Fixed preview not working properly first time when you open persona using skeleton

diff --git a/Engine/Source/Runtime/Engine/Private/Animation/Skeleton.cpp b/Engine/Source/Runtime/Engine/Private/Animation/Skeleton.cpp
--- a/Engine/Source/Runtime/Engine/Private/Animation/Skeleton.cpp
+++ b/Engine/Source/Runtime/Engine/Private/Animation/Skeleton.cpp
@@ -716,6 +716,8 @@ USkeletalMesh* USkeleton::GetPreviewMesh(bool bFindIfNotSet)
 			if(AssetList.Num() > 0)
 			{
 				SetPreviewMesh( Cast<USkeletalMesh>(AssetList[0].GetAsset()), false );
+				// update PreviewMesh
+				PreviewMesh = PreviewSkeletalMesh.Get();
 			}			
 		}
 	}
