#UE4:
* updated comments

diff --git a/Engine/Source/Runtime/Engine/Classes/Engine/TextureDefines.h b/Engine/Source/Runtime/Engine/Classes/Engine/TextureDefines.h
--- a/Engine/Source/Runtime/Engine/Classes/Engine/TextureDefines.h
+++ b/Engine/Source/Runtime/Engine/Classes/Engine/TextureDefines.h
@@ -6,10 +6,8 @@
 
 /**
  * @warning: if this is changed:
- *     update BaseEngine.ini SystemSettings
- *     update BaseCompat.ini AppCompatBucket 1 2 3
- *     update Game's DefaultEngine.ini SystemSettings
- *     update Game's BaseCompat.ini AppCompatBucket 1 2 3
+ *     update BaseEngine.ini [SystemSettings]
+ *     you might have to update the update Game's DefaultEngine.ini [SystemSettings]
  *     order and actual name can never change (order is important!)
  *
  * TEXTUREGROUP_Cinematic: should be used for Cinematics which will be baked out
@@ -43,10 +41,13 @@ enum TextureGroup
 	//obsolete - kept for back-compat
 	TEXTUREGROUP_ProcBuilding_LightMap UMETA(DisplayName="ProcBuilding_LightMap"),
 	TEXTUREGROUP_Shadowmap UMETA(DisplayName="Shadowmap"),
+	// no compression, no mips
 	TEXTUREGROUP_ColorLookupTable UMETA(DisplayName="ColorLookupTable"),
 	TEXTUREGROUP_Terrain_Heightmap UMETA(DisplayName="Terrain_Heightmap"),
 	TEXTUREGROUP_Terrain_Weightmap UMETA(DisplayName="Terrain_Weightmap"),
+	// using this TextureGroup triggers special mip map generation code only useful for the BokehDOF post process
 	TEXTUREGROUP_Bokeh UMETA(DisplayName="Bokeh"),
+	// no compression, created on import of a .IES file
 	TEXTUREGROUP_IESLightProfile UMETA(DisplayName="IESLightProfile"),
 	TEXTUREGROUP_MAX,
 };
