Fix overlap issues when character is standing on component that has sibling that's an overlap component

diff --git a/Engine/Source/Runtime/Engine/Classes/GameFramework/Actor.h b/Engine/Source/Runtime/Engine/Classes/GameFramework/Actor.h
--- a/Engine/Source/Runtime/Engine/Classes/GameFramework/Actor.h
+++ b/Engine/Source/Runtime/Engine/Classes/GameFramework/Actor.h
@@ -1624,7 +1624,14 @@ public:
 	 * @param Other the Actor to test for
 	 * @return true if this Actor is based on Other Actor
 	 */
-	virtual bool IsBasedOn( const AActor* Other ) const;
+	virtual bool AActor::IsBasedOnActor(const AActor* Other) const;
+	
+
+	/** iterates up the Base chain to see whether or not this Actor is attached to the given Actor
+	* @param Other the Actor to test for
+	* @return true if this Actor is attached on Other Actor
+	*/
+	virtual bool IsAttachedTo( const AActor* Other ) const;
 
 	/** Get the extent used when placing this actor in the editor, used for 'pulling back' hit. */
 	FVector GetPlacementExtent() const;
diff --git a/Engine/Source/Runtime/Engine/Classes/GameFramework/Character.h b/Engine/Source/Runtime/Engine/Classes/GameFramework/Character.h
--- a/Engine/Source/Runtime/Engine/Classes/GameFramework/Character.h
+++ b/Engine/Source/Runtime/Engine/Classes/GameFramework/Character.h
@@ -226,7 +226,6 @@ public:
 
 	// Begin AActor Interface.
 	virtual void TeleportSucceeded(bool bIsATest) OVERRIDE;
-	virtual bool IsBasedOn(const AActor* Other) const OVERRIDE;
 	virtual void ClearCrossLevelReferences() OVERRIDE;
 	virtual void PreNetReceive() OVERRIDE;
 	virtual void PostNetReceive() OVERRIDE;
diff --git a/Engine/Source/Runtime/Engine/Classes/GameFramework/Pawn.h b/Engine/Source/Runtime/Engine/Classes/GameFramework/Pawn.h
--- a/Engine/Source/Runtime/Engine/Classes/GameFramework/Pawn.h
+++ b/Engine/Source/Runtime/Engine/Classes/GameFramework/Pawn.h
@@ -142,6 +142,8 @@ public:
 	UFUNCTION(BlueprintPure, Category="Pawn", meta=(ToolTip="Gets the owning actor of the Movement Base Component on which the pawn is standing."))
 	static AActor* GetMovementBaseActor(const APawn* Pawn);
 
+	virtual bool IsBasedOnActor(const AActor * Other) const OVERRIDE;
+
 	virtual bool ReachedDesiredRotation();
 
 	/** @return The half-height of the default Pawn, scaled by the component scale. By default returns the half-height of the RootComponent, regardless of whether it is registered or collidable. */
diff --git a/Engine/Source/Runtime/Engine/Private/Actor.cpp b/Engine/Source/Runtime/Engine/Private/Actor.cpp
--- a/Engine/Source/Runtime/Engine/Private/Actor.cpp
+++ b/Engine/Source/Runtime/Engine/Private/Actor.cpp
@@ -811,18 +811,21 @@ void AActor::GetSimpleCollisionCylinder(float& CollisionRadius, float& Collision
 	}
 }
 
-
 bool AActor::IsRootComponentCollisionRegistered() const
 {
 	return RootComponent != NULL && RootComponent->IsRegistered() && RootComponent->IsCollisionEnabled();
 }
 
-
-bool AActor::IsBasedOn( const AActor* Other ) const
+bool AActor::IsAttachedTo(const AActor * Other) const
 {
 	return (RootComponent && Other && Other->RootComponent) ? RootComponent->IsAttachedTo(Other->RootComponent) : false;
 }
 
+bool AActor::IsBasedOnActor(const AActor* Other) const
+{
+	return IsAttachedTo(Other);
+}
+
 bool AActor::Modify( bool bAlwaysMarkDirty/*=true*/ )
 {
 	bool bSavedToTransactionBuffer = UObject::Modify( bAlwaysMarkDirty );
diff --git a/Engine/Source/Runtime/Engine/Private/Character.cpp b/Engine/Source/Runtime/Engine/Private/Character.cpp
--- a/Engine/Source/Runtime/Engine/Private/Character.cpp
+++ b/Engine/Source/Runtime/Engine/Private/Character.cpp
@@ -334,18 +334,6 @@ void ACharacter::ClearCrossLevelReferences()
 }
 
 
-bool ACharacter::IsBasedOn( const AActor* Other ) const
-{
-	AActor* MovementBaseActor = MovementBase ? MovementBase->GetOwner() : NULL;
-
-	if ( (MovementBaseActor == Other) || (MovementBaseActor && Cast<APawn>(MovementBaseActor) && Cast<APawn>(MovementBaseActor)->IsBasedOn(Other)) )
-	{
-		return true;
-	}
-	return Super::IsBasedOn(Other);
-}
-
-
 // Helper to set up tick dependencies on a movement base
 namespace MovementBaseUtility
 {
diff --git a/Engine/Source/Runtime/Engine/Private/Pawn.cpp b/Engine/Source/Runtime/Engine/Private/Pawn.cpp
--- a/Engine/Source/Runtime/Engine/Private/Pawn.cpp
+++ b/Engine/Source/Runtime/Engine/Private/Pawn.cpp
@@ -924,10 +924,24 @@ void APawn::PostNetReceiveLocation()
 	}
 }
 
+bool APawn::IsBasedOnActor(const AActor* Other) const
+{
+	UPrimitiveComponent * MovementBase = GetMovementBase();
+	AActor* MovementBaseActor = MovementBase ? MovementBase->GetOwner() : NULL;
+
+	if (MovementBaseActor && MovementBaseActor == Other)
+	{
+		return true;
+	}
+
+	return Super::IsBasedOnActor(Other);
+}
+
+
 bool APawn::IsNetRelevantFor(APlayerController* RealViewer, AActor* Viewer, const FVector& SrcLocation)
 {
 	if( bAlwaysRelevant || RealViewer == Controller || IsOwnedBy(Viewer) || IsOwnedBy(RealViewer) || this==Viewer || Viewer==Instigator
-		|| IsBasedOn(Viewer) || (Viewer && Viewer->IsBasedOn(this)) )
+		|| IsBasedOnActor(Viewer) || (Viewer && Viewer->IsBasedOnActor(this)))
 	{
 		return true;
 	}
diff --git a/Engine/Source/Runtime/Engine/Private/PrimitiveComponent.cpp b/Engine/Source/Runtime/Engine/Private/PrimitiveComponent.cpp
--- a/Engine/Source/Runtime/Engine/Private/PrimitiveComponent.cpp
+++ b/Engine/Source/Runtime/Engine/Private/PrimitiveComponent.cpp
@@ -1110,13 +1110,13 @@ const static float PERF_SHOW_MOVECOMPONENT_TAKING_LONG_TIME_AMOUNT = 2.0f; // mo
 static bool ShouldIgnoreHitResult(const UWorld* InWorld, FHitResult const& TestHit, FVector const& MovementDirDenormalized, AActor* MovingActor, EMoveComponentFlags MoveFlags)
 {
 	// check "ignore bases" functionality
-	if ( (MoveFlags & MOVECOMP_IgnoreBases) && MovingActor )
+	if ( (MoveFlags & MOVECOMP_IgnoreBases) && MovingActor && TestHit.bBlockingHit )	//we let overlap components go through because their overlap is still needed and will cause beginOverlap/endOverlap events
 	{
 		// ignore if there's a base relationship between moving actor and hit actor
 		AActor const* const HitActor = TestHit.GetActor();
 		if (HitActor)
 		{
-			if ( MovingActor->IsBasedOn(HitActor) || HitActor->IsBasedOn(MovingActor) )
+			if (MovingActor->IsBasedOnActor(HitActor) || HitActor->IsBasedOnActor(MovingActor))
 			{
 				return true;
 			}
@@ -2053,10 +2053,10 @@ void UPrimitiveComponent::UpdateOverlaps(TArray<FOverlapInfo> const* PendingOver
 					const FOverlapResult& Result = Overlaps[ResultIdx];
 
 					UPrimitiveComponent const* const HitComp = Result.Component.Get();
-					if (HitComp  && HitComp->bGenerateOverlapEvents)
+					if (!Result.bBlockingHit && HitComp  && HitComp->bGenerateOverlapEvents)
 					{
 						AActor* const ResultActor = Result.GetActor();
-						if ( (ResultActor != NULL) && (HitComp != this) && !ResultActor->IsBasedOn(MyActor) && (ResultActor != MyWorld->GetWorldSettings() && ResultActor->bActorInitialized) )
+						if ( (ResultActor != NULL) && (HitComp != this) && !ResultActor->IsAttachedTo(MyActor) && !MyActor->IsAttachedTo(ResultActor) && ResultActor != MyWorld->GetWorldSettings() && ResultActor->bActorInitialized)
 						{
 							NewOverlappingComponents.Add(FOverlapInfo(Result.Component.Get(), Result.ItemIndex));		// don't need to add unique unless the overlap check can return dupes
 						}						
