#ue4

- Analog movement in FloatingPawnMovement does not immediately cap the velocity when the max speed changes. Instead it decelerates smoothly to the new max speed.

diff --git a/Engine/Source/Runtime/Engine/Private/FloatingPawnMovement.cpp b/Engine/Source/Runtime/Engine/Private/FloatingPawnMovement.cpp
--- a/Engine/Source/Runtime/Engine/Private/FloatingPawnMovement.cpp
+++ b/Engine/Source/Runtime/Engine/Private/FloatingPawnMovement.cpp
@@ -87,12 +87,14 @@ bool UFloatingPawnMovement::LimitWorldBounds()
 void UFloatingPawnMovement::ApplyControlInputToVelocity(float DeltaTime)
 {
 	const FVector ControlAcceleration = GetInputVector().ClampMaxSize(1.f);
-	float AnalogInputModifier = 0.f;
-	
-	if (ControlAcceleration.SizeSquared() > 0.f)
+
+	const float AnalogInputModifier = (ControlAcceleration.SizeSquared() > 0.f ? ControlAcceleration.Size() : 0.f);
+	const float MaxSpeed = GetModifiedMaxSpeed() * AnalogInputModifier;
+	const bool bExceedingMaxSpeed = IsExceedingMaxSpeed(MaxSpeed);
+
+	if (AnalogInputModifier > 0.f && !bExceedingMaxSpeed)
 	{
 		// Apply change in velocity direction
-		AnalogInputModifier = ControlAcceleration.Size();
 		if (Velocity.SizeSquared() > 0.f)
 		{
 			Velocity -= (Velocity - ControlAcceleration * Velocity.Size()) * FMath::Min(DeltaTime * 8.f, 1.f);
@@ -103,15 +105,21 @@ void UFloatingPawnMovement::ApplyControlInputToVelocity(float DeltaTime)
 		// Dampen velocity magnitude based on deceleration.
 		if (Velocity.SizeSquared() > 0.f)
 		{
-			const float VelSize = FMath::Max(Velocity.Size() - Deceleration * DeltaTime, 0.f);
+			const FVector OldVelocity = Velocity;
+			const float VelSize = FMath::Max(Velocity.Size() - FMath::Abs(Deceleration) * DeltaTime, 0.f);
 			Velocity = Velocity.SafeNormal() * VelSize;
+
+			// Don't allow braking to lower us below max speed if we started above it.
+			if (bExceedingMaxSpeed && Velocity.SizeSquared() < FMath::Square(MaxSpeed))
+			{
+				Velocity = OldVelocity.SafeNormal() * MaxSpeed;
+			}
 		}
 	}
 
 	// Apply acceleration and clamp velocity magnitude.
-	const float MaxSpeed = GetModifiedMaxSpeed() * AnalogInputModifier;
 	const float NewMaxSpeed = (IsExceedingMaxSpeed(MaxSpeed)) ? Velocity.Size() : MaxSpeed;
-	Velocity += ControlAcceleration * Acceleration * DeltaTime;
+	Velocity += ControlAcceleration * FMath::Abs(Acceleration) * DeltaTime;
 	Velocity = Velocity.ClampMaxSize(NewMaxSpeed);
 
 	ConsumeInputVector();
