#UE4
Fix build break

diff --git a/Engine/Source/Runtime/Engine/Classes/GameFramework/Actor.h b/Engine/Source/Runtime/Engine/Classes/GameFramework/Actor.h
--- a/Engine/Source/Runtime/Engine/Classes/GameFramework/Actor.h
+++ b/Engine/Source/Runtime/Engine/Classes/GameFramework/Actor.h
@@ -1624,7 +1624,7 @@ public:
 	 * @param Other the Actor to test for
 	 * @return true if this Actor is based on Other Actor
 	 */
-	virtual bool AActor::IsBasedOnActor(const AActor* Other) const;
+	virtual bool IsBasedOnActor(const AActor* Other) const;
 	
 
 	/** iterates up the Base chain to see whether or not this Actor is attached to the given Actor
