#TTP 331039 - Landscape Foliage: Foliage is ignoring culldistance when not in G Mode

#proj UE4

#summary Foliage is now also distance-culled in editor viewports

diff --git a/Engine/Source/Runtime/Engine/Private/InstancedStaticMesh.cpp b/Engine/Source/Runtime/Engine/Private/InstancedStaticMesh.cpp
--- a/Engine/Source/Runtime/Engine/Private/InstancedStaticMesh.cpp
+++ b/Engine/Source/Runtime/Engine/Private/InstancedStaticMesh.cpp
@@ -799,8 +799,6 @@ public:
 	,	InstancedRenderData(InComponent)
 #if WITH_EDITOR
 	,	bHasSelectedInstances(InComponent->SelectedInstances.Num() > 0)
-#else
-	,	bHasSelectedInstances(false)
 #endif
 	{
 #if WITH_EDITOR
@@ -857,7 +855,7 @@ public:
 
 	// FPrimitiveSceneProxy interface.
 	
-	virtual FPrimitiveViewRelevance GetViewRelevance(const FSceneView* View)
+	virtual FPrimitiveViewRelevance GetViewRelevance(const FSceneView* View) OVERRIDE
 	{
 		FPrimitiveViewRelevance Result;
 		if(View->Family->EngineShowFlags.InstancedStaticMeshes)
@@ -883,49 +881,13 @@ public:
 	/** Sets up a wireframe FMeshBatch for a specific LOD. */
 	virtual bool GetWireframeMeshElement(int32 LODIndex, const FMaterialRenderProxy* WireframeRenderProxy, uint8 InDepthPriorityGroup, FMeshBatch& OutMeshElement) const OVERRIDE;
 
-	/**
-	 * Returns whether or not this component is instanced.
-	 * The base implementation returns false.  You should override this method in derived classes.
-	 *
-	 * @return	true if this component represents multiple instances of a primitive.
-	 */
-	virtual bool IsInstanced() const
-	{
-		return true;
-	}
-
-	/**
-	 * For instanced components, returns the number of instances.
-	 * The base implementation returns zero.  You should override this method in derived classes.
-	 *
-	 * @return	Number of instances
-	 */
-	virtual int32 GetInstanceCount() const
-	{
-		return PerInstanceSMData.Num();
-	}
-
-	/**
-	 * For instanced components, returns the Local -> World transform for the specific instance number.
-	 * If the function is called on non-instanced components, the component's LocalToWorld will be returned.
-	 * You should override this method in derived classes that support instancing.
-	 *
-	 * @param	InInstanceIndex	The index of the instance to return the Local -> World transform for
-	 *
-	 * @return	Number of instances
-	 */
-	virtual const FMatrix& GetInstanceLocalToWorld( int32 InInstanceIndex ) const
-	{
-		return PerInstanceSMData[ InInstanceIndex ].InstanceToWorld;
-	}
-
 	/**
 	 * Creates the hit proxies are used when DrawDynamicElements is called.
 	 * Called in the game thread.
 	 * @param OutHitProxies - Hit proxes which are created should be added to this array.
 	 * @return The hit proxy to use by default for elements drawn by DrawDynamicElements.
 	 */
-	virtual HHitProxy* CreateHitProxies(UPrimitiveComponent* Component,TArray<TRefCountPtr<HHitProxy> >& OutHitProxies)
+	virtual HHitProxy* CreateHitProxies(UPrimitiveComponent* Component,TArray<TRefCountPtr<HHitProxy> >& OutHitProxies) OVERRIDE
 	{
 		if( InstancedRenderData.HitProxies.Num() )
 		{
@@ -941,6 +903,11 @@ public:
 		}
 	}
 
+	virtual bool IsDetailMesh() const OVERRIDE
+	{
+		return true;
+	}
+
 private:
 
 	/** Array of per-instance static mesh rendering data for the scene proxy */
@@ -949,8 +916,12 @@ private:
 	/** Per component render data */
 	FInstancedStaticMeshRenderData InstancedRenderData;
 
+#if WITH_EDITOR
 	/* If we we have any selected instances */
 	bool bHasSelectedInstances;
+#else
+	static const bool bHasSelectedInstances = false;
+#endif
 
 	/** LOD transition info. */
 	FInstancingUserData UserData_AllInstances;
diff --git a/Engine/Source/Runtime/Engine/Public/PrimitiveSceneProxy.h b/Engine/Source/Runtime/Engine/Public/PrimitiveSceneProxy.h
--- a/Engine/Source/Runtime/Engine/Public/PrimitiveSceneProxy.h
+++ b/Engine/Source/Runtime/Engine/Public/PrimitiveSceneProxy.h
@@ -291,6 +291,15 @@ public:
 	inline int32 GetNumUncachedStaticLightingInteractions() { return NumUncachedStaticLightingInteractions; }
 #endif
 
+	/**
+	* Returns whether this proxy should be considered a "detail mesh".
+	* Detail meshes are distance culled even if distance culling is normally disabled for the view. (e.g. in editor)
+	*/
+	virtual bool IsDetailMesh() const
+	{
+		return false;
+	}
+
 	/**
 	 *	Returns whether the proxy utilizes custom occlusion bounds or not
 	 *
diff --git a/Engine/Source/Runtime/Renderer/Private/SceneVisibility.cpp b/Engine/Source/Runtime/Renderer/Private/SceneVisibility.cpp
--- a/Engine/Source/Runtime/Renderer/Private/SceneVisibility.cpp
+++ b/Engine/Source/Runtime/Renderer/Private/SceneVisibility.cpp
@@ -203,8 +203,9 @@ bool FViewInfo::IsDistanceCulled( float DistanceSquared, float MinDrawDistance,
 	float FadeRadius = GDisableLODFade ? 0.0f : GDistanceFadeMaxTravel;
 	float MaxDrawDistance = InMaxDrawDistance * MaxDrawDistanceScale;
 
-	// If cull distance is disabled, always show
-	if (Family->EngineShowFlags.DistanceCulledPrimitives)
+	// If cull distance is disabled, always show (except foliage)
+	if (Family->EngineShowFlags.DistanceCulledPrimitives
+		&& !PrimitiveSceneInfo->Proxy->IsDetailMesh())
 	{
 		return false;
 	}
@@ -258,8 +259,9 @@ static int32 FrustumCull(const FScene* Scene, FViewInfo& View)
 		float DistanceSquared = (Bounds.Origin - ViewOriginForDistanceCulling).SizeSquared();
 		float MaxDrawDistance = Bounds.MaxDrawDistance * MaxDrawDistanceScale;
 
-		// If cull distance is disabled, always show
-		if (View.Family->EngineShowFlags.DistanceCulledPrimitives)
+		// If cull distance is disabled, always show (except foliage)
+		if (View.Family->EngineShowFlags.DistanceCulledPrimitives
+			&& !Scene->Primitives[BitIt.GetIndex()]->Proxy->IsDetailMesh())
 		{
 			MaxDrawDistance = FLT_MAX;
 		}
