#TTP 330515 - FOLIAGE: Does not render when user enables show collision

#proj UE4

#summary Fixed foliage rendering in debug modes (e.g. "Show Collision") in non-editor builds

diff --git a/Engine/Source/Runtime/Engine/Private/InstancedStaticMesh.cpp b/Engine/Source/Runtime/Engine/Private/InstancedStaticMesh.cpp
--- a/Engine/Source/Runtime/Engine/Private/InstancedStaticMesh.cpp
+++ b/Engine/Source/Runtime/Engine/Private/InstancedStaticMesh.cpp
@@ -933,7 +933,7 @@ void FInstancedStaticMeshSceneProxy::DrawDynamicElements(FPrimitiveDrawInterface
 {
 	QUICK_SCOPE_CYCLE_COUNTER( STAT_InstancedStaticMeshSceneProxy_DrawDynamicElements );
 
-#if WITH_EDITOR
+#if !(UE_BUILD_SHIPPING || UE_BUILD_TEST)
 	const bool bSelectionRenderEnabled = GIsEditor && View->Family->EngineShowFlags.Selection;
 
 	// If the first pass rendered selected instances only, we need to render the deselected instances in a second pass
