#summary - PS4 - Texure updating / resolving now uses compute shaders
#change - RHICopyToResolveTarget now uses a compute shader. Now have the ability to copy from a tiled to a non-tiled texture.  (Also slightly faster).
#change - RHIUpdateTexture2D now uses a compute shader so it can work with tiled textures. (Also slightly faster).
#change - Surfaces created with CPUReadback are forced to be linear so that the CPU can read them directly.
#change - HZB occlusion culling now works on PS4.

diff --git a/Engine/Shaders/UpdateTextureShaders.usf b/Engine/Shaders/UpdateTextureShaders.usf
new file mode 100644
--- /dev/null
+++ b/Engine/Shaders/UpdateTextureShaders.usf
@@ -0,0 +1,30 @@
+// Copyright 1998-2014 Epic Games, Inc. All Rights Reserved.
+
+/*=============================================================================
+	UpdateTextureShaders.usf: Compute shaders for copying and updating textures.
+=============================================================================*/
+
+#include "Common.usf"
+
+Buffer<int4> SrcBuffer;
+RWTexture2D<int4> DestTexture;
+int SrcPitch;
+
+uint4 DestPosSize; // xy = pos, zw = Size of the destination sub rect
+
+[numthreads(8,8,1)]
+void UpdateTexture2DSubresourceCS( uint2 ThreadId : SV_DispatchThreadID )
+{
+	if( all( ThreadId.xy < DestPosSize.zw ) )
+	{
+		DestTexture[ ThreadId.xy + DestPosSize.xy ] = SrcBuffer[ ThreadId.x + ThreadId.y * SrcPitch ];
+	}
+}
+
+Texture2D<int4> SrcTexture;
+
+[numthreads(8,8,1)]
+void CopyTexture2DCS( uint3 ThreadId : SV_DispatchThreadID )
+{
+	DestTexture[ ThreadId.xy ] = SrcTexture.Load( int3( ThreadId.x, ThreadId.y, 0 ) );
+}
\ No newline at end of file
diff --git a/Engine/Source/Runtime/Engine/Private/UpdateTextureShaders.cpp b/Engine/Source/Runtime/Engine/Private/UpdateTextureShaders.cpp
new file mode 100644
--- /dev/null
+++ b/Engine/Source/Runtime/Engine/Private/UpdateTextureShaders.cpp
@@ -0,0 +1,8 @@
+// Copyright 1998-2014 Epic Games, Inc. All Rights Reserved.
+
+
+#include "EnginePrivate.h"
+#include "UpdateTextureShaders.h"
+
+IMPLEMENT_SHADER_TYPE(,FUpdateTexture2DSubresouceCS,TEXT("UpdateTextureShaders"),TEXT("UpdateTexture2DSubresourceCS"),SF_Compute);
+IMPLEMENT_SHADER_TYPE(,FCopyTexture2DCS,TEXT("UpdateTextureShaders"),TEXT("CopyTexture2DCS"),SF_Compute);
diff --git a/Engine/Source/Runtime/Engine/Public/UpdateTextureShaders.h b/Engine/Source/Runtime/Engine/Public/UpdateTextureShaders.h
new file mode 100644
--- /dev/null
+++ b/Engine/Source/Runtime/Engine/Public/UpdateTextureShaders.h
@@ -0,0 +1,72 @@
+// Copyright 1998-2014 Epic Games, Inc. All Rights Reserved.
+
+
+#pragma once
+
+#include "Engine.h"
+#include "GlobalShader.h"
+#include "ShaderParameters.h"
+
+class FUpdateTexture2DSubresouceCS : public FGlobalShader
+{
+	DECLARE_SHADER_TYPE(FUpdateTexture2DSubresouceCS,Global);
+public:
+	FUpdateTexture2DSubresouceCS() {}
+	FUpdateTexture2DSubresouceCS(const ShaderMetaType::CompiledShaderInitializerType& Initializer)
+		: FGlobalShader( Initializer )
+	{
+		SrcPitchParameter.Bind(Initializer.ParameterMap, TEXT("SrcPitch"), SPF_Mandatory);
+		SrcBuffer.Bind(Initializer.ParameterMap, TEXT("SrcBuffer"), SPF_Mandatory);
+		DestPosSizeParameter.Bind(Initializer.ParameterMap, TEXT("DestPosSize"), SPF_Mandatory);
+		DestTexture.Bind(Initializer.ParameterMap, TEXT("DestTexture"), SPF_Mandatory);
+	}
+
+	// FShader interface.
+	virtual bool Serialize(FArchive& Ar)
+	{
+		bool bShaderHasOutdatedParameters = FGlobalShader::Serialize(Ar);
+		Ar << SrcPitchParameter << SrcBuffer << DestPosSizeParameter << DestTexture;
+		return bShaderHasOutdatedParameters;
+	}
+
+	static bool ShouldCache(EShaderPlatform Platform)
+	{
+		return IsFeatureLevelSupported(Platform, ERHIFeatureLevel::SM5);
+	}
+
+//protected:
+	FShaderParameter SrcPitchParameter;
+	FShaderResourceParameter SrcBuffer;
+	FShaderParameter DestPosSizeParameter;
+	FShaderResourceParameter DestTexture;
+};
+
+class FCopyTexture2DCS : public FGlobalShader
+{
+	DECLARE_SHADER_TYPE(FCopyTexture2DCS,Global);
+public:
+	FCopyTexture2DCS() {}
+	FCopyTexture2DCS(const ShaderMetaType::CompiledShaderInitializerType& Initializer)
+		: FGlobalShader( Initializer )
+	{
+		SrcTexture.Bind(Initializer.ParameterMap, TEXT("SrcTexture"), SPF_Mandatory);
+		DestTexture.Bind(Initializer.ParameterMap, TEXT("DestTexture"), SPF_Mandatory);
+	}
+
+	// FShader interface.
+	virtual bool Serialize(FArchive& Ar)
+	{
+		bool bShaderHasOutdatedParameters = FGlobalShader::Serialize(Ar);
+		Ar << SrcTexture << DestTexture;
+		return bShaderHasOutdatedParameters;
+	}
+
+	static bool ShouldCache(EShaderPlatform Platform)
+	{
+		return IsFeatureLevelSupported(Platform, ERHIFeatureLevel::SM5);
+	}
+
+//protected:
+	FShaderResourceParameter SrcTexture;
+	FShaderResourceParameter DestTexture;
+};
