#UE4
- TTP#328091, navmesh bounds get updated in when Geometry-Editing ANavMeshBoundsVolume

diff --git a/Engine/Source/Runtime/Engine/Classes/AI/Navigation/NavMeshBoundsVolume.h b/Engine/Source/Runtime/Engine/Classes/AI/Navigation/NavMeshBoundsVolume.h
--- a/Engine/Source/Runtime/Engine/Classes/AI/Navigation/NavMeshBoundsVolume.h
+++ b/Engine/Source/Runtime/Engine/Classes/AI/Navigation/NavMeshBoundsVolume.h
@@ -15,8 +15,9 @@ class ANavMeshBoundsVolume : public AVolume
 
 	// Begin AActor Interface
 	virtual void PostInitializeComponents() OVERRIDE;
-	// End AActor Interface
 #if WITH_EDITOR
+	virtual void ReregisterAllComponents() OVERRIDE;
+	// End AActor Interface
 	// Begin UObject Interface
 	virtual void PostEditChangeProperty( struct FPropertyChangedEvent& PropertyChangedEvent) OVERRIDE;
 	// End UObject Interface
diff --git a/Engine/Source/Runtime/Engine/Classes/AI/Navigation/NavigationSystem.h b/Engine/Source/Runtime/Engine/Classes/AI/Navigation/NavigationSystem.h
--- a/Engine/Source/Runtime/Engine/Classes/AI/Navigation/NavigationSystem.h
+++ b/Engine/Source/Runtime/Engine/Classes/AI/Navigation/NavigationSystem.h
@@ -908,6 +908,8 @@ public:
 	void EndFakeComponentChanges() { bFakeComponentChangesBeingApplied = false; }
 
 	void UpdateLevelCollision(ULevel* InLevel);
+
+	virtual void OnEditorModeChanged(class FEdMode* Mode, bool IsEntering);
 #endif // WITH_EDITOR
 
 	/** register actor important for generation (navigation data will be build around them first) */
diff --git a/Engine/Source/Runtime/Engine/Private/AI/Navigation/NavigationSystem.cpp b/Engine/Source/Runtime/Engine/Private/AI/Navigation/NavigationSystem.cpp
--- a/Engine/Source/Runtime/Engine/Private/AI/Navigation/NavigationSystem.cpp
+++ b/Engine/Source/Runtime/Engine/Private/AI/Navigation/NavigationSystem.cpp
@@ -6,6 +6,11 @@
 #if WITH_RECAST
 #include "RecastNavMeshGenerator.h"
 #endif // WITH_RECAST
+#if WITH_EDITOR
+#include "UnrealEd.h"
+#include "Editor/GeometryMode/Public/GeometryEdMode.h"
+#include "Editor/GeometryMode/Public/EditorGeometry.h"
+#endif
 
 static const uint32 INITIAL_ASYNC_QUERIES_SIZE = 32;
 static const uint32 REGISTRATION_QUEUE_SIZE = 16;	// and we'll not reallocate
@@ -295,11 +300,24 @@ UNavigationSystem::UNavigationSystem(const class FPostConstructInitializePropert
 		DefaultWalkableArea = UNavArea_Default::StaticClass();
 		DefaultObstacleArea = UNavArea_Null::StaticClass();
 	}
+
+#if WITH_EDITOR
+	if (GIsEditor && HasAnyFlags(RF_ClassDefaultObject) == false)
+	{
+		GEditorModeTools().OnEditorModeChanged().AddUObject(this, &UNavigationSystem::OnEditorModeChanged);
+	}
+#endif // WITH_EDITOR
 }
 
 UNavigationSystem::~UNavigationSystem()
 {
 	CleanUp();
+#if WITH_EDITOR
+	if (GIsEditor)
+	{
+		GEditorModeTools().OnEditorModeChanged().RemoveAll(this);
+	}
+#endif // WITH_EDITOR
 }
 
 void UNavigationSystem::DoInitialSetup()
@@ -2193,6 +2211,28 @@ void UNavigationSystem::UpdateLevelCollision(ULevel* InLevel)
 		OnLevelAddedToWorld(InLevel, World);
 	}
 }
+
+void UNavigationSystem::OnEditorModeChanged(FEdMode* Mode, bool IsEntering)
+{
+	if (Mode == NULL)
+	{
+		return;
+	}
+
+	if (IsEntering == false && Mode->GetID() == FBuiltinEditorModes::EM_Geometry)
+	{
+		// check if any of modified brushes belongs to an ANavMeshBoundsVolume
+		FEdModeGeometry* GeometryMode = (FEdModeGeometry*)Mode;
+		for (auto GeomObjectIt = GeometryMode->GeomObjectItor(); GeomObjectIt; GeomObjectIt++)
+		{
+			ANavMeshBoundsVolume* Volume = Cast<ANavMeshBoundsVolume>((*GeomObjectIt)->GetActualBrush());
+			if (Volume)
+			{
+				OnNavigationBoundsUpdated(Volume);
+			}
+		}
+	}
+}
 #endif
 
 #if WITH_NAVIGATION_GENERATOR
diff --git a/Engine/Source/Runtime/Engine/Private/NavMeshBoundsVolume.cpp b/Engine/Source/Runtime/Engine/Private/NavMeshBoundsVolume.cpp
--- a/Engine/Source/Runtime/Engine/Private/NavMeshBoundsVolume.cpp
+++ b/Engine/Source/Runtime/Engine/Private/NavMeshBoundsVolume.cpp
@@ -19,6 +19,24 @@ ANavMeshBoundsVolume::ANavMeshBoundsVolume(const class FPostConstructInitializeP
 
 #if WITH_EDITOR
 
+void ANavMeshBoundsVolume::ReregisterAllComponents()
+{
+	Super::ReregisterAllComponents();
+	/*
+	// uncommenting this would make navigation rebuild during NavMeshBoundsVolume's
+	// changing with Geometry Tool. This gets called too often though, so we'll 
+	// stick to updating navigation upon leaving geometry tool
+	if (GIsEditor)
+	{
+		UWorld* World = GetWorld();
+		if (World != NULL && World->GetNavigationSystem() != NULL)
+		{
+		GetWorld()->GetNavigationSystem()->OnNavigationBoundsUpdated(this);
+		}
+	}
+	*/
+}
+
 void ANavMeshBoundsVolume::PostEditChangeProperty( struct FPropertyChangedEvent& PropertyChangedEvent)
 {
 	Super::PostEditChangeProperty(PropertyChangedEvent);
