Fix no hit notification when hitting BSP

diff --git a/Engine/Source/Runtime/Engine/Private/PrimitiveComponent.cpp b/Engine/Source/Runtime/Engine/Private/PrimitiveComponent.cpp
--- a/Engine/Source/Runtime/Engine/Private/PrimitiveComponent.cpp
+++ b/Engine/Source/Runtime/Engine/Private/PrimitiveComponent.cpp
@@ -1448,7 +1448,7 @@ bool UPrimitiveComponent::MoveComponent( const FVector& Delta, const FRotator& N
 		bMoved = InternalSetWorldLocationAndRotation(NewLocation, NewRotation, bSkipPhysicsMove);
 
 		// Handle hit notifications.
-		if( BlockingHit.GetActor() != NULL )
+		if (BlockingHit.bBlockingHit)
 		{
 			if (Actor != NULL)
 			{
