#ttp 319048 - Analytics warning: Event Editor.Usage.Modules has too many attributes (138).

#branch UE4
#proj Editor.UnrealEd, Runtime.Engine
#summary Only reporting Analytics stats for game modules' compile methods (not engine code) and reporting one event per module (not one event with many attributes).

#change FUnrealEdMisc::InitEngineAnalytics() - checking FModuleStatus::bIsGameModule before reporting compile method via Editor.Usage.Modules event. Reporting the event once per module rather than one attribute per module.

#extra Added DEBUG_ANALYTICS define and a dummy IAnalyticsProvider implementation to EngineAnalytics.cpp. Setting the define to non-zero makes all calls to FEngineAnalytics go via the dummy provider that simply logs event info to log category LogDebugAnalytics. This is useful when debugging changes to analytics events because they can be checked without setting up a real provider. Set an attribute limit of 40 based on the cut-off in the spreadsheet attached to this ttp. Breaking this limit using the debug provider generates a warning.

diff --git a/Engine/Source/Editor/UnrealEd/Private/UnrealEdMisc.cpp b/Engine/Source/Editor/UnrealEd/Private/UnrealEdMisc.cpp
--- a/Engine/Source/Editor/UnrealEd/Private/UnrealEdMisc.cpp
+++ b/Engine/Source/Editor/UnrealEd/Private/UnrealEdMisc.cpp
@@ -374,18 +374,20 @@ void FUnrealEdMisc::InitEngineAnalytics()
 			bIsAssetAnalyticsPending = true;
 		}
 
-// @todo: cwood: Disabled "Editor.Usage.Modules" event due to it spamming analytics. Needs better solution.
-#if 0
 		// Record known modules' compilation methods
 		TArray<FModuleManager::FModuleStatus> Modules;
 		FModuleManager::Get().QueryModules(Modules);
-		TArray< FAnalyticsEventAttribute > ModuleAttributes;
-		for (auto ModuleIter = Modules.CreateConstIterator(); ModuleIter; ++ModuleIter)
+		for (auto& Module : Modules)
 		{
-			ModuleAttributes.Add(FAnalyticsEventAttribute(ModuleIter->Name, ModuleIter->CompilationMethod));
+			// Record only game modules as these are the only ones that should be hot-reloaded
+			if (Module.bIsGameModule)
+			{
+				TArray< FAnalyticsEventAttribute > ModuleAttributes;
+				ModuleAttributes.Add(FAnalyticsEventAttribute(FString("ModuleName"), Module.Name));
+				ModuleAttributes.Add(FAnalyticsEventAttribute(FString("CompilationMethod"), Module.CompilationMethod));
+				EngineAnalytics.RecordEvent(FString("Editor.Usage.Modules"), ModuleAttributes);
+			}
 		}
-		EngineAnalytics.RecordEvent(FString("Editor.Usage.Modules"), ModuleAttributes);
-#endif
 	}
 }
 
diff --git a/Engine/Source/Runtime/Engine/Private/EngineAnalytics.cpp b/Engine/Source/Runtime/Engine/Private/EngineAnalytics.cpp
--- a/Engine/Source/Runtime/Engine/Private/EngineAnalytics.cpp
+++ b/Engine/Source/Runtime/Engine/Private/EngineAnalytics.cpp
@@ -9,6 +9,89 @@
 bool FEngineAnalytics::bIsInitialized;
 TSharedPtr<IAnalyticsProvider> FEngineAnalytics::Analytics;
 
+#define DEBUG_ANALYTICS 0
+
+#if DEBUG_ANALYTICS
+
+DEFINE_LOG_CATEGORY_STATIC(LogDebugAnalytics, Log, All);
+
+namespace DebugAnalyticsProviderDefs
+{
+	const static int32 MaxAttributeCount = 40;
+}
+
+/**
+ * Debug analytics provider for debugging analytics events
+ * without the need to connect to a real provider.
+ */
+class FDebugAnalyticsProvider : public IAnalyticsProvider
+{
+public:
+	virtual bool StartSession(const TArray<FAnalyticsEventAttribute>& Attributes)
+	{
+		UE_LOG(LogDebugAnalytics, Log, TEXT("FDebugAnalyticsProvider: StartSession..."));
+		OutputAttributesToLog(Attributes);
+		return true;
+	}
+
+	virtual void EndSession()
+	{
+		UE_LOG(LogDebugAnalytics, Log, TEXT("FDebugAnalyticsProvider: EndSession"));
+	}
+
+	virtual FString GetSessionID() const
+	{
+		return SessionID;
+	}
+
+	virtual bool SetSessionID(const FString& InSessionID)
+	{
+		UE_LOG(LogDebugAnalytics, Log, TEXT("FDebugAnalyticsProvider: SetSessionID SessionID=%s"), *InSessionID);
+		SessionID = InSessionID;
+		return true;
+	}
+
+	virtual void FlushEvents()
+	{
+	}
+
+	virtual void SetUserID(const FString& InUserID)
+	{
+		UE_LOG(LogDebugAnalytics, Log, TEXT("FDebugAnalyticsProvider: SetUserID UserID=%s"), *InUserID);
+		UserID = InUserID;
+	}
+
+	virtual FString GetUserID() const
+	{
+		return UserID;
+	}
+
+	virtual void RecordEvent(const FString& EventName, const TArray<FAnalyticsEventAttribute>& Attributes)
+	{
+		if (Attributes.Num() > DebugAnalyticsProviderDefs::MaxAttributeCount)
+		{
+			UE_LOG(LogDebugAnalytics, Warning, TEXT("FDebugAnalyticsProvider: Event %s has too many attributes (%d)"), *EventName, Attributes.Num());
+		}
+
+		UE_LOG(LogDebugAnalytics, Log, TEXT("FDebugAnalyticsProvider: RecordEvent Name=%s"), *EventName);
+		OutputAttributesToLog(Attributes);
+	}
+
+private:
+	void OutputAttributesToLog(const TArray<FAnalyticsEventAttribute>& Attributes)
+	{
+		for (auto& Attrib : Attributes)
+		{
+			UE_LOG(LogDebugAnalytics, Log, TEXT("AnalyticsAttrib Name=%s, Value=%s"), *Attrib.AttrName, *Attrib.AttrValue);
+		}
+	}
+
+	FString SessionID;
+	FString UserID;
+};
+
+#endif // DEBUG_ANALYTICS
+
 /**
  * Engine analytics config log to initialize the engine analytics provider.
  * External code should bind this delegate if engine analytics are desired,
@@ -35,7 +118,7 @@ IAnalyticsProvider& FEngineAnalytics::GetProvider()
 	checkf(bIsInitialized && Analytics.IsValid(), TEXT("FEngineAnalytics::GetProvider called outside of Initialize/Shutdown."));
 	return *Analytics.Get();
 }
-
+ 
 void FEngineAnalytics::Initialize()
 {
 	checkf(!bIsInitialized, TEXT("FEngineAnalytics::Initialize called more than once."));
@@ -44,6 +127,12 @@ void FEngineAnalytics::Initialize()
 	const bool bShouldInitAnalytics = !IsRunningCommandlet();
 	if( bShouldInitAnalytics )
 	{
+#if DEBUG_ANALYTICS
+		// Dummy analytics provider just logs events locally
+		Analytics = MakeShareable(new FDebugAnalyticsProvider());
+		Analytics->SetUserID(FPlatformMisc::GetUniqueDeviceId());
+		Analytics->StartSession();
+#else
 		// Connect the engine analytics provider (if there is a configuration delegate installed)
 		if (GetEngineAnalyticsConfigDelegate().IsBound())
 		{
@@ -57,6 +146,7 @@ void FEngineAnalytics::Initialize()
 				Analytics->StartSession();
 			}
 		}
+#endif
 	}
 	bIsInitialized = true;
 
