Fix destructible mesh ignoring Set Gravity flag

diff --git a/Engine/Source/Runtime/Engine/Private/DestructibleComponent.cpp b/Engine/Source/Runtime/Engine/Private/DestructibleComponent.cpp
--- a/Engine/Source/Runtime/Engine/Private/DestructibleComponent.cpp
+++ b/Engine/Source/Runtime/Engine/Private/DestructibleComponent.cpp
@@ -325,6 +325,8 @@ void UDestructibleComponent::CreatePhysicsState()
 	{
 		SCOPED_SCENE_WRITE_LOCK(PRootActor->getScene());
 
+		PRootActor->setActorFlag(PxActorFlag::eDISABLE_GRAVITY, !BodyInstance.bEnableGravity);
+
 		// Sleep/wake up as appropriate
 		if (BodyInstance.bStartAwake)
 		{
@@ -788,6 +790,8 @@ void UDestructibleComponent::SetChunkVisible( int32 ChunkIndex, bool bVisible )
 				CI->OwningComponent = this;
 				CI->Actor = PActor;
 
+				PActor->setActorFlag(PxActorFlag::eDISABLE_GRAVITY, !BodyInstance.bEnableGravity);
+
 				int32 UserDataIdx = PhysxChunkUserData.Add(FPhysxUserData(CI));
 				check(InfoIndex == UserDataIdx);
 
