#UE4
Pawn and Controller should expose playerstate to blueprints

diff --git a/Engine/Source/Runtime/Engine/Classes/GameFramework/Controller.h b/Engine/Source/Runtime/Engine/Classes/GameFramework/Controller.h
--- a/Engine/Source/Runtime/Engine/Classes/GameFramework/Controller.h
+++ b/Engine/Source/Runtime/Engine/Classes/GameFramework/Controller.h
@@ -50,7 +50,7 @@ private:
 
 public:
 	/** PlayerState containing replicated information about the player using this controller (only exists for players, not NPCs). */
-	UPROPERTY(editinline, replicatedUsing=OnRep_PlayerState)
+	UPROPERTY(editinline, replicatedUsing=OnRep_PlayerState, BlueprintReadOnly, Category="Controller")
 	class APlayerState* PlayerState;
 
 protected:
diff --git a/Engine/Source/Runtime/Engine/Classes/GameFramework/Pawn.h b/Engine/Source/Runtime/Engine/Classes/GameFramework/Pawn.h
--- a/Engine/Source/Runtime/Engine/Classes/GameFramework/Pawn.h
+++ b/Engine/Source/Runtime/Engine/Classes/GameFramework/Pawn.h
@@ -85,7 +85,7 @@ public:
 	TSubclassOf<class AAIController>  AIControllerClass;
 
 	/** If Pawn is possessed by a player, points to his playerstate.  Needed for network play as controllers are not replicated to clients. */
-	UPROPERTY(editinline, replicatedUsing=OnRep_PlayerState)
+	UPROPERTY(editinline, replicatedUsing=OnRep_PlayerState, BlueprintReadOnly, Category="Pawn")
 	class APlayerState* PlayerState;
 
 	/** Replicated so we can see where remote clients are looking. */
