Prevent divide by zero in CrossFade sound node when FadeOut or FadeIn MaxDistance is 0.

diff --git a/Engine/Source/Runtime/Engine/Private/SoundNodeDistanceCrossFade.cpp b/Engine/Source/Runtime/Engine/Private/SoundNodeDistanceCrossFade.cpp
--- a/Engine/Source/Runtime/Engine/Private/SoundNodeDistanceCrossFade.cpp
+++ b/Engine/Source/Runtime/Engine/Private/SoundNodeDistanceCrossFade.cpp
@@ -64,13 +64,13 @@ void USoundNodeDistanceCrossFade::ParseNodes( FAudioDevice* AudioDevice, const U
 			}
 			else if( ( Distance >= FadeInDistanceMin ) && ( Distance <= FadeInDistanceMax ) )
 			{
-				VolumeToSet = CrossFadeInput[ ChildNodeIndex ].Volume * ( 0.0f + ( Distance - FadeInDistanceMin ) / ( FadeInDistanceMax - FadeInDistanceMin ) );
+				VolumeToSet = (FadeInDistanceMax > 0.f ? CrossFadeInput[ ChildNodeIndex ].Volume * ( 0.0f + ( Distance - FadeInDistanceMin ) / ( FadeInDistanceMax - FadeInDistanceMin ) ) : 1.f);
 				//UE_LOG(LogAudio, Log,  TEXT("     FadeIn.  Distance: %f,  VolumeToSet: %f"), Distance, VolumeToSet );
 			}
 			// else if we are inside the FadeOut edge
 			else if( ( Distance >= FadeOutDistanceMin ) && ( Distance <= FadeOutDistanceMax ) )
 			{
-				VolumeToSet = CrossFadeInput[ ChildNodeIndex ].Volume * ( 1.0f - ( Distance - FadeOutDistanceMin ) / ( FadeOutDistanceMax - FadeOutDistanceMin ) );
+				VolumeToSet = (FadeInDistanceMax > 0.f ? CrossFadeInput[ ChildNodeIndex ].Volume * ( 1.0f - ( Distance - FadeOutDistanceMin ) / ( FadeOutDistanceMax - FadeOutDistanceMin ) ) : 0.f);
 				//UE_LOG(LogAudio, Log,  TEXT("     FadeOut.  Distance: %f,  VolumeToSet: %f"), Distance, VolumeToSet );
 			}
 			// else we are in between the fading edges of the CrossFaded sound and we should play the
@@ -84,7 +84,7 @@ void USoundNodeDistanceCrossFade::ParseNodes( FAudioDevice* AudioDevice, const U
 			else
 			{
 				//UE_LOG(LogAudio, Log,  TEXT("     OUTSIDE!!!" ));
-				VolumeToSet = 0; //CrossFadeInput( ChildNodeIndex ).Volume;
+				VolumeToSet = 0.f; //CrossFadeInput( ChildNodeIndex ).Volume;
 			}
 
 			UpdatedParams.Volume = ParseParams.Volume * VolumeToSet;
