Fix crash when assigning an animation blueprint to a skeletal mesh component

diff --git a/Engine/Source/Runtime/Engine/Private/SkeletalMeshComponent.cpp b/Engine/Source/Runtime/Engine/Private/SkeletalMeshComponent.cpp
--- a/Engine/Source/Runtime/Engine/Private/SkeletalMeshComponent.cpp
+++ b/Engine/Source/Runtime/Engine/Private/SkeletalMeshComponent.cpp
@@ -293,21 +293,7 @@ void USkeletalMeshComponent::PostEditChangeProperty(FPropertyChangedEvent& Prope
 
 		if ( PropertyThatChanged->GetFName() == GET_MEMBER_NAME_CHECKED( USkeletalMeshComponent, AnimBlueprintGeneratedClass ) )
 		{
-			if (AnimBlueprintGeneratedClass == NULL)
-			{
-				if (AnimationMode == EAnimationMode::AnimationBlueprint)
-				{
-					ClearAnimScriptInstance();
-				}
-			}
-			else
-			{
-				if (NeedToSpawnAnimScriptInstance(false))
-				{
-					AnimScriptInstance = NewObject<UAnimInstance>(this, AnimBlueprintGeneratedClass);
-					AnimScriptInstance->InitializeAnimation();
-				}
-			}
+			InitAnim(false);
 		}
 
 		if(PropertyThatChanged->GetFName() == GET_MEMBER_NAME_CHECKED( USkeletalMeshComponent, SkeletalMesh))
