#UE4 - Fix bug where rep notifies weren't being called under certain cases

diff --git a/Engine/Source/Runtime/Engine/Private/DataReplication.cpp b/Engine/Source/Runtime/Engine/Private/DataReplication.cpp
--- a/Engine/Source/Runtime/Engine/Private/DataReplication.cpp
+++ b/Engine/Source/Runtime/Engine/Private/DataReplication.cpp
@@ -1112,15 +1112,15 @@ bool FObjectReplicator::UpdateUnmappedObjects()
 		return false;
 	}
 
-	if ( !RepLayout->UpdateUnmappedObjects( RepState, Connection->PackageMap, Object ) )
-	{
-		return false;
-	}
+	check( RepState->RepNotifies.Num() == 0 );
+
+	// Cache the result, since this function will return false when there are no more unmapped properties, but we still need to call any rep notifies
+	const bool Result = RepLayout->UpdateUnmappedObjects( RepState, Connection->PackageMap, Object );
 
 	// Call any rep notifies that need to happen when object pointers change
 	RepLayout->CallRepNotifies( RepState, Object );
 
-	return true;
+	return Result;
 }
 
 void FObjectReplicator::QueuePropertyRepNotify( UObject * Object, UProperty * Property, const int32 ElementIndex, TArray< uint8 > & MetaData )
