#UE4 - Enable unmap logic on client now that a bug has been fixed

diff --git a/Engine/Source/Runtime/Engine/Private/RepLayout.cpp b/Engine/Source/Runtime/Engine/Private/RepLayout.cpp
--- a/Engine/Source/Runtime/Engine/Private/RepLayout.cpp
+++ b/Engine/Source/Runtime/Engine/Private/RepLayout.cpp
@@ -25,7 +25,7 @@ FAutoConsoleVariable CVarDoReplicationContextString( TEXT( "net.ContextDebug" ),
 
 static const int32 UNMAPPED_FRAMES_THRESHOLD = 100;
 
-#define ENABLE_CLIENT_UNMAP_LOGIC 0
+#define ENABLE_CLIENT_UNMAP_LOGIC 1
 
 #ifdef USE_CUSTOM_COMPARE
 static FORCEINLINE bool CompareBool( const FRepLayoutCmd & Cmd, const void * A, const void * B )
