Change soft limits to include mass

diff --git a/Engine/Source/Runtime/Core/Public/UObject/ObjectVersion.h b/Engine/Source/Runtime/Core/Public/UObject/ObjectVersion.h
--- a/Engine/Source/Runtime/Core/Public/UObject/ObjectVersion.h
+++ b/Engine/Source/Runtime/Core/Public/UObject/ObjectVersion.h
@@ -511,6 +511,9 @@ enum EUnrealEngineObjectUE4Version
 	// we still convert all to FVector in CPU time whenever any calculation
 	VER_UE4_MORPHTARGET_CPU_TANGENTZDELTA_FORMATCHANGE,
 
+	// Soft constraint limits will implicitly use the mass of the bodies
+	VER_UE4_SOFT_CONSTRAINTS_USE_MASS,
+
 	// -----<new versions can be added before this line>-------------------------------------------------
 	// - this needs to be the last line (see note below)
 	VER_UE4_AUTOMATIC_VERSION_PLUS_ONE,
diff --git a/Engine/Source/Runtime/Engine/Classes/Components/PrimitiveComponent.h b/Engine/Source/Runtime/Engine/Classes/Components/PrimitiveComponent.h
--- a/Engine/Source/Runtime/Engine/Classes/Components/PrimitiveComponent.h
+++ b/Engine/Source/Runtime/Engine/Classes/Components/PrimitiveComponent.h
@@ -1284,7 +1284,7 @@ public:
 	virtual float GetMass() const;
 	
 	/** Returns the calculated mass in kg. This is not 100% exactly the mass physx will calculate, but it is very close ( difference < 0.1kg ). */
-	virtual float CalculateMass();
+	virtual float CalculateMass(FName BoneName = NAME_None);
 
 	/**
 	 *	Force all bodies in this component to sleep.
diff --git a/Engine/Source/Runtime/Engine/Classes/Components/SkeletalMeshComponent.h b/Engine/Source/Runtime/Engine/Classes/Components/SkeletalMeshComponent.h
--- a/Engine/Source/Runtime/Engine/Classes/Components/SkeletalMeshComponent.h
+++ b/Engine/Source/Runtime/Engine/Classes/Components/SkeletalMeshComponent.h
@@ -693,7 +693,7 @@ public:
 	virtual void AddRadialForce(FVector Origin, float Radius, float Strength, ERadialImpulseFalloff Falloff) OVERRIDE;
 	virtual void SetAllPhysicsLinearVelocity(FVector NewVel,bool bAddToCurrent = false) OVERRIDE;
 	virtual float GetMass() const OVERRIDE;
-	virtual float CalculateMass() OVERRIDE;
+	virtual float CalculateMass(FName BoneName = NAME_None) OVERRIDE;
 	// End UPrimitiveComponent interface.
 
 	// Begin USkinnedMeshComponent interface
diff --git a/Engine/Source/Runtime/Engine/Private/PhysicsEngine/ConstraintInstance.cpp b/Engine/Source/Runtime/Engine/Private/PhysicsEngine/ConstraintInstance.cpp
--- a/Engine/Source/Runtime/Engine/Private/PhysicsEngine/ConstraintInstance.cpp
+++ b/Engine/Source/Runtime/Engine/Private/PhysicsEngine/ConstraintInstance.cpp
@@ -144,6 +144,28 @@ void FConstraintInstance::InitConstraint(USceneComponent* Owner, FBodyInstance*
 		}
 	}
 
+	float AverageMass = 0;
+	{
+		float TotalMass = 0;
+		int NumDynamic = 0;
+
+		if (PActor1 && PActor1->isRigidDynamic())
+		{
+			TotalMass += PActor1->isRigidDynamic()->getMass();
+			++NumDynamic;
+		}
+
+		if (PActor2 && PActor2->isRigidDynamic())
+		{
+			TotalMass += PActor2->isRigidDynamic()->getMass();
+			++NumDynamic;
+		}
+
+		check(NumDynamic);
+		AverageMass = TotalMass / NumDynamic;
+	}
+
+
 	// record if actors are asleep before creating joint, so we can sleep them afterwards if so (creating joint wakes them)
 	const bool bActor1Asleep = (PActor1 == NULL || !PActor1->isRigidDynamic() || PActor1->isRigidDynamic()->isSleeping());
 	const bool bActor2Asleep = (PActor2 == NULL || !PActor2->isRigidDynamic() || PActor2->isRigidDynamic()->isSleeping());
@@ -236,7 +258,7 @@ void FConstraintInstance::InitConstraint(USceneComponent* Owner, FBodyInstance*
 		// If angle drops below RB_MinAngleToLockDOF, just pass RB_MinAngleToLockDOF to physics - that axis will be locked anyway, and PhysX dislikes 0 here
 		float TwistLimitRad = TwistLimitAngle * (PI/180.0f);
 		PxJointAngularLimitPair PTwistLimitPair(-TwistLimitRad, TwistLimitRad, 1.f * (PI/180.0f));
-		SetSoftLimitParams(&PTwistLimitPair, bTwistLimitSoft, TwistLimitStiffness*CVarConstraintSitffnessScale->GetValueOnGameThread(), TwistLimitDamping*CVarConstraintDamplingScale->GetValueOnGameThread());
+		SetSoftLimitParams(&PTwistLimitPair, bTwistLimitSoft, TwistLimitStiffness*AverageMass*CVarConstraintSitffnessScale->GetValueOnGameThread(), TwistLimitDamping*AverageMass*CVarConstraintDamplingScale->GetValueOnGameThread());
 		PD6Joint->setTwistLimit(PTwistLimitPair);
 	}
 	else if (AngularTwistMotion==ACM_Locked)
@@ -274,7 +296,7 @@ void FConstraintInstance::InitConstraint(USceneComponent* Owner, FBodyInstance*
 		float Limit1Rad = FMath::ClampAngle(Swing1LimitAngle, KINDA_SMALL_NUMBER, 179.9999f) * (PI/180.0f);
 		float Limit2Rad = FMath::ClampAngle(Swing2LimitAngle, KINDA_SMALL_NUMBER, 179.9999f) * (PI/180.0f);
 		PxJointLimitCone PSwingLimitCone(Limit2Rad, Limit1Rad, 1.f * (PI/180.0f));
-		SetSoftLimitParams(&PSwingLimitCone, bSwingLimitSoft, SwingLimitStiffness*CVarConstraintSitffnessScale->GetValueOnGameThread(), SwingLimitDamping*CVarConstraintDamplingScale->GetValueOnGameThread());
+		SetSoftLimitParams(&PSwingLimitCone, bSwingLimitSoft, SwingLimitStiffness*AverageMass*CVarConstraintSitffnessScale->GetValueOnGameThread(), SwingLimitDamping*AverageMass*CVarConstraintDamplingScale->GetValueOnGameThread());
 		PD6Joint->setSwingLimit(PSwingLimitCone);
 	}
 
@@ -294,7 +316,7 @@ void FConstraintInstance::InitConstraint(USceneComponent* Owner, FBodyInstance*
 		// If limit drops below RB_MinSizeToLockDOF, just pass RB_MinSizeToLockDOF to physics - that axis will be locked anyway, and PhysX dislikes 0 here
 		float LinearLimit = FMath::Max<float>(LinearLimitSize, RB_MinSizeToLockDOF);
 		PxJointLinearLimit PLinearLimit(GPhysXSDK->getTolerancesScale(), LinearLimit, 0.05f * GPhysXSDK->getTolerancesScale().length);
-		SetSoftLimitParams(&PLinearLimit, bLinearLimitSoft, LinearLimitStiffness*CVarConstraintSitffnessScale->GetValueOnGameThread(), LinearLimitDamping*CVarConstraintDamplingScale->GetValueOnGameThread());
+		SetSoftLimitParams(&PLinearLimit, bLinearLimitSoft, LinearLimitStiffness*AverageMass*CVarConstraintSitffnessScale->GetValueOnGameThread(), LinearLimitDamping*AverageMass*CVarConstraintDamplingScale->GetValueOnGameThread());
 		PD6Joint->setLinearLimit(PLinearLimit);
 	}
 
diff --git a/Engine/Source/Runtime/Engine/Private/PhysicsEngine/PhysicsConstraintComponent.cpp b/Engine/Source/Runtime/Engine/Private/PhysicsEngine/PhysicsConstraintComponent.cpp
--- a/Engine/Source/Runtime/Engine/Private/PhysicsEngine/PhysicsConstraintComponent.cpp
+++ b/Engine/Source/Runtime/Engine/Private/PhysicsEngine/PhysicsConstraintComponent.cpp
@@ -279,6 +279,46 @@ void UPhysicsConstraintComponent::PostLoad()
 		ConstraintInstance.CopyConstraintParamsFrom(&ConstraintSetup_DEPRECATED->DefaultInstance);
 		ConstraintSetup_DEPRECATED = NULL;
 	}
+
+	if (GetLinkerUE4Version() < VER_UE4_SOFT_CONSTRAINTS_USE_MASS)
+	{
+		//In previous versions the mass was placed into the spring constant. This is correct because you use different springs for different mass - however, this makes tuning hard
+		//We now multiply mass into the spring constant. To fix old data we use CalculateMass which is not perfect but close (within 0.1kg)
+		//We also use the primitive body instance directly for determining if simulated - this is potentially wrong for fixed bones in skeletal mesh, but it's much more likely right (in skeletal case we don't have access to bodies to check)
+		
+		UPrimitiveComponent * Primitive1 = GetComponentInternal(EConstraintFrame::Frame1);
+		UPrimitiveComponent * Primitive2 = GetComponentInternal(EConstraintFrame::Frame2);
+		
+		int NumDynamic = 0;
+		float TotalMass = 0.f;
+
+		if (Primitive1 && Primitive1->BodyInstance.bSimulatePhysics)
+		{
+			FName BoneName = ConstraintInstance.ConstraintBone1;
+			++NumDynamic;
+			TotalMass += Primitive1->CalculateMass(BoneName);
+		}
+
+		if (Primitive2 && Primitive2->BodyInstance.bSimulatePhysics)
+		{
+			FName BoneName = ConstraintInstance.ConstraintBone2;
+			++NumDynamic;
+			TotalMass += Primitive2->CalculateMass(BoneName);
+		}
+
+		if (NumDynamic > 0)	//we don't support cases where both constrained bodies are static or NULL, but add this anyway to avoid crash
+		{
+			float AverageMass = TotalMass / NumDynamic;
+
+			ConstraintInstance.LinearLimitStiffness /= AverageMass;
+			ConstraintInstance.SwingLimitStiffness /= AverageMass;
+			ConstraintInstance.TwistLimitStiffness /= AverageMass;
+			ConstraintInstance.LinearLimitDamping /= AverageMass;
+			ConstraintInstance.SwingLimitDamping /= AverageMass;
+			ConstraintInstance.TwistLimitDamping /= AverageMass;
+		}
+
+	}
 }
 
 #if WITH_EDITOR
diff --git a/Engine/Source/Runtime/Engine/Private/PrimitiveComponentPhysics.cpp b/Engine/Source/Runtime/Engine/Private/PrimitiveComponentPhysics.cpp
--- a/Engine/Source/Runtime/Engine/Private/PrimitiveComponentPhysics.cpp
+++ b/Engine/Source/Runtime/Engine/Private/PrimitiveComponentPhysics.cpp
@@ -340,7 +340,7 @@ float UPrimitiveComponent::GetMass() const
 	return 0.0f;
 }
 
-float UPrimitiveComponent::CalculateMass()
+float UPrimitiveComponent::CalculateMass(FName)
 {
 	if (BodyInstance.BodySetup.IsValid())
 	{
diff --git a/Engine/Source/Runtime/Engine/Private/SkeletalMeshComponent.cpp b/Engine/Source/Runtime/Engine/Private/SkeletalMeshComponent.cpp
--- a/Engine/Source/Runtime/Engine/Private/SkeletalMeshComponent.cpp
+++ b/Engine/Source/Runtime/Engine/Private/SkeletalMeshComponent.cpp
@@ -1579,7 +1579,7 @@ FTransform USkeletalMeshComponent::ConvertLocalRootMotionToWorld(const FTransfor
 }
 
 
-float USkeletalMeshComponent::CalculateMass()
+float USkeletalMeshComponent::CalculateMass(FName BoneName)
 {
 	if (bUseSingleBodyPhysics)
 	{
@@ -1593,7 +1593,8 @@ float USkeletalMeshComponent::CalculateMass()
 		{
 			for (int32 i = 0; i < Bodies.Num(); ++i)
 			{
-				if (Bodies[i]->BodySetup.IsValid())
+				//if bone name is not provided calculate entire mass - otherwise get mass for just the bone
+				if (Bodies[i]->BodySetup.IsValid() && (BoneName == NAME_None || BoneName == Bodies[i]->BodySetup->BoneName))
 				{
 					Mass += Bodies[i]->BodySetup->CalculateMass(this);
 				}
@@ -1615,7 +1616,7 @@ float USkeletalMeshComponent::CalculateMass()
 			{
 				for (int32 i = 0; i < BodySetups->Num(); ++i)
 				{
-					if ((*BodySetups)[i])
+					if ((*BodySetups)[i] && (BoneName == NAME_None || BoneName == (*BodySetups)[i]->BoneName))
 					{
 						Mass += (*BodySetups)[i]->CalculateMass(this);
 					}
