- TTP 330609 : Allow box, capsule and sphere components to be derived from outside Engine.

diff --git a/Engine/Source/Runtime/Engine/Classes/Components/BoxComponent.h b/Engine/Source/Runtime/Engine/Classes/Components/BoxComponent.h
--- a/Engine/Source/Runtime/Engine/Classes/Components/BoxComponent.h
+++ b/Engine/Source/Runtime/Engine/Classes/Components/BoxComponent.h
@@ -7,8 +7,8 @@
 /** 
  * A box generally used for simple collision. Bounds are rendered as lines in the editor.
  */
-UCLASS(HeaderGroup=Component, ClassGroup=Shapes, hidecategories=(Object,LOD,Lighting,TextureStreaming), editinlinenew, MinimalAPI, meta=(BlueprintSpawnableComponent))
-class UBoxComponent : public UShapeComponent
+UCLASS(HeaderGroup=Component, ClassGroup=Shapes, hidecategories=(Object,LOD,Lighting,TextureStreaming), editinlinenew, meta=(BlueprintSpawnableComponent))
+class ENGINE_API UBoxComponent : public UShapeComponent
 {
 	GENERATED_UCLASS_BODY()
 
@@ -24,15 +24,15 @@ public:
 	 * @param	bUpdateOverlaps: if true and this shape is registered and collides, updates touching array for owner actor.
 	 */
 	UFUNCTION(BlueprintCallable, Category="Components|Box")
-	ENGINE_API void SetBoxExtent(FVector InBoxExtent, bool bUpdateOverlaps=true);
+	void SetBoxExtent(FVector InBoxExtent, bool bUpdateOverlaps=true);
 
 	// @return the box extent, scaled by the component scale.
 	UFUNCTION(BlueprintCallable, Category="Components|Box")
-	ENGINE_API FVector GetScaledBoxExtent() const;
+	FVector GetScaledBoxExtent() const;
 
 	// @return the box extent, ignoring component scale.
 	UFUNCTION(BlueprintCallable, Category="Components|Box")
-	ENGINE_API FVector GetUnscaledBoxExtent() const;
+	FVector GetUnscaledBoxExtent() const;
 
 	// Begin UPrimitiveComponent interface.
 	virtual bool IsZeroExtent() const OVERRIDE;
diff --git a/Engine/Source/Runtime/Engine/Classes/Components/CapsuleComponent.h b/Engine/Source/Runtime/Engine/Classes/Components/CapsuleComponent.h
--- a/Engine/Source/Runtime/Engine/Classes/Components/CapsuleComponent.h
+++ b/Engine/Source/Runtime/Engine/Classes/Components/CapsuleComponent.h
@@ -6,8 +6,8 @@
 /** 
  * A capsule generally used for simple collision. Bounds are rendered as lines in the editor.
  */
-UCLASS(HeaderGroup=Component, ClassGroup=Shapes, editinlinenew, hidecategories=(Object,LOD,Lighting,TextureStreaming), MinimalAPI, meta=(BlueprintSpawnableComponent))
-class UCapsuleComponent : public UShapeComponent
+UCLASS(HeaderGroup=Component, ClassGroup=Shapes, editinlinenew, hidecategories=(Object,LOD,Lighting,TextureStreaming), meta=(BlueprintSpawnableComponent))
+class ENGINE_API UCapsuleComponent : public UShapeComponent
 {
 	GENERATED_UCLASS_BODY()
 
@@ -32,7 +32,7 @@ public:
 	 * @param	bUpdateOverlaps: if true and this shape is registered and collides, updates touching array for owner actor.
 	 */
 	UFUNCTION(BlueprintCallable, Category="Components|Capsule")
-	ENGINE_API void SetCapsuleSize(float InRadius, float InHalfHeight, bool bUpdateOverlaps=true);
+	void SetCapsuleSize(float InRadius, float InHalfHeight, bool bUpdateOverlaps=true);
 
 	/**
 	 * Set the capsule radius. This is the unscaled radius, before component scale is applied.
@@ -41,7 +41,7 @@ public:
 	 * @param	bUpdateOverlaps: if true and this shape is registered and collides, updates touching array for owner actor.
 	 */
 	UFUNCTION(BlueprintCallable, Category="Components|Capsule")
-	ENGINE_API void SetCapsuleRadius(float Radius, bool bUpdateOverlaps=true);
+	void SetCapsuleRadius(float Radius, bool bUpdateOverlaps=true);
 	
 	/**
 	 * Set the capsule half-height. This is the unscaled half-height, before component scale is applied.
@@ -50,7 +50,7 @@ public:
 	 * @param	bUpdateOverlaps: if true and this shape is registered and collides, updates touching array for owner actor.
 	 */
 	UFUNCTION(BlueprintCallable, Category="Components|Capsule")
-	ENGINE_API void SetCapsuleHalfHeight(float HalfHeight, bool bUpdateOverlaps=true);
+	void SetCapsuleHalfHeight(float HalfHeight, bool bUpdateOverlaps=true);
 
 	// Begin UObject interface
 	virtual void Serialize(FArchive& Ar) OVERRIDE;
@@ -77,33 +77,33 @@ public:
 
 	// @return the capsule radius scaled by the component scale.
 	UFUNCTION(BlueprintCallable, Category="Components|Capsule")
-	ENGINE_API float GetScaledCapsuleRadius() const;
+	float GetScaledCapsuleRadius() const;
 
 	// @return the capsule half height scaled by the component scale.
 	UFUNCTION(BlueprintCallable, Category="Components|Capsule")
-	ENGINE_API float GetScaledCapsuleHalfHeight() const;
+	float GetScaledCapsuleHalfHeight() const;
 
 	// @return the capsule radius and half height scaled by the component scale.
 	UFUNCTION(BlueprintCallable, Category="Components|Capsule")
-	ENGINE_API void GetScaledCapsuleSize(float& OutRadius, float& OutHalfHeight) const;
+	void GetScaledCapsuleSize(float& OutRadius, float& OutHalfHeight) const;
 
 
 	// @return the capsule radius, ignoring component scaling.
 	UFUNCTION(BlueprintCallable, Category="Components|Capsule")
-	ENGINE_API float GetUnscaledCapsuleRadius() const;
+	float GetUnscaledCapsuleRadius() const;
 
 	// @return the capsule half height, ignoring component scaling.
 	UFUNCTION(BlueprintCallable, Category="Components|Capsule")
-	ENGINE_API float GetUnscaledCapsuleHalfHeight() const;
+	float GetUnscaledCapsuleHalfHeight() const;
 	
 	// @return the capsule radius and half height, ignoring component scaling.
 	UFUNCTION(BlueprintCallable, Category="Components|Capsule")
-	ENGINE_API void GetUnscaledCapsuleSize(float& OutRadius, float& OutHalfHeight) const;
+	void GetUnscaledCapsuleSize(float& OutRadius, float& OutHalfHeight) const;
 
 	// Get the scale used by this shape. This is a uniform scale that is the minimum of any non-uniform scaling.
 	// @return the scale used by this shape.
 	UFUNCTION(BlueprintCallable, Category="Components|Capsule")
-	ENGINE_API float GetShapeScale() const;
+	float GetShapeScale() const;
 
 	// Sets the capsule size without triggering a render or physics update.
 	FORCEINLINE void InitCapsuleSize(float InRadius, float InHalfHeight)
diff --git a/Engine/Source/Runtime/Engine/Classes/Components/SphereComponent.h b/Engine/Source/Runtime/Engine/Classes/Components/SphereComponent.h
--- a/Engine/Source/Runtime/Engine/Classes/Components/SphereComponent.h
+++ b/Engine/Source/Runtime/Engine/Classes/Components/SphereComponent.h
@@ -7,8 +7,8 @@
 /** 
  * A sphere generally used for simple collision. Bounds are rendered as lines in the editor.
  */
-UCLASS(HeaderGroup=Component, ClassGroup=Shapes, editinlinenew, hidecategories=(Object,LOD,Lighting,TextureStreaming), MinimalAPI, meta=(BlueprintSpawnableComponent))
-class USphereComponent : public UShapeComponent
+UCLASS(HeaderGroup=Component, ClassGroup=Shapes, editinlinenew, hidecategories=(Object,LOD,Lighting,TextureStreaming), meta=(BlueprintSpawnableComponent))
+class ENGINE_API USphereComponent : public UShapeComponent
 {
 	GENERATED_UCLASS_BODY()
 
@@ -24,15 +24,15 @@ public:
 	 * @param	bUpdateOverlaps: if true and this shape is registered and collides, updates touching array for owner actor.
 	 */
 	UFUNCTION(BlueprintCallable, Category="Components|Sphere")
-	ENGINE_API void SetSphereRadius(float InSphereRadius, bool bUpdateOverlaps=true);
+	void SetSphereRadius(float InSphereRadius, bool bUpdateOverlaps=true);
 
 	// @return the radius of the sphere, with component scale applied.
 	UFUNCTION(BlueprintCallable, Category="Components|Sphere")
-	ENGINE_API float GetScaledSphereRadius() const;
+	float GetScaledSphereRadius() const;
 
 	// @return the radius of the sphere, ignoring component scale.
 	UFUNCTION(BlueprintCallable, Category="Components|Sphere")
-	ENGINE_API float GetUnscaledSphereRadius() const;
+	float GetUnscaledSphereRadius() const;
 
 	// Begin UPrimitiveComponent interface.
 	virtual FPrimitiveSceneProxy* CreateSceneProxy() OVERRIDE;
@@ -53,7 +53,7 @@ public:
 	// Get the scale used by this shape. This is a uniform scale that is the minimum of any non-uniform scaling.
 	// @return the scale used by this shape.
 	UFUNCTION(BlueprintCallable, Category="Components|Sphere")
-	ENGINE_API float GetShapeScale() const;
+	float GetShapeScale() const;
 
 	// Sets the sphere radius without triggering a render or physics update.
 	FORCEINLINE void InitSphereRadius(float InSphereRadius) { SphereRadius = InSphereRadius; }
