#UE4
- fixed Pawns instantiated in editor and seved with a level not getting NavAgent's values set properly

diff --git a/Engine/Source/Runtime/Engine/Classes/GameFramework/Pawn.h b/Engine/Source/Runtime/Engine/Classes/GameFramework/Pawn.h
--- a/Engine/Source/Runtime/Engine/Classes/GameFramework/Pawn.h
+++ b/Engine/Source/Runtime/Engine/Classes/GameFramework/Pawn.h
@@ -192,6 +192,7 @@ public:
 	virtual class UNetConnection* GetNetConnection() OVERRIDE;
 	virtual void PostInitProperties() OVERRIDE;
 	virtual void PostLoad() OVERRIDE;
+	virtual void PostRegisterAllComponents() OVERRIDE;
 	virtual float TakeDamage(float Damage, struct FDamageEvent const& DamageEvent, class AController* EventInstigator, class AActor* DamageCauser) OVERRIDE;
 	virtual void BecomeViewTarget(class APlayerController* PC) OVERRIDE;
 	virtual bool UpdateNavigationRelevancy() OVERRIDE { SetNavigationRelevancy(false); return false; }
@@ -210,6 +211,9 @@ public:
 	virtual void GetMoveGoalReachTest(class AActor* MovingActor, const FVector& MoveOffset, FVector& GoalOffset, float& GoalRadius, float& GoalHalfHeight) const OVERRIDE;
 	// End INavAgentInterface Interface
 
+	/** updates MovementComponent's parameters used by navigation system */
+	void UpdateNavAgent();
+
 	/** @return true if we are in a state to take damage (checked at the start of TakeDamage.
 	*   Subclasses may check this as well if they override TakeDamage and don't want to potentially trigger TakeDamage actions by checking if it returns zero in the super class. */
 	virtual bool ShouldTakeDamage(float Damage, FDamageEvent const& DamageEvent, AController* EventInstigator, AActor* DamageCauser) const;
diff --git a/Engine/Source/Runtime/Engine/Private/Pawn.cpp b/Engine/Source/Runtime/Engine/Private/Pawn.cpp
--- a/Engine/Source/Runtime/Engine/Private/Pawn.cpp
+++ b/Engine/Source/Runtime/Engine/Private/Pawn.cpp
@@ -89,24 +89,8 @@ void APawn::PostInitializeComponents()
 			SpawnDefaultController();
 		}
 
-		UPawnMovementComponent* MovementComponent = GetMovementComponent();
-		// Update Nav Agent props with collision component's setup if it's not set yet
-		if (RootComponent != NULL && MovementComponent != NULL && MovementComponent->ShouldUpdateNavAgentWithOwnersCollision())
-		{
-			RootComponent->UpdateBounds();
-			MovementComponent->UpdateNavAgent(this);
-		}
-
-		//UPawnMovementComponent* MovementComponent = GetMovementComponent();
-		//if (MovementComponent != NULL)
-		//{
-		//	// Update Nav Agent props with collision component's setup if it's not set yet
-		//	float BoundRadius, BoundHalfHeight;
-		//	GetSimpleCollisionCylinder(BoundRadius, BoundHalfHeight);
-		//	FNavAgentProperties* NavAgent = MovementComponent->GetNavAgentProperties();
-		//	NavAgent->AgentRadius = BoundRadius;
-		//	NavAgent->AgentHeight = BoundHalfHeight * 2.f;
-		//}
+		// update movement component's nav agent values
+		UpdateNavAgent();
 	}
 }
 
@@ -124,6 +108,24 @@ void APawn::PostLoad()
 	AutoReceiveInput = EAutoReceiveInput::Disabled;
 }
 
+void APawn::PostRegisterAllComponents()
+{
+	Super::PostRegisterAllComponents();
+
+	UpdateNavAgent();
+}
+
+void APawn::UpdateNavAgent()
+{
+	UPawnMovementComponent* MovementComponent = GetMovementComponent();
+	//// Update Nav Agent props with collision component's setup if it's not set yet
+	if (RootComponent != NULL && MovementComponent != NULL && MovementComponent->ShouldUpdateNavAgentWithOwnersCollision())
+	{
+		RootComponent->UpdateBounds();
+		MovementComponent->UpdateNavAgent(this);
+	}
+}
+
 void APawn::PawnStartFire(uint8 FireModeNum) {}
 
 class UGameplayDebuggingComponent* APawn::GetDebugComponent(bool bCreateIfNotFound)
