- TTP 331256 : Hide functions like 'Get Actor Location', 'Get Root Component' etc. from Level BP

diff --git a/Engine/Source/Runtime/Engine/Classes/Engine/LevelScriptActor.h b/Engine/Source/Runtime/Engine/Classes/Engine/LevelScriptActor.h
--- a/Engine/Source/Runtime/Engine/Classes/Engine/LevelScriptActor.h
+++ b/Engine/Source/Runtime/Engine/Classes/Engine/LevelScriptActor.h
@@ -4,7 +4,7 @@
 
 #include "LevelScriptActor.generated.h"
 
-UCLASS(HeaderGroup=LevelScript, notplaceable, meta=(KismetHideOverrides = "ReceiveAnyDamage,ReceivePointDamage,ReceiveRadialDamage,ReceiveActorBeginOverlap,ReceiveActorEndOverlap,ReceiveHit,ReceiveDestroyed,ReceiveActorBeginCursorOver,ReceiveActorEndCursorOver,ReceiveActorOnClicked,ReceiveActorOnReleased,ReceiveActorOnInputTouchBegin,ReceiveActorOnInputTouchEnd,ReceiveActorOnInputTouchEnter,ReceiveActorOnInputTouchLeave"))
+UCLASS(HeaderGroup=LevelScript, notplaceable, meta=(KismetHideOverrides = "ReceiveAnyDamage,ReceivePointDamage,ReceiveRadialDamage,ReceiveActorBeginOverlap,ReceiveActorEndOverlap,ReceiveHit,ReceiveDestroyed,ReceiveActorBeginCursorOver,ReceiveActorEndCursorOver,ReceiveActorOnClicked,ReceiveActorOnReleased,ReceiveActorOnInputTouchBegin,ReceiveActorOnInputTouchEnd,ReceiveActorOnInputTouchEnter,ReceiveActorOnInputTouchLeave"), HideCategories=(Collision,Rendering,"Utilities|Orientation"))
 class ENGINE_API ALevelScriptActor : public AActor
 {
 	GENERATED_UCLASS_BODY()
