#UE4
Client z smoothing shouldn't be effected by RootComponent scale

diff --git a/Engine/Source/Runtime/Engine/Private/CharacterMovementComponent.cpp b/Engine/Source/Runtime/Engine/Private/CharacterMovementComponent.cpp
--- a/Engine/Source/Runtime/Engine/Private/CharacterMovementComponent.cpp
+++ b/Engine/Source/Runtime/Engine/Private/CharacterMovementComponent.cpp
@@ -4420,7 +4420,7 @@ void UCharacterMovementComponent::SmoothClientPosition(float DeltaSeconds)
 
 		if (CharacterOwner->Mesh)
 		{
-			const FVector NewRelTranslation = CharacterOwner->ActorToWorld().InverseTransformVector(ClientData->MeshTranslationOffset + CharacterOwner->GetBaseTranslationOffset());
+			const FVector NewRelTranslation = CharacterOwner->ActorToWorld().InverseTransformVectorNoScale(ClientData->MeshTranslationOffset + CharacterOwner->GetBaseTranslationOffset());
 			CharacterOwner->Mesh->SetRelativeLocation(NewRelTranslation);
 		}
 	}
