#UE4 - Adding code to find an invalid case that has shown up recently

diff --git a/Engine/Source/Runtime/Engine/Private/DataReplication.cpp b/Engine/Source/Runtime/Engine/Private/DataReplication.cpp
--- a/Engine/Source/Runtime/Engine/Private/DataReplication.cpp
+++ b/Engine/Source/Runtime/Engine/Private/DataReplication.cpp
@@ -450,8 +450,6 @@ bool FObjectReplicator::ReceivedBunch( FInBunch &Bunch, const FReplicationFlags
 
 				bool bDiscardLayout = false;
 
-				check( Bunch.PacketId >= Retire.InPacketId );
-
 				if ( Bunch.PacketId >= Retire.InPacketId ) //!! problem with reliable pkts containing dynamic references, being retransmitted, and overriding newer versions. Want "OriginalPacketId" for retransmissions?
 				{
 					// Receive this new property.
@@ -459,8 +457,11 @@ bool FObjectReplicator::ReceivedBunch( FInBunch &Bunch, const FReplicationFlags
 				}
 				else
 				{
+					UE_LOG( LogNet, Fatal, TEXT( "Bunch.PacketId < Retire.InPacketId. PacketId: %i, LastPacketId: %i, Partial: %s, Object: %s" ), Bunch.PacketId, Retire.InPacketId, Bunch.bPartial ? TEXT( "YES" ) : TEXT( "NO" ), *Object->GetFullName() );
 					bDiscardLayout = true;
 				}
+
+				check( Bunch.PacketId >= Retire.InPacketId );
 				
 				if ( !RepLayout->ReceiveProperties( ObjectClass, RepState, (void*)Object, Bunch, bDiscardLayout, bOutHasUnmapped ) )
 				{
@@ -494,8 +495,6 @@ bool FObjectReplicator::ReceivedBunch( FInBunch &Bunch, const FReplicationFlags
 				// Check property ordering.
 				FPropertyRetirement & Retire = Retirement[ ReplicatedProp->RepIndex + Element ];
 
-				check( Bunch.PacketId >= Retire.InPacketId );
-
 				if ( Bunch.PacketId >= Retire.InPacketId ) //!! problem with reliable pkts containing dynamic references, being retransmitted, and overriding newer versions. Want "OriginalPacketId" for retransmissions?
 				{
 					// Receive this new property.
@@ -503,6 +502,8 @@ bool FObjectReplicator::ReceivedBunch( FInBunch &Bunch, const FReplicationFlags
 				}
 				else
 				{
+					UE_LOG( LogNet, Fatal, TEXT( "Bunch.PacketId < Retire.InPacketId. PacketId: %i, LastPacketId: %i, Partial: %s, Object: %s" ), Bunch.PacketId, Retire.InPacketId, Bunch.bPartial ? TEXT( "YES" ) : TEXT( "NO" ), *Object->GetFullName() );
+
 					// Skip this property, because it's out-of-date.
 					UE_LOG( LogNetTraffic, Log, TEXT( "Received out-of-date %s" ), *ReplicatedProp->GetName() );
 
@@ -510,6 +511,8 @@ bool FObjectReplicator::ReceivedBunch( FInBunch &Bunch, const FReplicationFlags
 					DestRecent	= NULL;
 				}
 
+				check( Bunch.PacketId >= Retire.InPacketId );
+
 				FMemMark Mark(FMemStack::Get());
 				uint8 * Data = DestObj ? ReplicatedProp->ContainerPtrToValuePtr<uint8>(DestObj, Element) : NewZeroed<uint8>(FMemStack::Get(),ReplicatedProp->ElementSize);
 				TArray<uint8>	MetaData;
