#ue4
* fix for iOS build error in Test config

diff --git a/Engine/Source/Runtime/Engine/Private/UnrealEngine.cpp b/Engine/Source/Runtime/Engine/Private/UnrealEngine.cpp
--- a/Engine/Source/Runtime/Engine/Private/UnrealEngine.cpp
+++ b/Engine/Source/Runtime/Engine/Private/UnrealEngine.cpp
@@ -3827,15 +3827,20 @@ bool UEngine::HandleMemCommand( const TCHAR* Cmd, FOutputDevice& Ar )
 		FlushRenderingCommands();
 	}
 
+#if !NO_LOGGING
+	const FName CategoryName(LogMemory.GetCategoryName());
+#else
+	const FName CategoryName(TEXT("LogMemory"));
+#endif
 	FPlatformMemory::DumpStats( Ar );
-	Ar.CategorizedLogf( LogMemory.GetCategoryName(), ELogVerbosity::Log, TEXT("") );
+	Ar.CategorizedLogf( CategoryName, ELogVerbosity::Log, TEXT("") );
 	GMalloc->DumpAllocatorStats( Ar );
 
 	if( bDetailed || bReport)
 	{
-		Ar.CategorizedLogf( LogMemory.GetCategoryName(), ELogVerbosity::Log, TEXT("Memory Stats:") );
-		Ar.CategorizedLogf( LogMemory.GetCategoryName(), ELogVerbosity::Log, TEXT("FMemStack (gamethread) allocation size [used/ unused] = [%.2f / %.2f] MB"), FMemStack::Get().GetByteCount() / (1024.0f * 1024.0f), FMemStack::Get().GetUnusedByteCount() / (1024.0f * 1024.0f)  );
-		Ar.CategorizedLogf( LogMemory.GetCategoryName(), ELogVerbosity::Log, TEXT("Nametable memory usage = %.2f MB"), FName::GetNameTableMemorySize() / (1024.0f * 1024.0f) );
+		Ar.CategorizedLogf( CategoryName, ELogVerbosity::Log, TEXT("Memory Stats:") );
+		Ar.CategorizedLogf( CategoryName, ELogVerbosity::Log, TEXT("FMemStack (gamethread) allocation size [used/ unused] = [%.2f / %.2f] MB"), FMemStack::Get().GetByteCount() / (1024.0f * 1024.0f), FMemStack::Get().GetUnusedByteCount() / (1024.0f * 1024.0f)  );
+		Ar.CategorizedLogf( CategoryName, ELogVerbosity::Log, TEXT("Nametable memory usage = %.2f MB"), FName::GetNameTableMemorySize() / (1024.0f * 1024.0f) );
 
 #if STATS
 		TArray<FStatMessage> Stats;
@@ -3852,7 +3857,7 @@ bool UEngine::HandleMemCommand( const TCHAR* Cmd, FOutputDevice& Ar )
 			FName LastGroup = Meta.NameAndInfo.GetGroupName();
 			if ((LastGroup == NAME_STATGROUP_SceneMemory || LastGroup == NAME_STATGROUP_Memory || LastGroup == NAME_STATGROUP_TextureGroup || LastGroup == NAME_STATGROUP_RHI)  && Meta.NameAndInfo.GetFlag(EStatMetaFlags::IsMemory))
 			{
-				Ar.CategorizedLogf( LogMemory.GetCategoryName(), ELogVerbosity::Log, TEXT("%s"), *FStatsUtils::DebugPrint(Meta));
+				Ar.CategorizedLogf( CategoryName, ELogVerbosity::Log, TEXT("%s"), *FStatsUtils::DebugPrint(Meta));
 			}
 		}
 #endif
