Clear AxisKey values on PopInputComponent

diff --git a/Engine/Source/Runtime/Engine/Private/PlayerController.cpp b/Engine/Source/Runtime/Engine/Private/PlayerController.cpp
--- a/Engine/Source/Runtime/Engine/Private/PlayerController.cpp
+++ b/Engine/Source/Runtime/Engine/Private/PlayerController.cpp
@@ -3700,9 +3700,13 @@ bool APlayerController::PopInputComponent(UInputComponent* InputComponent)
 	{
 		if (CurrentInputStack.RemoveSingle(InputComponent) > 0)
 		{
-			for (int32 AxisIndex=0; AxisIndex<InputComponent->AxisBindings.Num(); ++AxisIndex)
+			for (FInputAxisBinding& AxisBinding : InputComponent->AxisBindings)
 			{
-				InputComponent->AxisBindings[AxisIndex].AxisValue = 0.f;
+				AxisBinding.AxisValue = 0.f;
+			}
+			for (FInputAxisKeyBinding& AxisKeyBinding : InputComponent->AxisKeyBindings)
+			{
+				AxisKeyBinding.AxisValue = 0.f;
 			}
 
 			return true;
