[GitHub]

- Fix slope check in UCharacterMovementComponent::PhysFalling. Gravity was being added as a scalar to all vector components, was supposed to be only in Z.

Pull request: https://github.com/EpicGames/UnrealEngine/pull/31

diff --git a/Engine/Source/Runtime/Engine/Private/CharacterMovementComponent.cpp b/Engine/Source/Runtime/Engine/Private/CharacterMovementComponent.cpp
--- a/Engine/Source/Runtime/Engine/Private/CharacterMovementComponent.cpp
+++ b/Engine/Source/Runtime/Engine/Private/CharacterMovementComponent.cpp
@@ -2408,7 +2408,7 @@ void UCharacterMovementComponent::PhysFalling(float deltaTime, int32 Iterations)
 		if ( TickAirControl > 0.0f && FallAcceleration.SizeSquared() > 0.f )
 		{
 			const float TestWalkTime = FMath::Max(deltaTime, 0.05f);
-			const FVector TestWalk = ((TickAirControl * GetModifiedMaxAcceleration() * FallAcceleration.SafeNormal() + GetGravityZ()) * TestWalkTime + Velocity) * TestWalkTime;
+			const FVector TestWalk = ((TickAirControl * GetModifiedMaxAcceleration() * FallAcceleration.SafeNormal() + FVector(0.f,0.f,GetGravityZ())) * TestWalkTime + Velocity) * TestWalkTime;
 			if(!TestWalk.IsZero())
 			{
 				static const FName FallingTraceParamsTag = FName(TEXT("PhysFalling"));
