#UE4:
* renamed some property and help to be more clear (for Light Propgation Volumes)

diff --git a/Engine/Source/Runtime/Engine/Classes/Materials/Material.h b/Engine/Source/Runtime/Engine/Classes/Materials/Material.h
--- a/Engine/Source/Runtime/Engine/Classes/Materials/Material.h
+++ b/Engine/Source/Runtime/Engine/Classes/Materials/Material.h
@@ -460,7 +460,7 @@ public:
 	/**
 	 * If enabled, the material's emissive colour is injected into the LightPropagationVolume
 	 */
-	UPROPERTY(EditAnywhere, BlueprintReadWrite, Category=Material, AdvancedDisplay)
+	UPROPERTY(EditAnywhere, BlueprintReadWrite, Category=Material, meta=(DisplayName = "Emissive (Dynamic Area Light)"), AdvancedDisplay)
 	uint32 bUseEmissiveForDynamicAreaLighting : 1;
 
 	/** Whether material uses BaseColor, Metallic, Specular */
diff --git a/Engine/Source/Runtime/Engine/Public/ShowFlagsValues.inl b/Engine/Source/Runtime/Engine/Public/ShowFlagsValues.inl
--- a/Engine/Source/Runtime/Engine/Public/ShowFlagsValues.inl
+++ b/Engine/Source/Runtime/Engine/Public/ShowFlagsValues.inl
@@ -245,7 +245,7 @@ SHOWFLAG_ALWAYS_ACCESSIBLE(PrecomputedVisibility, SFG_Advanced, LOCTEXT("Precomp
 /** Contribution from sky light */
 SHOWFLAG_ALWAYS_ACCESSIBLE(SkyLighting, SFG_LightingComponents, LOCTEXT("SkyLightingSF", "Sky Lighting"))
 /** Visualize Light Propagation Volume, for developer (by default off): */
-SHOWFLAG_ALWAYS_ACCESSIBLE(VisualizeLPV, SFG_Visualize, LOCTEXT("LpvLightingOnlySF", "LPV Lighting Only"))
+SHOWFLAG_ALWAYS_ACCESSIBLE(VisualizeLPV, SFG_Visualize, LOCTEXT("VisualizeLPVSF", "Visualize Light Propagation Volume"))
 /** Visualize preview shadow indicator */
 SHOWFLAG_ALWAYS_ACCESSIBLE(PreviewShadowsIndicator, SFG_Visualize, LOCTEXT("PreviewShadowIndicatorSF", "Preview Shadows Indicator"))
 /** Visualize precomputed visibilty cells */
