Fixed single node to use same system as normal anim instance. so that tick/curve collection is consistent.

diff --git a/Engine/Source/Editor/AnimGraph/Private/AnimPreviewInstance.cpp b/Engine/Source/Editor/AnimGraph/Private/AnimPreviewInstance.cpp
--- a/Engine/Source/Editor/AnimGraph/Private/AnimPreviewInstance.cpp
+++ b/Engine/Source/Editor/AnimGraph/Private/AnimPreviewInstance.cpp
@@ -173,10 +173,8 @@ void UAnimPreviewInstance::MontagePreview_SetReverse(bool bInReverse)
 
 	if (FAnimMontageInstance * CurMontageInstance = GetActiveMontageInstance())
 	{
-		if (bReverse == (CurMontageInstance->PlayRate > 0.f))
-		{
-			CurMontageInstance->PlayRate *= -1.f;
-		}
+		// copy the current playrate
+		CurMontageInstance->PlayRate = PlayRate;
 	}
 }
 
@@ -308,7 +306,7 @@ void UAnimPreviewInstance::MontagePreview_JumpToStart()
 		{
 			MontagePreview_Restart();
 		}
-		if (bReverse)
+		if (PlayRate < 0.f)
 		{
 			Montage_JumpToSectionsEnd(Montage->GetSectionName(SectionIdx));
 		}
@@ -338,7 +336,7 @@ void UAnimPreviewInstance::MontagePreview_JumpToEnd()
 		{
 			MontagePreview_Restart();
 		}
-		if (bReverse)
+		if (PlayRate < 0.f)
 		{
 			Montage_JumpToSection(Montage->GetSectionName(MontagePreview_FindLastSection(SectionIdx)));
 		}
@@ -368,7 +366,7 @@ void UAnimPreviewInstance::MontagePreview_JumpToPreviewStart()
 		{
 			MontagePreview_Restart();
 		}
-		Montage_JumpToSection(Montage->GetSectionName(! bReverse? SectionIdx : MontagePreview_FindLastSection(SectionIdx)));
+		Montage_JumpToSection(Montage->GetSectionName(PlayRate > 0.f? SectionIdx : MontagePreview_FindLastSection(SectionIdx)));
 		if (! bWasPlaying)
 		{
 			MontagePreview_SetPlaying(false);
@@ -416,7 +414,7 @@ void UAnimPreviewInstance::MontagePreview_PreviewNormal(int32 FromSectionIdx)
 			PreviewFromSection = MontagePreviewStartSectionIdx;
 		}
 		MontagePreviewType = EMPT_Normal;
-		Montage_Play(Montage, (bReverse? -1.0f : 1.0f) * PlayRate);
+		Montage_Play(Montage, PlayRate);
 		MontagePreview_SetLoopNormal(bLooping, FromSectionIdx);
 		Montage_JumpToSection(Montage->GetSectionName(PreviewFromSection));
 		MontagePreview_RemoveBlendOut();
@@ -430,7 +428,7 @@ void UAnimPreviewInstance::MontagePreview_PreviewAllSections()
 	if (Montage && Montage->SequenceLength > 0.f)
 	{
 		MontagePreviewType = EMPT_AllSections;
-		Montage_Play(Montage, (bReverse? -1.0f : 1.0f) * PlayRate);
+		Montage_Play(Montage, PlayRate);
 		MontagePreview_SetLoopAllSections(bLooping);
 		MontagePreview_JumpToPreviewStart();
 		MontagePreview_RemoveBlendOut();
diff --git a/Engine/Source/Runtime/Engine/Private/Animation/AnimSingleNodeInstance.cpp b/Engine/Source/Runtime/Engine/Private/Animation/AnimSingleNodeInstance.cpp
--- a/Engine/Source/Runtime/Engine/Private/Animation/AnimSingleNodeInstance.cpp
+++ b/Engine/Source/Runtime/Engine/Private/Animation/AnimSingleNodeInstance.cpp
@@ -194,39 +194,27 @@ void UAnimSingleNodeInstance::NativeInitializeAnimation()
 void UAnimSingleNodeInstance::NativeUpdateAnimation(float DeltaTimeX)
 {
 	float NewDeltaTime = DeltaTimeX;
-	if (bReverse)
-	{
-		NewDeltaTime = -NewDeltaTime;
-	}
-
 	UAnimSequence* PreviewBasePose = NULL;
 
-	if (!bPlaying)
+	if(CurrentAsset != NULL && bPlaying)
 	{
-		// if not playing, make DeltaTime to be 0.f;
-		// we still have to tick when NewDeltaTime == 0.f for retrieving curve datas
-		// -i.e. Morph Target Curves
-		// if we need an optimization, we can just grab curve separate, but for now this is better for consistency
-		NewDeltaTime = 0.f;
-	}
+		FAnimGroupInstance* SyncGroup;
+		FAnimTickRecord& TickRecord = CreateUninitializedTickRecord(INDEX_NONE, /*out*/ SyncGroup);
 
-	if(CurrentAsset != NULL)
-	{
 		if (UBlendSpaceBase* BlendSpace = Cast<UBlendSpaceBase>(CurrentAsset))
 		{
-			BlendSpaceAdvanceImmediate(BlendSpace, BlendSpaceInput, BlendSampleData, BlendFilter, bLooping, PlayRate, NewDeltaTime, CurrentTime);
+			MakeBlendSpaceTickRecord(TickRecord, BlendSpace, BlendSpaceInput, BlendSampleData, BlendFilter, bLooping, PlayRate, 1.f, /*inout*/ CurrentTime);
 #if WITH_EDITORONLY_DATA
 			PreviewBasePose = BlendSpace->PreviewBasePose;
 #endif
 		}
 		else if (UAnimSequence* Sequence = Cast<UAnimSequence>(CurrentAsset))
 		{
-			SequenceAdvanceImmediate(Sequence, bLooping, PlayRate, NewDeltaTime, CurrentTime);
-
+			MakeSequenceTickRecord(TickRecord, Sequence, bLooping, PlayRate, 1.f, /*inout*/ CurrentTime);
 			// if it's not looping, just set play to be false when reached to end
 			if (!bLooping)
 			{
-				if ((bReverse && CurrentTime <= 0.f) || (!bReverse && CurrentTime >= Sequence->SequenceLength))
+				if ((PlayRate < 0.f && CurrentTime <= 0.f) || (PlayRate > 0.f && CurrentTime >= Sequence->SequenceLength))
 				{
 					SetPlaying(false);
 				}
@@ -234,12 +222,11 @@ void UAnimSingleNodeInstance::NativeUpdateAnimation(float DeltaTimeX)
 		}
 		else if(UAnimComposite* Composite = Cast<UAnimComposite>(CurrentAsset))
 		{
-			SequenceAdvanceImmediate(Composite, bLooping, PlayRate, NewDeltaTime, CurrentTime);
-
+			MakeSequenceTickRecord(TickRecord, Composite, bLooping, PlayRate, 1.f, /*inout*/ CurrentTime);
 			// if it's not looping, just set play to be false when reached to end
 			if (!bLooping)
 			{
-				if ((bReverse && CurrentTime <= 0.f) || (!bReverse && CurrentTime >= Composite->SequenceLength))
+				if ((PlayRate < 0.f && CurrentTime <= 0.f) || (PlayRate > 0.f && CurrentTime >= Composite->SequenceLength))
 				{
 					SetPlaying(false);
 				}
@@ -467,6 +454,14 @@ void UAnimSingleNodeInstance::SetPlayRate(float InPlayRate)
 void UAnimSingleNodeInstance::SetReverse(bool bInReverse)
 {
 	bReverse = bInReverse;
+	if (bInReverse)
+	{
+		PlayRate = -FMath::Abs(PlayRate);
+	}
+	else
+	{
+		PlayRate = FMath::Abs(PlayRate);
+	}
 
 // reverse support is a bit tricky for montage
 // since we don't have delegate when it reached to the beginning
