Reflection capture data is now stored in the map package
* This allows cooking for ES2 to always have valid reflections for maps that have been resaved
* Old data is skipped so backwards compatibility is not necessary
* Also fixed invisible probes not having valid capture data, but bVisible can be changed at runtime

diff --git a/Engine/Source/Runtime/Core/Public/UObject/ObjectVersion.h b/Engine/Source/Runtime/Core/Public/UObject/ObjectVersion.h
--- a/Engine/Source/Runtime/Core/Public/UObject/ObjectVersion.h
+++ b/Engine/Source/Runtime/Core/Public/UObject/ObjectVersion.h
@@ -503,16 +503,15 @@ enum EUnrealEngineObjectUE4Version
 	VER_UE4_FOLIAGE_COLLISION,
 	// Added sky bent normal to indirect lighting cache
 	VER_UE4_SKY_BENT_NORMAL,
-
 	// Added cooking for landscape collision data
 	VER_UE4_LANDSCAPE_COLLISION_DATA_COOKING,
-	
 	// Convert CPU tangent Z delta to vector from PackedNormal since we don't get any benefit other than memory
 	// we still convert all to FVector in CPU time whenever any calculation
 	VER_UE4_MORPHTARGET_CPU_TANGENTZDELTA_FORMATCHANGE,
-
 	// Soft constraint limits will implicitly use the mass of the bodies
 	VER_UE4_SOFT_CONSTRAINTS_USE_MASS,
+	// Reflection capture data saved in packages
+	VER_UE4_REFLECTION_DATA_IN_PACKAGES,
 
 	// -----<new versions can be added before this line>-------------------------------------------------
 	// - this needs to be the last line (see note below)
diff --git a/Engine/Source/Runtime/Engine/Classes/Components/ReflectionCaptureComponent.h b/Engine/Source/Runtime/Engine/Classes/Components/ReflectionCaptureComponent.h
--- a/Engine/Source/Runtime/Engine/Classes/Components/ReflectionCaptureComponent.h
+++ b/Engine/Source/Runtime/Engine/Classes/Components/ReflectionCaptureComponent.h
@@ -94,7 +94,6 @@ class UReflectionCaptureComponent : public USceneComponent
 
 	/** Marks this component has having been recaptured. */
 	void SetCaptureCompleted() { bCaptureDirty = false; }
-	bool IsCaptureDirty() const { return bCaptureDirty; }
 
 	/** Gets the radius that bounds the shape's influence, used for culling. */
 	virtual float GetInfluenceBoundingRadius() const PURE_VIRTUAL(UReflectionCaptureComponent::GetInfluenceBoundingRadius,return 0;);
@@ -179,10 +178,8 @@ private:
 	 */
 	static TArray<UReflectionCaptureComponent*> ReflectionCapturesToUpdateForLoad;
 
-	/** Newly created captures don't get PostLoad called so we use this queue. */
-	static TArray<UReflectionCaptureComponent*> ReflectionCapturesToUpdateNewlyCreated;
-
-	void InvalidateDerivedData();
+	void UpdateDerivedData(FReflectionCaptureFullHDRDerivedData* NewDerivedData);
+	void SerializeSourceData(FArchive& Ar);
 
 	friend class FReflectionCaptureProxy;
 };
diff --git a/Engine/Source/Runtime/Engine/Private/ReflectionCaptureComponent.cpp b/Engine/Source/Runtime/Engine/Private/ReflectionCaptureComponent.cpp
--- a/Engine/Source/Runtime/Engine/Private/ReflectionCaptureComponent.cpp
+++ b/Engine/Source/Runtime/Engine/Private/ReflectionCaptureComponent.cpp
@@ -194,11 +194,13 @@ APlaneReflectionCapture::APlaneReflectionCapture(const class FPostConstructIniti
 }
 
 // Generate a new guid to force a recache of all reflection derived data
-#define REFLECTIONCAPTURE_FULL_DERIVEDDATA_VER			TEXT("82f35deee55a5ec9956d4897f3d46e5a")
+// Note: changing this will cause saved capture data in maps to be discarded
+// A resave of those maps will be required to guarantee valid reflections when cooking for ES2
+FGuid ReflectionCaptureDDCVer(0x299db8a0, 0x54254535, 0xa088cc40, 0x1ea245eb);
 
 FString FReflectionCaptureFullHDRDerivedData::GetDDCKeyString(const FGuid& StateId)
 {
-	return FDerivedDataCacheInterface::BuildCacheKey(TEXT("REFL_FULL"), REFLECTIONCAPTURE_FULL_DERIVEDDATA_VER, *StateId.ToString());
+	return FDerivedDataCacheInterface::BuildCacheKey(TEXT("REFL_FULL"), *ReflectionCaptureDDCVer.ToString(), *StateId.ToString());
 }
 
 FReflectionCaptureFullHDRDerivedData::~FReflectionCaptureFullHDRDerivedData()
@@ -703,7 +705,6 @@ private:
 
 TArray<UReflectionCaptureComponent*> UReflectionCaptureComponent::ReflectionCapturesToUpdate;
 TArray<UReflectionCaptureComponent*> UReflectionCaptureComponent::ReflectionCapturesToUpdateForLoad;
-TArray<UReflectionCaptureComponent*> UReflectionCaptureComponent::ReflectionCapturesToUpdateNewlyCreated;
 
 UReflectionCaptureComponent::UReflectionCaptureComponent(const class FPostConstructInitializeProperties& PCIP)
 	: Super(PCIP)
@@ -761,7 +762,62 @@ void UReflectionCaptureComponent::PostInitProperties()
 
 	if (!HasAnyFlags(RF_ClassDefaultObject))
 	{
-		ReflectionCapturesToUpdateNewlyCreated.AddUnique(this);
+		ReflectionCapturesToUpdateForLoad.AddUnique(this);
+		bCaptureDirty = true; 
+	}
+}
+
+void UReflectionCaptureComponent::SerializeSourceData(FArchive& Ar)
+{
+	if (Ar.UE4Ver() >= VER_UE4_REFLECTION_DATA_IN_PACKAGES)
+	{
+		if (Ar.IsSaving())
+		{
+			Ar << ReflectionCaptureDDCVer;
+
+			int32 StartOffset = Ar.Tell();
+			Ar << StartOffset;
+
+			bool bValid = FullHDRDerivedData != NULL;
+
+			Ar << bValid;
+
+			if (bValid)
+			{
+				Ar << FullHDRDerivedData->CompressedCapturedData;
+			}
+
+			int32 EndOffset = Ar.Tell();
+			Ar.Seek(StartOffset);
+			Ar << EndOffset;
+			Ar.Seek(EndOffset);
+		}
+		else if (Ar.IsLoading())
+		{
+			FGuid SavedVersion;
+			Ar << SavedVersion;
+
+			int32 EndOffset = 0;
+			Ar << EndOffset;
+
+			if (SavedVersion != ReflectionCaptureDDCVer)
+			{
+				// Guid version of saved source data doesn't match latest, skip the data
+				// The skipping is done so we don't have to maintain legacy serialization code paths when changing the format
+				Ar.Seek(EndOffset);
+			}
+			else
+			{
+				bool bValid = false;
+				Ar << bValid;
+
+				if (bValid)
+				{
+					FullHDRDerivedData = new FReflectionCaptureFullHDRDerivedData();
+					Ar << FullHDRDerivedData->CompressedCapturedData;
+				}
+			}
+		}
 	}
 }
 
@@ -884,6 +940,10 @@ void UReflectionCaptureComponent::Serialize(FArchive& Ar)
 			}
 		}
 	}
+	else
+	{
+		SerializeSourceData(Ar);
+	}
 }
 
 void UReflectionCaptureComponent::PostLoad()
@@ -953,16 +1013,6 @@ void UReflectionCaptureComponent::PostLoad()
 		delete FullHDRDerivedData;
 		FullHDRDerivedData = NULL;
 	}
-
-	// PostLoad was called, so we can make a proper decision based on visibility of whether or not to update
-	ReflectionCapturesToUpdateNewlyCreated.Remove(this);
-
-	// Add ourselves to the global list of reflection captures that need to be uploaded to the scene or recaptured
-	if (bVisible)
-	{
-		ReflectionCapturesToUpdateForLoad.AddUnique(this);
-		bCaptureDirty = true; 
-	}
 }
 
 void UReflectionCaptureComponent::PreSave() 
@@ -988,7 +1038,7 @@ void UReflectionCaptureComponent::PostDuplicate(bool bDuplicateForPIE)
 	}
 }
 
-void UReflectionCaptureComponent::InvalidateDerivedData()
+void UReflectionCaptureComponent::UpdateDerivedData(FReflectionCaptureFullHDRDerivedData* NewDerivedData)
 {
 	if (FullHDRDerivedData)
 	{
@@ -999,9 +1049,9 @@ void UReflectionCaptureComponent::InvalidateDerivedData()
 		{
 			delete FullHDRDerivedData;
 		});
-
-		FullHDRDerivedData = NULL;
 	}
+
+	FullHDRDerivedData = NewDerivedData;
 }
 
 FReflectionCaptureProxy* UReflectionCaptureComponent::CreateSceneProxy()
@@ -1037,7 +1087,6 @@ void UReflectionCaptureComponent::BeginDestroy()
 	{
 		ReflectionCapturesToUpdate.Remove(this);
 		ReflectionCapturesToUpdateForLoad.Remove(this);
-		ReflectionCapturesToUpdateNewlyCreated.Remove(this);
 	}
 
 	// Have to do this because we can't use GetWorld in BeginDestroy
@@ -1071,7 +1120,7 @@ bool UReflectionCaptureComponent::IsReadyForFinishDestroy()
 
 void UReflectionCaptureComponent::FinishDestroy()
 {
-	InvalidateDerivedData();
+	UpdateDerivedData(NULL);
 
 	if (SM4FullHDRCubemapTexture)
 	{
@@ -1092,7 +1141,7 @@ void UReflectionCaptureComponent::SetCaptureIsDirty()
 { 
 	if (bVisible)
 	{
-		InvalidateDerivedData();
+		UpdateDerivedData(NULL);
 		FPlatformMisc::CreateGuid(StateId);
 		bDerivedDataDirty = true;
 		ReflectionCapturesToUpdate.AddUnique(this);
@@ -1164,23 +1213,23 @@ void UReflectionCaptureComponent::ReadbackFromGPUAndSaveDerivedData(UWorld* Worl
 {
 	if (bDerivedDataDirty && !IsRunningCommandlet())
 	{
-		// Read back full HDR capture data and save it in the DDC
-		//@todo - not Updating DerivedData and setting bDerivedDataDirty to false yet because that would require syncing with the rendering thread
-		// This behavior means that reflection capture derived data will stay dirty (uncached for rendering) until the editor is restarted
-		FReflectionCaptureFullHDRDerivedData TemporaryDerivedData;
+		FReflectionCaptureFullHDRDerivedData* NewDerivedData = new FReflectionCaptureFullHDRDerivedData();
 
 		if (GRHIFeatureLevel == ERHIFeatureLevel::SM4)
 		{
-			ReadbackFromSM4Cubemap(SM4FullHDRCubemapTexture, &TemporaryDerivedData);
+			ReadbackFromSM4Cubemap(SM4FullHDRCubemapTexture, NewDerivedData);
 		}
 		else
 		{
-			WorldToUpdate->Scene->GetReflectionCaptureData(this, TemporaryDerivedData);
+			WorldToUpdate->Scene->GetReflectionCaptureData(this, *NewDerivedData);
 		}
 
-		if (TemporaryDerivedData.CompressedCapturedData.Num() > 0)
+		if (NewDerivedData->CompressedCapturedData.Num() > 0)
 		{
-			GetDerivedDataCacheRef().Put(*FReflectionCaptureFullHDRDerivedData::GetDDCKeyString(StateId), TemporaryDerivedData.CompressedCapturedData);
+			GetDerivedDataCacheRef().Put(*FReflectionCaptureFullHDRDerivedData::GetDDCKeyString(StateId), NewDerivedData->CompressedCapturedData);
+
+			// Update our copy in memory
+			UpdateDerivedData(NewDerivedData);
 		}
 	}
 }
@@ -1220,18 +1269,6 @@ void UReflectionCaptureComponent::UpdateReflectionCaptureContents(UWorld* WorldT
 			}
 		}
 
-		for (int32 CaptureIndex = ReflectionCapturesToUpdateNewlyCreated.Num() - 1; CaptureIndex >= 0; CaptureIndex--)
-		{
-			UReflectionCaptureComponent* CaptureComponent = ReflectionCapturesToUpdateNewlyCreated[CaptureIndex];
-
-			if (!CaptureComponent->GetOwner() || WorldToUpdate->ContainsActor(CaptureComponent->GetOwner()))
-			{
-				WorldCombinedCaptures.Add(CaptureComponent);
-				WorldCapturesToUpdateForLoad.Add(CaptureComponent);
-				ReflectionCapturesToUpdateNewlyCreated.RemoveAt(CaptureIndex);
-			}
-		}
-
 		if (GRHIFeatureLevel == ERHIFeatureLevel::SM4)
 		{
 			for (int32 CaptureIndex = 0; CaptureIndex < WorldCombinedCaptures.Num(); CaptureIndex++)
