Adjusting comment on SpawnActorDeferred to be explicit about what function to call to finish spawning.

diff --git a/Engine/Source/Runtime/Engine/Classes/Engine/World.h b/Engine/Source/Runtime/Engine/Classes/Engine/World.h
--- a/Engine/Source/Runtime/Engine/Classes/Engine/World.h
+++ b/Engine/Source/Runtime/Engine/Classes/Engine/World.h
@@ -2030,7 +2030,7 @@ public:
 	/**
 	* Spawns given class and returns class T pointer, forcibly sets world position. WILL NOT run Construction Script of Blueprints 
 	* to give caller an opportunity to set parameters beforehand.  Caller is responsible for invoking construction
-	* manually (see AActor::OnConstruction).
+	* manually by calling AActor::FinishSpawningActor (see AActor::OnConstruction).
 	*/
 	template< class T >
 	T* SpawnActorDeferred(
