Adding DisplayName metadata to all of the enums used by Cascade to greatly improve the readability of the property values in the editor.

diff --git a/Engine/Source/Runtime/Engine/Classes/Components/SceneComponent.h b/Engine/Source/Runtime/Engine/Classes/Components/SceneComponent.h
--- a/Engine/Source/Runtime/Engine/Classes/Components/SceneComponent.h
+++ b/Engine/Source/Runtime/Engine/Classes/Components/SceneComponent.h
@@ -29,9 +29,9 @@ struct FOverlapInfo
 UENUM()
 enum EDetailMode
 {
-	DM_Low,
-	DM_Medium,
-	DM_High,
+	DM_Low UMETA(DisplayName="Low"),
+	DM_Medium UMETA(DisplayName="Medium"),
+	DM_High UMETA(DisplayName="High"),
 	DM_MAX,
 };
 
diff --git a/Engine/Source/Runtime/Engine/Classes/Distributions/DistributionVector.h b/Engine/Source/Runtime/Engine/Classes/Distributions/DistributionVector.h
--- a/Engine/Source/Runtime/Engine/Classes/Distributions/DistributionVector.h
+++ b/Engine/Source/Runtime/Engine/Classes/Distributions/DistributionVector.h
@@ -7,20 +7,20 @@
 UENUM()
 enum EDistributionVectorLockFlags
 {
-	EDVLF_None,
-	EDVLF_XY,
-	EDVLF_XZ,
-	EDVLF_YZ,
-	EDVLF_XYZ,
+	EDVLF_None UMETA(DisplayName="None"),
+	EDVLF_XY UMETA(DisplayName="XY"),
+	EDVLF_XZ UMETA(DisplayName="XZ"),
+	EDVLF_YZ UMETA(DisplayName="YZ"),
+	EDVLF_XYZ UMETA(DisplayName="XYZ"),
 	EDVLF_MAX,
 };
 
 UENUM()
 enum EDistributionVectorMirrorFlags
 {
-	EDVMF_Same,
-	EDVMF_Different,
-	EDVMF_Mirror,
+	EDVMF_Same UMETA(DisplayName="Same"),
+	EDVMF_Different UMETA(DisplayName="Different"),
+	EDVMF_Mirror UMETA(DisplayName="Mirror"),
 	EDVMF_MAX,
 };
 
diff --git a/Engine/Source/Runtime/Engine/Classes/Particles/Attractor/ParticleModuleAttractorParticle.h b/Engine/Source/Runtime/Engine/Classes/Particles/Attractor/ParticleModuleAttractorParticle.h
--- a/Engine/Source/Runtime/Engine/Classes/Particles/Attractor/ParticleModuleAttractorParticle.h
+++ b/Engine/Source/Runtime/Engine/Classes/Particles/Attractor/ParticleModuleAttractorParticle.h
@@ -7,8 +7,8 @@
 UENUM()
 enum EAttractorParticleSelectionMethod
 {
-	EAPSM_Random,
-	EAPSM_Sequential,
+	EAPSM_Random UMETA(DisplayName="Random"),
+	EAPSM_Sequential UMETA(DisplayName="Sequential"),
 	EAPSM_MAX,
 };
 
diff --git a/Engine/Source/Runtime/Engine/Classes/Particles/Beam/ParticleModuleBeamBase.h b/Engine/Source/Runtime/Engine/Classes/Particles/Beam/ParticleModuleBeamBase.h
--- a/Engine/Source/Runtime/Engine/Classes/Particles/Beam/ParticleModuleBeamBase.h
+++ b/Engine/Source/Runtime/Engine/Classes/Particles/Beam/ParticleModuleBeamBase.h
@@ -17,22 +17,22 @@ enum Beam2SourceTargetMethod
 	/** Default	- use the distribution. 
 	 *	This is the fallback for when other modes can't be resolved.
 	 */
-	PEB2STM_Default,
+	PEB2STM_Default UMETA(DisplayName="Default"),
 	/** UserSet	- use the user set value. 
 	 *	Primarily intended for weapon effects.
 	 */
-	PEB2STM_UserSet,
+	PEB2STM_UserSet UMETA(DisplayName="User Set"),
 	/** Emitter	- use the emitter position as the source/target.
 	 */
-	PEB2STM_Emitter,
+	PEB2STM_Emitter UMETA(DisplayName="Emitter"),
 	/** Particle	- use the particles from a given emitter in the system.		
 	 *	The name of the emitter should be set in <Source/Target>Name.
 	 */
-	PEB2STM_Particle,
+	PEB2STM_Particle UMETA(DisplayName="Particle"),
 	/** Actor		- use the actor as the source/target.
 	 *	The name of the actor should be set in <Source/Target>Name.
 	 */
-	PEB2STM_Actor,
+	PEB2STM_Actor UMETA(DisplayName="Actor"),
 	PEB2STM_MAX,
 };
 
@@ -41,13 +41,13 @@ UENUM()
 enum Beam2SourceTargetTangentMethod
 {
 	/** Direct - a direct line between source and target.				 */
-	PEB2STTM_Direct,
+	PEB2STTM_Direct UMETA(DisplayName="Direct"),
 	/** UserSet	- use the user set value.								 */
-	PEB2STTM_UserSet,
+	PEB2STTM_UserSet UMETA(DisplayName="User Set"),
 	/** Distribution - use the distribution.							 */
-	PEB2STTM_Distribution,
+	PEB2STTM_Distribution UMETA(DisplayName="Distribution"),
 	/** Emitter	- use the emitter direction.							 */
-	PEB2STTM_Emitter,
+	PEB2STTM_Emitter UMETA(DisplayName="Emitter"),
 	PEB2STTM_MAX,
 };
 
diff --git a/Engine/Source/Runtime/Engine/Classes/Particles/Beam/ParticleModuleBeamModifier.h b/Engine/Source/Runtime/Engine/Classes/Particles/Beam/ParticleModuleBeamModifier.h
--- a/Engine/Source/Runtime/Engine/Classes/Particles/Beam/ParticleModuleBeamModifier.h
+++ b/Engine/Source/Runtime/Engine/Classes/Particles/Beam/ParticleModuleBeamModifier.h
@@ -15,9 +15,9 @@ UENUM()
 enum BeamModifierType
 {
 	/** Modify the source of the beam.				*/
-	PEB2MT_Source,
+	PEB2MT_Source UMETA(DisplayName="Source"),
 	/** Modify the target of the beam.				*/
-	PEB2MT_Target,
+	PEB2MT_Target UMETA(DisplayName="Target"),
 	PEB2MT_MAX,
 };
 
diff --git a/Engine/Source/Runtime/Engine/Classes/Particles/Camera/ParticleModuleCameraOffset.h b/Engine/Source/Runtime/Engine/Classes/Particles/Camera/ParticleModuleCameraOffset.h
--- a/Engine/Source/Runtime/Engine/Classes/Particles/Camera/ParticleModuleCameraOffset.h
+++ b/Engine/Source/Runtime/Engine/Classes/Particles/Camera/ParticleModuleCameraOffset.h
@@ -10,9 +10,9 @@
 UENUM()
 enum EParticleCameraOffsetUpdateMethod
 {
-	EPCOUM_DirectSet,
-	EPCOUM_Additive,
-	EPCOUM_Scalar,
+	EPCOUM_DirectSet UMETA(DisplayName="Direct Set"),
+	EPCOUM_Additive UMETA(DisplayName="Additive"),
+	EPCOUM_Scalar UMETA(DisplayName="Scalar"),
 	EPCOUM_MAX,
 };
 
@@ -58,5 +58,3 @@ class UParticleModuleCameraOffset : public UParticleModuleCameraBase
 	//End UParticleModule Interface
 };
 
-
-
diff --git a/Engine/Source/Runtime/Engine/Classes/Particles/Collision/ParticleModuleCollisionBase.h b/Engine/Source/Runtime/Engine/Classes/Particles/Collision/ParticleModuleCollisionBase.h
--- a/Engine/Source/Runtime/Engine/Classes/Particles/Collision/ParticleModuleCollisionBase.h
+++ b/Engine/Source/Runtime/Engine/Classes/Particles/Collision/ParticleModuleCollisionBase.h
@@ -11,17 +11,17 @@ UENUM()
 enum EParticleCollisionComplete
 {
 	/**	Kill the particle when MaxCollisions is reached		*/
-	EPCC_Kill,
+	EPCC_Kill UMETA(DisplayName="Kill"),
 	/**	Freeze the particle in place						*/
-	EPCC_Freeze,
+	EPCC_Freeze UMETA(DisplayName="Freeze"),
 	/**	Stop collision checks, but keep updating			*/
-	EPCC_HaltCollisions,
+	EPCC_HaltCollisions UMETA(DisplayName="Halt Collisions"),
 	/**	Stop translations of the particle					*/
-	EPCC_FreezeTranslation,
+	EPCC_FreezeTranslation UMETA(DisplayName="Freeze Translation"),
 	/**	Stop rotations of the particle						*/
-	EPCC_FreezeRotation,
+	EPCC_FreezeRotation UMETA(DisplayName="Freeze Rotation"),
 	/**	Stop all movement of the particle					*/
-	EPCC_FreezeMovement,
+	EPCC_FreezeMovement UMETA(DisplayName="Freeze Movement"),
 	EPCC_MAX,
 };
 
diff --git a/Engine/Source/Runtime/Engine/Classes/Particles/Location/ParticleModuleLocationPrimitiveCylinder.h b/Engine/Source/Runtime/Engine/Classes/Particles/Location/ParticleModuleLocationPrimitiveCylinder.h
--- a/Engine/Source/Runtime/Engine/Classes/Particles/Location/ParticleModuleLocationPrimitiveCylinder.h
+++ b/Engine/Source/Runtime/Engine/Classes/Particles/Location/ParticleModuleLocationPrimitiveCylinder.h
@@ -11,9 +11,9 @@
 UENUM()
 enum CylinderHeightAxis
 {
-	PMLPC_HEIGHTAXIS_X,
-	PMLPC_HEIGHTAXIS_Y,
-	PMLPC_HEIGHTAXIS_Z,
+	PMLPC_HEIGHTAXIS_X UMETA(DisplayName="X"),
+	PMLPC_HEIGHTAXIS_Y UMETA(DisplayName="Y"),
+	PMLPC_HEIGHTAXIS_Z UMETA(DisplayName="Z"),
 	PMLPC_HEIGHTAXIS_MAX,
 };
 
diff --git a/Engine/Source/Runtime/Engine/Classes/Particles/Location/ParticleModuleLocationSkelVertSurface.h b/Engine/Source/Runtime/Engine/Classes/Particles/Location/ParticleModuleLocationSkelVertSurface.h
--- a/Engine/Source/Runtime/Engine/Classes/Particles/Location/ParticleModuleLocationSkelVertSurface.h
+++ b/Engine/Source/Runtime/Engine/Classes/Particles/Location/ParticleModuleLocationSkelVertSurface.h
@@ -7,8 +7,8 @@
 UENUM()
 enum ELocationSkelVertSurfaceSource
 {
-	VERTSURFACESOURCE_Vert,
-	VERTSURFACESOURCE_Surface,
+	VERTSURFACESOURCE_Vert UMETA(DisplayName="Verticies"),
+	VERTSURFACESOURCE_Surface UMETA(DisplayName="Surfaces"),
 	VERTSURFACESOURCE_MAX,
 };
 
diff --git a/Engine/Source/Runtime/Engine/Classes/Particles/Orbit/ParticleModuleOrbit.h b/Engine/Source/Runtime/Engine/Classes/Particles/Orbit/ParticleModuleOrbit.h
--- a/Engine/Source/Runtime/Engine/Classes/Particles/Orbit/ParticleModuleOrbit.h
+++ b/Engine/Source/Runtime/Engine/Classes/Particles/Orbit/ParticleModuleOrbit.h
@@ -8,11 +8,11 @@ UENUM()
 enum EOrbitChainMode
 {
 	/** Add the module values to the previous results						*/
-	EOChainMode_Add,
+	EOChainMode_Add UMETA(DisplayName="Add"),
 	/**	Multiply the module values by the previous results					*/
-	EOChainMode_Scale,
+	EOChainMode_Scale UMETA(DisplayName="Scale"),
 	/**	'Break' the chain and apply the values from the	previous results	*/
-	EOChainMode_Link,
+	EOChainMode_Link UMETA(DisplayName="Link"),
 	EOChainMode_MAX,
 };
 
diff --git a/Engine/Source/Runtime/Engine/Classes/Particles/Orientation/ParticleModuleOrientationAxisLock.h b/Engine/Source/Runtime/Engine/Classes/Particles/Orientation/ParticleModuleOrientationAxisLock.h
--- a/Engine/Source/Runtime/Engine/Classes/Particles/Orientation/ParticleModuleOrientationAxisLock.h
+++ b/Engine/Source/Runtime/Engine/Classes/Particles/Orientation/ParticleModuleOrientationAxisLock.h
@@ -9,25 +9,25 @@ UENUM()
 enum EParticleAxisLock
 {
 	/** No locking to an axis...							*/
-	EPAL_NONE,
+	EPAL_NONE UMETA(DisplayName="None"),
 	/** Lock the sprite facing towards the positive X-axis	*/
-	EPAL_X,
+	EPAL_X UMETA(DisplayName="X"),
 	/** Lock the sprite facing towards the positive Y-axis	*/
-	EPAL_Y,
+	EPAL_Y UMETA(DisplayName="Y"),
 	/** Lock the sprite facing towards the positive Z-axis	*/
-	EPAL_Z,
+	EPAL_Z UMETA(DisplayName="Z"),
 	/** Lock the sprite facing towards the negative X-axis	*/
-	EPAL_NEGATIVE_X,
+	EPAL_NEGATIVE_X UMETA(DisplayName="-X"),
 	/** Lock the sprite facing towards the negative Y-axis	*/
-	EPAL_NEGATIVE_Y,
+	EPAL_NEGATIVE_Y UMETA(DisplayName="-Y"),
 	/** Lock the sprite facing towards the negative Z-axis	*/
-	EPAL_NEGATIVE_Z,
+	EPAL_NEGATIVE_Z UMETA(DisplayName="-Z"),
 	/** Lock the sprite rotation on the X-axis				*/
-	EPAL_ROTATE_X,
+	EPAL_ROTATE_X UMETA(DisplayName="Rotate X"),
 	/** Lock the sprite rotation on the Y-axis				*/
-	EPAL_ROTATE_Y,
+	EPAL_ROTATE_Y UMETA(DisplayName="Rotate Y"),
 	/** Lock the sprite rotation on the Z-axis				*/
-	EPAL_ROTATE_Z,
+	EPAL_ROTATE_Z UMETA(DisplayName="Rotate Z"),
 	EPAL_MAX,
 };
 
diff --git a/Engine/Source/Runtime/Engine/Classes/Particles/ParticleEmitter.h b/Engine/Source/Runtime/Engine/Classes/Particles/ParticleEmitter.h
--- a/Engine/Source/Runtime/Engine/Classes/Particles/ParticleEmitter.h
+++ b/Engine/Source/Runtime/Engine/Classes/Particles/ParticleEmitter.h
@@ -14,8 +14,8 @@
 UENUM()
 enum EParticleBurstMethod
 {
-	EPBM_Instant,
-	EPBM_Interpolated,
+	EPBM_Instant UMETA(DisplayName="Instant"),
+	EPBM_Interpolated UMETA(DisplayName="Interpolated"),
 	EPBM_MAX,
 };
 
@@ -25,11 +25,11 @@ enum EParticleBurstMethod
 UENUM()
 enum EParticleSubUVInterpMethod
 {
-	PSUVIM_None,
-	PSUVIM_Linear,
-	PSUVIM_Linear_Blend,
-	PSUVIM_Random,
-	PSUVIM_Random_Blend,
+	PSUVIM_None UMETA(DisplayName="None"),
+	PSUVIM_Linear UMETA(DisplayName="Linear"),
+	PSUVIM_Linear_Blend UMETA(DisplayName="Linear Blend"),
+	PSUVIM_Random UMETA(DisplayName="Random"),
+	PSUVIM_Random_Blend UMETA(DisplayName="Random Blend"),
 	PSUVIM_MAX,
 };
 
@@ -39,10 +39,10 @@ enum EParticleSubUVInterpMethod
 UENUM()
 enum EEmitterRenderMode
 {
-	ERM_Normal,
-	ERM_Point,
-	ERM_Cross,
-	ERM_None,
+	ERM_Normal UMETA(DisplayName="Normal"),
+	ERM_Point UMETA(DisplayName="Point"),
+	ERM_Cross UMETA(DisplayName="Cross"),
+	ERM_None UMETA(DisplayName="None"),
 	ERM_MAX,
 };
 
diff --git a/Engine/Source/Runtime/Engine/Classes/Particles/ParticleModule.h b/Engine/Source/Runtime/Engine/Classes/Particles/ParticleModule.h
--- a/Engine/Source/Runtime/Engine/Classes/Particles/ParticleModule.h
+++ b/Engine/Source/Runtime/Engine/Classes/Particles/ParticleModule.h
@@ -14,21 +14,21 @@ UENUM()
 enum EModuleType
 {
 	/** General - all emitter types can use it			*/
-	EPMT_General,
+	EPMT_General UMETA(DisplayName="General"),
 	/** TypeData - TypeData modules						*/
-	EPMT_TypeData,
+	EPMT_TypeData UMETA(DisplayName="Type Data"),
 	/** Beam - only applied to beam emitters			*/
-	EPMT_Beam,
+	EPMT_Beam UMETA(DisplayName="Beam"),
 	/** Trail - only applied to trail emitters			*/
-	EPMT_Trail,
+	EPMT_Trail UMETA(DisplayName="Trail"),
 	/** Spawn - all emitter types REQUIRE it			*/
-	EPMT_Spawn,
+	EPMT_Spawn UMETA(DisplayName="Spawn"),
 	/** Required - all emitter types REQUIRE it			*/
-	EPMT_Required,
+	EPMT_Required UMETA(DisplayName="Required"),
 	/** Event - event related modules					*/
-	EPMT_Event,
+	EPMT_Event UMETA(DisplayName="Event"),
 	/** Light related modules							*/
-	EPMT_Light,
+	EPMT_Light UMETA(DisplayName="Light"),
 	EPMT_MAX,
 };
 
@@ -40,9 +40,9 @@ UENUM()
 enum EParticleSourceSelectionMethod
 {
 	/** Random		- select a particle at random		*/
-	EPSSM_Random,
+	EPSSM_Random UMETA(DisplayName="Random"),
 	/** Sequential	- select a particle in order		*/
-	EPSSM_Sequential,
+	EPSSM_Sequential UMETA(DisplayName="Sequential"),
 	EPSSM_MAX,
 };
 
diff --git a/Engine/Source/Runtime/Engine/Classes/Particles/ParticleSystem.h b/Engine/Source/Runtime/Engine/Classes/Particles/ParticleSystem.h
--- a/Engine/Source/Runtime/Engine/Classes/Particles/ParticleSystem.h
+++ b/Engine/Source/Runtime/Engine/Classes/Particles/ParticleSystem.h
@@ -12,9 +12,9 @@ UENUM()
 enum EParticleSystemUpdateMode
 {
 	/** RealTime	- update via the delta time passed in				*/
-	EPSUM_RealTime,
+	EPSUM_RealTime UMETA(DisplayName="Real-Time"),
 	/** FixedTime	- update via a fixed time step						*/
-	EPSUM_FixedTime,
+	EPSUM_FixedTime UMETA(DisplayName="Fixed-Time"),
 	EPSUM_MAX,
 };
 
@@ -24,9 +24,9 @@ enum EParticleSystemUpdateMode
 UENUM()
 enum ParticleSystemLODMethod
 {
-	PARTICLESYSTEMLODMETHOD_Automatic,
-	PARTICLESYSTEMLODMETHOD_DirectSet,
-	PARTICLESYSTEMLODMETHOD_ActivateAutomatic,
+	PARTICLESYSTEMLODMETHOD_Automatic UMETA(DisplayName="Automatic"),
+	PARTICLESYSTEMLODMETHOD_DirectSet UMETA(DisplayName="Direct Set"),
+	PARTICLESYSTEMLODMETHOD_ActivateAutomatic UMETA(DisplayName="Activate Automatic"),
 	PARTICLESYSTEMLODMETHOD_MAX,
 };
 
@@ -35,11 +35,11 @@ UENUM()
 enum EParticleSystemOcclusionBoundsMethod
 {
 	/** Don't determine occlusion on this particle system */
-	EPSOBM_None,
+	EPSOBM_None UMETA(DisplayName="None"),
 	/** Use the bounds of the particle system component when determining occlusion */
-	EPSOBM_ParticleBounds,
+	EPSOBM_ParticleBounds UMETA(DisplayName="Particle Bounds"),
 	/** Use the custom occlusion bounds when determining occlusion */
-	EPSOBM_CustomBounds,
+	EPSOBM_CustomBounds UMETA(DisplayName="Custom Bounds"),
 	EPSOBM_MAX,
 };
 
diff --git a/Engine/Source/Runtime/Engine/Classes/Particles/ParticleSystemComponent.h b/Engine/Source/Runtime/Engine/Classes/Particles/ParticleSystemComponent.h
--- a/Engine/Source/Runtime/Engine/Classes/Particles/ParticleSystemComponent.h
+++ b/Engine/Source/Runtime/Engine/Classes/Particles/ParticleSystemComponent.h
@@ -8,14 +8,14 @@
 UENUM()
 enum EParticleSysParamType
 {
-	PSPT_None,
-	PSPT_Scalar,
-	PSPT_ScalarRand,
-	PSPT_Vector,
-	PSPT_VectorRand,
-	PSPT_Color,
-	PSPT_Actor,
-	PSPT_Material,
+	PSPT_None UMETA(DisplayName="None"),
+	PSPT_Scalar UMETA(DisplayName="Scalar"),
+	PSPT_ScalarRand UMETA(DisplayName="Scalar Random"),
+	PSPT_Vector UMETA(DisplayName="Vector"),
+	PSPT_VectorRand UMETA(DisplayName="Vector Random"),
+	PSPT_Color UMETA(DisplayName="Color"),
+	PSPT_Actor UMETA(DisplayName="Actor"),
+	PSPT_Material UMETA(DisplayName="Material"),
 	PSPT_MAX,
 };
 
@@ -24,13 +24,13 @@ UENUM()
 enum ParticleReplayState
 {
 	/** Replay system is disabled.  Particles are simulated and rendered normally. */
-	PRS_Disabled,
+	PRS_Disabled UMETA(DisplayName="Disabled"),
 	/** Capture all particle data to the clip specified by ReplayClipIDNumber.  The frame to capture
 		must be specified using the ReplayFrameIndex */
-	PRS_Capturing,
+	PRS_Capturing UMETA(DisplayName="Capturing"),
 	/** Replay captured particle state from the clip specified by ReplayClipIDNumber.  The frame to play
 		must be specified using the ReplayFrameIndex */
-	PRS_Replaying,
+	PRS_Replaying UMETA(DisplayName="Replaying"),
 	PRS_MAX,
 };
 
@@ -41,17 +41,17 @@ UENUM()
 enum EParticleEventType
 {
 	/** Any - allow any event */
-	EPET_Any,
+	EPET_Any UMETA(DisplayName="Any"),
 	/** Spawn - a particle spawn event */
-	EPET_Spawn,
+	EPET_Spawn UMETA(DisplayName="Spawn"),
 	/** Death - a particle death event */
-	EPET_Death,
+	EPET_Death UMETA(DisplayName="Death"),
 	/** Collision - a particle collision event */
-	EPET_Collision,
+	EPET_Collision UMETA(DisplayName="Collision"),
 	/** Burst - a particle burst event */
-	EPET_Burst,
+	EPET_Burst UMETA(DisplayName="Burst"),
 	/** Blueprint - an event generated by level script */
-	EPET_Blueprint,
+	EPET_Blueprint UMETA(DisplayName="Blueprint"),
 	EPET_MAX,
 };
 
@@ -665,7 +665,7 @@ public:
 	/** 
 	  * Decide which detail mode should be applied to this particle system. If we have an editor
 	  * override specified, use that. Otherwise use the global cached value 
- 	  */
+	  */
 	int32 GetCurrentDetailMode() const;
 
 private:
diff --git a/Engine/Source/Runtime/Engine/Classes/Particles/TypeData/ParticleModuleTypeDataBeam2.h b/Engine/Source/Runtime/Engine/Classes/Particles/TypeData/ParticleModuleTypeDataBeam2.h
--- a/Engine/Source/Runtime/Engine/Classes/Particles/TypeData/ParticleModuleTypeDataBeam2.h
+++ b/Engine/Source/Runtime/Engine/Classes/Particles/TypeData/ParticleModuleTypeDataBeam2.h
@@ -7,9 +7,9 @@
 UENUM()
 enum EBeam2Method
 {
-	PEB2M_Distance,
-	PEB2M_Target,
-	PEB2M_Branch,
+	PEB2M_Distance UMETA(DisplayName="Distance"),
+	PEB2M_Target UMETA(DisplayName="Target"),
+	PEB2M_Branch UMETA(DisplayName="Branch"),
 	PEB2M_MAX,
 };
 
@@ -19,9 +19,9 @@ enum EBeam2Method
 UENUM()
 enum EBeamTaperMethod
 {
-	PEBTM_None,
-	PEBTM_Full,
-	PEBTM_Partial,
+	PEBTM_None UMETA(DisplayName="None"),
+	PEBTM_Full UMETA(DisplayName="Full"),
+	PEBTM_Partial UMETA(DisplayName="Partial"),
 	PEBTM_MAX,
 };
 
diff --git a/Engine/Source/Runtime/Engine/Classes/Particles/TypeData/ParticleModuleTypeDataMesh.h b/Engine/Source/Runtime/Engine/Classes/Particles/TypeData/ParticleModuleTypeDataMesh.h
--- a/Engine/Source/Runtime/Engine/Classes/Particles/TypeData/ParticleModuleTypeDataMesh.h
+++ b/Engine/Source/Runtime/Engine/Classes/Particles/TypeData/ParticleModuleTypeDataMesh.h
@@ -7,39 +7,42 @@
 UENUM()
 enum EMeshScreenAlignment
 {
-	PSMA_MeshFaceCameraWithRoll,
-	PSMA_MeshFaceCameraWithSpin,
-	PSMA_MeshFaceCameraWithLockedAxis,
+	PSMA_MeshFaceCameraWithRoll UMETA(DisplayName="Face Camera With Roll"),
+	PSMA_MeshFaceCameraWithSpin UMETA(DisplayName="Face Camera With Spin"),
+	PSMA_MeshFaceCameraWithLockedAxis UMETA(DisplayName="Face Camera With Locked-Axis"),
 	PSMA_MAX,
 };
 
 UENUM()
 enum EMeshCameraFacingUpAxis
 {
-	CameraFacing_NoneUP,
-	CameraFacing_ZUp,
-	CameraFacing_NegativeZUp,
-	CameraFacing_YUp,
-	CameraFacing_NegativeYUp,
+	CameraFacing_NoneUP UMETA(DisplayName="None"),
+	CameraFacing_ZUp UMETA(DisplayName="Z Up"),
+	CameraFacing_NegativeZUp UMETA(DisplayName="-Z Up"),
+	CameraFacing_YUp UMETA(DisplayName="Y Up"),
+	CameraFacing_NegativeYUp UMETA(DisplayName="-Y Up"),
 	CameraFacing_MAX,
 };
 
 UENUM()
 enum EMeshCameraFacingOptions
 {
-	XAxisFacing_NoUp,
-	XAxisFacing_ZUp,
-	XAxisFacing_NegativeZUp,
-	XAxisFacing_YUp,
-	XAxisFacing_NegativeYUp,
-	LockedAxis_ZAxisFacing,
-	LockedAxis_NegativeZAxisFacing,
-	LockedAxis_YAxisFacing,
-	LockedAxis_NegativeYAxisFacing,
-	VelocityAligned_ZAxisFacing,
-	VelocityAligned_NegativeZAxisFacing,
-	VelocityAligned_YAxisFacing,
-	VelocityAligned_NegativeYAxisFacing,
+	XAxisFacing_NoUp UMETA(DisplayName="X Axis Facing : No Up"),
+	XAxisFacing_ZUp UMETA(DisplayName="X Axis Facing : Z Up"),
+	XAxisFacing_NegativeZUp UMETA(DisplayName="X Axis Facing : -Z Up"),
+	XAxisFacing_YUp UMETA(DisplayName="X Axis Facing : Y Up"),
+	XAxisFacing_NegativeYUp UMETA(DisplayName="X Axis Facing : -Y Up"),
+
+	LockedAxis_ZAxisFacing UMETA(DisplayName="Locked Axis : Z Axis Facing"),
+	LockedAxis_NegativeZAxisFacing UMETA(DisplayName="Locked Axis : -Z Axis Facing"),
+	LockedAxis_YAxisFacing UMETA(DisplayName="Locked Axis : Y Axis Facing"),
+	LockedAxis_NegativeYAxisFacing UMETA(DisplayName="Locked Axis : -Y Axis Facing"),
+
+	VelocityAligned_ZAxisFacing UMETA(DisplayName="Velocity Aligned : Z Axis Facing"),
+	VelocityAligned_NegativeZAxisFacing UMETA(DisplayName="Velocity Aligned : -Z Axis Facing"),
+	VelocityAligned_YAxisFacing UMETA(DisplayName="Velocity Aligned : Y Axis Facing"),
+	VelocityAligned_NegativeYAxisFacing UMETA(DisplayName="Velocity Aligned : -Y Axis Facing"),
+
 	EMeshCameraFacingOptions_MAX,
 };
 
diff --git a/Engine/Source/Runtime/Engine/Classes/Particles/TypeData/ParticleModuleTypeDataRibbon.h b/Engine/Source/Runtime/Engine/Classes/Particles/TypeData/ParticleModuleTypeDataRibbon.h
--- a/Engine/Source/Runtime/Engine/Classes/Particles/TypeData/ParticleModuleTypeDataRibbon.h
+++ b/Engine/Source/Runtime/Engine/Classes/Particles/TypeData/ParticleModuleTypeDataRibbon.h
@@ -12,9 +12,9 @@
 UENUM()
 enum ETrailsRenderAxisOption
 {
-	Trails_CameraUp,
-	Trails_SourceUp,
-	Trails_WorldUp,
+	Trails_CameraUp UMETA(DisplayName="Camera Up"),
+	Trails_SourceUp UMETA(DisplayName="Source Up"),
+	Trails_WorldUp UMETA(DisplayName="World Up"),
 	Trails_MAX,
 };
 
