#ue4 Fix for crash on NULL font for text render components

diff --git a/Engine/Source/Runtime/Engine/Private/TextRenderComponent.cpp b/Engine/Source/Runtime/Engine/Private/TextRenderComponent.cpp
--- a/Engine/Source/Runtime/Engine/Private/TextRenderComponent.cpp
+++ b/Engine/Source/Runtime/Engine/Private/TextRenderComponent.cpp
@@ -738,7 +738,7 @@ bool UTextRenderComponent::ShouldRecreateProxyOnUpdateTransform() const
 
 FBoxSphereBounds UTextRenderComponent::CalcBounds(const FTransform & LocalToWorld) const
 {
-	if(!Text.IsEmpty())
+	if(!Text.IsEmpty() && Font)
 	{
 		FVector2D Size(FLT_MIN, 0);
 		FVector2D LeftTop(FLT_MAX, FLT_MAX);
@@ -781,7 +781,7 @@ FBoxSphereBounds UTextRenderComponent::CalcBounds(const FTransform & LocalToWorl
 FMatrix UTextRenderComponent::GetRenderMatrix() const
 {
 	// Adjust LocalToWorld transform to account for vertical text alignment when rendering.
-	if(!Text.IsEmpty())
+	if(!Text.IsEmpty() && Font)
 	{
 		float SizeY = 0;
 		float FirstLineHeight = -1;
