331289  LIVE: CRASH: LANDSCAPE: COLLISION: Editor crashes when attempting to add an Object Channel when Landscape exists in the level

diff --git a/Engine/Source/Runtime/Engine/Private/PhysicsEngine/BodyInstance.cpp b/Engine/Source/Runtime/Engine/Private/PhysicsEngine/BodyInstance.cpp
--- a/Engine/Source/Runtime/Engine/Private/PhysicsEngine/BodyInstance.cpp
+++ b/Engine/Source/Runtime/Engine/Private/PhysicsEngine/BodyInstance.cpp
@@ -444,7 +444,11 @@ void FBodyInstance::UpdatePhysicsFilterData()
 		return;
 	}
 
-	check(BodySetup.IsValid());
+	// this can happen in landscape height field collision component
+	if (!BodySetup.IsValid())
+	{
+		return;
+	}
 
 	// Figure out if we are static
 	AActor* Owner = OwnerComponent.IsValid() ? OwnerComponent.Get()->GetOwner() : NULL;
