- Fix for TTP:325321 where music doesn't pause on app sleep on Android
- This fix is a bit sensitive to order of operations, but works for the Android suspend/resume process.

diff --git a/Engine/Source/Runtime/Engine/Private/AudioDevice.cpp b/Engine/Source/Runtime/Engine/Private/AudioDevice.cpp
--- a/Engine/Source/Runtime/Engine/Private/AudioDevice.cpp
+++ b/Engine/Source/Runtime/Engine/Private/AudioDevice.cpp
@@ -144,6 +144,11 @@ void FAudioDevice::Teardown()
 	FreeSources.Empty();
 }
 
+void FAudioDevice::Suspend(bool bGameTicking)
+{
+	HandlePause( bGameTicking, true );
+}
+
 void FAudioDevice::CountBytes(FArchive& Ar)
 {
 	Sources.CountBytes(Ar);
@@ -1634,7 +1639,7 @@ void FAudioDevice::DestroyEffect( FSoundSource* Source )
 }
 
 
-void FAudioDevice::HandlePause( bool bGameTicking )
+void FAudioDevice::HandlePause( bool bGameTicking, bool bGlobalPause )
 {
 	// Pause all sounds if transitioning to pause mode.
 	if( !bGameTicking && bGameWasTicking )
@@ -1642,7 +1647,7 @@ void FAudioDevice::HandlePause( bool bGameTicking )
 		for( int32 i = 0; i < Sources.Num(); i++ )
 		{
 			FSoundSource* Source = Sources[ i ];
-			if( Source->IsGameOnly() )
+			if( bGlobalPause || Source->IsGameOnly() )
 			{
 				Source->Pause();
 			}
@@ -1654,7 +1659,7 @@ void FAudioDevice::HandlePause( bool bGameTicking )
 		for( int32 i = 0; i < Sources.Num(); i++ )
 		{
 			FSoundSource* Source = Sources[ i ];
-			if( Source->IsGameOnly() )
+			if( bGlobalPause || Source->IsGameOnly() )
 			{
 				Source->Play();
 			}
diff --git a/Engine/Source/Runtime/Engine/Public/AudioDevice.h b/Engine/Source/Runtime/Engine/Public/AudioDevice.h
--- a/Engine/Source/Runtime/Engine/Public/AudioDevice.h
+++ b/Engine/Source/Runtime/Engine/Public/AudioDevice.h
@@ -232,6 +232,13 @@ public:
 	 */
 	void Update(bool bGameTicking);
 
+	/**
+	 * Suspend/resume all sounds (global pause for device suspend/resume, etc.)
+	 *
+	 * @param bGameTicking Whether the game is still ticking at the time of suspend
+	 */
+	void Suspend(bool bGameTicking);
+
 	/**
 	 * Counts the bytes for the structures used in this class
 	 */
@@ -483,9 +490,9 @@ protected:
 	friend class FSoundSource;
 
 	/**
-	 * Handle pausing/unpausing of sources when entering or leaving pause mode
+	 * Handle pausing/unpausing of sources when entering or leaving pause mode, or global pause (like device suspend)
 	 */
-	void HandlePause( bool bGameTicking );
+	void HandlePause( bool bGameTicking, bool bGlobalPause = false );
 
 	/**
 	 * Stop sources that need to be stopped, and touch the ones that need to be kept alive
diff --git a/Engine/Source/Runtime/Launch/Private/Android/AndroidEventManager.cpp b/Engine/Source/Runtime/Launch/Private/Android/AndroidEventManager.cpp
--- a/Engine/Source/Runtime/Launch/Private/Android/AndroidEventManager.cpp
+++ b/Engine/Source/Runtime/Launch/Private/Android/AndroidEventManager.cpp
@@ -222,16 +222,20 @@ void FAppEventManager::PauseAudio()
 {
 	bAudioPaused = true;
 
-	// AudioDevice::HandlePause() is called from the audio devices Update(), which isn't called when we're paused, so force it here
 	if (GEngine->GetAudioDevice())
 	{
-		GEngine->GetAudioDevice()->Update(false);
+		GEngine->GetAudioDevice()->Suspend(false);
 	}
 }
 
 void FAppEventManager::ResumeAudio()
 {
 	bAudioPaused = false;
+
+	if (GEngine->GetAudioDevice())
+	{
+		GEngine->GetAudioDevice()->Suspend(true);
+	}
 }
 
 void FAppEventManager::EnqueueAppEvent(EAppEventState InState, void* InData)
