#ue4
- Added an OverrideAsset hook to runtime animation nodes

diff --git a/Engine/Source/Runtime/Engine/Classes/Animation/AnimNodeBase.h b/Engine/Source/Runtime/Engine/Classes/Animation/AnimNodeBase.h
--- a/Engine/Source/Runtime/Engine/Classes/Animation/AnimNodeBase.h
+++ b/Engine/Source/Runtime/Engine/Classes/Animation/AnimNodeBase.h
@@ -308,5 +308,8 @@ struct ENGINE_API FAnimNode_Base
 	virtual void Update(const FAnimationUpdateContext& Context) {}
 	virtual void Evaluate(FPoseContext& Output) { check(false); }
 	virtual void EvaluateComponentSpace(FComponentSpacePoseContext& Output) { check(false); }
+
+	// If a derived anim node should respond to asset overrides, OverrideAsset should be defined to handle changing the asset
+	virtual void OverrideAsset(UAnimationAsset* NewAsset) {}
 	// End of interface to implement
 };
diff --git a/Engine/Source/Runtime/Engine/Classes/Animation/AnimNode_BlendSpacePlayer.h b/Engine/Source/Runtime/Engine/Classes/Animation/AnimNode_BlendSpacePlayer.h
--- a/Engine/Source/Runtime/Engine/Classes/Animation/AnimNode_BlendSpacePlayer.h
+++ b/Engine/Source/Runtime/Engine/Classes/Animation/AnimNode_BlendSpacePlayer.h
@@ -61,5 +61,6 @@ public:
 	virtual void CacheBones(const FAnimationCacheBonesContext & Context) OVERRIDE;
 	virtual void Update(const FAnimationUpdateContext& Context) OVERRIDE;
 	virtual void Evaluate(FPoseContext& Output) OVERRIDE;
+	virtual void OverrideAsset(UAnimationAsset* NewAsset) OVERRIDE;
 	// End of FAnimNode_Base interface
 };
diff --git a/Engine/Source/Runtime/Engine/Classes/Animation/AnimNode_SequenceEvaluator.h b/Engine/Source/Runtime/Engine/Classes/Animation/AnimNode_SequenceEvaluator.h
--- a/Engine/Source/Runtime/Engine/Classes/Animation/AnimNode_SequenceEvaluator.h
+++ b/Engine/Source/Runtime/Engine/Classes/Animation/AnimNode_SequenceEvaluator.h
@@ -31,5 +31,6 @@ public:
 	virtual void CacheBones(const FAnimationCacheBonesContext & Context) OVERRIDE;
 	virtual void Update(const FAnimationUpdateContext& Context) OVERRIDE;
 	virtual void Evaluate(FPoseContext& Output) OVERRIDE;
+	virtual void OverrideAsset(UAnimationAsset* NewAsset) OVERRIDE;
 	// End of FAnimNode_Base interface
 };
diff --git a/Engine/Source/Runtime/Engine/Classes/Animation/AnimNode_SequencePlayer.h b/Engine/Source/Runtime/Engine/Classes/Animation/AnimNode_SequencePlayer.h
--- a/Engine/Source/Runtime/Engine/Classes/Animation/AnimNode_SequencePlayer.h
+++ b/Engine/Source/Runtime/Engine/Classes/Animation/AnimNode_SequencePlayer.h
@@ -49,5 +49,6 @@ public:
 	virtual void CacheBones(const FAnimationCacheBonesContext & Context) OVERRIDE;
 	virtual void Update(const FAnimationUpdateContext& Context) OVERRIDE;
 	virtual void Evaluate(FPoseContext& Output) OVERRIDE;
+	virtual void OverrideAsset(UAnimationAsset* NewAsset) OVERRIDE;
 	// End of FAnimNode_Base interface
 };
diff --git a/Engine/Source/Runtime/Engine/Private/Animation/AnimNode_BlendSpacePlayer.cpp b/Engine/Source/Runtime/Engine/Private/Animation/AnimNode_BlendSpacePlayer.cpp
--- a/Engine/Source/Runtime/Engine/Private/Animation/AnimNode_BlendSpacePlayer.cpp
+++ b/Engine/Source/Runtime/Engine/Private/Animation/AnimNode_BlendSpacePlayer.cpp
@@ -75,3 +75,11 @@ void FAnimNode_BlendSpacePlayer::Evaluate(FPoseContext& Output)
 		Output.ResetToRefPose();
 	}
 }
+
+void FAnimNode_BlendSpacePlayer::OverrideAsset(UAnimationAsset* NewAsset)
+{
+	if(UBlendSpaceBase* NewBlendSpace = Cast<UBlendSpaceBase>(NewAsset))
+	{
+		BlendSpace = NewBlendSpace;
+	}
+}
diff --git a/Engine/Source/Runtime/Engine/Private/Animation/AnimNode_SequenceEvaluator.cpp b/Engine/Source/Runtime/Engine/Private/Animation/AnimNode_SequenceEvaluator.cpp
--- a/Engine/Source/Runtime/Engine/Private/Animation/AnimNode_SequenceEvaluator.cpp
+++ b/Engine/Source/Runtime/Engine/Private/Animation/AnimNode_SequenceEvaluator.cpp
@@ -28,4 +28,12 @@ void FAnimNode_SequenceEvaluator::Evaluate(FPoseContext& Output)
 	{
 		Output.ResetToRefPose();
 	}
-}
\ No newline at end of file
+}
+
+void FAnimNode_SequenceEvaluator::OverrideAsset(UAnimationAsset* NewAsset)
+{
+	if(UAnimSequenceBase* NewSequence = Cast<UAnimSequenceBase>(NewAsset))
+	{
+		Sequence = NewSequence;
+	}
+}
diff --git a/Engine/Source/Runtime/Engine/Private/Animation/AnimNode_SequencePlayer.cpp b/Engine/Source/Runtime/Engine/Private/Animation/AnimNode_SequencePlayer.cpp
--- a/Engine/Source/Runtime/Engine/Private/Animation/AnimNode_SequencePlayer.cpp
+++ b/Engine/Source/Runtime/Engine/Private/Animation/AnimNode_SequencePlayer.cpp
@@ -56,3 +56,11 @@ void FAnimNode_SequencePlayer::Evaluate(FPoseContext& Output)
 		Output.ResetToRefPose();
 	}
 }
+
+void FAnimNode_SequencePlayer::OverrideAsset(UAnimationAsset* NewAsset)
+{
+	if(UAnimSequenceBase* AnimSequence = Cast<UAnimSequenceBase>(NewAsset))
+	{
+		Sequence = AnimSequence;
+	}
+}
