Landscape : Added support to conisder LODBias/ForcedLOD affects to Lightmass geometry

diff --git a/Engine/Source/Runtime/Engine/Private/Landscape/LandscapeEdit.cpp b/Engine/Source/Runtime/Engine/Private/Landscape/LandscapeEdit.cpp
--- a/Engine/Source/Runtime/Engine/Private/Landscape/LandscapeEdit.cpp
+++ b/Engine/Source/Runtime/Engine/Private/Landscape/LandscapeEdit.cpp
@@ -3685,7 +3685,7 @@ void ALandscape::PostEditChangeProperty(FPropertyChangedEvent& PropertyChangedEv
 	else if (GIsEditor && PropertyName == GET_MEMBER_NAME_CHECKED(ALandscapeProxy, StaticLightingLOD))
 	{
 		StaticLightingLOD = FMath::Clamp<int32>(StaticLightingLOD, 0, FMath::CeilLogTwo(SubsectionSizeQuads+1)-1);
-		InvalidateLightingCache();
+		bChangedLighting = true;
 	}
 	else if (GIsEditor && PropertyName == GET_MEMBER_NAME_CHECKED(ALandscapeProxy, ExportLOD))
 	{
@@ -3734,6 +3734,11 @@ void ALandscape::PostEditChangeProperty(FPropertyChangedEvent& PropertyChangedEv
 					Comp->UpdateMaterialInstances();
 				}
 
+				if (bChangedLighting)
+				{
+					Comp->InvalidateLightingCache();
+				}
+
 				// Reattach all components
 				FComponentReregisterContext ReregisterContext(Comp);
 			}
@@ -3805,6 +3810,9 @@ void ULandscapeComponent::SetLOD(bool bForcedLODChanged, int32 InLODValue)
 		int32 MaxLOD = FMath::CeilLogTwo(SubsectionSizeQuads+1)-1;
 		LODBias = FMath::Clamp<int32>(InLODValue, -MaxLOD, MaxLOD);
 	}
+
+	InvalidateLightingCache();
+
 	// Update neighbor components
 	ULandscapeInfo* Info = GetLandscapeInfo(false);
 	if (Info)
@@ -3837,6 +3845,8 @@ void ULandscapeComponent::SetLOD(bool bForcedLODChanged, int32 InLODValue)
 					// Neighbor LODBias are saved in BYTE, so need to convert to range [-128:127]
 					Comp->NeighborLODBias[7-Idx] = LODBias + 128;
 				}
+
+				Comp->InvalidateLightingCache();
 				FComponentReregisterContext ReregisterContext(Comp);
 			}
 		}
diff --git a/Engine/Source/Runtime/Engine/Private/Landscape/LandscapeLight.cpp b/Engine/Source/Runtime/Engine/Private/Landscape/LandscapeLight.cpp
--- a/Engine/Source/Runtime/Engine/Private/Landscape/LandscapeLight.cpp
+++ b/Engine/Source/Runtime/Engine/Private/Landscape/LandscapeLight.cpp
@@ -120,14 +120,226 @@ FLandscapeStaticLightingMesh::FLandscapeStaticLightingMesh(ULandscapeComponent*
 	NumQuads = NumVertices - 1;
 	UVFactor = LightMapRatio / NumVertices;
 
-	GetHeightmapData(InLOD);
+	int32 GeometricLOD = FMath::Clamp<int32>(InComponent->ForcedLOD >= 0 ? InComponent->ForcedLOD : InComponent->LODBias, 0, FMath::CeilLogTwo(InComponent->SubsectionSizeQuads+1) - 1);
+	GetHeightmapData(InLOD, FMath::Max(GeometricLOD, InLOD));
 }
 
 FLandscapeStaticLightingMesh::~FLandscapeStaticLightingMesh()
 {
 }
 
-void FLandscapeStaticLightingMesh::GetHeightmapData(int32 InLOD)
+namespace
+{
+	void GetLODData(ULandscapeComponent* LandscapeComponent, int32 X, int32 Y, int32 HeightmapOffsetX, int32 HeightmapOffsetY, int32 LODValue, int32 HeightmapStride, FColor& OutHeight, FColor& OutXYOffset)
+	{
+		int32 ComponentSize = ((LandscapeComponent->SubsectionSizeQuads + 1) * LandscapeComponent->NumSubsections) >> LODValue;
+		int32 LODHeightmapSize = LandscapeComponent->HeightmapTexture->Source.GetSizeX() >> LODValue;
+		float Ratio = (float)(LODHeightmapSize) / (HeightmapStride);
+		float Offset = 0.5f * Ratio;
+
+		int32 CurrentHeightmapOffsetX = FMath::Round((float)(LODHeightmapSize)* LandscapeComponent->HeightmapScaleBias.Z);
+		int32 CurrentHeightmapOffsetY = FMath::Round((float)(LODHeightmapSize)* LandscapeComponent->HeightmapScaleBias.W);
+
+		// Need to match for component edge cases, otherwise it causes a little 
+		float XX = FMath::Clamp<float>((X - HeightmapOffsetX) * Ratio - Offset, 0.f, ComponentSize - 1.f) + CurrentHeightmapOffsetX;
+		int32 XI = (int32)XX;
+		float XF = XX - XI;
+
+		float YY = FMath::Clamp<float>((Y - HeightmapOffsetY) * Ratio - Offset, 0.f, ComponentSize - 1.f) + CurrentHeightmapOffsetY;
+		int32 YI = (int32)YY;
+		float YF = YY - YI;
+
+		FLandscapeComponentDataInterface DataInterface(LandscapeComponent, LODValue);
+		FColor* HeightMipData = DataInterface.GetRawHeightData();
+		FColor* XYOffsetMipData = DataInterface.GetRawXYOffsetData();
+
+		FColor H1 = HeightMipData[XI + YI * LODHeightmapSize];
+		FColor H2 = HeightMipData[FMath::Min(XI + 1, LODHeightmapSize - 1) + YI * LODHeightmapSize];
+		FColor H3 = HeightMipData[XI + FMath::Min(YI + 1, LODHeightmapSize - 1) * LODHeightmapSize];
+		FColor H4 = HeightMipData[FMath::Min(XI + 1, LODHeightmapSize - 1) + FMath::Min(YI + 1, LODHeightmapSize - 1) * LODHeightmapSize];
+
+		uint16 Height = FMath::Round(FMath::Lerp(FMath::Lerp<float>(((H1.R << 8) + H1.G), ((H2.R << 8) + H2.G), XF),
+			FMath::Lerp<float>(((H3.R << 8) + H3.G), ((H4.R << 8) + H4.G), XF), YF));
+		uint8 B = FMath::Round(FMath::Lerp(FMath::Lerp<float>((H1.B), (H2.B), XF),
+			FMath::Lerp<float>((H3.B), (H4.B), XF), YF));
+		uint8 A = FMath::Round(FMath::Lerp<float>(FMath::Lerp((H1.A), (H2.A), XF),
+			FMath::Lerp<float>((H3.A), (H4.A), XF), YF));
+
+		OutHeight = FColor((Height >> 8), Height & 255, B, A);
+
+		if (LandscapeComponent->XYOffsetmapTexture)
+		{
+			FColor X1 = XYOffsetMipData[XI + YI * LODHeightmapSize];
+			FColor X2 = XYOffsetMipData[FMath::Min(XI + 1, LODHeightmapSize - 1) + YI * LODHeightmapSize];
+			FColor X3 = XYOffsetMipData[XI + FMath::Min(YI + 1, LODHeightmapSize - 1) * LODHeightmapSize];
+			FColor X4 = XYOffsetMipData[FMath::Min(XI + 1, LODHeightmapSize - 1) + FMath::Min(YI + 1, LODHeightmapSize - 1) * LODHeightmapSize];
+
+			uint16 XComp = FMath::Round(FMath::Lerp(FMath::Lerp<float>(((X1.R << 8) + X1.G), ((X2.R << 8) + X2.G), XF),
+				FMath::Lerp<float>(((X3.R << 8) + X3.G), ((X4.R << 8) + X4.G), XF), YF));
+			uint16 YComp = FMath::Round(FMath::Lerp(FMath::Lerp<float>(((X1.B << 8) + X1.A), ((X2.B << 8) + X2.A), XF),
+				FMath::Lerp<float>(((X3.B << 8) + X3.A), ((X4.B << 8) + X4.A), XF), YF));
+
+			OutXYOffset = FColor((XComp >> 8), XComp & 255, YComp >> 8, YComp & 255);
+		}
+	}
+
+	void InternalUpscaling(FLandscapeComponentDataInterface& DataInterface, ULandscapeComponent* LandscapeComponent, int32 InLOD, int32 GeometryLOD, TArray<FColor>& CompHeightData, TArray<FColor>& CompXYOffsetData)
+	{
+		// Upscaling using Landscape LOD system 
+		ULandscapeInfo* const Info = LandscapeComponent->GetLandscapeInfo();
+		check(Info);
+
+		FIntPoint ComponentBase = LandscapeComponent->GetSectionBase() / LandscapeComponent->ComponentSizeQuads;
+		ULandscapeComponent* Neighbors[8] = { 0, 0, 0, 0, 0, 0, 0, 0 };
+		int32 NeighborLODs[8] = { 0, 0, 0, 0, 0, 0, 0, 0 };
+
+		Neighbors[0] = Info->XYtoComponentMap.FindRef(ComponentBase + FIntPoint(-1, -1));
+		Neighbors[1] = Info->XYtoComponentMap.FindRef(ComponentBase + FIntPoint(0, -1));
+		Neighbors[2] = Info->XYtoComponentMap.FindRef(ComponentBase + FIntPoint(1, -1));
+		Neighbors[3] = Info->XYtoComponentMap.FindRef(ComponentBase + FIntPoint(-1, 0));
+		Neighbors[4] = Info->XYtoComponentMap.FindRef(ComponentBase + FIntPoint(1, 0));
+		Neighbors[5] = Info->XYtoComponentMap.FindRef(ComponentBase + FIntPoint(-1, 1));
+		Neighbors[6] = Info->XYtoComponentMap.FindRef(ComponentBase + FIntPoint(0, 1));
+		Neighbors[7] = Info->XYtoComponentMap.FindRef(ComponentBase + FIntPoint(1, 1));
+
+		int32 MaxLOD = FMath::CeilLogTwo(LandscapeComponent->SubsectionSizeQuads + 1) - 1;
+		bool bNeedUpscaling = GeometryLOD > InLOD;
+		int32 NeigborMaxLOD = -1;
+
+		for (int32 i = 0; i < 8; ++i)
+		{
+			NeighborLODs[i] = Neighbors[i] ? FMath::Clamp<int32>(Neighbors[i]->ForcedLOD >= 0 ? Neighbors[i]->ForcedLOD : Neighbors[i]->LODBias, 0, MaxLOD) : -1;
+			bNeedUpscaling |= (NeighborLODs[i] > InLOD);
+			NeigborMaxLOD = FMath::Max(NeighborLODs[i], NeigborMaxLOD);
+		}
+
+		for (int32 i = 0; i < 8; ++i)
+		{
+			if (Neighbors[i] == NULL)
+			{
+				NeighborLODs[i] = NeigborMaxLOD;
+			}
+		}
+
+		if (bNeedUpscaling)
+		{
+			check(LandscapeComponent);
+			// Need Upscaling
+			int32 HeightmapStride = LandscapeComponent->HeightmapTexture->Source.GetSizeX() >> InLOD;
+
+			int32 HeightDataSize = HeightmapStride * HeightmapStride;
+			CompHeightData.Empty(HeightDataSize);
+			CompXYOffsetData.Empty(HeightDataSize);
+			CompHeightData.AddZeroed(HeightDataSize);
+			CompXYOffsetData.AddZeroed(HeightDataSize);
+
+			// Update for only component region for performance
+			int32 ComponentSize = ((LandscapeComponent->SubsectionSizeQuads + 1) * LandscapeComponent->NumSubsections) >> InLOD;
+
+			for (int32 Y = DataInterface.HeightmapComponentOffsetY; Y < DataInterface.HeightmapComponentOffsetY + ComponentSize; ++Y)
+			{
+				for (int32 X = DataInterface.HeightmapComponentOffsetX; X < DataInterface.HeightmapComponentOffsetX + ComponentSize; ++X)
+				{
+					// LOD System same as shader
+					FVector2D XY(float(X - DataInterface.HeightmapComponentOffsetX) / (ComponentSize - 1), float(Y - DataInterface.HeightmapComponentOffsetY) / (ComponentSize - 1));
+					XY = XY - 0.5f;
+
+					float Dist = XY.Size();
+
+					float RealLOD = GeometryLOD;
+
+					if (XY.X < 0.f)
+					{
+						if (XY.Y < 0.f)
+						{
+							RealLOD = FMath::Lerp(
+								FMath::Lerp<float>(NeighborLODs[0], NeighborLODs[1], XY.X + 1.f),
+								FMath::Lerp<float>(NeighborLODs[3], GeometryLOD, XY.X + 1.f),
+								XY.Y + 1.f); // 0
+						}
+						else
+						{
+							RealLOD = FMath::Lerp(
+								FMath::Lerp<float>(NeighborLODs[3], GeometryLOD, XY.X + 1.f),
+								FMath::Lerp<float>(NeighborLODs[5], NeighborLODs[6], XY.X + 1.f),
+								XY.Y); // 2
+						}
+					}
+					else
+					{
+						if (XY.Y < 0.f)
+						{
+							RealLOD = FMath::Lerp(
+								FMath::Lerp<float>(NeighborLODs[1], NeighborLODs[2], XY.X),
+								FMath::Lerp<float>(GeometryLOD, NeighborLODs[4], XY.X),
+								XY.Y + 1.f); // 1
+						}
+						else
+						{
+							RealLOD = FMath::Lerp(
+								FMath::Lerp<float>(GeometryLOD, NeighborLODs[4], XY.X),
+								FMath::Lerp<float>(NeighborLODs[6], NeighborLODs[7], XY.X),
+								XY.Y); // 3
+						}
+					}
+
+					RealLOD = FMath::Min(RealLOD, (float)MaxLOD);
+
+					int32 LODValue = (int32)RealLOD;
+					float MorphAlpha = FMath::Fractional(RealLOD);
+
+					FColor Height[2];
+					FColor XYOffset[2];
+					::GetLODData(LandscapeComponent, X, Y, DataInterface.HeightmapComponentOffsetX, DataInterface.HeightmapComponentOffsetY,
+						FMath::Min(MaxLOD, LODValue), HeightmapStride, Height[0], XYOffset[0]);
+
+					// Interpolation between two LOD
+					if ((RealLOD > InLOD) && (LODValue + 1 <= MaxLOD) && MorphAlpha != 0.f)
+					{
+						::GetLODData(LandscapeComponent, X, Y, DataInterface.HeightmapComponentOffsetX, DataInterface.HeightmapComponentOffsetY,
+							FMath::Min(MaxLOD, LODValue + 1), HeightmapStride, Height[1], XYOffset[1]);
+
+						// Need interpolation
+						uint16 Height0 = (Height[0].R << 8) + Height[0].G;
+						uint16 Height1 = (Height[1].R << 8) + Height[1].G;
+						uint16 LerpHeight = FMath::Round(FMath::Lerp<float>(Height0, Height1, MorphAlpha));
+
+						CompHeightData[X + Y * HeightmapStride] = 
+							FColor((LerpHeight >> 8), LerpHeight & 255, 
+							FMath::Round(FMath::Lerp<float>(Height[0].B, Height[1].B, MorphAlpha)), 
+							FMath::Round(FMath::Lerp<float>(Height[0].A, Height[1].A, MorphAlpha)));
+						if (LandscapeComponent->XYOffsetmapTexture)
+						{
+							uint16 XComp0 = (XYOffset[0].R << 8) + XYOffset[0].G;
+							uint16 XComp1 = (XYOffset[1].R << 8) + XYOffset[1].G;
+							uint16 LerpXComp = FMath::Round(FMath::Lerp<float>(XComp0, XComp1, MorphAlpha));
+
+							uint16 YComp0 = (XYOffset[0].B << 8) + XYOffset[0].A;
+							uint16 YComp1 = (XYOffset[1].B << 8) + XYOffset[1].A;
+							uint16 LerpYComp = FMath::Round(FMath::Lerp<float>(YComp0, YComp1, MorphAlpha));
+
+							CompXYOffsetData[X + Y * HeightmapStride] =
+								FColor(LerpXComp >> 8, LerpXComp & 255, LerpYComp >> 8, LerpYComp & 255);
+						}
+					}
+					else
+					{
+						CompHeightData[X + Y * HeightmapStride] = Height[0];
+						CompXYOffsetData[X + Y * HeightmapStride] = XYOffset[0];
+					}
+				}
+			}
+
+			DataInterface.SetRawHeightData(&CompHeightData[0]);
+			if (LandscapeComponent->XYOffsetmapTexture)
+			{
+				DataInterface.SetRawXYOffsetData(&CompXYOffsetData[0]);
+			}
+		}
+	}
+};
+
+void FLandscapeStaticLightingMesh::GetHeightmapData(int32 InLOD, int32 GeometryLOD)
 {
 	ULandscapeInfo* const Info = LandscapeComponent->GetLandscapeInfo();
 	check(Info);
@@ -144,9 +356,15 @@ void FLandscapeStaticLightingMesh::GetHeightmapData(int32 InLOD)
 	check(ExpandQuadsX <= SubsectionSizeQuads);
 	check(ExpandQuadsY <= SubsectionSizeQuads);
 
+	int32 MaxLOD = FMath::CeilLogTwo(LandscapeComponent->SubsectionSizeQuads + 1) - 1;
+
 	// copy heightmap data for this component...
 	{
+		// Data array for upscaling case
+		TArray<FColor> CompHeightData;
+		TArray<FColor> CompXYOffsetData;
 		FLandscapeComponentDataInterface DataInterface(LandscapeComponent, InLOD);
+		::InternalUpscaling(DataInterface, LandscapeComponent, InLOD, GeometryLOD, CompHeightData, CompXYOffsetData);
 
 		for (int32 Y = 0; Y < ComponentSizeQuads + 1; Y++)
 		{
@@ -179,18 +397,22 @@ void FLandscapeStaticLightingMesh::GetHeightmapData(int32 InLOD)
 			const int32 YSource = (ComponentY == 0) ? (ComponentSizeQuads - ExpandQuadsY) : ((ComponentY == 1) ? 0 : 1);
 			const int32 XDest = (ComponentX == 0) ? 0 : ((ComponentX == 1) ? ExpandQuadsX : (ComponentSizeQuads + ExpandQuadsX + 1));
 			const int32 YDest = (ComponentY == 0) ? 0 : ((ComponentY == 1) ? ExpandQuadsY : (ComponentSizeQuads + ExpandQuadsY + 1));
-			const int32 XNum = (ComponentX == 1) ? (ComponentSizeQuads + 1) : ExpandQuadsX;
-			const int32 YNum = (ComponentY == 1) ? (ComponentSizeQuads + 1) : ExpandQuadsY;
+			const int32 XNum = (ComponentX == 1) ? (ComponentSizeQuads + 1) : ExpandQuadsX; // ((ComponentX == 0 && ComponentY == 1) ? ExpandQuadsX + 1 : ExpandQuadsX);
+			const int32 YNum = (ComponentY == 1) ? (ComponentSizeQuads + 1) : ExpandQuadsY; // ((ComponentX == 1 && ComponentY == 0) ? ExpandQuadsY + 1 : ExpandQuadsY);
 			const int32 XBackup = (ComponentX == 2) ? (ComponentSizeQuads + ExpandQuadsX) : ExpandQuadsX;
 			const int32 YBackup = (ComponentY == 2) ? (ComponentSizeQuads + ExpandQuadsY) : ExpandQuadsY;
 			const int32 XBackupNum = (ComponentX == 1) ? (ComponentSizeQuads + 1) : 1;
 			const int32 YBackupNum = (ComponentY == 1) ? (ComponentSizeQuads + 1) : 1;
-
 			
 			ULandscapeComponent* Neighbor = Info->XYtoComponentMap.FindRef(ComponentBase + FIntPoint((ComponentX - 1), (ComponentY - 1)));
 			if (Neighbor)
 			{
+				// Data array for upscaling case
+				TArray<FColor> CompHeightData;
+				TArray<FColor> CompXYOffsetData;
+				int32 NeighborGeometricLOD = FMath::Clamp<int32>(Neighbor->ForcedLOD >= 0 ? Neighbor->ForcedLOD : Neighbor->LODBias, 0, MaxLOD);
 				FLandscapeComponentDataInterface DataInterface(Neighbor, InLOD);
+				::InternalUpscaling(DataInterface, Neighbor, InLOD, NeighborGeometricLOD, CompHeightData, CompXYOffsetData);
 				for (int32 Y = 0; Y < YNum; Y++)
 				{
 					const FColor* const Data = DataInterface.GetHeightData(0, YSource + Y);
diff --git a/Engine/Source/Runtime/Engine/Public/Landscape/LandscapeDataAccess.h b/Engine/Source/Runtime/Engine/Public/Landscape/LandscapeDataAccess.h
--- a/Engine/Source/Runtime/Engine/Public/Landscape/LandscapeDataAccess.h
+++ b/Engine/Source/Runtime/Engine/Public/Landscape/LandscapeDataAccess.h
@@ -123,6 +123,7 @@ struct FLandscapeComponentDataInterface
 	FLandscapeComponentDataInterface(ULandscapeComponent* InComponent, int32 InMipLevel = 0)
 	:	Component(InComponent),
 		HeightMipData(NULL),
+		XYOffsetMipData(NULL),
 		bNeedToDeleteDataInterface(false),
 		MipLevel(InMipLevel)
 	{
@@ -150,11 +151,9 @@ struct FLandscapeComponentDataInterface
 			HeightMipData = (FColor*)DataInterface->LockMip(Component->HeightmapTexture, MipLevel);
 			if (Component->XYOffsetmapTexture)
 			{
-				XYOffsetMipData = (FColor*)DataInterface->LockMip(Component->XYOffsetmapTexture,MipLevel);
+				XYOffsetMipData = (FColor*)DataInterface->LockMip(Component->XYOffsetmapTexture, MipLevel);
 			}
 		}
-
-		XYOffsetMipData = NULL;
 	}
 
 	~FLandscapeComponentDataInterface()
@@ -224,9 +223,24 @@ struct FLandscapeComponentDataInterface
 		return HeightMipData;
 	}
 
+	FColor* GetRawXYOffsetData() const
+	{
+		return XYOffsetMipData;
+	}
+
+	void SetRawHeightData(FColor* NewHeightData)
+	{
+		HeightMipData = NewHeightData;
+	}
+
+	void SetRawXYOffsetData(FColor* NewXYOffsetData)
+	{
+		XYOffsetMipData = NewXYOffsetData;
+	}
+
 	void UnlockRawHeightData() const
 	{
-		DataInterface->UnlockMip(Component->HeightmapTexture,0);
+		DataInterface->UnlockMip(Component->HeightmapTexture, MipLevel);
 	}
 
 	/* Return the raw heightmap data exactly same size for Heightmap texture which belong to only this component */
@@ -464,11 +478,14 @@ private:
 	struct FLandscapeDataInterface* DataInterface;
 	ULandscapeComponent* Component;
 
+public:
 	// offset of this component's data into heightmap texture
 	int32 HeightmapStride;
 	int32 HeightmapComponentOffsetX;
 	int32 HeightmapComponentOffsetY;
 	int32 HeightmapSubsectionOffset;
+
+private:
 	FColor* HeightMipData;
 	FColor* XYOffsetMipData;
 
diff --git a/Engine/Source/Runtime/Engine/Public/Landscape/LandscapeLight.h b/Engine/Source/Runtime/Engine/Public/Landscape/LandscapeLight.h
--- a/Engine/Source/Runtime/Engine/Public/Landscape/LandscapeLight.h
+++ b/Engine/Source/Runtime/Engine/Public/Landscape/LandscapeLight.h
@@ -59,7 +59,7 @@ public:
 #endif	//WITH_EDITOR
 
 protected:
-	void GetHeightmapData(int32 InLOD);
+	void GetHeightmapData(int32 InLOD, int32 GeometryLOD);
 
 	/** Fills in the static lighting vertex data for the Landscape vertex. */
 	void GetStaticLightingVertex(int32 VertexIndex, FStaticLightingVertex& OutVertex) const;
