#TTP 320713 - UE4: BLUEPRINTS: K2: == and != are not exposed for class variables
#Branch UE4
#Proj Engine

#Summary Added the requested class to class operators, exposed them to the blueprints and tested functionality.

diff --git a/Engine/Source/Runtime/Engine/Classes/Kismet/KismetMathLibrary.h b/Engine/Source/Runtime/Engine/Classes/Kismet/KismetMathLibrary.h
--- a/Engine/Source/Runtime/Engine/Classes/Kismet/KismetMathLibrary.h
+++ b/Engine/Source/Runtime/Engine/Classes/Kismet/KismetMathLibrary.h
@@ -784,6 +784,18 @@ class UKismetMathLibrary : public UBlueprintFunctionLibrary
 	UFUNCTION(BlueprintPure, meta=(FriendlyName = "NotEqual (Object)", CompactNodeTitle = "!=", Keywords = "!= not equal"), Category="Utilities")
 	static bool NotEqual_ObjectObject(class UObject* A, class UObject* B);
 
+	//
+	// Class operators and functions.
+	//
+
+	/* Returns true if A and B are equal (A == B) */
+	UFUNCTION(BlueprintPure, meta=(FriendlyName = "Equal (Class)", CompactNodeTitle = "==", Keywords = "== equal"), Category="Utilities")
+	static bool EqualEqual_ClassClass(class UClass* A, class UClass* B);
+
+	/* Returns true if A and B are not equal (A != B) */
+	UFUNCTION(BlueprintPure, meta=(FriendlyName = "NotEqual (Class)", CompactNodeTitle = "!=", Keywords = "!= not equal"), Category="Utilities")
+	static bool NotEqual_ClassClass(class UClass* A, class UClass* B);
+
 	/**
 	 * Determine if a class is a child of another class.
 	 *
diff --git a/Engine/Source/Runtime/Engine/Private/KismetMathLibrary.cpp b/Engine/Source/Runtime/Engine/Private/KismetMathLibrary.cpp
--- a/Engine/Source/Runtime/Engine/Private/KismetMathLibrary.cpp
+++ b/Engine/Source/Runtime/Engine/Private/KismetMathLibrary.cpp
@@ -906,6 +906,16 @@ bool UKismetMathLibrary::NotEqual_ObjectObject(class UObject* A, class UObject*
 	return A != B;
 }
 
+bool UKismetMathLibrary::EqualEqual_ClassClass(class UClass* A, class UClass* B)
+{
+	return A == B;
+}
+
+bool UKismetMathLibrary::NotEqual_ClassClass(class UClass* A, class UClass* B)
+{
+	return A != B;
+}
+
 bool UKismetMathLibrary::ClassIsChildOf(TSubclassOf<class UObject> TestClass, TSubclassOf<class UObject> ParentClass)
 {
 	return ((*ParentClass != NULL) && (*TestClass != NULL)) ? (*TestClass)->IsChildOf(*ParentClass) : false;
