#ttp 331525 - AudioComponent: Attenuation Overrides settings should be disabled when Override Attenuation is false

#ttp 325959 - AUDIO: The Soundcue categories should be renamed and its properties reorganized

#proj UE4.Engine

#summary Override Attenuation now toggles whether two options are available and renamed SoundCue categories

#change Used EditCondition meta data to determine when attenuation options can be changed.

#change Changed SoundCue and SoundBase category names.

diff --git a/Engine/Source/Runtime/Engine/Classes/Components/AudioComponent.h b/Engine/Source/Runtime/Engine/Classes/Components/AudioComponent.h
--- a/Engine/Source/Runtime/Engine/Classes/Components/AudioComponent.h
+++ b/Engine/Source/Runtime/Engine/Classes/Components/AudioComponent.h
@@ -151,10 +151,10 @@ class ENGINE_API UAudioComponent : public USceneComponent
 	UPROPERTY(EditAnywhere, BlueprintReadWrite, Category=Modulation)
 	float HighFrequencyGainMultiplier;
 
-	UPROPERTY(EditAnywhere, BlueprintReadWrite, Category=Attenuation)
+	UPROPERTY(EditAnywhere, BlueprintReadWrite, Category=Attenuation, meta=(EditCondition="!bOverrideAttenuation"))
 	class USoundAttenuation* AttenuationSettings;
 
-	UPROPERTY(EditAnywhere, BlueprintReadWrite, Category=Attenuation)
+	UPROPERTY(EditAnywhere, BlueprintReadWrite, Category=Attenuation, meta=(EditCondition="bOverrideAttenuation"))
 	struct FAttenuationSettings AttenuationOverrides;
 
 	/** while playing, this component will check for occlusion from its closest listener every this many seconds */
diff --git a/Engine/Source/Runtime/Engine/Classes/Sound/SoundBase.h b/Engine/Source/Runtime/Engine/Classes/Sound/SoundBase.h
--- a/Engine/Source/Runtime/Engine/Classes/Sound/SoundBase.h
+++ b/Engine/Source/Runtime/Engine/Classes/Sound/SoundBase.h
@@ -38,16 +38,16 @@ class USoundBase : public UObject
 	USoundClass* SoundClassObject;
 
 	/** For debugging purpose only . */
-	UPROPERTY(EditAnywhere, Category=Sound)
+	UPROPERTY(EditAnywhere, Category=Playback)
 	uint32 bDebug:1;
 
 	/** If we try to play a new version of this sound when at the max concurrent count how should it be resolved. */
-	UPROPERTY(EditAnywhere, Category=Sound)
+	UPROPERTY(EditAnywhere, Category=Playback)
 	TEnumAsByte<EMaxConcurrentResolutionRule::Type> MaxConcurrentResolutionRule;
 
 
 	/** Maximum number of times this sound can be played concurrently. */
-	UPROPERTY(EditAnywhere, Category=Sound)
+	UPROPERTY(EditAnywhere, Category=Playback)
 	int32 MaxConcurrentPlayCount;
 
 	/** Duration of sound in seconds. */
diff --git a/Engine/Source/Runtime/Engine/Classes/Sound/SoundCue.h b/Engine/Source/Runtime/Engine/Classes/Sound/SoundCue.h
--- a/Engine/Source/Runtime/Engine/Classes/Sound/SoundCue.h
+++ b/Engine/Source/Runtime/Engine/Classes/Sound/SoundCue.h
@@ -44,13 +44,13 @@ class USoundCue : public USoundBase
 	UPROPERTY()
 	class USoundNode* FirstNode;
 
-	UPROPERTY(EditAnywhere, Category=SoundCue, AssetRegistrySearchable)
+	UPROPERTY(EditAnywhere, Category=Sound, AssetRegistrySearchable)
 	float VolumeMultiplier;
 
-	UPROPERTY(EditAnywhere, Category=SoundCue, AssetRegistrySearchable)
+	UPROPERTY(EditAnywhere, Category=Sound, AssetRegistrySearchable)
 	float PitchMultiplier;
 
-	UPROPERTY(EditAnywhere, Category=Attenuation)
+	UPROPERTY(EditAnywhere, Category=Attenuation, meta=(EditCondition="bOverrideAttenuation"))
 	FAttenuationSettings AttenuationOverrides;
 
 #if WITH_EDITORONLY_DATA
diff --git a/Engine/Source/Runtime/Engine/Classes/Sound/SoundNodeAttenuation.h b/Engine/Source/Runtime/Engine/Classes/Sound/SoundNodeAttenuation.h
--- a/Engine/Source/Runtime/Engine/Classes/Sound/SoundNodeAttenuation.h
+++ b/Engine/Source/Runtime/Engine/Classes/Sound/SoundNodeAttenuation.h
@@ -12,10 +12,10 @@ class USoundNodeAttenuation : public USoundNode
 {
 	GENERATED_UCLASS_BODY()
 
-	UPROPERTY(EditAnywhere, Category=Attenuation)
+	UPROPERTY(EditAnywhere, Category=Attenuation, meta=(EditCondition="!bOverrideAttenuation"))
 	USoundAttenuation* AttenuationSettings;
 
-	UPROPERTY(EditAnywhere, Category=Attenuation)
+	UPROPERTY(EditAnywhere, Category=Attenuation, meta=(EditCondition="bOverrideAttenuation"))
 	FAttenuationSettings AttenuationOverrides;
 
 	UPROPERTY(EditAnywhere, Category=Attenuation)
