#TTP 328782 - USplineComponent with NULL StaticMesh leads to invisible collision

#proj UE4

#summary Destroy landscape spline collision if mesh is deleted

diff --git a/Engine/Source/Runtime/Engine/Classes/Components/SplineMeshComponent.h b/Engine/Source/Runtime/Engine/Classes/Components/SplineMeshComponent.h
--- a/Engine/Source/Runtime/Engine/Classes/Components/SplineMeshComponent.h
+++ b/Engine/Source/Runtime/Engine/Classes/Components/SplineMeshComponent.h
@@ -147,6 +147,8 @@ class USplineMeshComponent : public UStaticMeshComponent, public IInterface_Coll
 	virtual bool WantsNegXTriMesh() OVERRIDE { return false; }
 	// End Interface_CollisionDataProvider Interface
 
+	// Destroys the body setup, used to clear collision if the mesh goes missing
+	void DestroyBodySetup();
 #if WITH_EDITOR
 	// Builds collision for the spline mesh (if collision is enabled)
 	void RecreateCollision();
diff --git a/Engine/Source/Runtime/Engine/Private/Landscape/LandscapeSplines.cpp b/Engine/Source/Runtime/Engine/Private/Landscape/LandscapeSplines.cpp
--- a/Engine/Source/Runtime/Engine/Private/Landscape/LandscapeSplines.cpp
+++ b/Engine/Source/Runtime/Engine/Private/Landscape/LandscapeSplines.cpp
@@ -2500,10 +2500,18 @@ bool USplineMeshComponent::ContainsPhysicsTriMeshData(bool InUseAllTriData) cons
 void USplineMeshComponent::CreatePhysicsState()
 {
 #if WITH_EDITOR
-	if (StaticMesh != NULL && CachedMeshBodySetupGuid != StaticMesh->BodySetup->BodySetupGuid)
+	// With editor code we can recreate the collision if the mesh changes
+	const FGuid MeshBodySetupGuid = (StaticMesh != NULL ? StaticMesh->BodySetup->BodySetupGuid : FGuid());
+	if (CachedMeshBodySetupGuid != MeshBodySetupGuid)
 	{
 		RecreateCollision();
 	}
+#else
+	// Without editor code we can only destroy the collision if the mesh is missing
+	if (StaticMesh == NULL && BodySetup != NULL)
+	{
+		DestroyBodySetup();
+	}
 #endif
 
 	return Super::CreatePhysicsState();
@@ -2566,6 +2574,18 @@ bool USplineMeshComponent::DoCustomNavigableGeometryExport(struct FNavigableGeom
 	return true;
 }
 
+void USplineMeshComponent::DestroyBodySetup()
+{
+	if (BodySetup != NULL)
+	{
+		BodySetup->MarkPendingKill();
+		BodySetup = NULL;
+#if WITH_EDITORONLY_DATA
+		CachedMeshBodySetupGuid.Invalidate();
+#endif
+	}
+}
+
 #if WITH_EDITOR
 void USplineMeshComponent::RecreateCollision()
 {
@@ -2657,12 +2677,7 @@ void USplineMeshComponent::RecreateCollision()
 	}
 	else
 	{
-		if (BodySetup != NULL)
-		{
-			BodySetup->MarkPendingKill();
-			BodySetup = NULL;
-			CachedMeshBodySetupGuid = FGuid();
-		}
+		DestroyBodySetup();
 	}
 }
 #endif
