#UE4
- TTP#331527, modified Pawn's navigation location to be at its feet rather then at its pelvis

#QAGame
- disabled hierarchical pathfinding for qagame

diff --git a/Engine/Source/Runtime/Engine/Classes/GameFramework/Pawn.h b/Engine/Source/Runtime/Engine/Classes/GameFramework/Pawn.h
--- a/Engine/Source/Runtime/Engine/Classes/GameFramework/Pawn.h
+++ b/Engine/Source/Runtime/Engine/Classes/GameFramework/Pawn.h
@@ -207,7 +207,7 @@ public:
 	virtual const struct FNavAgentProperties* GetNavAgentProperties() const OVERRIDE { return GetMovementComponent() ? GetMovementComponent()->GetNavAgentProperties() : NULL;}
 	/** Basically retrieved pawn's location on navmesh */
 	UFUNCTION(BlueprintCallable, Category="Pawn")
-	virtual FVector GetNavAgentLocation() const OVERRIDE { return GetActorLocation(); }
+	virtual FVector GetNavAgentLocation() const OVERRIDE { return GetActorLocation() - FVector(0.f, 0.f, BaseEyeHeight); }
 	virtual void GetMoveGoalReachTest(class AActor* MovingActor, const FVector& MoveOffset, FVector& GoalOffset, float& GoalRadius, float& GoalHalfHeight) const OVERRIDE;
 	// End INavAgentInterface Interface
 
