#ttp 331537 - EDITOR: SKYLIGHT: Placing a skylight into the world adds dragon head icon to origin
#proj UE4
#branch UE4
#summary ASkyLight now sets up its SpriteComponent with the correct image

diff --git a/Engine/Source/Runtime/Engine/Private/SkyLightComponent.cpp b/Engine/Source/Runtime/Engine/Private/SkyLightComponent.cpp
--- a/Engine/Source/Runtime/Engine/Private/SkyLightComponent.cpp
+++ b/Engine/Source/Runtime/Engine/Private/SkyLightComponent.cpp
@@ -417,23 +417,37 @@ void USkyLightComponent::RecaptureSky()
 ASkyLight::ASkyLight(const class FPostConstructInitializeProperties& PCIP)
 	: Super(PCIP)
 {
-	// Structure to hold one-time initialization
-	struct FConstructorStatics
+	LightComponent = PCIP.CreateDefaultSubobject<USkyLightComponent>(this, TEXT("SkyLightComponent0"));
+	RootComponent = LightComponent;
+
+#if WITH_EDITORONLY_DATA
+	if (!IsRunningCommandlet())
 	{
-		FName ID_Sky;
-		FText NAME_Sky;
+		// Structure to hold one-time initialization
+		struct FConstructorStatics
+		{
+			ConstructorHelpers::FObjectFinderOptional<UTexture2D> SkyLightTextureObject;
+			FName ID_Sky;
+			FText NAME_Sky;
+
+			FConstructorStatics()
+				: SkyLightTextureObject(TEXT("/Engine/EditorResources/LightIcons/SkyLight"))
+				, ID_Sky(TEXT("Sky"))
+				, NAME_Sky(NSLOCTEXT( "SpriteCategory", "Sky", "Sky" ))
+			{
+			}
+		};
+		static FConstructorStatics ConstructorStatics;
 
-		FConstructorStatics()
-			: ID_Sky(TEXT("Sky"))
-			, NAME_Sky(NSLOCTEXT( "SpriteCategory", "Sky", "Sky" ))
+		if (SpriteComponent)
 		{
+			SpriteComponent->Sprite = ConstructorStatics.SkyLightTextureObject.Get();
+			SpriteComponent->SpriteInfo.Category = ConstructorStatics.ID_Sky;
+			SpriteComponent->SpriteInfo.DisplayName = ConstructorStatics.NAME_Sky;
+			SpriteComponent->AttachParent = LightComponent;
 		}
-	};
-	static FConstructorStatics ConstructorStatics;
-
-	LightComponent = PCIP.CreateDefaultSubobject<USkyLightComponent>(this, TEXT("SkyLightComponent0"));
-	RootComponent = LightComponent;
-
+	}
+#endif // WITH_EDITORONLY_DATA
 }
 
 void ASkyLight::GetLifetimeReplicatedProps( TArray< FLifetimeProperty > & OutLifetimeProps ) const
