#TTP 324427 - Please add functionality to save out log files while game/server are running, similar to UE3 Unreal Console functionality

#proj UE4.Editor

#summary Added FlushLog command to the cheat manager

#change Added Exec command for flushing the logs

diff --git a/Engine/Source/Runtime/Engine/Classes/GameFramework/CheatManager.h b/Engine/Source/Runtime/Engine/Classes/GameFramework/CheatManager.h
--- a/Engine/Source/Runtime/Engine/Classes/GameFramework/CheatManager.h
+++ b/Engine/Source/Runtime/Engine/Classes/GameFramework/CheatManager.h
@@ -281,6 +281,10 @@ class ENGINE_API UCheatManager : public UObject
 	UFUNCTION(exec)
 	virtual void BugItStringCreator(FVector ViewLocation, FRotator ViewRotation, FString& GoString, FString& LocString);
 
+	/** This will force a flush of the output log to file*/
+	UFUNCTION(exec)
+	virtual void FlushLog();
+
 	/** Logs the current location in bugit format without taking screenshot and further routing. */
 	UFUNCTION(exec)
 	virtual void LogLoc();
diff --git a/Engine/Source/Runtime/Engine/Private/CheatManager.cpp b/Engine/Source/Runtime/Engine/Private/CheatManager.cpp
--- a/Engine/Source/Runtime/Engine/Private/CheatManager.cpp
+++ b/Engine/Source/Runtime/Engine/Private/CheatManager.cpp
@@ -970,6 +970,12 @@ void UCheatManager::BugItStringCreator( FVector ViewLocation, FRotator ViewRotat
 	UE_LOG(LogCheatManager, Log, TEXT("%s"), *LocString);
 }
 
+void UCheatManager::FlushLog()
+{
+	GLog->FlushThreadedLogs();
+	GLog->Flush();
+}
+
 void UCheatManager::LogLoc()
 {
 	APlayerController* const MyPlayerController = GetOuterAPlayerController();
