Logic fix

diff --git a/Engine/Source/Runtime/Engine/Private/Blueprint.cpp b/Engine/Source/Runtime/Engine/Private/Blueprint.cpp
--- a/Engine/Source/Runtime/Engine/Private/Blueprint.cpp
+++ b/Engine/Source/Runtime/Engine/Private/Blueprint.cpp
@@ -218,7 +218,7 @@ void UBlueprint::Serialize(FArchive& Ar)
 
 			const UClass* ClassObject = Cast<UClass>(VarType.PinSubCategoryObject.Get());
 			// if the object type is an actor...
-			if ((ClassObject == NULL) && ClassObject->IsChildOf(AActor::StaticClass()))
+			if ((ClassObject != NULL) && ClassObject->IsChildOf(AActor::StaticClass()))
 			{
 				// hide the default value field
 				Variable.PropertyFlags |= CPF_DisableEditOnTemplate;
