#UE4 - Fixed logs that were too noisy

diff --git a/Engine/Source/Runtime/Engine/Private/RepLayout.cpp b/Engine/Source/Runtime/Engine/Private/RepLayout.cpp
--- a/Engine/Source/Runtime/Engine/Private/RepLayout.cpp
+++ b/Engine/Source/Runtime/Engine/Private/RepLayout.cpp
@@ -1180,7 +1180,7 @@ bool FRepLayout::ReadProperty(
 		const int32 LocalAbsOffset = AbsOffset + SwappedCmd.Offset;
 		UnmappedGuids->Map.Add( LocalAbsOffset, FUnmappedGuidMgrElement( ReaderState.Bunch.PackageMap->GetLastUnmappedNetGUID(), Cmd.ParentIndex, CmdIndex ) );
 
-		UE_LOG( LogNet, Warning, TEXT( "ADDED unmapped property: Offset: %i, Guid: %s, Name: %s"), LocalAbsOffset, *ReaderState.Bunch.PackageMap->GetLastUnmappedNetGUID().ToString(), *Cmd.Property->GetName() );
+		UE_LOG( LogNet, Verbose, TEXT( "ADDED unmapped property: Offset: %i, Guid: %s, Name: %s"), LocalAbsOffset, *ReaderState.Bunch.PackageMap->GetLastUnmappedNetGUID().ToString(), *Cmd.Property->GetName() );
 		ReaderState.bHasUnmapped = true;
 	}
 #endif
@@ -1410,7 +1410,7 @@ bool FRepLayout::UpdateUnmappedObjects_r( FRepState * RepState, FUnmappedGuidMgr
 		if ( AbsOffset >= MaxAbsOffset )
 		{
 			// Array must have shrunk, we can remove this item
-			//UE_LOG( LogNet, Log, TEXT( "REMOVED unmapped property: AbsOffset >= MaxAbsOffset. Offset: %i" ), AbsOffset );
+			UE_LOG( LogNet, Verbose, TEXT( "REMOVED unmapped property: AbsOffset >= MaxAbsOffset. Offset: %i" ), AbsOffset );
 			It.RemoveCurrent();
 			continue;
 		}
@@ -1435,7 +1435,7 @@ bool FRepLayout::UpdateUnmappedObjects_r( FRepState * RepState, FUnmappedGuidMgr
 
 		if ( Object != NULL )
 		{
-			//UE_LOG( LogNet, Log, TEXT( "REMOVED unmapped property: Offset: %i, Guid: %s, PropName: %s, ObjName: %s" ), AbsOffset, *It.Value().Guid.ToString(), *Cmd.Property->GetName(), *Object->GetName() );
+			UE_LOG( LogNet, Verbose, TEXT( "REMOVED unmapped property: Offset: %i, Guid: %s, PropName: %s, ObjName: %s" ), AbsOffset, *It.Value().Guid.ToString(), *Cmd.Property->GetName(), *Object->GetName() );
 			UObjectPropertyBase * ObjProperty = CastChecked< UObjectPropertyBase>( Cmd.Property );
 
 			UObject * OldObject = ObjProperty->GetObjectPropertyValue( Data + AbsOffset );
@@ -1453,7 +1453,7 @@ bool FRepLayout::UpdateUnmappedObjects_r( FRepState * RepState, FUnmappedGuidMgr
 
 		bHasUnmapped = true;
 
-		//UE_LOG( LogNet, Warning, TEXT( "UnmappedGuids: Offset: %i, Guid: %s" ), AbsOffset, *It.Value().Guid.ToString() );
+		UE_LOG( LogNet, Verbose, TEXT( "UnmappedGuids: Offset: %i, Guid: %s" ), AbsOffset, *It.Value().Guid.ToString() );
 	}
 
 	return bHasUnmapped;
