#ttp 331649 - EDITOR: Simulate: CRASH: Saving changes made to the builder brush during simulation will cause a crash when stopping simulation
#proj UE4
#branch UE4
#summary FActorEditorUtils::IsABuilderBrush now always performs the check in an editor build
#extra It was using UWorld::IsGameWorld() to ignore the check, which was also ignoring the check for PIE worlds. Rather than testing the world type, it now only performs the check in an editor build, and for a game will always return false.

diff --git a/Engine/Source/Runtime/Engine/Private/ActorEditorUtils.cpp b/Engine/Source/Runtime/Engine/Private/ActorEditorUtils.cpp
--- a/Engine/Source/Runtime/Engine/Private/ActorEditorUtils.cpp
+++ b/Engine/Source/Runtime/Engine/Private/ActorEditorUtils.cpp
@@ -7,14 +7,13 @@ namespace FActorEditorUtils
 {
 	bool IsABuilderBrush( const AActor* InActor )
 	{
-		if ( !InActor || !InActor->GetWorld() || InActor->GetWorld()->IsGameWorld() || InActor->HasAnyFlags(RF_ClassDefaultObject))
-		{
-			return false;
-		}
-		else
+#if WITH_EDITOR
+		if ( InActor && InActor->GetWorld() && !InActor->HasAnyFlags(RF_ClassDefaultObject) )
 		{
 			return InActor->GetLevel()->GetBrush() == InActor;
 		}
+#endif
+		return false;
 	}
 
 	bool IsAPreviewActor( const AActor* InActor )
