UE4: Fix for crash in PIE with muliple clients and dedicated server

diff --git a/Engine/Source/Runtime/Engine/Private/World.cpp b/Engine/Source/Runtime/Engine/Private/World.cpp
--- a/Engine/Source/Runtime/Engine/Private/World.cpp
+++ b/Engine/Source/Runtime/Engine/Private/World.cpp
@@ -1855,18 +1855,21 @@ UWorld* UWorld::DuplicateWorldForPIE(const FString& PackageName, UWorld* OwningW
 	PIELevelPackage->PIEInstanceID = WorldContext.PIEInstance;
 #endif
 
-	const FString PlayWorldMapName = ConvertToPIEPackageName(OwningWorld->GetOutermost()->GetName(), WorldContext.PIEInstance);
-
 	// Set up string asset reference fixups
 	TArray<FString> PackageNamesBeingDuplicatedForPIE;
-	PackageNamesBeingDuplicatedForPIE.Add(PlayWorldMapName);
-	for (auto LevelIt = OwningWorld->StreamingLevels.CreateConstIterator(); LevelIt; ++LevelIt)
+	if ( OwningWorld )
 	{
-		ULevelStreaming* StreamingLevel = *LevelIt;
-		if (StreamingLevel)
+		const FString PlayWorldMapName = ConvertToPIEPackageName(OwningWorld->GetOutermost()->GetName(), WorldContext.PIEInstance);
+
+		PackageNamesBeingDuplicatedForPIE.Add(PlayWorldMapName);
+		for (auto LevelIt = OwningWorld->StreamingLevels.CreateConstIterator(); LevelIt; ++LevelIt)
 		{
-			const FString StreamingLevelPIEName = UWorld::ConvertToPIEPackageName(StreamingLevel->PackageName.ToString(), WorldContext.PIEInstance);
-			PackageNamesBeingDuplicatedForPIE.Add(StreamingLevelPIEName);
+			ULevelStreaming* StreamingLevel = *LevelIt;
+			if (StreamingLevel)
+			{
+				const FString StreamingLevelPIEName = UWorld::ConvertToPIEPackageName(StreamingLevel->PackageName.ToString(), WorldContext.PIEInstance);
+				PackageNamesBeingDuplicatedForPIE.Add(StreamingLevelPIEName);
+			}
 		}
 	}
 
