Fixing compiler errors for platforms not the editor.

diff --git a/Engine/Source/Runtime/Engine/Private/EdGraph/EdGraph.cpp b/Engine/Source/Runtime/Engine/Private/EdGraph/EdGraph.cpp
--- a/Engine/Source/Runtime/Engine/Private/EdGraph/EdGraph.cpp
+++ b/Engine/Source/Runtime/Engine/Private/EdGraph/EdGraph.cpp
@@ -1200,8 +1200,10 @@ void UEdGraphSchema::BreakSinglePinLink(UEdGraphPin* SourcePin, UEdGraphPin* Tar
 
 FPinConnectionResponse UEdGraphSchema::MovePinLinks(UEdGraphPin& MoveFromPin, UEdGraphPin& MoveToPin, bool bIsIntermeadiateMove) const
 {
+#if WITH_EDITOR
 	ensureMsg(bIsIntermeadiateMove || !MoveToPin.GetOwningNode()->GetGraph()->HasAnyFlags(RF_Transient),
 		TEXT("When moving to an intermeadiate pin, use FKismetCompilerContext::MovePinLinksToIntermediate() instead of UEdGraphSchema::MovePinLinks()"));
+#endif // #if WITH_EDITOR
 
 	FPinConnectionResponse FinalResponse = FPinConnectionResponse(CONNECT_RESPONSE_MAKE, TEXT(""));
 	// First copy the current set of links
@@ -1231,8 +1233,10 @@ FPinConnectionResponse UEdGraphSchema::MovePinLinks(UEdGraphPin& MoveFromPin, UE
 
 FPinConnectionResponse UEdGraphSchema::CopyPinLinks(UEdGraphPin& CopyFromPin, UEdGraphPin& CopyToPin, bool bIsIntermeadiateCopy) const
 {
+#if WITH_EDITOR
 	ensureMsg(bIsIntermeadiateCopy || !CopyToPin.GetOwningNode()->GetGraph()->HasAnyFlags(RF_Transient),
 		TEXT("When copying to an intermeadiate pin, use FKismetCompilerContext::CopyPinLinksToIntermediate() instead of UEdGraphSchema::CopyPinLinks()"));
+#endif // #if WITH_EDITOR
 
 	FPinConnectionResponse FinalResponse = FPinConnectionResponse(CONNECT_RESPONSE_MAKE, TEXT(""));
 	for (int32 i=0; i<CopyFromPin.LinkedTo.Num(); i++)
