removed obsolete Z parameter in UCanvas::DrawTile

diff --git a/Engine/Source/Runtime/Engine/Classes/Engine/Canvas.h b/Engine/Source/Runtime/Engine/Classes/Engine/Canvas.h
--- a/Engine/Source/Runtime/Engine/Classes/Engine/Canvas.h
+++ b/Engine/Source/Runtime/Engine/Classes/Engine/Canvas.h
@@ -369,7 +369,6 @@ public:
 	 * @param Tex - Texture to draw
 	 * @param X - Position to draw X
 	 * @param Y - Position to draw Y
-	 * @param Z - Position to draw Z
 	 * @param XL - Width of tile
 	 * @param YL - Height of tile
 	 * @param U - Horizontal position of the upper left corner of the portion of the texture to be shown(texels)
@@ -379,7 +378,7 @@ public:
 	 * @param ClipTile - true to clip tile
 	 * @param BlendMode - blending mode of texture
 	 */
-	void DrawTile(UTexture* Tex, float X, float Y, float Z, float XL, float YL, float U, float V, float UL, float VL, EBlendMode BlendMode=BLEND_Translucent);
+	void DrawTile(UTexture* Tex, float X, float Y, float XL, float YL, float U, float V, float UL, float VL, EBlendMode BlendMode=BLEND_Translucent);
 
 	/** calculate the length of a string 
 	 * @param Font - font used
diff --git a/Engine/Source/Runtime/Engine/Private/UserInterface/Canvas.cpp b/Engine/Source/Runtime/Engine/Private/UserInterface/Canvas.cpp
--- a/Engine/Source/Runtime/Engine/Private/UserInterface/Canvas.cpp
+++ b/Engine/Source/Runtime/Engine/Private/UserInterface/Canvas.cpp
@@ -1182,7 +1182,7 @@ ESimpleElementBlendMode FCanvas::BlendToSimpleElementBlend(EBlendMode BlendMode)
 }
 
 
-void UCanvas::DrawTile( UTexture* Tex, float X, float Y, float Z, float XL, float YL, float U, float V, float UL, float VL, EBlendMode BlendMode )
+void UCanvas::DrawTile( UTexture* Tex, float X, float Y, float XL, float YL, float U, float V, float UL, float VL, EBlendMode BlendMode )
 {
 	if ( !Tex ) 
 	{
@@ -1497,9 +1497,9 @@ void UCanvas::DrawScaledIcon(FCanvasIcon Icon, float X, float Y, FVector Scale)
 		{
 			Icon.VL = Icon.Texture->GetSurfaceHeight();
 		}
-		const float Z = 0.1f;
+
 		// and draw the texture
-		DrawTile(Icon.Texture, X, Y, Z, FMath::Abs(Icon.UL) * Scale.X, FMath::Abs(Icon.VL) * Scale.Y, Icon.U, Icon.V, Icon.UL, Icon.VL);
+		DrawTile(Icon.Texture, X, Y, FMath::Abs(Icon.UL) * Scale.X, FMath::Abs(Icon.VL) * Scale.Y, Icon.U, Icon.V, Icon.UL, Icon.VL);
 	}
 }
 
@@ -1522,9 +1522,8 @@ void UCanvas::DrawIcon(FCanvasIcon Icon, float X, float Y, float Scale)
 			Icon.VL = Icon.Texture->GetSurfaceHeight();
 		}
 		
-		const float Z = 0.1f;
 		// and draw the texture
-		DrawTile(Icon.Texture, X, Y, Z, FMath::Abs(Icon.UL) * Scale, FMath::Abs(Icon.VL) * Scale, Icon.U, Icon.V, Icon.UL, Icon.VL );
+		DrawTile(Icon.Texture, X, Y, FMath::Abs(Icon.UL) * Scale, FMath::Abs(Icon.VL) * Scale, Icon.U, Icon.V, Icon.UL, Icon.VL );
 	}
 }
 
