Landscape : Fixed Selection rendering when re-registering

diff --git a/Engine/Source/Runtime/Engine/Private/Landscape/Landscape.cpp b/Engine/Source/Runtime/Engine/Private/Landscape/Landscape.cpp
--- a/Engine/Source/Runtime/Engine/Private/Landscape/Landscape.cpp
+++ b/Engine/Source/Runtime/Engine/Private/Landscape/Landscape.cpp
@@ -1627,6 +1627,21 @@ void ULandscapeInfo::RegisterActorComponent(ULandscapeComponent* Component, bool
 		// Show MapCheck window
 		FMessageLog("MapCheck").Open( EMessageSeverity::Warning );
 	}
+
+	// Update Selected Components/Regions
+#if WITH_EDITOR
+	if (Component->EditToolRenderData != NULL && Component->EditToolRenderData->SelectedType)
+	{
+		if (Component->EditToolRenderData->SelectedType & FLandscapeEditToolRenderData::ST_COMPONENT)
+		{
+			SelectedComponents.Add(Component);
+		}
+		else if (Component->EditToolRenderData->SelectedType & FLandscapeEditToolRenderData::ST_REGION)
+		{
+			SelectedRegionComponents.Add(Component);
+		}
+	}
+#endif
 }
 	
 void ULandscapeInfo::UnregisterActorComponent(ULandscapeComponent* Component)
