Fixed APawn EnableInput.

diff --git a/Engine/Source/Runtime/Engine/Private/Pawn.cpp b/Engine/Source/Runtime/Engine/Private/Pawn.cpp
--- a/Engine/Source/Runtime/Engine/Private/Pawn.cpp
+++ b/Engine/Source/Runtime/Engine/Private/Pawn.cpp
@@ -847,11 +847,11 @@ void APawn::DisableInput(class APlayerController* PlayerController)
 {
 	if (PlayerController == Controller || PlayerController == NULL)
 	{
-		bInputEnabled = true;
+		bInputEnabled = false;
 	}
 	else
 	{
-		UE_LOG(LogPawn, Error, TEXT("EnableInput can only be specified on a Pawn for its Controller"));
+		UE_LOG(LogPawn, Error, TEXT("DisableInput can only be specified on a Pawn for its Controller"));
 	}
 }
 
