Fix incorrect unit scale in vehicle code - wheels stop rolling backwards for no reason

diff --git a/Engine/Source/Runtime/Engine/Private/Vehicles/WheeledVehicleMovementComponent.cpp b/Engine/Source/Runtime/Engine/Private/Vehicles/WheeledVehicleMovementComponent.cpp
--- a/Engine/Source/Runtime/Engine/Private/Vehicles/WheeledVehicleMovementComponent.cpp
+++ b/Engine/Source/Runtime/Engine/Private/Vehicles/WheeledVehicleMovementComponent.cpp
@@ -253,7 +253,7 @@ void UWheeledVehicleMovementComponent::SetupWheels( PxVehicleWheelsSimData* PWhe
 
 	// Control substepping
 	PWheelsSimData->setSubStepCount( 5.f * LengthScale, 5, 2 ); // forcing constant substeps until all the kinks are ironed out	
-	PWheelsSimData->setMinLongSlipDenominator(4.f * LengthScale * LengthScale); 
+	PWheelsSimData->setMinLongSlipDenominator(4.f * LengthScale); 
 
 	// Prealloc data for the sprung masses
 	PxVec3 WheelOffsets[32];
