If Destructible Component has been destroyed do not propagate events to it

diff --git a/Engine/Source/Runtime/Engine/Private/PhysicsEngine/PhysXSupport.cpp b/Engine/Source/Runtime/Engine/Private/PhysicsEngine/PhysXSupport.cpp
--- a/Engine/Source/Runtime/Engine/Private/PhysicsEngine/PhysXSupport.cpp
+++ b/Engine/Source/Runtime/Engine/Private/PhysicsEngine/PhysXSupport.cpp
@@ -708,6 +708,11 @@ void FApexChunkReport::onDamageNotify(const NxApexDamageEventReportData& damageE
 	UDestructibleComponent* DestructibleComponent = Cast<UDestructibleComponent>(FPhysxUserData::Get<UPrimitiveComponent>(damageEvent.destructible->userData));
 	check(DestructibleComponent);
 
+	if (DestructibleComponent->IsPendingKill())	//don't notify if object is being destroyed
+	{
+		return;
+	}
+
 #if WITH_SUBSTEPPING
 	if (UPhysicsSettings::Get()->bSubstepping)
 	{
