Fix vehicle draw debug format

diff --git a/Engine/Source/Runtime/Engine/Private/Vehicles/WheeledVehicleMovementComponent.cpp b/Engine/Source/Runtime/Engine/Private/Vehicles/WheeledVehicleMovementComponent.cpp
--- a/Engine/Source/Runtime/Engine/Private/Vehicles/WheeledVehicleMovementComponent.cpp
+++ b/Engine/Source/Runtime/Engine/Private/Vehicles/WheeledVehicleMovementComponent.cpp
@@ -946,7 +946,7 @@ bool UWheeledVehicleMovementComponent::GetUseAutoGears() const
 
 #if WITH_PHYSX
 
-void DrawTelemetryGraph( uint32 Channel, const PxVehicleGraph& PGraph, UCanvas* Canvas, float GraphX, float GraphY, float GraphWidth, float GraphHeight )
+void DrawTelemetryGraph( uint32 Channel, const PxVehicleGraph& PGraph, UCanvas* Canvas, float GraphX, float GraphY, float GraphWidth, float GraphHeight, float & OutX )
 {
 	PxF32 PGraphXY[2*PxVehicleGraph::eMAX_NB_SAMPLES];
 	PxVec3 PGraphColor[PxVehicleGraph::eMAX_NB_SAMPLES];
@@ -966,6 +966,7 @@ void DrawTelemetryGraph( uint32 Channel, const PxVehicleGraph& PGraph, UCanvas*
 	float LineGraphY = GraphY + YL + 4.0f;
 
 	FCanvasTileItem TileItem( FVector2D(GraphX, LineGraphY), GWhiteTexture, FVector2D( GraphWidth, GraphWidth ), FLinearColor( 0.0f, 0.125f, 0.0f, 0.5f ) );
+	TileItem.BlendMode = SE_BLEND_Translucent;
 	Canvas->DrawItem( TileItem );
 	
 	Canvas->SetDrawColor( FColor( 0, 32, 0, 128 ) );	
@@ -985,6 +986,8 @@ void DrawTelemetryGraph( uint32 Channel, const PxVehicleGraph& PGraph, UCanvas*
 		LineItem.SetColor( FLinearColor( 1.0f, 0.5f, 0.0f, 1.0f ) );
 		LineItem.Draw( Canvas->Canvas );
 	}
+
+	OutX = FMath::Max(XL,GraphWidth);
 }
 
 void UWheeledVehicleMovementComponent::DrawDebug( UCanvas* Canvas, float& YL, float& YPos )
@@ -1007,6 +1010,7 @@ void UWheeledVehicleMovementComponent::DrawDebug( UCanvas* Canvas, float& YL, fl
 		YPos += Canvas->DrawText( RenderFont, FString::Printf( TEXT("Steering: %f"), SteeringInput ), 4, YPos  );
 		YPos += Canvas->DrawText( RenderFont, FString::Printf( TEXT("Throttle: %f"), ThrottleInput ), 4, YPos );
 		YPos += Canvas->DrawText( RenderFont, FString::Printf( TEXT("Brake: %f"), BrakeInput ), 4, YPos  );
+		YPos += Canvas->DrawText(RenderFont, FString::Printf(TEXT("Gear: %d"), GetCurrentGear()), 4, YPos);
 
 	}
 
@@ -1071,11 +1075,12 @@ void UWheeledVehicleMovementComponent::DrawDebug( UCanvas* Canvas, float& YL, fl
 			float CurX = 4;
 			for ( uint32 i = 0; i < ARRAY_COUNT(GraphChannels); ++i )
 			{
-				DrawTelemetryGraph( GraphChannels[i], TelemetryData->getWheelGraph(w), Canvas, CurX, YPos, GraphWidth, GraphHeight );
-				CurX += GraphWidth + 10.0f;
+				float OutX = GraphWidth;
+				DrawTelemetryGraph( GraphChannels[i], TelemetryData->getWheelGraph(w), Canvas, CurX, YPos, GraphWidth, GraphHeight, OutX );
+				CurX += OutX + 10.f;
 			}
 
-			YPos += GraphHeight;
+			YPos += GraphHeight + 10.f;
 			YPos += YL;
 		}
 	}
