Prevent capturing chord from getting stuck by missing the release event
TTP #331635

diff --git a/Engine/Source/Runtime/Engine/Private/UserInterface/PlayerInput.cpp b/Engine/Source/Runtime/Engine/Private/UserInterface/PlayerInput.cpp
--- a/Engine/Source/Runtime/Engine/Private/UserInterface/PlayerInput.cpp
+++ b/Engine/Source/Runtime/Engine/Private/UserInterface/PlayerInput.cpp
@@ -829,7 +829,7 @@ void UPlayerInput::ProcessInputStack(const TArray<UInputComponent*>& InputCompon
 				if (FoundChord.SourceAction && FoundChord.SourceAction->bPaired)
 				{
 					FActionKeyDetails& KeyDetails = ActionKeyMap.FindChecked(FoundChord.SourceAction->ActionName);
-					if (!KeyDetails.CapturingChord.Key.IsValid() || KeyDetails.CapturingChord == FoundChord.Chord)
+					if (!KeyDetails.CapturingChord.Key.IsValid() || KeyDetails.CapturingChord == FoundChord.Chord || !IsPressed(KeyDetails.CapturingChord.Key))
 					{
 						if (FoundChord.SourceAction->KeyEvent == IE_Pressed)
 						{
