More cleaning up vehicle debug drawing

diff --git a/Engine/Source/Runtime/Engine/Private/Vehicles/WheeledVehicleMovementComponent.cpp b/Engine/Source/Runtime/Engine/Private/Vehicles/WheeledVehicleMovementComponent.cpp
--- a/Engine/Source/Runtime/Engine/Private/Vehicles/WheeledVehicleMovementComponent.cpp
+++ b/Engine/Source/Runtime/Engine/Private/Vehicles/WheeledVehicleMovementComponent.cpp
@@ -955,7 +955,7 @@ void DrawTelemetryGraph( uint32 Channel, const PxVehicleGraph& PGraph, UCanvas*
 	PGraph.computeGraphChannel( Channel, PGraphXY, PGraphColor, PGraphTitle );
 
 	FString Label = ANSI_TO_TCHAR(PGraphTitle);
-	Canvas->SetDrawColor( FColor( 128, 255, 0 ) );
+	Canvas->SetDrawColor( FColor( 255, 255, 0 ) );
 	UFont* Font = GEngine->GetSmallFont();
 	Canvas->DrawText( Font, Label, GraphX, GraphY );
 
@@ -965,7 +965,7 @@ void DrawTelemetryGraph( uint32 Channel, const PxVehicleGraph& PGraph, UCanvas*
 	float LineGraphHeight = GraphHeight - YL - 4.0f;
 	float LineGraphY = GraphY + YL + 4.0f;
 
-	FCanvasTileItem TileItem( FVector2D(GraphX, LineGraphY), GWhiteTexture, FVector2D( GraphWidth, GraphWidth ), FLinearColor( 0.0f, 0.125f, 0.0f, 0.5f ) );
+	FCanvasTileItem TileItem( FVector2D(GraphX, LineGraphY), GWhiteTexture, FVector2D( GraphWidth, GraphWidth ), FLinearColor( 0.0f, 0.125f, 0.0f, 0.25f ) );
 	TileItem.BlendMode = SE_BLEND_Translucent;
 	Canvas->DrawItem( TileItem );
 	
@@ -1036,19 +1036,19 @@ void UWheeledVehicleMovementComponent::DrawDebug( UCanvas* Canvas, float& YL, fl
 		
 		Canvas->DrawText( RenderFont, FString::Printf( TEXT("LatSlip: %.3f"), LatSlip ), YL * 4 , YPos );
 		Canvas->DrawText( RenderFont, FString::Printf( TEXT("LongSlip: %.3f"), LongSlip ), YL * 12, YPos );
-		Canvas->DrawText( RenderFont, FString::Printf( TEXT("Speed: %d"), (int32)WheelSpeed ), YL * 20, YPos );
-		Canvas->DrawText( RenderFont, FString::Printf( TEXT("Contact Surface: %s"), *ContactSurfaceString ), YL * 60, YPos );
+		Canvas->DrawText( RenderFont, FString::Printf( TEXT("Speed: %d"), (int32)WheelSpeed ), YL * 22, YPos );
+		Canvas->DrawText( RenderFont, FString::Printf( TEXT("Contact Surface: %s"), *ContactSurfaceString ), YL * 74, YPos );
 		if( (int32)w < Wheels.Num() )
 		{
 			UVehicleWheel* Wheel = Wheels[w];
-			Canvas->DrawText( RenderFont, FString::Printf( TEXT("Load: %f"), (int32)Wheel->DebugNormalizedTireLoad ), YL * 28, YPos );
-			Canvas->DrawText( RenderFont, FString::Printf( TEXT("Torque: %d"), (int32)Wheel->DebugWheelTorque ), YL * 36, YPos );
-			Canvas->DrawText( RenderFont, FString::Printf( TEXT("Long Force: %d"), (int32)Wheel->DebugLongForce ),YL * 44, YPos );
-			Canvas->DrawText( RenderFont, FString::Printf( TEXT("Lat Force: %d"), (int32)Wheel->DebugLatForce ), YL * 52, YPos );
+			Canvas->DrawText( RenderFont, FString::Printf( TEXT("Load: %.3f"), Wheel->DebugNormalizedTireLoad ), YL * 30, YPos );
+			Canvas->DrawText( RenderFont, FString::Printf( TEXT("Torque: %d"), (int32)Wheel->DebugWheelTorque ), YL * 40, YPos );
+			Canvas->DrawText( RenderFont, FString::Printf( TEXT("Long Force: %d"), (int32)Wheel->DebugLongForce ),YL * 50, YPos );
+			Canvas->DrawText( RenderFont, FString::Printf( TEXT("Lat Force: %d"), (int32)Wheel->DebugLatForce ), YL * 62, YPos );
 		}
 		else
 		{
-			Canvas->DrawText( RenderFont, TEXT("Wheels array insufficiently sized!"), YL * 72, YPos );
+			Canvas->DrawText( RenderFont, TEXT("Wheels array insufficiently sized!"), YL * 50, YPos );
 		}
 
 		YPos += YL;
@@ -1062,12 +1062,14 @@ void UWheeledVehicleMovementComponent::DrawDebug( UCanvas* Canvas, float& YL, fl
 		const float GraphWidth(100.0f), GraphHeight(100.0f);
 	
 			int GraphChannels[] = {
-			PxVehicleWheelGraphChannel::eNORMALIZED_TIRELOAD,
 			PxVehicleWheelGraphChannel::eWHEEL_OMEGA,
+			PxVehicleWheelGraphChannel::eSUSPFORCE,
 			PxVehicleWheelGraphChannel::eTIRE_LONG_SLIP,
 			PxVehicleWheelGraphChannel::eNORM_TIRE_LONG_FORCE,
 			PxVehicleWheelGraphChannel::eTIRE_LAT_SLIP,
-			PxVehicleWheelGraphChannel::eNORM_TIRE_LAT_FORCE
+			PxVehicleWheelGraphChannel::eNORM_TIRE_LAT_FORCE,
+			PxVehicleWheelGraphChannel::eNORMALIZED_TIRELOAD,
+			PxVehicleWheelGraphChannel::eTIRE_FRICTION			
 		};
 
 		for ( uint32 w = 0; w < PVehicle->mWheelsSimData.getNbWheels(); ++w )
