331837	BLOCKER: CRITICAL: Regression: EDITOR: CRASH: Opening AnimMontage

diff --git a/Engine/Source/Runtime/Engine/Private/Animation/AnimSingleNodeInstance.cpp b/Engine/Source/Runtime/Engine/Private/Animation/AnimSingleNodeInstance.cpp
--- a/Engine/Source/Runtime/Engine/Private/Animation/AnimSingleNodeInstance.cpp
+++ b/Engine/Source/Runtime/Engine/Private/Animation/AnimSingleNodeInstance.cpp
@@ -199,10 +199,9 @@ void UAnimSingleNodeInstance::NativeUpdateAnimation(float DeltaTimeX)
 	if(CurrentAsset != NULL && bPlaying)
 	{
 		FAnimGroupInstance* SyncGroup;
-		FAnimTickRecord& TickRecord = CreateUninitializedTickRecord(INDEX_NONE, /*out*/ SyncGroup);
-
 		if (UBlendSpaceBase* BlendSpace = Cast<UBlendSpaceBase>(CurrentAsset))
 		{
+			FAnimTickRecord& TickRecord = CreateUninitializedTickRecord(INDEX_NONE, /*out*/ SyncGroup);
 			MakeBlendSpaceTickRecord(TickRecord, BlendSpace, BlendSpaceInput, BlendSampleData, BlendFilter, bLooping, PlayRate, 1.f, /*inout*/ CurrentTime);
 #if WITH_EDITORONLY_DATA
 			PreviewBasePose = BlendSpace->PreviewBasePose;
@@ -210,6 +209,7 @@ void UAnimSingleNodeInstance::NativeUpdateAnimation(float DeltaTimeX)
 		}
 		else if (UAnimSequence* Sequence = Cast<UAnimSequence>(CurrentAsset))
 		{
+			FAnimTickRecord& TickRecord = CreateUninitializedTickRecord(INDEX_NONE, /*out*/ SyncGroup);
 			MakeSequenceTickRecord(TickRecord, Sequence, bLooping, PlayRate, 1.f, /*inout*/ CurrentTime);
 			// if it's not looping, just set play to be false when reached to end
 			if (!bLooping)
@@ -222,6 +222,7 @@ void UAnimSingleNodeInstance::NativeUpdateAnimation(float DeltaTimeX)
 		}
 		else if(UAnimComposite* Composite = Cast<UAnimComposite>(CurrentAsset))
 		{
+			FAnimTickRecord& TickRecord = CreateUninitializedTickRecord(INDEX_NONE, /*out*/ SyncGroup);
 			MakeSequenceTickRecord(TickRecord, Composite, bLooping, PlayRate, 1.f, /*inout*/ CurrentTime);
 			// if it's not looping, just set play to be false when reached to end
 			if (!bLooping)
