Properly handle double clicking while paused
TTP #331862

diff --git a/Engine/Source/Runtime/Engine/Private/UserInterface/PlayerInput.cpp b/Engine/Source/Runtime/Engine/Private/UserInterface/PlayerInput.cpp
--- a/Engine/Source/Runtime/Engine/Private/UserInterface/PlayerInput.cpp
+++ b/Engine/Source/Runtime/Engine/Private/UserInterface/PlayerInput.cpp
@@ -99,13 +99,14 @@ bool UPlayerInput::InputKey( FKey Key, EInputEvent Event, float AmountDepressed,
 		{
 			// check for doubleclick
 			// note, a tripleclick will currently count as a 2nd double click.
-			if ( World->TimeSince(KeyState.LastUpDownTransitionTime) < GetDefault<UInputSettings>()->DoubleClickTime )
+			const float WorldRealTimeSeconds = World->GetRealTimeSeconds();
+			if ((WorldRealTimeSeconds - KeyState.LastUpDownTransitionTime) < GetDefault<UInputSettings>()->DoubleClickTime)
 			{
 				KeyState.EventAccumulator[IE_DoubleClick].Add(++EventCount);
 			}
 
 			// just went down
-			KeyState.LastUpDownTransitionTime = World->GetTimeSeconds();
+			KeyState.LastUpDownTransitionTime = WorldRealTimeSeconds;
 		}
 		break;
 	case IE_Released:
