Fix vehicle gear change and bad brake behavior

diff --git a/Engine/Source/Runtime/Engine/Classes/Vehicles/WheeledVehicleMovementComponent.h b/Engine/Source/Runtime/Engine/Classes/Vehicles/WheeledVehicleMovementComponent.h
--- a/Engine/Source/Runtime/Engine/Classes/Vehicles/WheeledVehicleMovementComponent.h
+++ b/Engine/Source/Runtime/Engine/Classes/Vehicles/WheeledVehicleMovementComponent.h
@@ -385,10 +385,14 @@ protected:
 	UPROPERTY(EditAnywhere, Category=VehicleInput)
 	float IdleBrakeInput;
 
-	// Auto-brake when absolute vehicle forward speed is less than this
+	// Auto-brake when absolute vehicle forward speed is less than this (cm/s)
 	UPROPERTY(EditAnywhere, Category=VehicleInput)
 	float StopThreshold;
 
+	// Auto-brake when vehicle forward speed is opposite of player input by at least this much (cm/s)
+	UPROPERTY(EditAnywhere, Category = VehicleInput)
+	float WrongDirectionThreshold;
+
 	// Rate at which input throttle can rise and fall
 	UPROPERTY(EditAnywhere, Category=VehicleInput, AdvancedDisplay)
 	FVehicleInputRate ThrottleInputRate;
diff --git a/Engine/Source/Runtime/Engine/Private/Vehicles/WheeledVehicleMovementComponent.cpp b/Engine/Source/Runtime/Engine/Private/Vehicles/WheeledVehicleMovementComponent.cpp
--- a/Engine/Source/Runtime/Engine/Private/Vehicles/WheeledVehicleMovementComponent.cpp
+++ b/Engine/Source/Runtime/Engine/Private/Vehicles/WheeledVehicleMovementComponent.cpp
@@ -94,8 +94,9 @@ UWheeledVehicleMovementComponent::UWheeledVehicleMovementComponent(const class F
 	MinNormalizedTireLoad = 0.0f;
 	MaxNormalizedTireLoad = 10.0f;
 	
-	IdleBrakeInput = 0.05f;
+	IdleBrakeInput = 0.0f;
 	StopThreshold = 10.0f; 
+	WrongDirectionThreshold = 100.f;
 	ThrottleInputRate.RiseRate = 6.0f;
 	ThrottleInputRate.FallRate = 10.0f;
 	BrakeInputRate.RiseRate = 6.0f;
@@ -666,6 +667,16 @@ void UWheeledVehicleMovementComponent::UpdateState( float DeltaTime )
 	APawn* MyOwner = UpdatedComponent ? Cast<APawn>(UpdatedComponent->GetOwner()) : NULL;
 	if (MyOwner && MyOwner->IsLocallyControlled())
 	{
+		//Manual shifting between reverse and first gear
+		if (RawThrottleInput < 0.f && GetCurrentGear() >= 0)
+		{
+			SetTargetGear(-1, true);
+		}
+		else if (RawThrottleInput > 0.f && GetCurrentGear() <= 0)
+		{
+			SetTargetGear(1, true);
+		}
+
 		SteeringInput = SteeringInputRate.InterpInputValue( DeltaTime, SteeringInput, CalcSteeringInput() );
 		ThrottleInput = ThrottleInputRate.InterpInputValue( DeltaTime, ThrottleInput, CalcThrottleInput() );
 		BrakeInput = BrakeInputRate.InterpInputValue( DeltaTime, BrakeInput, CalcBrakeInput() );
@@ -722,38 +733,32 @@ float UWheeledVehicleMovementComponent::CalcBrakeInput()
 	float NewBrakeInput = 0.0f;
 
 	// if player wants to move forwards...
-	if ( RawThrottleInput > 0.1f )  // MSS expanded hardcoded dead zone from 0.01f to filter weird throttle noise.
+	if ( RawThrottleInput > 0.f )
 	{
 		// if vehicle is moving backwards, then press brake
-		if ( ForwardSpeed < -StopThreshold )
+		if ( ForwardSpeed < -WrongDirectionThreshold)
 		{
 			NewBrakeInput = 1.0f;			
-			SetTargetGear(1, true);
 		}
 
 	}
 
 	// if player wants to move backwards...
-	else if ( RawThrottleInput < -0.1f )
+	else if ( RawThrottleInput < 0.f )
 	{
 		// if vehicle is moving forwards, then press brake
-		if ( ForwardSpeed > StopThreshold )
+		if (ForwardSpeed > WrongDirectionThreshold)
 		{
 			NewBrakeInput = 1.0f;			// Seems a bit severe to have 0 or 1 braking. Better control can be had by allowing continuous brake input values
-			// set to reverse
-			SetTargetGear(-1, true);	
-		}
-
-		
+		}	
 	}
 
 	// if player isn't pressing forward or backwards...
 	else
 	{
-		// If almost stationary, stick brakes on
-		if ( FMath::Abs(ForwardSpeed) <  StopThreshold )
+		if (ForwardSpeed < StopThreshold && ForwardSpeed > -StopThreshold)	//auto break 
 		{
-			NewBrakeInput = 0.1f + 0.1f * FMath::Abs(ForwardSpeed);
+			NewBrakeInput = 1.f;
 		}
 		else
 		{
