#UE4:
* fixed Console using scaled view rectangle where it should use the viewsized one.
  This was visible when using ScreenPercentage.

diff --git a/Engine/Source/Runtime/Engine/Private/Debug/DebugDrawService.cpp b/Engine/Source/Runtime/Engine/Private/Debug/DebugDrawService.cpp
--- a/Engine/Source/Runtime/Engine/Private/Debug/DebugDrawService.cpp
+++ b/Engine/Source/Runtime/Engine/Private/Debug/DebugDrawService.cpp
@@ -57,7 +57,7 @@ void UDebugDrawService::Draw(const FEngineShowFlags Flags, FViewport* Viewport,
 		CanvasObject->AddToRoot();
 	}
 	
-	CanvasObject->Init(View->ViewRect.Width(), View->ViewRect.Height(), View);
+	CanvasObject->Init(View->UnscaledViewRect.Width(), View->UnscaledViewRect.Height(), View);
 	CanvasObject->Update();	
 	CanvasObject->Canvas = Canvas;
 
