Fixed skylight on bsp

diff --git a/Engine/Source/Runtime/Engine/Private/ModelLight.cpp b/Engine/Source/Runtime/Engine/Private/ModelLight.cpp
--- a/Engine/Source/Runtime/Engine/Private/ModelLight.cpp
+++ b/Engine/Source/Runtime/Engine/Private/ModelLight.cpp
@@ -935,6 +935,11 @@ void UModel::ApplyStaticLighting()
 							DestSample.Coefficients[CoefficientIndex][ColorIndex] = (uint8)FMath::Clamp<int32>( FMath::Round( FMath::Pow( Repacked, 1.0f / Exponent ) * 255.0f ), 0, 255 );
 						}
 					}
+
+					for(int32 ColorIndex = 0;ColorIndex < 4;ColorIndex++)
+					{
+						DestSample.SkyOcclusion[ColorIndex] = SourceSample.SkyOcclusion[ColorIndex];
+					}
 				}
 			}
 
