- Excluding bHiddenInGame components from actor bounds calculations, so editor icons do not affect actor bounding box

diff --git a/Engine/Source/Runtime/Engine/Private/Actor.cpp b/Engine/Source/Runtime/Engine/Private/Actor.cpp
--- a/Engine/Source/Runtime/Engine/Private/Actor.cpp
+++ b/Engine/Source/Runtime/Engine/Private/Actor.cpp
@@ -848,7 +848,7 @@ FBox AActor::GetComponentsBoundingBox(bool bNonColliding) const
 		UPrimitiveComponent* PrimComp = Components[CompIdx];
 
 		// Only use collidable components to find collision bounding box.
-		if( PrimComp->IsRegistered() && (bNonColliding || PrimComp->IsCollisionEnabled()) )
+		if( PrimComp->IsRegistered() && ((bNonColliding && !PrimComp->bHiddenInGame) || PrimComp->IsCollisionEnabled()) )
 		{
 			Box += PrimComp->Bounds.GetBox();
 		}
