#TTP 221734 - Foliage: Navigation mesh does not get generated correctly for instanced static meshes

#proj UE4

#summary Implemented nav geometry export for instanced static meshes

diff --git a/Engine/Source/Runtime/Engine/Classes/Components/InstancedStaticMeshComponent.h b/Engine/Source/Runtime/Engine/Classes/Components/InstancedStaticMeshComponent.h
--- a/Engine/Source/Runtime/Engine/Classes/Components/InstancedStaticMeshComponent.h
+++ b/Engine/Source/Runtime/Engine/Classes/Components/InstancedStaticMeshComponent.h
@@ -135,6 +135,8 @@ public:
 #endif
 	virtual void GetLightAndShadowMapMemoryUsage( int32& LightMapMemoryUsage, int32& ShadowMapMemoryUsage ) const OVERRIDE;
 	virtual void ApplyWorldOffset(const FVector& InOffset, bool bWorldShift) OVERRIDE;
+
+	virtual bool DoCustomNavigableGeometryExport(struct FNavigableGeometryExport* GeomExport) const OVERRIDE;
 	// End UPrimitiveComponent Interface
 
 	//Begin UObject Interface
diff --git a/Engine/Source/Runtime/Engine/Private/AI/Navigation/RecastNavMeshGenerator.cpp b/Engine/Source/Runtime/Engine/Private/AI/Navigation/RecastNavMeshGenerator.cpp
--- a/Engine/Source/Runtime/Engine/Private/AI/Navigation/RecastNavMeshGenerator.cpp
+++ b/Engine/Source/Runtime/Engine/Private/AI/Navigation/RecastNavMeshGenerator.cpp
@@ -267,6 +267,7 @@ struct FRecastGeometryExport : public FNavigableGeometryExport
 	virtual void ExportPxHeightField(physx::PxHeightField const * const HeightField, const FTransform& LocalToWorld) OVERRIDE;
 #endif // WITH_PHYSX
 	virtual void ExportCustomMesh(const FVector* InVertices, int32 NumVerts, const int32* InIndices, int32 NumIndices, const FTransform& LocalToWorld) OVERRIDE;
+	virtual void ExportRigidBodySetup(UBodySetup& BodySetup, const FTransform& LocalToWorld) OVERRIDE;
 	virtual void AddNavModifiers(const FCompositeNavModifier& Modifiers) OVERRIDE;
 };
 
@@ -903,13 +904,14 @@ FORCEINLINE_DEBUGGABLE void ExportRigidBodySetup(UBodySetup& BodySetup, TNavStat
 	BodySetup.CreatePhysicsMeshes();
 
 	static TNavStatArray<int32> TemporaryShapeBuffer;
-	TemporaryShapeBuffer.Reset();
 
-	ExportRigidBodyTriMesh(BodySetup, VertexBuffer, IndexBuffer, UnrealBounds, LocalToWorld);		
+	ExportRigidBodyTriMesh(BodySetup, VertexBuffer, IndexBuffer, UnrealBounds, LocalToWorld);
 	ExportRigidBodyConvexElements(BodySetup, VertexBuffer, IndexBuffer, TemporaryShapeBuffer, UnrealBounds, LocalToWorld);
 	ExportRigidBodyBoxElements(BodySetup, VertexBuffer, IndexBuffer, TemporaryShapeBuffer, UnrealBounds, LocalToWorld);
 	ExportRigidBodySphylElements(BodySetup, VertexBuffer, IndexBuffer, TemporaryShapeBuffer, UnrealBounds, LocalToWorld);
 	ExportRigidBodySphereElements(BodySetup, VertexBuffer, IndexBuffer, TemporaryShapeBuffer, UnrealBounds, LocalToWorld);
+
+	TemporaryShapeBuffer.Reset();
 }
 
 FORCEINLINE_DEBUGGABLE void ExportComponent(UActorComponent& Component, FRecastGeometryExport* GeomExport, const FBox* ClipBounds=NULL)
@@ -1061,6 +1063,11 @@ void FRecastGeometryExport::ExportCustomMesh(const FVector* InVertices, int32 Nu
 	RecastGeometryExport::ExportCustomMesh(InVertices, NumVerts, InIndices, NumIndices, LocalToWorld, VertexBuffer, IndexBuffer, Data->Bounds);
 }
 
+void FRecastGeometryExport::ExportRigidBodySetup(UBodySetup& BodySetup, const FTransform& LocalToWorld)
+{
+	RecastGeometryExport::ExportRigidBodySetup(BodySetup, VertexBuffer, IndexBuffer, Data->Bounds, LocalToWorld);
+}
+
 void FRecastGeometryExport::AddNavModifiers(const FCompositeNavModifier& Modifiers)
 {
 	Data->Modifiers.Add(Modifiers);
diff --git a/Engine/Source/Runtime/Engine/Private/InstancedStaticMesh.cpp b/Engine/Source/Runtime/Engine/Private/InstancedStaticMesh.cpp
--- a/Engine/Source/Runtime/Engine/Private/InstancedStaticMesh.cpp
+++ b/Engine/Source/Runtime/Engine/Private/InstancedStaticMesh.cpp
@@ -1404,6 +1404,52 @@ void UInstancedStaticMeshComponent::ApplyWorldOffset(const FVector& InOffset, bo
 	MarkRenderStateDirty();
 }
 
+#include "AI/Navigation/RecastHelpers.h"
+bool UInstancedStaticMeshComponent::DoCustomNavigableGeometryExport(struct FNavigableGeometryExport* GeomExport) const
+{
+	if (StaticMesh != NULL)
+	{
+		UNavCollision* NavCollision = StaticMesh->NavCollision;
+		if (NavCollision != NULL && NavCollision->bHasConvexGeometry)
+		{
+			for (const FInstancedStaticMeshInstanceData& InstanceData : PerInstanceSMData)
+			{
+				const FVector Scale3D = InstanceData.Transform.GetScaleVector();
+				// if any of scales is 0 there's no point in exporting it
+				if (!Scale3D.IsZero())
+				{
+					GeomExport->ExportCustomMesh(NavCollision->ConvexCollision.VertexBuffer.GetData(), NavCollision->ConvexCollision.VertexBuffer.Num(),
+						NavCollision->ConvexCollision.IndexBuffer.GetData(), NavCollision->ConvexCollision.IndexBuffer.Num(), FTransform(InstanceData.Transform));
+
+					GeomExport->ExportCustomMesh(NavCollision->TriMeshCollision.VertexBuffer.GetData(), NavCollision->TriMeshCollision.VertexBuffer.Num(),
+						NavCollision->TriMeshCollision.IndexBuffer.GetData(), NavCollision->TriMeshCollision.IndexBuffer.Num(), FTransform(InstanceData.Transform));
+				}
+			}
+		}
+		else
+		{
+			UBodySetup* BodySetup = StaticMesh->BodySetup;
+			if (BodySetup)
+			{
+				for (const FInstancedStaticMeshInstanceData& InstanceData : PerInstanceSMData)
+				{
+					const FVector Scale3D = InstanceData.Transform.GetScaleVector();
+					// if any of scales is 0 there's no point in exporting it
+					if (!Scale3D.IsZero())
+					{
+						GeomExport->ExportRigidBodySetup(*BodySetup, FTransform(InstanceData.Transform));
+					}
+				}
+
+				//GeomExport->SlopeOverride = BodySetup->WalkableSlopeOverride;
+			}
+		}
+	}
+
+	// we don't want "regular" collision export for this component
+	return false;
+}
+
 SIZE_T UInstancedStaticMeshComponent::GetResourceSize( EResourceSizeMode::Type Mode )
 {
 	SIZE_T ResSize = 0;
diff --git a/Engine/Source/Runtime/Engine/Public/AI/NavigationSystemHelpers.h b/Engine/Source/Runtime/Engine/Public/AI/NavigationSystemHelpers.h
--- a/Engine/Source/Runtime/Engine/Public/AI/NavigationSystemHelpers.h
+++ b/Engine/Source/Runtime/Engine/Public/AI/NavigationSystemHelpers.h
@@ -13,6 +13,7 @@ struct FNavigableGeometryExport
 	virtual void ExportPxConvexMesh(physx::PxConvexMesh const * const ConvexMesh, const FTransform& LocalToWorld) = 0;
 	virtual void ExportPxHeightField(physx::PxHeightField const * const HeightField, const FTransform& LocalToWorld) = 0;
 #endif // WITH_PHYSX
+	virtual void ExportRigidBodySetup(UBodySetup& BodySetup, const FTransform& LocalToWorld) = 0;
 	virtual void ExportCustomMesh(const FVector* VertexBuffer, int32 NumVerts, const int32* IndexBuffer, int32 NumIndices, const FTransform& LocalToWorld) = 0;
 	
 	virtual void AddNavModifiers(const FCompositeNavModifier& Modifiers) = 0;
