#ttp 331534 - modifying an in level particle system during a paused pie session crashes the editor.
#summary Fixed a crash while editing a particle system with it running in a PIE world
#proj Editor
#branch UE4

diff --git a/Engine/Source/Runtime/Engine/Private/Particles/ParticleComponents.cpp b/Engine/Source/Runtime/Engine/Private/Particles/ParticleComponents.cpp
--- a/Engine/Source/Runtime/Engine/Private/Particles/ParticleComponents.cpp
+++ b/Engine/Source/Runtime/Engine/Private/Particles/ParticleComponents.cpp
@@ -4149,7 +4149,7 @@ void UParticleSystemComponent::InitParticles()
 		else
 		{
 			// create new instances as needed
-			for (int32 Idx = 0; Idx < Template->Emitters.Num(); Idx++)
+			for (int32 Idx = 0; Idx < EmitterInstances.Num(); Idx++)
 			{
 				FParticleEmitterInstance* Instance = EmitterInstances[Idx];
 				if (Instance)
