#ttp 312622 - UE4: RENDERING: Add PostProcessVolume component to Actor Blueprints to enable the creation of 'Lighting Rigs'
#proj Engine
#branch UE4
#summary Created PostProcessVolumeComponent for blueprints.
#add UPostProcessVolumeComponent
#add Postprocess volume interface.
#change Postprocess volumes are now stored in an array.
#change removed linked list vestiges
#add post process volumes array stores interface (to either an actor or component object).
#change fixed up affected code.

diff --git a/Engine/Source/Runtime/Engine/Classes/Components/PostProcessComponent.h b/Engine/Source/Runtime/Engine/Classes/Components/PostProcessComponent.h
new file mode 100644
--- /dev/null
+++ b/Engine/Source/Runtime/Engine/Classes/Components/PostProcessComponent.h
@@ -0,0 +1,65 @@
+// Copyright 1998-2013 Epic Games, Inc. All Rights Reserved.
+
+
+#include "PostProcessComponent.generated.h"
+
+/**
+ *  PostProcessComponent. Enables Post process controls for blueprints.
+ *	Will use a parent UShapeComponent to provide volume data if available.
+ */
+
+UCLASS(HeaderGroup = Component, ClassGroup = Rendering, collapsecategories, hidecategories = (Object), editinlinenew, meta = (BlueprintSpawnableComponent), MinimalAPI)
+class UPostProcessComponent : public USceneComponent, public IInterface_PostProcessVolume
+{
+	GENERATED_UCLASS_BODY()
+
+	/** Post process settings to use for this volume. */
+	UPROPERTY(interp, Category = PostProcessVolume, meta = (ShowOnlyInnerProperties))
+	struct FPostProcessSettings Settings;
+
+	/**
+	 * Priority of this volume. In the case of overlapping volumes the one with the highest priority
+	 * overrides the lower priority ones. The order is undefined if two or more overlapping volumes have the same priority.
+	 */
+	UPROPERTY(EditAnywhere, BlueprintReadWrite, Category = PostProcessVolume)
+	float Priority;
+
+	/** World space radius around the volume that is used for blending (only if not unbound).			*/
+	UPROPERTY(interp, Category = PostProcessVolume, meta = (ClampMin = "0.0", UIMin = "0.0", UIMax = "6000.0"))
+	float BlendRadius;
+
+	/** 0:no effect, 1:full effect */
+	UPROPERTY(interp, Category = PostProcessVolume, BlueprintReadWrite, meta = (UIMin = "0.0", UIMax = "1.0"))
+	float BlendWeight;
+
+	/** Whether this volume is enabled or not.															*/
+	UPROPERTY(EditAnywhere, BlueprintReadWrite, Category = PostProcessVolume)
+	uint32 bEnabled: 1;
+
+	/** set this to false to use the parent shape component as volume bounds. True affects the whole world regardless.		*/
+	UPROPERTY(EditAnywhere, BlueprintReadWrite, Category = PostProcessVolume)
+	uint32 bUnbound: 1;
+ 
+	// Begin IInterface_PostProcessVolume Interface
+	ENGINE_API virtual bool EncompassesPoint(FVector Point, float SphereRadius/*=0.f*/, float* OutDistanceToPoint) OVERRIDE;
+	ENGINE_API virtual FPostProcessVolumeProperties GetProperties() const OVERRIDE
+	{
+		FPostProcessVolumeProperties Ret;
+		Ret.bIsEnabled = bEnabled != 0;
+		Ret.bIsUnbound = bUnbound != 0 || Cast<UShapeComponent>(AttachParent) == nullptr;
+		Ret.BlendRadius = BlendRadius;
+		Ret.BlendWeight = BlendWeight;
+		Ret.Priority = Priority;
+		Ret.Settings = &Settings;
+		return Ret;
+	}
+	// End IInterface_PostProcessVolume Interface
+
+protected:
+
+	virtual void OnRegister();
+	virtual void OnUnregister();
+};
+
+
+
diff --git a/Engine/Source/Runtime/Engine/Classes/Engine/PostProcessVolume.h b/Engine/Source/Runtime/Engine/Classes/Engine/PostProcessVolume.h
--- a/Engine/Source/Runtime/Engine/Classes/Engine/PostProcessVolume.h
+++ b/Engine/Source/Runtime/Engine/Classes/Engine/PostProcessVolume.h
@@ -10,7 +10,7 @@
 
 	// for FPostprocessSettings
 UCLASS(dependson=UScene, autoexpandcategories=PostProcessVolume, hidecategories=(Advanced, Collision, Volume, Brush, Attachment), MinimalAPI)
-class APostProcessVolume : public AVolume
+class APostProcessVolume : public AVolume, public IInterface_PostProcessVolume
 {
 	GENERATED_UCLASS_BODY()
 
@@ -18,10 +18,6 @@ class APostProcessVolume : public AVolume
 	UPROPERTY(interp, Category=PostProcessVolume, meta=(ShowOnlyInnerProperties))
 	struct FPostProcessSettings Settings;
 
-	/** Next volume in linked listed, sorted by priority in ascending order.							*/
-	UPROPERTY(transient)
-	TAutoWeakObjectPtr<class APostProcessVolume> NextHigherPriorityVolume;
-
 	/**
 	 * Priority of this volume. In the case of overlapping volumes the one with the highest priority
 	 * overrides the lower priority ones. The order is undefined if two or more overlapping volumes have the same priority.
@@ -45,6 +41,21 @@ class APostProcessVolume : public AVolume
 	UPROPERTY(EditAnywhere, BlueprintReadWrite, Category=PostProcessVolume)
 	uint32 bUnbound:1;
 
+	// Begin IInterface_PostProcessVolume Interface
+	ENGINE_API virtual bool EncompassesPoint(FVector Point, float SphereRadius/*=0.f*/, float* OutDistanceToPoint) OVERRIDE;
+	ENGINE_API virtual FPostProcessVolumeProperties GetProperties() const OVERRIDE
+	{
+		FPostProcessVolumeProperties Ret;
+		Ret.bIsEnabled = bEnabled != 0;
+		Ret.bIsUnbound = bUnbound != 0;
+		Ret.BlendRadius = BlendRadius;
+		Ret.BlendWeight = BlendWeight;
+		Ret.Priority = Priority;
+		Ret.Settings = &Settings;
+		return Ret;
+	}
+	// End IInterface_PostProcessVolume Interface
+
 
 	// Begin AActor Interface
 	virtual void PostUnregisterAllComponents( void ) OVERRIDE;
diff --git a/Engine/Source/Runtime/Engine/Classes/Engine/World.h b/Engine/Source/Runtime/Engine/Classes/Engine/World.h
--- a/Engine/Source/Runtime/Engine/Classes/Engine/World.h
+++ b/Engine/Source/Runtime/Engine/Classes/Engine/World.h
@@ -790,9 +790,9 @@ public:
 	/** Moved from WorldSettings properties - START 														**/
 	/*****************************************************************************************************/
 		
-	/** Linked list of post processing volumes, sorted in ascending order of priority.					*/
-	TAutoWeakObjectPtr<class APostProcessVolume> LowestPriorityPostProcessVolume;
-	
+	/** An array of post processing volumes, sorted in ascending order of priority.					*/
+	TArray< IInterface_PostProcessVolume * > PostProcessVolumes;
+
 	/** Linked list of reverb volumes, sorted in descending order of priority.							*/
 	TAutoWeakObjectPtr<class AReverbVolume> HighestPriorityReverbVolume;
 
diff --git a/Engine/Source/Runtime/Engine/Classes/Interfaces/Interface_PostProcessVolume.h b/Engine/Source/Runtime/Engine/Classes/Interfaces/Interface_PostProcessVolume.h
new file mode 100644
--- /dev/null
+++ b/Engine/Source/Runtime/Engine/Classes/Interfaces/Interface_PostProcessVolume.h
@@ -0,0 +1,36 @@
+// Copyright 1998-2013 Epic Games, Inc. All Rights Reserved.
+#include "Interface_PostProcessVolume.generated.h"
+
+/** Interface for general PostProcessVolume access **/
+UINTERFACE()
+class UInterface_PostProcessVolume : public UInterface
+{
+	GENERATED_UINTERFACE_BODY()
+};
+
+struct FPostProcessVolumeProperties
+{
+	const FPostProcessSettings* Settings;
+	float Priority;
+	float BlendRadius;
+	float BlendWeight;
+	bool bIsEnabled;
+	bool bIsUnbound;
+};
+
+class IInterface_PostProcessVolume
+{
+	GENERATED_IINTERFACE_BODY()
+
+	ENGINE_API virtual bool EncompassesPoint(FVector Point, float SphereRadius/*=0.f*/, float* OutDistanceToPoint)
+	{
+		checkNoEntry();
+		return false;
+	}
+	ENGINE_API virtual FPostProcessVolumeProperties GetProperties() const
+	{
+		checkNoEntry();
+		FPostProcessVolumeProperties dummy = {};
+		return dummy;
+	}
+};
diff --git a/Engine/Source/Runtime/Engine/Private/Interface.cpp b/Engine/Source/Runtime/Engine/Private/Interface.cpp
--- a/Engine/Source/Runtime/Engine/Private/Interface.cpp
+++ b/Engine/Source/Runtime/Engine/Private/Interface.cpp
@@ -12,3 +12,9 @@ URuntimeMovieScenePlayerInterface::URuntimeMovieScenePlayerInterface( const clas
 {
 
 }
+
+UInterface_PostProcessVolume::UInterface_PostProcessVolume( const class FPostConstructInitializeProperties& PCIP )
+	: Super(PCIP)
+{
+
+}
diff --git a/Engine/Source/Runtime/Engine/Private/PostProcessComponent.cpp b/Engine/Source/Runtime/Engine/Private/PostProcessComponent.cpp
new file mode 100644
--- /dev/null
+++ b/Engine/Source/Runtime/Engine/Private/PostProcessComponent.cpp
@@ -0,0 +1,50 @@
+// Copyright 1998-2013 Epic Games, Inc. All Rights Reserved.
+
+#include "EnginePrivate.h"
+
+UPostProcessComponent::UPostProcessComponent(const class FPostConstructInitializeProperties& PCIP)
+	: Super(PCIP)
+{
+	bEnabled = true;
+	BlendRadius = 100.0f;
+	BlendWeight = 1.0f;
+	Priority = 0;
+	bUnbound = 1;
+}
+
+bool UPostProcessComponent::EncompassesPoint(FVector Point, float SphereRadius/*=0.f*/, float* OutDistanceToPoint)
+{
+	UShapeComponent* ParentShape = Cast<UShapeComponent>(AttachParent);
+	if (ParentShape != nullptr)
+	{
+#if WITH_PHYSX
+		FVector ClosestPoint;
+		float Distance = ParentShape->GetDistanceToCollision(Point, ClosestPoint);
+#else
+		FBoxSphereBounds Bounds = ParentShape->CalcBounds(ParentShape->ComponentToWorld);
+		float Distance = 0;
+		if (ParentShape->IsA<USphereComponent>())
+		{
+			const FSphere& Sphere = Bounds.GetSphere();
+			const FVector& Dist = Sphere.Center - Point;
+			Distance = FMath::Max(0.0f, Dist.Size() - Sphere.W);
+		}
+		else // UBox or UCapsule shape (approx).
+		{
+			Distance = FMath::Sqrt(Bounds.GetBox().ComputeSquaredDistanceToPoint(Point));
+		}
+#endif
+
+		if (OutDistanceToPoint)
+		{
+			*OutDistanceToPoint = Distance;
+		}
+
+		return Distance >= 0.f && Distance <= SphereRadius;
+	}
+	if (OutDistanceToPoint != nullptr)
+	{
+		*OutDistanceToPoint = 0;
+	}
+	return true;
+}
\ No newline at end of file
diff --git a/Engine/Source/Runtime/Engine/Private/PostProcessVolume.cpp b/Engine/Source/Runtime/Engine/Private/PostProcessVolume.cpp
--- a/Engine/Source/Runtime/Engine/Private/PostProcessVolume.cpp
+++ b/Engine/Source/Runtime/Engine/Private/PostProcessVolume.cpp
@@ -16,6 +16,11 @@ APostProcessVolume::APostProcessVolume(const class FPostConstructInitializePrope
 	BlendWeight = 1.0f;
 }
 
+bool APostProcessVolume::EncompassesPoint(FVector Point, float SphereRadius/*=0.f*/, float* OutDistanceToPoint)
+{
+	return Super::EncompassesPoint(Point, SphereRadius, OutDistanceToPoint);
+}
+
 #if WITH_EDITOR
 void APostProcessVolume::PostEditChangeProperty(FPropertyChangedEvent& PropertyChangedEvent)
 {
diff --git a/Engine/Source/Runtime/Engine/Private/SceneView.cpp b/Engine/Source/Runtime/Engine/Private/SceneView.cpp
--- a/Engine/Source/Runtime/Engine/Private/SceneView.cpp
+++ b/Engine/Source/Runtime/Engine/Private/SceneView.cpp
@@ -706,6 +706,52 @@ UBlendableInterface::UBlendableInterface(const class FPostConstructInitializePro
 {
 }
 
+void DoPostProcessVolume(IInterface_PostProcessVolume* Volume, FVector ViewLocation, FSceneView* SceneView)
+{
+	const FPostProcessVolumeProperties VolumeProperties = Volume->GetProperties();
+	if (!VolumeProperties.bIsEnabled)
+	{
+		return;
+	}
+
+	float DistanceToPoint = 0.0f;
+	float LocalWeight = FMath::Clamp(VolumeProperties.BlendWeight, 0.0f, 1.0f);
+
+	if (!VolumeProperties.bIsUnbound)
+	{
+		float SquaredBlendRadius = VolumeProperties.BlendRadius * VolumeProperties.BlendRadius;
+		Volume->EncompassesPoint(ViewLocation, 0.0f, &DistanceToPoint);
+
+		if (DistanceToPoint >= 0)
+		{
+			if (DistanceToPoint > VolumeProperties.BlendRadius)
+			{
+				// outside
+				LocalWeight = 0.0f;
+			}
+			else
+			{
+				// to avoid div by 0
+				if (VolumeProperties.BlendRadius >= 1.0f)
+				{
+					LocalWeight *= 1.0f - DistanceToPoint / VolumeProperties.BlendRadius;
+
+					check(LocalWeight >= 0 && LocalWeight <= 1.0f);
+				}
+			}
+		}
+		else
+		{
+			LocalWeight = 0;
+		}
+	}
+
+	if (LocalWeight > 0)
+	{
+		SceneView->OverridePostProcessSettings(*VolumeProperties.Settings, LocalWeight);
+	}
+}
+
 void FSceneView::StartFinalPostprocessSettings(FVector InViewLocation)
 {
 	check(IsInGameThread());
@@ -724,46 +770,11 @@ void FSceneView::StartFinalPostprocessSettings(FVector InViewLocation)
 	UWorld* World = Family->Scene->GetWorld();
 
 	// Some views have no world (e.g. material preview)
-	if(World)
+	if (World)
 	{
-		APostProcessVolume* Volume = World->LowestPriorityPostProcessVolume;
-		while(Volume)
+		for (auto VolumeIt = World->PostProcessVolumes.CreateIterator(); VolumeIt; ++VolumeIt)
 		{
-			// Volume encompasses
-			if(Volume->bEnabled)
-			{
-				float DistanceToPoint = 0.0f;
-				float LocalWeight = FMath::Clamp(Volume->BlendWeight, 0.0f, 1.0f);
-
-				if(!Volume->bUnbound)
-				{
-					Volume->EncompassesPoint(InViewLocation, 0.0f, &DistanceToPoint);
-
-					if (DistanceToPoint >= 0)
-					{
-						if(DistanceToPoint > Volume->BlendRadius)
-						{
-							// outside
-							LocalWeight = 0.0f;
-						}
-						else
-						{
-							// to avoid div by 0
-							if(Volume->BlendRadius >= 1.0f)
-							{
-								LocalWeight *= 1.0f - DistanceToPoint / Volume->BlendRadius;
-
-								check(LocalWeight >=0 && LocalWeight <= 1.0f);
-							}
-						}
-					}
-				}
-
-				OverridePostProcessSettings(Volume->Settings, LocalWeight);
-			}
-
-			// further traverse linked list.
-			Volume = Volume->NextHigherPriorityVolume;
+			DoPostProcessVolume(*VolumeIt, InViewLocation, this);
 		}
 	}
 }
diff --git a/Engine/Source/Runtime/Engine/Private/UnrealClient.cpp b/Engine/Source/Runtime/Engine/Private/UnrealClient.cpp
--- a/Engine/Source/Runtime/Engine/Private/UnrealClient.cpp
+++ b/Engine/Source/Runtime/Engine/Private/UnrealClient.cpp
@@ -509,91 +509,57 @@ void FViewport::EndRenderFrame( bool bPresent, bool bLockToVsync )
 	RHIEndDrawingViewport( GetViewportRHI(), bPresent, bLockToVsync );
 }
 
+void InsertVolume(IInterface_PostProcessVolume* Volume, TArray< IInterface_PostProcessVolume* >& VolumeArray)
+{
+	const int32 NumVolumes = VolumeArray.Num();
+	float TargetPriority = Volume->GetProperties().Priority;
+	int32 InsertIndex = 0;
+	// TODO: replace with binary search.
+	for (; InsertIndex < NumVolumes ; InsertIndex++)
+	{
+		IInterface_PostProcessVolume* CurrentVolume = VolumeArray[InsertIndex];
+		float CurrentPriority = CurrentVolume->GetProperties().Priority;
 
+		if (TargetPriority < CurrentPriority)
+		{
+			break;
+		}
+		if (CurrentVolume == Volume)
+		{
+			return;
+		}
+	}
+	VolumeArray.Insert(Volume, InsertIndex);
+}
 
 void APostProcessVolume::PostUnregisterAllComponents()
 {
 	// Route clear to super first.
 	Super::PostUnregisterAllComponents();
-
 	// World will be NULL during exit purge.
-	if( GetWorld() )
+	if (GetWorld())
 	{
-		APostProcessVolume* CurrentVolume  = GetWorld()->LowestPriorityPostProcessVolume;
-		APostProcessVolume*	PreviousVolume = NULL;
-
-		// Iterate over linked list, removing this volume if found.
-		while( CurrentVolume )
-		{
-			// Found.
-			if( CurrentVolume == this )
-			{
-				// Remove from linked list.
-				if( PreviousVolume )
-				{
-					PreviousVolume->NextHigherPriorityVolume = NextHigherPriorityVolume;
-				}
-				// Special case removal from first entry.
-				else
-				{
-					GetWorld()->LowestPriorityPostProcessVolume = NextHigherPriorityVolume;
-				}
-
-				// BREAK OUT OF LOOP
-				break;
-			}
-			// Further traverse linked list.
-			else
-			{
-				PreviousVolume	= CurrentVolume;
-				CurrentVolume	= CurrentVolume->NextHigherPriorityVolume;
-			}
-		}
-
-		// Reset next pointer to avoid dangling end bits and also for GC.
-		NextHigherPriorityVolume = NULL;
+		GetWorld()->PostProcessVolumes.RemoveSingle(this);
 	}
 }
 
-
 void APostProcessVolume::PostRegisterAllComponents()
 {
 	// Route update to super first.
 	Super::PostRegisterAllComponents();
+	InsertVolume(this, GetWorld()->PostProcessVolumes);
+}
 
-	APostProcessVolume* CurrentVolume  = GetWorld()->LowestPriorityPostProcessVolume;
-	TAutoWeakObjectPtr<APostProcessVolume>* PreviousVolumePtr = &GetWorld()->LowestPriorityPostProcessVolume;
-
-	// Find where to insert in sorted linked list.
-	while( CurrentVolume && CurrentVolume != this )
-	{
-		// We use < instead of <= to be sure that we are not inserting twice in the case of multiple volumes having
-		// the same priority and the current one already having being inserted after one with the same priority.
-		if( Priority < CurrentVolume->Priority )
-		{
-			// Insert before current node by fixing up previous to point to current.
-			*PreviousVolumePtr = this;
-
-			// Point to current volume, finalizing insertion.
-			NextHigherPriorityVolume = CurrentVolume;
-
-			// BREAK OUT OF LOOP.
-			return;
-		}
-		// Further traverse linked list.
-		else
-		{
-			PreviousVolumePtr = &CurrentVolume->NextHigherPriorityVolume;
-			CurrentVolume = CurrentVolume->NextHigherPriorityVolume;
-		}
-	}
+void UPostProcessComponent::OnRegister()
+{
+	Super::OnRegister();
+	InsertVolume(this, GetWorld()->PostProcessVolumes);
+}
 
-	// Avoid double insertion!
-	if(!CurrentVolume)
-	{
-		NextHigherPriorityVolume = *PreviousVolumePtr;
-		*PreviousVolumePtr = this;
-	}
+void UPostProcessComponent::OnUnregister()
+{
+	Super::OnUnregister();
+	GetWorld()->PostProcessVolumes.RemoveSingle(this);
 }
 
 /**
