Allow game to set player specific axis bindings similar to Action and Axis mappings and not be stuck with project defaults
TTP #331495

diff --git a/Engine/Source/Runtime/Engine/Classes/GameFramework/PlayerInput.h b/Engine/Source/Runtime/Engine/Classes/GameFramework/PlayerInput.h
--- a/Engine/Source/Runtime/Engine/Classes/GameFramework/PlayerInput.h
+++ b/Engine/Source/Runtime/Engine/Classes/GameFramework/PlayerInput.h
@@ -253,6 +253,9 @@ public:
 	UPROPERTY(config)
 	TArray<struct FKeyBind> DebugExecBindings;
 
+	/** This player's version of the Axis Properties */
+	TArray<struct FInputAxisConfigEntry> AxisConfig;
+
 	/** This player's version of the Action Mappings */
 	TArray<struct FInputActionKeyMapping> ActionMappings;
 
@@ -265,7 +268,7 @@ public:
 	//=============================================================================
 	// Input related functions.
 	UFUNCTION(exec)
-	void SetMouseSensitivity(float F);
+	void SetMouseSensitivity(const float Sensitivity);
 
 	UFUNCTION(exec)
 	void SetMouseSensitivityToDefault();
diff --git a/Engine/Source/Runtime/Engine/Private/UserInterface/PlayerInput.cpp b/Engine/Source/Runtime/Engine/Private/UserInterface/PlayerInput.cpp
--- a/Engine/Source/Runtime/Engine/Private/UserInterface/PlayerInput.cpp
+++ b/Engine/Source/Runtime/Engine/Private/UserInterface/PlayerInput.cpp
@@ -274,21 +274,18 @@ float UPlayerInput::GetMouseSensitivity()
 	return 1.0f;
 }
 
-void UPlayerInput::SetMouseSensitivity(float F)
+void UPlayerInput::SetMouseSensitivity(const float Sensitivity)
 {
-	ConditionalInitAxisProperties();
-
-	FInputAxisProperties* const XAxisProps = AxisProperties.Find(EKeys::MouseX);
-	if (XAxisProps)
+	for (FInputAxisConfigEntry& AxisConfigEntry : AxisConfig)
 	{
-		XAxisProps->Sensitivity = F;
+		const FKey AxisKey = AxisConfigEntry.AxisKeyName;
+		if (AxisKey == EKeys::MouseX || AxisKey == EKeys::MouseY)
+		{
+			AxisConfigEntry.AxisProperties.Sensitivity = Sensitivity;
+		}
 	}
 
-	FInputAxisProperties* const YAxisProps = AxisProperties.Find(EKeys::MouseY);
-	if (YAxisProps)
-	{
-		YAxisProps->Sensitivity = F;
-	}
+	AxisProperties.Empty();
 }
 
 void UPlayerInput::SetMouseSensitivityToDefault()
@@ -469,6 +466,7 @@ void UPlayerInput::ForceRebuildingKeyMaps(const bool bRestoreDefaults)
 {
 	if (bRestoreDefaults)
 	{
+		AxisConfig = GetDefault<UInputSettings>()->AxisConfig;
 		AxisMappings = GetDefault<UInputSettings>()->AxisMappings;
 		ActionMappings = GetDefault<UInputSettings>()->ActionMappings;
 	}
@@ -1373,7 +1371,7 @@ void UPlayerInput::ConditionalInitAxisProperties()
 	if ( AxisProperties.Num() == 0 )
 	{
 		// move stuff from config structure to our runtime structure
-		for (const FInputAxisConfigEntry& AxisConfigEntry : GetDefault<UInputSettings>()->AxisConfig)
+		for (const FInputAxisConfigEntry& AxisConfigEntry : AxisConfig)
 		{
 			const FKey AxisKey = AxisConfigEntry.AxisKeyName;
 			if (AxisKey.IsValid())
