#UE4:
* added show flags to disable directional/point/spot light contributions

diff --git a/Engine/Source/Runtime/Engine/Public/ShowFlagsValues.inl b/Engine/Source/Runtime/Engine/Public/ShowFlagsValues.inl
--- a/Engine/Source/Runtime/Engine/Public/ShowFlagsValues.inl
+++ b/Engine/Source/Runtime/Engine/Public/ShowFlagsValues.inl
@@ -53,7 +53,13 @@ SHOWFLAG_ALWAYS_ACCESSIBLE(ReflectionOverride, SFG_Hidden, LOCTEXT("ReflectionOv
 /** needed for VMI_VisualizeBuffer, Whether to enable the buffer visualization mode. */
 SHOWFLAG_ALWAYS_ACCESSIBLE(VisualizeBuffer, SFG_Hidden, LOCTEXT("VisualizeBufferSF", "Buffer Visualization"))
 /** Allows to disable all direct lighting (does not affect indirect light) */
-SHOWFLAG_ALWAYS_ACCESSIBLE(DirectLighting, SFG_LightingComponents, LOCTEXT("DirectLightingSF", "Direct Lighting"))
+SHOWFLAG_FIXED_IN_SHIPPING(DirectLighting, 1, SFG_LightingComponents, LOCTEXT("DirectLightingSF", "Direct Lighting"))
+/** Allows to disable lighting from Directional Lights */
+SHOWFLAG_FIXED_IN_SHIPPING(DirectionalLights, 1, SFG_LightingComponents, LOCTEXT("DirectionalLightsSF", "Directional Lights"))
+/** Allows to disable lighting from Point Lights */
+SHOWFLAG_FIXED_IN_SHIPPING(PointLights, 1, SFG_LightingComponents, LOCTEXT("PointLightsSF", "Point Lights"))
+/** Allows to disable lighting from Spot Lights */
+SHOWFLAG_FIXED_IN_SHIPPING(SpotLights, 1, SFG_LightingComponents, LOCTEXT("SpotLightsSF", "Spot Lights"))
 /** Color correction after tone mapping */
 SHOWFLAG_ALWAYS_ACCESSIBLE(ColorGrading, SFG_PostProcess, LOCTEXT("ColorGradingSF", "Color Grading"))
 /** Visualize vector fields. */
diff --git a/Engine/Source/Runtime/Renderer/Private/LightSceneInfo.cpp b/Engine/Source/Runtime/Renderer/Private/LightSceneInfo.cpp
--- a/Engine/Source/Runtime/Renderer/Private/LightSceneInfo.cpp
+++ b/Engine/Source/Runtime/Renderer/Private/LightSceneInfo.cpp
@@ -140,9 +140,36 @@ void FLightSceneInfo::Detach()
 
 bool FLightSceneInfo::ShouldRenderLight(const FViewInfo& View) const
 {
-	const bool bInViewFrustum = bVisible ? View.VisibleLightInfos[Id].bInViewFrustum : true;
 	// Only render the light if it is in the view frustum
-	return bInViewFrustum
+	bool bLocalVisible = bVisible ? View.VisibleLightInfos[Id].bInViewFrustum : true;
+
+#if !(UE_BUILD_SHIPPING || UE_BUILD_TEST)
+	ELightComponentType Type = (ELightComponentType)Proxy->GetLightType();
+
+	switch(Type)
+	{
+		case LightType_Directional:
+			if(!View.Family->EngineShowFlags.DirectionalLights) 
+			{
+				bLocalVisible = false;
+			}
+			break;
+		case LightType_Point:
+			if(!View.Family->EngineShowFlags.PointLights) 
+			{
+				bLocalVisible = false;
+			}
+			break;
+		case LightType_Spot:
+			if(!View.Family->EngineShowFlags.SpotLights)
+			{
+				bLocalVisible = false;
+			}
+			break;
+	}
+#endif
+
+	return bLocalVisible
 		// Only render lights with static shadowing for reflection captures, since they are only captured at edit time
 		&& (!View.bIsReflectionCapture || Proxy->HasStaticShadowing());
 }
