Shader compiling will fall back to singlethreaded dll if the worker process terminates unexpectedly

diff --git a/Engine/Source/Runtime/Engine/Private/ShaderCompiler/ShaderCompiler.cpp b/Engine/Source/Runtime/Engine/Private/ShaderCompiler/ShaderCompiler.cpp
--- a/Engine/Source/Runtime/Engine/Private/ShaderCompiler/ShaderCompiler.cpp
+++ b/Engine/Source/Runtime/Engine/Private/ShaderCompiler/ShaderCompiler.cpp
@@ -716,11 +716,12 @@ void FShaderCompileThreadRunnable::LaunchWorkerIfNeeded(FShaderCompileWorkerInfo
 #endif
 }
 
-void FShaderCompileThreadRunnable::LaunchWorkersIfNeeded()
+bool FShaderCompileThreadRunnable::LaunchWorkersIfNeeded()
 {
 	const double CurrentTime = FPlatformTime::Seconds();
 	// Limit how often we check for workers running since IsApplicationRunning eats up some CPU time on Windows
 	const bool bCheckForWorkerRunning = (CurrentTime - LastCheckForWorkersTime > .1f);
+	bool bAbandonWorkers = false;
 
 	if (bCheckForWorkerRunning)
 	{
@@ -761,7 +762,10 @@ void FShaderCompileThreadRunnable::LaunchWorkersIfNeeded()
 					}
 					else
 					{
-						UE_LOG(LogShaderCompilers, Fatal, TEXT("ShaderCompileWorker terminated unexpectedly! Thread %u."), WorkerIndex);
+						UE_LOG(LogShaderCompilers, Error, TEXT("ShaderCompileWorker terminated unexpectedly!  Falling back to directly compiling which will be very slow.  Thread %u."), WorkerIndex);
+
+						bAbandonWorkers = true;
+						break;
 					}
 				}
 
@@ -778,6 +782,8 @@ void FShaderCompileThreadRunnable::LaunchWorkersIfNeeded()
 			}
 		}
 	}
+
+	return bAbandonWorkers;
 }
 
 void FShaderCompileThreadRunnable::ReadAvailableResults()
@@ -930,10 +936,19 @@ int32 FShaderCompileThreadRunnable::CompilingLoop()
 
 			// Launch shader compile workers if they are not already running
 			// Workers can time out when idle so they may need to be relaunched
-			LaunchWorkersIfNeeded();
+			bool bAbandonWorkers = LaunchWorkersIfNeeded();
 
-			// Read files which are outputs from the shader compile workers
-			ReadAvailableResults();
+			if (bAbandonWorkers)
+			{
+				// Fall back to local compiles if the SCW crashed.
+				// This is nasty but needed to work around issues where message passing through files to SCW is unreliable on random PCs
+				Manager->bAllowCompilingThroughWorkers = false;
+			}
+			else
+			{
+				// Read files which are outputs from the shader compile workers
+				ReadAvailableResults();
+			}
 		}
 	}
 	else
@@ -956,6 +971,7 @@ FShaderCompilingManager::FShaderCompilingManager() :
 #endif
 {
 	WorkersBusyTime = 0;
+	bFallBackToDirectCompiles = false;
 
 	// Threads must use absolute paths on Windows in case the current directory is changed on another thread!
 	ShaderCompileWorkerName = FPaths::ConvertRelativePathToFull(ShaderCompileWorkerName);
diff --git a/Engine/Source/Runtime/Engine/Public/ShaderCompiler.h b/Engine/Source/Runtime/Engine/Public/ShaderCompiler.h
--- a/Engine/Source/Runtime/Engine/Public/ShaderCompiler.h
+++ b/Engine/Source/Runtime/Engine/Public/ShaderCompiler.h
@@ -96,7 +96,7 @@ private:
 	void WriteNewTasks();
 
 	/** Used when compiling through workers, launches worker processes if needed. */
-	void LaunchWorkersIfNeeded();
+	bool LaunchWorkersIfNeeded();
 
 	/** Used when compiling through workers, attempts to open the worker output file if the worker is done and read the results. */
 	void ReadAvailableResults();
@@ -200,6 +200,8 @@ private:
 	FString AbsoluteShaderDebugInfoDirectory;
 	/** Name of the shader worker application. */
 	FString ShaderCompileWorkerName;
+	/** Whether the SCW has crashed and we should fall back to calling the compiler dll's directly. */
+	bool bFallBackToDirectCompiles;
 
 	/** 
 	 * Tracks the total time that shader compile workers have been busy since startup.  
