Making it so components aren't registered until after their construction script has been ran (fixes issue where bad collision data was being generated).

diff --git a/Engine/Source/Runtime/Engine/Private/ActorConstruction.cpp b/Engine/Source/Runtime/Engine/Private/ActorConstruction.cpp
--- a/Engine/Source/Runtime/Engine/Private/ActorConstruction.cpp
+++ b/Engine/Source/Runtime/Engine/Private/ActorConstruction.cpp
@@ -4,7 +4,7 @@
 #include "BlueprintUtilities.h"
 #include "LatentActions.h"
 #include "EngineLevelScriptClasses.h"
-#include "DeferRegisterStaticComponents.h"
+#include "DeferRegisterComponents.h"
 
 #if WITH_EDITOR
 #include "Editor/UnrealEd/Public/Kismet2/BlueprintEditorUtils.h"
@@ -310,7 +310,7 @@ void AActor::OnConstruction(const FTransform& Transform)
 }
 
 
-void FDeferRegisterStaticComponents::DeferStaticComponent(AActor* Actor, USceneComponent* Component, EComponentMobility::Type OriginalMobility)
+void FDeferRegisterComponents::DeferComponentRegistration(AActor* Actor, USceneComponent* Component, EComponentMobility::Type OriginalMobility)
 {
 	TArray<FDeferredComponentInfo>* ActorComponentsToRegister = ComponentsToRegister.Find(Actor);
 	if (!ActorComponentsToRegister)
@@ -318,11 +318,15 @@ void FDeferRegisterStaticComponents::DeferStaticComponent(AActor* Actor, USceneC
 		ActorComponentsToRegister = &ComponentsToRegister.Add(Actor, TArray<FDeferredComponentInfo>());
 	}
 	ActorComponentsToRegister->Add(FDeferredComponentInfo(Component, OriginalMobility));
-	// Make sure Mobility is temporarily set to EComponentMobility::Movable.
+	
+	// Components with Mobility set to EComponentMobility::Static or EComponentMobility::Stationary 
+	// can't be properly set up from a construction script (all changes will be rejected due to 
+	// EComponentMobility::Static flag) so we're going to temporarily change the flag and defer 
+	// the registration until the construction-script has finished.
 	Component->Mobility = EComponentMobility::Movable;
 }
 
-void FDeferRegisterStaticComponents::RegisterComponents(AActor* Actor)
+void FDeferRegisterComponents::RegisterComponents(AActor* Actor)
 {
 	TArray<FDeferredComponentInfo>* ActorComponentsToRegister = ComponentsToRegister.Find(Actor);
 	if (ActorComponentsToRegister)
@@ -340,7 +344,7 @@ void FDeferRegisterStaticComponents::RegisterComponents(AActor* Actor)
 }
 
 // FGCObject interface
-void FDeferRegisterStaticComponents::AddReferencedObjects(FReferenceCollector& Collector)
+void FDeferRegisterComponents::AddReferencedObjects(FReferenceCollector& Collector)
 {
 	for (TMap<AActor*, TArray<FDeferredComponentInfo> >::TIterator It(ComponentsToRegister); It; ++It)
 	{
@@ -353,10 +357,10 @@ void FDeferRegisterStaticComponents::AddReferencedObjects(FReferenceCollector& C
 	}
 }
 
-/* Gets the FDeferRegisterStaticComponents singleton. */
-FDeferRegisterStaticComponents& FDeferRegisterStaticComponents::Get()
+/* Gets the FDeferRegisterComponents singleton. */
+FDeferRegisterComponents& FDeferRegisterComponents::Get()
 {
-	static FDeferRegisterStaticComponents Singleton;
+	static FDeferRegisterComponents Singleton;
 	return Singleton;
 }
 
@@ -368,7 +372,7 @@ void AActor::ProcessUserConstructionScript()
 	bRunningUserConstructionScript = false;
 
 	// Register all deferred components for this actor.
-	FDeferRegisterStaticComponents::Get().RegisterComponents(this);
+	FDeferRegisterComponents::Get().RegisterComponents(this);
 }
 
 void AActor::FinishAndRegisterComponent(UActorComponent* Component)
@@ -419,14 +423,12 @@ UActorComponent* AActor::AddComponent(FName TemplateName, bool bManualAttachment
 
 		if(NewSceneComp != NULL)
 		{
-			// Components with Mobility set to EComponentMobility::Static or EComponentMobility::Stationary can't be properly set up in UCS (all changes will be rejected
-			// due to EComponentMobility::Static flag) so we're going to temporarily change the flag and defer the registration until UCS has finished.
-			bDeferRegisterStaticComponent = bRunningUserConstructionScript && NewSceneComp->Mobility != EComponentMobility::Movable;
+			// if we end up deferring this component's registration below, then we 
+			// need to provide that deferred action the component's original mobility
+			// (because in the meantime, FDeferRegisterComponents alters the mobility
+			// to Movable so that construction-script changes won't be rejected due to
+			// a EComponentMobility::Static mobility)
 			OriginalMobility = NewSceneComp->Mobility;
-			if (bDeferRegisterStaticComponent)
-			{
-				NewSceneComp->Mobility = EComponentMobility::Movable;
-			}
 
 			if (!bManualAttachment)
 			{
@@ -446,10 +448,12 @@ UActorComponent* AActor::AddComponent(FName TemplateName, bool bManualAttachment
 		// Call function to notify component it has been created
 		NewActorComp->OnComponentCreated();
 
-		if (bDeferRegisterStaticComponent)
+		// need to defer component registration until after the construction script
+		// is ran, since the construction script can mutate the object (for collision, etc.)
+		if (bRunningUserConstructionScript)
 		{
 			// Defer registration until after UCS has completed.
-			FDeferRegisterStaticComponents::Get().DeferStaticComponent(this, NewSceneComp, OriginalMobility);
+			FDeferRegisterComponents::Get().DeferComponentRegistration(this, NewSceneComp, OriginalMobility);
 		}
 		else
 		{
diff --git a/Engine/Source/Runtime/Engine/Private/DeferRegisterComponents.h b/Engine/Source/Runtime/Engine/Private/DeferRegisterComponents.h
new file mode 100644
--- /dev/null
+++ b/Engine/Source/Runtime/Engine/Private/DeferRegisterComponents.h
@@ -0,0 +1,52 @@
+// Copyright 1998-2014 Epic Games, Inc. All Rights Reserved.
+
+#pragma once
+
+
+class FDeferredComponentInfo
+{
+public:
+	USceneComponent* Component;
+	EComponentMobility::Type SavedMobility;
+
+	FDeferredComponentInfo(USceneComponent* InComponent, EComponentMobility::Type OriginalMobility) :
+		Component(InComponent),
+		SavedMobility(OriginalMobility)
+	{}
+};
+
+/**
+ * Helper class to store components which registration should be deferred until
+ * the construction script has completed... generally, all components get stuck here
+ * because the construction script could modify something that is important prior to
+ * registration (transform data for physics, etc.).
+ */
+class FDeferRegisterComponents : public FGCObject
+{
+	/** Map of actors and their components to register. */
+	TMap<AActor*, TArray<FDeferredComponentInfo> > ComponentsToRegister;
+
+public:
+
+	/**
+	 * Stores the component to delay its registration.
+	 *
+	 * @param Actor Actor this component belongs to.
+	 * @param Component Component to register.
+	 */
+	void DeferComponentRegistration(AActor* Actor, USceneComponent* Component, EComponentMobility::Type OriginalMobility);
+
+	/**
+	 * Registers all deferred components that belong to the specified actor.
+	 *
+	 * @param Actor Actor which components should be registered.
+	 */
+	void RegisterComponents(AActor* Actor);
+
+	// FGCObject interface
+	virtual void AddReferencedObjects(FReferenceCollector& Collector) OVERRIDE;
+
+	/* Gets the FDeferRegisterComponents singleton. */
+	static FDeferRegisterComponents& Get();
+};
+
diff --git a/Engine/Source/Runtime/Engine/Private/DeferRegisterStaticComponents.h b/Engine/Source/Runtime/Engine/Private/DeferRegisterStaticComponents.h
deleted file mode 100644
diff --git a/Engine/Source/Runtime/Engine/Private/SCS_Node.cpp b/Engine/Source/Runtime/Engine/Private/SCS_Node.cpp
--- a/Engine/Source/Runtime/Engine/Private/SCS_Node.cpp
+++ b/Engine/Source/Runtime/Engine/Private/SCS_Node.cpp
@@ -7,7 +7,7 @@
 #if WITH_EDITORONLY_DATA
 #include "Kismet2/BlueprintEditorUtils.h"
 #endif
-#include "DeferRegisterStaticComponents.h"
+#include "DeferRegisterComponents.h"
 
 //////////////////////////////////////////////////////////////////////////
 // USCS_Node
@@ -44,14 +44,12 @@ void USCS_Node::ExecuteNodeOnActor(AActor* Actor, USceneComponent* ParentCompone
 		USceneComponent* NewSceneComp = Cast<USceneComponent>(NewActorComp);
 		if (NewSceneComp != NULL)
 		{
-			// Components with Mobility set to EComponentMobility::Static or EComponentMobility::Stationary can't be properly set up in SCS (all changes will be rejected
-			// due to EComponentMobility::Static flag) so we're going to temporarily change the flag and defer the registration until SCS has finished.
-			bDeferRegisterStaticComponent = NewSceneComp->Mobility != EComponentMobility::Movable;
+			// since we end up deferring this component's registration below, we 
+			// need to provide that deferred action the component's original mobility
+			// (because in the meantime, FDeferRegisterComponents alters the mobility
+			// to Movable so that construction-script changes won't be rejected due to
+			// a EComponentMobility::Static mobility)
 			OriginalMobility = NewSceneComp->Mobility;
-			if (bDeferRegisterStaticComponent)
-			{
-				NewSceneComp->Mobility = EComponentMobility::Movable;
-			}
 
 			// If NULL is passed in, we are the root, so set transform and assign as RootComponent on Actor
 			if (ParentComponent == NULL || (ParentComponent && ParentComponent->IsPendingKill()))
@@ -70,15 +68,9 @@ void USCS_Node::ExecuteNodeOnActor(AActor* Actor, USceneComponent* ParentCompone
 		// Call function to notify component it has been created
 		NewActorComp->OnComponentCreated();
 
-		if (bDeferRegisterStaticComponent)
-		{
-			// Defer registration until after SCS has completed.
-			FDeferRegisterStaticComponents::Get().DeferStaticComponent(Actor, NewSceneComp, OriginalMobility);
-		}
-		else
-		{
-			NewActorComp->RegisterComponent();
-		}
+		// need to defer component registration until after the construction script
+		// is ran, since the construction script can mutate the object (for collision, etc.)
+		FDeferRegisterComponents::Get().DeferComponentRegistration(Actor, NewSceneComp, OriginalMobility);
 
 		if (NewActorComp->GetIsReplicated())
 		{
