#UE4 - Send 0 bit to signify we are deleting sub-objects

diff --git a/Engine/Source/Runtime/Engine/Classes/Engine/ActorChannel.h b/Engine/Source/Runtime/Engine/Classes/Engine/ActorChannel.h
--- a/Engine/Source/Runtime/Engine/Classes/Engine/ActorChannel.h
+++ b/Engine/Source/Runtime/Engine/Classes/Engine/ActorChannel.h
@@ -127,6 +127,9 @@ class ENGINE_API UActorChannel : public UChannel
 	/** Writes the header for a content block of proeprties / RPCs for the given object (either the actor a subobject of the actor) */
 	void BeginContentBlock( UObject * Obj, FOutBunch &Bunch );
 
+	/** Writes the header for a content block specifically for deleting sub-objects */
+	void BeginContentBlockForSubObjectDelete( FOutBunch & Bunch, FNetworkGUID & GuidToDelete );
+
 	/** Writes the footer for a content block of proeprties / RPCs for the given object (either the actor a subobject of the actor) */
 	void EndContentBlock( UObject *Obj, FOutBunch &Bunch, FClassNetCache* ClassCache = NULL );
 
diff --git a/Engine/Source/Runtime/Engine/Private/DataChannel.cpp b/Engine/Source/Runtime/Engine/Private/DataChannel.cpp
--- a/Engine/Source/Runtime/Engine/Private/DataChannel.cpp
+++ b/Engine/Source/Runtime/Engine/Private/DataChannel.cpp
@@ -1675,13 +1675,7 @@ bool UActorChannel::ReplicateActor()
 		if (!RepComp.Key().IsValid())
 		{
 			// Write a deletion content header:
-			//	-Deleted object's NetGUID
-			//	-Invalid NetGUID (interpretted as delete)
-			Bunch << RepComp.Value()->ObjectNetGUID;
-
-			FNetworkGUID InvalidNetGUID;
-			InvalidNetGUID.Reset();
-			Bunch << InvalidNetGUID;
+			BeginContentBlockForSubObjectDelete( Bunch, RepComp.Value()->ObjectNetGUID );
 
 			WroteSomethingImportant = true;
 			Bunch.bReliable = true;
@@ -1871,6 +1865,23 @@ void UActorChannel::BeginContentBlock( UObject * Obj, FOutBunch &Bunch )
 #endif
 }
 
+void UActorChannel::BeginContentBlockForSubObjectDelete( FOutBunch & Bunch, FNetworkGUID & GuidToDelete )
+{
+	// Send a 0 bit to signify we are dealing with sub-objects
+	Bunch.WriteBit( 0 );
+
+	check( GuidToDelete.IsValid() );
+
+	//	-Deleted object's NetGUID
+	Bunch << GuidToDelete;
+	NET_CHECKSUM(Bunch);
+
+	//	-Invalid NetGUID (interpreted as delete)
+	FNetworkGUID InvalidNetGUID;
+	InvalidNetGUID.Reset();
+	Bunch << InvalidNetGUID;
+}
+
 void UActorChannel::EndContentBlock( UObject *Obj, FOutBunch &Bunch, FClassNetCache* ClassCache )
 {
 	check(Obj);
