Change the display name of OmniRadius to be Non-Spatialized Radius
Make Non-Spatialized Radius only editable if bSpatialize is true

diff --git a/Engine/Source/Runtime/Engine/Classes/Sound/SoundAttenuation.h b/Engine/Source/Runtime/Engine/Classes/Sound/SoundAttenuation.h
--- a/Engine/Source/Runtime/Engine/Classes/Sound/SoundAttenuation.h
+++ b/Engine/Source/Runtime/Engine/Classes/Sound/SoundAttenuation.h
@@ -73,8 +73,8 @@ struct ENGINE_API FAttenuationSettings
 	UPROPERTY(EditAnywhere, Category=Attenuation, meta=(DisplayName = "dB Attenuation At Max", ClampMax = "0" ))
 	float dBAttenuationAtMax;
 
-	/** At what distance we start transforming into omnidirectional soundsource */
-	UPROPERTY(EditAnywhere, Category=Attenuation, meta=(ClampMin = "0"))
+	/** At what distance we start treating the sound source as spatialized */
+	UPROPERTY(EditAnywhere, Category=Attenuation, meta=(ClampMin = "0", EditCondition="bSpatialize", DisplayName="Non-Spatialized Radius"))
 	float OmniRadius;
 
 	UPROPERTY()
