#UE4 - Fix URL parsing when there's no IPv6 address in the URL.

diff --git a/Engine/Source/Runtime/Engine/Private/URL.cpp b/Engine/Source/Runtime/Engine/Private/URL.cpp
--- a/Engine/Source/Runtime/Engine/Private/URL.cpp
+++ b/Engine/Source/Runtime/Engine/Private/URL.cpp
@@ -244,7 +244,7 @@ FURL::FURL( FURL* Base, const TCHAR* TextURL, ETravelType Type )
 			if
 			(	(FCString::Strchr(URL,':')!=NULL)
 			&&	(FCString::Strchr(URL,':')>URL+1)
-			&&  (FCString::Strchr(URL,':')<SquareBracket)
+			&&  (!SquareBracket || (FCString::Strchr(URL,':')<SquareBracket))
 			&&	(FCString::Strchr(URL,'.')==NULL || FCString::Strchr(URL,':')<FCString::Strchr(URL,'.')) )
 			{
 				TCHAR* ss = URL;
