Fixing up a crash from a null scene component (fallout from CL 2042611).

diff --git a/Engine/Source/Runtime/Engine/Private/ActorConstruction.cpp b/Engine/Source/Runtime/Engine/Private/ActorConstruction.cpp
--- a/Engine/Source/Runtime/Engine/Private/ActorConstruction.cpp
+++ b/Engine/Source/Runtime/Engine/Private/ActorConstruction.cpp
@@ -450,7 +450,7 @@ UActorComponent* AActor::AddComponent(FName TemplateName, bool bManualAttachment
 
 		// need to defer component registration until after the construction script
 		// is ran, since the construction script can mutate the object (for collision, etc.)
-		if (bRunningUserConstructionScript)
+		if (bRunningUserConstructionScript && (NewSceneComp != NULL))
 		{
 			// Defer registration until after UCS has completed.
 			FDeferRegisterComponents::Get().DeferComponentRegistration(this, NewSceneComp, OriginalMobility);
diff --git a/Engine/Source/Runtime/Engine/Private/SCS_Node.cpp b/Engine/Source/Runtime/Engine/Private/SCS_Node.cpp
--- a/Engine/Source/Runtime/Engine/Private/SCS_Node.cpp
+++ b/Engine/Source/Runtime/Engine/Private/SCS_Node.cpp
@@ -68,9 +68,16 @@ void USCS_Node::ExecuteNodeOnActor(AActor* Actor, USceneComponent* ParentCompone
 		// Call function to notify component it has been created
 		NewActorComp->OnComponentCreated();
 
-		// need to defer component registration until after the construction script
-		// is ran, since the construction script can mutate the object (for collision, etc.)
-		FDeferRegisterComponents::Get().DeferComponentRegistration(Actor, NewSceneComp, OriginalMobility);
+		if (NewSceneComp != NULL)
+		{
+			// need to defer component registration until after the construction script
+			// is ran, since the construction script can mutate the object (for collision, etc.)
+			FDeferRegisterComponents::Get().DeferComponentRegistration(Actor, NewSceneComp, OriginalMobility);
+		}
+		else 
+		{
+			NewActorComp->RegisterComponent();
+		}
 
 		if (NewActorComp->GetIsReplicated())
 		{
