#UE4 Change the Match state machine in GameMode to be explicit and based on the MatchState namespace, as opposed to being based on bools. Refactor related code to be clearer and easier to extend. States:

EnteringMap is the initial state. Actors are not yet ticking and the world hasn't been fully initialized. It will transition to the next state when things are fully loaded.
WaitingToStart (was bWaitingToStartMatch) is the next state, and HandleMatchIsWaitingToStart is called when entering. Actors are ticking, but players have not yet spawned in. It transitions to the next state if ReadyToStartMatch returns true, or if someone calls StartMatch.
InProgress (was bMatchIsInProgress) is the next state, and HandleMatchHasStarted (split off from StartMatch) is called when entering, which calls BeginPlay on all actors. Normal gameplay is in progress. It transitions to the next state if ReadyToEndMatch returns true, or if someone calls EndMatch.
WaitingPostMatch is the next state, and HandleMatchHasEnded (was PerformEndGameHandling) is called when entering. Actors are still ticking but new players are not accepted. It transitions to the next state when a map transfer starts.
LeavingMap (was bLevelChange) is the last state, and HandleLeavingMap (was GameEnding) is called when entering. The match stays in this state while transfering to a new map.
Aborted is the failure state, started from AbortMatch. This is set when there is an unrecoverable error.

Rename functions in GameState and GameSession to match the names above. StartPendingMatch is now HandleMatchIsWaitingToStart, EndPendingMatch/StartMatch is HandleMatchHasStarted, EndGame is HandleMatchHasEnded.
Split the spectator half of functionality of bDelayedStart into a new flag, bStartPlayersAsSpectators. If true, newly joining players will spawn as spectators, if false they'll try to spawn pawns. bDelayedStart now only affects rather the match should start as soon as one player joins, if true the match will not start.
Clean up GameMode::PlayerCanRestart, so it calls PlayerController::CanRestartPlayer, and change all places to call the GameMode. So now if either function returns false the player cannot restart, and games can cleanly override this behavior.
Change GameState::ElapsedTime to only increment when the match has started, the behavior before was inconsistent depending on if you were a client or server
Modify sample games to work with new changes

diff --git a/Engine/Source/Runtime/Engine/Classes/GameFramework/GameMode.h b/Engine/Source/Runtime/Engine/Classes/GameFramework/GameMode.h
--- a/Engine/Source/Runtime/Engine/Classes/GameFramework/GameMode.h
+++ b/Engine/Source/Runtime/Engine/Classes/GameFramework/GameMode.h
@@ -9,6 +9,19 @@
 //  instanced on the server and will never exist on the client.
 //=============================================================================
 
+/** Possible state of the current match, where a match is all the gameplay that happens on a single map */
+namespace MatchState
+{
+	extern ENGINE_API const FName EnteringMap;			// We are entering this map, actors are not yet ticking
+	extern ENGINE_API const FName WaitingToStart;		// Actors are ticking, but the match has not yet started
+	extern ENGINE_API const FName InProgress;			// Normal gameplay is occurring. Specific games will have their own state machine inside this state
+	extern ENGINE_API const FName WaitingPostMatch;		// Match has ended so we aren't accepting new players, but actors are still ticking
+	extern ENGINE_API const FName LeavingMap;			// We are transitioning out of the map to another location
+	extern ENGINE_API const FName Aborted;				// Match has failed due to network issues or other problems, cannot continue
+
+	// If a game needs to add additional states, you may need to override HasMatchStarted and HasMatchEnded to deal with the new states
+	// Do not add any states before WaitingToStart or after WaitingPostMatch
+}
 
 /** Default delegate that provides an implementation for those that don't have special needs other than a toggle */
 DECLARE_DELEGATE_RetVal( bool, FCanUnpause );
@@ -42,26 +55,91 @@ class ENGINE_API AGameMode : public AInfo
 {
 	GENERATED_UCLASS_BODY()
 
-public:
-	/** Whether the match is currently in progress */
+	// Code to deal with the match state machine
+
+	/** Returns the current match state, this is an accessor to protect the state machine flow */
+	FName GetMatchState() const { return MatchState; }
+
+	/** Returns true if the match state is InProgress or later */
+	virtual bool HasMatchStarted() const;
+
+	/** Returns true if the match state is InProgress or other gameplay state */
+	virtual bool IsMatchInProgress() const;
+
+	/** Returns true if the match state is WaitingPostMatch or later */
+	virtual bool HasMatchEnded() const;
+
+
+	/** Transition from WaitingToStart to InProgress. You can call this manually, will also get called if ReadyToStartMatch returns true */
+	virtual void StartMatch();
+
+	/** Transition from InProgress to WaitingPostMatch. You can call this manually, will also get called if ReadyToEndMatch returns true */
+	virtual void EndMatch();
+
+	/** Transition to LeavingMap state. Start the transition out of the current map. Called at start of seamless travel, or right before map change for hard travel. */
+	virtual void StartToLeaveMap();
+
+	/** Restart the game, by default travel to the current map */
+	virtual void RestartGame();
+
+	/** Return to main menu, and disconnect any players */
+	virtual void ReturnToMainMenuHost();
+
+	/** Report that a match has failed due to unrecoverable error */
+	virtual void AbortMatch();
+
+protected:
+
+	/** What match state we are currently in */
 	UPROPERTY()
-	uint32 bMatchIsInProgress:1;
+	FName MatchState;
 
-	/** Whether the game is pauseable. */
+	/** Updates the match state and calls the appropriate transition functions */
+	virtual void SetMatchState(FName NewState);
+
+	// Games should override these functions to deal with their game specific logic
+
+	/** Called when the state transitions to WaitingToStart */
+	virtual void HandleMatchIsWaitingToStart();
+
+	/** @return True if ready to Start Match. Games should override this */
+	virtual bool ReadyToStartMatch();
+
+	/** Called when the state transitions to InProgress */
+	virtual void HandleMatchHasStarted();
+
+	/** @return true if ready to End Match. Games should override this */
+	virtual bool ReadyToEndMatch();
+
+	/** Called when the map transitions to WaitingPostMatch */
+	virtual void HandleMatchHasEnded();
+
+	/** Called when the match transitions to LeavingMap */
+	virtual void HandleLeavingMap();
+
+	/** Called when the match transitions to Aborted */
+	virtual void HandleMatchAborted();
+
+public:
+
+	/**
+	 * perform map travels using SeamlessTravel() which loads in the background and doesn't disconnect clients
+	 * @see World::SeamlessTravel()
+	 */
 	UPROPERTY()
-	uint32 bPauseable:1;    
+	uint32 bUseSeamlessTravel : 1;
 
-	/** Match has not yet started - waiting for players to join and be ready. */
+	/** Whether the game is pauseable. */
 	UPROPERTY()
-	uint32 bWaitingToStartMatch:1; 
+	uint32 bPauseable:1;
 
-	/** Flag to make sure RestartGame() only gets executed once. */
+	/** Whether players should immediately spawn when logging in, or stay as spectators until they manually spawn */
 	UPROPERTY()
-	uint32 bGameRestarted:1; 
+	uint32 bStartPlayersAsSpectators : 1;
 
-	/** Level transition in progress. */
+	/** Whether the game should immediately start when the first player logs in. Affects the default behavior of ReadyToStartMatch */
 	UPROPERTY()
-	uint32 bLevelChange:1;    
+	uint32 bDelayedStart : 1;
 
 	/** Save options string and parse it when needed */
 	UPROPERTY()
@@ -149,20 +227,6 @@ public:
 	/** The list of delegates to check before unpausing a game */
 	TArray<FCanUnpause> Pausers;
 
-	/** 
-	 * perform map travels using SeamlessTravel() which loads in the background and doesn't disconnect clients
-	 * @see World::SeamlessTravel()
-	 */
-	UPROPERTY()
-	uint32 bUseSeamlessTravel:1;
-
-	/** Tracks whether the server can travel due to a critical network error. */
-	UPROPERTY()
-	uint32 bHasNetworkError:1;
-
-	/** Whether the game should immediately start when the player logs in. */
-	uint32 bDelayedStart:1;
-
 protected:
 
 	/** Handy alternate short names for GameMode classes (e.g. "DM" could be an alias for "MyProject.MyGameModeMP_DM". */
@@ -170,8 +234,6 @@ protected:
 	TArray<struct FGameClassShortName> GameModeClassAliases;
 
 public:
-	/** Does end of game handling for the online layer. */
-	virtual void PerformEndGameHandling();
 
 	/** Alters the synthetic bandwidth limit for a running game. */
 	UFUNCTION(exec)
@@ -179,11 +241,11 @@ public:
 
 	// Begin AActor interface
 	virtual void PreInitializeComponents() OVERRIDE;
-	virtual void PostInitializeComponents() OVERRIDE;
-	virtual void AddPlayerStart(APlayerStart* NewPlayerStart);
-	virtual void RemovePlayerStart(APlayerStart* RemovedPlayerStart);
+	/** This sets the match state to WaitingToStart unless we're in seamless travel */
+	virtual void PostInitializeComponents() OVERRIDE; 
 	virtual void DisplayDebug(class UCanvas* Canvas, const TArray<FName>& DebugDisplay, float& YL, float& YPos) OVERRIDE;
 	virtual void Reset() OVERRIDE;
+	virtual void Tick(float DeltaSeconds) OVERRIDE;
 	// End AActor interface
 
 	/** 
@@ -192,9 +254,6 @@ public:
 	 */
 	virtual bool ShouldReset(AActor* ActorToReset);
 
-	/** Called when StartSpot is selected for spawning NewPlayer to allow optional initialization. */
-	virtual void InitStartSpot(class AActor* StartSpot, AController* NewPlayer);
-
 	/** Resets level by calling Reset() on all actors */
 	virtual void ResetLevel();
 
@@ -209,8 +268,13 @@ public:
 	 */
 	bool ShouldStartInCinematicMode(bool& OutHidePlayer, bool& OutHideHud, bool& OutDisableMovement, bool& OutDisableTurning);
 
-	/** Called right before exiting or performing a server travel. */
-	virtual void GameEnding();
+	/**
+	 * Initialize the game.
+	 * The GameMode's InitGame() event is called before any other functions (including PreInitializeComponents() )
+	 * and is used by the GameMode to initialize parameters and spawn its helper classes.
+	 * @warning: this is called before actors' PreInitializeComponents.
+	 */
+	virtual void InitGame(const FString& MapName, const FString& Options, FString& ErrorMessage);
 
 	/**
 	 * Initialize the GameState actor with default settings
@@ -291,19 +355,13 @@ public:
 	/** @return GameSession class to use for this game  */
 	virtual TSubclassOf<class AGameSession> GetGameSessionClass() const;
 
-	/**
-	 * Initialize the game.
-	 * The GameMode's InitGame() event is called before any other scripts (including
-	 * PreInitializeComponents() ), and is used by the GameMode to initialize parameters and spawn
-	 * its helper classes.
-	 * @warning: this is called before actors' PreInitializeComponents.
-	 */
-	virtual void InitGame( const FString& MapName, const FString& Options, FString& ErrorMessage );
-
 	/** Called when a connection closes before getting to PostLogin() */
 	virtual void NotifyPendingConnectionLost();
 
-	/* Optional handling of ServerTravel for network games. */
+	/** @return true if we want to travel_absolute */
+	virtual bool GetTravelType();
+
+	/** Optional handling of ServerTravel for network games. */
 	virtual void ProcessServerTravel(const FString& URL, bool bAbsolute = false);
 
 	/**
@@ -343,12 +401,18 @@ public:
 	/** @Returns true if NewPlayer may only join the server as a spectator. */
 	virtual bool MustSpectate(APlayerController* NewPlayer);
 
-	/* Start the game - inform all actors that the match is starting, and spawn player pawns. */
-	virtual void StartMatch();
-
 	/** returns default pawn class for given controller */
 	virtual UClass* GetDefaultPawnClassForController(AController* InController);
 
+	/** Called when StartSpot is selected for spawning NewPlayer to allow optional initialization. */
+	virtual void InitStartSpot(class AActor* StartSpot, AController* NewPlayer);
+
+	/** Register new player start */
+	virtual void AddPlayerStart(APlayerStart* NewPlayerStart);
+
+	/** Remove an existing player start */
+	virtual void RemovePlayerStart(APlayerStart* RemovedPlayerStart);
+
 	/**
 	 * @param	NewPlayer - Controller for whom this pawn is spawned
 	 * @param	StartSpot - PlayerStart at which to spawn pawn
@@ -389,15 +453,6 @@ public:
 	/* Send a player to a URL.*/
 	virtual void SendPlayer( APlayerController* aPlayer, const FString& URL );
 
-	/** @return true if we want to travel_absolute */
-	virtual bool GetTravelType();
-
-	/** Restart the game */
-	virtual void RestartGame();
-
-	/** Return to main menu */
-	virtual void ReturnToMainMenuHost();
-
 	/** Broadcast a string to all players. */
 	virtual void Broadcast( AActor* Sender, const FString& Msg, FName Type = NAME_None );
 
@@ -427,11 +482,8 @@ public:
 	 */
 	virtual class AActor* ChoosePlayerStart( AController* Player );
 
-	/** @return true if player can restart the game */
-	virtual bool PlayerCanRestartGame( APlayerController* aPlayer );
-
-	/** @return true if player can be restarted */
-	virtual bool PlayerCanRestart( APlayerController* aPlayer );
+	/** @return true if it's valid to call RestartPlayer. Will call Player->CanRestartPlayer */
+	virtual bool PlayerCanRestart( APlayerController* Player );
 
 	/**
 	 * Used to notify the game type that it is ok to update a player's gameplay
@@ -450,13 +502,13 @@ public:
 	virtual bool AllowPausing( APlayerController* PC = NULL );
 
 	/**
-	 * Called from CommitMapChange before unloading previous level
+	 * Called from CommitMapChange before unloading previous level. Used for asynchronous level streaming
 	 * @param PreviousMapName - Name of the previous persistent level
 	 * @param NextMapName - Name of the persistent level being streamed to
 	 */
 	virtual void PreCommitMapChange(const FString& PreviousMapName, const FString& NextMapName);
 
-	/** Called from CommitMapChange after unloading previous level and loading new level+sublevels */
+	/** Called from CommitMapChange after unloading previous level and loading new level+sublevels. Used for asynchronous level streaming */
 	virtual void PostCommitMapChange();
 
 	/** Add PlayerState to the inactive list, remove from the active list */
@@ -482,7 +534,8 @@ public:
 	 */
 	virtual void GetSeamlessTravelActorList(bool bToEntry, TArray<AActor*>& ActorList);
 
-	/** used to swap a viewport/connection's PlayerControllers when seamless travelling and the new GameMode's
+	/** 
+	 * used to swap a viewport/connection's PlayerControllers when seamless travelling and the new GameMode's
 	 * controller class is different than the previous
 	 * includes network handling
 	 * @param OldPC - the old PC that should be discarded
@@ -490,14 +543,12 @@ public:
 	 */
 	virtual void SwapPlayerControllers(APlayerController* OldPC, APlayerController* NewPC);
 
-	/** called after a seamless level transition has been completed on the *new* GameMode
-	 * used to reinitialize players already in the game as they won't have *Login() called on them
+	/** 
+	 * Called after a seamless level transition has been completed on the *new* GameMode. This transitions to WaitingToStart
+	 * Used to reinitialize players already in the game as they won't have *Login() called on them
 	 */
 	virtual void PostSeamlessTravel();
 
-	/** @return true if ready to Start Match */
-	virtual bool ReadyToStartMatch();
-
 	/** 
 	 * Handles reinitializing players that remained through a seamless level transition
 	 * called from C++ for players that finished loading after the server
diff --git a/Engine/Source/Runtime/Engine/Classes/GameFramework/GameSession.h b/Engine/Source/Runtime/Engine/Classes/GameFramework/GameSession.h
--- a/Engine/Source/Runtime/Engine/Classes/GameFramework/GameSession.h
+++ b/Engine/Source/Runtime/Engine/Classes/GameFramework/GameSession.h
@@ -142,20 +142,20 @@ class ENGINE_API AGameSession : public AInfo
 	//=================================================================================
 	// MATCH INTERFACE
 
-	/** Start a pending match */
-	virtual void StartPendingMatch();
+	/** @RETURNS true if GameSession handled the request, in case it wants to stall for some reason. Otherwise, game mode will start immediately */
+	virtual bool HandleStartMatchRequest();
 
-	/** End a pending match */
-	virtual void EndPendingMatch();
+	/** Handle when the match enters waiting to start */
+	virtual void HandleMatchIsWaitingToStart();
 
-	/** Called from GameMode.RestartGame(). */
-	virtual bool CanRestartGame();
+	/** Handle when the match has started */
+	virtual void HandleMatchHasStarted();
 
-	/** End a game */
-	virtual void EndGame();
+	/** Handle when the match has completed */
+	virtual void HandleMatchHasEnded();
 
-	/** @RETURNS true if GameSession handled starting the match */
-	virtual bool HandleStartMatch();
+	/** Called from GameMode.RestartGame(). */
+	virtual bool CanRestartGame();
 
 private:
 	// Hidden functions that don't make sense to use on this class.
diff --git a/Engine/Source/Runtime/Engine/Classes/GameFramework/GameState.h b/Engine/Source/Runtime/Engine/Classes/GameFramework/GameState.h
--- a/Engine/Source/Runtime/Engine/Classes/GameFramework/GameState.h
+++ b/Engine/Source/Runtime/Engine/Classes/GameFramework/GameState.h
@@ -25,15 +25,44 @@ class ENGINE_API AGameState : public AInfo
 	UPROPERTY(replicatedUsing=OnRep_SpectatorClass)
 	TSubclassOf<class ASpectatorPawn> SpectatorClass;
 
-	/** Match is in progress. */
-	UPROPERTY(replicatedUsing=OnRep_bMatchHasBegun)
-	uint32 bMatchHasBegun:1;
+	// Code to deal with the match state machine
 
-	/** Match is over. */
-	UPROPERTY(replicatedUsing=OnRep_bMatchIsOver)
-	uint32 bMatchIsOver:1;
+	/** Returns the current match state, this is an accessor to protect the state machine flow */
+	FName GetMatchState() const { return MatchState; }
 
-	/** Elapsed game time since GameState was created. */
+	/** Returns true if the match state is InProgress or later */
+	virtual bool HasMatchStarted() const;
+
+	/** Returns true if we're in progress */
+	virtual bool IsMatchInProgress() const;
+
+	/** Returns true if match is WaitingPostMatch or later */
+	virtual bool HasMatchEnded() const;
+
+	/** Updates the match state and calls the appropriate transition functions, only valid on server */
+	void SetMatchState(FName NewState);
+
+protected:
+
+	/** What match state we are currently in */
+	UPROPERTY(replicatedUsing=OnRep_MatchState)
+	FName MatchState;
+
+	/** Called when the state transitions to WaitingToStart */
+	virtual void HandleMatchIsWaitingToStart();
+
+	/** Called when the state transitions to InProgress */
+	virtual void HandleMatchHasStarted();
+
+	/** Called when the map transitions to WaitingPostMatch */
+	virtual void HandleMatchHasEnded();
+
+	/** Called when the match transitions to LeavingMap */
+	virtual void HandleLeavingMap();
+
+public:
+
+	/** Elapsed game time since match has started. */
 	UPROPERTY(replicated, BlueprintReadOnly, Category=GameState)
 	int32 ElapsedTime;
 
@@ -53,13 +82,9 @@ class ENGINE_API AGameState : public AInfo
 	UFUNCTION()
 	virtual void OnRep_SpectatorClass();
 
-	/** Match begun Notification Callback */
+	/** Match state has changed */
 	UFUNCTION()
-	virtual void OnRep_bMatchHasBegun();
-
-	/** Match Is Over Notification Callback */
-	UFUNCTION()
-	virtual void OnRep_bMatchIsOver();
+	virtual void OnRep_MatchState();
 
 	// Begin AActor interface
 	virtual void PostInitializeComponents() OVERRIDE;
@@ -83,11 +108,6 @@ class ENGINE_API AGameState : public AInfo
 	/** Remove PlayerState from the PlayerArray. */
 	virtual void RemovePlayerState(class APlayerState* PlayerState);
 
-	/** Called on the server when the match has begun. */
-	virtual void StartMatch();
-
-	/** Called on the server when the match is over. */
-	virtual void EndGame();
 
 	/** Called periodically, overridden by subclasses */
 	virtual void DefaultTimer();
diff --git a/Engine/Source/Runtime/Engine/Classes/GameFramework/PlayerController.h b/Engine/Source/Runtime/Engine/Classes/GameFramework/PlayerController.h
--- a/Engine/Source/Runtime/Engine/Classes/GameFramework/PlayerController.h
+++ b/Engine/Source/Runtime/Engine/Classes/GameFramework/PlayerController.h
@@ -1094,7 +1094,7 @@ public:
 	 */
 	virtual void ViewAPlayer(int32 dir);
 
-	/** @return true if game allows this player to spawn */
+	/** @return true if this controller thinks it's able to restart. Called from GameMode::PlayerCanRestart */
 	virtual bool CanRestartPlayer();
 
 	/**
diff --git a/Engine/Source/Runtime/Engine/Classes/GameFramework/WorldSettings.h b/Engine/Source/Runtime/Engine/Classes/GameFramework/WorldSettings.h
--- a/Engine/Source/Runtime/Engine/Classes/GameFramework/WorldSettings.h
+++ b/Engine/Source/Runtime/Engine/Classes/GameFramework/WorldSettings.h
@@ -423,7 +423,7 @@ public:
 		return TimeDilation * MatineeTimeDilation;
 	}
 
-	/** Called by GameMode.StartMatch, used to notify native classes of match startup (such as level scripting).
+	/** Called by GameMode.HandleMatchHasStarted, used to notify native classes of match startup (such as level scripting).
 	 * @todo UE4 Need to send K2 event about match starting here
 	 * Activates LevelStartup and/or LevelBeginning events in the sequences in the world
 	 */	
diff --git a/Engine/Source/Runtime/Engine/Private/Actor.cpp b/Engine/Source/Runtime/Engine/Private/Actor.cpp
--- a/Engine/Source/Runtime/Engine/Private/Actor.cpp
+++ b/Engine/Source/Runtime/Engine/Private/Actor.cpp
@@ -2166,7 +2166,7 @@ void AActor::PostNetInit()
 	check(RemoteRole == ROLE_Authority);
 
 	AGameState *GameState = GetWorld()->GameState;
-	if (GameState && GameState->bMatchHasBegun)
+	if (GameState && GameState->HasMatchStarted())
 	{
 		BeginPlay();
 	}
diff --git a/Engine/Source/Runtime/Engine/Private/GameEngine.cpp b/Engine/Source/Runtime/Engine/Private/GameEngine.cpp
--- a/Engine/Source/Runtime/Engine/Private/GameEngine.cpp
+++ b/Engine/Source/Runtime/Engine/Private/GameEngine.cpp
@@ -699,7 +699,7 @@ bool UGameEngine::HandleExitCommand( const TCHAR* Cmd, FOutputDevice& Ar )
 
 		if (GameMode)
 		{
-			GameMode->GameEnding();
+			GameMode->StartToLeaveMap();
 		}
 
 		// Cancel any pending connection to a server
diff --git a/Engine/Source/Runtime/Engine/Private/GameMode.cpp b/Engine/Source/Runtime/Engine/Private/GameMode.cpp
--- a/Engine/Source/Runtime/Engine/Private/GameMode.cpp
+++ b/Engine/Source/Runtime/Engine/Private/GameMode.cpp
@@ -10,11 +10,22 @@
 
 DEFINE_LOG_CATEGORY_STATIC(LogGameMode, Log, All);
 
+namespace MatchState
+{
+	const FName EnteringMap = FName(TEXT("EnteringMap"));
+	const FName WaitingToStart = FName(TEXT("WaitingToStart"));
+	const FName InProgress = FName(TEXT("InProgress"));
+	const FName WaitingPostMatch = FName(TEXT("WaitingPostMatch"));
+	const FName LeavingMap = FName(TEXT("LeavingMap"));
+	const FName Aborted = FName(TEXT("Aborted"));
+}
+
 AGameMode::AGameMode(const class FPostConstructInitializeProperties& PCIP)
 	: Super(PCIP
 		.DoNotCreateDefaultSubobject(TEXT("Sprite"))
 	)
 {
+	bStartPlayersAsSpectators = false;
 	bDelayedStart = false;
 	bNetLoadOnClient = false;
 
@@ -22,7 +33,7 @@ AGameMode::AGameMode(const class FPostConstructInitializeProperties& PCIP)
 	PrimaryActorTick.bCanEverTick = true;
 	PrimaryActorTick.TickGroup = TG_PrePhysics;
 	HUDClass = AHUD::StaticClass();
-	bWaitingToStartMatch = false;
+	MatchState = MatchState::EnteringMap;
 	bPauseable = true;
 	DefaultPawnClass = ADefaultPawn::StaticClass();
 	PlayerControllerClass = APlayerController::StaticClass();
@@ -34,14 +45,12 @@ AGameMode::AGameMode(const class FPostConstructInitializeProperties& PCIP)
 	MinRespawnDelay = 1.0f;
 }
 
-
 FString AGameMode::GetNetworkNumber()
 {
 	UNetDriver* NetDriver = GetWorld()->GetNetDriver();
 	return NetDriver ? NetDriver->LowLevelGetNetworkNumber() : FString(TEXT(""));
 }
 
-
 void AGameMode::SwapPlayerControllers(APlayerController* OldPC, APlayerController* NewPC)
 {
 	if (OldPC != NULL && !OldPC->IsPendingKill() && NewPC != NULL && !NewPC->IsPendingKill() && OldPC->Player != NULL)
@@ -120,7 +129,7 @@ void AGameMode::ForceClearUnpauseDelegates( AActor* PauseActor )
 	}
 }
 
-void AGameMode::InitGame( const FString& MapName, const FString& Options, FString& ErrorMessage )
+void AGameMode::InitGame(const FString& MapName, const FString& Options, FString& ErrorMessage)
 {
 	UWorld* World = GetWorld();
 
@@ -130,7 +139,7 @@ void AGameMode::InitGame( const FString& MapName, const FString& Options, FStrin
 	UClass* const SessionClass = GetGameSessionClass();
 	FActorSpawnParameters SpawnInfo;
 	SpawnInfo.Instigator = Instigator;
-	GameSession = World->SpawnActor<AGameSession>( SessionClass, SpawnInfo );
+	GameSession = World->SpawnActor<AGameSession>(SessionClass, SpawnInfo);
 	GameSession->InitOptions(Options);
 
 	if (GetNetMode() != NM_Standalone)
@@ -147,6 +156,8 @@ void AGameMode::InitGame( const FString& MapName, const FString& Options, FStrin
 			GameSession->RegisterServer();
 		}
 	}
+
+	SetMatchState(MatchState::EnteringMap);
 }
 
 bool AGameMode::SetPause(APlayerController* PC, FCanUnpause CanUnpauseDelegate)
@@ -174,11 +185,10 @@ void AGameMode::RestartGame()
 {
 	if ( GameSession->CanRestartGame() )
 	{
-		if (bGameRestarted)
+		if (GetMatchState() == MatchState::LeavingMap)
 		{
 			return;
 		}
-		bGameRestarted = true;
 
 		GetWorld()->ServerTravel("?Restart",GetTravelType());
 	}
@@ -399,27 +409,16 @@ void AGameMode::PostInitializeComponents()
 {
 	Super::PostInitializeComponents();
 
-	bWaitingToStartMatch = true;
-	bMatchIsInProgress = false;
-
-	if(GameSession != NULL)
+	if (!GetWorld()->IsInSeamlessTravel())
 	{
-		GameSession->StartPendingMatch();
+		// If this isn't seamless travel, enter waiting state immediately. PostSeamlessTravel handles the state transition otherwise
+		SetMatchState(MatchState::WaitingToStart);
 	}
-
-	// validate default classes
-// 	check(DefaultPawnClass && DefaultPawnClass->IsChildOf(APawn::StaticClass()));
-// 	check(PlayerControllerClass && PlayerControllerClass->IsChildOf(APlayerController::StaticClass()));
-// 	check(SpectatorClass && SpectatorClass->IsChildOf(ASpectatorPawn::StaticClass()));
-// 	check(EngineMessageClass && EngineMessageClass->IsChildOf(UEngineMessage::StaticClass()));
-// 	check(GameStateClass && GameStateClass->IsChildOf(AGameState::StaticClass()));
-// 	check(PlayerStateClass && PlayerStateClass->IsChildOf(APlayerState::StaticClass()));
 }
 
-
 void AGameMode::RestartPlayer(AController* NewPlayer)
 {
-	if ( bWaitingToStartMatch || NewPlayer == NULL || NewPlayer->IsPendingKillPending() )
+	if ( NewPlayer == NULL || NewPlayer->IsPendingKillPending() )
 	{
 		return;
 	}
@@ -494,42 +493,58 @@ void AGameMode::InitStartSpot(AActor* StartSpot, AController* NewPlayer)
 {
 }
 
+void AGameMode::HandleMatchIsWaitingToStart()
+{
+	if (GameSession != NULL)
+	{
+		GameSession->HandleMatchIsWaitingToStart();
+	}
+}
+
+bool AGameMode::ReadyToStartMatch()
+{
+	// If bDelayed Start is set, wait for a manual match start
+	if (bDelayedStart)
+	{
+		return false;
+	}
+
+	// By default start when we have > 0 players
+	return (GetMatchState() == MatchState::WaitingToStart && NumPlayers + NumBots > 0);
+}
+
 void AGameMode::StartMatch()
 {
-	/**
-	 * Tells all of the currently connected clients to register with arbitration.
-	 * The clients will call back to the server once they have done so, which
-	 * will tell this state to see if it is time for the server to register with
-	 * arbitration.
-	 */
-	if (!bMatchIsInProgress && GameSession->HandleStartMatch())
+	if (HasMatchStarted())
 	{
+		// Already started
 		return;
 	}
 
-	if (bWaitingToStartMatch)
+	//Let the game session override the StartMatch function, in case it wants to wait for arbitration
+
+	if (GameSession->HandleStartMatchRequest())
 	{
-		GameSession->EndPendingMatch();
+		return;
 	}
 
-	bWaitingToStartMatch = false;
+	SetMatchState(MatchState::InProgress);
+}
+
+void AGameMode::HandleMatchHasStarted()
+{
+	GameSession->HandleMatchHasStarted();
 
 	// start human players first
 	for( FConstPlayerControllerIterator Iterator = GetWorld()->GetPlayerControllerIterator(); Iterator; ++Iterator )
 	{
 		APlayerController* PlayerController = *Iterator;
-		if( ( PlayerController->GetPawn() == NULL) && PlayerController->CanRestartPlayer() )
+		if ((PlayerController->GetPawn() == NULL) && PlayerCanRestart(PlayerController))
 		{
 			RestartPlayer(PlayerController);
 		}
 	}
 
-	// Notify all clients that the match has begun
-	if(GetWorld()->GameState != NULL)
-	{
-		GetWorld()->GameState->StartMatch();
-	}
-
 	// Make sure level streaming is up to date before triggering NotifyMatchStarted
 	GetWorld()->FlushLevelStreaming();
 
@@ -555,6 +570,134 @@ void AGameMode::StartMatch()
 	}
 }
 
+bool AGameMode::ReadyToEndMatch()
+{
+	// By default don't explicitly end match
+	return false;
+}
+
+void AGameMode::EndMatch()
+{
+	if (!IsMatchInProgress())
+	{
+		return;
+	}
+
+	SetMatchState(MatchState::WaitingPostMatch);
+}
+
+void AGameMode::HandleMatchHasEnded()
+{
+	GameSession->HandleMatchHasEnded();
+}
+
+void AGameMode::StartToLeaveMap()
+{
+	SetMatchState(MatchState::LeavingMap);
+}
+
+void AGameMode::HandleLeavingMap()
+{
+
+}
+
+void AGameMode::AbortMatch()
+{
+	SetMatchState(MatchState::Aborted);
+}
+
+void AGameMode::HandleMatchAborted()
+{
+
+}
+
+bool AGameMode::HasMatchStarted() const
+{
+	if (GetMatchState() == MatchState::EnteringMap || GetMatchState() == MatchState::WaitingToStart)
+	{
+		return false;
+	}
+
+	return true;
+}
+
+bool AGameMode::IsMatchInProgress() const
+{
+	if (GetMatchState() == MatchState::InProgress)
+	{
+		return true;
+	}
+
+	return false;
+}
+
+bool AGameMode::HasMatchEnded() const
+{
+	if (GetMatchState() == MatchState::WaitingPostMatch || GetMatchState() == MatchState::LeavingMap)
+	{
+		return true;
+	}
+
+	return false;
+}
+
+void AGameMode::SetMatchState(FName NewState)
+{
+	if (MatchState == NewState)
+	{
+		return;
+	}
+
+	MatchState = NewState;
+
+	// Call change callbacks
+
+	if (MatchState == MatchState::WaitingToStart)
+	{
+		HandleMatchIsWaitingToStart();
+	}
+	else if (MatchState == MatchState::InProgress)
+	{
+		HandleMatchHasStarted();
+	}
+	else if (MatchState == MatchState::WaitingPostMatch)
+	{
+		HandleMatchHasEnded();
+	}
+	else if (MatchState == MatchState::LeavingMap)
+	{
+		HandleLeavingMap();
+	}
+	else if (MatchState == MatchState::Aborted)
+	{
+		HandleMatchAborted();
+	}
+
+	if (GameState)
+	{
+		GameState->SetMatchState(NewState);
+	}
+}
+
+void AGameMode::Tick(float DeltaSeconds)
+{
+	if (GetMatchState() == MatchState::WaitingToStart)
+	{
+		// Check to see if we should start the match
+		if (ReadyToStartMatch())
+		{
+			StartMatch();
+		}
+	}
+	if (GetMatchState() == MatchState::InProgress)
+	{
+		// Check to see if we should start the match
+		if (ReadyToEndMatch())
+		{
+			EndMatch();
+		}
+	}
+}
 
 void AGameMode::ResetLevel()
 {
@@ -680,9 +823,6 @@ void AGameMode::HandleSeamlessTravelPlayer(AController*& C)
 			PC->bPlayerIsWaiting = true;
 			PC->ChangeState(NAME_Spectating);
 		}
-
-		// tell client what hud class to use
-		PC->ClientSetHUD(HUDClass);
 	}
 	else
 	{
@@ -691,11 +831,10 @@ void AGameMode::HandleSeamlessTravelPlayer(AController*& C)
 
 	GenericPlayerInitialization(C);
 
-	// RestartPlayer if match has started. This fixes bugs if the seamless travel player loads slow and
-	// the server calls StartMatch before he loads.
-	if (!bWaitingToStartMatch && PC)
+	if (PC)
 	{
-		RestartPlayer(PC);
+		// This handles setting hud, and may spawn the player pawn if the game is in progress
+		StartNewPlayer(PC);
 	}
 
 	UE_LOG(LogGameMode, Log, TEXT("<< GameMode::HandleSeamlessTravelPlayer: %s"), *C->GetName());
@@ -710,7 +849,7 @@ void AGameMode::SetSeamlessTravelViewTarget(APlayerController* PC)
 void AGameMode::ProcessServerTravel(const FString& URL, bool bAbsolute)
 {
 #if WITH_SERVER_CODE
-	bLevelChange = true;
+	StartToLeaveMap();
 
 	// force an old style load screen if the server has been up for a long time so that TimeSeconds doesn't overflow and break everything
 	bool bSeamless = (bUseSeamlessTravel && GetWorld()->TimeSeconds < 172800.0f); // 172800 seconds == 48 hours
@@ -855,15 +994,6 @@ APlayerController* AGameMode::Login(const FString& Portal, const FString& Option
 		return NewPlayer;
 	}
 
-	if (bDelayedStart)
-	{
-		NewPlayer->bPlayerIsWaiting = true;
-		NewPlayer->ChangeState(NAME_Spectating);
-		return NewPlayer;
-	}
-
-
-
 	return NewPlayer;
 }
 
@@ -886,8 +1016,6 @@ bool AGameMode::ShouldReset(AActor* ActorToReset)
 	return true;
 }
 
-void AGameMode::GameEnding() {}
-
 void AGameMode::PlayerSwitchedToSpectatorOnly(APlayerController* PC)
 {
 	RemovePlayerControllerFromPlayerCount(PC);
@@ -1183,29 +1311,24 @@ void AGameMode::StartNewPlayer(APlayerController* NewPlayer)
 	// tell client what hud class to use
 	NewPlayer->ClientSetHUD(HUDClass);
 
-	if (!bDelayedStart)
+	// If players should start as spectators, leave them in the spectator state
+	if (!bStartPlayersAsSpectators)
 	{
-		// start match, or let player enter, immediately
-		if ( bWaitingToStartMatch )
-		{
-			StartMatch();
-		}
-		else
+		// If match is in progress, start the player. Otherwise match will start from the game mode tick
+		if (IsMatchInProgress())
 		{
 			RestartPlayer(NewPlayer);
-		}
 
-		if (NewPlayer->GetPawn() != NULL)
-		{
-			NewPlayer->GetPawn()->ClientSetRotation(NewPlayer->GetPawn()->GetActorRotation());
+			if (NewPlayer->GetPawn() != NULL)
+			{
+				NewPlayer->GetPawn()->ClientSetRotation(NewPlayer->GetPawn()->GetActorRotation());
+			}
 		}
 	}
 }
 
 void AGameMode::PreExit() {}
 
-
-
 bool AGameMode::CanSpectate( APlayerController* Viewer, APlayerState* ViewTarget )
 {
 	return true;
@@ -1257,12 +1380,6 @@ void AGameMode::BroadcastLocalized( AActor* Sender, TSubclassOf<ULocalMessage> M
 	}
 }
 
-void AGameMode::PerformEndGameHandling()
-{
-	GameState->EndGame();
-	GameSession->EndGame();
-}
-
 bool AGameMode::ShouldSpawnAtStartSpot(AController* Player)
 {
 	return ( Player != NULL && Player->StartSpot != NULL );
@@ -1300,14 +1417,15 @@ AActor* AGameMode::ChoosePlayerStart( AController* Player )
 	return FoundPlayerStart;
 }
 
-bool AGameMode::PlayerCanRestartGame( APlayerController* aPlayer )
+bool AGameMode::PlayerCanRestart( APlayerController* Player )
 {
-	return true;
-}
+	if (!IsMatchInProgress() || Player == NULL || Player->IsPendingKillPending())
+	{
+		return false;
+	}
 
-bool AGameMode::PlayerCanRestart( APlayerController* aPlayer )
-{
-	return true;
+	// Ask the player controller if it's ready to restart as well
+	return Player->CanRestartPlayer();
 }
 
 void AGameMode::UpdateGameplayMuteList( APlayerController* aPlayer )
@@ -1469,15 +1587,7 @@ void AGameMode::PostSeamlessTravel()
 		}
 	}
 
-	if ( ReadyToStartMatch() )
-	{
-		StartMatch();
-	}
-}
-
-bool AGameMode::ReadyToStartMatch()
-{
-	return (bWaitingToStartMatch && NumPlayers + NumBots > 0);
+	SetMatchState(MatchState::WaitingToStart);
 }
 
 void AGameMode::MatineeCancelled()
diff --git a/Engine/Source/Runtime/Engine/Private/GameSession.cpp b/Engine/Source/Runtime/Engine/Private/GameSession.cpp
--- a/Engine/Source/Runtime/Engine/Private/GameSession.cpp
+++ b/Engine/Source/Runtime/Engine/Private/GameSession.cpp
@@ -39,10 +39,11 @@ AGameSession::AGameSession(const class FPostConstructInitializeProperties& PCIP)
 {
 }
 
-void AGameSession::StartPendingMatch() {}
-void AGameSession::EndPendingMatch() {}
+void AGameSession::HandleMatchIsWaitingToStart() {}
+void AGameSession::HandleMatchHasStarted() {}
+void AGameSession::HandleMatchHasEnded() {}
 
-bool AGameSession::HandleStartMatch() 
+bool AGameSession::HandleStartMatchRequest()
 {
 	return false;
 }
@@ -284,10 +285,6 @@ bool AGameSession::CanRestartGame()
 	return true;
 }
 
-void AGameSession::EndGame()
-{
-}
-
 void AGameSession::UpdateSessionJoinability(FName InSessionName, bool bPublicSearchable, bool bAllowInvites, bool bJoinViaPresence, bool bJoinViaPresenceFriendsOnly)
 {
 	if (GetNetMode() != NM_Standalone)
diff --git a/Engine/Source/Runtime/Engine/Private/GameState.cpp b/Engine/Source/Runtime/Engine/Private/GameState.cpp
--- a/Engine/Source/Runtime/Engine/Private/GameState.cpp
+++ b/Engine/Source/Runtime/Engine/Private/GameState.cpp
@@ -23,6 +23,7 @@ AGameState::AGameState(const class FPostConstructInitializeProperties& PCIP)
 	// and given static NetGUIDs. This also causes their deletions to be recorded and sent to new clients, which if unlucky due to name conflicts,
 	// may end up deleting the new PlayerStates they had just spaned.
 	bNetLoadOnClient = false;
+	MatchState = MatchState::EnteringMap;
 }
 
 /** Helper to return the default object of the GameMode class corresponding to this GameState. */
@@ -38,9 +39,7 @@ AGameMode* AGameState::GetDefaultGameMode() const
 
 void AGameState::DefaultTimer()
 {
-	AGameMode* const GameMode = GetWorld()->GetAuthGameMode();
-
-	if ( (GameMode == NULL) || GameMode->bMatchIsInProgress )
+	if (IsMatchInProgress())
 	{
 		++ElapsedTime;
 	}
@@ -78,22 +77,6 @@ void AGameState::OnRep_SpectatorClass()
 	ReceivedSpectatorClass();
 }
 
-void AGameState::OnRep_bMatchHasBegun()
-{
-	if (bMatchHasBegun)
-	{
-		GetWorldSettings()->NotifyMatchStarted();
-	}
-}
-
-void AGameState::OnRep_bMatchIsOver()
-{
-	if ( bMatchIsOver )
-	{
-		EndGame();
-	}
-}
-
 void AGameState::ReceivedGameModeClass()
 {
 	// Tell each PlayerController that the Game class is here
@@ -166,16 +149,86 @@ void AGameState::RemovePlayerState(APlayerState* PlayerState)
 	}
 }
 
-void AGameState::StartMatch()
+void AGameState::HandleMatchIsWaitingToStart()
+{
+
+}
+
+void AGameState::HandleMatchHasStarted()
+{
+	GetWorldSettings()->NotifyMatchStarted();
+}
+
+void AGameState::HandleMatchHasEnded()
+{
+
+}
+
+void AGameState::HandleLeavingMap()
 {
-	bMatchHasBegun = true;
+
 }
 
-void AGameState::EndGame()
+bool AGameState::HasMatchStarted() const
 {
-	bMatchIsOver = true;
+	if (GetMatchState() == MatchState::EnteringMap || GetMatchState() == MatchState::WaitingToStart)
+	{
+		return false;
+	}
+
+	return true;
 }
 
+bool AGameState::IsMatchInProgress() const
+{
+	if (GetMatchState() == MatchState::InProgress)
+	{
+		return true;
+	}
+
+	return false;
+}
+
+bool AGameState::HasMatchEnded() const
+{
+	if (GetMatchState() == MatchState::WaitingPostMatch || GetMatchState() == MatchState::LeavingMap)
+	{
+		return true;
+	}
+
+	return false;
+}
+
+void AGameState::SetMatchState(FName NewState)
+{
+	if (Role == ROLE_Authority)
+	{
+		MatchState = NewState;
+
+		// Call the onrep to make sure the callbacks happen
+		OnRep_MatchState();
+	}
+}
+
+void AGameState::OnRep_MatchState()
+{
+	if (MatchState == MatchState::WaitingToStart)
+	{
+		HandleMatchIsWaitingToStart();
+	}
+	else if (MatchState == MatchState::InProgress)
+	{
+		HandleMatchHasStarted();
+	}
+	else if (MatchState == MatchState::WaitingPostMatch)
+	{
+		HandleMatchHasEnded();
+	}
+	else if (MatchState == MatchState::LeavingMap)
+	{
+		HandleLeavingMap();
+	}
+}
 
 bool AGameState::ShouldShowGore() const
 {
@@ -186,8 +239,7 @@ void AGameState::GetLifetimeReplicatedProps( TArray< FLifetimeProperty > & OutLi
 {
 	Super::GetLifetimeReplicatedProps( OutLifetimeProps );
 
-	DOREPLIFETIME( AGameState, bMatchHasBegun );
-	DOREPLIFETIME( AGameState, bMatchIsOver );
+	DOREPLIFETIME( AGameState, MatchState );
 	DOREPLIFETIME( AGameState, SpectatorClass );
 
 	DOREPLIFETIME_CONDITION( AGameState, GameModeClass,	COND_InitialOnly );
diff --git a/Engine/Source/Runtime/Engine/Private/NetworkDriver.cpp b/Engine/Source/Runtime/Engine/Private/NetworkDriver.cpp
--- a/Engine/Source/Runtime/Engine/Private/NetworkDriver.cpp
+++ b/Engine/Source/Runtime/Engine/Private/NetworkDriver.cpp
@@ -913,7 +913,7 @@ void UNetDriver::UpdateStandbyCheatStatus(void)
 						NetworkManager->StandbyCheatDetected(STDBY_BadPing);
 					}
 					// Check for the host not sending to the clients, but only during a match
-					else if ( GameMode && !GameMode->bWaitingToStartMatch && !GameMode->bLevelChange &&
+					else if ( GameMode && GameMode->IsMatchInProgress() &&
 						float(CountBadTx) / float(ClientConnections.Num()) > PercentMissingForTxStandby)
 					{
 						bHasStandbyCheatTriggered = true;
diff --git a/Engine/Source/Runtime/Engine/Private/UnrealEngine.cpp b/Engine/Source/Runtime/Engine/Private/UnrealEngine.cpp
--- a/Engine/Source/Runtime/Engine/Private/UnrealEngine.cpp
+++ b/Engine/Source/Runtime/Engine/Private/UnrealEngine.cpp
@@ -8798,7 +8798,7 @@ static inline void CallHandleDisconnectForFailure(UWorld* InWorld, UNetDriver* N
 		if (GameMode)
 		{
 			// Mark the server as having a problem
-			GameMode->bHasNetworkError = true;
+			GameMode->AbortMatch();
 		}
 	}
 	
@@ -9372,7 +9372,7 @@ bool UEngine::TickWorldTravel(FWorldContext& Context, float DeltaSeconds)
 			UE_LOG(LogEngine, Log,  TEXT("Server switch level: %s"), *Context.World()->NextURL );
 			if (Context.World()->GetAuthGameMode() != NULL)
 			{
-				Context.World()->GetAuthGameMode()->GameEnding(); 
+				Context.World()->GetAuthGameMode()->StartToLeaveMap();
 			}
 			FString Error;
 			FString NextURL = Context.World()->NextURL;
@@ -9407,7 +9407,7 @@ bool UEngine::TickWorldTravel(FWorldContext& Context, float DeltaSeconds)
 		AGameMode* const GameMode = Context.World()->GetAuthGameMode();
 		if (GameMode)
 		{
-			GameMode->GameEnding(); 
+			GameMode->StartToLeaveMap();
 		}
 
 		FString Error, TravelURLCopy = Context.TravelURL;
diff --git a/Engine/Source/Runtime/Engine/Private/World.cpp b/Engine/Source/Runtime/Engine/Private/World.cpp
--- a/Engine/Source/Runtime/Engine/Private/World.cpp
+++ b/Engine/Source/Runtime/Engine/Private/World.cpp
@@ -4740,7 +4740,7 @@ void UWorld::ServerTravel(const FString& FURL, bool bAbsolute, bool bShouldSkipG
 	}
 
 	// Check for an error in the server's connection
-	if (AuthorityGameMode && AuthorityGameMode->bHasNetworkError)
+	if (AuthorityGameMode && AuthorityGameMode->GetMatchState() == MatchState::Aborted)
 	{
 		UE_LOG(LogWorld, Log, TEXT("Not traveling because of network error"));
 		return;
