#UE4 - Call Pre/PostNetReceive when mapping previously unmapped properties on client

diff --git a/Engine/Source/Runtime/Engine/Private/DataReplication.cpp b/Engine/Source/Runtime/Engine/Private/DataReplication.cpp
--- a/Engine/Source/Runtime/Engine/Private/DataReplication.cpp
+++ b/Engine/Source/Runtime/Engine/Private/DataReplication.cpp
@@ -1106,24 +1106,30 @@ void FObjectReplicator::StartBecomingDormant()
 	bLastUpdateEmpty = false; // Ensure we get one more attempt to update properties
 }
 
-bool FObjectReplicator::UpdateUnmappedObjects()
+void FObjectReplicator::UpdateUnmappedObjects( bool & bOutHasMoreUnmapped )
 {
 	UObject * Object = GetObject();
 
 	if ( Object == NULL || Object->IsPendingKill() )
 	{
-		return false;
+		bOutHasMoreUnmapped = false;
+		return;
 	}
 
 	check( RepState->RepNotifies.Num() == 0 );
 
-	// Cache the result, since this function will return false when there are no more unmapped properties, but we still need to call any rep notifies
-	const bool Result = RepLayout->UpdateUnmappedObjects( RepState, Connection->PackageMap, Object );
+	bool bSomeObjectsWereMapped = false;
+
+	RepLayout->UpdateUnmappedObjects( RepState, Connection->PackageMap, Object, bSomeObjectsWereMapped, bOutHasMoreUnmapped );
 
 	// Call any rep notifies that need to happen when object pointers change
 	RepLayout->CallRepNotifies( RepState, Object );
 
-	return Result;
+	if ( bSomeObjectsWereMapped )
+	{
+		// If we mapped some objects, make sure to call PostNetReceive (some game code will need to think this was actually replicated to work)
+		PostNetReceive();
+	}
 }
 
 void FObjectReplicator::QueuePropertyRepNotify( UObject * Object, UProperty * Property, const int32 ElementIndex, TArray< uint8 > & MetaData )
diff --git a/Engine/Source/Runtime/Engine/Private/NetworkDriver.cpp b/Engine/Source/Runtime/Engine/Private/NetworkDriver.cpp
--- a/Engine/Source/Runtime/Engine/Private/NetworkDriver.cpp
+++ b/Engine/Source/Runtime/Engine/Private/NetworkDriver.cpp
@@ -298,16 +298,18 @@ void UNetDriver::TickFlush(float DeltaSeconds)
 		if ( !It->IsValid() )
 		{
 			// These are weak references, so if the object has been freed, we can stop checking
-			//UE_LOG( LogNet, Log, TEXT( "REMOVED unmapped replicator: Replicator was destroyed." ) );
 			It.RemoveCurrent();
 			continue;
 		}
 
-		if ( !It->Pin().Get()->UpdateUnmappedObjects() )
+		bool bHasMoreUnmapped = false;
+
+		It->Pin().Get()->UpdateUnmappedObjects( bHasMoreUnmapped );
+		
+		if ( !bHasMoreUnmapped )
 		{
 			// If there are no more unmapped objects, we can also stop checking
 			It.RemoveCurrent();
-			//UE_LOG( LogNet, Log, TEXT( "REMOVED unmapped replicator: No unmapped properties. NumUnmappedReplicators: %i" ), UnmappedReplicators.Num() );
 		}
 	}
 }
diff --git a/Engine/Source/Runtime/Engine/Private/RepLayout.cpp b/Engine/Source/Runtime/Engine/Private/RepLayout.cpp
--- a/Engine/Source/Runtime/Engine/Private/RepLayout.cpp
+++ b/Engine/Source/Runtime/Engine/Private/RepLayout.cpp
@@ -1399,7 +1399,15 @@ FUnmappedGuidMgrElement::~FUnmappedGuidMgrElement()
 	}
 }
 
-bool FRepLayout::UpdateUnmappedObjects_r( FRepState * RepState, FUnmappedGuidMgr * UnmappedGuids, UPackageMap * PackageMap, uint8 * RESTRICT Data, const int32 MaxAbsOffset ) const
+void FRepLayout::UpdateUnmappedObjects_r( 
+	FRepState *			RepState, 
+	FUnmappedGuidMgr *	UnmappedGuids, 
+	UObject *			OriginalObject,
+	UPackageMap *		PackageMap, 
+	uint8 * RESTRICT	Data, 
+	const int32			MaxAbsOffset,
+	bool &				bOutSomeObjectsWereMapped,
+	bool &				bOutHasMoreUnmapped ) const
 {
 	bool bHasUnmapped = false;
 
@@ -1410,7 +1418,7 @@ bool FRepLayout::UpdateUnmappedObjects_r( FRepState * RepState, FUnmappedGuidMgr
 		if ( AbsOffset >= MaxAbsOffset )
 		{
 			// Array must have shrunk, we can remove this item
-			UE_LOG( LogNet, Verbose, TEXT( "REMOVED unmapped property: AbsOffset >= MaxAbsOffset. Offset: %i" ), AbsOffset );
+			UE_LOG( LogNet, VeryVerbose, TEXT( "REMOVED unmapped property: AbsOffset >= MaxAbsOffset. Offset: %i" ), AbsOffset );
 			It.RemoveCurrent();
 			continue;
 		}
@@ -1422,10 +1430,7 @@ bool FRepLayout::UpdateUnmappedObjects_r( FRepState * RepState, FUnmappedGuidMgr
 		{
 			check( Cmd.Type == REPCMD_DynamicArray );
 			FScriptArray * Array = (FScriptArray *)( Data + AbsOffset );
-			if ( UpdateUnmappedObjects_r( RepState, It.Value().Array, PackageMap, (uint8*)Array->GetData(), Array->Num() * Cmd.ElementSize ) )
-			{
-				bHasUnmapped = true;
-			}
+			UpdateUnmappedObjects_r( RepState, It.Value().Array, OriginalObject, PackageMap, (uint8*)Array->GetData(), Array->Num() * Cmd.ElementSize, bOutSomeObjectsWereMapped, bOutHasMoreUnmapped );
 			continue;
 		}
 
@@ -1435,36 +1440,46 @@ bool FRepLayout::UpdateUnmappedObjects_r( FRepState * RepState, FUnmappedGuidMgr
 
 		if ( Object != NULL )
 		{
-			UE_LOG( LogNet, Verbose, TEXT( "REMOVED unmapped property: Offset: %i, Guid: %s, PropName: %s, ObjName: %s" ), AbsOffset, *It.Value().Guid.ToString(), *Cmd.Property->GetName(), *Object->GetName() );
+			UE_LOG( LogNet, VeryVerbose, TEXT( "REMOVED unmapped property: Offset: %i, Guid: %s, PropName: %s, ObjName: %s" ), AbsOffset, *It.Value().Guid.ToString(), *Cmd.Property->GetName(), *Object->GetName() );
 			UObjectPropertyBase * ObjProperty = CastChecked< UObjectPropertyBase>( Cmd.Property );
 
 			UObject * OldObject = ObjProperty->GetObjectPropertyValue( Data + AbsOffset );
 
-			ObjProperty->SetObjectPropertyValue( Data + AbsOffset, Object );
-
-			if ( OldObject != Object && Parent.Property->HasAnyPropertyFlags( CPF_RepNotify ) )
+			if ( OldObject != Object )
 			{
-				RepState->RepNotifies.AddUnique( Parent.Property );
-			} 
+				if ( !bOutSomeObjectsWereMapped )
+				{
+					// Call PreNetReceive if we are going to change a value (some game code will need to think this is an actual replicated value)
+					OriginalObject->PreNetReceive();
+					bOutSomeObjectsWereMapped = true;
+				}
+
+				ObjProperty->SetObjectPropertyValue( Data + AbsOffset, Object );
+
+				// If this properties needs an OnRep, queue that up to be handled later
+				if ( Parent.Property->HasAnyPropertyFlags( CPF_RepNotify ) )
+				{
+					RepState->RepNotifies.AddUnique( Parent.Property );
+				} 
+			}
 
 			It.RemoveCurrent();
 			continue;
 		}
 
-		bHasUnmapped = true;
+		UE_LOG( LogNet, VeryVerbose, TEXT( "UnmappedGuids: Offset: %i, Guid: %s" ), AbsOffset, *It.Value().Guid.ToString() );
 
-		UE_LOG( LogNet, Verbose, TEXT( "UnmappedGuids: Offset: %i, Guid: %s" ), AbsOffset, *It.Value().Guid.ToString() );
+		bOutHasMoreUnmapped = true;
 	}
-
-	return bHasUnmapped;
 }
 
-bool FRepLayout::UpdateUnmappedObjects( FRepState *	RepState, UPackageMap * PackageMap, void * RESTRICT Data ) const
+void FRepLayout::UpdateUnmappedObjects( FRepState *	RepState, UPackageMap * PackageMap, UObject * OriginalObject, bool & bOutSomeObjectsWereMapped, bool & bOutHasMoreUnmapped ) const
 {
-#if ENABLE_CLIENT_UNMAP_LOGIC == 0
-	return false;
-#else
-	return UpdateUnmappedObjects_r( RepState, &RepState->UnmappedGuids, PackageMap, (uint8*)Data, RepState->StaticBuffer.Num() );
+	bOutSomeObjectsWereMapped	= false;
+	bOutHasMoreUnmapped			= false;
+
+#if ENABLE_CLIENT_UNMAP_LOGIC == 1
+	UpdateUnmappedObjects_r( RepState, &RepState->UnmappedGuids, OriginalObject, PackageMap, (uint8*)OriginalObject, RepState->StaticBuffer.Num(), bOutSomeObjectsWereMapped, bOutHasMoreUnmapped );
 #endif
 }
 
diff --git a/Engine/Source/Runtime/Engine/Public/Net/DataReplication.h b/Engine/Source/Runtime/Engine/Public/Net/DataReplication.h
--- a/Engine/Source/Runtime/Engine/Public/Net/DataReplication.h
+++ b/Engine/Source/Runtime/Engine/Public/Net/DataReplication.h
@@ -124,7 +124,7 @@ public:
 
 	void StartBecomingDormant();
 
-	bool UpdateUnmappedObjects();
+	void UpdateUnmappedObjects( bool & bOutHasMoreUnmapped );
 
 	FORCEINLINE UObject *	GetObject() const { return ObjectPtr.Get(); }
 	FORCEINLINE void		SetObject( UObject * NewObj ) { ObjectPtr = TWeakObjectPtr<UObject>( NewObj ); }
diff --git a/Engine/Source/Runtime/Engine/Public/Net/RepLayout.h b/Engine/Source/Runtime/Engine/Public/Net/RepLayout.h
--- a/Engine/Source/Runtime/Engine/Public/Net/RepLayout.h
+++ b/Engine/Source/Runtime/Engine/Public/Net/RepLayout.h
@@ -327,7 +327,7 @@ public:
 	void InitFromObjectClass( UClass * InObjectClass );
 
 	bool ReceiveProperties( UClass * InObjectClass, FRepState * RESTRICT RepState, void * RESTRICT Data, FNetBitReader & InBunch, bool bDiscard, bool & bOutHasUnmapped ) const;
-	bool UpdateUnmappedObjects( FRepState *	RepState, UPackageMap * PackageMap, void * RESTRICT Data ) const;
+	void UpdateUnmappedObjects( FRepState *	RepState, UPackageMap * PackageMap, UObject * Object, bool & bOutSomeObjectsWereMapped, bool & bOutHasMoreUnmapped ) const;
 
 	void CallRepNotifies( FRepState * RepState, UObject * Object ) const;
 	void PostReplicate( FRepState * RepState, FPacketIdRange & PacketRange, bool bReliable ) const;
@@ -459,7 +459,15 @@ private:
 		uint8 * RESTRICT	Data,
 		const bool			bDiscard ) const;
 
-	bool UpdateUnmappedObjects_r( FRepState * RepState, FUnmappedGuidMgr * UnmappedGuids, UPackageMap * PackageMap, uint8 * RESTRICT Data, const int32 MaxAbsOffset ) const;
+	void UpdateUnmappedObjects_r( 
+		FRepState *			RepState, 
+		FUnmappedGuidMgr *	UnmappedGuids, 
+		UObject *			OriginalObject,
+		UPackageMap *		PackageMap, 
+		uint8 * RESTRICT	Data, 
+		const int32			MaxAbsOffset,
+		bool &				bOutSomeObjectsWereMapped,
+		bool &				bOutHasMoreUnmapped ) const;
 
 	void ValidateWithChecksum_DynamicArray_r( const FRepLayoutCmd & Cmd, const int32 CmdIndex, const uint8 * RESTRICT Data, FArchive & Ar, const bool bDiscard ) const;
 	void ValidateWithChecksum_r( const int32 CmdStart, const int32 CmdEnd, const uint8 * RESTRICT Data, FArchive & Ar, const bool bDiscard ) const;
