#UE4:
* undo change as it broke existing content serialization of BreakMaterialAttributes. Old change:
  Change 2034682 by [REDACTED_PERSON][REDACTED_HANDLE] on 2014/03/27 09:34:33
#ttp 325589 - UE4: RENDERING: UDN: Rendering: Materials
#summary Removed mask specifiers for UMaterialExpressionBreakMaterialAttributes outputs.

diff --git a/Engine/Source/Runtime/Engine/Private/Materials/MaterialExpressions.cpp b/Engine/Source/Runtime/Engine/Private/Materials/MaterialExpressions.cpp
--- a/Engine/Source/Runtime/Engine/Private/Materials/MaterialExpressions.cpp
+++ b/Engine/Source/Runtime/Engine/Private/Materials/MaterialExpressions.cpp
@@ -3020,24 +3020,24 @@ UMaterialExpressionBreakMaterialAttributes::UMaterialExpressionBreakMaterialAttr
 	MenuCategories.Add(ConstructorStatics.NAME_MaterialAttributes);
 
 	Outputs.Reset();
-	Outputs.Add(FExpressionOutput(TEXT("BaseColor")));
-	Outputs.Add(FExpressionOutput(TEXT("Metallic")));
-	Outputs.Add(FExpressionOutput(TEXT("Specular")));
-	Outputs.Add(FExpressionOutput(TEXT("Roughness")));
-	Outputs.Add(FExpressionOutput(TEXT("EmissiveColor")));
-	Outputs.Add(FExpressionOutput(TEXT("Opacity")));
-	Outputs.Add(FExpressionOutput(TEXT("OpacityMask")));
-	Outputs.Add(FExpressionOutput(TEXT("Normal")));
-	Outputs.Add(FExpressionOutput(TEXT("WorldPositionOffset")));
-	Outputs.Add(FExpressionOutput(TEXT("WorldDisplacement")));
-	Outputs.Add(FExpressionOutput(TEXT("TessellationMultiplier")));
-	Outputs.Add(FExpressionOutput(TEXT("SubsurfaceColor")));
-	Outputs.Add(FExpressionOutput(TEXT("AmbientOcclusion")));
-	Outputs.Add(FExpressionOutput(TEXT("Refraction")));
+	Outputs.Add(FExpressionOutput(TEXT("BaseColor"), 1, 1, 1, 1, 0));
+	Outputs.Add(FExpressionOutput(TEXT("Metallic"), 1, 1, 1, 1, 0));
+	Outputs.Add(FExpressionOutput(TEXT("Specular"), 1, 1, 1, 1, 0));
+	Outputs.Add(FExpressionOutput(TEXT("Roughness"), 1, 1, 1, 1, 0));
+	Outputs.Add(FExpressionOutput(TEXT("EmissiveColor"), 1, 1, 1, 1, 0));
+	Outputs.Add(FExpressionOutput(TEXT("Opacity"), 1, 1, 1, 1, 0));
+	Outputs.Add(FExpressionOutput(TEXT("OpacityMask"), 1, 1, 1, 1, 0));
+	Outputs.Add(FExpressionOutput(TEXT("Normal"), 1, 1, 1, 1, 0));
+	Outputs.Add(FExpressionOutput(TEXT("WorldPositionOffset"), 1, 1, 1, 1, 0));
+	Outputs.Add(FExpressionOutput(TEXT("WorldDisplacement"), 1, 1, 1, 1, 0));
+	Outputs.Add(FExpressionOutput(TEXT("TessellationMultiplier"), 1, 1, 1, 1, 0));
+	Outputs.Add(FExpressionOutput(TEXT("SubsurfaceColor"), 1, 1, 1, 1, 0));
+	Outputs.Add(FExpressionOutput(TEXT("AmbientOcclusion"), 1, 1, 1, 1, 0));
+	Outputs.Add(FExpressionOutput(TEXT("Refraction"), 1, 1, 1, 0, 0));
 
 	for (int32 UVIndex = 0; UVIndex < MP_CustomizedUVs7 - MP_CustomizedUVs0; UVIndex++)
 	{
-		Outputs.Add(FExpressionOutput(*FString::Printf(TEXT("CustomizedUV%u"), UVIndex)));
+		Outputs.Add(FExpressionOutput(*FString::Printf(TEXT("CustomizedUV%u"), UVIndex), 1, 1, 1, 0, 0));
 	}
 }
 
