[INTEGRATE] Change 2042573 by [REDACTED_PERSON][REDACTED_HANDLE] on 2014/04/08 13:21:49

	#ttp 330647 - Unbuilt objects after fresh build.
	#proj UE4
	#branch UE4
	#summary Static lighting is now disabled for instanced static mesh components
	#extra The recent change to make foliage components static resulted in them being considered for static lighting; something not currently possible for instanced static mesh components.
	#change Foliage components have been changed back to be movable.
	#change Instanced static mesh components have static shadow casting disabled by default, and the UI option to change this in the foliage instance editor has been removed.

diff --git a/Engine/Source/Runtime/Core/Public/UObject/ObjectVersion.h b/Engine/Source/Runtime/Core/Public/UObject/ObjectVersion.h
--- a/Engine/Source/Runtime/Core/Public/UObject/ObjectVersion.h
+++ b/Engine/Source/Runtime/Core/Public/UObject/ObjectVersion.h
@@ -495,8 +495,8 @@ enum EUnrealEngineObjectUE4Version
 	VER_UE4_DISABLED_SCRIPT_LIMIT_BYTECODE,
 	// Made remote role private, exposed bReplicates
 	VER_UE4_PRIVATE_REMOTE_ROLE,
-	// Fix up old foliage components to have static mobility
-	VER_UE4_FOLIAGE_MOBILITY,
+	// Fix up old foliage components to have static mobility (superseded by VER_UE4_FOLIAGE_MOVABLE_MOBILITY)
+	VER_UE4_FOLIAGE_STATIC_MOBILITY,
 	// Change BuildScale from a float to a vector
 	VER_UE4_BUILD_SCALE_VECTOR,
 	// After implementing foliage collision, need to disable collision on old foliage instances
@@ -512,6 +512,8 @@ enum EUnrealEngineObjectUE4Version
 	VER_UE4_SOFT_CONSTRAINTS_USE_MASS,
 	// Reflection capture data saved in packages
 	VER_UE4_REFLECTION_DATA_IN_PACKAGES,
+	// Fix up old foliage components to have movable mobility
+	VER_UE4_FOLIAGE_MOVABLE_MOBILITY,
 
 	// -----<new versions can be added before this line>-------------------------------------------------
 	// - this needs to be the last line (see note below)
diff --git a/Engine/Source/Runtime/Engine/Classes/Foliage/InstancedFoliageSettings.h b/Engine/Source/Runtime/Engine/Classes/Foliage/InstancedFoliageSettings.h
--- a/Engine/Source/Runtime/Engine/Classes/Foliage/InstancedFoliageSettings.h
+++ b/Engine/Source/Runtime/Engine/Classes/Foliage/InstancedFoliageSettings.h
@@ -188,13 +188,16 @@ class UInstancedFoliageSettings : public UObject
 	UPROPERTY(EditAnywhere, BlueprintReadOnly, Category=Lighting)
 	uint32 bAffectDynamicIndirectLighting:1;
 
-	/** Controls whether the primitive should cast shadows in the case of non precomputed shadowing.  This flag is only used if CastShadow is true. **/
-	UPROPERTY(EditAnywhere, BlueprintReadOnly, Category=Lighting)
-	uint32 bCastDynamicShadow:1;
+	// Static lighting is not currently supported by instanced static meshes, so we treat CastShadow as controlling all shadow settings
+	// (enables dynamic shadows, disables static shadows) which makes these two settings (bCastDynamicShadow, bCastStaticShadow) irrelevant at this point in time
+	// See UInstancedStaticMeshComponent::GetStaticLightingInfo (it's empty, but it's what would add data to the Lightmass generation)
 
+	/** Controls whether the primitive should cast shadows in the case of non precomputed shadowing.  This flag is only used if CastShadow is true. **/
+	//UPROPERTY(EditAnywhere, BlueprintReadOnly, Category=Lighting)
+	//uint32 bCastDynamicShadow:1;
 	/** Whether the object should cast a static shadow from shadow casting lights.  This flag is only used if CastShadow is true. */
-	UPROPERTY(EditAnywhere, BlueprintReadOnly, Category=Lighting)
-	uint32 bCastStaticShadow:1;
+	//UPROPERTY(EditAnywhere, BlueprintReadOnly, Category=Lighting)
+	//uint32 bCastStaticShadow:1;
 
 	/** 
 	 *	If true, the primitive will cast shadows even if bHidden is true.
diff --git a/Engine/Source/Runtime/Engine/Private/InstancedFoliage.cpp b/Engine/Source/Runtime/Engine/Private/InstancedFoliage.cpp
--- a/Engine/Source/Runtime/Engine/Private/InstancedFoliage.cpp
+++ b/Engine/Source/Runtime/Engine/Private/InstancedFoliage.cpp
@@ -123,10 +123,13 @@ UInstancedFoliageSettings::UInstancedFoliageSettings(const class FPostConstructI
 	VertexColorMask = FOLIAGEVERTEXCOLORMASK_Disabled;
 	VertexColorMaskThreshold = 0.5f;
 
+	// Static lighting is not currently supported by instanced static meshes, so we treat CastShadow as controlling all shadow settings
+	// (enables dynamic shadows, disables static shadows) which makes these two settings (bCastDynamicShadow, bCastStaticShadow) irrelevant at this point in time
+	// See UInstancedStaticMeshComponent::GetStaticLightingInfo (it's empty, but it's what would add data to the Lightmass generation)
 	CastShadow = true;
-	bCastDynamicShadow = true;
+	//bCastDynamicShadow = true;
+	//bCastStaticShadow = true;
 	bAffectDynamicIndirectLighting = false;
-	bCastStaticShadow = true;
 	bCastHiddenShadow = false;
 	bCastShadowAsTwoSided = false;
 
@@ -291,7 +294,9 @@ void FFoliageMeshInfo::AddInstance( AInstancedFoliageActor* InIFA, UStaticMesh*
 			ConstructObject<UInstancedStaticMeshComponent>(UInstancedStaticMeshComponent::StaticClass(),InIFA,NAME_None,RF_Transactional),
 			InMesh->GetBounds().TransformBy(InstanceTransform)
 			);
-		BestCluster->ClusterComponent->Mobility = EComponentMobility::Static;
+
+		// Make the instanced static mesh component movable so it doesn't get statically lit; see the comment below about CastShadow for more details
+		BestCluster->ClusterComponent->Mobility = EComponentMobility::Movable;
 
 		BestCluster->ClusterComponent->StaticMesh = InMesh;
 		BestCluster->ClusterComponent->bSelectable = true;
@@ -300,10 +305,13 @@ void FFoliageMeshInfo::AddInstance( AInstancedFoliageActor* InIFA, UStaticMesh*
 		BestCluster->ClusterComponent->InstanceStartCullDistance = Settings->StartCullDistance;
 		BestCluster->ClusterComponent->InstanceEndCullDistance = Settings->EndCullDistance;
 
+		// Static lighting is not currently supported by instanced static meshes, so we treat CastShadow as controlling all shadow settings
+		// (enables dynamic shadows, disables static shadows) which makes these two settings (bCastDynamicShadow, bCastStaticShadow) irrelevant at this point in time
+		// See UInstancedStaticMeshComponent::GetStaticLightingInfo (it's empty, but it's what would add data to the Lightmass generation)
 		BestCluster->ClusterComponent->CastShadow = Settings->CastShadow;
-		BestCluster->ClusterComponent->bCastDynamicShadow = Settings->bCastDynamicShadow;
+		BestCluster->ClusterComponent->bCastDynamicShadow = true; //Settings->bCastDynamicShadow;
+		BestCluster->ClusterComponent->bCastStaticShadow = false; //Settings->bCastStaticShadow;
 		BestCluster->ClusterComponent->bAffectDynamicIndirectLighting = Settings->bAffectDynamicIndirectLighting;
-		BestCluster->ClusterComponent->bCastStaticShadow = Settings->bCastStaticShadow;
 		BestCluster->ClusterComponent->bCastHiddenShadow = Settings->bCastHiddenShadow;
 		BestCluster->ClusterComponent->bCastShadowAsTwoSided = Settings->bCastShadowAsTwoSided;
 
@@ -1460,14 +1468,14 @@ void AInstancedFoliageActor::Serialize(FArchive& Ar)
 	Super::Serialize(Ar);
 	Ar << FoliageMeshes;
 
-	if (Ar.UE4Ver() < VER_UE4_FOLIAGE_COLLISION)
+	if (Ar.UE4Ver() < VER_UE4_FOLIAGE_MOVABLE_MOBILITY)
 	{
 		for (const TPair<UStaticMesh*, FFoliageMeshInfo>& MeshInfo : FoliageMeshes)
 		{
 			for (const FFoliageInstanceCluster& Cluster : MeshInfo.Value.InstanceClusters)
 			{
-				if (Ar.UE4Ver() < VER_UE4_FOLIAGE_MOBILITY)
-					Cluster.ClusterComponent->SetMobility(EComponentMobility::Static);
+				if (Ar.UE4Ver() < VER_UE4_FOLIAGE_MOVABLE_MOBILITY)
+					Cluster.ClusterComponent->SetMobility(EComponentMobility::Movable);
 				if (Ar.UE4Ver() < VER_UE4_FOLIAGE_COLLISION)
 					Cluster.ClusterComponent->SetCollisionEnabled(ECollisionEnabled::NoCollision);
 			}
diff --git a/Engine/Source/Runtime/Engine/Private/InstancedStaticMesh.cpp b/Engine/Source/Runtime/Engine/Private/InstancedStaticMesh.cpp
--- a/Engine/Source/Runtime/Engine/Private/InstancedStaticMesh.cpp
+++ b/Engine/Source/Runtime/Engine/Private/InstancedStaticMesh.cpp
@@ -1061,6 +1061,10 @@ UInstancedStaticMeshComponent::UInstancedStaticMeshComponent(const class FPostCo
 {
 	Mobility = EComponentMobility::Movable;
 	BodyInstance.bSimulatePhysics = false;
+
+	// Static lighting is not currently supported by instanced static meshes, so disable it
+	// See UInstancedStaticMeshComponent::GetStaticLightingInfo (it's empty, but it's what would add data to the Lightmass generation)
+	bCastStaticShadow = false;
 }
 
 #if WITH_EDITOR
