[github]
pull request https://github.com/EpicGames/UnrealEngine/pull/17
from user StephenWhittle

diff --git a/Engine/Source/Editor/AnimGraph/Classes/AnimGraphNode_Fabrik.h b/Engine/Source/Editor/AnimGraph/Classes/AnimGraphNode_Fabrik.h
new file mode 100644
--- /dev/null
+++ b/Engine/Source/Editor/AnimGraph/Classes/AnimGraphNode_Fabrik.h
@@ -0,0 +1,25 @@
+// Copyright 1998-2013 Epic Games, Inc. All Rights Reserved.
+
+#pragma once
+
+#include "AnimGraphNode_Fabrik.generated.h"
+
+// Editor node for FABRIK IK skeletal controller
+UCLASS(MinimalAPI)
+class UAnimGraphNode_Fabrik : public UAnimGraphNode_SkeletalControlBase
+{
+	GENERATED_UCLASS_BODY()
+
+	UPROPERTY(EditAnywhere, Category = Settings)
+	FAnimNode_Fabrik Node;
+
+public:
+	// UEdGraphNode interface
+	virtual FString GetNodeTitle(ENodeTitleType::Type TitleType) const OVERRIDE;
+	// End of UEdGraphNode interface
+
+protected:
+	// UAnimGraphNode_SkeletalControlBase interface
+	virtual FString GetControllerDescription() const;
+	// End of UAnimGraphNode_SkeletalControlBase interface
+};
\ No newline at end of file
diff --git a/Engine/Source/Editor/AnimGraph/Private/AnimGraphNode_Fabrik.cpp b/Engine/Source/Editor/AnimGraph/Private/AnimGraphNode_Fabrik.cpp
new file mode 100644
--- /dev/null
+++ b/Engine/Source/Editor/AnimGraph/Private/AnimGraphNode_Fabrik.cpp
@@ -0,0 +1,22 @@
+// Copyright 1998-2013 Epic Games, Inc. All Rights Reserved.
+
+#include "AnimGraphPrivatePCH.h"
+
+/////////////////////////////////////////////////////
+// UAnimGraphNode_Fabrik 
+
+UAnimGraphNode_Fabrik::UAnimGraphNode_Fabrik(const FPostConstructInitializeProperties& PCIP)
+: Super(PCIP)
+{
+}
+
+FString UAnimGraphNode_Fabrik::GetControllerDescription() const
+{
+	return TEXT("Fabrik IK");
+}
+
+FString UAnimGraphNode_Fabrik::GetNodeTitle(ENodeTitleType::Type TitleType) const
+{
+	FString Result = *GetControllerDescription();
+	return Result;
+}
\ No newline at end of file
diff --git a/Engine/Source/Runtime/Engine/Classes/Animation/AnimInstance.h b/Engine/Source/Runtime/Engine/Classes/Animation/AnimInstance.h
--- a/Engine/Source/Runtime/Engine/Classes/Animation/AnimInstance.h
+++ b/Engine/Source/Runtime/Engine/Classes/Animation/AnimInstance.h
@@ -25,6 +25,18 @@ enum EBoneControlSpace
 	BCS_MAX,
 };
 
+/** Enum for specifying the source of a bone's rotation. */
+UENUM()
+enum EBoneRotationSource
+{
+	/** Don't change rotation at all. */
+	BRS_KeepComponentSpaceRotation UMETA(DisplayName = "No Change (Preserve Existing Component Space Rotation)"),
+	/** Keep forward direction vector relative to the parent bone. */
+	BRS_KeepLocalSpaceRotation UMETA(DisplayName = "Maintain Local Rotation Relative to Parent"),
+	/** Copy rotation of target to bone. */
+	BRS_CopyFromTarget UMETA(DisplayName = "Copy Target Rotation"),
+};
+
 USTRUCT()
 struct FA2Pose
 {
diff --git a/Engine/Source/Runtime/Engine/Classes/Animation/BoneControllers/AnimNode_Fabrik.h b/Engine/Source/Runtime/Engine/Classes/Animation/BoneControllers/AnimNode_Fabrik.h
new file mode 100644
--- /dev/null
+++ b/Engine/Source/Runtime/Engine/Classes/Animation/BoneControllers/AnimNode_Fabrik.h
@@ -0,0 +1,100 @@
+// Copyright 1998-2013 Epic Games, Inc. All Rights Reserved.
+
+#pragma once
+
+#include "AnimNode_SkeletalControlBase.h"
+#include "AnimNode_Fabrik.generated.h"
+
+/**
+*	Controller which implements the FABRIK IK approximation algorithm -  see http://andreasaristidou.com/publications/FABRIK.pdf for details
+*/
+
+/** Transient structure for FABRIK node evaluation */
+USTRUCT()
+struct FABRIKChainLink
+{
+	GENERATED_USTRUCT_BODY()
+
+public:
+	FVector Position;
+	float Length;
+	int32 BoneIndex;
+	int32 TransformIndex;
+
+	UPROPERTY()
+	TArray<int32> ChildZeroLengthTransformIndices;
+
+	FABRIKChainLink()
+		: Position(FVector::ZeroVector)
+		, Length(0.f)
+		, BoneIndex(INDEX_NONE)
+		, TransformIndex(INDEX_NONE)
+	{
+	}
+
+	FABRIKChainLink(FVector const & InPosition, float const & InLength, int32 const & InBoneIndex, int32 const & InTransformIndex)
+		: Position(InPosition)
+		, Length(InLength)
+		, BoneIndex(InBoneIndex)
+		, TransformIndex(InTransformIndex)
+	{
+	}
+};
+
+USTRUCT()
+struct ENGINE_API FAnimNode_Fabrik : public FAnimNode_SkeletalControlBase
+{
+	GENERATED_USTRUCT_BODY()
+
+	/** Coordinates for target location of tip bone - if EffectorLocationSpace is bone, this is the offset from Target Bone to use as target location*/
+	UPROPERTY(EditAnywhere, BlueprintReadWrite, Category = EndEffector, meta = (PinShownByDefault))
+	FTransform EffectorTransform;
+
+	/** Reference frame of Effector Transform. */
+	UPROPERTY(EditAnywhere, BlueprintReadWrite, Category = EndEffector)
+	TEnumAsByte<enum EBoneControlSpace> EffectorTransformSpace;
+
+	/** If EffectorTransformSpace is a bone, this is the bone to use. **/
+	UPROPERTY(EditAnywhere, BlueprintReadWrite, Category = EndEffector)
+	FBoneReference EffectorTransformBone;
+
+	UPROPERTY(EditAnywhere, BlueprintReadWrite, Category = EndEffector)
+	TEnumAsByte<enum EBoneRotationSource> EffectorRotationSource;
+
+	/** Name of tip bone */
+	UPROPERTY(EditAnywhere, BlueprintReadWrite, Category = Solver)
+	FBoneReference TipBone;
+
+	/** Name of the root bone*/
+	UPROPERTY(EditAnywhere, BlueprintReadWrite, Category = Solver)
+	FBoneReference RootBone;
+
+	/** Tolerance for final tip location delta from EffectorLocation*/
+	UPROPERTY(EditAnywhere, BlueprintReadWrite, Category = Solver)
+	float Precision;
+
+	/** Maximum number of iterations allowed, to control performance. */
+	UPROPERTY(EditAnywhere, BlueprintReadWrite, Category = Solver)
+	int32 MaxIterations;
+
+	/** Toggle drawing of axes to debug joint rotation*/
+	UPROPERTY(EditAnywhere, BlueprintReadWrite, Category = Solver)
+	bool bEnableDebugDraw;
+
+public:
+	FAnimNode_Fabrik();
+
+	// FAnimNode_SkeletalControlBase interface
+	virtual void EvaluateBoneTransforms(USkeletalMeshComponent* SkelComp, const FBoneContainer & RequiredBones, FA2CSPose& MeshBases, TArray<FBoneTransform>& OutBoneTransforms) OVERRIDE;
+	virtual bool IsValidToEvaluate(const USkeleton * Skeleton, const FBoneContainer & RequiredBones) OVERRIDE;
+	// End of FAnimNode_SkeletalControlBase interface
+
+
+private:
+	// FAnimNode_SkeletalControlBase interface
+	virtual void InitializeBoneReferences(const FBoneContainer & RequiredBones) OVERRIDE;
+	// End of FAnimNode_SkeletalControlBase interface
+
+	// Convenience function to get current (pre-translation iteration) component space location of bone by bone index
+	FVector GetCurrentLocation(FA2CSPose & MeshBases, int32 const & BoneIndex);
+};
\ No newline at end of file
diff --git a/Engine/Source/Runtime/Engine/Private/Animation/AnimNode_Fabrik.cpp b/Engine/Source/Runtime/Engine/Private/Animation/AnimNode_Fabrik.cpp
new file mode 100644
--- /dev/null
+++ b/Engine/Source/Runtime/Engine/Private/Animation/AnimNode_Fabrik.cpp
@@ -0,0 +1,251 @@
+// Copyright 1998-2013 Epic Games, Inc. All Rights Reserved.
+
+#include "EnginePrivate.h"
+
+/////////////////////////////////////////////////////
+// AnimNode_Fabrik
+// Implementation of the FABRIK IK Algorithm
+// Please see http://andreasaristidou.com/publications/FABRIK.pdf for more details
+
+FAnimNode_Fabrik::FAnimNode_Fabrik()
+	: EffectorTransform(FTransform::Identity)
+	, EffectorTransformSpace(BCS_ComponentSpace)
+	, EffectorRotationSource(BRS_KeepLocalSpaceRotation)
+	, Precision(1.f)
+	, MaxIterations(10)
+	, bEnableDebugDraw(false)
+{
+}
+
+FVector FAnimNode_Fabrik::GetCurrentLocation(FA2CSPose & MeshBases, int32 const & BoneIndex)
+{
+	return MeshBases.GetComponentSpaceTransform(BoneIndex).GetLocation();
+}
+
+void FAnimNode_Fabrik::EvaluateBoneTransforms(USkeletalMeshComponent* SkelComp, const FBoneContainer & RequiredBones, FA2CSPose& MeshBases, TArray<FBoneTransform>& OutBoneTransforms)
+{
+	// IsValidToEvaluate validated our inputs, we don't need to check pre-requisites again.
+	int32 const RootIndex = RootBone.BoneIndex;
+
+	// Update EffectorLocation if it is based off a bone position
+	FTransform CSEffectorTransform = FTransform(EffectorTransform);
+	FAnimationRuntime::ConvertBoneSpaceTransformToCS(SkelComp, MeshBases, CSEffectorTransform, EffectorTransformBone.BoneIndex, EffectorTransformSpace);
+	FVector const CSEffectorLocation = CSEffectorTransform.GetLocation();
+
+	// @fixme - we need better to draw widgets and debug information in editor.
+#if !(UE_BUILD_SHIPPING || UE_BUILD_TEST)
+	if (bEnableDebugDraw)
+	{
+		// Show end effector position.
+		DrawDebugBox(SkelComp->GetWorld(), CSEffectorLocation, FVector(Precision), FColor::Green, true, 0.1f);
+		DrawDebugCoordinateSystem(SkelComp->GetWorld(), CSEffectorLocation, CSEffectorTransform.GetRotation().Rotator(), 5.f, true, 0.1f);
+	}
+#endif
+
+	// Gather all bone indices between root and tip.
+	TArray<int32> BoneIndices;
+	int32 BoneIndex = TipBone.BoneIndex;
+	do 
+	{
+		BoneIndices.Insert(BoneIndex, 0);
+		BoneIndex = RequiredBones.GetParentBoneIndex(BoneIndex);
+	} while (BoneIndex != RootIndex);
+	BoneIndices.Insert(BoneIndex, 0);
+
+	// Maximum length of skeleton segment at full extension
+	float MaximumReach = 0;
+
+	// Gather transforms
+	int32 const NumTransforms = BoneIndices.Num();
+	OutBoneTransforms.AddUninitialized(NumTransforms);
+
+	// Gather chain links. These are non zero length bones.
+	TArray<FABRIKChainLink> Chain;
+	Chain.Reserve(NumTransforms);
+
+	// Start with Root Bone
+	{
+		int32 const & RootBoneIndex = BoneIndices[0];
+		FTransform const BoneCSTransform = MeshBases.GetComponentSpaceTransform(RootBoneIndex);
+		OutBoneTransforms[0] = FBoneTransform(RootBoneIndex, BoneCSTransform);
+
+		Chain.Add(FABRIKChainLink(BoneCSTransform.GetLocation(), 0.f, RootBoneIndex, 0));
+	}
+
+	// Go through remaining transforms
+	for (int32 TransformIndex = 1; TransformIndex < NumTransforms; TransformIndex++)
+	{
+		int32 const & BoneIndex = BoneIndices[TransformIndex];
+
+		FTransform const BoneCSTransform = MeshBases.GetComponentSpaceTransform(BoneIndex);
+		FVector const BoneCSPosition = BoneCSTransform.GetLocation();
+		OutBoneTransforms[TransformIndex] = FBoneTransform(BoneIndex, BoneCSTransform);
+
+		// Calculate the combined length of this segment of skeleton
+		float const BoneLength = FVector::Dist(BoneCSPosition, OutBoneTransforms[TransformIndex-1].Transform.GetLocation());
+
+		if (!FMath::IsNearlyZero(BoneLength))
+		{
+			Chain.Add(FABRIKChainLink(BoneCSPosition, BoneLength, BoneIndex, TransformIndex));
+			MaximumReach += BoneLength;
+		}
+		else
+		{
+			// Mark this transform as a zero length child of the last link.
+			// It will inherit position and delta rotation from parent link.
+			FABRIKChainLink & ParentLink = Chain[Chain.Num()-1];
+			ParentLink.ChildZeroLengthTransformIndices.Add(TransformIndex);
+		}
+	}
+
+	bool bBoneLocationUpdated = false;
+	float const RootToTargetDist = FVector::Dist(Chain[0].Position, CSEffectorLocation);
+	int32 const NumChainLinks = Chain.Num();
+
+	// FABRIK algorithm - bone translation calculation
+	// If the effector is further away than the distance from root to tip, simply move all bones in a line from root to effector location
+	if (RootToTargetDist > MaximumReach)
+	{
+		for (int32 LinkIndex = 1; LinkIndex < NumChainLinks; LinkIndex++)
+		{
+			FABRIKChainLink const & ParentLink = Chain[LinkIndex - 1];
+			FABRIKChainLink & CurrentLink = Chain[LinkIndex];
+			CurrentLink.Position = ParentLink.Position + (CSEffectorLocation - ParentLink.Position).UnsafeNormal() * CurrentLink.Length;
+		}
+		bBoneLocationUpdated = true;
+	}
+	else // Effector is within reach, calculate bone translations to position tip at effector location
+	{
+		int32 const TipBoneLinkIndex = NumChainLinks - 1;
+
+		// Check distance between tip location and effector location
+		float Slop = FVector::Dist(Chain[TipBoneLinkIndex].Position, CSEffectorLocation);
+		if (Slop > Precision)
+		{
+			// Set tip bone at end effector location.
+			Chain[TipBoneLinkIndex].Position = CSEffectorLocation;
+
+			int32 IterationCount = 0;
+			while ((Slop > Precision) && (IterationCount++ < MaxIterations))
+			{
+				// "Forward Reaching" stage - adjust bones from end effector.
+				for (int32 LinkIndex = TipBoneLinkIndex - 1; LinkIndex > 0; LinkIndex--)
+				{
+					FABRIKChainLink & CurrentLink = Chain[LinkIndex];
+					FABRIKChainLink const & ChildLink = Chain[LinkIndex + 1];
+
+					CurrentLink.Position = ChildLink.Position + (CurrentLink.Position - ChildLink.Position).UnsafeNormal() * ChildLink.Length;
+				}
+
+				// "Backward Reaching" stage - adjust bones from root.
+				for (int32 LinkIndex = 1; LinkIndex < TipBoneLinkIndex; LinkIndex++)
+				{
+					FABRIKChainLink const & ParentLink = Chain[LinkIndex - 1];
+					FABRIKChainLink & CurrentLink = Chain[LinkIndex];
+
+					CurrentLink.Position = ParentLink.Position + (CurrentLink.Position - ParentLink.Position).UnsafeNormal() * CurrentLink.Length;
+				}
+
+				// Re-check distance between tip location and effector location
+				// Since we're keeping tip on top of effector location, check with its parent bone.
+				Slop = FMath::Abs(Chain[TipBoneLinkIndex].Length - FVector::Dist(Chain[TipBoneLinkIndex - 1].Position, CSEffectorLocation));
+			}
+
+			// Place tip bone based on how close we got to target.
+			{
+				FABRIKChainLink const & ParentLink = Chain[TipBoneLinkIndex - 1];
+				FABRIKChainLink & CurrentLink = Chain[TipBoneLinkIndex];
+
+				CurrentLink.Position = ParentLink.Position + (CurrentLink.Position - ParentLink.Position).UnsafeNormal() * CurrentLink.Length;
+			}
+
+			bBoneLocationUpdated = true;
+		}
+	}
+
+	// If we moved some bones, update bone transforms.
+	if (bBoneLocationUpdated)
+	{
+		// First step: update bone transform positions from chain links.
+		for (int32 LinkIndex = 0; LinkIndex < NumChainLinks; LinkIndex++)
+		{
+			FABRIKChainLink const & ChainLink = Chain[LinkIndex];
+			OutBoneTransforms[ChainLink.TransformIndex].Transform.SetTranslation(ChainLink.Position);
+
+			// If there are any zero length children, update position of those
+			int32 const NumChildren = ChainLink.ChildZeroLengthTransformIndices.Num();
+			for (int32 ChildIndex = 0; ChildIndex < NumChildren; ChildIndex++)
+			{
+				OutBoneTransforms[ChainLink.ChildZeroLengthTransformIndices[ChildIndex]].Transform.SetTranslation(ChainLink.Position);
+			}
+		}
+
+		// FABRIK algorithm - re-orientation of bone local axes after translation calculation
+		for (int32 LinkIndex = 0; LinkIndex < NumChainLinks - 1; LinkIndex++)
+		{
+			FABRIKChainLink const & CurrentLink = Chain[LinkIndex];
+			FABRIKChainLink const & ChildLink = Chain[LinkIndex + 1];
+
+			// Calculate pre-translation vector between this bone and child
+			FVector const OldDir = (GetCurrentLocation(MeshBases, ChildLink.BoneIndex) - GetCurrentLocation(MeshBases, CurrentLink.BoneIndex)).UnsafeNormal();
+
+			// Get vector from the post-translation bone to it's child
+			FVector const NewDir = (ChildLink.Position - CurrentLink.Position).UnsafeNormal();
+
+			// Calculate axis of rotation from pre-translation vector to post-translation vector
+			FVector const RotationAxis = FVector::CrossProduct(OldDir, NewDir).SafeNormal();
+			float const RotationAngle = FMath::Acos(FVector::DotProduct(OldDir, NewDir));
+			FQuat const DeltaRotation = FQuat(RotationAxis, RotationAngle);
+			// We're going to multiply it, in order to not have to re-normalize the final quaternion, it has to be a unit quaternion.
+			checkSlow(DeltaRotation.IsNormalized());
+
+			// Calculate absolute rotation and set it
+			FTransform & CurrentBoneTransform = OutBoneTransforms[CurrentLink.TransformIndex].Transform;
+			CurrentBoneTransform.SetRotation(DeltaRotation * CurrentBoneTransform.GetRotation());
+
+			// Update zero length children if any
+			int32 const NumChildren = CurrentLink.ChildZeroLengthTransformIndices.Num();
+			for (int32 ChildIndex = 0; ChildIndex < NumChildren; ChildIndex++)
+			{
+				FTransform & ChildBoneTransform = OutBoneTransforms[CurrentLink.ChildZeroLengthTransformIndices[ChildIndex]].Transform;
+				ChildBoneTransform.SetRotation(DeltaRotation * ChildBoneTransform.GetRotation());
+			}
+		}
+	}
+
+	// Special handling for tip bone's rotation.
+	int32 const TipBoneTransformIndex = OutBoneTransforms.Num() - 1;
+	switch (EffectorRotationSource)
+	{
+	case BRS_KeepLocalSpaceRotation:
+		OutBoneTransforms[TipBoneTransformIndex].Transform = MeshBases.GetLocalSpaceTransform(BoneIndices[TipBoneTransformIndex]) * OutBoneTransforms[TipBoneTransformIndex - 1].Transform;
+		break;
+	case BRS_CopyFromTarget:
+		OutBoneTransforms[TipBoneTransformIndex].Transform.SetRotation(CSEffectorTransform.GetRotation());
+		break;
+	case BRS_KeepComponentSpaceRotation:
+		// Don't change the orientation at all
+		break;
+	default:
+		break;
+	}
+}
+
+bool FAnimNode_Fabrik::IsValidToEvaluate(const USkeleton * Skeleton, const FBoneContainer & RequiredBones)
+{
+	// Allow evaluation if all parameters are initialized and TipBone is child of RootBone
+	return
+		(
+		TipBone.IsValid(RequiredBones)
+		&& RootBone.IsValid(RequiredBones)
+		&& Precision > 0
+		&& RequiredBones.BoneIsChildOf(TipBone.BoneIndex, RootBone.BoneIndex)
+		);
+}
+
+void FAnimNode_Fabrik::InitializeBoneReferences(const FBoneContainer & RequiredBones) 
+{
+	TipBone.Initialize(RequiredBones);
+	RootBone.Initialize(RequiredBones);
+	EffectorTransformBone.Initialize(RequiredBones);
+}
