Fix scaling being a part of Box and Capsule component

diff --git a/Engine/Source/Runtime/Engine/Private/BoxComponent.cpp b/Engine/Source/Runtime/Engine/Private/BoxComponent.cpp
--- a/Engine/Source/Runtime/Engine/Private/BoxComponent.cpp
+++ b/Engine/Source/Runtime/Engine/Private/BoxComponent.cpp
@@ -61,17 +61,11 @@ void UBoxComponent::UpdateBodySetup()
 		BoxExtent.Z = 1.0f;
 	}
 
-	// apply non uniform scale factor
-	// min scale is applied when body is created
-	FVector Scale3D = ComponentToWorld.GetScale3D().GetAbs();
-	// divided by min since that will be applied to physX
-	Scale3D /= Scale3D.GetMin();
-
 	// now set the PhysX data values
 	se->SetTransform( FTransform::Identity );
-	se->X = BoxExtent.X*2*Scale3D.X;
-	se->Y = BoxExtent.Y*2*Scale3D.Y;
-	se->Z = BoxExtent.Z*2*Scale3D.Z;
+	se->X = BoxExtent.X*2;
+	se->Y = BoxExtent.Y*2;
+	se->Z = BoxExtent.Z*2;
 }
 
 bool UBoxComponent::IsZeroExtent() const
diff --git a/Engine/Source/Runtime/Engine/Private/CapsuleComponent.cpp b/Engine/Source/Runtime/Engine/Private/CapsuleComponent.cpp
--- a/Engine/Source/Runtime/Engine/Private/CapsuleComponent.cpp
+++ b/Engine/Source/Runtime/Engine/Private/CapsuleComponent.cpp
@@ -152,40 +152,10 @@ void UCapsuleComponent::UpdateBodySetup()
 
 	check(ShapeBodySetup->AggGeom.SphylElems.Num() == 1);
 	FKSphylElem* SE = ShapeBodySetup->AggGeom.SphylElems.GetTypedData();
-
-	// apply non uniform scale factor
-	// min scale is applied when body is created
-	FVector Scale3D = ComponentToWorld.GetScale3D().GetAbs();
-	float MinScale = Scale3D.GetMin();
-	float ScaleRadius;
-	float ScaleLength;
-	// do not set scale radius/length to be 0.f
-	// that will cause error on initializing physics
-	// physics will handle with min threshold number	
-	if (MinScale > 0.f)
-	{
-		// divided by min since that min scale will be applied in InitBody
-		ScaleRadius = Scale3D.GetMax() / MinScale;
-		ScaleLength = Scale3D.Z/MinScale;
-	}
-	else
-	{
-		ScaleRadius = 1.f;
-		ScaleLength = 1.f;
-	}
-
-	SE->SetTransform( FTransform::Identity );
-
-	// this is a bit confusing since radius and height is scaled
-	// first apply the scale first 
-	float Radius = FMath::Max(CapsuleRadius * ScaleRadius, 0.1f);
-	float HalfHeight = CapsuleHalfHeight * ScaleLength;
-
-	// now find half height without the caps
-	HalfHeight = FMath::Max<float>(HalfHeight - Radius, 1.f);
-	// set to final value
-	SE->Radius = Radius;
-	SE->Length = 2 * HalfHeight;
+	
+	SE->SetTransform(FTransform::Identity);
+	SE->Radius = CapsuleRadius;
+	SE->Length = 2 * CapsuleHalfHeight;
 }
 
 bool UCapsuleComponent::IsZeroExtent() const
diff --git a/Engine/Source/Runtime/Engine/Private/PhysicsEngine/BodySetup.cpp b/Engine/Source/Runtime/Engine/Private/PhysicsEngine/BodySetup.cpp
--- a/Engine/Source/Runtime/Engine/Private/PhysicsEngine/BodySetup.cpp
+++ b/Engine/Source/Runtime/Engine/Private/PhysicsEngine/BodySetup.cpp
@@ -314,7 +314,7 @@ void UBodySetup::AddShapesToRigidActor(PxRigidActor* PDestActor, FVector& Scale3
 			PxSphereGeometry PSphereGeom;
 			PSphereGeom.radius = (SphereElem->Radius * MinScaleAbs);
 
-			if(PSphereGeom.isValid())
+			if(ensure(PSphereGeom.isValid()))
 			{
 				PxTransform PLocalPose( U2PVector( RelativeTM ? RelativeTM->TransformPosition(SphereElem->Center) : SphereElem->Center ) );
 				PLocalPose.p *= MinScale;
@@ -327,7 +327,7 @@ void UBodySetup::AddShapesToRigidActor(PxRigidActor* PDestActor, FVector& Scale3
 			}
 			else
 			{
-				UE_LOG(LogPhysics, Log, TEXT("AddShapesToRigidActor: [%s] SphereElem[%d] invalid"), *GetPathNameSafe(GetOuter()), i);
+				UE_LOG(LogPhysics, Warning, TEXT("AddShapesToRigidActor: [%s] SphereElem[%d] invalid"), *GetPathNameSafe(GetOuter()), i);
 			}
 		}
 
@@ -337,15 +337,17 @@ void UBodySetup::AddShapesToRigidActor(PxRigidActor* PDestActor, FVector& Scale3
 			FKBoxElem* BoxElem = &AggGeom.BoxElems[i];
 
 			PxBoxGeometry PBoxGeom;
-			PBoxGeom.halfExtents.x = (0.5f * BoxElem->X * MinScaleAbs);
-			PBoxGeom.halfExtents.y = (0.5f * BoxElem->Y * MinScaleAbs);
-			PBoxGeom.halfExtents.z = (0.5f * BoxElem->Z * MinScaleAbs);
+			PBoxGeom.halfExtents.x = (0.5f * BoxElem->X * Scale3DAbs.X);
+			PBoxGeom.halfExtents.y = (0.5f * BoxElem->Y * Scale3DAbs.Y);
+			PBoxGeom.halfExtents.z = (0.5f * BoxElem->Z * Scale3DAbs.Z);
 
 			FTransform BoxTransform = RelativeTM ? BoxElem->GetTransform() * *RelativeTM : BoxElem->GetTransform();
-			if(PBoxGeom.isValid() && BoxTransform.IsValid())
+			if(ensure(PBoxGeom.isValid() && BoxTransform.IsValid()))
 			{
 				PxTransform PLocalPose( U2PTransform( BoxTransform ));
-				PLocalPose.p *= MinScale;
+				PLocalPose.p.x *= Scale3D.X;
+				PLocalPose.p.y *= Scale3D.Y;
+				PLocalPose.p.z *= Scale3D.Z;
 
 				ensure(PLocalPose.isValid());
 				PxShape* NewShape = PDestActor->createShape(PBoxGeom, *PDefaultMat, PLocalPose);
@@ -355,34 +357,49 @@ void UBodySetup::AddShapesToRigidActor(PxRigidActor* PDestActor, FVector& Scale3
 			}
 			else
 			{
-				UE_LOG(LogPhysics, Log, TEXT("AddShapesToRigidActor: [%s] BoxElems[%d] invalid or has invalid transform"), *GetPathNameSafe(GetOuter()), i);
+				UE_LOG(LogPhysics, Warning, TEXT("AddShapesToRigidActor: [%s] BoxElems[%d] invalid or has invalid transform"), *GetPathNameSafe(GetOuter()), i);
 			}
 		}
 
 		// Sphyl (aka Capsule) primitives
-		for (int32 i =0; i < AggGeom.SphylElems.Num(); i++)
 		{
-			FKSphylElem* SphylElem = &AggGeom.SphylElems[i];
-
-			PxCapsuleGeometry PCapsuleGeom;
-			PCapsuleGeom.halfHeight = (0.5f * SphylElem->Length * MinScaleAbs);
-			PCapsuleGeom.radius = (SphylElem->Radius * MinScaleAbs);
+			float ScaleRadius = FMath::Max(Scale3DAbs.X, Scale3DAbs.Y);
+			float ScaleLength = Scale3DAbs.Z;
 
-			if(PCapsuleGeom.isValid())
+			for (int32 i = 0; i < AggGeom.SphylElems.Num(); i++)
 			{
-				// The stored sphyl transform assumes the sphyl axis is down Z. In PhysX, it points down X, so we twiddle the matrix a bit here (swap X and Z and negate Y).
-				PxTransform PLocalPose(U2PVector(RelativeTM ? RelativeTM->TransformPosition(SphylElem->Center) : SphylElem->Center), U2PQuat(SphylElem->Orientation) * U2PSphylBasis);
-				PLocalPose.p *= MinScale;
+				FKSphylElem* SphylElem = &AggGeom.SphylElems[i];
 
-				ensure(PLocalPose.isValid());
-				PxShape* NewShape = PDestActor->createShape(PCapsuleGeom, *PDefaultMat, PLocalPose);
+				// this is a bit confusing since radius and height is scaled
+				// first apply the scale first 
+				float Radius = FMath::Max(SphylElem->Radius * ScaleRadius, 0.1f);
+				float HalfHeight = SphylElem->Length * ScaleLength * 0.5f;
 
-				const float ContactOffset = FMath::Min(MaxContactOffset, ContactOffsetFactor * PCapsuleGeom.radius);
-				NewShape->setContactOffset(ContactOffset);
-			}
-			else
-			{
-				UE_LOG(LogPhysics, Log, TEXT("AddShapesToRigidActor: [%s] SphylElems[%d] invalid"), *GetPathNameSafe(GetOuter()), i);
+				// now find half height without the caps
+				HalfHeight = FMath::Max<float>(HalfHeight - Radius, 1.f);
+
+				PxCapsuleGeometry PCapsuleGeom;
+				PCapsuleGeom.halfHeight = (HalfHeight);
+				PCapsuleGeom.radius = (Radius * ScaleRadius);
+
+				if (ensure(PCapsuleGeom.isValid()))
+				{
+					// The stored sphyl transform assumes the sphyl axis is down Z. In PhysX, it points down X, so we twiddle the matrix a bit here (swap X and Z and negate Y).
+					PxTransform PLocalPose(U2PVector(RelativeTM ? RelativeTM->TransformPosition(SphylElem->Center) : SphylElem->Center), U2PQuat(SphylElem->Orientation) * U2PSphylBasis);
+					PLocalPose.p.x *= Scale3D.X;
+					PLocalPose.p.y *= Scale3D.Y;
+					PLocalPose.p.z *= Scale3D.Z;
+
+					ensure(PLocalPose.isValid());
+					PxShape* NewShape = PDestActor->createShape(PCapsuleGeom, *PDefaultMat, PLocalPose);
+
+					const float ContactOffset = FMath::Min(MaxContactOffset, ContactOffsetFactor * PCapsuleGeom.radius);
+					NewShape->setContactOffset(ContactOffset);
+				}
+				else
+				{
+					UE_LOG(LogPhysics, Warning, TEXT("AddShapesToRigidActor: [%s] SphylElems[%d] invalid"), *GetPathNameSafe(GetOuter()), i);
+				}
 			}
 		}
 
@@ -391,7 +408,7 @@ void UBodySetup::AddShapesToRigidActor(PxRigidActor* PDestActor, FVector& Scale3
 		{
 			FKConvexElem* ConvexElem = &(AggGeom.ConvexElems[i]);
 
-			if(ConvexElem->ConvexMesh != NULL)
+			if(ensure(ConvexElem->ConvexMesh != NULL))
 			{
 				PxTransform PLocalPose;
 				bool bUseNegX = CalcMeshNegScaleCompensation(Scale3D, PLocalPose);
@@ -400,13 +417,13 @@ void UBodySetup::AddShapesToRigidActor(PxRigidActor* PDestActor, FVector& Scale3
 				PConvexGeom.convexMesh = bUseNegX ? ConvexElem->ConvexMeshNegX : ConvexElem->ConvexMesh;
 				PConvexGeom.scale.scale = U2PVector(Scale3DAbs * ConvexElem->GetTransform().GetScale3D().GetAbs());
 				FTransform ConvexTransform = ConvexElem->GetTransform();
-				if ( ConvexTransform.IsValid() )
+				if (ensure(ConvexTransform.IsValid()))
 				{
 					PxTransform PElementTransform = U2PTransform(RelativeTM ? ConvexTransform * *RelativeTM : ConvexTransform);
 					PLocalPose.q *= PElementTransform.q;
 					PLocalPose.p += PElementTransform.p * MinScale;
 
-					if(PConvexGeom.isValid())
+					if(ensure(PConvexGeom.isValid()))
 					{
 						PxVec3 PBoundsExtents = PConvexGeom.convexMesh->getLocalBounds().getExtents();
 
@@ -418,12 +435,12 @@ void UBodySetup::AddShapesToRigidActor(PxRigidActor* PDestActor, FVector& Scale3
 					}
 					else
 					{
-						UE_LOG(LogPhysics, Log, TEXT("AddShapesToRigidActor: [%s] ConvexElem[%d] invalid"), *GetPathNameSafe(GetOuter()), i);
+						UE_LOG(LogPhysics, Warning, TEXT("AddShapesToRigidActor: [%s] ConvexElem[%d] invalid"), *GetPathNameSafe(GetOuter()), i);
 					}
 				}
 				else
 				{
-					UE_LOG(LogPhysics, Log, TEXT("AddShapesToRigidActor: [%s] ConvexElem[%d] has invalid transform"), *GetPathNameSafe(GetOuter()), i);
+					UE_LOG(LogPhysics, Warning, TEXT("AddShapesToRigidActor: [%s] ConvexElem[%d] has invalid transform"), *GetPathNameSafe(GetOuter()), i);
 				}				
 			}
 			else
