UGameUserSettings::LoadSettings() now uses the actual class instead of UGameUserSettings::StaticClass when calling UObject::LoadConfig().

diff --git a/Engine/Source/Runtime/Engine/Private/GameUserSettings.cpp b/Engine/Source/Runtime/Engine/Private/GameUserSettings.cpp
--- a/Engine/Source/Runtime/Engine/Private/GameUserSettings.cpp
+++ b/Engine/Source/Runtime/Engine/Private/GameUserSettings.cpp
@@ -242,7 +242,7 @@ void UGameUserSettings::LoadSettings( bool bForceReload/*=false*/ )
 	{
 		LoadConfigIni( bForceReload );
 	}
-	LoadConfig(UGameUserSettings::StaticClass(), *GGameUserSettingsIni);
+	LoadConfig(GetClass(), *GGameUserSettingsIni);
 
 
 	// Note: Scalability::LoadState() should not be needed as we already loaed the settings earlier (needed so the engine can startup with that before the game is initialized)
