Fixed morphtarget not playing in matinee

diff --git a/Engine/Source/Runtime/Engine/Private/Animation/AnimSingleNodeInstance.cpp b/Engine/Source/Runtime/Engine/Private/Animation/AnimSingleNodeInstance.cpp
--- a/Engine/Source/Runtime/Engine/Private/Animation/AnimSingleNodeInstance.cpp
+++ b/Engine/Source/Runtime/Engine/Private/Animation/AnimSingleNodeInstance.cpp
@@ -193,16 +193,22 @@ void UAnimSingleNodeInstance::NativeInitializeAnimation()
 
 void UAnimSingleNodeInstance::NativeUpdateAnimation(float DeltaTimeX)
 {
-	float NewDeltaTime = DeltaTimeX;
+	float NewPlayRate = PlayRate;
 	UAnimSequence* PreviewBasePose = NULL;
 
-	if(CurrentAsset != NULL && bPlaying)
+	if (bPlaying == false)
+	{
+		// we still have to tick animation when bPlaying is false because 
+		NewPlayRate = 0.f;
+	}
+
+	if(CurrentAsset != NULL)
 	{
 		FAnimGroupInstance* SyncGroup;
 		if (UBlendSpaceBase* BlendSpace = Cast<UBlendSpaceBase>(CurrentAsset))
 		{
 			FAnimTickRecord& TickRecord = CreateUninitializedTickRecord(INDEX_NONE, /*out*/ SyncGroup);
-			MakeBlendSpaceTickRecord(TickRecord, BlendSpace, BlendSpaceInput, BlendSampleData, BlendFilter, bLooping, PlayRate, 1.f, /*inout*/ CurrentTime);
+			MakeBlendSpaceTickRecord(TickRecord, BlendSpace, BlendSpaceInput, BlendSampleData, BlendFilter, bLooping, NewPlayRate, 1.f, /*inout*/ CurrentTime);
 #if WITH_EDITORONLY_DATA
 			PreviewBasePose = BlendSpace->PreviewBasePose;
 #endif
@@ -210,11 +216,11 @@ void UAnimSingleNodeInstance::NativeUpdateAnimation(float DeltaTimeX)
 		else if (UAnimSequence* Sequence = Cast<UAnimSequence>(CurrentAsset))
 		{
 			FAnimTickRecord& TickRecord = CreateUninitializedTickRecord(INDEX_NONE, /*out*/ SyncGroup);
-			MakeSequenceTickRecord(TickRecord, Sequence, bLooping, PlayRate, 1.f, /*inout*/ CurrentTime);
+			MakeSequenceTickRecord(TickRecord, Sequence, bLooping, NewPlayRate, 1.f, /*inout*/ CurrentTime);
 			// if it's not looping, just set play to be false when reached to end
 			if (!bLooping)
 			{
-				if ((PlayRate < 0.f && CurrentTime <= 0.f) || (PlayRate > 0.f && CurrentTime >= Sequence->SequenceLength))
+				if ((NewPlayRate < 0.f && CurrentTime <= 0.f) || (NewPlayRate > 0.f && CurrentTime >= Sequence->SequenceLength))
 				{
 					SetPlaying(false);
 				}
@@ -223,11 +229,11 @@ void UAnimSingleNodeInstance::NativeUpdateAnimation(float DeltaTimeX)
 		else if(UAnimComposite* Composite = Cast<UAnimComposite>(CurrentAsset))
 		{
 			FAnimTickRecord& TickRecord = CreateUninitializedTickRecord(INDEX_NONE, /*out*/ SyncGroup);
-			MakeSequenceTickRecord(TickRecord, Composite, bLooping, PlayRate, 1.f, /*inout*/ CurrentTime);
+			MakeSequenceTickRecord(TickRecord, Composite, bLooping, NewPlayRate, 1.f, /*inout*/ CurrentTime);
 			// if it's not looping, just set play to be false when reached to end
 			if (!bLooping)
 			{
-				if ((PlayRate < 0.f && CurrentTime <= 0.f) || (PlayRate > 0.f && CurrentTime >= Composite->SequenceLength))
+				if ((NewPlayRate < 0.f && CurrentTime <= 0.f) || (NewPlayRate > 0.f && CurrentTime >= Composite->SequenceLength))
 				{
 					SetPlaying(false);
 				}
@@ -258,14 +264,14 @@ void UAnimSingleNodeInstance::NativeUpdateAnimation(float DeltaTimeX)
 	}
 	else if(CurrentVertexAnim != NULL)
 	{
-		float MoveDelta = NewDeltaTime * PlayRate;
+		float MoveDelta = DeltaTimeX * NewPlayRate;
 		FAnimationRuntime::AdvanceTime(bLooping, MoveDelta, CurrentTime, CurrentVertexAnim->GetAnimLength());
 	}
 
 #if WITH_EDITORONLY_DATA
 	if(PreviewBasePose)
 	{
-		float MoveDelta = NewDeltaTime * PlayRate;
+		float MoveDelta = DeltaTimeX * NewPlayRate;
 		const bool bIsPreviewPoseLooping = true;
 
 		FAnimationRuntime::AdvanceTime(bIsPreviewPoseLooping, MoveDelta, PreviewPoseCurrentTime, PreviewBasePose->SequenceLength);
