#ttp 331240 - CORE: Remove GIsBenchamrking, GUseFixedTimeStep from global scope
#proj core
#branch UE4
#summary Moved global variables to statics in FApp. Removed unused GUnclampedDeltaTime.

diff --git a/Engine/Plugins/Experimental/Paper2D/Source/Paper2DEditor/Private/PaperAnimatedSpriteThumbnailRenderer.cpp b/Engine/Plugins/Experimental/Paper2D/Source/Paper2DEditor/Private/PaperAnimatedSpriteThumbnailRenderer.cpp
--- a/Engine/Plugins/Experimental/Paper2D/Source/Paper2DEditor/Private/PaperAnimatedSpriteThumbnailRenderer.cpp
+++ b/Engine/Plugins/Experimental/Paper2D/Source/Paper2DEditor/Private/PaperAnimatedSpriteThumbnailRenderer.cpp
@@ -16,7 +16,7 @@ void UPaperFlipbookThumbnailRenderer::Draw(UObject* Object, int32 X, int32 Y, ui
 	{
 		if (Flipbook->KeyFrames.Num() > 0)
 		{
-			const double DeltaTime = GCurrentTime - GStartTime;
+			const double DeltaTime = FApp::GetCurrentTime() - GStartTime;
 			const float PlayTime = FMath::Fmod(DeltaTime, Flipbook->GetDuration());
 
 			UPaperSprite* Sprite = Flipbook->GetSpriteAtTime(PlayTime);
diff --git a/Engine/Source/Developer/CrashTracker/Private/CrashTrackerModule.cpp b/Engine/Source/Developer/CrashTracker/Private/CrashTrackerModule.cpp
--- a/Engine/Source/Developer/CrashTracker/Private/CrashTrackerModule.cpp
+++ b/Engine/Source/Developer/CrashTracker/Private/CrashTrackerModule.cpp
@@ -118,7 +118,7 @@ void FCrashTrackerModule::StartupModule()
 			false && // Disable crash tracker by default for now unless someone enables it.
 			!GIsDemoMode &&
 			!FParse::Param( FCommandLine::Get(), TEXT("disablecrashtracker") ) &&
-			!GIsBenchmarking &&
+			!FApp::IsBenchmarking() &&
 			!FPlatformMisc::IsDebuggerPresent();
 #endif
 		bool bForceEnableCrashTracker = FParse::Param( FCommandLine::Get(), TEXT("forceenablecrashtracker") );
diff --git a/Engine/Source/Editor/GraphEditor/Private/ConnectionDrawingPolicy.cpp b/Engine/Source/Editor/GraphEditor/Private/ConnectionDrawingPolicy.cpp
--- a/Engine/Source/Editor/GraphEditor/Private/ConnectionDrawingPolicy.cpp
+++ b/Engine/Source/Editor/GraphEditor/Private/ConnectionDrawingPolicy.cpp
@@ -742,9 +742,9 @@ void FKismetConnectionDrawingPolicy::BuildExecutionRoadmap()
 		// executions (so we don't care about this sequence of statements)
 	}
 
-	// Fade only when free-running (since we're using GCurrentTime, instead of FPlatformTime::Seconds)
-	const double MaxTimeAhead = FMath::Min(GCurrentTime + 2*TracePositionBonusPeriod, LatestTimeDiscovered); //@TODO: Rough clamping; should be exposed as a parameter
-	CurrentTime = FMath::Max(GCurrentTime, MaxTimeAhead);
+	// Fade only when free-running (since we're using FApp::GetCurrentTime(), instead of FPlatformTime::Seconds)
+	const double MaxTimeAhead = FMath::Min(FApp::GetCurrentTime() + 2*TracePositionBonusPeriod, LatestTimeDiscovered); //@TODO: Rough clamping; should be exposed as a parameter
+	CurrentTime = FMath::Max(FApp::GetCurrentTime(), MaxTimeAhead);
 }
 
 void FKismetConnectionDrawingPolicy::CalculateEnvelopeAlphas(double ExecutionTime, /*out*/ float& AttackAlpha, /*out*/ float& SustainAlpha) const
@@ -1229,11 +1229,11 @@ void FSoundCueGraphConnectionDrawingPolicy::BuildAudioFlowRoadmap()
 					GraphNodes.Add(RootNode[0]);
 
 					TArray<double> NodeTimes;
-					NodeTimes.Add(GCurrentTime); // Time for the root node
+					NodeTimes.Add(FApp::GetCurrentTime()); // Time for the root node
 
 					for (int32 i = 0; i < PathToWaveInstance.Num(); ++i)
 					{
-						const double ObservationTime = GCurrentTime + 1.f;
+						const double ObservationTime = FApp::GetCurrentTime() + 1.f;
 
 						NodeTimes.Add(ObservationTime);
 						GraphNodes.Add(PathToWaveInstance[i]->GraphNode);
diff --git a/Engine/Source/Editor/GraphEditor/Private/MaterialNodes/SGraphNodeMaterialBase.cpp b/Engine/Source/Editor/GraphEditor/Private/MaterialNodes/SGraphNodeMaterialBase.cpp
--- a/Engine/Source/Editor/GraphEditor/Private/MaterialNodes/SGraphNodeMaterialBase.cpp
+++ b/Engine/Source/Editor/GraphEditor/Private/MaterialNodes/SGraphNodeMaterialBase.cpp
@@ -183,8 +183,8 @@ void FPreviewElement::DrawRenderThread( const void* InWindowBackBuffer )
 	RenderTarget->SetRenderTargetTexture( *(FTexture2DRHIRef*)InWindowBackBuffer );
 	{
 		// Check realtime mode for whether to pass current time to canvas
-		float CurrentTime = bIsRealtime ? (GCurrentTime - GStartTime) : 0.0f;
-		float DeltaTime = bIsRealtime ? GDeltaTime : 0.0f;
+		float CurrentTime = bIsRealtime ? (FApp::GetCurrentTime() - GStartTime) : 0.0f;
+		float DeltaTime = bIsRealtime ? FApp::GetDeltaTime() : 0.0f;
 
 		FCanvas Canvas(RenderTarget, NULL, CurrentTime, CurrentTime, DeltaTime);
 		{
diff --git a/Engine/Source/Editor/Matinee/Private/Matinee.cpp b/Engine/Source/Editor/Matinee/Private/Matinee.cpp
--- a/Engine/Source/Editor/Matinee/Private/Matinee.cpp
+++ b/Engine/Source/Editor/Matinee/Private/Matinee.cpp
@@ -2029,7 +2029,7 @@ void FMatinee::OnClose()
 	}
 
 	// Make sure benchmarking mode is disabled (we may have turned it on for 'fixed time step playback')
-	GIsBenchmarking = false;
+	FApp::SetBenchmarking(false);
 
 	// Update UI to reflect any change in realtime status
 	FEditorSupportDelegates::UpdateUI.Broadcast();
diff --git a/Engine/Source/Editor/Matinee/Private/MatineeTools.cpp b/Engine/Source/Editor/Matinee/Private/MatineeTools.cpp
--- a/Engine/Source/Editor/Matinee/Private/MatineeTools.cpp
+++ b/Engine/Source/Editor/Matinee/Private/MatineeTools.cpp
@@ -265,7 +265,7 @@ void FMatinee::ConstrainFixedTimeStepFrameRate()
 	if( bSnapToFrames && bFixedTimeStepPlayback )
 	{
 		// NOTE: Its important that PlaybackStartRealTime and NumContinuousFixedTimeStepFrames are reset
-		//    when anything timing-related changes, like GFixedDeltaTime or playback direction.
+		//    when anything timing-related changes, like FApp::GetFixedDeltaTime() or playback direction.
 
 		double CurRealTime = FPlatformTime::Seconds();
 
@@ -283,7 +283,7 @@ void FMatinee::ConstrainFixedTimeStepFrameRate()
 
 		// How long should have it taken to get to the current frame?
 		const double ExpectedPlaybackTime =
-			NumContinuousFixedTimeStepFrames * GFixedDeltaTime * PlaybackSpeed;
+			NumContinuousFixedTimeStepFrames * FApp::GetFixedDeltaTime() * PlaybackSpeed;
 
 		// How long has it been (in real-time) since we started playback?
 		double RealTimeSincePlaybackStarted = CurRealTime - PlaybackStartRealTime;
@@ -1584,10 +1584,11 @@ void FMatinee::SetFixedTimeStepPlayback( bool bInValue )
 void FMatinee::UpdateFixedTimeStepPlayback()
 {
 	// Turn on 'benchmarking' mode if we're using a fixed time step
-	GIsBenchmarking = MatineeActor->bIsPlaying && bSnapToFrames && bFixedTimeStepPlayback;
+	bool bIsBenchmarking = MatineeActor->bIsPlaying && bSnapToFrames && bFixedTimeStepPlayback;
+	FApp::SetBenchmarking(bIsBenchmarking);
 
 	// Set the time interval between fixed ticks
-	GFixedDeltaTime = SnapAmount;
+	FApp::SetFixedDeltaTime(SnapAmount);
 }
 
 
diff --git a/Engine/Source/Editor/StatsViewer/Private/StatsPages/TextureStatsPage.cpp b/Engine/Source/Editor/StatsViewer/Private/StatsPages/TextureStatsPage.cpp
--- a/Engine/Source/Editor/StatsViewer/Private/StatsPages/TextureStatsPage.cpp
+++ b/Engine/Source/Editor/StatsViewer/Private/StatsPages/TextureStatsPage.cpp
@@ -202,7 +202,7 @@ struct TextureStatsGenerator : public FFindReferencedAssets
 			const FTexture* Resource = InTexture->Resource; 
 			if(Resource)
 			{
-				Entry->LastTimeRendered = (float)FMath::Max( GLastTime - Resource->LastRenderTime, 0.0 );
+				Entry->LastTimeRendered = (float)FMath::Max( FApp::GetLastTime() - Resource->LastRenderTime, 0.0 );
 			}
 
 			const UTexture2D* Texture2D = Cast<const UTexture2D>(InTexture);
diff --git a/Engine/Source/Editor/UnrealEd/Private/Commandlets/FileServerCommandlet.cpp b/Engine/Source/Editor/UnrealEd/Private/Commandlets/FileServerCommandlet.cpp
--- a/Engine/Source/Editor/UnrealEd/Private/Commandlets/FileServerCommandlet.cpp
+++ b/Engine/Source/Editor/UnrealEd/Private/Commandlets/FileServerCommandlet.cpp
@@ -93,11 +93,11 @@ int32 UFileServerCommandlet::Main( const FString& Params )
 	while (GIsRunning && !GIsRequestingExit)
 	{
 		GEngine->UpdateTimeAndHandleMaxTickRate();
-		GEngine->Tick(GDeltaTime, false);
+		GEngine->Tick(FApp::GetDeltaTime(), false);
 
 		// tick the directory watcher
 		FDirectoryWatcherModule& DirectoryWatcherModule = FModuleManager::Get().LoadModuleChecked<FDirectoryWatcherModule>(TEXT("DirectoryWatcher"));
-		DirectoryWatcherModule.Get()->Tick(GDeltaTime);
+		DirectoryWatcherModule.Get()->Tick(FApp::GetDeltaTime());
 
 		// update task graph
 		FTaskGraphInterface::Get().ProcessThreadUntilIdle(ENamedThreads::GameThread);
diff --git a/Engine/Source/Editor/UnrealEd/Private/Editor.cpp b/Engine/Source/Editor/UnrealEd/Private/Editor.cpp
--- a/Engine/Source/Editor/UnrealEd/Private/Editor.cpp
+++ b/Engine/Source/Editor/UnrealEd/Private/Editor.cpp
@@ -347,7 +347,7 @@ void UEditorEngine::InitEditor(IEngineLoop* InEngineLoop)
 
     if( !FEngineBuildSettings::IsInternalBuild() &&
 		!FEngineBuildSettings::IsPerforceBuild() && 
-		!GIsBenchmarking &&
+		!FApp::IsBenchmarking() &&
 		!GIsDemoMode && 
 		!IsRunningCommandlet() &&
 		!FPlatformProcess::IsApplicationRunning(TEXT("UnrealEngineLauncher") ) &&
diff --git a/Engine/Source/Editor/UnrealEd/Private/EditorExporters.cpp b/Engine/Source/Editor/UnrealEd/Private/EditorExporters.cpp
--- a/Engine/Source/Editor/UnrealEd/Private/EditorExporters.cpp
+++ b/Engine/Source/Editor/UnrealEd/Private/EditorExporters.cpp
@@ -2524,7 +2524,7 @@ namespace MaterialExportUtils
 		RHIBeginScene();
 		{
 			// Create a canvas for the render target and clear it to black
-			FCanvas Canvas(RTResource, NULL, GCurrentTime - GStartTime, GDeltaTime, GCurrentTime - GStartTime );
+			FCanvas Canvas(RTResource, NULL, FApp::GetCurrentTime() - GStartTime, FApp::GetDeltaTime(), FApp::GetCurrentTime() - GStartTime);
 			Canvas.Clear(FLinearColor::Black);
 			FCanvasTileItem TileItem(FVector2D(0.0f, 0.0f), MaterialProxy, FVector2D(InRenderTarget->SizeX, InRenderTarget->SizeY));
 			TileItem.bFreezeTime = true;
@@ -2670,7 +2670,7 @@ namespace MaterialExportUtils
 
 				FSceneViewFamilyContext ViewFamily(
 					FSceneViewFamily::ConstructionValues(RenderTargetResource, Scene, FEngineShowFlags(ESFIM_Game))
-						.SetWorldTimes(GCurrentTime - GStartTime, GDeltaTime, GCurrentTime - GStartTime)
+						.SetWorldTimes(FApp::GetCurrentTime() - GStartTime, FApp::GetDeltaTime(), FApp::GetCurrentTime() - GStartTime)
 					);
 
 				ViewFamily.EngineShowFlags.DisableAdvancedFeatures();
@@ -2720,7 +2720,7 @@ namespace MaterialExportUtils
 					}
 				}
 			
-				FCanvas Canvas(RenderTargetResource, NULL, GCurrentTime - GStartTime, GDeltaTime, GCurrentTime - GStartTime);
+				FCanvas Canvas(RenderTargetResource, NULL, FApp::GetCurrentTime() - GStartTime, FApp::GetDeltaTime(), FApp::GetCurrentTime() - GStartTime);
 				Canvas.Clear(FLinearColor::Black);
 				GetRendererModule().BeginRenderingViewFamily(&Canvas, &ViewFamily);
 
diff --git a/Engine/Source/Editor/UnrealEd/Private/EditorViewportClient.cpp b/Engine/Source/Editor/UnrealEd/Private/EditorViewportClient.cpp
--- a/Engine/Source/Editor/UnrealEd/Private/EditorViewportClient.cpp
+++ b/Engine/Source/Editor/UnrealEd/Private/EditorViewportClient.cpp
@@ -2431,9 +2431,9 @@ void FEditorViewportClient::Draw(FViewport* InViewport, FCanvas* Canvas)
 	if (( GetScene() != World->Scene) || (IsRealtime() == true))
 	{
 		// Use time relative to start time to avoid issues with float vs double
-		TimeSeconds = GCurrentTime - GStartTime;
-		RealTimeSeconds = GCurrentTime - GStartTime;
-		DeltaTimeSeconds = GDeltaTime;
+		TimeSeconds = FApp::GetCurrentTime() - GStartTime;
+		RealTimeSeconds = FApp::GetCurrentTime() - GStartTime;
+		DeltaTimeSeconds = FApp::GetDeltaTime();
 	}
 	else
 	{
diff --git a/Engine/Source/Editor/UnrealEd/Private/ObjectTools.cpp b/Engine/Source/Editor/UnrealEd/Private/ObjectTools.cpp
--- a/Engine/Source/Editor/UnrealEd/Private/ObjectTools.cpp
+++ b/Engine/Source/Editor/UnrealEd/Private/ObjectTools.cpp
@@ -3284,7 +3284,7 @@ namespace ThumbnailTools
 		});
 
 		// Create a canvas for the render target and clear it to black
-		FCanvas Canvas( RenderTargetResource, NULL, GCurrentTime - GStartTime, GDeltaTime, GCurrentTime - GStartTime );
+		FCanvas Canvas( RenderTargetResource, NULL, FApp::GetCurrentTime() - GStartTime, FApp::GetDeltaTime(), FApp::GetCurrentTime() - GStartTime );
 		Canvas.Clear( FLinearColor::Black );
 
 
diff --git a/Engine/Source/Editor/UnrealEd/Private/ThumbnailRendering/BlueprintThumbnailRenderer.cpp b/Engine/Source/Editor/UnrealEd/Private/ThumbnailRendering/BlueprintThumbnailRenderer.cpp
--- a/Engine/Source/Editor/UnrealEd/Private/ThumbnailRendering/BlueprintThumbnailRenderer.cpp
+++ b/Engine/Source/Editor/UnrealEd/Private/ThumbnailRendering/BlueprintThumbnailRenderer.cpp
@@ -23,7 +23,7 @@ void UBlueprintThumbnailRenderer::Draw(UObject* Object, int32 X, int32 Y, uint32
 
 		ThumbnailScene->SetBlueprint(Blueprint);
 		FSceneViewFamilyContext ViewFamily( FSceneViewFamily::ConstructionValues( RenderTarget, ThumbnailScene->GetScene(), FEngineShowFlags(ESFIM_Game) )
-			.SetWorldTimes(GCurrentTime - GStartTime,GDeltaTime,GCurrentTime - GStartTime ));
+			.SetWorldTimes(FApp::GetCurrentTime() - GStartTime, FApp::GetDeltaTime(), FApp::GetCurrentTime() - GStartTime));
 
 		ViewFamily.EngineShowFlags.DisableAdvancedFeatures();
 		ViewFamily.EngineShowFlags.MotionBlur = 0;
diff --git a/Engine/Source/Editor/UnrealEd/Private/ThumbnailRendering/LevelThumbnailRenderer.cpp b/Engine/Source/Editor/UnrealEd/Private/ThumbnailRendering/LevelThumbnailRenderer.cpp
--- a/Engine/Source/Editor/UnrealEd/Private/ThumbnailRendering/LevelThumbnailRenderer.cpp
+++ b/Engine/Source/Editor/UnrealEd/Private/ThumbnailRendering/LevelThumbnailRenderer.cpp
@@ -13,7 +13,7 @@ void ULevelThumbnailRenderer::Draw(UObject* Object, int32 X, int32 Y, uint32 Wid
 	if (Level != nullptr)
 	{
 		FSceneViewFamilyContext ViewFamily( FSceneViewFamily::ConstructionValues( RenderTarget, Level->OwningWorld->Scene, FEngineShowFlags(ESFIM_Game) )
-			.SetWorldTimes(GCurrentTime - GStartTime,GDeltaTime,GCurrentTime - GStartTime) );
+			.SetWorldTimes(FApp::GetCurrentTime() - GStartTime, FApp::GetDeltaTime(), FApp::GetCurrentTime() - GStartTime));
 
 		ViewFamily.EngineShowFlags.DisableAdvancedFeatures();
 		ViewFamily.EngineShowFlags.MotionBlur = 0;
diff --git a/Engine/Source/Editor/UnrealEd/Private/ThumbnailRendering/MaterialFunctionThumbnailRenderer.cpp b/Engine/Source/Editor/UnrealEd/Private/ThumbnailRendering/MaterialFunctionThumbnailRenderer.cpp
--- a/Engine/Source/Editor/UnrealEd/Private/ThumbnailRendering/MaterialFunctionThumbnailRenderer.cpp
+++ b/Engine/Source/Editor/UnrealEd/Private/ThumbnailRendering/MaterialFunctionThumbnailRenderer.cpp
@@ -27,7 +27,7 @@ void UMaterialFunctionThumbnailRenderer::Draw(UObject* Object, int32 X, int32 Y,
 			PreviewMaterial->ThumbnailInfo = MatFunc->ThumbnailInfo;
 			ThumbnailScene->SetMaterialInterface( PreviewMaterial );
 			FSceneViewFamilyContext ViewFamily( FSceneViewFamily::ConstructionValues( RenderTarget, ThumbnailScene->GetScene(), FEngineShowFlags(ESFIM_Game) )
-				.SetWorldTimes(GCurrentTime - GStartTime,GDeltaTime,GCurrentTime - GStartTime) );
+				.SetWorldTimes(FApp::GetCurrentTime() - GStartTime, FApp::GetDeltaTime(), FApp::GetCurrentTime() - GStartTime));
 
 			ViewFamily.EngineShowFlags.DisableAdvancedFeatures();
 			ViewFamily.EngineShowFlags.MotionBlur = 0;
diff --git a/Engine/Source/Editor/UnrealEd/Private/ThumbnailRendering/MaterialInstanceThumbnailRenderer.cpp b/Engine/Source/Editor/UnrealEd/Private/ThumbnailRendering/MaterialInstanceThumbnailRenderer.cpp
--- a/Engine/Source/Editor/UnrealEd/Private/ThumbnailRendering/MaterialInstanceThumbnailRenderer.cpp
+++ b/Engine/Source/Editor/UnrealEd/Private/ThumbnailRendering/MaterialInstanceThumbnailRenderer.cpp
@@ -23,7 +23,7 @@ void UMaterialInstanceThumbnailRenderer::Draw(UObject* Object, int32 X, int32 Y,
 
 		ThumbnailScene->SetMaterialInterface(MatInst);
 		FSceneViewFamilyContext ViewFamily( FSceneViewFamily::ConstructionValues( RenderTarget, ThumbnailScene->GetScene(), FEngineShowFlags(ESFIM_Game) )
-			.SetWorldTimes(GCurrentTime - GStartTime,GDeltaTime,GCurrentTime - GStartTime) );
+			.SetWorldTimes(FApp::GetCurrentTime() - GStartTime, FApp::GetDeltaTime(), FApp::GetCurrentTime() - GStartTime));
 
 		ViewFamily.EngineShowFlags.DisableAdvancedFeatures();
 		ViewFamily.EngineShowFlags.MotionBlur = 0;
diff --git a/Engine/Source/Editor/UnrealEd/Private/ThumbnailRendering/ParticleSystemThumbnailRenderer.cpp b/Engine/Source/Editor/UnrealEd/Private/ThumbnailRendering/ParticleSystemThumbnailRenderer.cpp
--- a/Engine/Source/Editor/UnrealEd/Private/ThumbnailRendering/ParticleSystemThumbnailRenderer.cpp
+++ b/Engine/Source/Editor/UnrealEd/Private/ThumbnailRendering/ParticleSystemThumbnailRenderer.cpp
@@ -67,7 +67,7 @@ void UParticleSystemThumbnailRenderer::Draw(UObject* Object, int32 X, int32 Y, u
 
 				ThumbnailScene->SetParticleSystem(ParticleSystem);
 				FSceneViewFamilyContext ViewFamily( FSceneViewFamily::ConstructionValues( RenderTarget, ThumbnailScene->GetScene(), FEngineShowFlags(ESFIM_Game))
-					.SetWorldTimes(GCurrentTime - GStartTime,GDeltaTime,GCurrentTime - GStartTime) );
+					.SetWorldTimes(FApp::GetCurrentTime() - GStartTime, FApp::GetDeltaTime(), FApp::GetCurrentTime() - GStartTime));
 			
 				ViewFamily.EngineShowFlags.DisableAdvancedFeatures();
 				ViewFamily.EngineShowFlags.MotionBlur = 0;
diff --git a/Engine/Source/Editor/UnrealEd/Private/ThumbnailRendering/SkeletalMeshThumbnailRenderer.cpp b/Engine/Source/Editor/UnrealEd/Private/ThumbnailRendering/SkeletalMeshThumbnailRenderer.cpp
--- a/Engine/Source/Editor/UnrealEd/Private/ThumbnailRendering/SkeletalMeshThumbnailRenderer.cpp
+++ b/Engine/Source/Editor/UnrealEd/Private/ThumbnailRendering/SkeletalMeshThumbnailRenderer.cpp
@@ -23,7 +23,7 @@ void USkeletalMeshThumbnailRenderer::Draw(UObject* Object, int32 X, int32 Y, uin
 
 		ThumbnailScene->SetSkeletalMesh(SkeletalMesh);
 		FSceneViewFamilyContext ViewFamily( FSceneViewFamily::ConstructionValues( RenderTarget, ThumbnailScene->GetScene(), FEngineShowFlags(ESFIM_Game) )
-			.SetWorldTimes(GCurrentTime - GStartTime,GDeltaTime,GCurrentTime - GStartTime) );
+			.SetWorldTimes(FApp::GetCurrentTime() - GStartTime, FApp::GetDeltaTime(), FApp::GetCurrentTime() - GStartTime));
 
 		ViewFamily.EngineShowFlags.DisableAdvancedFeatures();
 		ViewFamily.EngineShowFlags.MotionBlur = 0;
diff --git a/Engine/Source/Editor/UnrealEd/Private/ThumbnailRendering/StaticMeshThumbnailRenderer.cpp b/Engine/Source/Editor/UnrealEd/Private/ThumbnailRendering/StaticMeshThumbnailRenderer.cpp
--- a/Engine/Source/Editor/UnrealEd/Private/ThumbnailRendering/StaticMeshThumbnailRenderer.cpp
+++ b/Engine/Source/Editor/UnrealEd/Private/ThumbnailRendering/StaticMeshThumbnailRenderer.cpp
@@ -25,7 +25,7 @@ void UStaticMeshThumbnailRenderer::Draw(UObject* Object, int32 X, int32 Y, uint3
 		ThumbnailScene->GetScene()->UpdateSpeedTreeWind(0.0);
 
 		FSceneViewFamilyContext ViewFamily( FSceneViewFamily::ConstructionValues( RenderTarget, ThumbnailScene->GetScene(), FEngineShowFlags(ESFIM_Game) )
-			.SetWorldTimes(GCurrentTime - GStartTime,GDeltaTime,GCurrentTime - GStartTime ));
+			.SetWorldTimes(FApp::GetCurrentTime() - GStartTime, FApp::GetDeltaTime(), FApp::GetCurrentTime() - GStartTime));
 
 		ViewFamily.EngineShowFlags.DisableAdvancedFeatures();
 		ViewFamily.EngineShowFlags.MotionBlur = 0;
diff --git a/Engine/Source/Editor/UnrealEd/Private/VertexSnapping.cpp b/Engine/Source/Editor/UnrealEd/Private/VertexSnapping.cpp
--- a/Engine/Source/Editor/UnrealEd/Private/VertexSnapping.cpp
+++ b/Engine/Source/Editor/UnrealEd/Private/VertexSnapping.cpp
@@ -296,7 +296,7 @@ void FVertexSnappingImpl::ClearSnappingHelpers( bool bClearImmediately )
 	else if( ActorVertsToDraw.IsValid() )
 	{
 		// Fade out previous verts
-		ActorVertsToFade.Add( ActorVertsToDraw, GCurrentTime );
+		ActorVertsToFade.Add( ActorVertsToDraw, FApp::GetCurrentTime() );
 		ActorVertsToDraw.Reset();
 	}
 }
@@ -342,15 +342,15 @@ void FVertexSnappingImpl::DrawSnappingHelpers(const FSceneView* View,FPrimitiveD
 		{
 			float PointSize = View->IsPerspectiveProjection() ? 4.0f : 5.0f;
 
-			if( GCurrentTime-FadeStart <= VertexSnappingConstants::FadeTime )
+			if( FApp::GetCurrentTime()-FadeStart <= VertexSnappingConstants::FadeTime )
 			{
-				PointSize = FMath::Lerp( PointSize, 0.0f, (GCurrentTime-FadeStart)/VertexSnappingConstants::FadeTime );
+				PointSize = FMath::Lerp( PointSize, 0.0f, (FApp::GetCurrentTime()-FadeStart)/VertexSnappingConstants::FadeTime );
 
 				DrawSnapVertices( Actor.Get(), PointSize, PDI );
 			}
 		}
 	
-		if( !Actor.IsValid() || GCurrentTime-FadeStart > VertexSnappingConstants::FadeTime )
+		if( !Actor.IsValid() || FApp::GetCurrentTime()-FadeStart > VertexSnappingConstants::FadeTime )
 		{
 			It.RemoveCurrent();
 		}
@@ -547,7 +547,7 @@ FSnappingVertex FVertexSnappingImpl::GetClosestVertex( const TArray<FSnapActor>&
 
 	if( ActorVertsToDraw.IsValid() )
 	{
-		ActorVertsToFade.Add( ActorVertsToDraw, GCurrentTime );
+		ActorVertsToFade.Add( ActorVertsToDraw, FApp::GetCurrentTime() );
 	}
 
 
diff --git a/Engine/Source/Runtime/Core/Private/Misc/App.cpp b/Engine/Source/Runtime/Core/Private/Misc/App.cpp
--- a/Engine/Source/Runtime/Core/Private/Misc/App.cpp
+++ b/Engine/Source/Runtime/Core/Private/Misc/App.cpp
@@ -16,7 +16,12 @@ FGuid FApp::SessionId = FGuid::NewGuid();
 FString FApp::SessionName = FString();
 FString FApp::SessionOwner = FString();
 bool FApp::Standalone = true;
-
+bool FApp::bIsBenchmarking = false;
+double FApp::FixedDeltaTime = 1 / 30.0;
+bool FApp::bUseFixedTimeStep = false;
+double FApp::CurrentTime = 0.0;
+double FApp::LastTime = 0.0;
+double FApp::DeltaTime = 1 / 30.0;
 
 /* FApp static interface
  *****************************************************************************/
diff --git a/Engine/Source/Runtime/Core/Private/Misc/Core.cpp b/Engine/Source/Runtime/Core/Private/Misc/Core.cpp
--- a/Engine/Source/Runtime/Core/Private/Misc/Core.cpp
+++ b/Engine/Source/Runtime/Core/Private/Misc/Core.cpp
@@ -111,7 +111,6 @@ bool					GSlowTaskOccurred				= false;					/* Whether a slow task began last tic
 bool					GIsRequestingExit				= false;					/* Indicates that MainLoop() should be exited at the end of the current iteration */
 /** Archive for serializing arbitrary data to and from memory												*/
 FReloadObjectArc*		GMemoryArchive					= NULL;
-bool					GIsBenchmarking					= 0;						/* Whether we are in benchmark mode or not */
 bool					GAreScreenMessagesEnabled		= true;						/* Whether onscreen warnings/messages are enabled */
 bool					GScreenMessagesRestoreState		= false;					/* Used to restore state after a screenshot */
 int32					GIsDumpingMovie					= 0;						/* Whether we are dumping screenshots (!= 0), exposed as console variable r.DumpingMovie */
@@ -133,16 +132,6 @@ FString				GGameIni;													/* Game ini filename */
 FString				GGameUserSettingsIni;										/* User Game Settings ini filename */
 
 float					GNearClippingPlane				= 10.0f;					/* Near clipping plane */
-/** Timestep if a fixed delta time is wanted.																*/
-double					GFixedDeltaTime					= 1 / 30.f;
-/** Current delta time in seconds.																			*/
-double					GDeltaTime						= 1 / 30.f;
-/** Current unclamped delta time in seconds.																*/
-double					GUnclampedDeltaTime				= 1 / 30.f;
-/* Current time																								*/
-double					GCurrentTime					= 0;						
-/* Last GCurrentTime																						*/
-double					GLastTime						= 0;						
 
 bool					GExitPurge						= false;
 
@@ -215,8 +204,6 @@ void					(*GFlushStreamingFunc)(void)	  = &appNoop;
 bool					GEmitDrawEvents					= false;
 /** Whether we want the rendering thread to be suspended, used e.g. for tracing.							*/
 bool					GShouldSuspendRenderingThread	= false;
-/** Whether we want to use a fixed time step or not.														*/
-bool					GUseFixedTimeStep				= false;
 /** Determines what kind of trace should occur, NAME_None for none.											*/
 FName					GCurrentTraceName				= NAME_None;
 /** How to print the time in log output																		*/
diff --git a/Engine/Source/Runtime/Core/Private/ProfilingDebugging/MallocProfiler.cpp b/Engine/Source/Runtime/Core/Private/ProfilingDebugging/MallocProfiler.cpp
--- a/Engine/Source/Runtime/Core/Private/ProfilingDebugging/MallocProfiler.cpp
+++ b/Engine/Source/Runtime/Core/Private/ProfilingDebugging/MallocProfiler.cpp
@@ -793,7 +793,7 @@ bool FMallocProfiler::HandleSnapshotMemoryFrameCommand( const TCHAR* Cmd, FOutpu
 {
 	if (!bEndProfilingHasBeenCalled)
 	{
-		EmbedFloatMarker(SUBTYPE_FrameTimeMarker, (float)GDeltaTime);
+		EmbedFloatMarker(SUBTYPE_FrameTimeMarker, (float)FApp::GetDeltaTime());
 	}
 	return true;
 }
@@ -1136,4 +1136,4 @@ FScopedMallocProfilerLock::~FScopedMallocProfilerLock()
 // Suppress linker warning "warning LNK4221: no public symbols found; archive member will be inaccessible"
 int32 FMallocProfilerLinkerHelper;
 
-#endif //
\ No newline at end of file
+#endif //
diff --git a/Engine/Source/Runtime/Core/Public/Core.h b/Engine/Source/Runtime/Core/Public/Core.h
--- a/Engine/Source/Runtime/Core/Public/Core.h
+++ b/Engine/Source/Runtime/Core/Public/Core.h
@@ -220,7 +220,6 @@ extern CORE_API bool GIsRequestingExit;
 
 /** Archive for serializing arbitrary data to and from memory						*/
 extern CORE_API class FReloadObjectArc* GMemoryArchive;
-extern CORE_API bool GIsBenchmarking;
 
 /**
  *	Global value indicating on-screen warnings/message should be displayed.
@@ -251,17 +250,6 @@ extern CORE_API FString GGameUserSettingsIni;
 
 extern CORE_API float GNearClippingPlane;
 
-/** Timestep if a fixed delta time is wanted. */
-extern CORE_API double GFixedDeltaTime;
-
-/** Current delta time in seconds. */
-extern CORE_API double GDeltaTime;
-
-/** Unclamped delta time in seconds. */
-extern CORE_API double GUnclampedDeltaTime;
-extern CORE_API double GCurrentTime;
-extern CORE_API double GLastTime;
-
 extern CORE_API bool GExitPurge;
 extern CORE_API TCHAR GGameName[64];
 
@@ -342,9 +330,6 @@ extern CORE_API bool GEmitDrawEvents;
 /** Whether we want the rendering thread to be suspended, used e.g. for tracing. */
 extern CORE_API bool GShouldSuspendRenderingThread;
 
-/** Whether we want to use a fixed time step or not. */
-extern CORE_API bool GUseFixedTimeStep;
-
 /** Determines what kind of trace should occur, NAME_None for none. */
 extern CORE_API FName GCurrentTraceName;
 
diff --git a/Engine/Source/Runtime/Core/Public/Misc/App.h b/Engine/Source/Runtime/Core/Public/Misc/App.h
--- a/Engine/Source/Runtime/Core/Public/Misc/App.h
+++ b/Engine/Source/Runtime/Core/Public/Misc/App.h
@@ -19,7 +19,7 @@ public:
 	 *
 	 * @return true if the application can render, false otherwise.
 	 */
-	FORCEINLINE static bool CanEverRender( )
+	FORCEINLINE static bool CanEverRender()
 	{
 		return !IsRunningCommandlet() && !IsRunningDedicatedServer();
 	}
@@ -29,7 +29,7 @@ public:
 	 *
 	 * @return The branch name.
 	 */
-	static FString GetBranchName( );
+	static FString GetBranchName();
 
 	/**
 	 * Gets the application's build configuration, i.e. Debug or Shipping.
@@ -43,14 +43,14 @@ public:
 	 *
 	 * @return Build date string.
 	 */
-	static FString GetBuildDate( );
+	static FString GetBuildDate();
 
 	/**
 	 * Gets the name of the currently running game.
 	 *
 	 * @return The game name.
 	 */
-	static const TCHAR* GetGameName( )
+	FORCEINLINE static const TCHAR* GetGameName()
 	{
 		return GGameName;
 	}
@@ -65,7 +65,7 @@ public:
 	 *
 	 * @see GetSessionId
 	 */
-	static FGuid GetInstanceId( )
+	FORCEINLINE static FGuid GetInstanceId()
 	{
 		return InstanceId;
 	}
@@ -77,7 +77,7 @@ public:
 	 *
 	 * @return Instance name string.
 	 */
-	static FString GetInstanceName( )
+	static FString GetInstanceName()
 	{
 		return FString::Printf(TEXT("%s-%i"), FPlatformProcess::ComputerName(), FPlatformProcess::GetCurrentProcessId());
 	}
@@ -89,7 +89,7 @@ public:
 	 *
 	 * @return Application name string.
 	 */
-	static FString GetName( )
+	static FString GetName()
 	{
 		FString ExecutableName = FPlatformProcess::ExecutableName();
 
@@ -120,7 +120,7 @@ public:
 	 *
 	 * @see GetInstanceId
 	 */
-	static FGuid GetSessionId( )
+	FORCEINLINE static FGuid GetSessionId()
 	{
 		return SessionId;
 	}
@@ -130,7 +130,7 @@ public:
 	 *
 	 * @return Session name string.
 	 */
-	static FString GetSessionName( )
+	FORCEINLINE static FString GetSessionName()
 	{
 		return SessionName;
 	}
@@ -145,7 +145,7 @@ public:
 	 * 
 	 * @return Name of session owner.
 	 */
-	static FString GetSessionOwner( )
+	FORCEINLINE static FString GetSessionOwner()
 	{
 		return SessionOwner;
 	}
@@ -155,7 +155,7 @@ public:
 	 *
 	 * @return true if the game name has been set
 	 */
-	static bool HasGameName( )
+	static bool HasGameName()
 	{
 		return (GGameName[0] != 0) && (FCString::Stricmp(GGameName, TEXT("None")) != 0);
 	}
@@ -163,7 +163,7 @@ public:
 	/**
 	 * Initializes the application session.
 	 */
-	static void InitializeSession( );
+	static void InitializeSession();
 
 	/** 
 	 * Checks whether this application is a game.
@@ -174,7 +174,7 @@ public:
 	 *
 	 * @return true if the application is a game, false otherwise.
 	 */
-	FORCEINLINE static bool IsGame( )
+	FORCEINLINE static bool IsGame()
 	{
 #if WITH_EDITOR
 		return !GIsEditor || GIsPlayInEditorWorld || IsRunningGame();
@@ -200,15 +200,15 @@ public:
 	static bool IsInstalled();
 
 	/**
-	* Checks whether the engine components of this application have been installed.
-	*
-	* In binary UE4 releases, the engine can be installed while the game is not. The game IsInstalled()
-	* setting will take precedence over this flag.
-	* 
-	* To override, pass -engineinstalled or -enginenotinstalled on the command line.
-	*
-	* @return true if the engine is installed, false otherwise.
-	*/
+	 * Checks whether the engine components of this application have been installed.
+	 *
+	 * In binary UE4 releases, the engine can be installed while the game is not. The game IsInstalled()
+	 * setting will take precedence over this flag.
+	 * 
+	 * To override, pass -engineinstalled or -enginenotinstalled on the command line.
+	 *
+	 * @return true if the engine is installed, false otherwise.
+	 */
 	static bool IsEngineInstalled();
 
 	/**
@@ -219,7 +219,7 @@ public:
 	 *
 	 * @return true if this application is standalone, false otherwise.
 	 */
-	static bool IsStandalone( )
+	FORCEINLINE static bool IsStandalone()
 	{
 		return Standalone;
 	}
@@ -232,7 +232,7 @@ public:
 	 *
 	 * @return true if the application runs unattended, false otherwise.
 	 */
-	static bool IsUnattended( )
+	static bool IsUnattended()
 	{
 		// FCommandLine::Get() will assert that the command line has been set.
 		// This function may not be used before FCommandLine::Set() is called.
@@ -241,15 +241,123 @@ public:
 	}
 
 	/**
-	 * Checks whether the application should run mult-threaded for performance critical features.
+	 * Checks whether the application should run multi-threaded for performance critical features.
 	 *
 	 * This method is used for performance based threads (like rendering, task graph).
 	 * This will not disable async IO or similar threads needed to disable hitching
 	 *
 	 * @return true if this isn't a server, has more than one core, does not have a -onethread command line options, etc.
 	 */
-	static bool ShouldUseThreadingForPerformance( );
+	static bool ShouldUseThreadingForPerformance();
+
+	/**
+	 * Checks whether application is in benchmark mode.
+	 *
+	 * @return true if application is in benchmark mode, false otherwise.
+	 */
+	FORCEINLINE static bool IsBenchmarking()
+	{
+		return bIsBenchmarking;
+	}
+
+	/**
+	 * Sets application benchmarking mode.
+	 *
+	 * @param bVal - True sets application in benchmark mode, false sets to non-benchmark mode.
+	 */
+	static void SetBenchmarking(bool bVal)
+	{
+		bIsBenchmarking = bVal;
+	}
+
+	/**
+	 * Gets time step in seconds if a fixed delta time is wanted.
+	 *
+	 * @return Time step in seconds for fixed delta time.
+	 */
+	FORCEINLINE static double GetFixedDeltaTime()
+	{
+		return FixedDeltaTime;
+	}
+
+	/**
+	 * Sets time step in seconds if a fixed delta time is wanted.
+	 *
+	 * @param seconds - Time step in seconds for fixed delta time.
+	 */
+	static void SetFixedDeltaTime(double Seconds)
+	{
+		FixedDeltaTime = Seconds;
+	}
+
+	/**
+	 * Gets whether we want to use a fixed time step or not.
+	 *
+	 * @return True if using fixed time step, false otherwise.
+	 */
+	static bool UseFixedTimeStep()
+	{
+		static bool bUseFixedTimeStep = FParse::Param(FCommandLine::Get(), TEXT("UseFixedTimeStep"));
+		return bUseFixedTimeStep;
+	}
+
+	/**
+	 * Gets current time in seconds.
+	 *
+	 * @return Current time in seconds.
+	 */
+	FORCEINLINE static double GetCurrentTime()
+	{
+		return CurrentTime;
+	}
+
+	/**
+	 * Sets current time in seconds.
+	 *
+	 * @param seconds - Time in seconds.
+	 */
+	static void SetCurrentTime(double Seconds)
+	{
+		CurrentTime = Seconds;
+	}
 
+	/**
+	 * Gets previous value of CurrentTime.
+	 *
+	 * @return Previous value of CurrentTime.
+	 */
+	FORCEINLINE static double GetLastTime()
+	{
+		return LastTime;
+	}
+
+	/**
+	 * Updates Last time to CurrentTime.
+	 */
+	static void UpdateLastTime()
+	{
+		LastTime = CurrentTime;
+	}
+
+	/**
+	 * Gets time delta in seconds.
+	 *
+	 * @return Time delta in seconds.
+	 */
+	FORCEINLINE static double GetDeltaTime()
+	{
+		return DeltaTime;
+	}
+
+	/**
+	 * Sets time delta in seconds.
+	 *
+	 * @param seconds - Time in seconds.
+	 */
+	static void SetDeltaTime(double Seconds)
+	{
+		DeltaTime = Seconds;
+	}
 
 private:
 
@@ -267,6 +375,24 @@ private:
 
 	// Holds a flag indicating whether this is a standalone session.
 	static bool Standalone;
+
+	// Holds a flag Whether we are in benchmark mode or not.
+	static bool bIsBenchmarking;
+
+	// Holds time step if a fixed delta time is wanted.
+	static double FixedDeltaTime;
+
+	// Holds whether we want to use a fixed time step or not.
+	static bool bUseFixedTimeStep;
+
+	// Holds current time.
+	static double CurrentTime;
+
+	// Holds previous value of CurrentTime.
+	static double LastTime;
+
+	// Holds current delta time in seconds.
+	static double DeltaTime;
 };
 
 
diff --git a/Engine/Source/Runtime/CoreUObject/Private/Serialization/AsyncLoading.cpp b/Engine/Source/Runtime/CoreUObject/Private/Serialization/AsyncLoading.cpp
--- a/Engine/Source/Runtime/CoreUObject/Private/Serialization/AsyncLoading.cpp
+++ b/Engine/Source/Runtime/CoreUObject/Private/Serialization/AsyncLoading.cpp
@@ -643,7 +643,7 @@ EAsyncPackageState::Type FAsyncPackage::CreateExports()
 		// Data isn't ready yet. Give up remainder of time slice if we're not using a time limit.
 		else if( GiveUpTimeSlice() )
 		{
-			INC_FLOAT_STAT_BY(STAT_AsyncIO_AsyncPackagePrecacheWaitTime,(float)GDeltaTime);
+			INC_FLOAT_STAT_BY(STAT_AsyncIO_AsyncPackagePrecacheWaitTime, (float)FApp::GetDeltaTime());
 			return EAsyncPackageState::TimeOut;
 		}
 	}
diff --git a/Engine/Source/Runtime/Engine/Classes/Engine/Engine.h b/Engine/Source/Runtime/Engine/Classes/Engine/Engine.h
--- a/Engine/Source/Runtime/Engine/Classes/Engine/Engine.h
+++ b/Engine/Source/Runtime/Engine/Classes/Engine/Engine.h
@@ -1598,7 +1598,7 @@ public:
 	virtual void Tick( float DeltaSeconds, bool bIdleMode ) PURE_VIRTUAL(UEngine::Tick,);
 
 	/**
-	 * Update GCurrentTime/ GDeltaTime while taking into account max tick rate.
+	 * Update FApp::CurrentTime/ FApp::DeltaTime while taking into account max tick rate.
 	 */
 	void UpdateTimeAndHandleMaxTickRate();
 
diff --git a/Engine/Source/Runtime/Engine/Classes/Engine/GameEngine.h b/Engine/Source/Runtime/Engine/Classes/Engine/GameEngine.h
--- a/Engine/Source/Runtime/Engine/Classes/Engine/GameEngine.h
+++ b/Engine/Source/Runtime/Engine/Classes/Engine/GameEngine.h
@@ -18,7 +18,7 @@ class ENGINE_API UGameEngine : public UEngine
 	UPROPERTY(transient)
 	uint32 bCheckForMovieCapture:1;
 
-	/** Maximium delta time the engine uses to populate GDeltaTime. If 0, unbound.									*/
+	/** Maximium delta time the engine uses to populate FApp::DeltaTime. If 0, unbound.									*/
 	UPROPERTY(config)
 	float MaxDeltaTime;
 
diff --git a/Engine/Source/Runtime/Engine/Private/ActiveSound.cpp b/Engine/Source/Runtime/Engine/Private/ActiveSound.cpp
--- a/Engine/Source/Runtime/Engine/Private/ActiveSound.cpp
+++ b/Engine/Source/Runtime/Engine/Private/ActiveSound.cpp
@@ -321,7 +321,7 @@ void FActiveSound::HandleInteriorVolumes( const FListener& Listener, FSoundParse
 	{
 		SourceInteriorVolume = CurrentInteriorVolume;
 		SourceInteriorLPF = CurrentInteriorLPF;
-		LastUpdateTime = GCurrentTime;
+		LastUpdateTime = FApp::GetCurrentTime();
 	}
 
 	if( Listener.Volume == ReverbVolume || !bAllowSpatialization )
@@ -578,4 +578,4 @@ void FActiveSound::CollectAttenuationShapesForVisualization(TMap<EAttenuationSha
 			}
 		}
 	}
-}
\ No newline at end of file
+}
diff --git a/Engine/Source/Runtime/Engine/Private/AudioDevice.cpp b/Engine/Source/Runtime/Engine/Private/AudioDevice.cpp
--- a/Engine/Source/Runtime/Engine/Private/AudioDevice.cpp
+++ b/Engine/Source/Runtime/Engine/Private/AudioDevice.cpp
@@ -1051,7 +1051,7 @@ bool FAudioDevice::ApplySoundMix( USoundMix* NewMix, FSoundMixState* SoundMixSta
 	{
 		UE_LOG(LogAudio, Log, TEXT( "FAudioDevice::ApplySoundMix(): %s" ), *NewMix->GetName() );
 
-		SoundMixState->StartTime = GCurrentTime;
+		SoundMixState->StartTime = FApp::GetCurrentTime();
 		SoundMixState->FadeInStartTime = SoundMixState->StartTime + NewMix->InitialDelay;
 		SoundMixState->FadeInEndTime = SoundMixState->FadeInStartTime + NewMix->FadeInTime;
 		SoundMixState->FadeOutStartTime = -1.0;
@@ -1079,7 +1079,7 @@ void FAudioDevice::UpdateSoundMix(class USoundMix* SoundMix, FSoundMixState* Sou
 		if (SoundMixState->CurrentState == ESoundMixState::FadingOut)
 		{
 			// In process of fading out, need to adjust timing to fade in again.
-			SoundMixState->FadeInStartTime = GCurrentTime - (SoundMixState->InterpValue * SoundMix->FadeInTime);
+			SoundMixState->FadeInStartTime = FApp::GetCurrentTime() - (SoundMixState->InterpValue * SoundMix->FadeInTime);
 			SoundMixState->FadeInEndTime = SoundMixState->FadeInStartTime + SoundMix->FadeInTime;
 			SoundMixState->FadeOutStartTime = -1.0;
 			SoundMixState->EndTime = -1.0;
@@ -1095,7 +1095,7 @@ void FAudioDevice::UpdateSoundMix(class USoundMix* SoundMix, FSoundMixState* Sou
 		else if (SoundMixState->CurrentState == ESoundMixState::Active)
 		{
 			// Duration may be -1, this guarantees the fade will at least start at the current time
-			double NewFadeOutStartTime = FMath::Max(GCurrentTime, GCurrentTime + SoundMix->Duration);
+			double NewFadeOutStartTime = FMath::Max(FApp::GetCurrentTime(), FApp::GetCurrentTime() + SoundMix->Duration);
 			if (NewFadeOutStartTime > SoundMixState->FadeOutStartTime)
 			{
 				SoundMixState->FadeOutStartTime = NewFadeOutStartTime;
@@ -1167,7 +1167,7 @@ bool FAudioDevice::TryClearingSoundMix(USoundMix* SoundMix, FSoundMixState* Soun
 				else if (SoundMixState->CurrentState == ESoundMixState::FadingIn)
 				{
 					// Currently fading up, force fade in to complete and start fade out from current fade level
-					SoundMixState->FadeOutStartTime = GCurrentTime - ((1.0f - SoundMixState->InterpValue) * SoundMix->FadeOutTime);
+					SoundMixState->FadeOutStartTime = FApp::GetCurrentTime() - ((1.0f - SoundMixState->InterpValue) * SoundMix->FadeOutTime);
 					SoundMixState->EndTime = SoundMixState->FadeOutStartTime + SoundMix->FadeOutTime;
 					SoundMixState->StartTime = SoundMixState->FadeInStartTime = SoundMixState->FadeInEndTime = SoundMixState->FadeOutStartTime - 1.0;
 
@@ -1176,7 +1176,7 @@ bool FAudioDevice::TryClearingSoundMix(USoundMix* SoundMix, FSoundMixState* Soun
 				else if (SoundMixState->CurrentState == ESoundMixState::Active)
 				{
 					// SoundMix active, start fade out early
-					SoundMixState->FadeOutStartTime = GCurrentTime;
+					SoundMixState->FadeOutStartTime = FApp::GetCurrentTime();
 					SoundMixState->EndTime = SoundMixState->FadeOutStartTime + SoundMix->FadeOutTime;
 
 					TryClearingEQSoundMix(SoundMix);
@@ -1308,29 +1308,29 @@ void FAudioDevice::UpdateSoundClassProperties()
 		FSoundMixState* SoundMixState = &(It.Value());
 
 		// Initial delay before mix is applied
-		if( GCurrentTime >= SoundMixState->StartTime && GCurrentTime < SoundMixState->FadeInStartTime )
+		if( FApp::GetCurrentTime() >= SoundMixState->StartTime && FApp::GetCurrentTime() < SoundMixState->FadeInStartTime )
 		{
 			SoundMixState->InterpValue = 0.0f;
 			SoundMixState->CurrentState = ESoundMixState::Inactive;
 		}
-		else if( GCurrentTime >= SoundMixState->FadeInStartTime && GCurrentTime < SoundMixState->FadeInEndTime )
+		else if( FApp::GetCurrentTime() >= SoundMixState->FadeInStartTime && FApp::GetCurrentTime() < SoundMixState->FadeInEndTime )
 		{
 			// Work out the fade in portion
-			SoundMixState->InterpValue = ( float )( ( GCurrentTime - SoundMixState->FadeInStartTime ) / ( SoundMixState->FadeInEndTime - SoundMixState->FadeInStartTime ) );
+			SoundMixState->InterpValue = ( float )( ( FApp::GetCurrentTime() - SoundMixState->FadeInStartTime ) / ( SoundMixState->FadeInEndTime - SoundMixState->FadeInStartTime ) );
 			SoundMixState->CurrentState = ESoundMixState::FadingIn;
 		}
-		else if( GCurrentTime >= SoundMixState->FadeInEndTime
-			&& ( SoundMixState->IsBaseSoundMix || SoundMixState->PassiveRefCount > 0 || SoundMixState->FadeOutStartTime < 0.0 || GCurrentTime < SoundMixState->FadeOutStartTime ) )
+		else if( FApp::GetCurrentTime() >= SoundMixState->FadeInEndTime
+			&& ( SoundMixState->IsBaseSoundMix || SoundMixState->PassiveRefCount > 0 || SoundMixState->FadeOutStartTime < 0.0 || FApp::GetCurrentTime() < SoundMixState->FadeOutStartTime ) )
 		{
 			// .. ensure the full mix is applied between the end of the fade in time and the start of the fade out time
 			// or if SoundMix is the base or active via a passive push - ignores duration.
 			SoundMixState->InterpValue = 1.0f;
 			SoundMixState->CurrentState = ESoundMixState::Active;
 		}
-		else if( GCurrentTime >= SoundMixState->FadeOutStartTime && GCurrentTime < SoundMixState->EndTime )
+		else if( FApp::GetCurrentTime() >= SoundMixState->FadeOutStartTime && FApp::GetCurrentTime() < SoundMixState->EndTime )
 		{
 			// Work out the fade out portion
-			SoundMixState->InterpValue = 1.0f - ( float )( ( GCurrentTime - SoundMixState->FadeOutStartTime ) / ( SoundMixState->EndTime - SoundMixState->FadeOutStartTime ) );
+			SoundMixState->InterpValue = 1.0f - ( float )( ( FApp::GetCurrentTime() - SoundMixState->FadeOutStartTime ) / ( SoundMixState->EndTime - SoundMixState->FadeOutStartTime ) );
 			if( SoundMixState->CurrentState != ESoundMixState::FadingOut )
 			{
 				// Start fading EQ at same time
@@ -1340,7 +1340,7 @@ void FAudioDevice::UpdateSoundClassProperties()
 		}
 		else
 		{
-			check( SoundMixState->EndTime >= 0.0 && GCurrentTime >= SoundMixState->EndTime );
+			check( SoundMixState->EndTime >= 0.0 && FApp::GetCurrentTime() >= SoundMixState->EndTime );
 			// Clear the effect of this SoundMix - may need to revisit for passive
 			SoundMixState->InterpValue = 0.0f;
 			SoundMixState->CurrentState = ESoundMixState::AwaitingRemoval;
@@ -1358,17 +1358,17 @@ void FAudioDevice::UpdateSoundClassProperties()
 
 float FListener::Interpolate( const double EndTime )
 {
-	if( GCurrentTime < InteriorStartTime )
+	if( FApp::GetCurrentTime() < InteriorStartTime )
 	{
 		return( 0.0f );
 	}
 
-	if( GCurrentTime >= EndTime )
+	if( FApp::GetCurrentTime() >= EndTime )
 	{
 		return( 1.0f );
 	}
 
-	float InterpValue = ( float )( ( GCurrentTime - InteriorStartTime ) / ( EndTime - InteriorStartTime ) );
+	float InterpValue = ( float )( ( FApp::GetCurrentTime() - InteriorStartTime ) / ( EndTime - InteriorStartTime ) );
 	return( InterpValue );
 }
 
@@ -1395,7 +1395,7 @@ void FListener::ApplyInteriorSettings( class AReverbVolume* InVolume, const FInt
 	if( InVolume != Volume )
 	{
 		// Use previous/ current interpolation time if we're transitioning to the default worldsettings zone.
-		InteriorStartTime = GCurrentTime;
+		InteriorStartTime = FApp::GetCurrentTime();
 		InteriorEndTime = InteriorStartTime + (Settings.bIsWorldSettings ? InteriorSettings.InteriorTime : Settings.InteriorTime);
 		ExteriorEndTime = InteriorStartTime + (Settings.bIsWorldSettings ? InteriorSettings.ExteriorTime : Settings.ExteriorTime);
 		InteriorLPFEndTime = InteriorStartTime + (Settings.bIsWorldSettings ? InteriorSettings.InteriorLPFTime : Settings.InteriorLPFTime);
@@ -1518,9 +1518,9 @@ void FAudioDevice::PopSoundMixModifier(USoundMix* SoundMix, bool bIsPassive)
 				if (SoundMixState->PassiveRefCount == 0)
 				{
 					// Check whether Fade out time was previously set and reset it to current time
-					if (SoundMixState->FadeOutStartTime >= 0.0 && GCurrentTime > SoundMixState->FadeOutStartTime)
+					if (SoundMixState->FadeOutStartTime >= 0.0 && FApp::GetCurrentTime() > SoundMixState->FadeOutStartTime)
 					{
-						SoundMixState->FadeOutStartTime = GCurrentTime;
+						SoundMixState->FadeOutStartTime = FApp::GetCurrentTime();
 						SoundMixState->EndTime = SoundMixState->FadeOutStartTime + SoundMix->FadeOutTime;
 					}
 				}
@@ -1703,7 +1703,7 @@ int32 FAudioDevice::GetSortedActiveWaveInstances(TArray<FWaveInstance*>& WaveIns
 			else
 			{
 				// If not in game, do not advance sounds unless they are UI sounds.
-				float UsedDeltaTime = GDeltaTime;
+				float UsedDeltaTime = FApp::GetDeltaTime();
 				if (GetType == ESortedActiveWaveGetType::QueryOnly || (GetType == ESortedActiveWaveGetType::PausedUpdate && !ActiveSound->bIsUISound))
 				{
 					UsedDeltaTime = 0.0f;
@@ -2548,4 +2548,4 @@ void FAudioDevice::OnEndPIE(const bool bIsSimulating)
 		SoundNode->OnEndPIE(bIsSimulating);
 	}
 }
-#endif
\ No newline at end of file
+#endif
diff --git a/Engine/Source/Runtime/Engine/Private/AudioEffect.cpp b/Engine/Source/Runtime/Engine/Private/AudioEffect.cpp
--- a/Engine/Source/Runtime/Engine/Private/AudioEffect.cpp
+++ b/Engine/Source/Runtime/Engine/Private/AudioEffect.cpp
@@ -90,7 +90,7 @@ FAudioReverbEffect& FAudioReverbEffect::operator=(class UReverbEffect* InReverbE
 void FAudioReverbEffect::Interpolate( float InterpValue, const FAudioReverbEffect& Start, const FAudioReverbEffect& End )
 {
 	float InvInterpValue = 1.0f - InterpValue;
-	Time = GCurrentTime;
+	Time = FApp::GetCurrentTime();
 	Volume = ( Start.Volume * InvInterpValue ) + ( End.Volume * InterpValue );
 	Density = ( Start.Density * InvInterpValue ) + ( End.Density * InterpValue );				
 	Diffusion = ( Start.Diffusion * InvInterpValue ) + ( End.Diffusion * InterpValue );				
@@ -126,7 +126,7 @@ void FAudioEQEffect::ClampValues( void )
 void FAudioEQEffect::Interpolate( float InterpValue, const FAudioEQEffect& Start, const FAudioEQEffect& End )
 {
 	float InvInterpValue = 1.0f - InterpValue;
-	RootTime = GCurrentTime;	
+	RootTime = FApp::GetCurrentTime();
 	HFFrequency = ( Start.HFFrequency * InvInterpValue ) + ( End.HFFrequency * InterpValue );
 	HFGain = ( Start.HFGain * InvInterpValue ) + ( End.HFGain * InterpValue );
 	MFCutoffFrequency = ( Start.MFCutoffFrequency * InvInterpValue ) + ( End.MFCutoffFrequency * InterpValue );
@@ -175,7 +175,7 @@ void FAudioEffectsManager::Interpolate( FAudioReverbEffect& Current, const FAudi
 	float InterpValue = 1.0f;
 	if( End.Time - Start.Time > 0.0 )
 	{
-		InterpValue = ( float )( ( GCurrentTime - Start.Time ) / ( End.Time - Start.Time ) );
+		InterpValue = ( float )( ( FApp::GetCurrentTime() - Start.Time ) / ( End.Time - Start.Time ) );
 	}
 
 	if( InterpValue >= 1.0f )
@@ -201,7 +201,7 @@ void FAudioEffectsManager::Interpolate( FAudioEQEffect& Current, const FAudioEQE
 	float InterpValue = 1.0f;
 	if( End.RootTime - Start.RootTime > 0.0 )
 	{
-		InterpValue = ( float )( ( GCurrentTime - Start.RootTime ) / ( End.RootTime - Start.RootTime ) );
+		InterpValue = ( float )( ( FApp::GetCurrentTime() - Start.RootTime ) / ( End.RootTime - Start.RootTime ) );
 	}
 
 	if( InterpValue >= 1.0f )
@@ -278,10 +278,10 @@ void FAudioEffectsManager::SetReverbSettings( const FReverbSettings& ReverbSetti
 		}
 
 		SourceReverbEffect = CurrentReverbEffect;
-		SourceReverbEffect.Time = GCurrentTime;
+		SourceReverbEffect.Time = FApp::GetCurrentTime();
 
 		DestinationReverbEffect = ReverbSettings.ReverbEffect;
-		DestinationReverbEffect.Time = GCurrentTime + ReverbSettings.FadeTime;
+		DestinationReverbEffect.Time = FApp::GetCurrentTime() + ReverbSettings.FadeTime;
 		DestinationReverbEffect.Volume = ReverbSettings.Volume;
 		if( !ReverbSettings.ReverbEffect )
 		{
@@ -306,7 +306,7 @@ void FAudioEffectsManager::SetMixSettings( USoundMix* NewMix, bool bIgnorePriori
 			UE_LOG(LogAudio, Log, TEXT( "FAudioEffectsManager::SetMixSettings(): %s" ), *NewMix->GetName() );
 
 			SourceEQEffect = CurrentEQEffect;
-			SourceEQEffect.RootTime = GCurrentTime;
+			SourceEQEffect.RootTime = FApp::GetCurrentTime();
 
 			if( NewMix->bApplyEQ )
 			{
@@ -318,7 +318,7 @@ void FAudioEffectsManager::SetMixSettings( USoundMix* NewMix, bool bIgnorePriori
 				DestinationEQEffect = FAudioEQEffect();
 			}
 
-			DestinationEQEffect.RootTime = GCurrentTime + NewMix->FadeInTime;
+			DestinationEQEffect.RootTime = FApp::GetCurrentTime() + NewMix->FadeInTime;
 			DestinationEQEffect.ClampValues();
 
 			CurrentEQMix = NewMix;
@@ -336,11 +336,11 @@ void FAudioEffectsManager::ClearMixSettings()
 		UE_LOG(LogAudio, Log, TEXT( "FAudioEffectsManager::ClearMixSettings(): %s" ), *CurrentEQMix->GetName() );
 
 		SourceEQEffect = CurrentEQEffect;
-		SourceEQEffect.RootTime = GCurrentTime;
+		SourceEQEffect.RootTime = FApp::GetCurrentTime();
 
 		// interpolate back to default
 		DestinationEQEffect = FAudioEQEffect();
-		DestinationEQEffect.RootTime = GCurrentTime + CurrentEQMix->FadeOutTime;
+		DestinationEQEffect.RootTime = FApp::GetCurrentTime() + CurrentEQMix->FadeOutTime;
 
 		CurrentEQMix = NULL;
 	}
diff --git a/Engine/Source/Runtime/Engine/Private/Commandlets/SmokeTestCommandlet.cpp b/Engine/Source/Runtime/Engine/Private/Commandlets/SmokeTestCommandlet.cpp
--- a/Engine/Source/Runtime/Engine/Private/Commandlets/SmokeTestCommandlet.cpp
+++ b/Engine/Source/Runtime/Engine/Private/Commandlets/SmokeTestCommandlet.cpp
@@ -36,11 +36,11 @@ bool FRunServer::RunTest(const FString& Parameters)
 	// own commandlet
 	if( GEngine && IsRunningDedicatedServer() )
 	{
-		// Update GCurrentTime/ GDeltaTime while taking into account max tick rate.
+		// Update FApp::CurrentTime / FApp::DeltaTime while taking into account max tick rate.
 		GEngine->UpdateTimeAndHandleMaxTickRate();
 
 		// Tick the engine.
-		GEngine->Tick( GDeltaTime, false );
+		GEngine->Tick( FApp::GetDeltaTime(), false );
 	}
 
 	return true;
diff --git a/Engine/Source/Runtime/Engine/Private/ContentStreaming.cpp b/Engine/Source/Runtime/Engine/Private/ContentStreaming.cpp
--- a/Engine/Source/Runtime/Engine/Private/ContentStreaming.cpp
+++ b/Engine/Source/Runtime/Engine/Private/ContentStreaming.cpp
@@ -258,7 +258,7 @@ struct FStreamingTexture
 		MaxAllowedMips = MipCount;	//FMath::Max(Texture->ResidentMips, Texture->RequestedMips);
 		STAT( MaxAllowedOptimalMips = MaxAllowedMips );
 		STAT( MostResidentMips = FMath::Max(MostResidentMips, Texture->ResidentMips) );
-		LastRenderTime = (GCurrentTime > Texture->Resource->LastRenderTime) ? float( GCurrentTime - Texture->Resource->LastRenderTime ) : 0.0f;
+		LastRenderTime = (FApp::GetCurrentTime() > Texture->Resource->LastRenderTime) ? float( FApp::GetCurrentTime() - Texture->Resource->LastRenderTime ) : 0.0f;
 		MinDistance = MAX_STREAMINGDISTANCE;
 		bForceFullyLoad = Texture->ShouldMipLevelsBeForcedResident() || (ForceLoadRefCount > 0);
 		TextureLODBias = Texture->GetCachedLODBias();
@@ -420,7 +420,7 @@ struct FStreamingTexture
 	/** If non-zero, the most recent time an instance location was removed for this texture. */
 	double			InstanceRemovedTimestamp;
 
-	/** Set to GCurrentTime every time LastRenderTimeRefCount is modified. */
+	/** Set to FApp::GetCurrentTime() every time LastRenderTimeRefCount is modified. */
 	double			LastRenderTimeRefCountTimestamp;
 	/** Current number of instances that need LRT heuristics for this texture. */
 	int32			LastRenderTimeRefCount;
@@ -793,7 +793,7 @@ struct FTextureStreamingStats
 	,	SizeY( InTexture2D->GetSizeY() )
 	,	NumMips( InTexture2D->GetNumMips() )
 	,	LODBias( InTexture2D->GetCachedLODBias() )
-	,	LastRenderTime( FMath::Clamp( InTexture2D->Resource ? (GCurrentTime - InTexture2D->Resource->LastRenderTime) : 1000000.0, 0.0, 1000000.0) )
+	,	LastRenderTime( FMath::Clamp( InTexture2D->Resource ? (FApp::GetCurrentTime() - InTexture2D->Resource->LastRenderTime) : 1000000.0, 0.0, 1000000.0) )
 	,	StreamType( InType )
 	,	ResidentMips( InResidentMips )
 	,	WantedMips( InWantedMips )
@@ -817,7 +817,7 @@ struct FTextureStreamingStats
 	int32		NumMips;
 	/** Mirror of UTexture2D::GetCachedLODBias() */
 	int32		LODBias;
-	/** How many seconds since it was last rendered: GCurrentTime - UTexture2D::Resource->LastRenderTime */
+	/** How many seconds since it was last rendered: FApp::GetCurrentTime() - UTexture2D::Resource->LastRenderTime */
 	float		LastRenderTime;
 	/** What streaming heuristics were primarily used for this texture. */
 	ETextureStreamingType	StreamType;
@@ -2809,7 +2809,7 @@ void FStreamingManagerTexture::RemoveLevel( ULevel* Level )
 			if ( Texture2D && IsManagedStreamingTexture( Texture2D ) )
 			{
 				FStreamingTexture& StreamingTexture = GetStreamingTexture( Texture2D );
-				StreamingTexture.InstanceRemovedTimestamp = GCurrentTime;
+				StreamingTexture.InstanceRemovedTimestamp = FApp::GetCurrentTime();
 			}
 		}
 
@@ -3055,7 +3055,7 @@ void FStreamingManagerTexture::SetInstanceRemovedTimestamp( FSpawnedPrimitiveDat
 		if ( Texture2D && IsManagedStreamingTexture(Texture2D) )
 		{
 			FStreamingTexture& StreamingTexture = GetStreamingTexture( Texture2D );
-			StreamingTexture.InstanceRemovedTimestamp = GCurrentTime;
+			StreamingTexture.InstanceRemovedTimestamp = FApp::GetCurrentTime();
 		}
 	}
 }
@@ -3120,7 +3120,7 @@ void FStreamingManagerTexture::NotifyTimedPrimitiveAttached( const UPrimitiveCom
 				// Note: Doesn't have to be cycle-perfect for thread safety.
 				FStreamingTexture& StreamingTexture = GetStreamingTexture( Texture2D );
 				StreamingTexture.LastRenderTimeRefCount++;
-				StreamingTexture.LastRenderTimeRefCountTimestamp = GCurrentTime;
+				StreamingTexture.LastRenderTimeRefCountTimestamp = FApp::GetCurrentTime();
 			}
 		}
 	}
@@ -3153,7 +3153,7 @@ void FStreamingManagerTexture::NotifyTimedPrimitiveDetached( const UPrimitiveCom
 				if ( StreamingTexture.LastRenderTimeRefCount > 0 )
 				{
 					StreamingTexture.LastRenderTimeRefCount--;
-					StreamingTexture.LastRenderTimeRefCountTimestamp = GCurrentTime;
+					StreamingTexture.LastRenderTimeRefCountTimestamp = FApp::GetCurrentTime();
 				}
 			}
 		}
@@ -3731,14 +3731,14 @@ void FStreamingManagerTexture::CheckUserSettings()
 		// Only adjust the pool size once every 10 seconds, but immediately in some other cases.
 		if ( PoolSizeSetting != PreviousPoolSizeSetting ||
 			 TexturePoolSize > GTexturePoolSize ||
-			 (GCurrentTime - PreviousPoolSizeTimestamp) > 10.0 )
+			 (FApp::GetCurrentTime() - PreviousPoolSizeTimestamp) > 10.0 )
 		{
 			if ( TexturePoolSize != GTexturePoolSize )
 			{
 				UE_LOG(LogContentStreaming,Log,TEXT("Texture pool size now %d MB"), int(TexturePoolSize/1024/1024));
 			}
 			PreviousPoolSizeSetting = PoolSizeSetting;
-			PreviousPoolSizeTimestamp = GCurrentTime;
+			PreviousPoolSizeTimestamp = FApp::GetCurrentTime();
 			GTexturePoolSize = TexturePoolSize;
 		}
 	}
@@ -4127,7 +4127,7 @@ void FStreamingManagerTexture::CalcWantedMips( FStreamingTexture& StreamingTextu
 		}
 
 		bool bShouldAlsoUseLastRenderTime = false;
-		if ( StreamingTexture.LastRenderTimeRefCount > 0 || (GCurrentTime - StreamingTexture.LastRenderTimeRefCountTimestamp) < 91.0 )
+		if ( StreamingTexture.LastRenderTimeRefCount > 0 || (FApp::GetCurrentTime() - StreamingTexture.LastRenderTimeRefCountTimestamp) < 91.0 )
 		{
 			bShouldAlsoUseLastRenderTime = true;
 
@@ -4188,7 +4188,7 @@ FFloatMipLevel FStreamingManagerTexture::GetWantedMipsForOrphanedTexture( FStrea
 	FFloatMipLevel WantedMips;
 
 	// Did we recently remove instance locations for this texture?
-	const float TimeSinceInstanceWasRemoved = float(GCurrentTime - StreamingTexture.InstanceRemovedTimestamp);
+	const float TimeSinceInstanceWasRemoved = float(FApp::GetCurrentTime() - StreamingTexture.InstanceRemovedTimestamp);
 
 	// Was it less than 91 seconds ago?
 	if ( TimeSinceInstanceWasRemoved < 91.0f )
@@ -5522,7 +5522,7 @@ FFloatMipLevel FStreamingHandlerTextureLastRender::GetWantedMips( FStreamingMana
 {
 	FFloatMipLevel Ret;
 
-	float SecondsSinceLastRender = StreamingTexture.LastRenderTime;	//(GCurrentTime - Texture->Resource->LastRenderTime);;
+	float SecondsSinceLastRender = StreamingTexture.LastRenderTime;	//(FApp::GetCurrentTime() - Texture->Resource->LastRenderTime);;
 	StreamingTexture.bUsesLastRenderHeuristics = true;
 
 	if( SecondsSinceLastRender < 45.0f && GStreamWithTimeFactor )
diff --git a/Engine/Source/Runtime/Engine/Private/GameViewportClient.cpp b/Engine/Source/Runtime/Engine/Private/GameViewportClient.cpp
--- a/Engine/Source/Runtime/Engine/Private/GameViewportClient.cpp
+++ b/Engine/Source/Runtime/Engine/Private/GameViewportClient.cpp
@@ -399,8 +399,8 @@ void UGameViewportClient::SetDropDetail(float DeltaSeconds)
 		const float FrameRate	= FrameTime > 0 ? 1 / FrameTime : 0;
 
 		// Drop detail if framerate is below threshold.
-		GetWorld()->bDropDetail		= FrameRate < FMath::Clamp(GEngine->MinDesiredFrameRate, 1.f, 100.f) && !GIsBenchmarking && !GUseFixedTimeStep;
-		GetWorld()->bAggressiveLOD	= FrameRate < FMath::Clamp(GEngine->MinDesiredFrameRate - 5.f, 1.f, 100.f) && !GIsBenchmarking && !GUseFixedTimeStep;
+		GetWorld()->bDropDetail		= FrameRate < FMath::Clamp(GEngine->MinDesiredFrameRate, 1.f, 100.f) && !FApp::IsBenchmarking() && !FApp::UseFixedTimeStep();
+		GetWorld()->bAggressiveLOD	= FrameRate < FMath::Clamp(GEngine->MinDesiredFrameRate - 5.f, 1.f, 100.f) && !FApp::IsBenchmarking() && !FApp::UseFixedTimeStep();
 
 		// this is slick way to be able to do something based on the frametime and whether we are bound by one thing or another
 #if 0 
@@ -2671,4 +2671,4 @@ bool UGameViewportClient::RequestBugScreenShot(const TCHAR* Cmd, bool bDisplayHU
 	}
 
 	return true;
-}
\ No newline at end of file
+}
diff --git a/Engine/Source/Runtime/Engine/Private/Interpolation.cpp b/Engine/Source/Runtime/Engine/Private/Interpolation.cpp
--- a/Engine/Source/Runtime/Engine/Private/Interpolation.cpp
+++ b/Engine/Source/Runtime/Engine/Private/Interpolation.cpp
@@ -3906,13 +3906,13 @@ void UInterpTrackMove::UpdateTrack(float NewPosition, UInterpTrackInst* TrInst,
 	}
 
 	FVector NewLocation = Actor->GetRootComponent()->GetComponentLocation();
-	// if GDeltaTime == 0.f, I'd think that's paused, then we won't need to update Velocity
-	if ( GDeltaTime > 0.f ) 
+	// if FApp::GetDeltaTime() == 0.f, I'd think that's paused, then we won't need to update Velocity
+	if ( FApp::GetDeltaTime() > 0.f ) 
 	{
 		// we're not using PreviPosition to NewPosition because MatineeActor itself can have different playrate
 		// so we can't guarantee that's the time it took to get there. 
 		// this should approximately safe in replication as well
-		FVector ComponentVelocity = (NewLocation-PrevLocation)/GDeltaTime;
+		FVector ComponentVelocity = (NewLocation-PrevLocation)/FApp::GetDeltaTime();
 		Actor->GetRootComponent()->ComponentVelocity = ComponentVelocity;
 	}
 }
diff --git a/Engine/Source/Runtime/Engine/Private/LevelActor.cpp b/Engine/Source/Runtime/Engine/Private/LevelActor.cpp
--- a/Engine/Source/Runtime/Engine/Private/LevelActor.cpp
+++ b/Engine/Source/Runtime/Engine/Private/LevelActor.cpp
@@ -955,7 +955,7 @@ void UWorld::SetMapNeedsLightingFullyRebuilt(int32 InNumLightingUnbuiltObjects)
 		// Update last time unbuilt lighting was encountered.
 		if (NumLightingUnbuiltObjects > 0)
 		{
-			LastTimeUnbuiltLightingWasEncountered = GCurrentTime;
+			LastTimeUnbuiltLightingWasEncountered = FApp::GetCurrentTime();
 		}
 	}
 }
diff --git a/Engine/Source/Runtime/Engine/Private/LevelTick.cpp b/Engine/Source/Runtime/Engine/Private/LevelTick.cpp
--- a/Engine/Source/Runtime/Engine/Private/LevelTick.cpp
+++ b/Engine/Source/Runtime/Engine/Private/LevelTick.cpp
@@ -215,7 +215,7 @@ void FDetailedTickStats::DumpStats()
 	bool bShouldDump = false;
 	
 	// Dump request due to interval.
-	if( GCurrentTime > (LastTimeOfLogDump + TimeBetweenLogDumps) )
+	if( FApp::GetCurrentTime() > (LastTimeOfLogDump + TimeBetweenLogDumps) )
 	{
 		bShouldDump = true;
 	}
@@ -237,9 +237,9 @@ void FDetailedTickStats::DumpStats()
 
 	// Only dump every TimeBetweenLogDumps seconds.
 	if( bShouldDump 
-	&& ((GCurrentTime - LastTimeOfLogDump) > MinTimeBetweenLogDumps) )
+	&& ((FApp::GetCurrentTime() - LastTimeOfLogDump) > MinTimeBetweenLogDumps) )
 	{
-		LastTimeOfLogDump = GCurrentTime;
+		LastTimeOfLogDump = FApp::GetCurrentTime();
 
 		// Array of stats, used for sorting.
 		TArray<FTickStats> SortedTickStats;
diff --git a/Engine/Source/Runtime/Engine/Private/Particles/ParticleSystemRender.cpp b/Engine/Source/Runtime/Engine/Private/Particles/ParticleSystemRender.cpp
--- a/Engine/Source/Runtime/Engine/Private/Particles/ParticleSystemRender.cpp
+++ b/Engine/Source/Runtime/Engine/Private/Particles/ParticleSystemRender.cpp
@@ -6794,7 +6794,7 @@ FParticleSystemSceneProxy::FParticleSystemSceneProxy(const UParticleSystemCompon
 	, LastFramePreRendered(-1)
 {
 #if STATS
-	LastStatCaptureTime = GCurrentTime;
+	LastStatCaptureTime = FApp::GetCurrentTime();
 	bCountedThisFrame = false;
 #endif
 	LODMethod = Component->LODMethod;
@@ -6928,9 +6928,9 @@ void FParticleSystemSceneProxy::DrawDynamicElements(FPrimitiveDrawInterface* PDI
 			// This is needed to prevent capturing particle rendering time during cinematics where the engine switches back to high detail mode
 			&& GetCachedScalabilityCVars().DetailMode != DM_High
 			&& NumDraws > 0
-			&& GCurrentTime - LastStatCaptureTime > GTimeBetweenParticleRenderStatCaptures)
+			&& FApp::GetCurrentTime() - LastStatCaptureTime > GTimeBetweenParticleRenderStatCaptures)
 		{
-			LastStatCaptureTime = GCurrentTime;
+			LastStatCaptureTime = FApp::GetCurrentTime();
 
 			const double EndTime = FPlatformTime::Seconds();
 
diff --git a/Engine/Source/Runtime/Engine/Private/Slate/SceneViewport.cpp b/Engine/Source/Runtime/Engine/Private/Slate/SceneViewport.cpp
--- a/Engine/Source/Runtime/Engine/Private/Slate/SceneViewport.cpp
+++ b/Engine/Source/Runtime/Engine/Private/Slate/SceneViewport.cpp
@@ -530,7 +530,7 @@ FReply FSceneViewport::OnControllerAnalogValueChanged( const FGeometry& MyGeomet
 		// Switch to the viewport clients world before processing input
 		FScopedConditionalWorldSwitcher WorldSwitcher( ViewportClient );
 
-		if( !ViewportClient->InputAxis( this, ControllerEvent.GetUserIndex(), ControllerEvent.GetEffectingButton(), ControllerEvent.GetEffectingButton() == EKeys::Gamepad_RightY ? -ControllerEvent.GetAnalogValue() : ControllerEvent.GetAnalogValue(), GDeltaTime, 1, true  ) )
+		if (!ViewportClient->InputAxis(this, ControllerEvent.GetUserIndex(), ControllerEvent.GetEffectingButton(), ControllerEvent.GetEffectingButton() == EKeys::Gamepad_RightY ? -ControllerEvent.GetAnalogValue() : ControllerEvent.GetAnalogValue(), FApp::GetDeltaTime(), 1, true))
 		{
 			CurrentReplyState = FReply::Unhandled(); 
 		}
diff --git a/Engine/Source/Runtime/Engine/Private/UnrealEngine.cpp b/Engine/Source/Runtime/Engine/Private/UnrealEngine.cpp
--- a/Engine/Source/Runtime/Engine/Private/UnrealEngine.cpp
+++ b/Engine/Source/Runtime/Engine/Private/UnrealEngine.cpp
@@ -890,32 +890,32 @@ void UEngine::UpdateTimeAndHandleMaxTickRate()
 	static bool bTimeWasManipulated = false;
 
 	// Figure out whether we want to use real or fixed time step.
-	const bool bUseFixedTimeStep = GIsBenchmarking || GUseFixedTimeStep;
+	const bool bUseFixedTimeStep = FApp::IsBenchmarking() || FApp::UseFixedTimeStep();
 
-	GLastTime = GCurrentTime;
+	FApp::UpdateLastTime();
 
 	// Calculate delta time and update time.
 	if( bUseFixedTimeStep )
 	{
 		bTimeWasManipulated = true;
 
-		GDeltaTime		= GFixedDeltaTime;
-		LastTime		= GCurrentTime;
-		GCurrentTime	+= GDeltaTime;
+		FApp::SetDeltaTime(FApp::GetFixedDeltaTime());
+		LastTime = FApp::GetCurrentTime();
+		FApp::SetCurrentTime(FApp::GetCurrentTime() + FApp::GetDeltaTime());
 	}
 	else
 	{
-		GCurrentTime = FPlatformTime::Seconds();
+		FApp::SetCurrentTime(FPlatformTime::Seconds());
 		// Did we just switch from a fixed time step to real-time?  If so, then we'll update our
 		// cached 'last time' so our current interval isn't huge (or negative!)
 		if( bTimeWasManipulated )
 		{
-			LastTime = GCurrentTime - GDeltaTime;
+			LastTime = FApp::GetCurrentTime() - FApp::GetDeltaTime();
 			bTimeWasManipulated = false;
 		}
 
 		// Calculate delta time.
-		float DeltaTime = GCurrentTime - LastTime;
+		float DeltaTime = FApp::GetCurrentTime() - LastTime;
 
 		// Negative delta time means something is wrong with the system. Error out so user can address issue.
 		if( DeltaTime < 0 )
@@ -938,7 +938,7 @@ void UEngine::UpdateTimeAndHandleMaxTickRate()
 		STAT( double ActualWaitTime = 0.f ); 
 		if( WaitTime > 0 )
 		{
-			double WaitEndTime = GCurrentTime + WaitTime;
+			double WaitEndTime = FApp::GetCurrentTime() + WaitTime;
 			SCOPE_SECONDS_COUNTER(ActualWaitTime);
 			SCOPE_CYCLE_COUNTER(STAT_GameTickWaitTime);
 			SCOPE_CYCLE_COUNTER(STAT_GameIdleTime);
@@ -963,24 +963,23 @@ void UEngine::UpdateTimeAndHandleMaxTickRate()
 					FPlatformProcess::Sleep( 0 );
 				}
 			}
-			GCurrentTime = FPlatformTime::Seconds();
+			FApp::SetCurrentTime(FPlatformTime::Seconds());
 		}
 
 
 		SET_FLOAT_STAT(STAT_GameTickWantedWaitTime,WaitTime * 1000.f);
 		SET_FLOAT_STAT(STAT_GameTickAdditionalWaitTime,FMath::Max<float>((ActualWaitTime-WaitTime)*1000.f,0.f));
 
-		GDeltaTime = GCurrentTime - LastTime;
+		FApp::SetDeltaTime(FApp::GetCurrentTime() - LastTime);
 
 		// Negative delta time means something is wrong with the system. Error out so user can address issue.
-		if( GDeltaTime < 0 )
+		if( FApp::GetDeltaTime() < 0 )
 		{
 			// AMD dual-core systems are a known issue that require AMD CPU drivers to be installed. Installer will take care of this for shipping.
 			UE_LOG(LogEngine, Fatal,TEXT("Detected negative delta time - on AMD systems please install http://files.aoaforums.com/I3199-setup.zip.html"));
-			GDeltaTime = 0.01;
+			FApp::SetDeltaTime(0.01);
 		}
-		GUnclampedDeltaTime = GDeltaTime;
-		LastTime			= GCurrentTime;
+		LastTime			= FApp::GetCurrentTime();
 
 		// Enforce a maximum delta time if wanted.
 		UGameEngine* GameEngine = Cast<UGameEngine>(this);
@@ -1008,7 +1007,7 @@ void UEngine::UpdateTimeAndHandleMaxTickRate()
 				&&	World->GetAuthGameMode()->NumPlayers == NumGamePlayers ) ) )
 			{
 				// Happy clamping!
-				GDeltaTime = FMath::Min<double>( GDeltaTime, MaxDeltaTime );
+				FApp::SetDeltaTime(FMath::Min<double>(FApp::GetDeltaTime(), MaxDeltaTime));
 			}
 		}
 	}
@@ -1022,9 +1021,9 @@ void UEngine::UpdateTimeAndHandleMaxTickRate()
 		if(OverrideFPS >= 0.001f)
 		{
 			// in seconds
-			GDeltaTime = 1.0f / OverrideFPS;
-			LastTime = GCurrentTime;
-			GCurrentTime += GDeltaTime;
+			FApp::SetDeltaTime(1.0f / OverrideFPS);
+			LastTime = FApp::GetCurrentTime();
+			FApp::SetCurrentTime(FApp::GetCurrentTime() + FApp::GetDeltaTime());
 			bTimeWasManipulated = true;
 		}
 	}
@@ -1035,7 +1034,7 @@ void UEngine::UpdateTimeAndHandleMaxTickRate()
 void UEngine::ParseCommandline()
 {
 	// If dedicated server, the -nosound, or -benchmark parameters are used, disable sound.
-	if(FParse::Param(FCommandLine::Get(),TEXT("nosound")) || GIsBenchmarking || IsRunningDedicatedServer() || IsRunningCommandlet())
+	if(FParse::Param(FCommandLine::Get(),TEXT("nosound")) || FApp::IsBenchmarking() || IsRunningDedicatedServer() || IsRunningCommandlet())
 	{
 		bUseSound = false;
 	}
@@ -3115,7 +3114,7 @@ bool UEngine::HandleRemoteTextureStatsCommand( const TCHAR* Cmd, FOutputDevice&
 	//SrcAddr.Addr = FCString::Atoi( *Addr );
 
 	// Gather stats.
-	double LastTime = GLastTime;
+	double LastTime = FApp::GetLastTime();
 
 	UE_LOG(LogEngine, Log,  TEXT("Remote AssetsStats request received.") );
 
@@ -7328,7 +7327,7 @@ void DrawUnitTimeGraph( FViewport* Viewport, FCanvas* Canvas, const bool bHaveGP
 
 		// Compute pulse effect for lines above alert threshold
 		const float AlertPulseFreq = 8.0f;
-		const float AlertPulse = 0.5f + 0.5f * (float)sin( ( 0.25f * PI * 2.0 ) + ( GCurrentTime * PI * 2.0 ) * AlertPulseFreq );
+		const float AlertPulse = 0.5f + 0.5f * (float)sin( ( 0.25f * PI * 2.0 ) + ( FApp::GetCurrentTime() * PI * 2.0 ) * AlertPulseFreq );
 
 
 		// For each type of statistic that we want to graph (0=Render, 1=Game, 2=GPU, 3=Frame)
@@ -7535,9 +7534,9 @@ int32 DrawUnitTimes( FViewport* Viewport, FCanvas* Canvas, int32 X, int32 Y )
 	extern uint32 GGameThreadTime;
 
 	float FrameTime;
-	if (GIsBenchmarking || GUseFixedTimeStep)
+	if (FApp::IsBenchmarking() || FApp::UseFixedTimeStep())
 	{
-		/** If we're in fixed time step mode, GCurrentTime will be incorrect for benchmarking */
+		/** If we're in fixed time step mode, FApp::GetCurrentTime() will be incorrect for benchmarking */
 		static double LastTime = 0.0;
 		const double CurrentTime = FPlatformTime::Seconds();
 		if ( LastTime == 0 )
@@ -7550,7 +7549,7 @@ int32 DrawUnitTimes( FViewport* Viewport, FCanvas* Canvas, int32 X, int32 Y )
 	else
 	{
 		/** Use the FrameTime we computed last frame, because it correctly handles the end of frame idling and corresponds better to the other unit times. */
-		FrameTime = GCurrentTime - GLastTime;
+		FrameTime = FApp::GetCurrentTime() - FApp::GetLastTime();
 	}
 	
 	GUnit_RawFrameTime		= FrameTime * 1000.0f;
@@ -7760,7 +7759,7 @@ void DrawStatsHUD( UWorld* World, FViewport* Viewport, FCanvas* Canvas, UCanvas*
 			{
 				SmallTextItem.SetColor( FLinearColor::White );
 				// Color unbuilt lighting red if encountered within the last second
-				if( GCurrentTime - World->LastTimeUnbuiltLightingWasEncountered < 1 )
+				if( FApp::GetCurrentTime() - World->LastTimeUnbuiltLightingWasEncountered < 1 )
 				{
 					SmallTextItem.SetColor( FLinearColor::Red );
 				}
diff --git a/Engine/Source/Runtime/Launch/Private/LaunchEngineLoop.cpp b/Engine/Source/Runtime/Launch/Private/LaunchEngineLoop.cpp
--- a/Engine/Source/Runtime/Launch/Private/LaunchEngineLoop.cpp
+++ b/Engine/Source/Runtime/Launch/Private/LaunchEngineLoop.cpp
@@ -769,13 +769,13 @@ int32 FEngineLoop::PreInit( const TCHAR* CmdLine )
 
 #if !UE_BUILD_SHIPPING
 	// Benchmarking.
-	GIsBenchmarking	= FParse::Param(FCommandLine::Get(),TEXT("BENCHMARK"));
+	FApp::SetBenchmarking(FParse::Param(FCommandLine::Get(),TEXT("BENCHMARK")));
 #else
-	GIsBenchmarking = false;
+	FApp::SetBenchmarking(false);
 #endif // !UE_BUILD_SHIPPING
 
 	// Initialize random number generator.
-	if( GIsBenchmarking || FParse::Param(FCommandLine::Get(),TEXT("FIXEDSEED")) )
+	if( FApp::IsBenchmarking() || FParse::Param(FCommandLine::Get(),TEXT("FIXEDSEED")) )
 	{
 		FMath::RandInit( 0 );
 		FMath::SRandInit( 0 );
@@ -1462,7 +1462,7 @@ int32 FEngineLoop::PreInit( const TCHAR* CmdLine )
 	}
 
 	// Don't update INI files if benchmarking or -noini
-	if( GIsBenchmarking || FParse::Param(FCommandLine::Get(),TEXT("NOINI")))
+	if( FApp::IsBenchmarking() || FParse::Param(FCommandLine::Get(),TEXT("NOINI")))
 	{
 		GConfig->Detach( GEngineIni );
 		GConfig->Detach( GInputIni );
@@ -1646,7 +1646,7 @@ bool FEngineLoop::LoadStartupModules()
 void FEngineLoop::InitTime()
 {
 	// Init variables used for benchmarking and ticking.
-	GCurrentTime				= FPlatformTime::Seconds();
+	FApp::SetCurrentTime(FPlatformTime::Seconds());
 	MaxFrameCounter				= 0;
 	MaxTickTime					= 0;
 	TotalTickTime				= 0;
@@ -1657,12 +1657,12 @@ void FEngineLoop::InitTime()
 	FParse::Value(FCommandLine::Get(),TEXT("SECONDS="),FloatMaxTickTime);
 	MaxTickTime					= FloatMaxTickTime;
 
-	// look of a version of seconds that only is applied if GIsBenchmarking is set. This makes it easier on
+	// look of a version of seconds that only is applied if FApp::IsBenchmarking() is set. This makes it easier on
 	// say, iOS, where we have a toggle setting to enable benchmarking, but don't want to have to make user
 	// also disable the seconds setting as well. -seconds= will exit the app after time even if benchmarking
 	// is not enabled
 	// NOTE: This will override -seconds= if it's specified
-	if (GIsBenchmarking)
+	if (FApp::IsBenchmarking())
 	{
 		if (FParse::Value(FCommandLine::Get(),TEXT("BENCHMARKSECONDS="),FloatMaxTickTime) && FloatMaxTickTime)
 		{
@@ -1676,16 +1676,13 @@ void FEngineLoop::InitTime()
 	if( FixedFPS > 0 )
 	{
 		GEngine->MatineeCaptureFPS = (int32)FixedFPS;
-		GFixedDeltaTime = 1 / FixedFPS;
+		FApp::SetFixedDeltaTime(1 / FixedFPS);
 	}
 
-	// Whether we want to use a fixed time step or not.
-	GUseFixedTimeStep = FParse::Param( FCommandLine::Get(), TEXT( "UseFixedTimeStep" ) );
-
 #endif // !UE_BUILD_SHIPPING
 
-	// convert FloatMaxTickTime into number of frames (using 1 / GFixedDeltaTime to convert fps to seconds )
-	MaxFrameCounter				= FMath::Trunc( MaxTickTime / GFixedDeltaTime );
+	// convert FloatMaxTickTime into number of frames (using 1 / FApp::GetFixedDeltaTime() to convert fps to seconds )
+	MaxFrameCounter = FMath::Trunc(MaxTickTime / FApp::GetFixedDeltaTime());
 }
 
 void FlushStatsFrame(bool bDiscardCallstack)
@@ -1935,17 +1932,17 @@ void FEngineLoop::Tick()
 		GLog->FlushThreadedLogs();
 
 		// Exit if frame limit is reached in benchmark mode.
-		if( (GIsBenchmarking && MaxFrameCounter && (GFrameCounter > MaxFrameCounter))
+		if( (FApp::IsBenchmarking() && MaxFrameCounter && (GFrameCounter > MaxFrameCounter))
 			// or time limit is reached if set.
 			|| (MaxTickTime && (TotalTickTime > MaxTickTime)) )
 		{
 			FPlatformMisc::RequestExit(0);
 		}
 
-		// Set GCurrentTime, GDeltaTime and potentially wait to enforce max tick rate.
+		// Set FApp::CurrentTime, FApp::DeltaTime and potentially wait to enforce max tick rate.
 		GEngine->UpdateTimeAndHandleMaxTickRate();
 
-		GEngine->TickFPSChart( GDeltaTime );
+		GEngine->TickFPSChart( FApp::GetDeltaTime() );
 
 		// Update platform memory and memory allocator stats.
 		FPlatformMemory::UpdateStats();
@@ -1994,7 +1991,7 @@ void FEngineLoop::Tick()
 			FSlateApplication::Get().PollGameDeviceState();
 		}
 
-		GEngine->Tick( GDeltaTime, bIdleMode );
+		GEngine->Tick( FApp::GetDeltaTime(), bIdleMode );
 
 		if (GShaderCompilingManager)
 		{
@@ -2036,7 +2033,7 @@ void FEngineLoop::Tick()
 #if WITH_ENGINE
 			SCOPE_CYCLE_COUNTER(STAT_RHITickTime);
 #endif
-			RHITick( GDeltaTime ); // Update RHI.
+			RHITick( FApp::GetDeltaTime() ); // Update RHI.
 		}
 
 		// Increment global frame counter. Once for each engine tick.
@@ -2045,7 +2042,7 @@ void FEngineLoop::Tick()
 		// Disregard first few ticks for total tick time as it includes loading and such.
 		if( GFrameCounter > 6 )
 		{
-			TotalTickTime+=GDeltaTime;
+			TotalTickTime+=FApp::GetDeltaTime();
 		}
 
 
@@ -2072,7 +2069,7 @@ void FEngineLoop::Tick()
 			// Delete the objects which were enqueued for deferred cleanup before the previous frame.
 			delete PreviousPendingCleanupObjects;
 
-			FTicker::GetCoreTicker().Tick(GDeltaTime);
+			FTicker::GetCoreTicker().Tick(FApp::GetDeltaTime());
 
 			FSingleThreadManager::Get().Tick();
 
diff --git a/Engine/Source/Runtime/Launch/Private/Windows/LaunchWindows.cpp b/Engine/Source/Runtime/Launch/Private/Windows/LaunchWindows.cpp
--- a/Engine/Source/Runtime/Launch/Private/Windows/LaunchWindows.cpp
+++ b/Engine/Source/Runtime/Launch/Private/Windows/LaunchWindows.cpp
@@ -166,7 +166,7 @@ int32 WINAPI WinMain( HINSTANCE hInInstance, HINSTANCE hPrevInstance, char*, int
 	// The inner exception handler in GuardedMainWrapper() catches crashes/asserts in native C++ code and is the only way to get the
 	// correct callstack when running a 64-bit executable. However, XAudio2 sometimes (?) don't like this and it may result in no sound.
 #if _WIN64
-	if ( FParse::Param(CmdLine,TEXT("noinnerexception")) || GIsBenchmarking )
+	if ( FParse::Param(CmdLine,TEXT("noinnerexception")) || FApp::IsBenchmarking() )
 	{
 		GEnableInnerException = false;
 	}
diff --git a/Engine/Source/Runtime/Launch/Public/LaunchEngineLoop.h b/Engine/Source/Runtime/Launch/Public/LaunchEngineLoop.h
--- a/Engine/Source/Runtime/Launch/Public/LaunchEngineLoop.h
+++ b/Engine/Source/Runtime/Launch/Public/LaunchEngineLoop.h
@@ -121,7 +121,7 @@ public:
 protected:
 
 	/**
-	 * Holds a dynamically expanding array of frame times in milliseconds (if GIsBenchmarking is set).
+	 * Holds a dynamically expanding array of frame times in milliseconds (if FApp::IsBenchmarking() is set).
 	 */
 	TArray<float> FrameTimes;
 
diff --git a/Engine/Source/Runtime/RenderCore/Private/RenderResource.cpp b/Engine/Source/Runtime/RenderCore/Private/RenderResource.cpp
--- a/Engine/Source/Runtime/RenderCore/Private/RenderResource.cpp
+++ b/Engine/Source/Runtime/RenderCore/Private/RenderResource.cpp
@@ -512,14 +512,14 @@ float GMipLevelFadingAgeThreshold = 0.5f;
  *	Sets up a new interpolation target for the mip-bias.
  *	@param ActualMipCount	Number of mip-levels currently in memory
  *	@param TargetMipCount	Number of mip-levels we're changing to
- *	@param LastRenderTime	Timestamp when it was last rendered (GCurrentTime time space)
+ *	@param LastRenderTime	Timestamp when it was last rendered (FApp::CurrentTime time space)
  *	@param FadeSetting		Which fade speed settings to use
  */
 void FMipBiasFade::SetNewMipCount( float ActualMipCount, float TargetMipCount, double LastRenderTime, EMipFadeSettings FadeSetting )
 {
 	check( ActualMipCount >=0 && TargetMipCount <= ActualMipCount );
 
-	float TimeSinceLastRendered = float(GCurrentTime - LastRenderTime);
+	float TimeSinceLastRendered = float(FApp::GetCurrentTime() - LastRenderTime);
 
 	// Is this a new texture or is this not in-game?
 	if ( TotalMipCount == 0 || TimeSinceLastRendered >= GMipLevelFadingAgeThreshold || GEnableMipLevelFading < 0.0f )
diff --git a/Engine/Source/Runtime/RenderCore/Public/RenderResource.h b/Engine/Source/Runtime/RenderCore/Public/RenderResource.h
--- a/Engine/Source/Runtime/RenderCore/Public/RenderResource.h
+++ b/Engine/Source/Runtime/RenderCore/Public/RenderResource.h
@@ -241,7 +241,7 @@ struct FMipBiasFade
 	 *	Sets up a new interpolation target for the mip-bias.
 	 *	@param ActualMipCount	Number of mip-levels currently in memory
 	 *	@param TargetMipCount	Number of mip-levels we're changing to
-	 *	@param LastRenderTime	Timestamp when it was last rendered (GCurrentTime time space)
+	 *	@param LastRenderTime	Timestamp when it was last rendered (FApp::CurrentTime time space)
 	 *	@param FadeSetting		Which fade speed settings to use
 	 */
 	RENDERCORE_API void	SetNewMipCount( float ActualMipCount, float TargetMipCount, double LastRenderTime, EMipFadeSettings FadeSetting );
diff --git a/Engine/Source/Runtime/Renderer/Private/PostProcess/VisualizeTexture.cpp b/Engine/Source/Runtime/Renderer/Private/PostProcess/VisualizeTexture.cpp
--- a/Engine/Source/Runtime/Renderer/Private/PostProcess/VisualizeTexture.cpp
+++ b/Engine/Source/Runtime/Renderer/Private/PostProcess/VisualizeTexture.cpp
@@ -80,7 +80,7 @@ public:
 		{
 			// alternates between 0 and 1 with a short pause
 			const float FracTimeScale = 2.0f;
-			float FracTime = GCurrentTime * FracTimeScale - floor(GCurrentTime * FracTimeScale);
+			float FracTime = FApp::GetCurrentTime() * FracTimeScale - floor(FApp::GetCurrentTime() * FracTimeScale);
 			float BlinkState = (FracTime > 0.5f) ? 1.0f : 0.0f;
 
 			FVector4 VisualizeParamValue[3];
diff --git a/Engine/Source/Runtime/Renderer/Private/ShaderBaseClasses.h b/Engine/Source/Runtime/Renderer/Private/ShaderBaseClasses.h
--- a/Engine/Source/Runtime/Renderer/Private/ShaderBaseClasses.h
+++ b/Engine/Source/Runtime/Renderer/Private/ShaderBaseClasses.h
@@ -232,7 +232,7 @@ public:
 				}
 
 				checkSlow(Value);
-				Value->LastRenderTime = GCurrentTime;
+				Value->LastRenderTime = FApp::GetCurrentTime();
 
 				SetTextureParameter(
 					ShaderRHI, 
@@ -261,7 +261,7 @@ public:
 				}
 
 				checkSlow(Value);
-				Value->LastRenderTime = GCurrentTime;
+				Value->LastRenderTime = FApp::GetCurrentTime();
 
 				SetTextureParameter(
 					ShaderRHI,
diff --git a/Engine/Source/Runtime/ShaderCore/Public/ShaderParameters.h b/Engine/Source/Runtime/ShaderCore/Public/ShaderParameters.h
--- a/Engine/Source/Runtime/ShaderCore/Public/ShaderParameters.h
+++ b/Engine/Source/Runtime/ShaderCore/Public/ShaderParameters.h
@@ -232,7 +232,7 @@ FORCEINLINE void SetTextureParameter(
 	checkSlow(SamplerParameter.IsInitialized());
 	if(TextureParameter.IsBound())
 	{
-		Texture->LastRenderTime = GCurrentTime;
+		Texture->LastRenderTime = FApp::GetCurrentTime();
 
 		if (ElementIndex < TextureParameter.GetNumResources())
 		{
diff --git a/Engine/Source/Runtime/SlateRHIRenderer/Private/SlateRHIResourceManager.cpp b/Engine/Source/Runtime/SlateRHIRenderer/Private/SlateRHIResourceManager.cpp
--- a/Engine/Source/Runtime/SlateRHIRenderer/Private/SlateRHIResourceManager.cpp
+++ b/Engine/Source/Runtime/SlateRHIRenderer/Private/SlateRHIResourceManager.cpp
@@ -435,7 +435,7 @@ FSlateShaderResourceProxy* FSlateRHIResourceManager::GetDynamicTextureResource(
 				{
 					DynamicTextureResource->RHIRefTexture->SetRHIRef( InFTexture->TextureRHI->GetTexture2D(), InFTexture->GetSizeX(), InFTexture->GetSizeY() );
 					// Let the streaming manager know we are using this texture now
-					InFTexture->LastRenderTime = GCurrentTime;
+					InFTexture->LastRenderTime = FApp::GetCurrentTime();
 				});
 
 				AccessedUTextures.Add( TextureResource );
