Landscape : Fixed problems for neighbor component when LODBias/ForcedLOD change with multiple component selection

diff --git a/Engine/Source/Runtime/Engine/Classes/Landscape/LandscapeComponent.h b/Engine/Source/Runtime/Engine/Classes/Landscape/LandscapeComponent.h
--- a/Engine/Source/Runtime/Engine/Classes/Landscape/LandscapeComponent.h
+++ b/Engine/Source/Runtime/Engine/Classes/Landscape/LandscapeComponent.h
@@ -230,6 +230,7 @@ public:
 #if WITH_EDITOR
 	virtual void PostLoad() OVERRIDE;
 	virtual void PostEditUndo() OVERRIDE;
+	virtual void PreEditChange(UProperty* PropertyThatWillChange) OVERRIDE;
 	virtual void PostEditChangeProperty(FPropertyChangedEvent& PropertyChangedEvent) OVERRIDE;
 	// End UObject interface
 
diff --git a/Engine/Source/Runtime/Engine/Private/Landscape/LandscapeEdit.cpp b/Engine/Source/Runtime/Engine/Private/Landscape/LandscapeEdit.cpp
--- a/Engine/Source/Runtime/Engine/Private/Landscape/LandscapeEdit.cpp
+++ b/Engine/Source/Runtime/Engine/Private/Landscape/LandscapeEdit.cpp
@@ -3854,6 +3854,26 @@ void ULandscapeComponent::SetLOD(bool bForcedLODChanged, int32 InLODValue)
 	FComponentReregisterContext ReregisterContext(this);
 }
 
+void ULandscapeComponent::PreEditChange(UProperty* PropertyThatWillChange)
+{
+	Super::PreEditChange(PropertyThatWillChange);
+	if (GIsEditor && PropertyThatWillChange && (PropertyThatWillChange->GetName() == TEXT("ForcedLOD") || PropertyThatWillChange->GetName() == TEXT("LODBias")))
+	{
+		// PreEdit unregister component and re-register after PostEdit so we will lose XYtoComponentMap for this component
+		ULandscapeInfo* Info = GetLandscapeInfo(false);
+		if (Info)
+		{
+			FIntPoint ComponentKey = GetSectionBase() / ComponentSizeQuads;
+			auto RegisteredComponent = Info->XYtoComponentMap.FindRef(ComponentKey);
+
+			if (RegisteredComponent == NULL)
+			{
+				Info->XYtoComponentMap.Add(ComponentKey, this);
+			}
+		}
+	}
+}
+
 void ULandscapeComponent::PostEditChangeProperty(FPropertyChangedEvent& PropertyChangedEvent)
 {
 	Super::PostEditChangeProperty(PropertyChangedEvent);
