#ttp 331962 - UE4: RENDERING: Fix rare crash in anim trails

#branch UE4

#proj Engine

#summary Reasonabley confident I've found the cause of the crash and fixed it. Can't be 100% though because of the bugs rarity.

#extra The cause of this was some particles being killed correctly. If a MIDDLE particle was killed and it's prev particle was a dead trail then the dead particle would not have it's state cleared correcly meaning that it could throw out later triangle calculations and trigger a check.

diff --git a/Engine/Source/Runtime/Engine/Private/Particles/ParticleTrail2EmitterInstance.cpp b/Engine/Source/Runtime/Engine/Private/Particles/ParticleTrail2EmitterInstance.cpp
--- a/Engine/Source/Runtime/Engine/Private/Particles/ParticleTrail2EmitterInstance.cpp
+++ b/Engine/Source/Runtime/Engine/Private/Particles/ParticleTrail2EmitterInstance.cpp
@@ -53,6 +53,11 @@ void FParticleTrailsEmitterInstance_Base::Tick(float DeltaTime, bool bSuppressSp
 	SCOPE_CYCLE_COUNTER(STAT_TrailTickTime);
 	if (Component)
 	{
+
+#if ULTRA_VERBOSE_TRAILS_DEBUG
+		PrintTrails();
+#endif
+
 		check(SpriteTemplate);
 		check(SpriteTemplate->LODLevels.Num() > 0);
 
@@ -409,6 +414,12 @@ void FParticleTrailsEmitterInstance_Base::KillParticles()
 						{
 							PrevTrailData->Flags = TRAIL_EMITTER_SET_ONLY(PrevTrailData->Flags);
 						}
+						else if (TRAIL_EMITTER_IS_DEADTRAIL(PrevTrailData->Flags))
+						{
+							// Nothing to do in this case.
+							PrevTrailData->TriangleCount = 0;
+							PrevTrailData->RenderingInterpCount = 1;
+						}
 						else
 						{
 							PrevTrailData->Flags = TRAIL_EMITTER_SET_END(PrevTrailData->Flags);
@@ -2805,6 +2816,7 @@ FParticleAnimTrailEmitterInstance::FParticleAnimTrailEmitterInstance() :
 	, bRenderTangents(false)
 	, bRenderTessellation(false)
 #endif
+	, HeadOnlyParticles(0)
 {
 }
 
@@ -3663,6 +3675,7 @@ void FParticleAnimTrailEmitterInstance::DetermineVertexAndTriangleCount()
 
 	VertexCount		= 0;
 	TriangleCount	= 0;
+	HeadOnlyParticles = 0;
 
 	int32 CheckParticleCount = 0;
 
@@ -3693,10 +3706,11 @@ void FParticleAnimTrailEmitterInstance::DetermineVertexAndTriangleCount()
 		DECLARE_PARTICLE_PTR(Particle, ParticleData + ParticleStride * ParticleIndices[ii]);
 		FBaseParticle* CurrParticle = Particle;
 		FAnimTrailTypeDataPayload*	CurrTrailData = ((FAnimTrailTypeDataPayload*)((uint8*)Particle + TypeDataOffset));
-		if (TRAIL_EMITTER_IS_ONLY(CurrTrailData->Flags))
+		if (TRAIL_EMITTER_IS_HEADONLY(CurrTrailData->Flags))
 		{
 			CurrTrailData->RenderingInterpCount = 0;
 			CurrTrailData->TriangleCount = 0;
+			++HeadOnlyParticles;
 		}
 		else if (TRAIL_EMITTER_IS_END(CurrTrailData->Flags))
 		{
@@ -4131,7 +4145,7 @@ bool FParticleAnimTrailEmitterInstance::FillReplayData( FDynamicEmitterReplayDat
 
 	if (TriangleCount <= 0)
 	{
-		if (ActiveParticles > 1)
+		if (ActiveParticles > 0 && ActiveParticles != HeadOnlyParticles)
 		{
 // 				UE_LOG(LogParticles, Warning, TEXT("TRAIL: GetDynamicData -- TriangleCount == 0 (APC = %4d) for PSys %s"),
 // 					ActiveParticles, 
@@ -4149,6 +4163,9 @@ bool FParticleAnimTrailEmitterInstance::FillReplayData( FDynamicEmitterReplayDat
 				GEngine->AddOnScreenDebugMessage((uint64)((PTRINT)this), 5.0f, ErrorColor,ErrorMessage);
 				UE_LOG(LogParticles, Log, TEXT("%s"), *ErrorMessage);
 			}
+
+			PrintAllActiveParticles();
+
 #endif	//#if !(UE_BUILD_SHIPPING || UE_BUILD_TEST)
 		}
 		return false;
@@ -4178,3 +4195,90 @@ bool FParticleAnimTrailEmitterInstance::FillReplayData( FDynamicEmitterReplayDat
 
 	return true;
 }
+
+void FParticleAnimTrailEmitterInstance::PrintParticleData(FBaseParticle* Particle, FTrailsBaseTypeDataPayload* TrailData, int32 CurrentIndex, int32 TrailIndex)
+{
+	UE_LOG(LogParticles, Log, TEXT("%d| Particle %4d - Next = %4d, Prev = %4d, Type = %8s, Life=%.5f"),
+		TrailIndex,
+		CurrentIndex,
+		TRAIL_EMITTER_GET_NEXT(TrailData->Flags),
+		TRAIL_EMITTER_GET_PREV(TrailData->Flags),
+		TRAIL_EMITTER_IS_ONLY(TrailData->Flags) ? TEXT("ONLY") :
+		TRAIL_EMITTER_IS_START(TrailData->Flags) ? TEXT("START") :
+		TRAIL_EMITTER_IS_END(TrailData->Flags) ? TEXT("END") :
+		TRAIL_EMITTER_IS_MIDDLE(TrailData->Flags) ? TEXT("MIDDLE") :
+		TRAIL_EMITTER_IS_DEADTRAIL(TrailData->Flags) ? TEXT("DEAD") :
+		TEXT("????"),
+		Particle->RelativeTime
+		);
+	UE_LOG(LogParticles, Log, TEXT("TI=%4d - TC=%4d, SpT=%.5f, SpD= %.5f, TU=%.4f, SpTP=%4d, RIntpC=%4d, PSF=%.4f, %s, %s"),
+		TrailData->TrailIndex, TrailData->TriangleCount, TrailData->SpawnTime, TrailData->SpawnDelta, TrailData->TiledU, TrailData->SpawnedTessellationPoints,
+		TrailData->RenderingInterpCount, TrailData->PinchScaleFactor, TrailData->bInterpolatedSpawn ? "1" : "0", TrailData->bMovementSpawned ? "1" : "0");
+}
+
+void FParticleAnimTrailEmitterInstance::PrintAllActiveParticles()
+{
+	UE_LOG(LogParticles, Log, TEXT("==========================================================================="));
+	//Print out the state of all the particles. 
+	for (int32 ActiveIdx = 0; ActiveIdx < ActiveParticles; ++ActiveIdx)
+	{
+		int32 CurrentIndex = ParticleIndices[ActiveIdx];
+		DECLARE_PARTICLE_PTR(Particle, ParticleData + ParticleStride * CurrentIndex);
+		FTrailsBaseTypeDataPayload* TrailData = ((FTrailsBaseTypeDataPayload*)((uint8*)Particle + TypeDataOffset));
+		PrintParticleData(Particle, TrailData, CurrentIndex, -1);
+		UE_LOG(LogParticles, Log, TEXT("---------------------------------------------------------------------------------------"));
+	}
+	UE_LOG(LogParticles, Log, TEXT("==========================================================================="));
+}
+
+void FParticleAnimTrailEmitterInstance::PrintTrails()
+{
+	if (ActiveParticles > 0)
+	{
+		UE_LOG(LogParticles, Log, TEXT("==========================================="));
+		UE_LOG(LogParticles, Log, TEXT("Active: %d"), ActiveParticles);
+		UE_LOG(LogParticles, Log, TEXT("==========================================="));
+		TArray<int32> ParticlesVisited;
+		TArray<int32> TrailHeads;
+		for (int32 ActiveIdx = 0; ActiveIdx < ActiveParticles; ++ActiveIdx)
+		{
+			int32 CurrentIndex = ParticleIndices[ActiveIdx];
+			//Work forwards through the trail printing data about the trail.
+			DECLARE_PARTICLE_PTR(Particle, ParticleData + ParticleStride * CurrentIndex);
+			FTrailsBaseTypeDataPayload* TrailData = ((FTrailsBaseTypeDataPayload*)((uint8*)Particle + TypeDataOffset));
+			if (TRAIL_EMITTER_IS_HEAD(TrailData->Flags))
+			{
+				TrailHeads.Add(CurrentIndex);
+				while (true)
+				{
+					//Ensure we've not already visited this particle.
+					check(!ParticlesVisited.Contains(CurrentIndex));
+
+					ParticlesVisited.Add(CurrentIndex);
+
+					PrintParticleData(Particle, TrailData, CurrentIndex, TrailHeads.Num() - 1);
+
+					int32 Next = TRAIL_EMITTER_GET_NEXT(TrailData->Flags);
+					if (Next == TRAIL_EMITTER_NULL_NEXT)
+					{
+						UE_LOG(LogParticles, Log, TEXT("==========================================================================="));
+						break;
+					}
+					else
+					{
+						UE_LOG(LogParticles, Log, TEXT("---------------------------------------------------------------------------------------"));
+						CurrentIndex = Next;
+						Particle = (FBaseParticle*)(ParticleData + ParticleStride * Next);
+						TrailData = ((FTrailsBaseTypeDataPayload*)((uint8*)Particle + TypeDataOffset));
+					}
+				}
+			}
+		}
+
+		//check that all particles were visited. If not then there are some orphaned particles munging things up.
+		if (ParticlesVisited.Num() != ActiveParticles)
+		{
+			PrintAllActiveParticles();
+		}
+	}
+}
diff --git a/Engine/Source/Runtime/Engine/Public/ParticleEmitterInstances.h b/Engine/Source/Runtime/Engine/Public/ParticleEmitterInstances.h
--- a/Engine/Source/Runtime/Engine/Public/ParticleEmitterInstances.h
+++ b/Engine/Source/Runtime/Engine/Public/ParticleEmitterInstances.h
@@ -1149,6 +1149,13 @@ protected:
 		EGetTrailParticleOption InGetOption,
 		FBaseParticle*& OutParticle,
 		FTrailsBaseTypeDataPayload*& OutTrailData);
+
+	/** Prints out info for a single particle. */
+	virtual void PrintParticleData(FBaseParticle* Particle, FTrailsBaseTypeDataPayload* TrailData, int32 CurrentIndex, int32 TrailIndex){}
+	/** Prints out info for all active particles. */
+	virtual void PrintAllActiveParticles(){}
+	/** Traverses all trails and prints out debugging info. */
+	virtual void PrintTrails(){}
 };
 
 struct FParticleRibbonEmitterInstance : public FParticleTrailsEmitterInstance_Base
@@ -1394,6 +1401,9 @@ struct FParticleAnimTrailEmitterInstance : public FParticleTrailsEmitterInstance
 	uint32 bRenderTessellation : 1;
 #endif
 
+	/** The number of particles in the trail which are either marked */
+	int32 HeadOnlyParticles;
+
 	/** Constructor	*/
 	FParticleAnimTrailEmitterInstance();
 
@@ -1537,4 +1547,11 @@ protected:
 	*	@param	Params			Various parameters for spawning.
 	*/
 	void SpawnParticle( int32& StartParticleIndex, const struct FAnimTrailParticleSpawnParams& Params );
+
+	/** Prints out info for a single particle. */
+	virtual void PrintParticleData(FBaseParticle* Particle, FTrailsBaseTypeDataPayload* TrailData, int32 CurrentIndex, int32 TrailIndex);
+	/** Prints out info for all active particles. */
+	virtual void PrintAllActiveParticles();
+	/** Traverses all trails and prints out debugging info. */
+	virtual void PrintTrails();
 };
