Make HUD hit box click and Component OnClick events respond to double clicks
TTP #331926

diff --git a/Engine/Source/Runtime/Engine/Private/HUD.cpp b/Engine/Source/Runtime/Engine/Private/HUD.cpp
--- a/Engine/Source/Runtime/Engine/Private/HUD.cpp
+++ b/Engine/Source/Runtime/Engine/Private/HUD.cpp
@@ -708,18 +708,18 @@ const TArray<FIntRect>& AHUD::GetUIBlurRectangles() const
 
 bool AHUD::UpdateAndDispatchHitBoxClickEvents(EInputEvent InEventType)
 {
-	for (int32 iBox = 0; iBox < HitBoxHits.Num(); iBox++)
+	for (FHUDHitBox* HitBoxHit : HitBoxHits)
 	{
-		if(InEventType == IE_Pressed)
+		if(InEventType == IE_Pressed || InEventType == IE_DoubleClick)
 		{
-			ReceiveHitBoxClick( HitBoxHits[ iBox ]->GetName() );
+			ReceiveHitBoxClick( HitBoxHit->GetName() );
 		}
 		else if(InEventType == IE_Released)
 		{
-			ReceiveHitBoxRelease( HitBoxHits[ iBox ]->GetName() );
+			ReceiveHitBoxRelease( HitBoxHit->GetName() );
 		}
 
-		if( HitBoxHits[ iBox ]->ConsumesInput() == true )
+		if( HitBoxHit->ConsumesInput() == true )
 		{
 			break;	//Early out if this box consumed the click
 		}
diff --git a/Engine/Source/Runtime/Engine/Private/PlayerController.cpp b/Engine/Source/Runtime/Engine/Private/PlayerController.cpp
--- a/Engine/Source/Runtime/Engine/Private/PlayerController.cpp
+++ b/Engine/Source/Runtime/Engine/Private/PlayerController.cpp
@@ -1961,6 +1961,7 @@ bool APlayerController::InputKey(FKey Key, EInputEvent EventType, float AmountDe
 				switch(EventType)
 				{
 				case IE_Pressed:
+				case IE_DoubleClick:
 					ClickedPrimitive->DispatchOnClicked();
 					break;
 
