Fixed comment so that u sers can understand it better - by request of Zak P

diff --git a/Engine/Source/Runtime/Engine/Classes/Animation/AnimInstance.h b/Engine/Source/Runtime/Engine/Classes/Animation/AnimInstance.h
--- a/Engine/Source/Runtime/Engine/Classes/Animation/AnimInstance.h
+++ b/Engine/Source/Runtime/Engine/Classes/Animation/AnimInstance.h
@@ -294,7 +294,15 @@ public:
 	UFUNCTION(BlueprintCallable, Category="Animation")
 	void Montage_JumpToSectionsEnd(FName SectionName);
 
-	/** Changes the next section in the montage to a different one. */
+	/** Relink new next section AFTER SectionNameToChange in run-time
+	 *	You can link section order the way you like in editor, but in run-time if you'd like to change it dynamically, 
+	 *	use this function to relink the next section
+	 *	For example, you can have Start->Loop->Loop->Loop.... but when you want it to end, you can relink
+	 *	next section of Loop to be End to finish the montage, in which case, it stops looping by Loop->End. 
+	 
+	 * @param SectionNameToChange : This should be the name of the Montage Section after which you want to insert a new next section
+	 * @param NextSection	: new next section 
+	 */
 	UFUNCTION(BlueprintCallable, Category="Animation")
 	void Montage_SetNextSection(FName SectionNameToChange, FName NextSection);
 
