Prevent LocalPlayers from being created with an already used controller ID

diff --git a/Engine/Source/Runtime/Engine/Private/GameViewportClient.cpp b/Engine/Source/Runtime/Engine/Private/GameViewportClient.cpp
--- a/Engine/Source/Runtime/Engine/Private/GameViewportClient.cpp
+++ b/Engine/Source/Runtime/Engine/Private/GameViewportClient.cpp
@@ -1115,7 +1115,11 @@ ULocalPlayer* UGameViewportClient::CreatePlayer(int32 ControllerId, FString& Out
 	ULocalPlayer* NewPlayer = NULL;
 	int32 InsertIndex = INDEX_NONE;
 
-	if (GetOuterUEngine()->GetNumGamePlayers(this) < MaxSplitscreenPlayers)
+	if (GetOuterUEngine()->GetLocalPlayerFromControllerId(this, ControllerId) != NULL)
+	{
+		OutError = FString::Printf(TEXT("A local player already exists for controller ID %d,"), ControllerId);
+	}
+	else if (GetOuterUEngine()->GetNumGamePlayers(this) < MaxSplitscreenPlayers)
 	{
 		NewPlayer = CastChecked<ULocalPlayer>(StaticConstructObject(GetOuterUEngine()->LocalPlayerClass, GetOuter()));	
 		InsertIndex = AddLocalPlayer(NewPlayer, ControllerId);
@@ -1199,6 +1203,10 @@ void UGameViewportClient::DebugCreatePlayer(int32 ControllerId)
 #if !UE_BUILD_SHIPPING
 	FString Error;
 	CreatePlayer(ControllerId, Error, true);
+	if (Error.Len() > 0)
+	{
+		UE_LOG(LogPlayerManagement, Warning, TEXT("Failed to DebugCreatePlayer. Error: %s"), *Error);
+	}
 #endif
 }
 
