#ue4

- InitCapsuleSize() enforces that CapsuleHalfHeight must be at least CapsuleRadius in size (matching SetCapsuleSize()).
- Remove InitCapsuleRadius() and InitCapsuleHalfHeight() because they can't correctly enforce the size bounds unless used in a specific order.

diff --git a/Engine/Source/Runtime/Engine/Classes/Components/CapsuleComponent.h b/Engine/Source/Runtime/Engine/Classes/Components/CapsuleComponent.h
--- a/Engine/Source/Runtime/Engine/Classes/Components/CapsuleComponent.h
+++ b/Engine/Source/Runtime/Engine/Classes/Components/CapsuleComponent.h
@@ -105,18 +105,12 @@ public:
 	UFUNCTION(BlueprintCallable, Category="Components|Capsule")
 	float GetShapeScale() const;
 
-	// Sets the capsule size without triggering a render or physics update.
+	// Sets the capsule size without triggering a render or physics update. This is the preferred method when initializing a component in a class constructor.
 	FORCEINLINE void InitCapsuleSize(float InRadius, float InHalfHeight)
 	{
 		CapsuleRadius = InRadius;
-		CapsuleHalfHeight = InHalfHeight;
+		CapsuleHalfHeight = FMath::Max(InHalfHeight, InRadius);
 	}
-
-	// Sets the capsule radius without triggering a render or physics update.
-	FORCEINLINE void InitCapsuleRadius(float InRadius) { CapsuleRadius = InRadius; }
-
-	// Sets the capsule half-height without triggering a render or physics update.
-	FORCEINLINE void InitCapsuleHalfHeight(float InHalfHeight) { CapsuleHalfHeight = InHalfHeight; }
 };
 
 
diff --git a/Engine/Source/Runtime/Engine/Private/PlayerStart.cpp b/Engine/Source/Runtime/Engine/Private/PlayerStart.cpp
--- a/Engine/Source/Runtime/Engine/Private/PlayerStart.cpp
+++ b/Engine/Source/Runtime/Engine/Private/PlayerStart.cpp
@@ -5,8 +5,7 @@
 APlayerStart::APlayerStart(const class FPostConstructInitializeProperties& PCIP)
 	: Super(PCIP)
 {
-	CapsuleComponent->InitCapsuleRadius(40.0f);
-	CapsuleComponent->InitCapsuleHalfHeight(92.0f);
+	CapsuleComponent->InitCapsuleSize(40.0f, 92.0f);
 
 #if WITH_EDITORONLY_DATA
 	ArrowComponent = PCIP.CreateEditorOnlyDefaultSubobject<UArrowComponent>(this, TEXT("Arrow"));
