Reverting CL 2042649 (was causing stretching decal issues in Orion art demo).

diff --git a/Engine/Source/Runtime/Engine/Private/ActorConstruction.cpp b/Engine/Source/Runtime/Engine/Private/ActorConstruction.cpp
--- a/Engine/Source/Runtime/Engine/Private/ActorConstruction.cpp
+++ b/Engine/Source/Runtime/Engine/Private/ActorConstruction.cpp
@@ -423,12 +423,14 @@ UActorComponent* AActor::AddComponent(FName TemplateName, bool bManualAttachment
 
 		if(NewSceneComp != NULL)
 		{
-			// if we end up deferring this component's registration below, then we 
-			// need to provide that deferred action the component's original mobility
-			// (because in the meantime, FDeferRegisterComponents alters the mobility
-			// to Movable so that construction-script changes won't be rejected due to
-			// a EComponentMobility::Static mobility)
+			// Components with Mobility set to EComponentMobility::Static or EComponentMobility::Stationary can't be properly set up in UCS (all changes will be rejected
+			// due to EComponentMobility::Static flag) so we're going to temporarily change the flag and defer the registration until UCS has finished.
+			bDeferRegisterStaticComponent = bRunningUserConstructionScript && NewSceneComp->Mobility != EComponentMobility::Movable;
 			OriginalMobility = NewSceneComp->Mobility;
+			if (bDeferRegisterStaticComponent)
+			{
+				NewSceneComp->Mobility = EComponentMobility::Movable;
+			}
 
 			if (!bManualAttachment)
 			{
@@ -448,9 +450,7 @@ UActorComponent* AActor::AddComponent(FName TemplateName, bool bManualAttachment
 		// Call function to notify component it has been created
 		NewActorComp->OnComponentCreated();
 
-		// need to defer component registration until after the construction script
-		// is ran, since the construction script can mutate the object (for collision, etc.)
-		if (bRunningUserConstructionScript && (NewSceneComp != NULL))
+		if (bDeferRegisterStaticComponent)
 		{
 			// Defer registration until after UCS has completed.
 			FDeferRegisterComponents::Get().DeferComponentRegistration(this, NewSceneComp, OriginalMobility);
diff --git a/Engine/Source/Runtime/Engine/Private/SCS_Node.cpp b/Engine/Source/Runtime/Engine/Private/SCS_Node.cpp
--- a/Engine/Source/Runtime/Engine/Private/SCS_Node.cpp
+++ b/Engine/Source/Runtime/Engine/Private/SCS_Node.cpp
@@ -44,12 +44,14 @@ void USCS_Node::ExecuteNodeOnActor(AActor* Actor, USceneComponent* ParentCompone
 		USceneComponent* NewSceneComp = Cast<USceneComponent>(NewActorComp);
 		if (NewSceneComp != NULL)
 		{
-			// since we end up deferring this component's registration below, we 
-			// need to provide that deferred action the component's original mobility
-			// (because in the meantime, FDeferRegisterComponents alters the mobility
-			// to Movable so that construction-script changes won't be rejected due to
-			// a EComponentMobility::Static mobility)
+			// Components with Mobility set to EComponentMobility::Static or EComponentMobility::Stationary can't be properly set up in SCS (all changes will be rejected
+			// due to EComponentMobility::Static flag) so we're going to temporarily change the flag and defer the registration until SCS has finished.
+			bDeferRegisterStaticComponent = NewSceneComp->Mobility != EComponentMobility::Movable;
 			OriginalMobility = NewSceneComp->Mobility;
+			if (bDeferRegisterStaticComponent)
+			{
+				NewSceneComp->Mobility = EComponentMobility::Movable;
+			}
 
 			// If NULL is passed in, we are the root, so set transform and assign as RootComponent on Actor
 			if (ParentComponent == NULL || (ParentComponent && ParentComponent->IsPendingKill()))
@@ -68,13 +70,13 @@ void USCS_Node::ExecuteNodeOnActor(AActor* Actor, USceneComponent* ParentCompone
 		// Call function to notify component it has been created
 		NewActorComp->OnComponentCreated();
 
-		if (NewSceneComp != NULL)
+		if (bDeferRegisterStaticComponent)
 		{
 			// need to defer component registration until after the construction script
 			// is ran, since the construction script can mutate the object (for collision, etc.)
 			FDeferRegisterComponents::Get().DeferComponentRegistration(Actor, NewSceneComp, OriginalMobility);
 		}
-		else 
+		else
 		{
 			NewActorComp->RegisterComponent();
 		}
