Increased max surfacetype to be 64

diff --git a/Engine/Source/Runtime/Engine/Classes/Engine/EngineTypes.h b/Engine/Source/Runtime/Engine/Classes/Engine/EngineTypes.h
--- a/Engine/Source/Runtime/Engine/Classes/Engine/EngineTypes.h
+++ b/Engine/Source/Runtime/Engine/Classes/Engine/EngineTypes.h
@@ -2280,6 +2280,38 @@ enum EPhysicalSurface
 	SurfaceType28 UMETA(Hidden),
 	SurfaceType29 UMETA(Hidden),
 	SurfaceType30 UMETA(Hidden),
+	SurfaceType31 UMETA(Hidden),
+	SurfaceType32 UMETA(Hidden),
+	SurfaceType33 UMETA(Hidden),
+	SurfaceType34 UMETA(Hidden),
+	SurfaceType35 UMETA(Hidden),
+	SurfaceType36 UMETA(Hidden),
+	SurfaceType37 UMETA(Hidden),
+	SurfaceType38 UMETA(Hidden),
+	SurfaceType39 UMETA(Hidden),
+	SurfaceType40 UMETA(Hidden),
+	SurfaceType41 UMETA(Hidden),
+	SurfaceType42 UMETA(Hidden),
+	SurfaceType43 UMETA(Hidden),
+	SurfaceType44 UMETA(Hidden),
+	SurfaceType45 UMETA(Hidden),
+	SurfaceType46 UMETA(Hidden),
+	SurfaceType47 UMETA(Hidden),
+	SurfaceType48 UMETA(Hidden),
+	SurfaceType49 UMETA(Hidden),
+	SurfaceType50 UMETA(Hidden),
+	SurfaceType51 UMETA(Hidden),
+	SurfaceType52 UMETA(Hidden),
+	SurfaceType53 UMETA(Hidden),
+	SurfaceType54 UMETA(Hidden),
+	SurfaceType55 UMETA(Hidden),
+	SurfaceType56 UMETA(Hidden),
+	SurfaceType57 UMETA(Hidden),
+	SurfaceType58 UMETA(Hidden),
+	SurfaceType59 UMETA(Hidden),
+	SurfaceType60 UMETA(Hidden),
+	SurfaceType61 UMETA(Hidden),
+	SurfaceType62 UMETA(Hidden),
 	SurfaceType_Max UMETA(Hidden)
 };
 
