#UE4 Fix it so BeginPlay doesn't get called twice on servers, I broke this yesterday

diff --git a/Engine/Source/Runtime/Engine/Private/GameState.cpp b/Engine/Source/Runtime/Engine/Private/GameState.cpp
--- a/Engine/Source/Runtime/Engine/Private/GameState.cpp
+++ b/Engine/Source/Runtime/Engine/Private/GameState.cpp
@@ -156,7 +156,11 @@ void AGameState::HandleMatchIsWaitingToStart()
 
 void AGameState::HandleMatchHasStarted()
 {
-	GetWorldSettings()->NotifyMatchStarted();
+	if (Role != ROLE_Authority)
+	{
+		// Server handles this in AGameMode::HandleMatchHasStarted
+		GetWorldSettings()->NotifyMatchStarted();
+	}
 }
 
 void AGameState::HandleMatchHasEnded()
