Add a GetGameViewport function to UWorld

diff --git a/Engine/Source/Runtime/Engine/Classes/Engine/Engine.h b/Engine/Source/Runtime/Engine/Classes/Engine/Engine.h
--- a/Engine/Source/Runtime/Engine/Classes/Engine/Engine.h
+++ b/Engine/Source/Runtime/Engine/Classes/Engine/Engine.h
@@ -2292,7 +2292,7 @@ public:
 
 	virtual bool WorldIsPIEInNewViewport(UWorld *InWorld);
 
-	FWorldContext* GetWorldContextFromWorld(UWorld * InWorld);
+	FWorldContext* GetWorldContextFromWorld(const UWorld* InWorld);
 	FWorldContext* GetWorldContextFromGameViewport(const UGameViewportClient *InViewport);
 	FWorldContext* GetWorldContextFromPendingNetGame(const UPendingNetGame *InPendingNetGame);	
 	FWorldContext* GetWorldContextFromPendingNetGameNetDriver(const UNetDriver *InPendingNetGame);	
diff --git a/Engine/Source/Runtime/Engine/Classes/Engine/World.h b/Engine/Source/Runtime/Engine/Classes/Engine/World.h
--- a/Engine/Source/Runtime/Engine/Classes/Engine/World.h
+++ b/Engine/Source/Runtime/Engine/Classes/Engine/World.h
@@ -1406,6 +1406,7 @@ public:
 	 */	
 	ULocalPlayer* GetFirstLocalPlayerFromController() const;
 	
+	UGameViewportClient* GetGameViewport() const;
 
 	/** Returns the default brush. */
 	ABrush* GetBrush() const;
diff --git a/Engine/Source/Runtime/Engine/Private/UnrealEngine.cpp b/Engine/Source/Runtime/Engine/Private/UnrealEngine.cpp
--- a/Engine/Source/Runtime/Engine/Private/UnrealEngine.cpp
+++ b/Engine/Source/Runtime/Engine/Private/UnrealEngine.cpp
@@ -10343,7 +10343,7 @@ FWorldContext& UEngine::GetWorldContextFromHandleChecked(int32 WorldContextHandl
 	return HandleInvalidWorldContext();
 }
 
-FWorldContext* UEngine::GetWorldContextFromWorld(UWorld *InWorld)
+FWorldContext* UEngine::GetWorldContextFromWorld(const UWorld* InWorld)
 {
 	for (FWorldContext& WorldContext : WorldList)
 	{
diff --git a/Engine/Source/Runtime/Engine/Private/World.cpp b/Engine/Source/Runtime/Engine/Private/World.cpp
--- a/Engine/Source/Runtime/Engine/Private/World.cpp
+++ b/Engine/Source/Runtime/Engine/Private/World.cpp
@@ -2709,6 +2709,12 @@ void UWorld::CleanupWorld(bool bSessionEnded, bool bCleanupResources)
 	}
 }
 
+UGameViewportClient* UWorld::GetGameViewport() const
+{
+	FWorldContext* WorldContext = GEngine->GetWorldContextFromWorld(this);
+	return (WorldContext ? WorldContext->GameViewport : NULL);
+}
+
 FConstControllerIterator UWorld::GetControllerIterator() const
 {
 	return ControllerList.CreateConstIterator();
