#UE4
Added missing comments on FABRIK node.
Renamed node from "Fabrik IK" to "FABRIK", since "FABRIK" is an acronym already containing "IK".

diff --git a/Engine/Source/Editor/AnimGraph/Private/AnimGraphNode_Fabrik.cpp b/Engine/Source/Editor/AnimGraph/Private/AnimGraphNode_Fabrik.cpp
--- a/Engine/Source/Editor/AnimGraph/Private/AnimGraphNode_Fabrik.cpp
+++ b/Engine/Source/Editor/AnimGraph/Private/AnimGraphNode_Fabrik.cpp
@@ -12,7 +12,7 @@ UAnimGraphNode_Fabrik::UAnimGraphNode_Fabrik(const FPostConstructInitializePrope
 
 FString UAnimGraphNode_Fabrik::GetControllerDescription() const
 {
-	return TEXT("Fabrik IK");
+	return TEXT("FABRIK");
 }
 
 FString UAnimGraphNode_Fabrik::GetNodeTitle(ENodeTitleType::Type TitleType) const
diff --git a/Engine/Source/Runtime/Engine/Classes/Animation/BoneControllers/AnimNode_Fabrik.h b/Engine/Source/Runtime/Engine/Classes/Animation/BoneControllers/AnimNode_Fabrik.h
--- a/Engine/Source/Runtime/Engine/Classes/Animation/BoneControllers/AnimNode_Fabrik.h
+++ b/Engine/Source/Runtime/Engine/Classes/Animation/BoneControllers/AnimNode_Fabrik.h
@@ -16,11 +16,20 @@ struct FABRIKChainLink
 	GENERATED_USTRUCT_BODY()
 
 public:
+	/** Position of bone in component space. */
 	FVector Position;
+
+	/** Distance to its parent link. */
 	float Length;
+
+	/** Bone Index in SkeletalMesh */
 	int32 BoneIndex;
+
+	/** Transform Index that this control will output */
 	int32 TransformIndex;
 
+	/** Child bones which are overlapping this bone. 
+	 * They have a zero length distance, so they will inherit this bone's transformation. */
 	UPROPERTY()
 	TArray<int32> ChildZeroLengthTransformIndices;
 
