- Added GetPlatformName BP node

diff --git a/Engine/Source/Runtime/Engine/Classes/Kismet/GameplayStatics.h b/Engine/Source/Runtime/Engine/Classes/Kismet/GameplayStatics.h
--- a/Engine/Source/Runtime/Engine/Classes/Kismet/GameplayStatics.h
+++ b/Engine/Source/Runtime/Engine/Classes/Kismet/GameplayStatics.h
@@ -442,6 +442,10 @@ class ENGINE_API UGameplayStatics : public UBlueprintFunctionLibrary
 	UFUNCTION(BlueprintCallable, Category="Game", meta=(Enable="true"))
 	static void EnableLiveStreaming(bool Enable);
 
+	/** Returns the string name of the current platform, to perform different behavior based on platorm */
+	UFUNCTION(BlueprintCallable, Category="Game")
+	static FString GetPlatformName();
+
 	/**
 	 * Calculates an launch velocity for a projectile to hit a specified point.
 	 * @param TossVelocity		(output) Result launch velocity.
@@ -460,5 +464,6 @@ class ENGINE_API UGameplayStatics : public UBlueprintFunctionLibrary
 
 	/** Native version, has more options than the Blueprint version. */
 	static bool SuggestProjectileVelocity(UObject* WorldContextObject, FVector& TossVelocity, FVector StartLocation, FVector EndLocation, float TossSpeed, bool bHighArc=false, float CollisionRadius = 0.f, float OverrideGravityZ = 0, ESuggestProjVelocityTraceOption::Type TraceOption = ESuggestProjVelocityTraceOption::TraceFullPath, bool bDrawDebug=false);
+
 };
 
diff --git a/Engine/Source/Runtime/Engine/Private/GameplayStatics.cpp b/Engine/Source/Runtime/Engine/Private/GameplayStatics.cpp
--- a/Engine/Source/Runtime/Engine/Private/GameplayStatics.cpp
+++ b/Engine/Source/Runtime/Engine/Private/GameplayStatics.cpp
@@ -1041,6 +1041,12 @@ void UGameplayStatics::EnableLiveStreaming(bool Enable)
 	}
 }
 
+FString UGameplayStatics::GetPlatformName()
+{
+	// the string that BP users care about is actually the platform name that we'd name the .ini file directory (Windows, not WindowsEditor)
+	return FPlatformProperties::IniPlatformName();
+}
+
 bool UGameplayStatics::BlueprintSuggestProjectileVelocity(UObject* WorldContextObject, FVector& OutTossVelocity, FVector StartLocation, FVector EndLocation, float LaunchSpeed, float OverrideGravityZ, ESuggestProjVelocityTraceOption::Type TraceOption, float CollisionRadius, bool bFavorHighArc, bool bDrawDebug)
 {
 	// simple pass-through to the C++ interface
