Integrate overflow bug fix from Fornite.

diff --git a/Engine/Source/Runtime/Engine/Private/NetConnection.cpp b/Engine/Source/Runtime/Engine/Private/NetConnection.cpp
--- a/Engine/Source/Runtime/Engine/Private/NetConnection.cpp
+++ b/Engine/Source/Runtime/Engine/Private/NetConnection.cpp
@@ -1084,7 +1084,8 @@ void UNetConnection::SendAck( int32 AckPacketId, bool FirstTime/*=1*/, bool bHav
 		
 		if (bPingAck)
 		{
-			SendSize += 1 + (bHavePingAckData ? sizeof(uint32) : 0);
+			const int32 PING_ACK_DATA_NUM_BITS = sizeof(uint32) * 8;
+			SendSize += 1 + (bHavePingAckData ? PING_ACK_DATA_NUM_BITS : 0);
 		}
 
 		PreSend(SendSize);
