#UE4 - removed no longer necessary boilerplate

diff --git a/Engine/Source/Runtime/Engine/Classes/GameFramework/OnlineReplStructs.h b/Engine/Source/Runtime/Engine/Classes/GameFramework/OnlineReplStructs.h
--- a/Engine/Source/Runtime/Engine/Classes/GameFramework/OnlineReplStructs.h
+++ b/Engine/Source/Runtime/Engine/Classes/GameFramework/OnlineReplStructs.h
@@ -137,13 +137,5 @@ struct TStructOpsTypeTraits<FUniqueNetIdRepl> : public TStructOpsTypeTraitsBase
 	};
 };
 
-/** Need this to have a .h file */
-UCLASS(Abstract, CustomConstructor)
-class UOnlineReplStructs: public UObject
-{
-	GENERATED_UCLASS_BODY()
-	UOnlineReplStructs(const class FPostConstructInitializeProperties& PCIP) : Super(PCIP) {}
-};
-
 /** Test harness for Unique Id replication */
 ENGINE_API void TestUniqueIdRepl(class UWorld* InWorld);
