- Change comment for density to read 'g per cubic cm' instead of 'g per cubic uu'

diff --git a/Engine/Source/Runtime/Engine/Classes/PhysicalMaterials/PhysicalMaterial.h b/Engine/Source/Runtime/Engine/Classes/PhysicalMaterials/PhysicalMaterial.h
--- a/Engine/Source/Runtime/Engine/Classes/PhysicalMaterials/PhysicalMaterial.h
+++ b/Engine/Source/Runtime/Engine/Classes/PhysicalMaterials/PhysicalMaterial.h
@@ -54,7 +54,7 @@ class ENGINE_API UPhysicalMaterial : public UObject
 	// Object properties.
 	//
 	
-	/** Used with the shape of the object to calculate its mass properties. The higher the number, the heavier the object. g per cubic uu. */
+	/** Used with the shape of the object to calculate its mass properties. The higher the number, the heavier the object. g per cubic cm. */
 	UPROPERTY(EditAnywhere, BlueprintReadWrite, Category=PhysicalMaterial)
 	float Density;
 
