Landscape : Fixed TTP #331854, performance issue in case of using upscaling

diff --git a/Engine/Source/Editor/UnrealEd/Private/Lightmass/Lightmass.cpp b/Engine/Source/Editor/UnrealEd/Private/Lightmass/Lightmass.cpp
--- a/Engine/Source/Editor/UnrealEd/Private/Lightmass/Lightmass.cpp
+++ b/Engine/Source/Editor/UnrealEd/Private/Lightmass/Lightmass.cpp
@@ -7,6 +7,7 @@
 #include "UnrealEd.h"
 #include "PrecomputedLightVolume.h"
 #include "Runtime/Engine/Public/Landscape/LandscapeRender.h"
+#include "Runtime/Engine/Public/Landscape/LandscapeLight.h"
 
 extern FSwarmDebugOptions GSwarmDebugOptions;
 
@@ -2254,6 +2255,9 @@ FLightmassProcessor::~FLightmassProcessor()
 		delete ImportData;
 	}
 	ImportedMappings.Empty();
+
+	FLandscapeStaticLightingMesh::LandscapeUpscaleHeightDataCache.Empty();
+	FLandscapeStaticLightingMesh::LandscapeUpscaleXYOffsetDataCache.Empty();
 }
 
 /** Retrieve an exporter for the given channel name */
diff --git a/Engine/Source/Runtime/Engine/Private/Landscape/LandscapeLight.cpp b/Engine/Source/Runtime/Engine/Private/Landscape/LandscapeLight.cpp
--- a/Engine/Source/Runtime/Engine/Private/Landscape/LandscapeLight.cpp
+++ b/Engine/Source/Runtime/Engine/Private/Landscape/LandscapeLight.cpp
@@ -13,6 +13,9 @@ LandscapeLight.cpp: Static lighting for LandscapeComponents
 
 #define LANDSCAPE_LIGHTMAP_UV_INDEX 1
 
+TMap<FIntPoint, FColor> FLandscapeStaticLightingMesh::LandscapeUpscaleHeightDataCache;
+TMap<FIntPoint, FColor> FLandscapeStaticLightingMesh::LandscapeUpscaleXYOffsetDataCache;
+
 /** A texture mapping for landscapes */
 /** Initialization constructor. */
 FLandscapeStaticLightingTextureMapping::FLandscapeStaticLightingTextureMapping(
@@ -240,93 +243,115 @@ namespace
 			{
 				for (int32 X = DataInterface.HeightmapComponentOffsetX; X < DataInterface.HeightmapComponentOffsetX + ComponentSize; ++X)
 				{
-					// LOD System same as shader
-					FVector2D XY(float(X - DataInterface.HeightmapComponentOffsetX) / (ComponentSize - 1), float(Y - DataInterface.HeightmapComponentOffsetY) / (ComponentSize - 1));
-					XY = XY - 0.5f;
-
-					float Dist = XY.Size();
+					FIntPoint IXY(X - DataInterface.HeightmapComponentOffsetX, Y - DataInterface.HeightmapComponentOffsetY);
+					IXY += ComponentBase * (ComponentSize - 1);
 
-					float RealLOD = GeometryLOD;
+					FColor* CachedHeight = FLandscapeStaticLightingMesh::LandscapeUpscaleHeightDataCache.Find(IXY);
+					FColor* CachedXYOffset = FLandscapeStaticLightingMesh::LandscapeUpscaleXYOffsetDataCache.Find(IXY);
 
-					if (XY.X < 0.f)
+					if (CachedHeight)
 					{
-						if (XY.Y < 0.f)
-						{
-							RealLOD = FMath::Lerp(
-								FMath::Lerp<float>(NeighborLODs[0], NeighborLODs[1], XY.X + 1.f),
-								FMath::Lerp<float>(NeighborLODs[3], GeometryLOD, XY.X + 1.f),
-								XY.Y + 1.f); // 0
-						}
-						else
+						CompHeightData[X + Y * HeightmapStride] = *CachedHeight;
+						if (CachedXYOffset)
 						{
-							RealLOD = FMath::Lerp(
-								FMath::Lerp<float>(NeighborLODs[3], GeometryLOD, XY.X + 1.f),
-								FMath::Lerp<float>(NeighborLODs[5], NeighborLODs[6], XY.X + 1.f),
-								XY.Y); // 2
+							CompXYOffsetData[X + Y * HeightmapStride] = *CachedXYOffset;
 						}
 					}
 					else
 					{
-						if (XY.Y < 0.f)
+						// LOD System similar to the shader
+						FVector2D XY(float(X - DataInterface.HeightmapComponentOffsetX) / (ComponentSize - 1), float(Y - DataInterface.HeightmapComponentOffsetY) / (ComponentSize - 1));
+						XY = XY - 0.5f;
+
+						float RealLOD = GeometryLOD;
+
+						if (XY.X < 0.f)
 						{
-							RealLOD = FMath::Lerp(
-								FMath::Lerp<float>(NeighborLODs[1], NeighborLODs[2], XY.X),
-								FMath::Lerp<float>(GeometryLOD, NeighborLODs[4], XY.X),
-								XY.Y + 1.f); // 1
+							if (XY.Y < 0.f)
+							{
+								RealLOD = FMath::Lerp(
+									FMath::Lerp<float>(NeighborLODs[0], NeighborLODs[1], XY.X + 1.f),
+									FMath::Lerp<float>(NeighborLODs[3], GeometryLOD, XY.X + 1.f),
+									XY.Y + 1.f); // 0
+							}
+							else
+							{
+								RealLOD = FMath::Lerp(
+									FMath::Lerp<float>(NeighborLODs[3], GeometryLOD, XY.X + 1.f),
+									FMath::Lerp<float>(NeighborLODs[5], NeighborLODs[6], XY.X + 1.f),
+									XY.Y); // 2
+							}
 						}
 						else
 						{
-							RealLOD = FMath::Lerp(
-								FMath::Lerp<float>(GeometryLOD, NeighborLODs[4], XY.X),
-								FMath::Lerp<float>(NeighborLODs[6], NeighborLODs[7], XY.X),
-								XY.Y); // 3
+							if (XY.Y < 0.f)
+							{
+								RealLOD = FMath::Lerp(
+									FMath::Lerp<float>(NeighborLODs[1], NeighborLODs[2], XY.X),
+									FMath::Lerp<float>(GeometryLOD, NeighborLODs[4], XY.X),
+									XY.Y + 1.f); // 1
+							}
+							else
+							{
+								RealLOD = FMath::Lerp(
+									FMath::Lerp<float>(GeometryLOD, NeighborLODs[4], XY.X),
+									FMath::Lerp<float>(NeighborLODs[6], NeighborLODs[7], XY.X),
+									XY.Y); // 3
+							}
 						}
-					}
 
-					RealLOD = FMath::Min(RealLOD, (float)MaxLOD);
+						RealLOD = FMath::Min(RealLOD, (float)MaxLOD);
 
-					int32 LODValue = (int32)RealLOD;
-					float MorphAlpha = FMath::Fractional(RealLOD);
+						int32 LODValue = (int32)RealLOD;
+						float MorphAlpha = FMath::Fractional(RealLOD);
 
-					FColor Height[2];
-					FColor XYOffset[2];
-					::GetLODData(LandscapeComponent, X, Y, DataInterface.HeightmapComponentOffsetX, DataInterface.HeightmapComponentOffsetY,
-						FMath::Min(MaxLOD, LODValue), HeightmapStride, Height[0], XYOffset[0]);
-
-					// Interpolation between two LOD
-					if ((RealLOD > InLOD) && (LODValue + 1 <= MaxLOD) && MorphAlpha != 0.f)
-					{
+						FColor Height[2];
+						FColor XYOffset[2];
 						::GetLODData(LandscapeComponent, X, Y, DataInterface.HeightmapComponentOffsetX, DataInterface.HeightmapComponentOffsetY,
-							FMath::Min(MaxLOD, LODValue + 1), HeightmapStride, Height[1], XYOffset[1]);
+							FMath::Min(MaxLOD, LODValue), HeightmapStride, Height[0], XYOffset[0]);
 
-						// Need interpolation
-						uint16 Height0 = (Height[0].R << 8) + Height[0].G;
-						uint16 Height1 = (Height[1].R << 8) + Height[1].G;
-						uint16 LerpHeight = FMath::Round(FMath::Lerp<float>(Height0, Height1, MorphAlpha));
+						// Interpolation between two LOD
+						if ((RealLOD > InLOD) && (LODValue + 1 <= MaxLOD) && MorphAlpha != 0.f)
+						{
+							::GetLODData(LandscapeComponent, X, Y, DataInterface.HeightmapComponentOffsetX, DataInterface.HeightmapComponentOffsetY,
+								FMath::Min(MaxLOD, LODValue + 1), HeightmapStride, Height[1], XYOffset[1]);
+
+							// Need interpolation
+							uint16 Height0 = (Height[0].R << 8) + Height[0].G;
+							uint16 Height1 = (Height[1].R << 8) + Height[1].G;
+							uint16 LerpHeight = FMath::Round(FMath::Lerp<float>(Height0, Height1, MorphAlpha));
+
+							CompHeightData[X + Y * HeightmapStride] =
+								FColor((LerpHeight >> 8), LerpHeight & 255,
+								FMath::Round(FMath::Lerp<float>(Height[0].B, Height[1].B, MorphAlpha)),
+								FMath::Round(FMath::Lerp<float>(Height[0].A, Height[1].A, MorphAlpha)));
+							if (LandscapeComponent->XYOffsetmapTexture)
+							{
+								uint16 XComp0 = (XYOffset[0].R << 8) + XYOffset[0].G;
+								uint16 XComp1 = (XYOffset[1].R << 8) + XYOffset[1].G;
+								uint16 LerpXComp = FMath::Round(FMath::Lerp<float>(XComp0, XComp1, MorphAlpha));
+
+								uint16 YComp0 = (XYOffset[0].B << 8) + XYOffset[0].A;
+								uint16 YComp1 = (XYOffset[1].B << 8) + XYOffset[1].A;
+								uint16 LerpYComp = FMath::Round(FMath::Lerp<float>(YComp0, YComp1, MorphAlpha));
+
+								CompXYOffsetData[X + Y * HeightmapStride] =
+									FColor(LerpXComp >> 8, LerpXComp & 255, LerpYComp >> 8, LerpYComp & 255);
+							}
+						}
+						else
+						{
+							CompHeightData[X + Y * HeightmapStride] = Height[0];
+							CompXYOffsetData[X + Y * HeightmapStride] = XYOffset[0];
+						}
 
-						CompHeightData[X + Y * HeightmapStride] = 
-							FColor((LerpHeight >> 8), LerpHeight & 255, 
-							FMath::Round(FMath::Lerp<float>(Height[0].B, Height[1].B, MorphAlpha)), 
-							FMath::Round(FMath::Lerp<float>(Height[0].A, Height[1].A, MorphAlpha)));
+						// Caching current calculated value
+						FLandscapeStaticLightingMesh::LandscapeUpscaleHeightDataCache.Add(IXY, CompHeightData[X + Y * HeightmapStride]);
 						if (LandscapeComponent->XYOffsetmapTexture)
 						{
-							uint16 XComp0 = (XYOffset[0].R << 8) + XYOffset[0].G;
-							uint16 XComp1 = (XYOffset[1].R << 8) + XYOffset[1].G;
-							uint16 LerpXComp = FMath::Round(FMath::Lerp<float>(XComp0, XComp1, MorphAlpha));
-
-							uint16 YComp0 = (XYOffset[0].B << 8) + XYOffset[0].A;
-							uint16 YComp1 = (XYOffset[1].B << 8) + XYOffset[1].A;
-							uint16 LerpYComp = FMath::Round(FMath::Lerp<float>(YComp0, YComp1, MorphAlpha));
-
-							CompXYOffsetData[X + Y * HeightmapStride] =
-								FColor(LerpXComp >> 8, LerpXComp & 255, LerpYComp >> 8, LerpYComp & 255);
+							FLandscapeStaticLightingMesh::LandscapeUpscaleHeightDataCache.Add(IXY, CompXYOffsetData[X + Y * HeightmapStride]);
 						}
 					}
-					else
-					{
-						CompHeightData[X + Y * HeightmapStride] = Height[0];
-						CompXYOffsetData[X + Y * HeightmapStride] = XYOffset[0];
-					}
 				}
 			}
 
@@ -397,8 +422,8 @@ void FLandscapeStaticLightingMesh::GetHeightmapData(int32 InLOD, int32 GeometryL
 			const int32 YSource = (ComponentY == 0) ? (ComponentSizeQuads - ExpandQuadsY) : ((ComponentY == 1) ? 0 : 1);
 			const int32 XDest = (ComponentX == 0) ? 0 : ((ComponentX == 1) ? ExpandQuadsX : (ComponentSizeQuads + ExpandQuadsX + 1));
 			const int32 YDest = (ComponentY == 0) ? 0 : ((ComponentY == 1) ? ExpandQuadsY : (ComponentSizeQuads + ExpandQuadsY + 1));
-			const int32 XNum = (ComponentX == 1) ? (ComponentSizeQuads + 1) : ExpandQuadsX; // ((ComponentX == 0 && ComponentY == 1) ? ExpandQuadsX + 1 : ExpandQuadsX);
-			const int32 YNum = (ComponentY == 1) ? (ComponentSizeQuads + 1) : ExpandQuadsY; // ((ComponentX == 1 && ComponentY == 0) ? ExpandQuadsY + 1 : ExpandQuadsY);
+			const int32 XNum = (ComponentX == 1) ? (ComponentSizeQuads + 1) : ExpandQuadsX;
+			const int32 YNum = (ComponentY == 1) ? (ComponentSizeQuads + 1) : ExpandQuadsY;
 			const int32 XBackup = (ComponentX == 2) ? (ComponentSizeQuads + ExpandQuadsX) : ExpandQuadsX;
 			const int32 YBackup = (ComponentY == 2) ? (ComponentSizeQuads + ExpandQuadsY) : ExpandQuadsY;
 			const int32 XBackupNum = (ComponentX == 1) ? (ComponentSizeQuads + 1) : 1;
diff --git a/Engine/Source/Runtime/Engine/Public/Landscape/LandscapeLight.h b/Engine/Source/Runtime/Engine/Public/Landscape/LandscapeLight.h
--- a/Engine/Source/Runtime/Engine/Public/Landscape/LandscapeLight.h
+++ b/Engine/Source/Runtime/Engine/Public/Landscape/LandscapeLight.h
@@ -80,7 +80,13 @@ protected:
 	float UVFactor;
 
 	friend class FLightmassExporter;
-};
 
+#if WITH_EDITOR
+public:
+	// Cache data for Landscape upscaling data
+	ENGINE_API static TMap<FIntPoint, FColor> LandscapeUpscaleHeightDataCache;
+	ENGINE_API static TMap<FIntPoint, FColor> LandscapeUpscaleXYOffsetDataCache;
+#endif
+};
 
 #endif // _LANDSCAPELIGHT_H
