#ttp 316728 - UE4: RENDERING: SHADOW: Kerning issues with text render components
#proj Engine
#branch UE4

#summary Fixed several issues with text render component string processing

#change Correct issues with the iterator Peek function, which was incorrectly returning false on the penultimate character of the string
#change FMath::Max doesn't do a per-component Max on two FVector objects, it does a length comparison. Change text size calculation to do a Max on both components of the size instead.

diff --git a/Engine/Source/Runtime/Engine/Private/TextRenderComponent.cpp b/Engine/Source/Runtime/Engine/Private/TextRenderComponent.cpp
--- a/Engine/Source/Runtime/Engine/Private/TextRenderComponent.cpp
+++ b/Engine/Source/Runtime/Engine/Private/TextRenderComponent.cpp
@@ -81,13 +81,13 @@ struct FTextIterator
 	bool Peek(int32& Ch)
 	{
 		check(CurrentPosition);
-		if ( CurrentPosition[0] =='\0' || CurrentPosition[1] == '\0' || (CurrentPosition[1] == '<' && CurrentPosition[2] == 'b' && CurrentPosition[3] == 'r' && CurrentPosition[4] == '>'))
+		if ( CurrentPosition[0] =='\0' || (CurrentPosition[0] == '<' && CurrentPosition[1] == 'b' && CurrentPosition[2] == 'r' && CurrentPosition[3] == '>'))
 		{
 			return false;
 		}
 		else
 		{
-			Ch = CurrentPosition[1];
+			Ch = CurrentPosition[0];
 			return true;
 		}
 	}
@@ -211,7 +211,8 @@ FVector2D ComputeTextSize(FTextIterator It, class UFont* Font,
 			float Right = X + SizeX;
 			float Bottom = Y + SizeY;
 
-			Ret = FMath::Max(Ret, FVector2D(Right, Bottom));
+			Ret.X = FMath::Max(Ret.X, Right);
+			Ret.Y = FMath::Max(Ret.Y, Bottom);
 
 			// if we have another non-whitespace character to render, add the font's kerning.
 			int32 NextCh;
