#UE4 - removed some DependsOn clauses

diff --git a/Engine/Source/Runtime/Engine/Classes/GameFramework/PlayerController.h b/Engine/Source/Runtime/Engine/Classes/GameFramework/PlayerController.h
--- a/Engine/Source/Runtime/Engine/Classes/GameFramework/PlayerController.h
+++ b/Engine/Source/Runtime/Engine/Classes/GameFramework/PlayerController.h
@@ -9,6 +9,7 @@
 #include "../Engine/GameViewportClient.h"
 #include "SceneManagement.h"
 #include "ForceFeedbackEffect.h"
+#include "OnlineReplStructs.h"
 #include "PlayerController.generated.h"
 
 /** delegate used to override default viewport audio listener position calculated from camera */
@@ -26,7 +27,7 @@ DECLARE_DELEGATE_ThreeParams(FGetAudioListenerPos, FVector& /*Location*/, FVecto
  */
 //=============================================================================
 
-UCLASS(config=Game, BlueprintType, Blueprintable, DependsOn=UOnlineReplStructs)
+UCLASS(config=Game, BlueprintType, Blueprintable)
 class ENGINE_API APlayerController : public AController
 {
 	GENERATED_UCLASS_BODY()
diff --git a/Engine/Source/Runtime/Engine/Classes/GameFramework/PlayerState.h b/Engine/Source/Runtime/Engine/Classes/GameFramework/PlayerState.h
--- a/Engine/Source/Runtime/Engine/Classes/GameFramework/PlayerState.h
+++ b/Engine/Source/Runtime/Engine/Classes/GameFramework/PlayerState.h
@@ -1,6 +1,7 @@
 // Copyright 1998-2014 Epic Games, Inc. All Rights Reserved.
 
 #pragma once
+#include "OnlineReplStructs.h"
 #include "PlayerState.generated.h"
 
 /**
@@ -29,7 +30,7 @@ struct PingAvgData
 // such as playername, score, etc.
 //=============================================================================
 
-UCLASS(HeaderGroup=ReplicationInfo, BlueprintType, Blueprintable, notplaceable, DependsOn=UOnlineReplStructs)
+UCLASS(HeaderGroup=ReplicationInfo, BlueprintType, Blueprintable, notplaceable)
 class ENGINE_API APlayerState : public AInfo
 {
 	GENERATED_UCLASS_BODY()
