#summary 	Fixed one-frame discrepancy when playing back Matinee where animations did not get updated after camera cuts.
#ttp 329487 	UE4/Matinee/Animation/Orion/ Animations often do not render for the first 2-3 frames of a shot
#detail 	When Matinee takes control of skeletal meshes we now set the MeshComponentUpdateFlag so that camera cuts do not result in a one-frame delay in animation updates.

diff --git a/Engine/Source/Runtime/Engine/Classes/Animation/SkeletalMeshActor.h b/Engine/Source/Runtime/Engine/Classes/Animation/SkeletalMeshActor.h
--- a/Engine/Source/Runtime/Engine/Classes/Animation/SkeletalMeshActor.h
+++ b/Engine/Source/Runtime/Engine/Classes/Animation/SkeletalMeshActor.h
@@ -76,7 +76,10 @@ public:
 	// End IMatineeAnimInterface Interface
 
 private:
+	// Saved flag used to restore animation mode
 	EAnimationMode::Type	SavedAnimationMode;
+	// Saved flag used to restore mesh update mode
+	EMeshComponentUpdateFlag::Type	SavedMeshUpdateFlag;
 	// utility function to see if it can play animation or not
 	bool CanPlayAnimation();
 };
diff --git a/Engine/Source/Runtime/Engine/Private/SkeletalMesh.cpp b/Engine/Source/Runtime/Engine/Private/SkeletalMesh.cpp
--- a/Engine/Source/Runtime/Engine/Private/SkeletalMesh.cpp
+++ b/Engine/Source/Runtime/Engine/Private/SkeletalMesh.cpp
@@ -3472,6 +3472,8 @@ void ASkeletalMeshActor::PreviewBeginAnimControl(UInterpGroup* InInterpGroup)
 	{
 		SavedAnimationMode = SkeletalMeshComponent->GetAnimationMode();
 		SkeletalMeshComponent->SetAnimationMode(EAnimationMode::Type::AnimationSingleNode);
+		SavedMeshUpdateFlag = SkeletalMeshComponent->MeshComponentUpdateFlag;
+		SkeletalMeshComponent->MeshComponentUpdateFlag = EMeshComponentUpdateFlag::AlwaysTickPoseAndRefreshBones;
 	}
 }
 
@@ -3483,6 +3485,7 @@ void ASkeletalMeshActor::PreviewFinishAnimControl(UInterpGroup* InInterpGroup)
 		// if in editor, reset the Animations, makes easier for artist to see them visually and align them
 		// in game, we keep the last pose that matinee kept. If you'd like it to have animation, you'll need to have AnimTree or AnimGraph to handle correctly
 		SkeletalMeshComponent->SetAnimationMode(SavedAnimationMode);
+		SkeletalMeshComponent->MeshComponentUpdateFlag = SavedMeshUpdateFlag;
 
 		// Update space bases to reset it back to ref pose
 		SkeletalMeshComponent->RefreshBoneTransforms();
@@ -3641,6 +3644,8 @@ void ASkeletalMeshActor::BeginAnimControl(UInterpGroup* InInterpGroup)
 	{
 		SavedAnimationMode = SkeletalMeshComponent->GetAnimationMode();
 		SkeletalMeshComponent->SetAnimationMode(EAnimationMode::Type::AnimationSingleNode);
+		SavedMeshUpdateFlag = SkeletalMeshComponent->MeshComponentUpdateFlag;
+		SkeletalMeshComponent->MeshComponentUpdateFlag = EMeshComponentUpdateFlag::AlwaysTickPoseAndRefreshBones;
 	}
 }
 
@@ -3691,6 +3696,7 @@ void ASkeletalMeshActor::FinishAnimControl(UInterpGroup* InInterpGroup)
 		// if in editor, reset the Animations, makes easier for artist to see them visually and align them
 		// in game, we keep the last pose that matinee kept. If you'd like it to have animation, you'll need to have AnimTree or AnimGraph to handle correctly
 		SkeletalMeshComponent->SetAnimationMode(SavedAnimationMode);
+		SkeletalMeshComponent->MeshComponentUpdateFlag = SavedMeshUpdateFlag;
 	}
 }
 
