Refactored UTexture platform cache

diff --git a/Engine/Source/Runtime/Engine/Classes/Engine/Texture.h b/Engine/Source/Runtime/Engine/Classes/Engine/Texture.h
--- a/Engine/Source/Runtime/Engine/Classes/Engine/Texture.h
+++ b/Engine/Source/Runtime/Engine/Classes/Engine/Texture.h
@@ -278,8 +278,6 @@ struct FTexturePlatformData
 #if WITH_EDITORONLY_DATA
 	/** The key associated with this derived data. */
 	FString DerivedDataKey;
-	/** The next cached platform data in the list. */
-	TScopedPointer<FTexturePlatformData> NextCachedPlatformData;
 	/** Async cache task if one is outstanding. */
 	struct FTextureAsyncCacheDerivedDataTask* AsyncTask;
 #endif
@@ -313,7 +311,6 @@ struct FTexturePlatformData
 		uint32 InFlags
 		);
 	void FinishCache();
-	void MakeCurrent(class UTexture& InTexture, TScopedPointer<FTexturePlatformData>* PlatformDataLink);
 	ENGINE_API bool TryInlineMipData();
 	bool AreDerivedMipsAvailable() const;
 #endif
@@ -563,7 +560,11 @@ public:
 	 * Textures that use the derived data cache must override this function and
 	 * provide a pointer to the linked list of platform data.
 	 */
-	virtual TScopedPointer<FTexturePlatformData>* GetPlatformDataLink() { return NULL; }
+	virtual FTexturePlatformData** GetRunningPlatformData() { return NULL; }
+	virtual TMap<FString, FTexturePlatformData*>* GetCookedPlatformData() { return NULL; }
+
+	void CleanupCachedRunningPlatformData();
+	void CleanupCachedCookedPlatformData();
 
 	/**
 	 * Serializes cooked platform data.
@@ -574,7 +575,13 @@ public:
 	/**
 	 * Caches platform data for the texture.
 	 */
-	void CachePlatformData();
+	void CachePlatformData(bool bAsyncCache = false);
+
+	/**
+	 * Begins caching platform data in the background for the platform requested
+	 */
+	virtual void BeginCacheForCookedPlatformData(  const ITargetPlatform *TargetPlatform ) OVERRIDE;
+
 
 	/**
 	 * Begins caching platform data in the background.
@@ -596,6 +603,8 @@ public:
 	 */
 	void ForceRebuildPlatformData();
 
+	void UpdateCachedLODBias();
+
 	/**
 	 * Marks platform data as transient. This optionally removes persistent or cached data associated with the platform.
 	 */
diff --git a/Engine/Source/Runtime/Engine/Classes/Engine/Texture2D.h b/Engine/Source/Runtime/Engine/Classes/Engine/Texture2D.h
--- a/Engine/Source/Runtime/Engine/Classes/Engine/Texture2D.h
+++ b/Engine/Source/Runtime/Engine/Classes/Engine/Texture2D.h
@@ -122,9 +122,10 @@ public:
 	FGuid TextureFileCacheGuid_DEPRECATED;
 
 public:
-	/** The derived data for this texture. */
-	TScopedPointer<FTexturePlatformData> PlatformData;
-
+	/** The derived data for this texture on this platform. */
+	FTexturePlatformData *PlatformData;
+	/* cooked platform data for this texture */
+	TMap<FString, FTexturePlatformData*> CookedPlatformData;
 	/**
 	 * Thread-safe counter indicating the texture streaming state. The definitions below are mirrored in Texture.h.
 	 *
@@ -177,7 +178,8 @@ public:
 	virtual EMaterialValueType GetMaterialType() OVERRIDE { return MCT_Texture2D; }
 	virtual void UpdateResource() OVERRIDE;
 	virtual float GetAverageBrightness(bool bIgnoreTrueBlack, bool bUseGrayscale) OVERRIDE;
-	virtual TScopedPointer<FTexturePlatformData>* GetPlatformDataLink() OVERRIDE { return &PlatformData; }
+	virtual FTexturePlatformData** GetRunningPlatformData() OVERRIDE { return &PlatformData; }
+	virtual TMap<FString,FTexturePlatformData*>* GetCookedPlatformData() OVERRIDE { return &CookedPlatformData; }
 	// End UTexture interface.
 
 	/** Trivial accessors. */
diff --git a/Engine/Source/Runtime/Engine/Classes/Engine/TextureCube.h b/Engine/Source/Runtime/Engine/Classes/Engine/TextureCube.h
--- a/Engine/Source/Runtime/Engine/Classes/Engine/TextureCube.h
+++ b/Engine/Source/Runtime/Engine/Classes/Engine/TextureCube.h
@@ -11,7 +11,8 @@ class UTextureCube : public UTexture
 
 public:
 	/** Platform data. */
-	TScopedPointer<FTexturePlatformData> PlatformData;
+	FTexturePlatformData *PlatformData;
+	TMap<FString, FTexturePlatformData*> CookedPlatformData;
 
 	// Begin UObject interface.
 	virtual void Serialize(FArchive& Ar) OVERRIDE;
@@ -61,7 +62,8 @@ public:
 	virtual FTextureResource* CreateResource() OVERRIDE;
 	virtual void UpdateResource() OVERRIDE;
 	virtual EMaterialValueType GetMaterialType() OVERRIDE { return MCT_TextureCube; }
-	virtual TScopedPointer<FTexturePlatformData>* GetPlatformDataLink() OVERRIDE { return &PlatformData; }
+	virtual FTexturePlatformData** GetRunningPlatformData() OVERRIDE { return &PlatformData; }
+	virtual TMap<FString, FTexturePlatformData*> *GetCookedPlatformData() OVERRIDE { return &CookedPlatformData; }
 	// End UTexture interface
 	
 	/**
diff --git a/Engine/Source/Runtime/Engine/Private/Texture.cpp b/Engine/Source/Runtime/Engine/Private/Texture.cpp
--- a/Engine/Source/Runtime/Engine/Private/Texture.cpp
+++ b/Engine/Source/Runtime/Engine/Private/Texture.cpp
@@ -348,6 +348,10 @@ void UTexture::FinishDestroy()
 		delete Resource;
 		Resource = NULL;
 	}
+
+	CleanupCachedRunningPlatformData();
+	CleanupCachedCookedPlatformData();
+
 }
 
 void UTexture::PreSave()
diff --git a/Engine/Source/Runtime/Engine/Private/Texture2D.cpp b/Engine/Source/Runtime/Engine/Private/Texture2D.cpp
--- a/Engine/Source/Runtime/Engine/Private/Texture2D.cpp
+++ b/Engine/Source/Runtime/Engine/Private/Texture2D.cpp
@@ -490,16 +490,7 @@ void UTexture2D::CookerWillNeverCookAgain()
 {
 	Super::CookerWillNeverCookAgain();
 
-	TScopedPointer<FTexturePlatformData>* PlatformDataLinkPtr = GetPlatformDataLink();
-	if (!IsTemplate() && PlatformDataLinkPtr)
-	{
-		// Release the resource. This will block on the rendering thread if we actually created a resource which we shouldn't have.
-		// We could probably check(Resource == NULL) instead.
-		ReleaseResource();
-
-		// This will free any platform data currently allocated. It's actually a linked list of scoped pointers so the entire list will be destroyed.
-		*PlatformDataLinkPtr = NULL;
-	}
+	CleanupCachedCookedPlatformData();
 }
 void UTexture2D::PostLinkerChange()
 {
diff --git a/Engine/Source/Runtime/Engine/Private/TextureDerivedData.cpp b/Engine/Source/Runtime/Engine/Private/TextureDerivedData.cpp
--- a/Engine/Source/Runtime/Engine/Private/TextureDerivedData.cpp
+++ b/Engine/Source/Runtime/Engine/Private/TextureDerivedData.cpp
@@ -343,36 +343,43 @@ static void GetBuildSettingsForRunningPlatform(
 	FTextureBuildSettings& OutBuildSettings
 	)
 {
-	// Compress to whatever formats the active target platforms want
-	ITargetPlatformManagerModule* TPM = GetTargetPlatformManager();
-	if (TPM)
+	static bool bInitialized = false;
+	static TArray<FName> PlatformFormats;
+
+	if ( !bInitialized )
 	{
-		ITargetPlatform* CurrentPlatform = NULL;
-		const TArray<ITargetPlatform*>& Platforms = TPM->GetActiveTargetPlatforms();
-		
-		check(Platforms.Num());
+		bInitialized = true;
+		// Compress to whatever formats the active target platforms want
+		ITargetPlatformManagerModule* TPM = GetTargetPlatformManager();
+		if (TPM)
+		{
+			ITargetPlatform* CurrentPlatform = NULL;
+			const TArray<ITargetPlatform*>& Platforms = TPM->GetActiveTargetPlatforms();
 
-		CurrentPlatform = Platforms[0];
+			check(Platforms.Num());
 
-		for (int32 Index = 1; Index < Platforms.Num(); Index++)
-		{
-			if (Platforms[Index]->IsRunningPlatform())
+			CurrentPlatform = Platforms[0];
+
+			for (int32 Index = 1; Index < Platforms.Num(); Index++)
 			{
-				CurrentPlatform = Platforms[Index];
-				break;
+				if (Platforms[Index]->IsRunningPlatform())
+				{
+					CurrentPlatform = Platforms[Index];
+					break;
+				}
 			}
-		}
 
-		check(CurrentPlatform != NULL);
+			check(CurrentPlatform != NULL);
 
-		TArray<FName> PlatformFormats;
-		CurrentPlatform->GetTextureFormats(&Texture, PlatformFormats);
-
-		// Assume there is at least one format and the first one is what we want at runtime.
-		check(PlatformFormats.Num());
-		GetTextureBuildSettings(Texture, GSystemSettings.TextureLODSettings, OutBuildSettings);
-		OutBuildSettings.TextureFormatName = PlatformFormats[0];
+			CurrentPlatform->GetTextureFormats(&Texture, PlatformFormats);
+		}
 	}
+
+
+	// Assume there is at least one format and the first one is what we want at runtime.
+	check(PlatformFormats.Num());
+	GetTextureBuildSettings(Texture, GSystemSettings.TextureLODSettings, OutBuildSettings);
+	OutBuildSettings.TextureFormatName = PlatformFormats[0];
 }
 
 /**
@@ -967,42 +974,6 @@ void FTexturePlatformData::FinishCache()
 	}
 }
 
-void FTexturePlatformData::MakeCurrent(UTexture& InTexture, TScopedPointer<FTexturePlatformData>* InPlatformDataLink)
-{
-	TScopedPointer<FTexturePlatformData>& PlatformDataLink = *InPlatformDataLink;
-
-	// Update the linked list of derived data to put this at the head.
-	if (PlatformDataLink != this)
-	{
-		// Find the derived data from the list.
-		FTexturePlatformData* PlatformDataPtr = PlatformDataLink;
-		while (PlatformDataPtr && PlatformDataPtr->NextCachedPlatformData != this)
-		{
-			PlatformDataPtr = PlatformDataPtr->NextCachedPlatformData;
-		}
-
-		if (PlatformDataPtr && PlatformDataPtr->NextCachedPlatformData == this)
-		{
-			// Unlink.
-			NextCachedPlatformData.Swap(PlatformDataPtr->NextCachedPlatformData);
-			// Link at head.
-			NextCachedPlatformData.Swap(PlatformDataLink);
-		}
-		else
-		{
-			check(NextCachedPlatformData == NULL);
-			// Link at head.
-			NextCachedPlatformData.Swap(PlatformDataLink);
-			PlatformDataLink = this;
-		}
-		check(PlatformDataLink == this);
-	}
-
-	// Update cached LOD bias. This must be done because the size of the 
-	// texture isn't known until derived data has been cached!
-	InTexture.CachedCombinedLODBias = GSystemSettings.TextureLODSettings.CalculateLODBias(&InTexture);
-}
-
 typedef TArray<uint32, TInlineAllocator<MAX_TEXTURE_MIP_COUNT> > FAsyncMipHandles;
 
 /**
@@ -1346,13 +1317,32 @@ void UTexture2D::GetMipData(int32 FirstMipToLoad, void** OutMipData)
 	}
 }
 
+
+
+void UTexture::CleanupCachedCookedPlatformData()
+{
+	TMap<FString, FTexturePlatformData*> *CookedPlatformDataPtr = GetCookedPlatformData();
+
+	if ( CookedPlatformDataPtr )
+	{
+		TMap<FString, FTexturePlatformData*> &CookedPlatformData = *CookedPlatformDataPtr;
+
+		for ( auto It : CookedPlatformData )
+		{
+			delete It.Value;
+		}
+
+		CookedPlatformData.Empty();
+	}
+}
+
 #if WITH_EDITORONLY_DATA
-void UTexture::CachePlatformData()
+void UTexture::CachePlatformData(bool bAsyncCache)
 {
-	TScopedPointer<FTexturePlatformData>* PlatformDataLinkPtr = GetPlatformDataLink();
+	FTexturePlatformData** PlatformDataLinkPtr = GetRunningPlatformData();
 	if (PlatformDataLinkPtr)
 	{
-		TScopedPointer<FTexturePlatformData>& PlatformDataLink = *PlatformDataLinkPtr;
+		FTexturePlatformData*& PlatformDataLink = *PlatformDataLinkPtr;
 		if (Source.IsValid() && FApp::CanEverRender())
 		{
 			FString DerivedDataKey;
@@ -1371,7 +1361,7 @@ void UTexture::CachePlatformData()
 				{
 					PlatformDataLink = new FTexturePlatformData();
 				}
-				PlatformDataLink->Cache(*this, BuildSettings, ETextureCacheFlags::None);
+				PlatformDataLink->Cache(*this, BuildSettings, bAsyncCache ? ETextureCacheFlags::Async : ETextureCacheFlags::None);
 			}
 		}
 		else if (PlatformDataLink == NULL)
@@ -1380,73 +1370,67 @@ void UTexture::CachePlatformData()
 			PlatformDataLink = new FTexturePlatformData();
 		}
 
-		PlatformDataLink->MakeCurrent(*this, PlatformDataLinkPtr);
+		
+		UpdateCachedLODBias();
 	}
 }
 
+void UTexture::UpdateCachedLODBias()
+{
+	CachedCombinedLODBias = GSystemSettings.TextureLODSettings.CalculateLODBias( this );
+}
+
 void UTexture::BeginCachePlatformData()
 {
-	TScopedPointer<FTexturePlatformData>* PlatformDataLinkPtr = GetPlatformDataLink();
-	if (PlatformDataLinkPtr)
+	CachePlatformData(true);
+
+
+	ITargetPlatformManagerModule* TPM = GetTargetPlatformManager();
+	
+	if (TPM)
 	{
-		TScopedPointer<FTexturePlatformData>& PlatformDataLink = *PlatformDataLinkPtr;
-		check(PlatformDataLink == NULL);
+		const TArray<ITargetPlatform*>& Platforms = TPM->GetActiveTargetPlatforms();
+		for (int32 PlatformIndex = 0; PlatformIndex < Platforms.Num(); PlatformIndex++)
+		{
+			BeginCacheForCookedPlatformData( Platforms[PlatformIndex]);
+		}
+	}
+
 
+}
+
+void UTexture::BeginCacheForCookedPlatformData( const ITargetPlatform *TargetPlatform )
+{
+	TMap<FString, FTexturePlatformData*>* CookedPlatformDataPtr = GetCookedPlatformData();
+	if (CookedPlatformDataPtr)
+	{
+		TMap<FString,FTexturePlatformData*>& CookedPlatformData = *CookedPlatformDataPtr;
+		
 		// Make sure the pixel format enum has been cached.
 		UTexture::GetPixelFormatEnum();
 
-		// Retrieve formats to cache for all platforms.
-		bool bInlineMips = false;
+		// Retrieve formats to cache for targetplatform.
+		
+		TArray<FName> PlatformFormats;
+		
+
 		TArray<FTextureBuildSettings> BuildSettingsToCache;
-		ITargetPlatformManagerModule* TPM = GetTargetPlatformManager();
 
-		if (TPM)
+		FTextureBuildSettings BuildSettings;
+		GetTextureBuildSettings(*this, TargetPlatform->GetTextureLODSettings(), BuildSettings);
+		TargetPlatform->GetTextureFormats(this, PlatformFormats);
+		for (int32 FormatIndex = 0; FormatIndex < PlatformFormats.Num(); ++FormatIndex)
 		{
-			bInlineMips = TPM->RestrictFormatsToRuntimeOnly() == false;
-			const TArray<ITargetPlatform*>& Platforms = TPM->GetActiveTargetPlatforms();
-
-			for (int32 PlatformIndex = 0; PlatformIndex < Platforms.Num(); PlatformIndex++)
+			BuildSettings.TextureFormatName = PlatformFormats[FormatIndex];
+			if (BuildSettings.TextureFormatName != NAME_None)
 			{
-				FTextureBuildSettings BuildSettings;
-				TArray<FName> PlatformFormats;
-
-				ITargetPlatform* Platform = Platforms[PlatformIndex];
-				GetTextureBuildSettings(*this, Platform->GetTextureLODSettings(), BuildSettings);
-				Platform->GetTextureFormats(this, PlatformFormats);
-				for (int32 FormatIndex = 0; FormatIndex < PlatformFormats.Num(); ++FormatIndex)
-				{
-					BuildSettings.TextureFormatName = PlatformFormats[FormatIndex];
-					if (BuildSettings.TextureFormatName != NAME_None)
-					{
-						BuildSettingsToCache.Add(BuildSettings);
-					}
-				}
+				BuildSettingsToCache.Add(BuildSettings);
 			}
 		}
 
-		// Cull redundant settings by comparing derived data keys.
-		if (BuildSettingsToCache.Num() > 1)
-		{
-			TSet<FString> UniqueDerivedDataKeys;
-			for (int32 SettingsIndex = 0; SettingsIndex < BuildSettingsToCache.Num(); ++SettingsIndex)
-			{
-				FString DerivedDataKey;
-				bool bAlreadyInSet = false;
 
-				GetTextureDerivedDataKeySuffix(*this, BuildSettingsToCache[SettingsIndex], DerivedDataKey);
-				UniqueDerivedDataKeys.Add(DerivedDataKey, &bAlreadyInSet);
-				if (bAlreadyInSet)
-				{
-					BuildSettingsToCache.RemoveAtSwap(SettingsIndex--);
-				}
-			}
-		}
 
-		uint32 CacheFlags = ETextureCacheFlags::Async;
-		if (bInlineMips)
-		{
-			CacheFlags |= ETextureCacheFlags::InlineMips;
-		}
+		uint32 CacheFlags = ETextureCacheFlags::Async | ETextureCacheFlags::InlineMips;
 
 		// If source data is resident in memory then allow the texture to be built
 		// in a background thread.
@@ -1454,18 +1438,27 @@ void UTexture::BeginCachePlatformData()
 		{
 			CacheFlags |= ETextureCacheFlags::AllowAsyncBuild;
 		}
-	
+
+
+		// Cull redundant settings by comparing derived data keys.
 		for (int32 SettingsIndex = 0; SettingsIndex < BuildSettingsToCache.Num(); ++SettingsIndex)
 		{
-			TScopedPointer<FTexturePlatformData> PlatformDataToCache;
-			PlatformDataToCache = new FTexturePlatformData();
-			PlatformDataToCache->Cache(
-				*this,
-				BuildSettingsToCache[SettingsIndex],
-				CacheFlags
-				);
-			PlatformDataToCache->NextCachedPlatformData.Swap(PlatformDataLink);
-			PlatformDataLink.Swap(PlatformDataToCache);
+			FString DerivedDataKey;
+			GetTextureDerivedDataKeySuffix(*this, BuildSettingsToCache[SettingsIndex], DerivedDataKey);
+
+			FTexturePlatformData *PlatformData = CookedPlatformData.FindRef( DerivedDataKey );
+
+			if ( PlatformData == NULL )
+			{
+				FTexturePlatformData* PlatformDataToCache;
+				PlatformDataToCache = new FTexturePlatformData();
+				PlatformDataToCache->Cache(
+					*this,
+					BuildSettingsToCache[SettingsIndex],
+					CacheFlags
+					);
+				CookedPlatformData.Add( DerivedDataKey, PlatformDataToCache );
+			}
 		}
 	}
 }
@@ -1473,72 +1466,75 @@ void UTexture::BeginCachePlatformData()
 bool UTexture::IsAsyncCacheComplete()
 {
 	bool bComplete = true;
-	TScopedPointer<FTexturePlatformData>* PlatformDataLinkPtr = GetPlatformDataLink();
-	if (PlatformDataLinkPtr)
+	FTexturePlatformData** RunningPlatformDataPtr = GetRunningPlatformData();
+	if (RunningPlatformDataPtr)
 	{
-		FTexturePlatformData* PlatformDataPtr = *PlatformDataLinkPtr;
-		while (PlatformDataPtr && bComplete)
+		FTexturePlatformData* RunningPlatformData = *RunningPlatformDataPtr;
+		if (RunningPlatformData )
 		{
-			bComplete = (PlatformDataPtr->AsyncTask == NULL) || PlatformDataPtr->AsyncTask->IsWorkDone();
-			PlatformDataPtr = PlatformDataPtr->NextCachedPlatformData;
+			bComplete &= (RunningPlatformData->AsyncTask == NULL) || RunningPlatformData->AsyncTask->IsWorkDone();
 		}
 	}
+
+	
+
+	TMap<FString,FTexturePlatformData*>* CookedPlatformDataPtr = GetCookedPlatformData();
+	if (CookedPlatformDataPtr)
+	{
+		for ( auto It : *CookedPlatformDataPtr )
+		{
+			FTexturePlatformData* PlatformData = It.Value;
+			if (PlatformData)
+			{
+				bComplete &= (PlatformData->AsyncTask == NULL) || PlatformData->AsyncTask->IsWorkDone();
+			}
+		}
+	}
+
 	return bComplete;
 }
 	
 void UTexture::FinishCachePlatformData()
 {
-	TScopedPointer<FTexturePlatformData>* PlatformDataLinkPtr = GetPlatformDataLink();
-	if (PlatformDataLinkPtr)
+	FTexturePlatformData** RunningPlatformDataPtr = GetRunningPlatformData();
+	if (RunningPlatformDataPtr)
 	{
-		ITargetPlatformManagerModule* TPM = GetTargetPlatformManager();
-		if (TPM->RestrictFormatsToRuntimeOnly() && FApp::CanEverRender())
+		FTexturePlatformData*& RunningPlatformData = *RunningPlatformDataPtr;
+		
+		if ( FApp::CanEverRender() )
 		{
-			TScopedPointer<FTexturePlatformData>& PlatformDataLink = *PlatformDataLinkPtr;
-			FTexturePlatformData* PlatformDerivedData = PlatformDataLink;
-
-			if (PlatformDerivedData)
+			if ( RunningPlatformData == NULL )
 			{
-				// If derived data has been cached for more than one platform, we must find
-				// the current platform and move it to the head of the list.
-				if (PlatformDerivedData->NextCachedPlatformData)
-				{
-					// Get the derived data key for the current platform.
-					FString DerivedDataKey;
-					FTextureBuildSettings BuildSettings;
-					GetBuildSettingsForRunningPlatform(*this, BuildSettings);
-					GetTextureDerivedDataKey(*this, BuildSettings, DerivedDataKey);
-
-					// Traverse list to find derived data for this platform.
-					while (PlatformDerivedData &&
-						PlatformDerivedData->DerivedDataKey != DerivedDataKey)
-					{
-						PlatformDerivedData = PlatformDerivedData->NextCachedPlatformData;
-					}
-					check(PlatformDerivedData); // We didn't cache data for this platform?
-				}
-
-				// Ensure async caching/building has completed.
-				PlatformDerivedData->FinishCache();
-
-				// Make this the current derived data for this texture.
-				PlatformDerivedData->MakeCurrent(*this, PlatformDataLinkPtr);
+				// begin cache never called
+				CachePlatformData();
 			}
 			else
 			{
-				// BeginCachePlatformData was never called, go ahead and cache it now.
-				CachePlatformData();
+				// make sure async requests are finished
+				RunningPlatformData->FinishCache();
 			}
+
+#if DO_CHECK
+			FString DerivedDataKey;
+			FTextureBuildSettings BuildSettings;
+			GetBuildSettingsForRunningPlatform(*this, BuildSettings);
+			GetTextureDerivedDataKey(*this, BuildSettings, DerivedDataKey);
+
+			check( RunningPlatformData->DerivedDataKey == DerivedDataKey );
+#endif
 		}
 	}
+
+	UpdateCachedLODBias();
+
 }
 
 void UTexture::ForceRebuildPlatformData()
 {
-	TScopedPointer<FTexturePlatformData>* PlatformDataLinkPtr = GetPlatformDataLink();
+	FTexturePlatformData** PlatformDataLinkPtr = GetRunningPlatformData();
 	if (PlatformDataLinkPtr && *PlatformDataLinkPtr && FApp::CanEverRender())
 	{
-		TScopedPointer<FTexturePlatformData>& PlatformDataLink = *PlatformDataLinkPtr;
+		FTexturePlatformData *&PlatformDataLink = *PlatformDataLinkPtr;
 		FlushRenderingCommands();
 		FTextureBuildSettings BuildSettings;
 		GetBuildSettingsForRunningPlatform(*this, BuildSettings);
@@ -1550,15 +1546,37 @@ void UTexture::ForceRebuildPlatformData()
 	}
 }
 
+
+
 void UTexture::MarkPlatformDataTransient()
 {
-	TScopedPointer<FTexturePlatformData>* PlatformDataLinkPtr = GetPlatformDataLink();
-	if (PlatformDataLinkPtr)
+	FDerivedDataCacheInterface& DDC = GetDerivedDataCacheRef();
+
+	FTexturePlatformData** RunningPlatformData = GetRunningPlatformData();
+	if (RunningPlatformData)
+	{
+		FTexturePlatformData* PlatformData = *RunningPlatformData;
+		if (PlatformData)
+		{
+			for (int32 MipIndex = 0; MipIndex < PlatformData->Mips.Num(); ++MipIndex)
+			{
+				FTexture2DMipMap& Mip = PlatformData->Mips[MipIndex];
+				if (Mip.DerivedDataKey.Len() > 0)
+				{
+					DDC.MarkTransient(*Mip.DerivedDataKey);
+				}
+			}
+			DDC.MarkTransient(*PlatformData->DerivedDataKey);
+		}
+	}
+
+	TMap<FString, FTexturePlatformData*> *CookedPlatformData = GetCookedPlatformData();
+	if ( CookedPlatformData )
 	{
 		FDerivedDataCacheInterface& DDC = GetDerivedDataCacheRef();
-		FTexturePlatformData* PlatformData = *PlatformDataLinkPtr;
-		while (PlatformData)
+		for ( auto It : *CookedPlatformData )
 		{
+			FTexturePlatformData* PlatformData = It.Value;
 			for (int32 MipIndex = 0; MipIndex < PlatformData->Mips.Num(); ++MipIndex)
 			{
 				FTexture2DMipMap& Mip = PlatformData->Mips[MipIndex];
@@ -1568,23 +1586,37 @@ void UTexture::MarkPlatformDataTransient()
 				}
 			}
 			DDC.MarkTransient(*PlatformData->DerivedDataKey);
-			PlatformData = PlatformData->NextCachedPlatformData;
 		}
 	}
 }
 #endif // #if WITH_EDITORONLY_DATA
 
+void UTexture::CleanupCachedRunningPlatformData()
+{
+	FTexturePlatformData **RunningPlatformDataPtr = GetRunningPlatformData();
+
+	if ( RunningPlatformDataPtr )
+	{
+		FTexturePlatformData *&RunningPlatformData = *RunningPlatformDataPtr;
+
+		if ( RunningPlatformData != NULL )
+		{
+			delete RunningPlatformData;
+			RunningPlatformData = NULL;
+		}
+	}
+}
+
+
 void UTexture::SerializeCookedPlatformData(FArchive& Ar)
 {
-	TScopedPointer<FTexturePlatformData>* PlatformDataLinkPtr = GetPlatformDataLink();
-	if (IsTemplate() || PlatformDataLinkPtr == NULL)
+	if (IsTemplate() )
 	{
 		return;
 	}
 
-	TScopedPointer<FTexturePlatformData>& PlatformDataLink = *PlatformDataLinkPtr;
 	UEnum* PixelFormatEnum = UTexture::GetPixelFormatEnum();
-	FName PixelFormatName;
+
 
 #if WITH_EDITORONLY_DATA
 	TextureDerivedDataTimings::FScopedMeasurement Timer(TextureDerivedDataTimings::SerializeCookedTime);
@@ -1592,6 +1624,12 @@ void UTexture::SerializeCookedPlatformData(FArchive& Ar)
 	{
 		if (!Ar.CookingTarget()->IsServerOnly())
 		{
+			TMap<FString, FTexturePlatformData*> *CookedPlatformDataPtr = GetCookedPlatformData();
+			if ( CookedPlatformDataPtr == NULL )
+				return;
+
+
+
 			FTextureBuildSettings BuildSettings;
 			TArray<FName> PlatformFormats;
 			TArray<FTexturePlatformData*> PlatformDataToSerialize;
@@ -1604,40 +1642,26 @@ void UTexture::SerializeCookedPlatformData(FArchive& Ar)
 				BuildSettings.TextureFormatName = PlatformFormats[FormatIndex];
 				GetTextureDerivedDataKey(*this, BuildSettings, DerivedDataKey);
 
-				FTexturePlatformData* PlatformDataPtr = PlatformDataLink;
-				while (PlatformDataPtr)
-				{
-					if (PlatformDataPtr->DerivedDataKey == DerivedDataKey)
-					{
-						break;
-					}
-					PlatformDataPtr = PlatformDataPtr->NextCachedPlatformData;
-				}
-
+				FTexturePlatformData *PlatformDataPtr = (*CookedPlatformDataPtr).FindRef(DerivedDataKey);
+				
 				if (PlatformDataPtr == NULL)
 				{
-					TScopedPointer<FTexturePlatformData> NewPlatformData;
-					NewPlatformData = new FTexturePlatformData();
-					NewPlatformData->Cache(*this, BuildSettings, ETextureCacheFlags::InlineMips | ETextureCacheFlags::Async);
-					PlatformDataPtr = NewPlatformData;
-					if (PlatformDataLink == NULL)
-					{
-						PlatformDataLink.Swap(NewPlatformData);
-					}
-					else
-					{
-						NewPlatformData->NextCachedPlatformData.Swap(PlatformDataLink->NextCachedPlatformData);
-						PlatformDataLink->NextCachedPlatformData.Swap(NewPlatformData);
-					}
+					PlatformDataPtr = new FTexturePlatformData();
+					PlatformDataPtr->Cache(*this, BuildSettings, ETextureCacheFlags::InlineMips | ETextureCacheFlags::Async);
+					
+					CookedPlatformDataPtr->Add( DerivedDataKey, PlatformDataPtr );
+					
 				}
 				PlatformDataToSerialize.Add(PlatformDataPtr);
 			}
 
 			for (int32 i = 0; i < PlatformDataToSerialize.Num(); ++i)
 			{
+				
+
 				FTexturePlatformData* PlatformDataToSave = PlatformDataToSerialize[i];
 				PlatformDataToSave->FinishCache();
-				PixelFormatName = PixelFormatEnum->GetEnum(PlatformDataToSave->PixelFormat);
+				FName PixelFormatName = PixelFormatEnum->GetEnum(PlatformDataToSave->PixelFormat);
 				Ar << PixelFormatName;
 				int32 SkipOffsetLoc = Ar.Tell();
 				int32 SkipOffset = 0;
@@ -1649,13 +1673,25 @@ void UTexture::SerializeCookedPlatformData(FArchive& Ar)
 				Ar.Seek(SkipOffset);
 			}
 		}
-		PixelFormatName = NAME_None;
+		FName PixelFormatName = NAME_None;
 		Ar << PixelFormatName;
 	}
 	else
 #endif // #if WITH_EDITORONLY_DATA
 	{
-		PlatformDataLink = new FTexturePlatformData();
+
+		FTexturePlatformData** RunningPlatformDataPtr = GetRunningPlatformData();
+		if ( RunningPlatformDataPtr == NULL )
+			return;
+
+		FTexturePlatformData*& RunningPlatformData = *RunningPlatformDataPtr;
+
+		FName PixelFormatName = NAME_None;
+
+		CleanupCachedRunningPlatformData();
+		check( RunningPlatformData == NULL );
+
+		RunningPlatformData = new FTexturePlatformData();
 		Ar << PixelFormatName;
 		while (PixelFormatName != NAME_None)
 		{
@@ -1663,9 +1699,9 @@ void UTexture::SerializeCookedPlatformData(FArchive& Ar)
 			int32 SkipOffset = 0;
 			Ar << SkipOffset;
 			bool bFormatSupported = GPixelFormats[PixelFormat].Supported;
-			if (PlatformDataLink->PixelFormat == PF_Unknown && bFormatSupported)
+			if (RunningPlatformData->PixelFormat == PF_Unknown && bFormatSupported)
 			{
-				PlatformDataLink->SerializeCooked(Ar, this);
+				RunningPlatformData->SerializeCooked(Ar, this);
 			}
 			else
 			{
