Expose CreatePlayer to Blueprints

diff --git a/Engine/Source/Runtime/Engine/Classes/Kismet/GameplayStatics.h b/Engine/Source/Runtime/Engine/Classes/Kismet/GameplayStatics.h
--- a/Engine/Source/Runtime/Engine/Classes/Kismet/GameplayStatics.h
+++ b/Engine/Source/Runtime/Engine/Classes/Kismet/GameplayStatics.h
@@ -83,6 +83,13 @@ class ENGINE_API UGameplayStatics : public UBlueprintFunctionLibrary
 	UFUNCTION(BlueprintPure, Category="Game", meta=(HidePin="WorldContextObject", DefaultToSelf="WorldContextObject"))
 	static class APlayerCameraManager* GetPlayerCameraManager(UObject* WorldContextObject, int32 PlayerIndex);
 
+	/** Create a new player for this game.  
+	 *  @param ControllerId		The ID of the controller that the should control the newly created player.  A value of -1 specifies to use the next available ID
+	 *  @param bSpawnPawn		Whether a pawn should be spawned immediately. If false a pawn will not be created until transition to the next map.
+	 */
+	UFUNCTION(BlueprintCallable, Category="Game", meta=(HidePin="WorldContextObject", DefaultToSelf="WorldContextObject", ControllerId="-1", bSpawnPawn="true", AdvancedDisplay="2"))
+	static class APlayerController* CreatePlayer(UObject* WorldContextObject, int32 ControllerId, bool bSpawnPawn = true);
+
 	// --- Level Streaming functions ------------------------
 	
 	/** Stream the level with the LevelName ; Calling again before it finishes has no effect */
diff --git a/Engine/Source/Runtime/Engine/Private/GameViewportClient.cpp b/Engine/Source/Runtime/Engine/Private/GameViewportClient.cpp
--- a/Engine/Source/Runtime/Engine/Private/GameViewportClient.cpp
+++ b/Engine/Source/Runtime/Engine/Private/GameViewportClient.cpp
@@ -1121,7 +1121,24 @@ ULocalPlayer* UGameViewportClient::CreatePlayer(int32 ControllerId, FString& Out
 	}
 	else if (GetOuterUEngine()->GetNumGamePlayers(this) < MaxSplitscreenPlayers)
 	{
-		NewPlayer = CastChecked<ULocalPlayer>(StaticConstructObject(GetOuterUEngine()->LocalPlayerClass, GetOuter()));	
+		// If the controller ID is not specified then find the first available
+		if (ControllerId < 0)
+		{
+			for (ControllerId = 0; ControllerId < MaxSplitscreenPlayers; ++ControllerId)
+			{
+				if (GetOuterUEngine()->GetLocalPlayerFromControllerId(this, ControllerId) == NULL)
+				{
+					break;
+				}
+			}
+			check(ControllerId < MaxSplitscreenPlayers);
+		}
+		else if (ControllerId >= MaxSplitscreenPlayers)
+		{
+			UE_LOG(LogPlayerManagement, Warning, TEXT("Controller ID (%d) is unlikely to map to any physical device, so this player will not receive input"), ControllerId);
+		}
+
+		NewPlayer = CastChecked<ULocalPlayer>(StaticConstructObject(GetOuterUEngine()->LocalPlayerClass, GetOuter()));
 		InsertIndex = AddLocalPlayer(NewPlayer, ControllerId);
 		if ( bSpawnActor && InsertIndex != INDEX_NONE )
 		{
@@ -1205,7 +1222,7 @@ void UGameViewportClient::DebugCreatePlayer(int32 ControllerId)
 	CreatePlayer(ControllerId, Error, true);
 	if (Error.Len() > 0)
 	{
-		UE_LOG(LogPlayerManagement, Warning, TEXT("Failed to DebugCreatePlayer. Error: %s"), *Error);
+		UE_LOG(LogPlayerManagement, Error, TEXT("Failed to DebugCreatePlayer: %s"), *Error);
 	}
 #endif
 }
diff --git a/Engine/Source/Runtime/Engine/Private/GameplayStatics.cpp b/Engine/Source/Runtime/Engine/Private/GameplayStatics.cpp
--- a/Engine/Source/Runtime/Engine/Private/GameplayStatics.cpp
+++ b/Engine/Source/Runtime/Engine/Private/GameplayStatics.cpp
@@ -59,6 +59,21 @@ APlayerCameraManager* UGameplayStatics::GetPlayerCameraManager(UObject* WorldCon
 	return PC ? PC->PlayerCameraManager : NULL;
 }
 
+APlayerController* UGameplayStatics::CreatePlayer(UObject* WorldContextObject, int32 ControllerId, bool bSpawnPawn)
+{
+	UWorld* World = GEngine->GetWorldFromContextObject(WorldContextObject);
+	FString Error;
+
+	ULocalPlayer* LocalPlayer = World->GetGameViewport()->CreatePlayer(ControllerId, Error, bSpawnPawn);
+
+	if (Error.Len() > 0)
+	{
+		UE_LOG(LogPlayerManagement, Error, TEXT("Failed to Create Player: %s"), *Error);
+	}
+
+	return (LocalPlayer ? LocalPlayer->PlayerController : NULL);
+}
+
 AGameMode* UGameplayStatics::GetGameMode(UObject* WorldContextObject)
 {
 	UWorld* const World = GEngine->GetWorldFromContextObject(WorldContextObject);
