#summary Can search for nodes in the palette menu in Blueprints with both the native name and the localized name.

diff --git a/Engine/Source/Runtime/Engine/Classes/EdGraph/EdGraphSchema.h b/Engine/Source/Runtime/Engine/Classes/EdGraph/EdGraphSchema.h
--- a/Engine/Source/Runtime/Engine/Classes/EdGraph/EdGraphSchema.h
+++ b/Engine/Source/Runtime/Engine/Classes/EdGraph/EdGraphSchema.h
@@ -53,7 +53,7 @@ struct ENGINE_API FEdGraphSchemaAction
 
 	/** The menu text that should be displayed for this node in the creation menu */
 	UPROPERTY()
-	FString MenuDescription;
+	FText MenuDescription;
 
 	/** The tooltip text that should be displayed for this node in the creation menu */
 	UPROPERTY()
@@ -75,6 +75,10 @@ struct ENGINE_API FEdGraphSchemaAction
 	UPROPERTY()
 	int32 SectionID;
 
+	/** Search title for the action */
+	UPROPERTY()
+	FText SearchTitle;
+
 	FEdGraphSchemaAction() 
 		: Grouping(0)
 		, SectionID(0)
@@ -82,13 +86,14 @@ struct ENGINE_API FEdGraphSchemaAction
 	
 	virtual ~FEdGraphSchemaAction() {}
 
-	FEdGraphSchemaAction(const FString& InNodeCategory, const FString& InMenuDesc, const FString& InToolTip, const int32 InGrouping)
+	FEdGraphSchemaAction(const FString& InNodeCategory, const FText& InMenuDesc, const FString& InToolTip, const int32 InGrouping)
 		: MenuDescription(InMenuDesc)
 		, TooltipDescription(InToolTip)
 		, Category(InNodeCategory)
 		, Grouping(InGrouping)
 		, SectionID(0)
-	{}
+	{
+	}
 
 	/** Whether or not this action can be parented to other actions of the same type */
 	virtual bool IsParentable() const { return false; }
@@ -112,6 +117,22 @@ struct ENGINE_API FEdGraphSchemaAction
 		return NewNode;
 	}
 
+	/** Retrieves the full searchable title for this action */
+	FText GetSearchTitle()
+	{
+		if(SearchTitle.IsEmpty())
+		{
+			if(const FString* SourceString = FTextInspector::GetSourceString(MenuDescription))
+			{
+				FFormatNamedArguments Args;
+				Args.Add(TEXT("LocalizedTitle"), MenuDescription);
+				Args.Add(TEXT("SourceTitle"), FText::FromString(*SourceString));
+				SearchTitle = FText::Format(FText::FromString("{LocalizedTitle} {SourceTitle}"), Args);
+			}
+		}
+		return SearchTitle;
+	}
+
 	// GC.
 	virtual void AddReferencedObjects(FReferenceCollector& Collector) {}
 };
@@ -136,7 +157,7 @@ struct ENGINE_API FEdGraphSchemaAction_NewNode : public FEdGraphSchemaAction
 		, NodeTemplate(NULL)
 	{}
 
-	FEdGraphSchemaAction_NewNode(const FString& InNodeCategory, const FString& InMenuDesc, const FString& InToolTip, const int32 InGrouping)
+	FEdGraphSchemaAction_NewNode(const FString& InNodeCategory, const FText& InMenuDesc, const FString& InToolTip, const int32 InGrouping)
 		: FEdGraphSchemaAction(InNodeCategory, InMenuDesc, InToolTip, InGrouping) 
 		, NodeTemplate(NULL)
 	{}
@@ -169,7 +190,7 @@ struct FEdGraphSchemaAction_Dummy : public FEdGraphSchemaAction
 	: FEdGraphSchemaAction()
 	{}
 
-	FEdGraphSchemaAction_Dummy(const FString& InNodeCategory, const FString& InMenuDesc, const FString& InToolTip, const int32 InGrouping)
+	FEdGraphSchemaAction_Dummy(const FString& InNodeCategory, const FText& InMenuDesc, const FString& InToolTip, const int32 InGrouping)
 		: FEdGraphSchemaAction(InNodeCategory, InMenuDesc, InToolTip, InGrouping)		
 	{}
 };
@@ -349,6 +370,8 @@ public:
 struct FGraphDisplayInfo
 {
 public:
+	/** Plain name for this graph */
+	FText PlainName;
 	/** Friendly name to display for this graph */
 	FText DisplayName;
 	/** Text to show as tooltip for this graph */
diff --git a/Engine/Source/Runtime/Engine/Private/AI/BehaviorTree/Tasks/BTTask_RunEQSQuery.cpp b/Engine/Source/Runtime/Engine/Private/AI/BehaviorTree/Tasks/BTTask_RunEQSQuery.cpp
--- a/Engine/Source/Runtime/Engine/Private/AI/BehaviorTree/Tasks/BTTask_RunEQSQuery.cpp
+++ b/Engine/Source/Runtime/Engine/Private/AI/BehaviorTree/Tasks/BTTask_RunEQSQuery.cpp
@@ -103,7 +103,7 @@ void UBTTask_RunEQSQuery::OnQueryFinished(TSharedPtr<struct FEnvQueryResult> Res
 		if (!bSuccess)
 		{
 			UE_VLOG(MyOwner, LogBehaviorTree, Warning, TEXT("Failed to store query result! item:%s key:%s"),
-				*UEnvQueryTypes::GetShortTypeName(Result->ItemType),
+				*UEnvQueryTypes::GetShortTypeName(Result->ItemType).ToString(),
 				*UBehaviorTreeTypes::GetShortTypeName(BlackboardKey.SelectedKeyType));
 		}
 	}
diff --git a/Engine/Source/Runtime/Engine/Private/AI/EnvironmentQuery/EnvQueryGenerator.cpp b/Engine/Source/Runtime/Engine/Private/AI/EnvironmentQuery/EnvQueryGenerator.cpp
--- a/Engine/Source/Runtime/Engine/Private/AI/EnvironmentQuery/EnvQueryGenerator.cpp
+++ b/Engine/Source/Runtime/Engine/Private/AI/EnvironmentQuery/EnvQueryGenerator.cpp
@@ -6,14 +6,14 @@ UEnvQueryGenerator::UEnvQueryGenerator(const class FPostConstructInitializePrope
 {
 }
 
-FString UEnvQueryGenerator::GetDescriptionTitle() const
+FText UEnvQueryGenerator::GetDescriptionTitle() const
 {
 	return UEnvQueryTypes::GetShortTypeName(this);
 }
 
-FString UEnvQueryGenerator::GetDescriptionDetails() const
+FText UEnvQueryGenerator::GetDescriptionDetails() const
 {
-	return TEXT("");
+	return FText::GetEmpty();
 }
 
 #if WITH_EDITOR
diff --git a/Engine/Source/Runtime/Engine/Private/AI/EnvironmentQuery/EnvQueryInstance.cpp b/Engine/Source/Runtime/Engine/Private/AI/EnvironmentQuery/EnvQueryInstance.cpp
--- a/Engine/Source/Runtime/Engine/Private/AI/EnvironmentQuery/EnvQueryInstance.cpp
+++ b/Engine/Source/Runtime/Engine/Private/AI/EnvironmentQuery/EnvQueryInstance.cpp
@@ -50,9 +50,9 @@ bool FEnvQueryInstance::PrepareContext(UClass* ContextClass, FEnvQueryContextDat
 	if (ContextData.NumValues == 0)
 	{
 		UE_LOG(LogEQS, Log, TEXT("Query [%s] is missing values for context [%s], skipping test %d:%d [%s]"),
-			*QueryName, *UEnvQueryTypes::GetShortTypeName(ContextClass),
+			*QueryName, *UEnvQueryTypes::GetShortTypeName(ContextClass).ToString(),
 			OptionIndex, CurrentTest,
-			CurrentTest >=0 ? *UEnvQueryTypes::GetShortTypeName(Options[OptionIndex].TestDelegates[CurrentTest].GetUObject()) : TEXT("Generator")
+			CurrentTest >=0 ? *UEnvQueryTypes::GetShortTypeName(Options[OptionIndex].TestDelegates[CurrentTest].GetUObject()).ToString() : TEXT("Generator")
 			);
 
 		return false;
@@ -482,7 +482,7 @@ void FEnvQueryInstance::FinalizeTest()
 #if WITH_EDITOR
 	if (bStoreDebugInfo)
 	{
-		DebugData.PerformedTestNames.Add(UEnvQueryTypes::GetShortTypeName(DefTestOb));
+		DebugData.PerformedTestNames.Add(UEnvQueryTypes::GetShortTypeName(DefTestOb).ToString());
 	}
 #endif
 
diff --git a/Engine/Source/Runtime/Engine/Private/AI/EnvironmentQuery/EnvQueryOption.cpp b/Engine/Source/Runtime/Engine/Private/AI/EnvironmentQuery/EnvQueryOption.cpp
--- a/Engine/Source/Runtime/Engine/Private/AI/EnvironmentQuery/EnvQueryOption.cpp
+++ b/Engine/Source/Runtime/Engine/Private/AI/EnvironmentQuery/EnvQueryOption.cpp
@@ -6,12 +6,12 @@ UEnvQueryOption::UEnvQueryOption(const class FPostConstructInitializeProperties&
 {
 }
 
-FString UEnvQueryOption::GetDescriptionTitle() const
+FText UEnvQueryOption::GetDescriptionTitle() const
 {
-	return Generator ? Generator->GetDescriptionTitle() : FString();
+	return Generator ? Generator->GetDescriptionTitle() : FText::GetEmpty();
 }
 
-FString UEnvQueryOption::GetDescriptionDetails() const
+FText UEnvQueryOption::GetDescriptionDetails() const
 {
-	return Generator ? Generator->GetDescriptionDetails() : FString();
+	return Generator ? Generator->GetDescriptionDetails() : FText::GetEmpty();
 }
diff --git a/Engine/Source/Runtime/Engine/Private/AI/EnvironmentQuery/EnvQueryTest.cpp b/Engine/Source/Runtime/Engine/Private/AI/EnvironmentQuery/EnvQueryTest.cpp
--- a/Engine/Source/Runtime/Engine/Private/AI/EnvironmentQuery/EnvQueryTest.cpp
+++ b/Engine/Source/Runtime/Engine/Private/AI/EnvironmentQuery/EnvQueryTest.cpp
@@ -2,6 +2,8 @@
 
 #include "EnginePrivate.h"
 
+#define LOCTEXT_NAMESPACE "EnvQueryGenerator"
+
 UEnvQueryTest::UEnvQueryTest(const class FPostConstructInitializeProperties& PCIP) : Super(PCIP)
 {
 	Cost = EEnvTestCost::Low;
@@ -126,12 +128,12 @@ AActor* UEnvQueryTest::GetItemActor(struct FEnvQueryInstance& QueryInstance, int
 
 FString UEnvQueryTest::GetDescriptionTitle() const
 {
-	return UEnvQueryTypes::GetShortTypeName(this);
+	return UEnvQueryTypes::GetShortTypeName(this).ToString();
 }
 
-FString UEnvQueryTest::GetDescriptionDetails() const
+FText UEnvQueryTest::GetDescriptionDetails() const
 {
-	return FString();
+	return FText::GetEmpty();
 }
 
 #if WITH_EDITOR
@@ -142,67 +144,98 @@ void UEnvQueryTest::PostEditChangeProperty(FPropertyChangedEvent& PropertyChange
 }
 #endif //WITH_EDITOR
 
-FString UEnvQueryTest::DescribeFloatTestParams() const
+FText UEnvQueryTest::DescribeFloatTestParams() const
 {
-	FString ParamDesc;
+	FText ParamDesc;
 
 	if (Condition != EEnvTestCondition::NoCondition)
 	{
-		FString InvalidConditionDesc("invalid condition");
+		FFormatNamedArguments Args;
+		Args.Add(TEXT("Filter"), FText::FromString(UEnvQueryTypes::DescribeFloatParam(FloatFilter)));
+
+		FText InvalidConditionDesc = LOCTEXT("InvalidCondition", "invalid condition");
 		ParamDesc =
-			(Condition == EEnvTestCondition::AtLeast) ? FString::Printf(TEXT("at least %s"), *UEnvQueryTypes::DescribeFloatParam(FloatFilter)) :
-			(Condition == EEnvTestCondition::UpTo) ? FString::Printf(TEXT("up to %s"), *UEnvQueryTypes::DescribeFloatParam(FloatFilter)) :
+			(Condition == EEnvTestCondition::AtLeast) ? FText::Format(LOCTEXT("AtLeastFiltered", "at least {Filter}"), Args) :
+			(Condition == EEnvTestCondition::UpTo) ? FText::Format(LOCTEXT("UpToFiltered", "up to %s"), Args) :
 			InvalidConditionDesc;
-
-		ParamDesc += TEXT(", ");
 	}
 
+	FFormatNamedArguments Args;
+	Args.Add(TEXT("ParmDesc"), ParamDesc);
+
 	if (WeightModifier == EEnvTestWeight::Flat)
 	{
-		ParamDesc += TEXT("don't score");
+		ParamDesc = FText::Format(LOCTEXT("ParmDescWithDontScore", "{ParmDesc}, don't score"), Args);
 	}
 	else if (Weight.IsNamedParam())
 	{
-		ParamDesc += FString::Printf(TEXT("weight: %s"), *Weight.ParamName.ToString());
+		Args.Add(TEXT("WeightName"), FText::FromName(Weight.ParamName));
+		ParamDesc = FText::Format(LOCTEXT("ParmDescWithWeight", "{ParmDesc}, weight: {WeightName}"), Args);
 	}
 	else
 	{
-		ParamDesc += FString::Printf(TEXT("prefer %s [x%.2f]"), (Weight.Value > 0) ? TEXT("greater") : TEXT("lesser"), Weight.Value);
+		FNumberFormattingOptions NumberFormattingOptions;
+		NumberFormattingOptions.MaximumFractionalDigits = 2;
+
+		Args.Add(TEXT("WeightCondition"), (Weight.Value > 0) ? LOCTEXT("Greater", "greater") : LOCTEXT("Lesser", "lesser"));
+		Args.Add(TEXT("WeightValue"), FText::AsNumber(Weight.Value, &NumberFormattingOptions));
+		ParamDesc = FText::Format(LOCTEXT("DescriptionPreferWeightValue", "{ParmDesc}, prefer {WeightCondition} [x{WeightValue}]"), Args);
 	}
 
 	return ParamDesc;
 }
 
-FString UEnvQueryTest::DescribeBoolTestParams(const FString& ConditionDesc) const
+FText UEnvQueryTest::DescribeBoolTestParams(const FString& ConditionDesc) const
 {
-	FString ParamDesc;
+	FText ParamDesc;
 
 	if (Condition != EEnvTestCondition::NoCondition)
 	{
-		FString InvalidConditionDesc("invalid condition");
+		FFormatNamedArguments Args;
+		Args.Add(TEXT("ConditionDesc"), FText::FromString(ConditionDesc));
+
+		if(BoolFilter.IsNamedParam())
+		{
+			Args.Add(TEXT("Filter"), FText::FromString(UEnvQueryTypes::DescribeBoolParam(BoolFilter)));
+		}
+		else
+		{
+			Args.Add(TEXT("Filter"), BoolFilter.Value ? FText::GetEmpty() : LOCTEXT("NotWithSpace", "not "));
+		}
+		FText InvalidConditionDesc = LOCTEXT("InvalidCondition", "invalid condition");
 		ParamDesc =
 			(Condition == EEnvTestCondition::Match) ?
 				(BoolFilter.IsNamedParam() ?
-					FString::Printf(TEXT("require %s: %s"), *ConditionDesc, *UEnvQueryTypes::DescribeBoolParam(BoolFilter)) :
-					FString::Printf(TEXT("require %s%s"), BoolFilter.Value ? TEXT("") : TEXT("not "), *ConditionDesc)
+					FText::Format(LOCTEXT("RequiresCondition", "require {ConditionDesc}: {Filter}"), Args) :
+					FText::Format(LOCTEXT("RequiresFilter", "require {Filter}{ConditionDesc}%s"), Args)
 				) :
 			InvalidConditionDesc;
-
-		ParamDesc += TEXT(", ");
 	}
 
+	FFormatNamedArguments Args;
+	Args.Add(TEXT("ParmDesc"), ParamDesc);
+
 	if (WeightModifier == EEnvTestWeight::Flat)
 	{
-		ParamDesc += TEXT("don't score");
+		ParamDesc = FText::Format(LOCTEXT("ParmDescWithDontScore", "{ParmDesc}, don't score"), Args);
 	}
 	else if (Weight.IsNamedParam())
 	{
-		ParamDesc += FString::Printf(TEXT("weight: %s"), *Weight.ParamName.ToString());
+		Args.Add(TEXT("WeightName"), FText::FromName(Weight.ParamName));
+		ParamDesc = FText::Format(LOCTEXT("ParmDescWithWeight", "{ParmDesc}, weight: {WeightName}"), Args);
 	}
 	else
 	{
-		ParamDesc += FString::Printf(TEXT("prefer %s%s [x%.2f]"), (Weight.Value > 0) ? TEXT("") : TEXT("not "), *ConditionDesc, Weight.Value);
+		FNumberFormattingOptions NumberFormattingOptions;
+		NumberFormattingOptions.MaximumFractionalDigits = 2;
+
+		Args.Add(TEXT("ConditionDesc"), FText::FromString(ConditionDesc));
+		Args.Add(TEXT("WeightCondition"), (Weight.Value > 0) ? FText::GetEmpty() : LOCTEXT("NotWithSpace", "not "));
+		Args.Add(TEXT("WeightValue"), FText::AsNumber(Weight.Value, &NumberFormattingOptions));
+		ParamDesc = FText::Format(LOCTEXT("ParmDescWithPreference", "{ParmDesc}, prefer {WeightCondition}{ConditionDesc} [x{WeightValue}]"), Args);
 	}
 
 	return ParamDesc;
 }
+
+#undef LOCTEXT_NAMESPACE
diff --git a/Engine/Source/Runtime/Engine/Private/AI/EnvironmentQuery/EnvQueryTypes.cpp b/Engine/Source/Runtime/Engine/Private/AI/EnvironmentQuery/EnvQueryTypes.cpp
--- a/Engine/Source/Runtime/Engine/Private/AI/EnvironmentQuery/EnvQueryTypes.cpp
+++ b/Engine/Source/Runtime/Engine/Private/AI/EnvironmentQuery/EnvQueryTypes.cpp
@@ -2,6 +2,8 @@
 
 #include "EnginePrivate.h"
 
+#define LOCTEXT_NAMESPACE "EnvQueryGenerator"
+
 AActor* FEnvQueryResult::GetItemAsActor(int32 Index) const
 {
 	if (Items.IsValidIndex(Index) &&
@@ -26,17 +28,17 @@ FVector FEnvQueryResult::GetItemAsLocation(int32 Index) const
 	return FVector::ZeroVector;
 }
 
-FString UEnvQueryTypes::GetShortTypeName(const UObject* Ob)
+FText UEnvQueryTypes::GetShortTypeName(const UObject* Ob)
 {
 	if (Ob == NULL)
 	{
-		return TEXT("unknown");
+		return LOCTEXT("Unknown", "unknown");
 	}
 
 	const UClass* ObClass = Ob->IsA(UClass::StaticClass()) ? (const UClass*)Ob : Ob->GetClass();
 	if (ObClass->HasAnyClassFlags(CLASS_CompiledFromBlueprint))
 	{
-		return ObClass->GetName().LeftChop(2);
+		return FText::FromString(ObClass->GetName().LeftChop(2));
 	}
 
 	FString TypeDesc = ObClass->GetName();
@@ -46,10 +48,10 @@ FString UEnvQueryTypes::GetShortTypeName(const UObject* Ob)
 		TypeDesc = TypeDesc.Mid(ShortNameIdx + 1);
 	}
 
-	return TypeDesc;
+	return FText::FromString(TypeDesc);
 }
 
-FString UEnvQueryTypes::DescribeContext(TSubclassOf<class UEnvQueryContext> ContextClass)
+FText UEnvQueryTypes::DescribeContext(TSubclassOf<class UEnvQueryContext> ContextClass)
 {
 	return GetShortTypeName(ContextClass);
 }
@@ -84,66 +86,114 @@ FString UEnvQueryTypes::DescribeBoolParam(const FEnvBoolParam& Param)
 	return Param.Value ? TEXT("true") : TEXT("false");
 }
 
-FString FEnvDirection::ToString() const
+FText FEnvDirection::ToText() const
 {
-	return (DirMode == EEnvDirection::TwoPoints) ?
-		FString::Printf(TEXT("[%s - %s]"), *UEnvQueryTypes::DescribeContext(LineFrom), *UEnvQueryTypes::DescribeContext(LineTo)) :
-		FString::Printf(TEXT("[%s rotation]"), *UEnvQueryTypes::DescribeContext(Rotation));
+	if(DirMode == EEnvDirection::TwoPoints)
+	{
+		FFormatNamedArguments Args;
+		Args.Add(TEXT("LineFrom"), UEnvQueryTypes::DescribeContext(LineFrom));
+		Args.Add(TEXT("LineTo"), UEnvQueryTypes::DescribeContext(LineTo));
+
+		return FText::Format(LOCTEXT("DescribeLineFromAndTo", "[{LineFrom} - {LineTo}]"), Args);
+	}
+	else
+	{
+		FFormatNamedArguments Args;
+		Args.Add(TEXT("Rotation"), UEnvQueryTypes::DescribeContext(Rotation));
+
+		return FText::Format(LOCTEXT("DescribeRotation", "[{Rotation} rotation]"), Args);
+	}
 }
 
-FString FEnvTraceData::ToString(FEnvTraceData::EDescriptionMode DescMode) const
+FText FEnvTraceData::ToText(FEnvTraceData::EDescriptionMode DescMode) const
 {
-	FString Desc;
+	FText Desc;
 
 	if (TraceMode == EEnvQueryTrace::Geometry)
 	{
-		Desc = (TraceShape == EEnvTraceShape::Line) ? TEXT("line") :
-			(TraceShape == EEnvTraceShape::Sphere) ? FString::Printf(TEXT("sphere (radius: %.2f)"), ExtentX) :
-			(TraceShape == EEnvTraceShape::Capsule) ? FString::Printf(TEXT("capsule (radius: %.2f, half height: %.2f)"), ExtentX, ExtentZ) :
-			(TraceShape == EEnvTraceShape::Box) ? FString::Printf(TEXT("box (extent: %.2f %.2f %.2f)"), ExtentX, ExtentY, ExtentZ) :
-			TEXT("unknown");
+		FNumberFormattingOptions NumberFormatOptions;
+		NumberFormatOptions.MaximumFractionalDigits = 2;
+
+		FFormatNamedArguments Args;
+		Args.Add(TEXT("ExtentX"), FText::AsNumber(ExtentX, &NumberFormatOptions));
+		Args.Add(TEXT("ExtentY"), FText::AsNumber(ExtentY, &NumberFormatOptions));
+		Args.Add(TEXT("ExtentZ"), FText::AsNumber(ExtentZ, &NumberFormatOptions));
+
+		Desc = (TraceShape == EEnvTraceShape::Line) ? LOCTEXT("Line", "line") :
+			(TraceShape == EEnvTraceShape::Sphere) ? FText::Format(LOCTEXT("SphereWithRadius", "sphere (radius: {ExtentX})"), Args) :
+			(TraceShape == EEnvTraceShape::Capsule) ? FText::Format(LOCTEXT("CasuleWithRadiusHalfHeight", "capsule (radius: {ExtentX}, half height: {ExtentZ})"), Args) :
+			(TraceShape == EEnvTraceShape::Box) ? FText::Format(LOCTEXT("BoxWithExtents", "box (extent: {ExtentX} {ExtentY} {ExtentZ)"), Args) :
+			LOCTEXT("Unknown", "unknown");
 
 		if (DescMode == FEnvTraceData::Brief)
 		{
 			static UEnum* ChannelEnum = FindObject<UEnum>(ANY_PACKAGE, TEXT("ETraceTypeQuery"), true);
-			Desc += FString::Printf(TEXT(" %s on %s"), bCanProjectDown ? TEXT("projection") : TEXT("trace"), *ChannelEnum->GetEnumText(TraceChannel).ToString());
+
+			FFormatNamedArguments Args;
+			Args.Add(TEXT("ExtentDescription"), Desc);
+			Args.Add(TEXT("ProjectionTraceDesc"), bCanProjectDown ? LOCTEXT("Projection", "projection") : LOCTEXT("Trace", "trace"));
+			Args.Add(TEXT("Channel"), ChannelEnum->GetEnumText(TraceChannel));
+
+			Desc = FText::Format(LOCTEXT("GeometryBriefDescription", "{ExtentDescription} {ProjectionTraceDesc} on {Channel}"), Args);
 		}
 		else
 		{
+			FFormatNamedArguments Args;
+			Args.Add(TEXT("Description"), Desc);
+
 			if (bTraceComplex)
 			{
-				Desc += TEXT(", complex collision");
+				Desc = FText::Format(LOCTEXT("DescWithComplexCollision", "{Description}, complex collision"), Args);
 			}
 
 			if (!bOnlyBlockingHits)
 			{
-				Desc += TEXT(", accept non blocking");
+				Desc = FText::Format(LOCTEXT("DescWithNonBlocking", "{Description}, accept non blocking"), Args);
 			}
 		}
 	}
 	else if (TraceMode == EEnvQueryTrace::Navigation)
 	{
-		Desc = TEXT("navmesh ");
+		Desc = LOCTEXT("Navmesh", "navmesh");
 		if (bCanProjectDown)
 		{
-			Desc += FString::Printf(TEXT("projection (%s: %.0f"),
-				(ProjectDown == ProjectUp) ? TEXT("height") : (ProjectDown > ProjectUp) ? TEXT("down") : TEXT("up"),
-				FMath::Max(ProjectDown, ProjectUp));
+			FNumberFormattingOptions NumberFormatOptions;
+			NumberFormatOptions.MaximumFractionalDigits = 0;
+
+			FFormatNamedArguments Args;
+			Args.Add(TEXT("Description"), Desc);
+			Args.Add(TEXT("Direction"), (ProjectDown == ProjectUp) ? LOCTEXT("Height", "height") : (ProjectDown > ProjectUp) ? LOCTEXT("Down", "down") : LOCTEXT("Up", "up"));
+			Args.Add(TEXT("ProjectionAmount"), FText::AsNumber(FMath::Max(ProjectDown, ProjectUp), &NumberFormatOptions));
+
+			Desc = FText::Format(LOCTEXT("DescriptionWithProjection", "projection ({Direction}: {ProjectionAmount}"), Args);
 
 			if (ExtentX > 1.0f)
 			{
-				Desc += FString::Printf(TEXT(", radius: %.2f"), ExtentX);
+				NumberFormatOptions.MaximumFractionalDigits = 2;
+
+				FFormatNamedArguments RadiusArgs;
+				RadiusArgs.Add(TEXT("Description"), Desc);
+				RadiusArgs.Add(TEXT("Radius"), FText::AsNumber(ExtentX, &NumberFormatOptions));
+
+				Desc = FText::Format(LOCTEXT("DescriptionWithRadius", "{Description}, radius {Radius}"), RadiusArgs);
 			}
-			Desc += TEXT(")");
+
+			FFormatNamedArguments EndingArgs;
+			EndingArgs.Add(TEXT("Description"), Desc);
+			Desc = FText::Format(LOCTEXT("DescriptionWithEnding", "{Description})"), EndingArgs);
 		}
 		else
 		{
-			Desc += TEXT("trace");
+			FFormatNamedArguments Args;
+			Args.Add(TEXT("Description"), Desc);
+
+			Desc = FText::Format(LOCTEXT("DescriptionWithTrace", "{Description} trace"), Args);
 			if (NavigationFilter)
 			{
-				Desc += TEXT(" (filter: ");
-				Desc += NavigationFilter->GetName();
-				Desc += TEXT(")");
+				FFormatNamedArguments NavFilterArgs;
+				NavFilterArgs.Add(TEXT("Description"), Desc);
+				NavFilterArgs.Add(TEXT("NavigationFilter"), FText::FromString(NavigationFilter->GetName()));
+				Desc = FText::Format(LOCTEXT("DescriptionWithNavigationFilter", "{Description} (filter: {NavigationFilter})"), NavFilterArgs);
 			}
 		}
 	}
@@ -189,4 +239,6 @@ namespace FEQSHelpers
 
 		return Cast<const ARecastNavMesh>(NavSys->GetMainNavData());
 	}
-}
\ No newline at end of file
+}
+
+#undef LOCTEXT_NAMESPACE
diff --git a/Engine/Source/Runtime/Engine/Private/AI/EnvironmentQuery/Generators/EnvQueryGenerator_OnCircle.cpp b/Engine/Source/Runtime/Engine/Private/AI/EnvironmentQuery/Generators/EnvQueryGenerator_OnCircle.cpp
--- a/Engine/Source/Runtime/Engine/Private/AI/EnvironmentQuery/Generators/EnvQueryGenerator_OnCircle.cpp
+++ b/Engine/Source/Runtime/Engine/Private/AI/EnvironmentQuery/Generators/EnvQueryGenerator_OnCircle.cpp
@@ -2,6 +2,8 @@
 
 #include "EnginePrivate.h"
 
+#define LOCTEXT_NAMESPACE "EnvQueryGenerator"
+
 UEnvQueryGenerator_OnCircle::UEnvQueryGenerator_OnCircle(const class FPostConstructInitializeProperties& PCIP) 
 	: Super(PCIP)
 {
@@ -220,28 +222,40 @@ void UEnvQueryGenerator_OnCircle::RunBoxTrace(const FVector& StartPos, const FVe
 	HitPos = bHit ? OutHit.ImpactPoint : EndPos;
 };
 
-FString UEnvQueryGenerator_OnCircle::GetDescriptionTitle() const
+FText UEnvQueryGenerator_OnCircle::GetDescriptionTitle() const
 {
-	return FString::Printf(TEXT("%s: generate items on circle around %s"),
-		*Super::GetDescriptionTitle(), *UEnvQueryTypes::DescribeContext(CircleCenter));
+	FFormatNamedArguments Args;
+	Args.Add(TEXT("DescriptionTitle"), Super::GetDescriptionTitle());
+	Args.Add(TEXT("DescribeContext"), UEnvQueryTypes::DescribeContext(CircleCenter));
+
+	return FText::Format(LOCTEXT("DescriptionGenerateCircleAroundContext", "{DescriptionTitle}: generate items on circle around {DescribeContext}"), Args);
 }
 
-FString UEnvQueryGenerator_OnCircle::GetDescriptionDetails() const
+FText UEnvQueryGenerator_OnCircle::GetDescriptionDetails() const
 {
-	FString Desc = FString::Printf(TEXT("radius: %s, item span: %.2f\n%s"),
-		*UEnvQueryTypes::DescribeFloatParam(Radius), *UEnvQueryTypes::DescribeFloatParam(ItemSpacing),
-		*TraceData.ToString(FEnvTraceData::Detailed));
+	FFormatNamedArguments Args;
+	Args.Add(TEXT("Radius"), FText::FromString(UEnvQueryTypes::DescribeFloatParam(Radius)));
+	Args.Add(TEXT("ItemSpacing"), FText::FromString(UEnvQueryTypes::DescribeFloatParam(ItemSpacing)));
+	Args.Add(TEXT("TraceData"), TraceData.ToText(FEnvTraceData::Detailed));
+
+	FText Desc = FText::Format(LOCTEXT("OnCircleDescription", "radius: {Radius}, item span: {ItemSpacing}\n{TraceData}"), Args);
 
 	if (bDefineArc)
 	{
-		Desc += FString::Printf(TEXT("\nLimit to %.2f angle both sides on %s"), Angle.Value, *ArcDirection.ToString());
+		FFormatNamedArguments ArcArgs;
+		ArcArgs.Add(TEXT("Description"), Desc);
+		ArcArgs.Add(TEXT("AngleValue"), Angle.Value);
+		ArcArgs.Add(TEXT("ArcDirection"), ArcDirection.ToText());
+		Desc = FText::Format(LOCTEXT("DescriptionWithArc", "{Description}\nLimit to {AngleValue} angle both sides on {ArcDirection}"), ArcArgs);
 	}
 
-	FString ProjDesc = ProjectionData.ToString(FEnvTraceData::Brief);
-	if (ProjDesc.Len() > 0)
+	FText ProjDesc = ProjectionData.ToText(FEnvTraceData::Brief);
+	if (!ProjDesc.IsEmpty())
 	{
-		Desc += TEXT(", ");
-		Desc += ProjDesc;
+		FFormatNamedArguments ProjArgs;
+		ProjArgs.Add(TEXT("Description"), Desc);
+		ProjArgs.Add(TEXT("ProjectionDescription"), ProjDesc);
+		Desc = FText::Format(LOCTEXT("DescriptionWithProjection", "{Description}, {ProjectionDescription}"), ProjArgs);
 	}
 
 	return Desc;
@@ -277,4 +291,6 @@ void UEnvQueryGenerator_OnCircle::PostEditChangeProperty( struct FPropertyChange
 	Super::PostEditChangeProperty(PropertyChangedEvent);
 }
 
-#endif // WITH_EDITOR
\ No newline at end of file
+#endif // WITH_EDITOR
+
+#undef LOCTEXT_NAMESPACE
\ No newline at end of file
diff --git a/Engine/Source/Runtime/Engine/Private/AI/EnvironmentQuery/Generators/EnvQueryGenerator_PathingGrid.cpp b/Engine/Source/Runtime/Engine/Private/AI/EnvironmentQuery/Generators/EnvQueryGenerator_PathingGrid.cpp
--- a/Engine/Source/Runtime/Engine/Private/AI/EnvironmentQuery/Generators/EnvQueryGenerator_PathingGrid.cpp
+++ b/Engine/Source/Runtime/Engine/Private/AI/EnvironmentQuery/Generators/EnvQueryGenerator_PathingGrid.cpp
@@ -2,6 +2,8 @@
 
 #include "EnginePrivate.h"
 
+#define LOCTEXT_NAMESPACE "EnvQueryGenerator"
+
 UEnvQueryGenerator_PathingGrid::UEnvQueryGenerator_PathingGrid(const class FPostConstructInitializeProperties& PCIP) : Super(PCIP)
 {
 	GenerateDelegate.BindUObject(this, &UEnvQueryGenerator_PathingGrid::GenerateItems);
@@ -81,18 +83,23 @@ void UEnvQueryGenerator_PathingGrid::GenerateItems(struct FEnvQueryInstance& Que
 #endif // WITH_RECAST
 }
 
-FString UEnvQueryGenerator_PathingGrid::GetDescriptionTitle() const
+FText UEnvQueryGenerator_PathingGrid::GetDescriptionTitle() const
 {
-	return FString::Printf(TEXT("%s: generate around %s"),
-		*Super::GetDescriptionTitle(), *UEnvQueryTypes::DescribeContext(GenerateAround));
+	FFormatNamedArguments Args;
+	Args.Add(TEXT("DescriptionTitle"), Super::GetDescriptionTitle());
+	Args.Add(TEXT("DescribeContext"), UEnvQueryTypes::DescribeContext(GenerateAround));
+
+	return FText::Format(LOCTEXT("DescriptionGenerateAroundContext", "{DescriptionTitle}: generate around {DescribeContext}"), Args);
 };
 
-FString UEnvQueryGenerator_PathingGrid::GetDescriptionDetails() const
+FText UEnvQueryGenerator_PathingGrid::GetDescriptionDetails() const
 {
-	return FString::Printf(TEXT("max distance: %s, density: %s, path from context: %s"),
-		*UEnvQueryTypes::DescribeFloatParam(MaxPathDistance),
-		*UEnvQueryTypes::DescribeFloatParam(Density),
-		*UEnvQueryTypes::DescribeBoolParam(PathFromContext));
+	FFormatNamedArguments Args;
+	Args.Add(TEXT("MaxPathDistance"), FText::FromString(UEnvQueryTypes::DescribeFloatParam(MaxPathDistance)));
+	Args.Add(TEXT("Density"), FText::FromString(UEnvQueryTypes::DescribeFloatParam(Density)));
+	Args.Add(TEXT("PathFromContext"), FText::FromString(UEnvQueryTypes::DescribeBoolParam(PathFromContext)));
+
+	return FText::Format(LOCTEXT("DescriptionDetailsPathingGrid", "max distance: {MaxPathDistance}, density: {Density}, path from context: {PathFromContext}"), Args);
 }
 
 #if WITH_RECAST
@@ -153,4 +160,6 @@ bool UEnvQueryGenerator_PathingGrid::IsNavLocationInPathDistance(const class ARe
 
 #undef ENVQUERY_CLUSTER_SEARCH
 
-#endif // WITH_RECAST
\ No newline at end of file
+#endif // WITH_RECAST
+
+#undef LOCTEXT_NAMESPACE
\ No newline at end of file
diff --git a/Engine/Source/Runtime/Engine/Private/AI/EnvironmentQuery/Generators/EnvQueryGenerator_SimpleGrid.cpp b/Engine/Source/Runtime/Engine/Private/AI/EnvironmentQuery/Generators/EnvQueryGenerator_SimpleGrid.cpp
--- a/Engine/Source/Runtime/Engine/Private/AI/EnvironmentQuery/Generators/EnvQueryGenerator_SimpleGrid.cpp
+++ b/Engine/Source/Runtime/Engine/Private/AI/EnvironmentQuery/Generators/EnvQueryGenerator_SimpleGrid.cpp
@@ -2,6 +2,8 @@
 
 #include "EnginePrivate.h"
 
+#define LOCTEXT_NAMESPACE "EnvQueryGenerator"
+
 UEnvQueryGenerator_SimpleGrid::UEnvQueryGenerator_SimpleGrid(const class FPostConstructInitializeProperties& PCIP) : Super(PCIP)
 {
 	GenerateDelegate.BindUObject(this, &UEnvQueryGenerator_SimpleGrid::GenerateItems);
@@ -76,24 +78,33 @@ void UEnvQueryGenerator_SimpleGrid::GenerateItems(struct FEnvQueryInstance& Quer
 	}
 }
 
-FString UEnvQueryGenerator_SimpleGrid::GetDescriptionTitle() const
+FText UEnvQueryGenerator_SimpleGrid::GetDescriptionTitle() const
 {
-	return FString::Printf(TEXT("%s: generate around %s"),
-		*Super::GetDescriptionTitle(), *UEnvQueryTypes::DescribeContext(GenerateAround));
+	FFormatNamedArguments Args;
+	Args.Add(TEXT("DescriptionTitle"), Super::GetDescriptionTitle());
+	Args.Add(TEXT("DescribeContext"), UEnvQueryTypes::DescribeContext(GenerateAround));
+
+	return FText::Format(LOCTEXT("DescriptionGenerateSimpleGridAroundContext", "{DescriptionTitle}: generate around {DescribeContext}"), Args);
 };
 
-FString UEnvQueryGenerator_SimpleGrid::GetDescriptionDetails() const
+FText UEnvQueryGenerator_SimpleGrid::GetDescriptionDetails() const
 {
-	FString Desc = FString::Printf(TEXT("radius: %s, density: %s"),
-		*UEnvQueryTypes::DescribeFloatParam(Radius),
-		*UEnvQueryTypes::DescribeFloatParam(Density));
+	FFormatNamedArguments Args;
+	Args.Add(TEXT("Radius"), FText::FromString(UEnvQueryTypes::DescribeFloatParam(Radius)));
+	Args.Add(TEXT("Density"), FText::FromString(UEnvQueryTypes::DescribeFloatParam(Density)));
 
-	FString ProjDesc = ProjectionData.ToString(FEnvTraceData::Brief);
-	if (ProjDesc.Len() > 0)
+	FText Desc = FText::Format(LOCTEXT("SimpleGridDescription", "radius: {Radius}, density: {Density}"), Args);
+
+	FText ProjDesc = ProjectionData.ToText(FEnvTraceData::Brief);
+	if (!ProjDesc.IsEmpty())
 	{
-		Desc += TEXT(", ");
-		Desc += ProjDesc;
+		FFormatNamedArguments ProjArgs;
+		ProjArgs.Add(TEXT("Description"), Desc);
+		ProjArgs.Add(TEXT("ProjectionDescription"), ProjDesc);
+		Desc = FText::Format(LOCTEXT("DescriptionWithProjection", "{Description}, {ProjectionDescription}"), ProjArgs);
 	}
 
 	return Desc;
 }
+
+#undef LOCTEXT_NAMESPACE
\ No newline at end of file
diff --git a/Engine/Source/Runtime/Engine/Private/AI/EnvironmentQuery/Tests/EnvQueryTest_Distance.cpp b/Engine/Source/Runtime/Engine/Private/AI/EnvironmentQuery/Tests/EnvQueryTest_Distance.cpp
--- a/Engine/Source/Runtime/Engine/Private/AI/EnvironmentQuery/Tests/EnvQueryTest_Distance.cpp
+++ b/Engine/Source/Runtime/Engine/Private/AI/EnvironmentQuery/Tests/EnvQueryTest_Distance.cpp
@@ -97,10 +97,10 @@ FString UEnvQueryTest_Distance::GetDescriptionTitle() const
 
 	return FString::Printf(TEXT("%s%s: to %s"), 
 		*Super::GetDescriptionTitle(), *ModeDesc,
-		*UEnvQueryTypes::DescribeContext(DistanceTo));
+		*UEnvQueryTypes::DescribeContext(DistanceTo).ToString());
 }
 
-FString UEnvQueryTest_Distance::GetDescriptionDetails() const
+FText UEnvQueryTest_Distance::GetDescriptionDetails() const
 {
 	return DescribeFloatTestParams();
 }
diff --git a/Engine/Source/Runtime/Engine/Private/AI/EnvironmentQuery/Tests/EnvQueryTest_Dot.cpp b/Engine/Source/Runtime/Engine/Private/AI/EnvironmentQuery/Tests/EnvQueryTest_Dot.cpp
--- a/Engine/Source/Runtime/Engine/Private/AI/EnvironmentQuery/Tests/EnvQueryTest_Dot.cpp
+++ b/Engine/Source/Runtime/Engine/Private/AI/EnvironmentQuery/Tests/EnvQueryTest_Dot.cpp
@@ -164,10 +164,10 @@ bool UEnvQueryTest_Dot::RequiresPerItemUpdates(TSubclassOf<class UEnvQueryContex
 
 FString UEnvQueryTest_Dot::GetDescriptionTitle() const
 {
-	return FString::Printf(TEXT("%s: %s and %s"), *Super::GetDescriptionTitle(), *LineA.ToString(), *LineB.ToString());
+	return FString::Printf(TEXT("%s: %s and %s"), *Super::GetDescriptionTitle(), *LineA.ToText().ToString(), *LineB.ToText().ToString());
 }
 
-FString UEnvQueryTest_Dot::GetDescriptionDetails() const
+FText UEnvQueryTest_Dot::GetDescriptionDetails() const
 {
 	return DescribeFloatTestParams();
 }
diff --git a/Engine/Source/Runtime/Engine/Private/AI/EnvironmentQuery/Tests/EnvQueryTest_Pathfinding.cpp b/Engine/Source/Runtime/Engine/Private/AI/EnvironmentQuery/Tests/EnvQueryTest_Pathfinding.cpp
--- a/Engine/Source/Runtime/Engine/Private/AI/EnvironmentQuery/Tests/EnvQueryTest_Pathfinding.cpp
+++ b/Engine/Source/Runtime/Engine/Private/AI/EnvironmentQuery/Tests/EnvQueryTest_Pathfinding.cpp
@@ -2,6 +2,8 @@
 
 #include "EnginePrivate.h"
 
+#define LOCTEXT_NAMESPACE "EnvQueryGenerator"
+
 UEnvQueryTest_Pathfinding::UEnvQueryTest_Pathfinding(const class FPostConstructInitializeProperties& PCIP) : Super(PCIP)
 {
 	ExecuteDelegate.BindUObject(this, &UEnvQueryTest_Pathfinding::RunTest);
@@ -95,41 +97,67 @@ FString UEnvQueryTest_Pathfinding::GetDescriptionTitle() const
 	FString ModeDesc[] = { TEXT("PathExist"), TEXT("PathCost"), TEXT("PathLength") };
 
 	FString DirectionDesc = PathToItem.IsNamedParam() ?
-		FString::Printf(TEXT("%s, direction: %s"), *UEnvQueryTypes::DescribeContext(Context), *UEnvQueryTypes::DescribeBoolParam(PathToItem)) :
-		FString::Printf(TEXT("%s %s"), PathToItem.Value ? TEXT("from") : TEXT("to"), *UEnvQueryTypes::DescribeContext(Context));
+		FString::Printf(TEXT("%s, direction: %s"), *UEnvQueryTypes::DescribeContext(Context).ToString(), *UEnvQueryTypes::DescribeBoolParam(PathToItem)) :
+		FString::Printf(TEXT("%s %s"), PathToItem.Value ? TEXT("from") : TEXT("to"), *UEnvQueryTypes::DescribeContext(Context).ToString());
 
 	return FString::Printf(TEXT("%s: %s"), *ModeDesc[TestMode], *DirectionDesc);
 }
 
-FString UEnvQueryTest_Pathfinding::GetDescriptionDetails() const
+FText UEnvQueryTest_Pathfinding::GetDescriptionDetails() const
 {
-	FString AdditionalDesc;
+	FText AdditionalDesc;
 	if (HierarchicalPathfinding.IsNamedParam())
 	{
-		AdditionalDesc = FString::Printf(TEXT("Hierarchical pathfinding: %s"), *HierarchicalPathfinding.ParamName.ToString());
+		FFormatNamedArguments Args;
+		Args.Add(TEXT("PathName"), FText::FromString(HierarchicalPathfinding.ParamName.ToString()));
+		AdditionalDesc = FText::Format(LOCTEXT("HierarchicalPathfindingWithName", "Hierarchical pathfinding: {PathName}"), Args);
 	}
 	else if (HierarchicalPathfinding.Value)
 	{
-		AdditionalDesc = TEXT("Use hierarchical pathfinding");
+		AdditionalDesc = LOCTEXT("UseHierarchicalPathfinding", "Use hierarchical pathfinding");
 	}
 
 	if (DiscardUnreachable.IsNamedParam())
 	{
-		AdditionalDesc += FString::Printf(TEXT("%siscard unreachable: %s"), AdditionalDesc.Len() ? TEXT(", d") : TEXT("D"), *DiscardUnreachable.ParamName.ToString());
+		FFormatNamedArguments Args;
+		Args.Add(TEXT("Description"), AdditionalDesc);
+		Args.Add(TEXT("UnreachableParamName"), FText::FromName(DiscardUnreachable.ParamName));
+
+		if(AdditionalDesc.IsEmpty())
+		{
+			AdditionalDesc = FText::Format(LOCTEXT("DiscardUnreachableWithParam", "Discard unreachable: {UnreachableParamName}"), Args);
+		}
+		else
+		{
+			AdditionalDesc = FText::Format(LOCTEXT("DescWithUnreachableAndParam", "{Description}, discard unreachable: {UnreachableParamName}"), Args);
+		}
 	}
 	else if (DiscardUnreachable.Value)
 	{
-		AdditionalDesc += FString::Printf(TEXT("%siscard unreachable"), AdditionalDesc.Len() ? TEXT(", d") : TEXT("D"));
+		if(AdditionalDesc.IsEmpty())
+		{
+			AdditionalDesc = LOCTEXT("DiscardUnreachable", "Discard unreachable");
+		}
+		else
+		{
+			FFormatNamedArguments Args;
+			Args.Add(TEXT("Description"), AdditionalDesc);
+			AdditionalDesc = FText::Format(LOCTEXT("DescWithUnreachable", "{Description}, discard unreachable"), Args);
+		}
 	}
 
-	AdditionalDesc.AppendChar(TEXT('\n'));
+	FFormatNamedArguments Args;
+	Args.Add(TEXT("Description"), AdditionalDesc);
+
 	if (bWorkOnFloatValues)
 	{
-		AdditionalDesc += DescribeFloatTestParams();
+		Args.Add(TEXT("FloatParamsDescription"), DescribeFloatTestParams());
+		AdditionalDesc = FText::Format(LOCTEXT("DescriptionWithDescribedFloatParams", "{Description}\n{FloatParamsDescription}"), Args);
 	}
 	else
 	{
-		AdditionalDesc += DescribeBoolTestParams("existing path");
+		Args.Add(TEXT("BoolParamsDescription"), DescribeBoolTestParams("existing path"));
+		AdditionalDesc = FText::Format(LOCTEXT("DescriptionWithDescribedBoolParams", "{Description}\n{BoolParamsDescription}"), Args);
 	}
 
 	return AdditionalDesc;
@@ -193,3 +221,5 @@ ANavigationData* UEnvQueryTest_Pathfinding::FindNavigationData(UNavigationSystem
 
 	return NavSys->GetMainNavData(NavigationSystem::DontCreate);
 }
+
+#undef LOCTEXT_NAMESPACE
diff --git a/Engine/Source/Runtime/Engine/Private/AI/EnvironmentQuery/Tests/EnvQueryTest_Trace.cpp b/Engine/Source/Runtime/Engine/Private/AI/EnvironmentQuery/Tests/EnvQueryTest_Trace.cpp
--- a/Engine/Source/Runtime/Engine/Private/AI/EnvironmentQuery/Tests/EnvQueryTest_Trace.cpp
+++ b/Engine/Source/Runtime/Engine/Private/AI/EnvironmentQuery/Tests/EnvQueryTest_Trace.cpp
@@ -2,6 +2,8 @@
 
 #include "EnginePrivate.h"
 
+#define LOCTEXT_NAMESPACE "EnvQueryGenerator"
+
 UEnvQueryTest_Trace::UEnvQueryTest_Trace(const class FPostConstructInitializeProperties& PCIP) : Super(PCIP)
 {
 	ExecuteDelegate.BindUObject(this, &UEnvQueryTest_Trace::RunTest);
@@ -98,22 +100,20 @@ FString UEnvQueryTest_Trace::GetDescriptionTitle() const
 	FString ChannelDesc = ChannelEnum->GetEnumText(TraceData.TraceChannel).ToString();
 
 	FString DirectionDesc = TraceToItem.IsNamedParam() ?
-		FString::Printf(TEXT("%s, direction: %s"), *UEnvQueryTypes::DescribeContext(Context), *UEnvQueryTypes::DescribeBoolParam(TraceToItem)) :
-		FString::Printf(TEXT("%s %s"), TraceToItem.Value ? TEXT("from") : TEXT("to"), *UEnvQueryTypes::DescribeContext(Context));
+		FString::Printf(TEXT("%s, direction: %s"), *UEnvQueryTypes::DescribeContext(Context).ToString(), *UEnvQueryTypes::DescribeBoolParam(TraceToItem)) :
+		FString::Printf(TEXT("%s %s"), TraceToItem.Value ? TEXT("from") : TEXT("to"), *UEnvQueryTypes::DescribeContext(Context).ToString());
 
 	return FString::Printf(TEXT("%s: %s on %s"), 
 		*Super::GetDescriptionTitle(), *DirectionDesc, *ChannelDesc);
 }
 
-FString UEnvQueryTest_Trace::GetDescriptionDetails() const
+FText UEnvQueryTest_Trace::GetDescriptionDetails() const
 {
-	FString Desc;
-
-	Desc = TraceData.ToString(FEnvTraceData::Detailed);
-	Desc.AppendChar(TEXT('\n'));
-	Desc += DescribeBoolTestParams("hit");
-
-	return Desc;
+	FFormatNamedArguments Args;
+	Args.Add(TEXT("TraceData"),  TraceData.ToText(FEnvTraceData::Detailed));
+	Args.Add(TEXT("TestParams"),  DescribeBoolTestParams("hit"));
+	
+	return LOCTEXT("TraceDescription", "{TraceData}\n{TestParams}");
 }
 
 bool UEnvQueryTest_Trace::RunLineTraceTo(const FVector& ItemPos, const FVector& ContextPos, AActor* ItemActor, UWorld* World, enum ECollisionChannel Channel, const FCollisionQueryParams& Params, const FVector& Extent)
@@ -187,3 +187,5 @@ bool UEnvQueryTest_Trace::RunCapsuleTraceFrom(const FVector& ItemPos, const FVec
 	const bool bHit = World->SweepTest(ItemPos, ContextPos, FQuat::Identity, Channel, FCollisionShape::MakeCapsule(Extent.X, Extent.Z), TraceParams);
 	return bHit;
 }
+
+#undef LOCTEXT_NAMESPACE
diff --git a/Engine/Source/Runtime/Engine/Private/EdGraph/EdGraph.cpp b/Engine/Source/Runtime/Engine/Private/EdGraph/EdGraph.cpp
--- a/Engine/Source/Runtime/Engine/Private/EdGraph/EdGraph.cpp
+++ b/Engine/Source/Runtime/Engine/Private/EdGraph/EdGraph.cpp
@@ -614,9 +614,36 @@ bool UEdGraphNode::CanUserDeleteNode() const
 	return true;
 }
 
-FString UEdGraphNode::GetNodeTitle(ENodeTitleType::Type TitleType) const
+FText UEdGraphNode::GetNodeTitle(ENodeTitleType::Type TitleType) const
 {
-	return GetClass()->GetName();
+	return FText::FromString(GetClass()->GetName());
+}
+
+FString UEdGraphNode::GetNodeNativeTitle(ENodeTitleType::Type TitleType) const
+{
+	FText NodeTitle = GetNodeTitle(TitleType);
+	if(const FString* SourceString = FTextInspector::GetSourceString(NodeTitle))
+	{
+		return *SourceString;
+	}
+	return NodeTitle.ToString();
+}
+
+FText UEdGraphNode::GetNodeSearchTitle() const
+{
+	TArray< FText > SearchMetadata;
+
+	FText LocalizedTitle = GetNodeTitle(ENodeTitleType::ListView);
+	FText NativeTitle = FText::FromString(GetNodeNativeTitle(ENodeTitleType::ListView));
+	if(FInternationalization::GetCurrentCulture()->GetName() == TEXT("en_US"))
+	{
+		ensure( LocalizedTitle.CompareTo(NativeTitle) == 0 );
+	}
+
+	FFormatNamedArguments Args;
+	Args.Add(TEXT("LocalizedTitle"), LocalizedTitle);
+	Args.Add(TEXT("NativeTitle"), NativeTitle);
+	return FText::Format(LOCTEXT("NodeSearchTitle", "{LocalizedTitle} {NativeTitle}"), Args);
 }
 
 UObject* UEdGraphNode::GetJumpTargetForDoubleClick() const
@@ -689,19 +716,19 @@ FString UEdGraphNode_Comment::GetDocumentationExcerptName() const
 	return TEXT("UEdGraphNode_Comment");
 }
 
-FString UEdGraphNode_Comment::GetNodeTitle(ENodeTitleType::Type TitleType) const
+FText UEdGraphNode_Comment::GetNodeTitle(ENodeTitleType::Type TitleType) const
 {
 	if(TitleType == ENodeTitleType::ListView)
 	{
-		return NSLOCTEXT("K2Node", "CommentBlock_Title", "Comment").ToString();
+		return NSLOCTEXT("K2Node", "CommentBlock_Title", "Comment");
 	}
 
-	return NodeComment;
+	return FText::FromString(NodeComment);
 }
 
 FString UEdGraphNode_Comment::GetPinNameOverride(const UEdGraphPin& Pin) const
 {
-	return GetNodeTitle(ENodeTitleType::ListView);
+	return GetNodeTitle(ENodeTitleType::ListView).ToString();
 }
 
 FLinearColor UEdGraphNode_Comment::GetNodeCommentColor() const
@@ -885,11 +912,11 @@ const FString UEdGraphPin::GetLinkInfoString( const FString& InFunctionName, con
 #if WITH_EDITOR
 	const FString FromPinName = PinName;
 	const UEdGraphNode* FromPinNode = Cast<UEdGraphNode>(GetOuter());
-	const FString FromPinNodeName = (FromPinNode != NULL) ? FromPinNode->GetNodeTitle(ENodeTitleType::ListView) : TEXT("Unknown");
+	const FString FromPinNodeName = (FromPinNode != NULL) ? FromPinNode->GetNodeTitle(ENodeTitleType::ListView).ToString() : TEXT("Unknown");
 	
 	const FString ToPinName = InToPin->PinName;
 	const UEdGraphNode* ToPinNode = Cast<UEdGraphNode>(InToPin->GetOuter());
-	const FString ToPinNodeName = (ToPinNode != NULL) ? ToPinNode->GetNodeTitle(ENodeTitleType::ListView) : TEXT("Unknown");
+	const FString ToPinNodeName = (ToPinNode != NULL) ? ToPinNode->GetNodeTitle(ENodeTitleType::ListView).ToString() : TEXT("Unknown");
 	const FString LinkInfo = FString::Printf( TEXT("UEdGraphPin::%s Pin '%s' on node '%s' %s '%s' on node '%s'"), *InFunctionName, *ToPinName, *ToPinNodeName, *InInfoData, *FromPinName, *FromPinNodeName);
 	return LinkInfo;
 #else
@@ -1262,7 +1289,7 @@ FPinConnectionResponse UEdGraphSchema::CopyPinLinks(UEdGraphPin& CopyFromPin, UE
 FString UEdGraphSchema::GetPinDisplayName(const UEdGraphPin* Pin) const
 {
 	check(Pin != NULL);
-	return (Pin->PinFriendlyName.Len() > 0) ? Pin->PinFriendlyName : Pin->PinName;
+	return !Pin->PinFriendlyName.IsEmpty() ? Pin->PinFriendlyName.ToString() : Pin->PinName;
 }
 
 void UEdGraphSchema::ConstructBasicPinTooltip(UEdGraphPin const& Pin, FString const& PinDescription, FString& TooltipOut) const
@@ -1299,7 +1326,8 @@ void UEdGraphSchema::ReconstructNode(UEdGraphNode& TargetNode, bool bIsBatchRequ
 
 void UEdGraphSchema::GetGraphDisplayInformation(const UEdGraph& Graph, /*out*/ FGraphDisplayInfo& DisplayInfo) const
 {
-	DisplayInfo.DisplayName = FText::FromString( Graph.GetName() );
+	DisplayInfo.PlainName = FText::FromString( Graph.GetName() );
+	DisplayInfo.DisplayName = DisplayInfo.PlainName;
 }
 
 void UEdGraphSchema::GetContextMenuActions(const UEdGraph* CurrentGraph, const UEdGraphNode* InGraphNode, const UEdGraphPin* InGraphPin, class FMenuBuilder* MenuBuilder, bool bIsDebugging) const
