MIssed file in my last change.

diff --git a/Engine/Source/Runtime/Engine/Classes/Materials/MaterialInterface.h b/Engine/Source/Runtime/Engine/Classes/Materials/MaterialInterface.h
--- a/Engine/Source/Runtime/Engine/Classes/Materials/MaterialInterface.h
+++ b/Engine/Source/Runtime/Engine/Classes/Materials/MaterialInterface.h
@@ -508,4 +508,4 @@ private:
 };
 
 /** Helper function to serialize inline shader maps for the given material resources. */
-extern void SerializeInlineShaderMaps(const TArray<FMaterialResource*>& MaterialResourcesToSave, FArchive& Ar, FMaterialResource* OutMaterialResourcesLoaded[][ERHIFeatureLevel::Num]);
\ No newline at end of file
+extern void SerializeInlineShaderMaps(const TMap<const class ITargetPlatform*,TArray<FMaterialResource*>>& PlatformMaterialResourcesToSave, FArchive& Ar, FMaterialResource* OutMaterialResourcesLoaded[][ERHIFeatureLevel::Num]);
\ No newline at end of file
