#UE4 - Cleanup packet sending/queuing logic.

* Rename Out to SendBuffer
* Remove PreSend/PostSend in favor of a single pass
* You now buffer bits up in a temp write, and send using a single call: WriteBitsToSendBuffer( Bits, NumBits );

diff --git a/Engine/Source/Runtime/Engine/Classes/Engine/ChildConnection.h b/Engine/Source/Runtime/Engine/Classes/Engine/ChildConnection.h
--- a/Engine/Source/Runtime/Engine/Classes/Engine/ChildConnection.h
+++ b/Engine/Source/Runtime/Engine/Classes/Engine/ChildConnection.h
@@ -33,9 +33,9 @@ class UChildConnection : public UNetConnection
 	{
 	}
 
-	virtual void InitOut() OVERRIDE
+	virtual void InitSendBuffer() OVERRIDE
 	{
-		Parent->InitOut();
+		Parent->InitSendBuffer();
 	}
 
 	virtual void AssertValid() OVERRIDE
diff --git a/Engine/Source/Runtime/Engine/Classes/Engine/NetConnection.h b/Engine/Source/Runtime/Engine/Classes/Engine/NetConnection.h
--- a/Engine/Source/Runtime/Engine/Classes/Engine/NetConnection.h
+++ b/Engine/Source/Runtime/Engine/Classes/Engine/NetConnection.h
@@ -206,7 +206,7 @@ public:
 	int32 InPacketsLost, OutPacketsLost;
 
 	// Packet.
-	FBitWriter		Out;					// Outgoing packet.
+	FBitWriter		SendBuffer;				// Queued up bits waiting to send
 	double			OutLagTime[256];		// For lag measuring.
 	int32			OutLagPacketId[256];	// For lag measuring.
 	int32			InPacketId;				// Full incoming packet index.
@@ -355,11 +355,8 @@ public:
 	 */
 	virtual void LowLevelSend( void* Data, int32 Count ) PURE_VIRTUAL(UNetConnection::LowLevelSend,); //!! "Looks like an FArchive"
 
-	/** Validates the FBitWriter to make sure it's not in an error state */
-	ENGINE_API virtual void ValidateOut();
-
 	/** Resets the FBitWriter to its default state */
-	ENGINE_API virtual void InitOut();
+	ENGINE_API virtual void InitSendBuffer();
 
 	/** Make sure this connection is in a reasonable state. */
 	ENGINE_API virtual void AssertValid();
@@ -461,14 +458,21 @@ public:
 	/** Send package map to the remote. */
 	void SendPackageMap();
 
-	/** Called before sending anything. */
-	void PreSend( int32 SizeBits );
-
-	/** Called after sending anything. */
-	void PostSend();
+	/** 
+	 * Appends the passed in data to the SendBuffer to be sent when FlushNet is called
+	 * @param Bits Data as bits to be appended to the send buffer
+	 * @param SizeInBits Number of bits to append
+	 * @param ExtraBits (optional) Second set of bits to be appended to the send buffer that need to send with the first set of bits
+	 * @param ExtraSizeInBits (optional) Number of secondary bits to append
+	 */
+	void WriteBitsToSendBuffer( 
+		const uint8 *	Bits, 
+		const int32		SizeInBits, 
+		const uint8 *	ExtraBits = NULL, 
+		const int32		ExtraSizeInBits = 0 );
 
 	/** Returns number of bits left in current packet that can be used without causing a flush.  */
-	int64 GetRemainingPacketBits();
+	int64 GetFreeSendBufferBits();
 
 	/** 
 	 * parses the passed in control channel bunch and fills the given package info struct with that data
diff --git a/Engine/Source/Runtime/Engine/Private/DataChannel.cpp b/Engine/Source/Runtime/Engine/Private/DataChannel.cpp
--- a/Engine/Source/Runtime/Engine/Private/DataChannel.cpp
+++ b/Engine/Source/Runtime/Engine/Private/DataChannel.cpp
@@ -206,13 +206,7 @@ void UChannel::ReceivedAcks()
 
 int32 UChannel::MaxSendBytes()
 {
-	int32 ResultBits
-	=	Connection->MaxPacket*8
-	-	(Connection->Out.GetNumBits() ? 0 : MAX_PACKET_HEADER_BITS)
-	-	Connection->Out.GetNumBits()
-	-	MAX_PACKET_TRAILER_BITS
-	-	MAX_BUNCH_HEADER_BITS;
-	return FMath::Max( 0, ResultBits/8 );
+	return FMath::Max( 0, Connection->GetFreeSendBufferBits() / 8 );
 }
 
 
@@ -561,8 +555,6 @@ FPacketIdRange UChannel::SendBunch( FOutBunch* Bunch, bool Merge )
 		Merge = false;
 	}
 
-	Connection->ValidateOut();
-
 	//-----------------------------------------------------
 	// Contemplate merging.
 	//-----------------------------------------------------
@@ -573,8 +565,8 @@ FPacketIdRange UChannel::SendBunch( FOutBunch* Bunch, bool Merge )
 	&&	Connection->LastOut.ChIndex == Bunch->ChIndex
 	&&	Connection->AllowMerge
 	&&	Connection->LastEnd.GetNumBits()
-	&&	Connection->LastEnd.GetNumBits()==Connection->Out.GetNumBits()
-	&&	Connection->Out.GetNumBytes()+Bunch->GetNumBytes()+(MAX_BUNCH_HEADER_BITS+MAX_PACKET_TRAILER_BITS+7)/8<=Connection->MaxPacket )
+	&&	Connection->LastEnd.GetNumBits()==Connection->SendBuffer.GetNumBits()
+	&&	Connection->LastOut.GetNumBits() + Bunch->GetNumBits() <= MAX_SINGLE_BUNCH_SIZE_BITS )
 	{
 		// Merge.
 		check(!Connection->LastOut.IsError());
@@ -586,12 +578,10 @@ FPacketIdRange UChannel::SendBunch( FOutBunch* Bunch, bool Merge )
 		OutBunch                       = Connection->LastOutBunch;
 		Bunch                          = &Connection->LastOut;
 		check(!Bunch->IsError());
-		Connection->LastStart.Pop( Connection->Out );
+		Connection->LastStart.Pop( Connection->SendBuffer );
 		Connection->Driver->OutBunches--;
 	}
 
-	Connection->ValidateOut();
-
 	//-----------------------------------------------------
 	// Possibly split large bunch into list of smaller partial bunches
 	//-----------------------------------------------------
@@ -697,7 +687,7 @@ FPacketIdRange UChannel::SendBunch( FOutBunch* Bunch, bool Merge )
 
 		// Update channel sequence count.
 		Connection->LastOut = *ThisOutBunch;
-		Connection->LastEnd	= FBitWriterMark(Connection->Out);
+		Connection->LastEnd	= FBitWriterMark( Connection->SendBuffer );
 	}
 
 	// Update open range if necessary
diff --git a/Engine/Source/Runtime/Engine/Private/NetConnection.cpp b/Engine/Source/Runtime/Engine/Private/NetConnection.cpp
--- a/Engine/Source/Runtime/Engine/Private/NetConnection.cpp
+++ b/Engine/Source/Runtime/Engine/Private/NetConnection.cpp
@@ -49,7 +49,7 @@ UNetConnection::UNetConnection(const class FPostConstructInitializeProperties& P
 ,	AverageFrameTime	( 0 )
 ,	CountedFrames		( 0 )
 
-,	Out					( 0 )
+,	SendBuffer			( 0 )
 ,	InPacketId			( -1 )
 ,	OutPacketId			( 0 ) // must be initialized as OutAckPacketId + 1 so loss of first packet can be detected
 ,	OutAckPacketId		( -1 )
@@ -440,30 +440,20 @@ bool UNetConnection::ClientHasInitializedLevelFor(const UObject* TestObject)
 			ClientVisibleLevelNames.Contains(Level->GetOutermost()->GetFName()) );
 }
 
-void UNetConnection::ValidateOut()
-{
-	if ( Out.IsError() )
-	{
-		UE_LOG( LogNetTraffic, Fatal, TEXT( "UNetConnection::ValidateOut: Out.IsError() == true. NumBits: %i, NumBytes: %i, MaxBits: %i" ), Out.GetNumBits(), Out.GetNumBytes(), Out.GetMaxBits() );
-	}
-}
-
-void UNetConnection::InitOut()
+void UNetConnection::InitSendBuffer()
 {
 	check(MaxPacket > 0);
 	// Initialize the one outgoing buffer.
-	if (MaxPacket * 8 == Out.GetMaxBits())
+	if (MaxPacket * 8 == SendBuffer.GetMaxBits())
 	{
 		// Reset all of our values to their initial state without a malloc/free
-		Out.Reset();
+		SendBuffer.Reset();
 	}
 	else
 	{
 		// First time initialization needs to allocate the buffer
-		Out = FBitWriter(MaxPacket * 8);
+		SendBuffer = FBitWriter(MaxPacket * 8);
 	}
-
-	ValidateOut();
 }
 
 void UNetConnection::ReceivedRawPacket( void* InData, int32 Count )
@@ -507,27 +497,28 @@ void UNetConnection::ReceivedRawPacket( void* InData, int32 Count )
 void UNetConnection::FlushNet(bool bIgnoreSimulation)
 {
 	// Update info.
-	ValidateOut();
+	check( !SendBuffer.IsError() );
+
 	LastEnd = FBitWriterMark();
 	TimeSensitive = 0;
 
 	// If there is any pending data to send, send it.
-	if( Out.GetNumBits() || Driver->Time-LastSendTime>Driver->KeepAliveTime )
+	if( SendBuffer.GetNumBits() || Driver->Time-LastSendTime>Driver->KeepAliveTime )
 	{
-		// If sending keepalive packet, still generate header.
-		if( Out.GetNumBits()==0 )
+		// If sending keepalive packet, still write the packet id
+		if ( SendBuffer.GetNumBits() == 0 )
 		{
-			PreSend( 0 );
+			WriteBitsToSendBuffer( NULL, 0 );		// This will force the packet id to be written
 		}
 
 		// Make sure packet size is byte-aligned.
-		Out.WriteBit( 1 );
-		ValidateOut();
-		while( Out.GetNumBits() & 7 )
+		SendBuffer.WriteBit( 1 );
+		while( SendBuffer.GetNumBits() & 7 )
 		{
-			Out.WriteBit( 0 );
+			SendBuffer.WriteBit( 0 );
 		}
-		ValidateOut();
+
+		check( !SendBuffer.IsError() );
 
 		// Send now.
 #if DO_ENABLE_NET_TEST
@@ -538,14 +529,14 @@ void UNetConnection::FlushNet(bool bIgnoreSimulation)
 			// Checked in FlushNet() so each child class doesn't have to implement this
 			if (Driver->IsNetResourceValid())
 			{
-				LowLevelSend(Out.GetData(), Out.GetNumBytes());
+				LowLevelSend(SendBuffer.GetData(), SendBuffer.GetNumBytes());
 			}
 		}
 		else if( PacketSimulationSettings.PktOrder )
 		{
 			DelayedPacket& B = *(new(Delayed)DelayedPacket);
-			B.Data.AddUninitialized( Out.GetNumBytes() );
-			FMemory::Memcpy( B.Data.GetTypedData(), Out.GetData(), Out.GetNumBytes() );
+			B.Data.AddUninitialized( SendBuffer.GetNumBytes() );
+			FMemory::Memcpy( B.Data.GetTypedData(), SendBuffer.GetData(), SendBuffer.GetNumBytes() );
 
 			for( int32 i=Delayed.Num()-1; i>=0; i-- )
 			{
@@ -568,8 +559,8 @@ void UNetConnection::FlushNet(bool bIgnoreSimulation)
 			if( !PacketSimulationSettings.PktLoss || FMath::FRand()*100.f > PacketSimulationSettings.PktLoss )
 			{
 				DelayedPacket& B = *(new(Delayed)DelayedPacket);
-				B.Data.AddUninitialized( Out.GetNumBytes() );
-				FMemory::Memcpy( B.Data.GetTypedData(), Out.GetData(), Out.GetNumBytes() );
+				B.Data.AddUninitialized( SendBuffer.GetNumBytes() );
+				FMemory::Memcpy( B.Data.GetTypedData(), SendBuffer.GetData(), SendBuffer.GetNumBytes() );
 				B.SendTime = FPlatformTime::Seconds() + (double(PacketSimulationSettings.PktLag)  + 2.0f * (FMath::FRand() - 0.5f) * double(PacketSimulationSettings.PktLagVariance))/ 1000.f;
 			}
 		}
@@ -579,7 +570,7 @@ void UNetConnection::FlushNet(bool bIgnoreSimulation)
 			// Checked in FlushNet() so each child class doesn't have to implement this
 			if (Driver->IsNetResourceValid())
 			{
-				LowLevelSend( Out.GetData(), Out.GetNumBytes() );
+				LowLevelSend( SendBuffer.GetData(), SendBuffer.GetNumBytes() );
 			}
 #if DO_ENABLE_NET_TEST
 			if( PacketSimulationSettings.PktDup && FMath::FRand()*100.f < PacketSimulationSettings.PktDup )
@@ -587,23 +578,23 @@ void UNetConnection::FlushNet(bool bIgnoreSimulation)
 				// Checked in FlushNet() so each child class doesn't have to implement this
 				if (Driver->IsNetResourceValid())
 				{
-					LowLevelSend( (char*)Out.GetData(), Out.GetNumBytes() );
+					LowLevelSend( (char*)SendBuffer.GetData(), SendBuffer.GetNumBytes() );
 				}
 			}
 		}
 #endif
 		// Update stuff.
-		int32 Index = OutPacketId & (ARRAY_COUNT(OutLagPacketId)-1);
+		const int32 Index = OutPacketId & (ARRAY_COUNT(OutLagPacketId)-1);
 		OutLagPacketId [Index] = OutPacketId;
 		OutLagTime     [Index] = Driver->Time;
 		OutPacketId++;
 		Driver->OutPackets++;
 		LastSendTime = Driver->Time;
-		int32 PacketBytes = Out.GetNumBytes() + PacketOverhead;
+		const int32 PacketBytes = SendBuffer.GetNumBytes() + PacketOverhead;
 		QueuedBytes += PacketBytes;
 		OutBytes += PacketBytes;
 		Driver->OutBytes += PacketBytes;
-		InitOut();
+		InitSendBuffer();
 	}
 
 	// Move acks around.
@@ -618,8 +609,8 @@ int32 UNetConnection::IsNetReady( bool Saturate )
 {
 	// Return whether we can send more data without saturation the connection.
 	if( Saturate )
-		QueuedBytes = -Out.GetNumBytes();
-	return QueuedBytes+Out.GetNumBytes() <= 0;
+		QueuedBytes = -SendBuffer.GetNumBytes();
+	return QueuedBytes+SendBuffer.GetNumBytes() <= 0;
 }
 
 void UNetConnection::ReadInput( float DeltaSeconds )
@@ -667,8 +658,6 @@ void UNetConnection::ReceivedPacket( FBitReader& Reader )
 		return;
 	}
 
-	ValidateOut();
-
 	// Update receive time to avoid timeout.
 	LastReceiveTime = Driver->Time;
 
@@ -1010,8 +999,6 @@ void UNetConnection::ReceivedPacket( FBitReader& Reader )
 		}
 	}
 
-	ValidateOut();
-
 	// Acknowledge the packet.
 	if ( !bSkipAck )
 	{
@@ -1019,58 +1006,80 @@ void UNetConnection::ReceivedPacket( FBitReader& Reader )
 	}
 }
 
-void UNetConnection::PreSend( int32 SizeBits )
+void UNetConnection::WriteBitsToSendBuffer( 
+	const uint8 *	Bits, 
+	const int32		SizeInBits, 
+	const uint8 *	ExtraBits, 
+	const int32		ExtraSizeInBits )
 {
-	ValidateOut();
+	check( !SendBuffer.IsError() );
+
+	const int32 TotalSizeInBits = SizeInBits + ExtraSizeInBits;
 
-	// Flush if not enough space.
-	if( GetRemainingPacketBits() < SizeBits)
+	// Flush if we can't add to current buffer
+	if ( TotalSizeInBits > GetFreeSendBufferBits() )
+	{
 		FlushNet();
+	}
+
+	// Remember start position in case we want to undo this write
+	LastStart = FBitWriterMark( SendBuffer );
 
-	// If start of packet, send packet id.
-	if( Out.GetNumBits()==0 )
+	// If this is the start of the queue, make sure to add the packet id
+	if ( SendBuffer.GetNumBits() == 0 )
 	{
-		Out.WriteIntWrapped(OutPacketId, MAX_PACKETID);
-		check(Out.GetNumBits()<=MAX_PACKET_HEADER_BITS);
-		ValidateOut();
+		SendBuffer.WriteIntWrapped( OutPacketId, MAX_PACKETID );
+		check( !SendBuffer.IsError() );
 	}
 
-	// Make sure there's enough space now.
-	if( Out.GetNumBits() + SizeBits + MAX_PACKET_TRAILER_BITS > MaxPacket*8 )
+	// Add the bits to the queue
+	if ( SizeInBits )
 	{
-		UE_LOG(LogNet, Fatal, TEXT("PreSend overflowed: %i+%i>%i"), Out.GetNumBits(), SizeBits, MaxPacket*8 );
+		SendBuffer.SerializeBits( const_cast< uint8 * >( Bits ), SizeInBits );
+		check( !SendBuffer.IsError() );
 	}
-}
 
-/** Returns number of bits left in current packet that can be used without causing a flush  */
-int64 UNetConnection::GetRemainingPacketBits()
-{
-	return (MaxPacket*8) - (Out.GetNumBits() + MAX_PACKET_TRAILER_BITS);
-}
+	// Add any extra bits
+	if ( ExtraSizeInBits )
+	{
+		SendBuffer.SerializeBits( const_cast< uint8 * >( ExtraBits ), ExtraSizeInBits );
+		check( !SendBuffer.IsError() );
+	}
 
-void UNetConnection::PostSend()
-{
-	// If absolutely filled now, flush so that MaxSend() doesn't return zero unnecessarily.
-	check(Out.GetNumBits()<=MaxPacket*8);
-	if( Out.GetNumBits()==MaxPacket*8 )
+	// Flush now if we are full
+	if ( GetFreeSendBufferBits() == 0 )
 	{
 		FlushNet();
 	}
 }
 
+/** Returns number of bits left in current packet that can be used without causing a flush  */
+int64 UNetConnection::GetFreeSendBufferBits()
+{
+	// If we haven't sent anything yet, make sure to account for the packet header + trailer size
+	// Otherwise, we only need to account for trailer size
+	const int32 ExtraBits = ( SendBuffer.GetNumBits() > 0 ) ? MAX_PACKET_TRAILER_BITS : MAX_PACKET_HEADER_BITS + MAX_PACKET_TRAILER_BITS;
+
+	const int32 NumberOfFreeBits = ( MaxPacket * 8 ) - ( SendBuffer.GetNumBits() + ExtraBits );
+
+	check( NumberOfFreeBits >= 0 );
+
+	return NumberOfFreeBits;
+}
 
 void UNetConnection::PurgeAcks()
 {
-	for( int32 i=0; i<ResendAcks.Num(); i++ )
+	for ( int32 i = 0; i < ResendAcks.Num(); i++ )
+	{
 		SendAck( ResendAcks[i], 0 );
+	}
+
 	ResendAcks.Empty(32);
 }
 
 
 void UNetConnection::SendAck( int32 AckPacketId, bool FirstTime/*=1*/, bool bHavePingAckData/*=0*/, uint32 PingAckData/*=0*/ )
 {
-	ValidateOut();
-
 	if( !InternalAck )
 	{
 		if( FirstTime )
@@ -1079,35 +1088,26 @@ void UNetConnection::SendAck( int32 AckPacketId, bool FirstTime/*=1*/, bool bHav
 			QueuedAcks.Add(AckPacketId);
 		}
 
-		bool bPingAck = Driver->ServerConnection != NULL && (AckPacketId % PING_ACK_PACKET_INTERVAL) == 0;
-		int32 SendSize = FMath::CeilLogTwo(MAX_PACKETID) + 1;
-		
-		if (bPingAck)
-		{
-			const int32 PING_ACK_DATA_NUM_BITS = sizeof(uint32) * 8;
-			SendSize += 1 + (bHavePingAckData ? PING_ACK_DATA_NUM_BITS : 0);
-		}
+		FBitWriter AckData( 32, true );
 
-		PreSend(SendSize);
-		Out.WriteBit( 1 );
-		ValidateOut();
-		Out.WriteIntWrapped(AckPacketId, MAX_PACKETID);
-		ValidateOut();
+		AckData.WriteBit( 1 );
+		AckData.WriteIntWrapped(AckPacketId, MAX_PACKETID);
 
-		if (bPingAck)
+		const bool bPingAck = Driver->ServerConnection != NULL && (AckPacketId % PING_ACK_PACKET_INTERVAL) == 0;
+
+		if ( bPingAck )
 		{
-			Out.WriteBit(bHavePingAckData);
-			ValidateOut();
+			AckData.WriteBit( bHavePingAckData );
 
-			if (bHavePingAckData)
+			if ( bHavePingAckData )
 			{
-				Out.Serialize(&PingAckData, sizeof(uint32));
-				ValidateOut();
+				AckData.Serialize( &PingAckData, sizeof( uint32 ) );
 			}
 		}
 
+		WriteBitsToSendBuffer( AckData.GetData(), AckData.GetNumBits() );
+
 		AllowMerge = false;
-		PostSend();
 
 		UE_LOG(LogNetTraffic, Log, TEXT("   Send ack %i"), AckPacketId);
 	}
@@ -1116,7 +1116,6 @@ void UNetConnection::SendAck( int32 AckPacketId, bool FirstTime/*=1*/, bool bHav
 
 int32 UNetConnection::SendRawBunch( FOutBunch& Bunch, bool InAllowMerge )
 {
-	ValidateOut();
 	check(!Bunch.ReceivedAck);
 	check(!Bunch.IsError());
 	Driver->OutBunches++;
@@ -1156,9 +1155,6 @@ int32 UNetConnection::SendRawBunch( FOutBunch& Bunch, bool InAllowMerge )
 	Header.WriteIntWrapped(Bunch.GetNumBits(), UNetConnection::MaxPacket * 8);
 	check(!Header.IsError());
 
-	// If this data doesn't fit in the current packet, flush it.
-	PreSend( Header.GetNumBits() + Bunch.GetNumBits() );
-
 	// Remember start position.
 	AllowMerge      = InAllowMerge;
 	Bunch.PacketId  = OutPacketId;
@@ -1176,16 +1172,7 @@ int32 UNetConnection::SendRawBunch( FOutBunch& Bunch, bool InAllowMerge )
 
 	UE_LOG(LogNetTraffic, Verbose, TEXT("UNetConnection::SendRawBunch. ChIndex: %d. Bits: %d. PacketId: %d"), Bunch.ChIndex, Bunch.GetNumBits(), OutPacketId );
 
-	// Remember start position, and write data.
-	LastStart = FBitWriterMark( Out );
-
-	Out.SerializeBits( Header.GetData(), Header.GetNumBits() );
-
-	ValidateOut();
-
-	Out.SerializeBits( Bunch .GetData(), Bunch .GetNumBits() );
-
-	ValidateOut();
+	WriteBitsToSendBuffer( Header.GetData(), Header.GetNumBits(), Bunch.GetData(), Bunch.GetNumBits() );
 
 	if ( PackageMap && Bunch.bHasGUIDs )
 	{
@@ -1197,7 +1184,6 @@ int32 UNetConnection::SendRawBunch( FOutBunch& Bunch, bool InAllowMerge )
 		Driver->NetGUIDOutBytes += (Header.GetNumBits() + Bunch.GetNumBits()) >> 3;
 	}
 
-
 	// Verified client ping tracking - caches some semi-random bytes of the packet, for ping validation
 	if (Driver->ServerConnection == NULL && OutPacketId != LastPingAckPacketId && (OutPacketId % PING_ACK_PACKET_INTERVAL) == 0 &&
 		Bunch.GetNumBits() >= 32)
@@ -1212,10 +1198,6 @@ int32 UNetConnection::SendRawBunch( FOutBunch& Bunch, bool InAllowMerge )
 		LastPingAckPacketId = OutPacketId;
 	}
 
-
-	// Finished.
-	PostSend();
-
 	return Bunch.PacketId;
 }
 
diff --git a/Engine/Source/Runtime/Online/OnlineSubsystemUtils/Private/IpConnection.cpp b/Engine/Source/Runtime/Online/OnlineSubsystemUtils/Private/IpConnection.cpp
--- a/Engine/Source/Runtime/Online/OnlineSubsystemUtils/Private/IpConnection.cpp
+++ b/Engine/Source/Runtime/Online/OnlineSubsystemUtils/Private/IpConnection.cpp
@@ -71,7 +71,7 @@ void UIpConnection::InitLocalConnection(UNetDriver* InDriver, class FSocket* InS
 	}
 
 	// Initialize our send bunch
-	InitOut();
+	InitSendBuffer();
 }
 
 void UIpConnection::InitRemoteConnection(UNetDriver* InDriver, class FSocket* InSocket, const FURL& InURL, const class FInternetAddr& InRemoteAddr, EConnectionState InState, int32 InMaxPacket, int32 InPacketOverhead)
@@ -91,7 +91,7 @@ void UIpConnection::InitRemoteConnection(UNetDriver* InDriver, class FSocket* In
 	URL.Host = RemoteAddr->ToString(false);
 
 	// Initialize our send bunch
-	InitOut();
+	InitSendBuffer();
 
 	// This is for a client that needs to log in, setup ClientLoginState and ExpectedClientLoginMsgType to reflect that
 	SetClientLoginState( EClientLoginState::LoggingIn );
