#UE4 - Compile fix

diff --git a/Engine/Source/Runtime/Engine/Classes/Engine/Channel.h b/Engine/Source/Runtime/Engine/Classes/Engine/Channel.h
--- a/Engine/Source/Runtime/Engine/Classes/Engine/Channel.h
+++ b/Engine/Source/Runtime/Engine/Classes/Engine/Channel.h
@@ -110,9 +110,6 @@ class UChannel : public UObject
 	/** Make sure the incoming buffer is in sequence and there are no duplicates. */
 	void AssertInSequenced();
 
-	/** Return the maximum amount of data that can be sent in this bunch without overflow. */
-	int32 MaxSendBytes();
-
 	/** cleans up channel if it hasn't already been */
 	void ConditionalCleanUp()
 	{
diff --git a/Engine/Source/Runtime/Engine/Private/DataChannel.cpp b/Engine/Source/Runtime/Engine/Private/DataChannel.cpp
--- a/Engine/Source/Runtime/Engine/Private/DataChannel.cpp
+++ b/Engine/Source/Runtime/Engine/Private/DataChannel.cpp
@@ -203,13 +203,6 @@ void UChannel::ReceivedAcks()
 
 }
 
-
-int32 UChannel::MaxSendBytes()
-{
-	return FMath::Max( 0, Connection->GetFreeSendBufferBits() / 8 );
-}
-
-
 void UChannel::Tick()
 {
 	checkSlow(Connection->Channels[ChIndex]==this);
