#UE4 - Cleanup, removed static keyword on variable that can change

diff --git a/Engine/Source/Runtime/Engine/Private/DataChannel.cpp b/Engine/Source/Runtime/Engine/Private/DataChannel.cpp
--- a/Engine/Source/Runtime/Engine/Private/DataChannel.cpp
+++ b/Engine/Source/Runtime/Engine/Private/DataChannel.cpp
@@ -531,13 +531,13 @@ FPacketIdRange UChannel::SendBunch( FOutBunch* Bunch, bool Merge )
 	check(!OpenTemporary || !Bunch->bReliable);
 
 	// This is the max number of bits we can have in a single bunch
-	static const int64 MAX_SINGLE_BUNCH_SIZE_BITS  = Connection->MaxPacket*8-MAX_BUNCH_HEADER_BITS-MAX_PACKET_TRAILER_BITS-MAX_PACKET_HEADER_BITS;
+	const int64 MAX_SINGLE_BUNCH_SIZE_BITS  = Connection->MaxPacket*8-MAX_BUNCH_HEADER_BITS-MAX_PACKET_TRAILER_BITS-MAX_PACKET_HEADER_BITS;
 
 	// Max bytes we'll put in a partial bunch
-	static const int64 MAX_SINGLE_BUNCH_SIZE_BYTES = MAX_SINGLE_BUNCH_SIZE_BITS / 8;
+	const int64 MAX_SINGLE_BUNCH_SIZE_BYTES = MAX_SINGLE_BUNCH_SIZE_BITS / 8;
 
 	// Max bits will put in a partial bunch (byte aligned, we dont want to deal with partial bytes in the partial bunches)
-	static const int64 MAX_PARTIAL_BUNCH_SIZE_BITS = MAX_SINGLE_BUNCH_SIZE_BYTES * 8;
+	const int64 MAX_PARTIAL_BUNCH_SIZE_BITS = MAX_SINGLE_BUNCH_SIZE_BYTES * 8;
 
 	TArray<FOutBunch *> OutgoingBunches;
 
@@ -651,7 +651,7 @@ FPacketIdRange UChannel::SendBunch( FOutBunch* Bunch, bool Merge )
 			NextBunch->bClose = (Bunch->bClose && (OutgoingBunches.Num()-1 == PartialNum)); // Only last bunch should have bClose bit set
 		}
 
-		FOutBunch *ThisOutBunch = PrepBunch(NextBunch, OutBunch, Merge); // This handles queueing reliable bunches into the ack list
+		FOutBunch *ThisOutBunch = PrepBunch(NextBunch, OutBunch, Merge); // This handles queuing reliable bunches into the ack list
 			
 		if (ThisOutBunch->bPartial && !ThisOutBunch->bReliable)
 		{
