#UE4 - Fix a bug that was causing the wrong packet id to be assigned to the bunch

diff --git a/Engine/Source/Runtime/Engine/Classes/Engine/NetConnection.h b/Engine/Source/Runtime/Engine/Classes/Engine/NetConnection.h
--- a/Engine/Source/Runtime/Engine/Classes/Engine/NetConnection.h
+++ b/Engine/Source/Runtime/Engine/Classes/Engine/NetConnection.h
@@ -465,7 +465,7 @@ public:
 	 * @param ExtraBits (optional) Second set of bits to be appended to the send buffer that need to send with the first set of bits
 	 * @param ExtraSizeInBits (optional) Number of secondary bits to append
 	 */
-	void WriteBitsToSendBuffer( 
+	int32 WriteBitsToSendBuffer( 
 		const uint8 *	Bits, 
 		const int32		SizeInBits, 
 		const uint8 *	ExtraBits = NULL, 
diff --git a/Engine/Source/Runtime/Engine/Private/NetConnection.cpp b/Engine/Source/Runtime/Engine/Private/NetConnection.cpp
--- a/Engine/Source/Runtime/Engine/Private/NetConnection.cpp
+++ b/Engine/Source/Runtime/Engine/Private/NetConnection.cpp
@@ -1006,7 +1006,7 @@ void UNetConnection::ReceivedPacket( FBitReader& Reader )
 	}
 }
 
-void UNetConnection::WriteBitsToSendBuffer( 
+int32 UNetConnection::WriteBitsToSendBuffer( 
 	const uint8 *	Bits, 
 	const int32		SizeInBits, 
 	const uint8 *	ExtraBits, 
@@ -1023,6 +1023,7 @@ void UNetConnection::WriteBitsToSendBuffer(
 	}
 
 	// Remember start position in case we want to undo this write
+	// Store this after the possible flush above so we have the correct start position in the case that we do flush
 	LastStart = FBitWriterMark( SendBuffer );
 
 	// If this is the start of the queue, make sure to add the packet id
@@ -1046,11 +1047,15 @@ void UNetConnection::WriteBitsToSendBuffer(
 		check( !SendBuffer.IsError() );
 	}
 
+	const int32 RememberedPacketId = OutPacketId;
+
 	// Flush now if we are full
 	if ( GetFreeSendBufferBits() == 0 )
 	{
 		FlushNet();
 	}
+
+	return RememberedPacketId;
 }
 
 /** Returns number of bits left in current packet that can be used without causing a flush  */
@@ -1157,7 +1162,6 @@ int32 UNetConnection::SendRawBunch( FOutBunch& Bunch, bool InAllowMerge )
 
 	// Remember start position.
 	AllowMerge      = InAllowMerge;
-	Bunch.PacketId  = OutPacketId;
 	Bunch.Time      = Driver->Time;
 
 	if ((Bunch.bClose || Bunch.bOpen) && UE_LOG_ACTIVE(LogNetDormancy,VeryVerbose) )
@@ -1170,9 +1174,10 @@ int32 UNetConnection::SendRawBunch( FOutBunch& Bunch, bool InAllowMerge )
 		UE_LOG(LogNetTraffic, VeryVerbose, TEXT("Sending: %s"), *Bunch.ToString());
 	}
 
-	UE_LOG(LogNetTraffic, Verbose, TEXT("UNetConnection::SendRawBunch. ChIndex: %d. Bits: %d. PacketId: %d"), Bunch.ChIndex, Bunch.GetNumBits(), OutPacketId );
+	// Write the bits to the buffer and remember the packet id used
+	Bunch.PacketId = WriteBitsToSendBuffer( Header.GetData(), Header.GetNumBits(), Bunch.GetData(), Bunch.GetNumBits() );
 
-	WriteBitsToSendBuffer( Header.GetData(), Header.GetNumBits(), Bunch.GetData(), Bunch.GetNumBits() );
+	UE_LOG(LogNetTraffic, Verbose, TEXT("UNetConnection::SendRawBunch. ChIndex: %d. Bits: %d. PacketId: %d"), Bunch.ChIndex, Bunch.GetNumBits(), Bunch.PacketId );
 
 	if ( PackageMap && Bunch.bHasGUIDs )
 	{
@@ -1185,17 +1190,17 @@ int32 UNetConnection::SendRawBunch( FOutBunch& Bunch, bool InAllowMerge )
 	}
 
 	// Verified client ping tracking - caches some semi-random bytes of the packet, for ping validation
-	if (Driver->ServerConnection == NULL && OutPacketId != LastPingAckPacketId && (OutPacketId % PING_ACK_PACKET_INTERVAL) == 0 &&
+	if (Driver->ServerConnection == NULL && Bunch.PacketId != LastPingAckPacketId && (Bunch.PacketId % PING_ACK_PACKET_INTERVAL) == 0 &&
 		Bunch.GetNumBits() >= 32)
 	{
-		const uint32 PingAckData = CalcPingAckData( Bunch.GetData(), Bunch.GetNumBits(), OutPacketId );
+		const uint32 PingAckData = CalcPingAckData( Bunch.GetData(), Bunch.GetNumBits(), Bunch.PacketId );
 
-		const int32 PingAckIdx = (OutPacketId % MAX_PACKETID) / PING_ACK_PACKET_INTERVAL;
+		const int32 PingAckIdx = (Bunch.PacketId % MAX_PACKETID) / PING_ACK_PACKET_INTERVAL;
 
 		PingAckDataCache[PingAckIdx] = PingAckData;
 
 		// Multiple bunches get written per-packet, but PingAck only uses the first bunch, so don't process more bunches this PacketId
-		LastPingAckPacketId = OutPacketId;
+		LastPingAckPacketId = Bunch.PacketId;
 	}
 
 	return Bunch.PacketId;
