Back out changelist 2042107
Fixes TTP# 332266 Lightmass broken in several ways due to AActor::GetComponentsBoundingBox  change

diff --git a/Engine/Source/Runtime/Engine/Private/Actor.cpp b/Engine/Source/Runtime/Engine/Private/Actor.cpp
--- a/Engine/Source/Runtime/Engine/Private/Actor.cpp
+++ b/Engine/Source/Runtime/Engine/Private/Actor.cpp
@@ -848,7 +848,7 @@ FBox AActor::GetComponentsBoundingBox(bool bNonColliding) const
 		UPrimitiveComponent* PrimComp = Components[CompIdx];
 
 		// Only use collidable components to find collision bounding box.
-		if( PrimComp->IsRegistered() && ((bNonColliding && !PrimComp->bHiddenInGame) || PrimComp->IsCollisionEnabled()) )
+		if( PrimComp->IsRegistered() && (bNonColliding || PrimComp->IsCollisionEnabled()) )
 		{
 			Box += PrimComp->Bounds.GetBox();
 		}
