322887 - ANIM: Add in game animation node graph debug output

#proj UE4

#summary Added debug display of active anim graph

diff --git a/Engine/Config/BaseInput.ini b/Engine/Config/BaseInput.ini
--- a/Engine/Config/BaseInput.ini
+++ b/Engine/Config/BaseInput.ini
@@ -137,8 +137,17 @@ HistoryBot=-1
 +ManualAutoCompleteList=(Command="ShowDebug AI",Desc="ShowDebug AI")
 +ManualAutoCompleteList=(Command="ShowDebug CAMERA",Desc="ShowDebug CAMERA")
 +ManualAutoCompleteList=(Command="ShowDebug WEAPON",Desc="ShowDebug WEAPON")
-+ManualAutoCompleteList=(Command="ShowDebug ANIMATION",Desc="ShowDebug ANIMATION")
++ManualAutoCompleteList=(Command="ShowDebug ANIMATION",Desc="ShowDebug ANIMATION (Toggles display state of animation debug data)")
++ManualAutoCompleteList=(Command="ShowDebug BONES",Desc="ShowDebug BONES (Toggles display of skeletalmesh bones)")
 +ManualAutoCompleteList=(Command="ShowDebug INPUT",Desc="ShowDebug INPUT")
++ManualAutoCompleteList=(Command="ShowDebugToggleSubCategory 3DBONES",Desc="ShowDebugToggleSubCategory 3DBONES (With ShowDebug Bones: Toggles bone rendering between single lines and a more complex 3D model per bone)")
++ManualAutoCompleteList=(Command="ShowDebugToggleSubCategory SYNCGROUPS",Desc="ShowDebugToggleSubCategory SYNCGROUPS (With ShowDebug Animation: Toggles display of sync group data)")
++ManualAutoCompleteList=(Command="ShowDebugToggleSubCategory MONTAGES",Desc="ShowDebugToggleSubCategory MONTAGES (With ShowDebug Animation: Toggles display of montage data)")
++ManualAutoCompleteList=(Command="ShowDebugToggleSubCategory GRAPH",Desc="ShowDebugToggleSubCategory GRAPH (With ShowDebug Animation: Toggles display of animation blueprint graph)")
++ManualAutoCompleteList=(Command="ShowDebugToggleSubCategory CURVES",Desc="ShowDebugToggleSubCategory CURVES (With ShowDebug Animation: Toggles display of curve data)")
++ManualAutoCompleteList=(Command="ShowDebugToggleSubCategory NOTIFIES",Desc="ShowDebugToggleSubCategory NOTIFIES(With ShowDebug Animation: Toggles display of notify data)")
++ManualAutoCompleteList=(Command="ShowDebugToggleSubCategory FULLGRAPH",Desc="ShowDebugToggleSubCategory FULLGRAPH (With ShowDebug Animation: Toggles graph display between active nodes only and all nodes)")
++ManualAutoCompleteList=(Command="ShowDebugToggleSubCategory FULLBLENDSPACEDISPLAY",Desc="ShowDebugToggleSubCategory FULLBLENDSPACEDISPLAY (With ShowDebug Animation: Toggles display of sample blend weights on blendspaces)")
 +ManualAutoCompleteList=(Command="togglehdwarning",Desc=)
 +ManualAutoCompleteList=(Command="Stat SOUNDMIXES",Desc="Stat SOUNDMIXES (Shows active SoundMixes)")
 +ManualAutoCompleteList=(Command="Stat SOUNDWAVES",Desc="Stat SOUNDWAVES (Shows active SoundWaves)")
diff --git a/Engine/Source/Runtime/Engine/Classes/AI/Navigation/PathFollowingComponent.h b/Engine/Source/Runtime/Engine/Classes/AI/Navigation/PathFollowingComponent.h
--- a/Engine/Source/Runtime/Engine/Classes/AI/Navigation/PathFollowingComponent.h
+++ b/Engine/Source/Runtime/Engine/Classes/AI/Navigation/PathFollowingComponent.h
@@ -188,7 +188,7 @@ class ENGINE_API UPathFollowingComponent : public UActorComponent, public IAIRes
 	void DebugReachTest(float& CurrentDot, float& CurrentDistance, float& CurrentHeight, uint8& bDotFailed, uint8& bDistanceFailed, uint8& bHeightFailed); 
 
 	virtual FString GetDebugString() const;
-	virtual void DisplayDebug(class UCanvas* Canvas, const TArray<FName>& DebugDisplay, float& YL, float& YPos) const;
+	virtual void DisplayDebug(class UCanvas* Canvas, const FDebugDisplayInfo& DebugDisplay, float& YL, float& YPos) const;
 #if ENABLE_VISUAL_LOG
 	virtual void DescribeSelfToVisLog(struct FVisLogEntry* Snapshot) const;
 #endif // ENABLE_VISUAL_LOG
diff --git a/Engine/Source/Runtime/Engine/Classes/Animation/AnimInstance.h b/Engine/Source/Runtime/Engine/Classes/Animation/AnimInstance.h
--- a/Engine/Source/Runtime/Engine/Classes/Animation/AnimInstance.h
+++ b/Engine/Source/Runtime/Engine/Classes/Animation/AnimInstance.h
@@ -410,6 +410,9 @@ public:
 	// @return true if this function is implemented, false otherwise.
 	// Note: the node graph will not be evaluated if this function returns true
 	virtual bool NativeEvaluateAnimation(FPoseContext& Output);
+
+	// Debug output for this anim instance 
+	void DisplayDebug(class UCanvas* Canvas, const FDebugDisplayInfo& DebugDisplay, float& YL, float& YPos);
 public:
 
 	/** Temporary array of bone indices required this frame. Should be subset of Skeleton and Mesh's RequiredBones */
@@ -502,6 +505,9 @@ public:
 	/** Get Current Active Montage in this AnimInstance **/
 	UAnimMontage * GetCurrentActiveMontage();
 
+	/** Get Currently active montage instace **/
+	FAnimMontageInstance* GetActiveMontageInstance();
+
 protected:
 	/** Update weight of montages  **/
 	virtual void Montage_UpdateWeight(float DeltaSeconds);
@@ -509,8 +515,6 @@ protected:
 	virtual void Montage_Advance(float DeltaSeconds);
 	/** Stop all montages that are active **/
 	void StopAllMontages(float BlendOut);
-	/** Get Currently active montage instace **/
-	FAnimMontageInstance* GetActiveMontageInstance();
 	/** Called by blueprint functions that modify the montages current position. */
 	void OnMontagePositionChanged(FAnimMontageInstance* MontageInstance, FName ToSectionName);
 
diff --git a/Engine/Source/Runtime/Engine/Classes/Animation/AnimMontage.h b/Engine/Source/Runtime/Engine/Classes/Animation/AnimMontage.h
--- a/Engine/Source/Runtime/Engine/Classes/Animation/AnimMontage.h
+++ b/Engine/Source/Runtime/Engine/Classes/Animation/AnimMontage.h
@@ -228,6 +228,8 @@ public:
 	void Advance(float DeltaTime, struct FRootMotionMovementParams & OutRootMotionParams);
 
 	FName GetCurrentSection() const;
+	FName GetNextSection() const;
+
 	// reference has to be managed manually
 	void AddReferencedObjects( FReferenceCollector& Collector );
 private:
diff --git a/Engine/Source/Runtime/Engine/Classes/Animation/AnimNodeBase.h b/Engine/Source/Runtime/Engine/Classes/Animation/AnimNodeBase.h
--- a/Engine/Source/Runtime/Engine/Classes/Animation/AnimNodeBase.h
+++ b/Engine/Source/Runtime/Engine/Classes/Animation/AnimNodeBase.h
@@ -175,7 +175,67 @@ public:
 	bool IsNormalized() const;
 };
 
+struct FNodeDebugData
+{
+private:
+	struct DebugItem
+	{
+		DebugItem(FString Data, bool bPoseSource) : DebugData(Data), bPoseSource(bPoseSource) {}
+		
+		// This node items debug text to display
+		FString DebugData;
+
+		// Whether we are supplying a pose instead of modifying one (e.g. an playing animation)
+		bool	bPoseSource;
+
+		// Nodes that we are connected to
+		TArray<FNodeDebugData> ChildNodeChain;
+	};
+
+	// This nodes final contribution weight (based on its own weight and the weight of its parents)
+	float			AbsoluteWeight;
+
+	// Nodes that we are dependent on
+	TArray<DebugItem> NodeChain;
+
+public:
+	struct FFlattenedDebugData
+	{
+		FFlattenedDebugData(FString Line, float AbsWeight, int32 Indent, int32 ChainID, bool bPoseSource) : DebugLine(Line), AbsoluteWeight(AbsWeight), Indent(Indent), ChainID(ChainID), bPoseSource(bPoseSource){}
+		FString DebugLine;
+		float	AbsoluteWeight;
+		int32	Indent;
+		int32	ChainID;
+		bool	bPoseSource;
+
+		bool IsOnActiveBranch() { return AbsoluteWeight > ZERO_ANIMWEIGHT_THRESH; }
+	};
+
+	FNodeDebugData(const class UAnimInstance* AnimInstance) : AnimInstance(AnimInstance), AbsoluteWeight(1.f) {}
+	FNodeDebugData(const class UAnimInstance* AnimInstance, float AbsWeight) : AnimInstance(AnimInstance), AbsoluteWeight(AbsWeight) {}
 
+	void			AddDebugItem(FString DebugData, bool bPoseSource = false);
+	FNodeDebugData&	BranchFlow(float BranchWeight);
+
+	template<class Type>
+	FString GetNodeName(Type* Node)
+	{
+		return FString::Printf(TEXT("%s<W:%.1f%%>"), *Node->StaticStruct()->GetName(), AbsoluteWeight*100.f);
+	}
+
+	void GetFlattenedDebugData(TArray<FFlattenedDebugData>& FlattenedDebugData, int32 Indent, int32& ChainID);
+
+	TArray<FFlattenedDebugData> GetFlattenedDebugData()
+	{
+		TArray<FFlattenedDebugData> Data;
+		int32 ChainID = 0;
+		GetFlattenedDebugData(Data, 0, ChainID);
+		return Data;
+	}
+
+	// Anim instance that we are generating debug data for
+	const class UAnimInstance* AnimInstance;
+};
 
 // The display mode of editable values on an animation node
 UENUM()
@@ -235,6 +295,7 @@ public:
 	void Initialize(const FAnimationInitializeContext& Context);
 	void CacheBones(const FAnimationCacheBonesContext & Context) ;
 	void Update(const FAnimationUpdateContext& Context);
+	void GatherDebugData(FNodeDebugData& DebugData);
 
 	// Try to re-establish the linked node pointer
 	void AttemptRelink(const FAnimationBaseContext& Context);
@@ -285,7 +346,6 @@ struct FExposedValueHandler
 	}
 };
 
-
 // To create a new animation node:
 //   Create a struct derived from FAnimNode_Base - this is your runtime node
 //   Create a class derived from UAnimGraphNode_Base, containing an instance of your runtime as a member - this is your visual/editor-only node
@@ -311,5 +371,10 @@ struct ENGINE_API FAnimNode_Base
 
 	// If a derived anim node should respond to asset overrides, OverrideAsset should be defined to handle changing the asset
 	virtual void OverrideAsset(UAnimationAsset* NewAsset) {}
+
+	virtual void GatherDebugData(FNodeDebugData& DebugData)
+	{ 
+		DebugData.AddDebugItem(TEXT("Non Overriden GatherDebugData")); 
+	}
 	// End of interface to implement
 };
diff --git a/Engine/Source/Runtime/Engine/Classes/Animation/AnimNodeSpaceConversions.h b/Engine/Source/Runtime/Engine/Classes/Animation/AnimNodeSpaceConversions.h
--- a/Engine/Source/Runtime/Engine/Classes/Animation/AnimNodeSpaceConversions.h
+++ b/Engine/Source/Runtime/Engine/Classes/Animation/AnimNodeSpaceConversions.h
@@ -22,6 +22,7 @@ public:
 	virtual void CacheBones(const FAnimationCacheBonesContext & Context) OVERRIDE;
 	virtual void Update(const FAnimationUpdateContext& Context) OVERRIDE;
 	virtual void Evaluate(FPoseContext& Output) OVERRIDE;
+	virtual void GatherDebugData(FNodeDebugData& DebugData) OVERRIDE;
 	// End of FAnimNode_Base interface
 
 };
@@ -42,6 +43,7 @@ public:
 	virtual void Initialize(const FAnimationInitializeContext& Context) OVERRIDE;
 	virtual void CacheBones(const FAnimationCacheBonesContext & Context) OVERRIDE;
 	virtual void Update(const FAnimationUpdateContext& Context) OVERRIDE;
+	virtual void GatherDebugData(FNodeDebugData& DebugData) OVERRIDE;
 	virtual void EvaluateComponentSpace(FComponentSpacePoseContext& Output) OVERRIDE;
 	// End of FAnimNode_Base interface
 };
diff --git a/Engine/Source/Runtime/Engine/Classes/Animation/AnimNode_ApplyAdditive.h b/Engine/Source/Runtime/Engine/Classes/Animation/AnimNode_ApplyAdditive.h
--- a/Engine/Source/Runtime/Engine/Classes/Animation/AnimNode_ApplyAdditive.h
+++ b/Engine/Source/Runtime/Engine/Classes/Animation/AnimNode_ApplyAdditive.h
@@ -31,5 +31,6 @@ public:
 	virtual void CacheBones(const FAnimationCacheBonesContext & Context) OVERRIDE;
 	virtual void Update(const FAnimationUpdateContext& Context) OVERRIDE;
 	virtual void Evaluate(FPoseContext& Output) OVERRIDE;
+	virtual void GatherDebugData(FNodeDebugData& DebugData) OVERRIDE;
 	// End of FAnimNode_Base interface
 };
diff --git a/Engine/Source/Runtime/Engine/Classes/Animation/AnimNode_BlendListBase.h b/Engine/Source/Runtime/Engine/Classes/Animation/AnimNode_BlendListBase.h
--- a/Engine/Source/Runtime/Engine/Classes/Animation/AnimNode_BlendListBase.h
+++ b/Engine/Source/Runtime/Engine/Classes/Animation/AnimNode_BlendListBase.h
@@ -39,6 +39,7 @@ public:
 	virtual void CacheBones(const FAnimationCacheBonesContext & Context) OVERRIDE;
 	virtual void Update(const FAnimationUpdateContext& Context) OVERRIDE;
 	virtual void Evaluate(FPoseContext& Output) OVERRIDE;
+	virtual void GatherDebugData(FNodeDebugData& DebugData) OVERRIDE;
 	// End of FAnimNode_Base interface
 
 #if WITH_EDITOR
@@ -57,4 +58,5 @@ public:
 
 protected:
 	virtual int32 GetActiveChildIndex() { return 0; }
+	virtual FString GetNodeName(FNodeDebugData& DebugData) { return DebugData.GetNodeName(this); }
 };
diff --git a/Engine/Source/Runtime/Engine/Classes/Animation/AnimNode_BlendListByBool.h b/Engine/Source/Runtime/Engine/Classes/Animation/AnimNode_BlendListByBool.h
--- a/Engine/Source/Runtime/Engine/Classes/Animation/AnimNode_BlendListByBool.h
+++ b/Engine/Source/Runtime/Engine/Classes/Animation/AnimNode_BlendListByBool.h
@@ -27,4 +27,5 @@ public:
 
 protected:
 	virtual int32 GetActiveChildIndex();
+	virtual FString GetNodeName(FNodeDebugData& DebugData) { return DebugData.GetNodeName(this); }
 };
diff --git a/Engine/Source/Runtime/Engine/Classes/Animation/AnimNode_BlendListByEnum.h b/Engine/Source/Runtime/Engine/Classes/Animation/AnimNode_BlendListByEnum.h
--- a/Engine/Source/Runtime/Engine/Classes/Animation/AnimNode_BlendListByEnum.h
+++ b/Engine/Source/Runtime/Engine/Classes/Animation/AnimNode_BlendListByEnum.h
@@ -27,4 +27,5 @@ public:
 
 protected:
 	virtual int32 GetActiveChildIndex();
+	virtual FString GetNodeName(FNodeDebugData& DebugData) { return DebugData.GetNodeName(this); }
 };
diff --git a/Engine/Source/Runtime/Engine/Classes/Animation/AnimNode_BlendListByInt.h b/Engine/Source/Runtime/Engine/Classes/Animation/AnimNode_BlendListByInt.h
--- a/Engine/Source/Runtime/Engine/Classes/Animation/AnimNode_BlendListByInt.h
+++ b/Engine/Source/Runtime/Engine/Classes/Animation/AnimNode_BlendListByInt.h
@@ -24,4 +24,5 @@ public:
 
 protected:
 	virtual int32 GetActiveChildIndex();
+	virtual FString GetNodeName(FNodeDebugData& DebugData) { return DebugData.GetNodeName(this); }
 };
diff --git a/Engine/Source/Runtime/Engine/Classes/Animation/AnimNode_BlendSpaceEvaluator.h b/Engine/Source/Runtime/Engine/Classes/Animation/AnimNode_BlendSpaceEvaluator.h
--- a/Engine/Source/Runtime/Engine/Classes/Animation/AnimNode_BlendSpaceEvaluator.h
+++ b/Engine/Source/Runtime/Engine/Classes/Animation/AnimNode_BlendSpaceEvaluator.h
@@ -23,5 +23,6 @@ public:
 
 	// FAnimNode_Base interface
 	virtual void Update(const FAnimationUpdateContext& Context) OVERRIDE;
+	virtual void GatherDebugData(FNodeDebugData& DebugData) OVERRIDE;
 	// End of FAnimNode_Base interface
 };
diff --git a/Engine/Source/Runtime/Engine/Classes/Animation/AnimNode_BlendSpacePlayer.h b/Engine/Source/Runtime/Engine/Classes/Animation/AnimNode_BlendSpacePlayer.h
--- a/Engine/Source/Runtime/Engine/Classes/Animation/AnimNode_BlendSpacePlayer.h
+++ b/Engine/Source/Runtime/Engine/Classes/Animation/AnimNode_BlendSpacePlayer.h
@@ -62,5 +62,6 @@ public:
 	virtual void Update(const FAnimationUpdateContext& Context) OVERRIDE;
 	virtual void Evaluate(FPoseContext& Output) OVERRIDE;
 	virtual void OverrideAsset(UAnimationAsset* NewAsset) OVERRIDE;
+	virtual void GatherDebugData(FNodeDebugData& DebugData) OVERRIDE;
 	// End of FAnimNode_Base interface
 };
diff --git a/Engine/Source/Runtime/Engine/Classes/Animation/AnimNode_LayeredBoneBlend.h b/Engine/Source/Runtime/Engine/Classes/Animation/AnimNode_LayeredBoneBlend.h
--- a/Engine/Source/Runtime/Engine/Classes/Animation/AnimNode_LayeredBoneBlend.h
+++ b/Engine/Source/Runtime/Engine/Classes/Animation/AnimNode_LayeredBoneBlend.h
@@ -46,6 +46,7 @@ public:
 	virtual void CacheBones(const FAnimationCacheBonesContext & Context) OVERRIDE;
 	virtual void Update(const FAnimationUpdateContext& Context) OVERRIDE;
 	virtual void Evaluate(FPoseContext& Output) OVERRIDE;
+	virtual void GatherDebugData(FNodeDebugData& DebugData) OVERRIDE;
 	// End of FAnimNode_Base interface
 
 #if WITH_EDITOR
diff --git a/Engine/Source/Runtime/Engine/Classes/Animation/AnimNode_RefPose.h b/Engine/Source/Runtime/Engine/Classes/Animation/AnimNode_RefPose.h
--- a/Engine/Source/Runtime/Engine/Classes/Animation/AnimNode_RefPose.h
+++ b/Engine/Source/Runtime/Engine/Classes/Animation/AnimNode_RefPose.h
@@ -28,6 +28,8 @@ public:
 	}
 
 	virtual void Evaluate(FPoseContext& Output) OVERRIDE;
+	virtual void GatherDebugData(FNodeDebugData& DebugData) OVERRIDE;
+
 };
 
 USTRUCT()
diff --git a/Engine/Source/Runtime/Engine/Classes/Animation/AnimNode_Root.h b/Engine/Source/Runtime/Engine/Classes/Animation/AnimNode_Root.h
--- a/Engine/Source/Runtime/Engine/Classes/Animation/AnimNode_Root.h
+++ b/Engine/Source/Runtime/Engine/Classes/Animation/AnimNode_Root.h
@@ -21,5 +21,6 @@ public:
 	virtual void CacheBones(const FAnimationCacheBonesContext & Context) OVERRIDE;
 	virtual void Update(const FAnimationUpdateContext& Context) OVERRIDE;
 	virtual void Evaluate(FPoseContext& Output) OVERRIDE;
+	virtual void GatherDebugData(FNodeDebugData& DebugData) OVERRIDE;
 	// End of FAnimNode_Base interface
 };
diff --git a/Engine/Source/Runtime/Engine/Classes/Animation/AnimNode_RotateRootBone.h b/Engine/Source/Runtime/Engine/Classes/Animation/AnimNode_RotateRootBone.h
--- a/Engine/Source/Runtime/Engine/Classes/Animation/AnimNode_RotateRootBone.h
+++ b/Engine/Source/Runtime/Engine/Classes/Animation/AnimNode_RotateRootBone.h
@@ -31,5 +31,6 @@ public:
 	virtual void CacheBones(const FAnimationCacheBonesContext & Context) OVERRIDE;
 	virtual void Update(const FAnimationUpdateContext& Context) OVERRIDE;
 	virtual void Evaluate(FPoseContext& Output) OVERRIDE;
+	virtual void GatherDebugData(FNodeDebugData& DebugData) OVERRIDE;
 	// End of FAnimNode_Base interface
 };
diff --git a/Engine/Source/Runtime/Engine/Classes/Animation/AnimNode_RotationOffsetBlendSpace.h b/Engine/Source/Runtime/Engine/Classes/Animation/AnimNode_RotationOffsetBlendSpace.h
--- a/Engine/Source/Runtime/Engine/Classes/Animation/AnimNode_RotationOffsetBlendSpace.h
+++ b/Engine/Source/Runtime/Engine/Classes/Animation/AnimNode_RotationOffsetBlendSpace.h
@@ -22,5 +22,6 @@ public:
 	virtual void CacheBones(const FAnimationCacheBonesContext & Context) OVERRIDE;
 	virtual void Update(const FAnimationUpdateContext& Context) OVERRIDE;
 	virtual void Evaluate(FPoseContext& Output) OVERRIDE;
+	virtual void GatherDebugData(FNodeDebugData& DebugData) OVERRIDE;
 	// End of FAnimNode_Base interface
 };
diff --git a/Engine/Source/Runtime/Engine/Classes/Animation/AnimNode_SaveCachedPose.h b/Engine/Source/Runtime/Engine/Classes/Animation/AnimNode_SaveCachedPose.h
--- a/Engine/Source/Runtime/Engine/Classes/Animation/AnimNode_SaveCachedPose.h
+++ b/Engine/Source/Runtime/Engine/Classes/Animation/AnimNode_SaveCachedPose.h
@@ -34,5 +34,6 @@ public:
 	virtual void CacheBones(const FAnimationCacheBonesContext & Context) OVERRIDE;
 	virtual void Update(const FAnimationUpdateContext& Context) OVERRIDE;
 	virtual void Evaluate(FPoseContext& Output) OVERRIDE;
+	virtual void GatherDebugData(FNodeDebugData& DebugData) OVERRIDE;
 	// End of FAnimNode_Base interface
 };
diff --git a/Engine/Source/Runtime/Engine/Classes/Animation/AnimNode_SequenceEvaluator.h b/Engine/Source/Runtime/Engine/Classes/Animation/AnimNode_SequenceEvaluator.h
--- a/Engine/Source/Runtime/Engine/Classes/Animation/AnimNode_SequenceEvaluator.h
+++ b/Engine/Source/Runtime/Engine/Classes/Animation/AnimNode_SequenceEvaluator.h
@@ -32,5 +32,6 @@ public:
 	virtual void Update(const FAnimationUpdateContext& Context) OVERRIDE;
 	virtual void Evaluate(FPoseContext& Output) OVERRIDE;
 	virtual void OverrideAsset(UAnimationAsset* NewAsset) OVERRIDE;
+	virtual void GatherDebugData(FNodeDebugData& DebugData) OVERRIDE;
 	// End of FAnimNode_Base interface
 };
diff --git a/Engine/Source/Runtime/Engine/Classes/Animation/AnimNode_SequencePlayer.h b/Engine/Source/Runtime/Engine/Classes/Animation/AnimNode_SequencePlayer.h
--- a/Engine/Source/Runtime/Engine/Classes/Animation/AnimNode_SequencePlayer.h
+++ b/Engine/Source/Runtime/Engine/Classes/Animation/AnimNode_SequencePlayer.h
@@ -50,5 +50,6 @@ public:
 	virtual void Update(const FAnimationUpdateContext& Context) OVERRIDE;
 	virtual void Evaluate(FPoseContext& Output) OVERRIDE;
 	virtual void OverrideAsset(UAnimationAsset* NewAsset) OVERRIDE;
+	virtual void GatherDebugData(FNodeDebugData& DebugData) OVERRIDE;
 	// End of FAnimNode_Base interface
 };
diff --git a/Engine/Source/Runtime/Engine/Classes/Animation/AnimNode_Slot.h b/Engine/Source/Runtime/Engine/Classes/Animation/AnimNode_Slot.h
--- a/Engine/Source/Runtime/Engine/Classes/Animation/AnimNode_Slot.h
+++ b/Engine/Source/Runtime/Engine/Classes/Animation/AnimNode_Slot.h
@@ -31,5 +31,6 @@ public:
 	virtual void CacheBones(const FAnimationCacheBonesContext & Context) OVERRIDE;
 	virtual void Update(const FAnimationUpdateContext& Context) OVERRIDE;
 	virtual void Evaluate(FPoseContext& Output) OVERRIDE;
+	virtual void GatherDebugData(FNodeDebugData& DebugData) OVERRIDE;
 	// End of FAnimNode_Base interface
 };
diff --git a/Engine/Source/Runtime/Engine/Classes/Animation/AnimNode_StateMachine.h b/Engine/Source/Runtime/Engine/Classes/Animation/AnimNode_StateMachine.h
--- a/Engine/Source/Runtime/Engine/Classes/Animation/AnimNode_StateMachine.h
+++ b/Engine/Source/Runtime/Engine/Classes/Animation/AnimNode_StateMachine.h
@@ -163,6 +163,7 @@ public:
 	virtual void CacheBones(const FAnimationCacheBonesContext & Context) OVERRIDE;
 	virtual void Update(const FAnimationUpdateContext& Context) OVERRIDE;
 	virtual void Evaluate(FPoseContext& Output) OVERRIDE;
+	virtual void GatherDebugData(FNodeDebugData& DebugData) OVERRIDE;
 	// End of FAnimNode_Base interface
 
 	// Returns the blend weight of the specified state, as calculated by the last call to Update()
diff --git a/Engine/Source/Runtime/Engine/Classes/Animation/AnimNode_TransitionPoseEvaluator.h b/Engine/Source/Runtime/Engine/Classes/Animation/AnimNode_TransitionPoseEvaluator.h
--- a/Engine/Source/Runtime/Engine/Classes/Animation/AnimNode_TransitionPoseEvaluator.h
+++ b/Engine/Source/Runtime/Engine/Classes/Animation/AnimNode_TransitionPoseEvaluator.h
@@ -63,6 +63,7 @@ public:
 	virtual void CacheBones(const FAnimationCacheBonesContext & Context) OVERRIDE;
 	virtual void Update(const FAnimationUpdateContext& Context) OVERRIDE;
 	virtual void Evaluate(FPoseContext& Output) OVERRIDE;
+	virtual void GatherDebugData(FNodeDebugData& DebugData) OVERRIDE;
 	// End of FAnimNode_Base interface
 
 	bool InputNodeNeedsUpdate() const;
diff --git a/Engine/Source/Runtime/Engine/Classes/Animation/AnimNode_TransitionResult.h b/Engine/Source/Runtime/Engine/Classes/Animation/AnimNode_TransitionResult.h
--- a/Engine/Source/Runtime/Engine/Classes/Animation/AnimNode_TransitionResult.h
+++ b/Engine/Source/Runtime/Engine/Classes/Animation/AnimNode_TransitionResult.h
@@ -21,5 +21,6 @@ public:
 	virtual void CacheBones(const FAnimationCacheBonesContext & Context) OVERRIDE;
 	virtual void Update(const FAnimationUpdateContext& Context) OVERRIDE;
 	virtual void Evaluate(FPoseContext& Output) OVERRIDE;
+	virtual void GatherDebugData(FNodeDebugData& DebugData) OVERRIDE;
 	// End of FAnimNode_Base interface
 };
diff --git a/Engine/Source/Runtime/Engine/Classes/Animation/AnimNode_TwoWayBlend.h b/Engine/Source/Runtime/Engine/Classes/Animation/AnimNode_TwoWayBlend.h
--- a/Engine/Source/Runtime/Engine/Classes/Animation/AnimNode_TwoWayBlend.h
+++ b/Engine/Source/Runtime/Engine/Classes/Animation/AnimNode_TwoWayBlend.h
@@ -35,6 +35,7 @@ public:
 	virtual void CacheBones(const FAnimationCacheBonesContext & Context) OVERRIDE;
 	virtual void Update(const FAnimationUpdateContext& Context) OVERRIDE;
 	virtual void Evaluate(FPoseContext& Output) OVERRIDE;
+	virtual void GatherDebugData(FNodeDebugData& DebugData) OVERRIDE;
 	// End of FAnimNode_Base interface
 };
 
diff --git a/Engine/Source/Runtime/Engine/Classes/Animation/AnimNode_UseCachedPose.h b/Engine/Source/Runtime/Engine/Classes/Animation/AnimNode_UseCachedPose.h
--- a/Engine/Source/Runtime/Engine/Classes/Animation/AnimNode_UseCachedPose.h
+++ b/Engine/Source/Runtime/Engine/Classes/Animation/AnimNode_UseCachedPose.h
@@ -20,5 +20,6 @@ public:
 	virtual void CacheBones(const FAnimationCacheBonesContext & Context) OVERRIDE;
 	virtual void Update(const FAnimationUpdateContext& Context) OVERRIDE;
 	virtual void Evaluate(FPoseContext& Output) OVERRIDE;
+	virtual void GatherDebugData(FNodeDebugData& DebugData) OVERRIDE;
 	// End of FAnimNode_Base interface
 };
diff --git a/Engine/Source/Runtime/Engine/Classes/Animation/BoneControllers/AnimNode_CopyBone.h b/Engine/Source/Runtime/Engine/Classes/Animation/BoneControllers/AnimNode_CopyBone.h
--- a/Engine/Source/Runtime/Engine/Classes/Animation/BoneControllers/AnimNode_CopyBone.h
+++ b/Engine/Source/Runtime/Engine/Classes/Animation/BoneControllers/AnimNode_CopyBone.h
@@ -36,6 +36,10 @@ struct ENGINE_API FAnimNode_CopyBone : public FAnimNode_SkeletalControlBase
 
 	FAnimNode_CopyBone();
 
+	// FAnimNode_Base interface
+	virtual void GatherDebugData(FNodeDebugData& DebugData) OVERRIDE;
+	// End of FAnimNode_Base interface
+
 	// FAnimNode_SkeletalControlBase interface
 	virtual void EvaluateBoneTransforms(USkeletalMeshComponent* SkelComp, const FBoneContainer & RequiredBones, FA2CSPose& MeshBases, TArray<FBoneTransform>& OutBoneTransforms) OVERRIDE;
 	virtual bool IsValidToEvaluate(const USkeleton * Skeleton, const FBoneContainer & RequiredBones) OVERRIDE;
diff --git a/Engine/Source/Runtime/Engine/Classes/Animation/BoneControllers/AnimNode_ModifyBone.h b/Engine/Source/Runtime/Engine/Classes/Animation/BoneControllers/AnimNode_ModifyBone.h
--- a/Engine/Source/Runtime/Engine/Classes/Animation/BoneControllers/AnimNode_ModifyBone.h
+++ b/Engine/Source/Runtime/Engine/Classes/Animation/BoneControllers/AnimNode_ModifyBone.h
@@ -68,6 +68,10 @@ struct ENGINE_API FAnimNode_ModifyBone : public FAnimNode_SkeletalControlBase
 
 	FAnimNode_ModifyBone();
 
+	// FAnimNode_Base interface
+	virtual void GatherDebugData(FNodeDebugData& DebugData) OVERRIDE;
+	// End of FAnimNode_Base interface
+
 	// FAnimNode_SkeletalControlBase interface
 	virtual void EvaluateBoneTransforms(USkeletalMeshComponent* SkelComp, const FBoneContainer & RequiredBones, FA2CSPose& MeshBases, TArray<FBoneTransform>& OutBoneTransforms) OVERRIDE;
 	virtual bool IsValidToEvaluate(const USkeleton * Skeleton, const FBoneContainer & RequiredBones) OVERRIDE;
diff --git a/Engine/Source/Runtime/Engine/Classes/Animation/BoneControllers/AnimNode_RotationMultiplier.h b/Engine/Source/Runtime/Engine/Classes/Animation/BoneControllers/AnimNode_RotationMultiplier.h
--- a/Engine/Source/Runtime/Engine/Classes/Animation/BoneControllers/AnimNode_RotationMultiplier.h
+++ b/Engine/Source/Runtime/Engine/Classes/Animation/BoneControllers/AnimNode_RotationMultiplier.h
@@ -40,6 +40,10 @@ struct ENGINE_API FAnimNode_RotationMultiplier : public FAnimNode_SkeletalContro
 
 	FAnimNode_RotationMultiplier();
 
+	// FAnimNode_Base interface
+	virtual void GatherDebugData(FNodeDebugData& DebugData) OVERRIDE;
+	// End of FAnimNode_Base interface
+
 	// FAnimNode_SkeletalControlBase interface
 	virtual void EvaluateBoneTransforms(USkeletalMeshComponent* SkelComp, const FBoneContainer & RequiredBones, FA2CSPose& MeshBases, TArray<FBoneTransform>& OutBoneTransforms) OVERRIDE;
 	virtual bool IsValidToEvaluate(const USkeleton * Skeleton, const FBoneContainer & RequiredBones) OVERRIDE;
diff --git a/Engine/Source/Runtime/Engine/Classes/Animation/BoneControllers/AnimNode_SkeletalControlBase.h b/Engine/Source/Runtime/Engine/Classes/Animation/BoneControllers/AnimNode_SkeletalControlBase.h
--- a/Engine/Source/Runtime/Engine/Classes/Animation/BoneControllers/AnimNode_SkeletalControlBase.h
+++ b/Engine/Source/Runtime/Engine/Classes/Animation/BoneControllers/AnimNode_SkeletalControlBase.h
@@ -53,4 +53,7 @@ protected:
 	virtual bool IsValidToEvaluate(const USkeleton * Skeleton, const FBoneContainer & RequiredBones) { return false; }
 	// initialize any bone references you have
 	virtual void InitializeBoneReferences(const FBoneContainer & RequiredBones){};
+
+	/** Allow base to add info to the node debug output */
+	void AddDebugNodeData(FString& OutDebugData);
 };
diff --git a/Engine/Source/Runtime/Engine/Classes/Animation/BoneControllers/AnimNode_SpringBone.h b/Engine/Source/Runtime/Engine/Classes/Animation/BoneControllers/AnimNode_SpringBone.h
--- a/Engine/Source/Runtime/Engine/Classes/Animation/BoneControllers/AnimNode_SpringBone.h
+++ b/Engine/Source/Runtime/Engine/Classes/Animation/BoneControllers/AnimNode_SpringBone.h
@@ -61,6 +61,7 @@ public:
 	virtual void Initialize(const FAnimationInitializeContext& Context) OVERRIDE;
 	virtual void CacheBones(const FAnimationCacheBonesContext & Context) OVERRIDE;
 	virtual void Update(const FAnimationUpdateContext& Context) OVERRIDE;
+	virtual void GatherDebugData(FNodeDebugData& DebugData) OVERRIDE;
 	// End of FAnimNode_Base interface
 
 	// FAnimNode_SkeletalControlBase interface
diff --git a/Engine/Source/Runtime/Engine/Classes/Animation/BoneControllers/AnimNode_TwoBoneIK.h b/Engine/Source/Runtime/Engine/Classes/Animation/BoneControllers/AnimNode_TwoBoneIK.h
--- a/Engine/Source/Runtime/Engine/Classes/Animation/BoneControllers/AnimNode_TwoBoneIK.h
+++ b/Engine/Source/Runtime/Engine/Classes/Animation/BoneControllers/AnimNode_TwoBoneIK.h
@@ -60,6 +60,10 @@ struct ENGINE_API FAnimNode_TwoBoneIK : public FAnimNode_SkeletalControlBase
 
 	FAnimNode_TwoBoneIK();
 
+	// FAnimNode_Base interface
+	virtual void GatherDebugData(FNodeDebugData& DebugData) OVERRIDE;
+	// End of FAnimNode_Base interface
+
 	// FAnimNode_SkeletalControlBase interface
 	virtual void EvaluateBoneTransforms(USkeletalMeshComponent* SkelComp, const FBoneContainer & RequiredBones, FA2CSPose& MeshBases, TArray<FBoneTransform>& OutBoneTransforms) OVERRIDE;
 	virtual bool IsValidToEvaluate(const USkeleton * Skeleton, const FBoneContainer & RequiredBones) OVERRIDE;
diff --git a/Engine/Source/Runtime/Engine/Classes/Camera/PlayerCameraManager.h b/Engine/Source/Runtime/Engine/Classes/Camera/PlayerCameraManager.h
--- a/Engine/Source/Runtime/Engine/Classes/Camera/PlayerCameraManager.h
+++ b/Engine/Source/Runtime/Engine/Classes/Camera/PlayerCameraManager.h
@@ -428,7 +428,7 @@ public:
 	virtual bool ShouldTickIfViewportsOnly() const OVERRIDE;
 	virtual void PostInitializeComponents() OVERRIDE;
 	virtual void Destroyed() OVERRIDE;
-	virtual void DisplayDebug(class UCanvas* Canvas, const TArray<FName>& DebugDisplay, float& YL, float& YPos) OVERRIDE;
+	virtual void DisplayDebug(class UCanvas* Canvas, const FDebugDisplayInfo& DebugDisplay, float& YL, float& YPos) OVERRIDE;
 	virtual void ApplyWorldOffset(const FVector& InOffset, bool bWorldShift) OVERRIDE;
 	// End AActor Interface
 
diff --git a/Engine/Source/Runtime/Engine/Classes/GameFramework/AIController.h b/Engine/Source/Runtime/Engine/Classes/GameFramework/AIController.h
--- a/Engine/Source/Runtime/Engine/Classes/GameFramework/AIController.h
+++ b/Engine/Source/Runtime/Engine/Classes/GameFramework/AIController.h
@@ -228,7 +228,7 @@ public:
 
 	// Begin AController Interface
 	virtual void Possess(class APawn* InPawn) OVERRIDE;
-	virtual void DisplayDebug(class UCanvas* Canvas, const TArray<FName>& DebugDisplay, float& YL, float& YPos) OVERRIDE;
+	virtual void DisplayDebug(class UCanvas* Canvas, const FDebugDisplayInfo& DebugDisplay, float& YL, float& YPos) OVERRIDE;
 
 #if ENABLE_VISUAL_LOG
 	virtual void GrabDebugSnapshot(struct FVisLogEntry* Snapshot) const OVERRIDE;
diff --git a/Engine/Source/Runtime/Engine/Classes/GameFramework/Actor.h b/Engine/Source/Runtime/Engine/Classes/GameFramework/Actor.h
--- a/Engine/Source/Runtime/Engine/Classes/GameFramework/Actor.h
+++ b/Engine/Source/Runtime/Engine/Classes/GameFramework/Actor.h
@@ -109,6 +109,7 @@ public:
 	}
 };
 
+
 /**
  * Base class for an object that can be placed or spawned in a level. Actors may contain a collection of Components, and support network replication.
 *
@@ -1940,11 +1941,11 @@ public:
 	 * Draw important Actor variables on canvas.  HUD will call DisplayDebug() on the current ViewTarget when the ShowDebug exec is used
 	 *
 	 * @param Canvas - Canvas to draw on
-	 * @param DebugDisplay - List of names specifying what debug info to display
+	 * @param DebugDisplay - Contains information about what debug data to display
 	 * @param YL - Height of the current font
 	 * @param YPos - Y position on Canvas. YPos += YL, gives position to draw text for next debug line.
 	 */
-	virtual void DisplayDebug(class UCanvas* Canvas, const TArray<FName>& DebugDisplay, float& YL, float& YPos);
+	virtual void DisplayDebug(class UCanvas* Canvas, const class FDebugDisplayInfo& DebugDisplay, float& YL, float& YPos);
 
 	/** Retrieves actor's name used for logging, or string "NULL" if Actor == NULL */
 	static FString GetDebugName(const AActor* Actor) { return Actor ? Actor->GetName() : TEXT("NULL"); }
diff --git a/Engine/Source/Runtime/Engine/Classes/GameFramework/Character.h b/Engine/Source/Runtime/Engine/Classes/GameFramework/Character.h
--- a/Engine/Source/Runtime/Engine/Classes/GameFramework/Character.h
+++ b/Engine/Source/Runtime/Engine/Classes/GameFramework/Character.h
@@ -252,7 +252,7 @@ public:
 	virtual void PawnClientRestart() OVERRIDE;
 	virtual void UnPossessed() OVERRIDE;
 	virtual void SetupPlayerInputComponent(class UInputComponent* InputComponent) OVERRIDE;
-	virtual void DisplayDebug(class UCanvas* Canvas, const TArray<FName>& DebugDisplay, float& YL, float& YPos) OVERRIDE;
+	virtual void DisplayDebug(class UCanvas* Canvas, const FDebugDisplayInfo& DebugDisplay, float& YL, float& YPos) OVERRIDE;
 	virtual void RecalculateBaseEyeHeight() OVERRIDE;
 	// End APawn Interface
 
diff --git a/Engine/Source/Runtime/Engine/Classes/GameFramework/CharacterMovementComponent.h b/Engine/Source/Runtime/Engine/Classes/GameFramework/CharacterMovementComponent.h
--- a/Engine/Source/Runtime/Engine/Classes/GameFramework/CharacterMovementComponent.h
+++ b/Engine/Source/Runtime/Engine/Classes/GameFramework/CharacterMovementComponent.h
@@ -885,11 +885,11 @@ public:
 	 * Draw important variables on canvas.  Character will call DisplayDebug() on the current ViewTarget when the ShowDebug exec is used
 	 *
 	 * @param Canvas - Canvas to draw on
-	 * @param DebugDisplay - List of names specifying what debug info to display
+	 * @param DebugDisplay - Contains information about what debug data to display
 	 * @param YL - Height of the current font
 	 * @param YPos - Y position on Canvas. YPos += YL, gives position to draw text for next debug line.
 	 */
-	virtual void DisplayDebug(class UCanvas* Canvas, const TArray<FName>& DebugDisplay, float& YL, float& YPos);
+	virtual void DisplayDebug(class UCanvas* Canvas, const FDebugDisplayInfo& DebugDisplay, float& YL, float& YPos);
 
 	/** Check if swimming pawn just ran into edge of the pool and should jump out. */
 	virtual bool CheckWaterJump(FVector CheckPoint, FVector& WallNormal);
diff --git a/Engine/Source/Runtime/Engine/Classes/GameFramework/Controller.h b/Engine/Source/Runtime/Engine/Classes/GameFramework/Controller.h
--- a/Engine/Source/Runtime/Engine/Classes/GameFramework/Controller.h
+++ b/Engine/Source/Runtime/Engine/Classes/GameFramework/Controller.h
@@ -165,7 +165,7 @@ public:
 	// Begin AActor interface
 	virtual void TickActor( float DeltaTime, enum ELevelTick TickType, FActorTickFunction& ThisTickFunction ) OVERRIDE;
 	virtual void K2_DestroyActor() OVERRIDE;
-	virtual void DisplayDebug(class UCanvas* Canvas, const TArray<FName>& DebugDisplay, float& YL, float& YPos) OVERRIDE;
+	virtual void DisplayDebug(class UCanvas* Canvas, const FDebugDisplayInfo& DebugDisplay, float& YL, float& YPos) OVERRIDE;
 	virtual void GetActorEyesViewPoint( FVector& out_Location, FRotator& out_Rotation ) const OVERRIDE;
 	virtual FString GetHumanReadableName() const OVERRIDE;
 
diff --git a/Engine/Source/Runtime/Engine/Classes/GameFramework/GameMode.h b/Engine/Source/Runtime/Engine/Classes/GameFramework/GameMode.h
--- a/Engine/Source/Runtime/Engine/Classes/GameFramework/GameMode.h
+++ b/Engine/Source/Runtime/Engine/Classes/GameFramework/GameMode.h
@@ -243,7 +243,7 @@ public:
 	virtual void PreInitializeComponents() OVERRIDE;
 	/** This sets the match state to WaitingToStart unless we're in seamless travel */
 	virtual void PostInitializeComponents() OVERRIDE; 
-	virtual void DisplayDebug(class UCanvas* Canvas, const TArray<FName>& DebugDisplay, float& YL, float& YPos) OVERRIDE;
+	virtual void DisplayDebug(class UCanvas* Canvas, const FDebugDisplayInfo& DebugDisplay, float& YL, float& YPos) OVERRIDE;
 	virtual void Reset() OVERRIDE;
 	virtual void Tick(float DeltaSeconds) OVERRIDE;
 	// End AActor interface
diff --git a/Engine/Source/Runtime/Engine/Classes/GameFramework/HUD.h b/Engine/Source/Runtime/Engine/Classes/GameFramework/HUD.h
--- a/Engine/Source/Runtime/Engine/Classes/GameFramework/HUD.h
+++ b/Engine/Source/Runtime/Engine/Classes/GameFramework/HUD.h
@@ -208,7 +208,11 @@ class ENGINE_API AHUD : public AActor
 
 	/** Array of names specifying what debug info to display for viewtarget actor. */
 	UPROPERTY(globalconfig)
-	TArray<FName> DebugDisplay;    
+	TArray<FName> DebugDisplay;
+
+	/** Array of names specifying what subsets of debug info to display for viewtarget actor. */
+	UPROPERTY(globalconfig)
+	TArray<FName> ToggledDebugCategories;
 
 protected:
 	/** Canvas to Draw HUD on.  Only valid during PostRender() event.  */
@@ -241,6 +245,12 @@ public:
 	UFUNCTION(exec)
 	virtual void ShowDebug(FName DebugType = NAME_None);
 
+	/**
+	 * Toggles sub categories of show debug to customize display
+	 */
+	UFUNCTION(exec)
+	void ShowDebugToggleSubCategory(FName Category);
+
 	UFUNCTION(reliable, client, SealedEvent)
 	void AddDebugText(const FString& DebugText, class AActor* SrcActor = NULL, float Duration = 0, FVector Offset = FVector(ForceInit), FVector DesiredOffset = FVector(ForceInit), FColor TextColor = FColor(ForceInit), bool bSkipOverwriteCheck = false, bool bAbsoluteLocation = false, bool bKeepAttachedToActor = false, class UFont* InFont = NULL, float FontScale = 1.0);
 
diff --git a/Engine/Source/Runtime/Engine/Classes/GameFramework/Pawn.h b/Engine/Source/Runtime/Engine/Classes/GameFramework/Pawn.h
--- a/Engine/Source/Runtime/Engine/Classes/GameFramework/Pawn.h
+++ b/Engine/Source/Runtime/Engine/Classes/GameFramework/Pawn.h
@@ -182,7 +182,7 @@ public:
 	virtual bool IsNetRelevantFor(APlayerController* RealViewer, AActor* Viewer, const FVector& SrcLocation) OVERRIDE;
 	virtual void PostNetReceiveLocation() OVERRIDE;
 	virtual void PostNetReceiveVelocity(const FVector& NewVelocity) OVERRIDE;
-	virtual void DisplayDebug(class UCanvas* Canvas, const TArray<FName>& DebugDisplay, float& YL, float& YPos) OVERRIDE;
+	virtual void DisplayDebug(class UCanvas* Canvas, const FDebugDisplayInfo& DebugDisplay, float& YL, float& YPos) OVERRIDE;
 	virtual void GetActorEyesViewPoint( FVector& Location, FRotator& Rotation ) const OVERRIDE;
 	virtual void OutsideWorldBounds() OVERRIDE;
 	virtual void Destroyed() OVERRIDE;
diff --git a/Engine/Source/Runtime/Engine/Classes/GameFramework/PlayerController.h b/Engine/Source/Runtime/Engine/Classes/GameFramework/PlayerController.h
--- a/Engine/Source/Runtime/Engine/Classes/GameFramework/PlayerController.h
+++ b/Engine/Source/Runtime/Engine/Classes/GameFramework/PlayerController.h
@@ -965,7 +965,7 @@ public:
 	virtual float GetNetPriority(const FVector& ViewPos, const FVector& ViewDir, APlayerController* Viewer, UActorChannel* InChannel, float Time, bool bLowBandwidth) OVERRIDE;
 	virtual class UPlayer* GetNetOwningPlayer() OVERRIDE;
 	virtual class UNetConnection* GetNetConnection() OVERRIDE;
-	virtual void DisplayDebug(class UCanvas* Canvas, const TArray<FName>& DebugDisplay, float& YL, float& YPos) OVERRIDE;
+	virtual void DisplayDebug(class UCanvas* Canvas, const FDebugDisplayInfo& DebugDisplay, float& YL, float& YPos) OVERRIDE;
 	virtual void PostInitializeComponents() OVERRIDE;
 	virtual void EnableInput(class APlayerController* PlayerController) OVERRIDE;
 	virtual void DisableInput(class APlayerController* PlayerController) OVERRIDE;
diff --git a/Engine/Source/Runtime/Engine/Classes/GameFramework/PlayerInput.h b/Engine/Source/Runtime/Engine/Classes/GameFramework/PlayerInput.h
--- a/Engine/Source/Runtime/Engine/Classes/GameFramework/PlayerInput.h
+++ b/Engine/Source/Runtime/Engine/Classes/GameFramework/PlayerInput.h
@@ -401,11 +401,11 @@ public:
 	 * Draw important PlayerInput variables on canvas.  HUD will call DisplayDebug() on the current ViewTarget when the ShowDebug exec is used
 	 *
 	 * @param Canvas - Canvas to draw on
-	 * @param DebugDisplay - List of names specifying what debug info to display
+	 * @param DebugDisplay - Contains information about what debug data to display
 	 * @param YL - Height of the current font
 	 * @param YPos - Y position on Canvas. YPos += YL, gives position to draw text for next debug line.
 	 */
-	void DisplayDebug(class UCanvas* Canvas, const TArray<FName>& DebugDisplay, float& YL, float& YPos);
+	void DisplayDebug(class UCanvas* Canvas, const FDebugDisplayInfo& DebugDisplay, float& YL, float& YPos);
 
 	bool IsPressed( FKey InKey ) const;
 	
diff --git a/Engine/Source/Runtime/Engine/Private/AI/Navigation/PathFollowingComponent.cpp b/Engine/Source/Runtime/Engine/Private/AI/Navigation/PathFollowingComponent.cpp
--- a/Engine/Source/Runtime/Engine/Private/AI/Navigation/PathFollowingComponent.cpp
+++ b/Engine/Source/Runtime/Engine/Private/AI/Navigation/PathFollowingComponent.cpp
@@ -1196,7 +1196,7 @@ FString UPathFollowingComponent::GetResultDesc(EPathFollowingResult::Type Result
 	return TEXT("Unknown");
 }
 
-void UPathFollowingComponent::DisplayDebug(class UCanvas* Canvas, const TArray<FName>& DebugDisplay, float& YL, float& YPos) const
+void UPathFollowingComponent::DisplayDebug(class UCanvas* Canvas, const FDebugDisplayInfo& DebugDisplay, float& YL, float& YPos) const
 {
 	Canvas->SetDrawColor(FColor::Blue);
 	UFont* RenderFont = GEngine->GetSmallFont();
diff --git a/Engine/Source/Runtime/Engine/Private/AIController.cpp b/Engine/Source/Runtime/Engine/Private/AIController.cpp
--- a/Engine/Source/Runtime/Engine/Private/AIController.cpp
+++ b/Engine/Source/Runtime/Engine/Private/AIController.cpp
@@ -4,7 +4,7 @@
 #include "NavigationPathBuilder.h"
 #include "VisualLog.h"
 #include "../Classes/Kismet/GameplayStatics.h"
-
+#include "DisplayDebugHelpers.h"
 
 bool AAIController::bAIIgnorePlayers = false;
 
@@ -53,12 +53,12 @@ void AAIController::Reset()
 	}
 }
 
-void AAIController::DisplayDebug(class UCanvas* Canvas, const TArray<FName>& DebugDisplay, float& YL, float& YPos)
+void AAIController::DisplayDebug(class UCanvas* Canvas, const FDebugDisplayInfo& DebugDisplay, float& YL, float& YPos)
 {
 	Super::DisplayDebug(Canvas, DebugDisplay, YL, YPos);
 
 	static FName NAME_AI = FName(TEXT("AI"));
-	if (DebugDisplay.Contains(NAME_AI))
+	if (DebugDisplay.IsDisplayOn(NAME_AI))
 	{
 		if (PathFollowingComponent)
 		{
diff --git a/Engine/Source/Runtime/Engine/Private/Actor.cpp b/Engine/Source/Runtime/Engine/Private/Actor.cpp
--- a/Engine/Source/Runtime/Engine/Private/Actor.cpp
+++ b/Engine/Source/Runtime/Engine/Private/Actor.cpp
@@ -9,6 +9,7 @@
 #include "LatentActions.h"
 #include "MessageLog.h"
 #include "Net/UnrealNetwork.h"
+#include "DisplayDebugHelpers.h"
 
 //DEFINE_LOG_CATEGORY_STATIC(LogActor, Log, All);
 DEFINE_LOG_CATEGORY(LogActor);
@@ -1712,7 +1713,7 @@ FString AActor::GetHumanReadableName() const
 	return GetName();
 }
 
-void AActor::DisplayDebug(UCanvas* Canvas, const TArray<FName>& DebugDisplay, float& YL, float& YPos)
+void AActor::DisplayDebug(UCanvas* Canvas, const FDebugDisplayInfo& DebugDisplay, float& YL, float& YPos)
 {
 	Canvas->SetDrawColor(255,0,0);
 
@@ -1730,7 +1731,7 @@ void AActor::DisplayDebug(UCanvas* Canvas, const TArray<FName>& DebugDisplay, fl
 
 	Canvas->SetDrawColor(255,255,255);
 
-	if( DebugDisplay.Contains(TEXT("net")) )
+	if( DebugDisplay.IsDisplayOn(TEXT("net")) )
 	{
 		if( GetNetMode() != NM_Standalone )
 		{
@@ -1749,13 +1750,13 @@ void AActor::DisplayDebug(UCanvas* Canvas, const TArray<FName>& DebugDisplay, fl
 	Canvas->DrawText(RenderFont, FString::Printf(TEXT("Location: %s Rotation: %s"), *GetActorLocation().ToString(), *GetActorRotation().ToString()), 4.0f, YPos);
 	YPos += YL;
 
-	if( DebugDisplay.Contains(TEXT("physics")) )
+	if( DebugDisplay.IsDisplayOn(TEXT("physics")) )
 	{
 		Canvas->DrawText(RenderFont,FString::Printf(TEXT("Velocity: %s Speed: %f Speed2D: %f"), *GetVelocity().ToString(), GetVelocity().Size(), GetVelocity().Size2D()), 4.0f, YPos);
 		YPos += YL;
 	}
 
-	if( DebugDisplay.Contains(TEXT("collision")) )
+	if( DebugDisplay.IsDisplayOn(TEXT("collision")) )
 	{
 		Canvas->DrawColor.B = 0;
 		float MyRadius, MyHeight;
@@ -1795,6 +1796,36 @@ void AActor::DisplayDebug(UCanvas* Canvas, const TArray<FName>& DebugDisplay, fl
 	Canvas->DrawText( RenderFont,FString::Printf(TEXT(" Instigator: %s Owner: %s"), (Instigator ? *Instigator->GetName() : TEXT("None")),
 		(Owner ? *Owner->GetName() : TEXT("None"))), 4,YPos);
 	YPos += YL;
+
+	static FName NAME_Animation(TEXT("Animation"));
+	static FName NAME_Bones = FName(TEXT("Bones"));
+	if (DebugDisplay.IsDisplayOn(NAME_Animation) || DebugDisplay.IsDisplayOn(NAME_Bones))
+	{
+		TArray<USkeletalMeshComponent*> Components;
+		GetComponents(Components);
+
+		if (DebugDisplay.IsDisplayOn(NAME_Animation))
+		{
+			for (USkeletalMeshComponent* Comp : Components)
+			{
+				UAnimInstance* AnimInstance = Comp->GetAnimInstance();
+				if (AnimInstance)
+				{
+					AnimInstance->DisplayDebug(Canvas, DebugDisplay, YL, YPos);
+				}
+			}
+		}
+
+		static FName NAME_3DBones = FName(TEXT("3DBones"));
+		if (DebugDisplay.IsDisplayOn(NAME_Bones))
+		{
+			bool bSimpleBones = !DebugDisplay.IsCategoryToggledOn(NAME_3DBones, false);
+			for (USkeletalMeshComponent* Comp : Components)
+			{
+				Comp->DebugDrawBones(Canvas, bSimpleBones);
+			}
+		}
+	}
 }
 
 void AActor::OutsideWorldBounds()
diff --git a/Engine/Source/Runtime/Engine/Private/Animation/AnimInstance.cpp b/Engine/Source/Runtime/Engine/Private/Animation/AnimInstance.cpp
--- a/Engine/Source/Runtime/Engine/Private/Animation/AnimInstance.cpp
+++ b/Engine/Source/Runtime/Engine/Private/Animation/AnimInstance.cpp
@@ -8,6 +8,7 @@
 #include "AnimationRuntime.h"
 #include "AnimationUtils.h"
 #include "ParticleDefinitions.h"
+#include "DisplayDebugHelpers.h"
 
 #include "MessageLog.h"
 
@@ -494,6 +495,312 @@ bool UAnimInstance::NativeEvaluateAnimation(FPoseContext& Output)
 	return false;
 }
 
+void OutputCurveMap(TMap<FName, float>& CurveMap, UCanvas* Canvas, UFont* RenderFont, float Indent, float& YPos, FFontRenderInfo RenderInfo, float& YL)
+{
+	TArray<FName> Names;
+	CurveMap.GetKeys(Names);
+	Names.Sort();
+	for (FName CurveName : Names)
+	{
+		FString CurveEntry = FString::Printf(TEXT("%s: %.3f"), *CurveName.ToString(), CurveMap[CurveName]);
+		Canvas->DrawText(RenderFont, CurveEntry, Indent, YPos, 1.f, 1.f, RenderInfo);
+		YPos += YL;
+	}
+}
+
+void OutputTickRecords(const TArray<FAnimTickRecord>& Records, UCanvas* Canvas, float Indent, const int32 HighlightIndex, FLinearColor TextColor, FLinearColor HighlightColor, FLinearColor InactiveColor, UFont* RenderFont, float& YPos, FFontRenderInfo RenderInfo, float& YL, bool bFullBlendspaceDisplay)
+{
+	for (int32 PlayerIndex = 0; PlayerIndex < Records.Num(); ++PlayerIndex)
+	{
+		const FAnimTickRecord& Player = Records[PlayerIndex];
+
+		Canvas->SetLinearDrawColor((PlayerIndex == HighlightIndex) ? HighlightColor : TextColor);
+
+		FString PlayerEntry = FString::Printf(TEXT("%i) %s (%s) W:%.1f%%"), PlayerIndex, *Player.SourceAsset->GetName(), *Player.SourceAsset->GetClass()->GetName(), Player.EffectiveBlendWeight*100.f);
+		Canvas->DrawText(RenderFont, PlayerEntry, Indent, YPos, 1.f, 1.f, RenderInfo);
+		YPos += YL;
+
+		if (UBlendSpaceBase* BlendSpace = Cast<UBlendSpaceBase>(Player.SourceAsset))
+		{
+			if (bFullBlendspaceDisplay && Player.BlendSampleDataCache && Player.BlendSampleDataCache->Num() > 0)
+			{
+				TArray<FBlendSampleData> SampleData = *Player.BlendSampleDataCache;
+				SampleData.Sort([](const FBlendSampleData& L, const FBlendSampleData& R) { return L.SampleDataIndex < R.SampleDataIndex; });
+
+				FIndenter BlendspaceIndent(Indent);
+				FString BlendspaceHeader = FString::Printf(TEXT("Blendspace Input (%.2f, %.2f, %.2f)"), Player.BlendSpacePosition.X, Player.BlendSpacePosition.Y, Player.BlendSpacePosition.Z);
+				Canvas->DrawText(RenderFont, BlendspaceHeader, Indent, YPos, 1.f, 1.f, RenderInfo);
+				YPos += YL;
+
+				const TArray<FBlendSample>& BlendSamples = BlendSpace->GetBlendSamples();
+
+				int32 WeightedSampleIndex = 0;
+
+				for (int32 SampleIndex = 0; SampleIndex < BlendSamples.Num(); ++SampleIndex)
+				{
+					const FBlendSample& BlendSample = BlendSamples[SampleIndex];
+
+					float Weight = 0.f;
+					for (; WeightedSampleIndex < SampleData.Num(); ++WeightedSampleIndex)
+					{
+						FBlendSampleData& WeightedSample = SampleData[WeightedSampleIndex];
+						if (WeightedSample.SampleDataIndex == SampleIndex)
+						{
+							Weight += WeightedSample.GetWeight();
+						}
+						else if (WeightedSample.SampleDataIndex > SampleIndex)
+						{
+							break;
+						}
+					}
+
+					FIndenter SampleIndent(Indent);
+
+					Canvas->SetLinearDrawColor((Weight > 0.f) ? TextColor : InactiveColor);
+
+					FString SampleEntry = FString::Printf(TEXT("%s W:%.1f%%"), *BlendSample.Animation->GetName(), Weight*100.f);
+					Canvas->DrawText(RenderFont, SampleEntry, Indent, YPos, 1.f, 1.f, RenderInfo);
+					YPos += YL;
+				}
+			}
+		}
+	}
+}
+
+void UAnimInstance::DisplayDebug(class UCanvas* Canvas, const FDebugDisplayInfo& DebugDisplay, float& YL, float& YPos)
+{
+	float Indent = 0.f;
+
+	UFont* RenderFont = GEngine->GetSmallFont();
+
+	FLinearColor TextYellow(0.86f, 0.69f, 0.f);
+	FLinearColor TextWhite(0.9f, 0.9f, 0.9f);
+	FLinearColor ActiveColor(0.1f, 0.6f, 0.1f);
+	FLinearColor InactiveColor(0.2f, 0.2f, 0.2f);
+	FLinearColor PoseSourceColor(0.5f, 0.25f, 0.5f);
+
+	Canvas->SetLinearDrawColor(TextYellow);
+
+	static FName CAT_SyncGroups(TEXT("SyncGroups"));
+	static FName CAT_Montages(TEXT("Montages"));
+	static FName CAT_Graph(TEXT("Graph"));
+	static FName CAT_Curves(TEXT("Curves"));
+	static FName CAT_Notifies(TEXT("Notifies"));
+	static FName CAT_FullAnimGraph(TEXT("FullGraph"));
+	static FName CAT_FullBlendspaceDisplay(TEXT("FullBlendspaceDisplay"));
+
+	const bool bShowSyncGroups = DebugDisplay.IsCategoryToggledOn(CAT_SyncGroups, true);
+	const bool bShowMontages = DebugDisplay.IsCategoryToggledOn(CAT_Montages, true);
+	const bool bShowGraph = DebugDisplay.IsCategoryToggledOn(CAT_Graph, true);
+	const bool bShowCurves = DebugDisplay.IsCategoryToggledOn(CAT_Curves, true);
+	const bool bShowNotifies = DebugDisplay.IsCategoryToggledOn(CAT_Notifies, true);
+	const bool bFullGraph = DebugDisplay.IsCategoryToggledOn(CAT_FullAnimGraph, false);
+	const bool bFullBlendspaceDisplay = DebugDisplay.IsCategoryToggledOn(CAT_FullBlendspaceDisplay, true);
+
+	FFontRenderInfo RenderInfo;
+	RenderInfo.bEnableShadow = true;
+
+	YPos += YL;
+
+	Canvas->SetLinearDrawColor(TextYellow);
+
+	FString Heading = FString::Printf(TEXT("Animation: %s"), *GetName());
+	Canvas->DrawText(RenderFont, Heading, Indent, YPos, 1.f, 1.f, RenderInfo);
+	YPos += YL;
+	if (bShowSyncGroups)
+	{
+		FIndenter AnimIndent(Indent);
+
+		//Display Sync Groups
+		Heading = FString::Printf(TEXT("SyncGroups: %i"), SyncGroups.Num());
+		Canvas->DrawText(RenderFont, Heading, Indent, YPos, 1.f, 1.f, RenderInfo);
+		YPos += YL;
+
+		for (int32 GroupIndex = 0; GroupIndex < SyncGroups.Num(); ++GroupIndex)
+		{
+			FIndenter GroupIndent(Indent);
+			FAnimGroupInstance& SyncGroup = SyncGroups[GroupIndex];
+
+			Canvas->SetLinearDrawColor(TextYellow);
+
+			FString GroupLabel = FString::Printf(TEXT("Group %i - Players %i"), GroupIndex, SyncGroup.ActivePlayers.Num());
+			Canvas->DrawText(RenderFont, GroupLabel, Indent, YPos, 1.f, 1.f, RenderInfo);
+			YPos += YL;
+
+			if (SyncGroup.ActivePlayers.Num() > 0)
+			{
+				const int32 GroupLeaderIndex = FMath::Max(SyncGroup.GroupLeaderIndex, 0);
+				OutputTickRecords(SyncGroup.ActivePlayers, Canvas, Indent, GroupLeaderIndex, TextWhite, ActiveColor, InactiveColor, RenderFont, YPos, RenderInfo, YL, bFullBlendspaceDisplay);
+			}
+		}
+
+		Canvas->SetLinearDrawColor(TextYellow);
+
+		Heading = FString::Printf(TEXT("Ungrouped: %i"), UngroupedActivePlayers.Num());
+		Canvas->DrawText(RenderFont, Heading, Indent, YPos, 1.f, 1.f, RenderInfo);
+		YPos += YL;
+
+		Canvas->SetLinearDrawColor(TextWhite);
+
+		OutputTickRecords(UngroupedActivePlayers, Canvas, Indent, -1, TextWhite, ActiveColor, InactiveColor, RenderFont, YPos, RenderInfo, YL, bFullBlendspaceDisplay);
+	}
+
+	if (bShowMontages)
+	{
+		Canvas->SetLinearDrawColor(TextYellow);
+
+		Heading = FString::Printf(TEXT("Montages: %i"), MontageInstances.Num());
+		Canvas->DrawText(RenderFont, Heading, Indent, YPos, 1.f, 1.f, RenderInfo);
+		YPos += YL;
+
+		FAnimMontageInstance* ActiveMontageInstance = GetActiveMontageInstance();
+
+		for (int32 MontageIndex = 0; MontageIndex < MontageInstances.Num(); ++MontageIndex)
+		{
+			FIndenter PlayerIndent(Indent);
+
+			FAnimMontageInstance* MontageInstance = MontageInstances[MontageIndex];
+
+			Canvas->SetLinearDrawColor((MontageInstance == ActiveMontageInstance) ? ActiveColor : TextWhite);
+
+			FString MontageEntry = FString::Printf(TEXT("%i) %s Sec: %s W:%.3f DW:%.3f"), MontageIndex, *MontageInstance->Montage->GetName(), *MontageInstance->GetCurrentSection().ToString(), *MontageInstance->GetNextSection().ToString(), MontageInstance->Weight, MontageInstance->DesiredWeight);
+			Canvas->DrawText(RenderFont, MontageEntry, Indent, YPos, 1.f, 1.f, RenderInfo);
+			YPos += YL;
+		}
+	}
+
+	if (bShowNotifies)
+	{
+		Canvas->SetLinearDrawColor(TextYellow);
+
+		Heading = FString::Printf(TEXT("Active Notify States: %i"), ActiveAnimNotifyState.Num());
+		Canvas->DrawText(RenderFont, Heading, Indent, YPos, 1.f, 1.f, RenderInfo);
+		YPos += YL;
+
+		Canvas->SetLinearDrawColor(TextWhite);
+
+		for (int32 NotifyIndex = 0; NotifyIndex < ActiveAnimNotifyState.Num(); ++NotifyIndex)
+		{
+			FIndenter NotifyIndent(Indent);
+
+			const FAnimNotifyEvent* NotifyState = ActiveAnimNotifyState[NotifyIndex];
+
+			FString NotifyEntry = FString::Printf(TEXT("%i) %s Class: %s Dur:%.3f"), NotifyIndex, *NotifyState->NotifyName.ToString(), *NotifyState->NotifyStateClass->GetName(), NotifyState->Duration);
+			Canvas->DrawText(RenderFont, NotifyEntry, Indent, YPos, 1.f, 1.f, RenderInfo);
+			YPos += YL;
+		}
+	}
+
+	if (bShowCurves)
+	{
+		Canvas->SetLinearDrawColor(TextYellow);
+
+		Canvas->DrawText(RenderFont, TEXT("Curves"), Indent, YPos, 1.f, 1.f, RenderInfo);
+		YPos += YL;
+
+		{
+			FIndenter CurveIndent(Indent);
+
+			Heading = FString::Printf(TEXT("Morph Curves: %i"), MorphTargetCurves.Num());
+			Canvas->DrawText(RenderFont, Heading, Indent, YPos, 1.f, 1.f, RenderInfo);
+			YPos += YL;
+
+			Canvas->SetLinearDrawColor(TextWhite);
+
+			{
+				FIndenter MorphCurveIndent(Indent);
+				OutputCurveMap(MorphTargetCurves, Canvas, RenderFont, Indent, YPos, RenderInfo, YL);
+			}
+
+			Canvas->SetLinearDrawColor(TextYellow);
+
+			Heading = FString::Printf(TEXT("Material Curves: %i"), MaterialParameterCurves.Num());
+			Canvas->DrawText(RenderFont, Heading, Indent, YPos, 1.f, 1.f, RenderInfo);
+			YPos += YL;
+
+			Canvas->SetLinearDrawColor(TextWhite);
+
+			{
+				FIndenter MaterialCurveIndent(Indent);
+				OutputCurveMap(MaterialParameterCurves, Canvas, RenderFont, Indent, YPos, RenderInfo, YL);
+			}
+		}
+	}
+
+	if (bShowGraph)
+	{
+		Canvas->SetLinearDrawColor(TextYellow);
+
+		YPos += YL;
+		Canvas->DrawText(RenderFont, TEXT("Anim Node Tree"), Indent, YPos, 1.f, 1.f, RenderInfo);
+
+		const float NodeIndent = 8.f;
+		const float LineIndent = 4.f;
+		const float AttachLineLength = NodeIndent - LineIndent;
+
+		YPos += YL;
+		FIndenter AnimNodeTreeIndent(Indent);
+
+		FNodeDebugData NodeDebugData(this);
+		RootNode->GatherDebugData(NodeDebugData);
+
+		TArray<FNodeDebugData::FFlattenedDebugData> FlattenedData = NodeDebugData.GetFlattenedDebugData();
+
+		TArray<float> VerticalLineStarts; // Index represents indent level, track the current starting point for that 
+
+		int32 HalfStep = int32(YL / 2);
+		int32 PrevChainID = -1;
+
+		for (FNodeDebugData::FFlattenedDebugData& Line : FlattenedData)
+		{
+			if (!Line.IsOnActiveBranch() && !bFullGraph)
+			{
+				continue;
+			}
+			float CurrIndent = Indent + (Line.Indent * NodeIndent);
+			float CurrLineYBase = YPos + YL;
+
+			if (PrevChainID != Line.ChainID)
+			{
+				YPos += HalfStep; // Extra spacing to delimit different chains, CurrLineYBase now 
+				// roughly represents middle of text line, so we can use it for line drawing
+
+				//Handle line drawing
+				int32 VerticalLineIndex = Line.Indent - 1;
+				if (VerticalLineStarts.IsValidIndex(VerticalLineIndex))
+				{
+					float VerticalLineStartY = VerticalLineStarts[VerticalLineIndex];
+					VerticalLineStarts[VerticalLineIndex] = CurrLineYBase;
+
+					float EndX = CurrIndent;
+					float StartX = EndX - AttachLineLength;
+
+					//horizontal line to node
+					DrawDebugCanvas2DLine(Canvas, FVector(StartX, CurrLineYBase, 0.f), FVector(EndX, CurrLineYBase, 0.f), ActiveColor);
+
+					//vertical line
+					DrawDebugCanvas2DLine(Canvas, FVector(StartX, VerticalLineStartY, 0.f), FVector(StartX, CurrLineYBase, 0.f), ActiveColor);
+				}
+
+				CurrLineYBase += HalfStep; // move CurrYLineBase back to base of line
+			}
+
+			// Update our base position for subsequent line drawing
+			if (!VerticalLineStarts.IsValidIndex(Line.Indent))
+			{
+				VerticalLineStarts.AddZeroed(Line.Indent + 1 - VerticalLineStarts.Num());
+			}
+			VerticalLineStarts[Line.Indent] = CurrLineYBase;
+
+			PrevChainID = Line.ChainID;
+			FLinearColor ItemColor = Line.bPoseSource ? PoseSourceColor : ActiveColor;
+			Canvas->SetLinearDrawColor(Line.IsOnActiveBranch() ? ItemColor : InactiveColor);
+			Canvas->DrawText(RenderFont, Line.DebugLine, CurrIndent, YPos, 1.f, 1.f, RenderInfo);
+			YPos += YL;
+		}
+	}
+}
+
 void UAnimInstance::BlendSpaceEvaluatePose(class UBlendSpaceBase* BlendSpace, TArray<FBlendSampleData>& BlendSampleDataCache, struct FA2Pose& Pose, bool bIsLooping)
 {
 	SCOPE_CYCLE_COUNTER(STAT_AnimNativeEvaluatePoses);
diff --git a/Engine/Source/Runtime/Engine/Private/Animation/AnimMontage.cpp b/Engine/Source/Runtime/Engine/Private/Animation/AnimMontage.cpp
--- a/Engine/Source/Runtime/Engine/Private/Animation/AnimMontage.cpp
+++ b/Engine/Source/Runtime/Engine/Private/Animation/AnimMontage.cpp
@@ -830,6 +830,22 @@ FName FAnimMontageInstance::GetCurrentSection() const
 	return NAME_None;
 }
 
+FName FAnimMontageInstance::GetNextSection() const
+{
+	if (Montage)
+	{
+		float CurrentPosition;
+		const int32 CurrentSectionIndex = Montage->GetAnimCompositeSectionIndexFromPos(Position, CurrentPosition);
+		if (Montage->IsValidSectionIndex(CurrentSectionIndex))
+		{
+			FCompositeSection& CurrentSection = Montage->GetAnimCompositeSection(CurrentSectionIndex);
+			return CurrentSection.NextSectionName;
+		}
+	}
+
+	return NAME_None;
+}
+
 bool FAnimMontageInstance::ChangePositionToSection(FName SectionName, bool bEndOfSection)
 {
 	const int32 SectionID = Montage->GetSectionIndex(SectionName);
diff --git a/Engine/Source/Runtime/Engine/Private/Animation/AnimNodeBase.cpp b/Engine/Source/Runtime/Engine/Private/Animation/AnimNodeBase.cpp
--- a/Engine/Source/Runtime/Engine/Private/Animation/AnimNodeBase.cpp
+++ b/Engine/Source/Runtime/Engine/Private/Animation/AnimNodeBase.cpp
@@ -82,6 +82,14 @@ void FPoseLinkBase::Update(const FAnimationUpdateContext& Context)
 	}
 }
 
+void FPoseLinkBase::GatherDebugData(FNodeDebugData& DebugData)
+{
+	if(LinkedNode != NULL)
+	{
+		LinkedNode->GatherDebugData(DebugData);
+	}
+}
+
 /////////////////////////////////////////////////////
 // FPoseLink
 
@@ -206,4 +214,42 @@ bool FComponentSpacePoseContext::IsNormalized() const
 	}
 
 	return true;
+}
+
+/////////////////////////////////////////////////////
+// FNodeDebugData
+
+void FNodeDebugData::AddDebugItem(FString DebugData, bool bPoseSource)
+{
+	check(NodeChain.Num() == 0 || NodeChain.Last().ChildNodeChain.Num() == 0); //Cannot add to this chain once we have branched
+
+	NodeChain.Add( DebugItem(DebugData, bPoseSource) );
+}
+
+FNodeDebugData& FNodeDebugData::BranchFlow(float BranchWeight)
+{
+	DebugItem& LatestItem = NodeChain.Last();
+	new (LatestItem.ChildNodeChain) FNodeDebugData(AnimInstance, BranchWeight*AbsoluteWeight);
+	return LatestItem.ChildNodeChain.Last();
+}
+
+void FNodeDebugData::GetFlattenedDebugData(TArray<FFlattenedDebugData>& FlattenedDebugData, int32 Indent, int32& ChainID)
+{
+	int32 CurrChainID = ChainID;
+	for(DebugItem& Item : NodeChain)
+	{
+		FlattenedDebugData.Add( FFlattenedDebugData(Item.DebugData, AbsoluteWeight, Indent, CurrChainID, Item.bPoseSource) );
+		bool bMultiBranch = Item.ChildNodeChain.Num() > 1;
+		int32 ChildIndent = bMultiBranch ? Indent + 1 : Indent;
+		for(FNodeDebugData& Child : Item.ChildNodeChain)
+		{
+			if(bMultiBranch)
+			{
+				// If we only have one branch we treat it as the same really
+				// as we may have only changed active status
+				++ChainID;
+			}
+			Child.GetFlattenedDebugData(FlattenedDebugData, ChildIndent, ChainID);
+		}
+	}
 }
\ No newline at end of file
diff --git a/Engine/Source/Runtime/Engine/Private/Animation/AnimNodeRotationMultiplier.cpp b/Engine/Source/Runtime/Engine/Private/Animation/AnimNodeRotationMultiplier.cpp
--- a/Engine/Source/Runtime/Engine/Private/Animation/AnimNodeRotationMultiplier.cpp
+++ b/Engine/Source/Runtime/Engine/Private/Animation/AnimNodeRotationMultiplier.cpp
@@ -10,6 +10,18 @@ FAnimNode_RotationMultiplier::FAnimNode_RotationMultiplier()
 {
 }
 
+void FAnimNode_RotationMultiplier::GatherDebugData(FNodeDebugData& DebugData)
+{
+	FString DebugLine = DebugData.GetNodeName(this);
+
+	DebugLine += "(";
+	AddDebugNodeData(DebugLine);
+	DebugLine += FString::Printf(TEXT(" Src: %s Dst: %s Multiplier: %.2f)"), *SourceBone.BoneName.ToString(), *TargetBone.BoneName.ToString(), Multiplier);
+	DebugData.AddDebugItem(DebugLine);
+
+	ComponentPose.GatherDebugData(DebugData);
+}
+
 FVector GetAxisVector(const EBoneAxis Axis)
 {
 	switch (Axis)
diff --git a/Engine/Source/Runtime/Engine/Private/Animation/AnimNodeSkeletalControlBase.cpp b/Engine/Source/Runtime/Engine/Private/Animation/AnimNodeSkeletalControlBase.cpp
--- a/Engine/Source/Runtime/Engine/Private/Animation/AnimNodeSkeletalControlBase.cpp
+++ b/Engine/Source/Runtime/Engine/Private/Animation/AnimNodeSkeletalControlBase.cpp
@@ -61,3 +61,8 @@ void FAnimNode_SkeletalControlBase::EvaluateComponentSpace(FComponentSpacePoseCo
 	}
 }
 
+void FAnimNode_SkeletalControlBase::AddDebugNodeData(FString& OutDebugData)
+{
+	const float ActualAlpha = AlphaScaleBias.ApplyTo(Alpha);
+	OutDebugData += FString::Printf(TEXT("Alpha: %.1f%%"), ActualAlpha*100.f);
+}
\ No newline at end of file
diff --git a/Engine/Source/Runtime/Engine/Private/Animation/AnimNodeSpaceConversions.cpp b/Engine/Source/Runtime/Engine/Private/Animation/AnimNodeSpaceConversions.cpp
--- a/Engine/Source/Runtime/Engine/Private/Animation/AnimNodeSpaceConversions.cpp
+++ b/Engine/Source/Runtime/Engine/Private/Animation/AnimNodeSpaceConversions.cpp
@@ -34,6 +34,13 @@ void FAnimNode_ConvertComponentToLocalSpace::Evaluate(FPoseContext & Output)
 	InputCSPose.Pose.ConvertToLocalPoses(Output.Pose);
 }
 
+void FAnimNode_ConvertComponentToLocalSpace::GatherDebugData(FNodeDebugData& DebugData)
+{
+	FString DebugLine = DebugData.GetNodeName(this);
+	DebugData.AddDebugItem(DebugLine);
+	ComponentPose.GatherDebugData(DebugData);
+}
+
 /////////////////////////////////////////////////////
 // FAnimNode_ConvertLocalToComponentSpace
 
@@ -56,6 +63,13 @@ void FAnimNode_ConvertLocalToComponentSpace::Update(const FAnimationUpdateContex
 	LocalPose.Update(Context);
 }
 
+void FAnimNode_ConvertLocalToComponentSpace::GatherDebugData(FNodeDebugData& DebugData)
+{
+	FString DebugLine = DebugData.GetNodeName(this);
+	DebugData.AddDebugItem(DebugLine);
+	LocalPose.GatherDebugData(DebugData);
+}
+
 void FAnimNode_ConvertLocalToComponentSpace::EvaluateComponentSpace(FComponentSpacePoseContext & OutputCSPose)
 {
 	// Evaluate the child and convert
diff --git a/Engine/Source/Runtime/Engine/Private/Animation/AnimNode_ApplyAdditive.cpp b/Engine/Source/Runtime/Engine/Private/Animation/AnimNode_ApplyAdditive.cpp
--- a/Engine/Source/Runtime/Engine/Private/Animation/AnimNode_ApplyAdditive.cpp
+++ b/Engine/Source/Runtime/Engine/Private/Animation/AnimNode_ApplyAdditive.cpp
@@ -53,3 +53,15 @@ FAnimNode_ApplyAdditive::FAnimNode_ApplyAdditive()
 	: Alpha(1.0f)
 {
 }
+
+void FAnimNode_ApplyAdditive::GatherDebugData(FNodeDebugData& DebugData)
+{
+	const float ActualAlpha = AlphaScaleBias.ApplyTo(Alpha);
+
+	FString DebugLine = DebugData.GetNodeName(this);
+	DebugLine += FString::Printf(TEXT("(Alpha: %.1f%%)"), ActualAlpha*100.f);
+
+	DebugData.AddDebugItem(DebugLine);
+	Base.GatherDebugData(DebugData.BranchFlow(1.f));
+	Additive.GatherDebugData(DebugData.BranchFlow(ActualAlpha));
+}
\ No newline at end of file
diff --git a/Engine/Source/Runtime/Engine/Private/Animation/AnimNode_BlendListBase.cpp b/Engine/Source/Runtime/Engine/Private/Animation/AnimNode_BlendListBase.cpp
--- a/Engine/Source/Runtime/Engine/Private/Animation/AnimNode_BlendListBase.cpp
+++ b/Engine/Source/Runtime/Engine/Private/Animation/AnimNode_BlendListBase.cpp
@@ -147,4 +147,20 @@ void FAnimNode_BlendListBase::Evaluate(FPoseContext& Output)
 	{
 		Output.ResetToRefPose();
 	}
+}
+
+void FAnimNode_BlendListBase::GatherDebugData(FNodeDebugData& DebugData)
+{
+	const int NumPoses = BlendPose.Num();
+	const int32 ChildIndex = GetActiveChildIndex();
+
+	FString DebugLine = GetNodeName(DebugData);
+	DebugLine += FString::Printf(TEXT("(Active: (%i/%i) BlendWeight: %.1f%% BlendTime %.3f)"), ChildIndex+1, NumPoses, BlendWeights[ChildIndex]*100.f, BlendTime[ChildIndex]);
+
+	DebugData.AddDebugItem(DebugLine);
+	
+	for(int32 Pose = 0; Pose < NumPoses; ++Pose)
+	{
+		BlendPose[Pose].GatherDebugData(DebugData.BranchFlow(BlendWeights[Pose]));
+	}
 }
\ No newline at end of file
diff --git a/Engine/Source/Runtime/Engine/Private/Animation/AnimNode_BlendSpaceEvaluator.cpp b/Engine/Source/Runtime/Engine/Private/Animation/AnimNode_BlendSpaceEvaluator.cpp
--- a/Engine/Source/Runtime/Engine/Private/Animation/AnimNode_BlendSpaceEvaluator.cpp
+++ b/Engine/Source/Runtime/Engine/Private/Animation/AnimNode_BlendSpaceEvaluator.cpp
@@ -21,3 +21,11 @@ void FAnimNode_BlendSpaceEvaluator::Update(const FAnimationUpdateContext& Contex
 	FAnimNode_BlendSpacePlayer::Update(Context);
 }
 
+
+void FAnimNode_BlendSpaceEvaluator::GatherDebugData(FNodeDebugData& DebugData)
+{
+	FString DebugLine = DebugData.GetNodeName(this);
+	
+	DebugLine += FString::Printf(TEXT("('%s' Play Time: %.3f)"), *BlendSpace->GetName(), InternalTimeAccumulator);
+	DebugData.AddDebugItem(DebugLine, true);
+}
diff --git a/Engine/Source/Runtime/Engine/Private/Animation/AnimNode_BlendSpacePlayer.cpp b/Engine/Source/Runtime/Engine/Private/Animation/AnimNode_BlendSpacePlayer.cpp
--- a/Engine/Source/Runtime/Engine/Private/Animation/AnimNode_BlendSpacePlayer.cpp
+++ b/Engine/Source/Runtime/Engine/Private/Animation/AnimNode_BlendSpacePlayer.cpp
@@ -83,3 +83,16 @@ void FAnimNode_BlendSpacePlayer::OverrideAsset(UAnimationAsset* NewAsset)
 		BlendSpace = NewBlendSpace;
 	}
 }
+
+void FAnimNode_BlendSpacePlayer::GatherDebugData(FNodeDebugData& DebugData)
+{
+	FString DebugLine = DebugData.GetNodeName(this);
+	if (BlendSpace)
+	{
+		DebugLine += FString::Printf(TEXT("('%s' Play Time: %.3f)"), *BlendSpace->GetName(), InternalTimeAccumulator);
+
+		DebugData.AddDebugItem(DebugLine, true);
+
+		//BlendSpace->GetBlendSamples();
+	}
+}
\ No newline at end of file
diff --git a/Engine/Source/Runtime/Engine/Private/Animation/AnimNode_CopyBone.cpp b/Engine/Source/Runtime/Engine/Private/Animation/AnimNode_CopyBone.cpp
--- a/Engine/Source/Runtime/Engine/Private/Animation/AnimNode_CopyBone.cpp
+++ b/Engine/Source/Runtime/Engine/Private/Animation/AnimNode_CopyBone.cpp
@@ -12,6 +12,18 @@ FAnimNode_CopyBone::FAnimNode_CopyBone()
 {
 }
 
+void FAnimNode_CopyBone::GatherDebugData(FNodeDebugData& DebugData)
+{
+	FString DebugLine = DebugData.GetNodeName(this);
+	
+	DebugLine += "(";
+	AddDebugNodeData(DebugLine);
+	DebugLine += FString::Printf(TEXT(" Src: %s Dst: %s)"), *SourceBone.BoneName.ToString(), *TargetBone.BoneName.ToString());
+	DebugData.AddDebugItem(DebugLine);
+
+	ComponentPose.GatherDebugData(DebugData);
+}
+
 void FAnimNode_CopyBone::EvaluateBoneTransforms(USkeletalMeshComponent* SkelComp, const FBoneContainer & RequiredBones, FA2CSPose& MeshBases, TArray<FBoneTransform>& OutBoneTransforms)
 {
 	check(OutBoneTransforms.Num() == 0);
diff --git a/Engine/Source/Runtime/Engine/Private/Animation/AnimNode_LayeredBoneBlend.cpp b/Engine/Source/Runtime/Engine/Private/Animation/AnimNode_LayeredBoneBlend.cpp
--- a/Engine/Source/Runtime/Engine/Private/Animation/AnimNode_LayeredBoneBlend.cpp
+++ b/Engine/Source/Runtime/Engine/Private/Animation/AnimNode_LayeredBoneBlend.cpp
@@ -107,4 +107,21 @@ void FAnimNode_LayeredBoneBlend::Evaluate(FPoseContext& Output)
 
 		FAnimationRuntime::BlendPosesPerBoneFilter(BasePoseContext.Pose, TargetBlendPoses, Output.Pose, CurrentBoneBlendWeights, bMeshSpaceRotationBlend, Output.AnimInstance->RequiredBones, Output.AnimInstance->CurrentSkeleton);
 	}
+}
+
+
+void FAnimNode_LayeredBoneBlend::GatherDebugData(FNodeDebugData& DebugData)
+{
+	const int NumPoses = BlendPoses.Num();
+
+	FString DebugLine = DebugData.GetNodeName(this);
+	DebugLine += FString::Printf(TEXT("(Num Poses: %i)"), NumPoses);
+	DebugData.AddDebugItem(DebugLine);
+
+	BasePose.GatherDebugData(DebugData.BranchFlow(1.f));
+	
+	for (int32 ChildIndex = 0; ChildIndex < NumPoses; ++ChildIndex)
+	{
+		BlendPoses[ChildIndex].GatherDebugData(DebugData.BranchFlow(BlendWeights[ChildIndex]));
+	}
 }
\ No newline at end of file
diff --git a/Engine/Source/Runtime/Engine/Private/Animation/AnimNode_ModifyBone.cpp b/Engine/Source/Runtime/Engine/Private/Animation/AnimNode_ModifyBone.cpp
--- a/Engine/Source/Runtime/Engine/Private/Animation/AnimNode_ModifyBone.cpp
+++ b/Engine/Source/Runtime/Engine/Private/Animation/AnimNode_ModifyBone.cpp
@@ -19,6 +19,18 @@ FAnimNode_ModifyBone::FAnimNode_ModifyBone()
 {
 }
 
+void FAnimNode_ModifyBone::GatherDebugData(FNodeDebugData& DebugData)
+{
+	FString DebugLine = DebugData.GetNodeName(this);
+
+	DebugLine += "(";
+	AddDebugNodeData(DebugLine);
+	DebugLine += FString::Printf(TEXT(" Target: %s)"), *BoneToModify.BoneName.ToString());
+	DebugData.AddDebugItem(DebugLine);
+
+	ComponentPose.GatherDebugData(DebugData);
+}
+
 void FAnimNode_ModifyBone::EvaluateBoneTransforms(USkeletalMeshComponent* SkelComp, const FBoneContainer & RequiredBones, FA2CSPose& MeshBases, TArray<FBoneTransform>& OutBoneTransforms)
 {
 	check(OutBoneTransforms.Num() == 0);
diff --git a/Engine/Source/Runtime/Engine/Private/Animation/AnimNode_RefPose.cpp b/Engine/Source/Runtime/Engine/Private/Animation/AnimNode_RefPose.cpp
--- a/Engine/Source/Runtime/Engine/Private/Animation/AnimNode_RefPose.cpp
+++ b/Engine/Source/Runtime/Engine/Private/Animation/AnimNode_RefPose.cpp
@@ -26,4 +26,25 @@ void FAnimNode_RefPose::Evaluate(FPoseContext& Output)
 void FAnimNode_MeshSpaceRefPose::EvaluateComponentSpace(FComponentSpacePoseContext& Output)
 {
 	Output.ResetToRefPose();
+}
+
+/** Helper for enum output... */
+#ifndef CASE_ENUM_TO_TEXT
+#define CASE_ENUM_TO_TEXT(txt) case txt: return TEXT(#txt);
+#endif
+
+const TCHAR* GetRefPostTypeText(ERefPoseType RefPose)
+{
+	switch (RefPose)
+	{
+		FOREACH_ENUM_EREFPOSETYPE(CASE_ENUM_TO_TEXT)
+	}
+	return TEXT("Unknown Ref Pose Type");
+}
+
+void FAnimNode_RefPose::GatherDebugData(FNodeDebugData& DebugData)
+{
+	FString DebugLine = DebugData.GetNodeName(this);
+	DebugLine += FString::Printf(TEXT("(Ref Pose Type: %s)"), GetRefPostTypeText(RefPoseType));
+	DebugData.AddDebugItem(DebugLine, true);
 }
\ No newline at end of file
diff --git a/Engine/Source/Runtime/Engine/Private/Animation/AnimNode_Root.cpp b/Engine/Source/Runtime/Engine/Private/Animation/AnimNode_Root.cpp
--- a/Engine/Source/Runtime/Engine/Private/Animation/AnimNode_Root.cpp
+++ b/Engine/Source/Runtime/Engine/Private/Animation/AnimNode_Root.cpp
@@ -29,3 +29,10 @@ void FAnimNode_Root::Evaluate(FPoseContext& Output)
 {
 	Result.Evaluate(Output);
 }
+
+void FAnimNode_Root::GatherDebugData(FNodeDebugData& DebugData)
+{
+	FString DebugLine = DebugData.GetNodeName(this);
+	DebugData.AddDebugItem(DebugLine);
+	Result.GatherDebugData(DebugData);
+}
\ No newline at end of file
diff --git a/Engine/Source/Runtime/Engine/Private/Animation/AnimNode_RotateRootBone.cpp b/Engine/Source/Runtime/Engine/Private/Animation/AnimNode_RotateRootBone.cpp
--- a/Engine/Source/Runtime/Engine/Private/Animation/AnimNode_RotateRootBone.cpp
+++ b/Engine/Source/Runtime/Engine/Private/Animation/AnimNode_RotateRootBone.cpp
@@ -41,6 +41,17 @@ void FAnimNode_RotateRootBone::Evaluate(FPoseContext& Output)
 	Output.Pose.Bones[0].SetRotation(Output.Pose.Bones[0].GetRotation() * MeshSpaceDeltaQuat);
 }
 
+
+void FAnimNode_RotateRootBone::GatherDebugData(FNodeDebugData& DebugData)
+{
+	FString DebugLine = DebugData.GetNodeName(this);
+	
+	DebugLine += FString::Printf(TEXT("(Pitch: %.2f Yaw: %.2f)"), Pitch, Yaw);
+	DebugData.AddDebugItem(DebugLine);
+
+	BasePose.GatherDebugData(DebugData);
+}
+
 FAnimNode_RotateRootBone::FAnimNode_RotateRootBone()
 	: Pitch(0.0f)
 	, Yaw(0.0f)
diff --git a/Engine/Source/Runtime/Engine/Private/Animation/AnimNode_RotationOffsetBlendSpace.cpp b/Engine/Source/Runtime/Engine/Private/Animation/AnimNode_RotationOffsetBlendSpace.cpp
--- a/Engine/Source/Runtime/Engine/Private/Animation/AnimNode_RotationOffsetBlendSpace.cpp
+++ b/Engine/Source/Runtime/Engine/Private/Animation/AnimNode_RotationOffsetBlendSpace.cpp
@@ -35,6 +35,16 @@ void FAnimNode_RotationOffsetBlendSpace::Evaluate(FPoseContext& Context)
 	Context.AnimInstance->BlendRotationOffset(BasePoseContext.Pose, BlendSpacePose.Pose, Alpha, Context.Pose);
 }
 
+void FAnimNode_RotationOffsetBlendSpace::GatherDebugData(FNodeDebugData& DebugData)
+{
+	FString DebugLine = DebugData.GetNodeName(this);
+	
+	DebugLine += FString::Printf(TEXT("(Play Time: %.3f)"), InternalTimeAccumulator);
+	DebugData.AddDebugItem(DebugLine);
+	
+	BasePose.GatherDebugData(DebugData);
+}
+
 FAnimNode_RotationOffsetBlendSpace::FAnimNode_RotationOffsetBlendSpace()
 {
 }
diff --git a/Engine/Source/Runtime/Engine/Private/Animation/AnimNode_SaveCachedPose.cpp b/Engine/Source/Runtime/Engine/Private/Animation/AnimNode_SaveCachedPose.cpp
--- a/Engine/Source/Runtime/Engine/Private/Animation/AnimNode_SaveCachedPose.cpp
+++ b/Engine/Source/Runtime/Engine/Private/Animation/AnimNode_SaveCachedPose.cpp
@@ -60,3 +60,12 @@ void FAnimNode_SaveCachedPose::Evaluate(FPoseContext& Output)
 	// Return the cached result
 	Output.AnimInstance->CopyPose(CachedPose, Output.Pose);
 }
+
+
+void FAnimNode_SaveCachedPose::GatherDebugData(FNodeDebugData& DebugData)
+{
+	FString DebugLine = DebugData.GetNodeName(this);
+	DebugData.AddDebugItem(DebugLine);
+
+	Pose.GatherDebugData(DebugData);
+}
\ No newline at end of file
diff --git a/Engine/Source/Runtime/Engine/Private/Animation/AnimNode_SequenceEvaluator.cpp b/Engine/Source/Runtime/Engine/Private/Animation/AnimNode_SequenceEvaluator.cpp
--- a/Engine/Source/Runtime/Engine/Private/Animation/AnimNode_SequenceEvaluator.cpp
+++ b/Engine/Source/Runtime/Engine/Private/Animation/AnimNode_SequenceEvaluator.cpp
@@ -29,7 +29,6 @@ void FAnimNode_SequenceEvaluator::Evaluate(FPoseContext& Output)
 		Output.ResetToRefPose();
 	}
 }
-
 void FAnimNode_SequenceEvaluator::OverrideAsset(UAnimationAsset* NewAsset)
 {
 	if(UAnimSequenceBase* NewSequence = Cast<UAnimSequenceBase>(NewAsset))
@@ -37,3 +36,11 @@ void FAnimNode_SequenceEvaluator::OverrideAsset(UAnimationAsset* NewAsset)
 		Sequence = NewSequence;
 	}
 }
+
+void FAnimNode_SequenceEvaluator::GatherDebugData(FNodeDebugData& DebugData)
+{
+	FString DebugLine = DebugData.GetNodeName(this);
+	
+	DebugLine += FString::Printf(TEXT("('%s' Play Time: %.3f)"), *Sequence->GetName(), ExplicitTime);
+	DebugData.AddDebugItem(DebugLine, true);
+}
\ No newline at end of file
diff --git a/Engine/Source/Runtime/Engine/Private/Animation/AnimNode_SequencePlayer.cpp b/Engine/Source/Runtime/Engine/Private/Animation/AnimNode_SequencePlayer.cpp
--- a/Engine/Source/Runtime/Engine/Private/Animation/AnimNode_SequencePlayer.cpp
+++ b/Engine/Source/Runtime/Engine/Private/Animation/AnimNode_SequencePlayer.cpp
@@ -64,3 +64,11 @@ void FAnimNode_SequencePlayer::OverrideAsset(UAnimationAsset* NewAsset)
 		Sequence = AnimSequence;
 	}
 }
+
+void FAnimNode_SequencePlayer::GatherDebugData(FNodeDebugData& DebugData)
+{
+	FString DebugLine = DebugData.GetNodeName(this);
+	
+	DebugLine += FString::Printf(TEXT("('%s' Play Time: %.3f)"), *Sequence->GetName(), InternalTimeAccumulator);
+	DebugData.AddDebugItem(DebugLine, true);
+}
\ No newline at end of file
diff --git a/Engine/Source/Runtime/Engine/Private/Animation/AnimNode_Slot.cpp b/Engine/Source/Runtime/Engine/Private/Animation/AnimNode_Slot.cpp
--- a/Engine/Source/Runtime/Engine/Private/Animation/AnimNode_Slot.cpp
+++ b/Engine/Source/Runtime/Engine/Private/Animation/AnimNode_Slot.cpp
@@ -53,6 +53,36 @@ void FAnimNode_Slot::Evaluate(FPoseContext& Output)
 	Output.AnimInstance->SlotEvaluatePose(SlotName, SourceContext.Pose, Output.Pose, SlotNodeWeight);
 }
 
+
+void FAnimNode_Slot::GatherDebugData(FNodeDebugData& DebugData)
+{
+	FString DebugLine = DebugData.GetNodeName(this);
+	DebugLine += FString::Printf(TEXT("(Slot Name: '%s' Weight:%.1f%%)"), *SlotName.ToString(), SlotNodeWeight*100.f);
+
+	bool bIsPoseSource = !bNeedToEvaluateSource;
+	DebugData.AddDebugItem(DebugLine, bIsPoseSource);
+	Source.GatherDebugData(DebugData.BranchFlow(bNeedToEvaluateSource ? 1.f : 0.f));
+
+	for (FAnimMontageInstance* MontageInstance : DebugData.AnimInstance->MontageInstances)
+	{
+		if (MontageInstance->IsValid() && MontageInstance->Montage->IsValidSlot(SlotName))
+		{
+			const FAnimTrack* Track = MontageInstance->Montage->GetAnimationData(SlotName);
+			for (const FAnimSegment& Segment : Track->AnimSegments)
+			{
+				float Weight;
+				float CurrentAnimPos;
+				if (UAnimSequenceBase* Anim = Segment.GetAnimationData(MontageInstance->Position, CurrentAnimPos, Weight))
+				{
+					FString MontageLine = FString::Printf(TEXT("Montage: '%s' Anim: '%s' Play Time:%.2f W:%.2f%%"), *MontageInstance->Montage->GetName(), *Anim->GetName(), CurrentAnimPos, Weight*100.f);
+					DebugData.BranchFlow(1.0f).AddDebugItem(MontageLine, true);
+					break;
+				}
+			}
+		}
+	}
+}
+
 FAnimNode_Slot::FAnimNode_Slot()
 	: bNeedToEvaluateSource(false)
 {
diff --git a/Engine/Source/Runtime/Engine/Private/Animation/AnimNode_SpringBone.cpp b/Engine/Source/Runtime/Engine/Private/Animation/AnimNode_SpringBone.cpp
--- a/Engine/Source/Runtime/Engine/Private/Animation/AnimNode_SpringBone.cpp
+++ b/Engine/Source/Runtime/Engine/Private/Animation/AnimNode_SpringBone.cpp
@@ -38,6 +38,18 @@ void FAnimNode_SpringBone::Update(const FAnimationUpdateContext& Context)
 	RemainingTime += Context.GetDeltaTime();
 }
 
+void FAnimNode_SpringBone::GatherDebugData(FNodeDebugData& DebugData)
+{
+	const float ActualAlpha = AlphaScaleBias.ApplyTo(Alpha);
+
+	//MDW_TODO Add more output info?
+	FString DebugLine = DebugData.GetNodeName(this);
+	DebugLine += FString::Printf(TEXT("(Alpha: %.1f%% RemainingTime: %.3f)"), ActualAlpha*100.f, RemainingTime);
+
+	DebugData.AddDebugItem(DebugLine);
+	ComponentPose.GatherDebugData(DebugData);
+}
+
 void FAnimNode_SpringBone::EvaluateBoneTransforms(USkeletalMeshComponent* SkelComp, const FBoneContainer & RequiredBones, FA2CSPose& MeshBases, TArray<FBoneTransform>& OutBoneTransforms)
 {
 	check(OutBoneTransforms.Num() == 0);
diff --git a/Engine/Source/Runtime/Engine/Private/Animation/AnimNode_StateMachine.cpp b/Engine/Source/Runtime/Engine/Private/Animation/AnimNode_StateMachine.cpp
--- a/Engine/Source/Runtime/Engine/Private/Animation/AnimNode_StateMachine.cpp
+++ b/Engine/Source/Runtime/Engine/Private/Animation/AnimNode_StateMachine.cpp
@@ -634,6 +634,78 @@ void FAnimNode_StateMachine::EvaluateTransitionCustomBlend(FPoseContext& Output,
 	}
 }
 
+void AddStateWeight(TMap<int32, float>& StateWeightMap, int32 StateIndex, float Weight)
+{
+	if (!StateWeightMap.Find(StateIndex))
+	{
+		StateWeightMap.Add(StateIndex) = Weight;
+	}
+}
+
+void FAnimNode_StateMachine::GatherDebugData(FNodeDebugData& DebugData)
+{
+	FString DebugLine = DebugData.GetNodeName(this);
+	DebugLine += FString::Printf(TEXT("(%s->%s)"), *GetMachineDescription()->MachineName.ToString(), *GetStateInfo().StateName.ToString());
+
+	TMap<int32, float> StateWeightMap;
+
+	if (ActiveTransitionArray.Num() > 0)
+	{
+		for (int32 Index = 0; Index < ActiveTransitionArray.Num(); ++Index)
+		{
+			// if there is any source pose, blend it here
+			FAnimationActiveTransitionEntry& ActiveTransition = ActiveTransitionArray[Index];
+
+			if (ActiveTransition.bActive)
+			{
+				switch (ActiveTransition.LogicType)
+				{
+					case ETransitionLogicType::TLT_StandardBlend:
+					{
+						AddStateWeight(StateWeightMap, ActiveTransition.PreviousState, (1.0f - ActiveTransition.Alpha));
+						AddStateWeight(StateWeightMap, ActiveTransition.NextState, (ActiveTransition.Alpha));
+						break;
+					}
+					case ETransitionLogicType::TLT_Custom:
+					{
+						if (ActiveTransition.CustomTransitionGraph.LinkID != INDEX_NONE)
+						{
+							for (TArray<FAnimNode_TransitionPoseEvaluator*>::TIterator PoseEvaluatorListIt = ActiveTransition.PoseEvaluators.CreateIterator(); PoseEvaluatorListIt; ++PoseEvaluatorListIt)
+							{
+								FAnimNode_TransitionPoseEvaluator* Evaluator = *PoseEvaluatorListIt;
+								if (Evaluator->InputNodeNeedsUpdate())
+								{
+									const bool bUsePreviousState = (Evaluator->DataSource == EEvaluatorDataSource::EDS_SourcePose);
+									const int32 EffectiveStateIndex = bUsePreviousState ? ActiveTransition.PreviousState : ActiveTransition.NextState;
+									AddStateWeight(StateWeightMap, EffectiveStateIndex, 1.f);
+								}
+							}
+						}
+						break;
+					}
+					default:
+					{
+						break;
+					}
+				}
+			}
+		}
+	}
+	else if (!IsAConduitState(CurrentState))
+	{
+		StateWeightMap.Add(CurrentState) = 1.0f;
+	}
+
+	DebugData.AddDebugItem(DebugLine);
+	for (int32 PoseIndex = 0; PoseIndex < StatePoseLinks.Num(); ++PoseIndex)
+	{
+		float* WeightPtr = StateWeightMap.Find(PoseIndex);
+		float Weight = WeightPtr ? *WeightPtr : 0.f;
+
+		StatePoseLinks[PoseIndex].GatherDebugData(DebugData.BranchFlow(Weight));
+	}
+}
+
 void FAnimNode_StateMachine::SetStateInternal(int32 NewStateIndex)
 {
 	checkSlow(PRIVATE_MachineDescription);
diff --git a/Engine/Source/Runtime/Engine/Private/Animation/AnimNode_TransitionPoseEvaluator.cpp b/Engine/Source/Runtime/Engine/Private/Animation/AnimNode_TransitionPoseEvaluator.cpp
--- a/Engine/Source/Runtime/Engine/Private/Animation/AnimNode_TransitionPoseEvaluator.cpp
+++ b/Engine/Source/Runtime/Engine/Private/Animation/AnimNode_TransitionPoseEvaluator.cpp
@@ -50,6 +50,14 @@ void FAnimNode_TransitionPoseEvaluator::Evaluate(FPoseContext& Output)
 	}
 }
 
+void FAnimNode_TransitionPoseEvaluator::GatherDebugData(FNodeDebugData& DebugData)
+{
+	FString DebugLine = DebugData.GetNodeName(this);
+	
+	DebugLine += FString::Printf(TEXT("(Cached Frames Remaining: %i)"), CacheFramesRemaining);
+	DebugData.AddDebugItem(DebugLine);
+}
+
 bool FAnimNode_TransitionPoseEvaluator::InputNodeNeedsUpdate() const
 {
 	// EM_Standard mode always updates and EM_DelayedFreeze mode only updates if there are cache frames remaining
diff --git a/Engine/Source/Runtime/Engine/Private/Animation/AnimNode_TransitionResult.cpp b/Engine/Source/Runtime/Engine/Private/Animation/AnimNode_TransitionResult.cpp
--- a/Engine/Source/Runtime/Engine/Private/Animation/AnimNode_TransitionResult.cpp
+++ b/Engine/Source/Runtime/Engine/Private/Animation/AnimNode_TransitionResult.cpp
@@ -24,3 +24,10 @@ void FAnimNode_TransitionResult::Update(const FAnimationUpdateContext& Context)
 void FAnimNode_TransitionResult::Evaluate(FPoseContext& Output)
 {
 }
+
+
+void FAnimNode_TransitionResult::GatherDebugData(FNodeDebugData& DebugData)
+{
+	FString DebugLine = DebugData.GetNodeName(this);
+	DebugData.AddDebugItem(DebugLine);
+}
\ No newline at end of file
diff --git a/Engine/Source/Runtime/Engine/Private/Animation/AnimNode_TwoBoneIK.cpp b/Engine/Source/Runtime/Engine/Private/Animation/AnimNode_TwoBoneIK.cpp
--- a/Engine/Source/Runtime/Engine/Private/Animation/AnimNode_TwoBoneIK.cpp
+++ b/Engine/Source/Runtime/Engine/Private/Animation/AnimNode_TwoBoneIK.cpp
@@ -18,6 +18,18 @@ FAnimNode_TwoBoneIK::FAnimNode_TwoBoneIK()
 {
 }
 
+void FAnimNode_TwoBoneIK::GatherDebugData(FNodeDebugData& DebugData)
+{
+	FString DebugLine = DebugData.GetNodeName(this);
+
+	DebugLine += "(";
+	AddDebugNodeData(DebugLine);
+	DebugLine += FString::Printf(TEXT(" IKBone: %s)"), *IKBone.BoneName.ToString());
+	DebugData.AddDebugItem(DebugLine);
+
+	ComponentPose.GatherDebugData(DebugData);
+}
+
 void FAnimNode_TwoBoneIK::EvaluateBoneTransforms(USkeletalMeshComponent* SkelComp, const FBoneContainer & RequiredBones, FA2CSPose& MeshBases, TArray<FBoneTransform>& OutBoneTransforms)
 {
 	check(OutBoneTransforms.Num() == 0);
diff --git a/Engine/Source/Runtime/Engine/Private/Animation/AnimNode_TwoWayBlend.cpp b/Engine/Source/Runtime/Engine/Private/Animation/AnimNode_TwoWayBlend.cpp
--- a/Engine/Source/Runtime/Engine/Private/Animation/AnimNode_TwoWayBlend.cpp
+++ b/Engine/Source/Runtime/Engine/Private/Animation/AnimNode_TwoWayBlend.cpp
@@ -68,3 +68,16 @@ void FAnimationNode_TwoWayBlend::Evaluate(FPoseContext& Output)
 		A.Evaluate(Output);
 	}
 }
+
+
+void FAnimationNode_TwoWayBlend::GatherDebugData(FNodeDebugData& DebugData)
+{
+	const float ActualAlpha = AlphaScaleBias.ApplyTo(Alpha);
+
+	FString DebugLine = DebugData.GetNodeName(this);
+	DebugLine += FString::Printf(TEXT("(Alpha: %.1f%%)"), ActualAlpha*100);
+	DebugData.AddDebugItem(DebugLine);
+
+	A.GatherDebugData(DebugData.BranchFlow(1.f - ActualAlpha));
+	B.GatherDebugData(DebugData.BranchFlow(ActualAlpha));
+}
\ No newline at end of file
diff --git a/Engine/Source/Runtime/Engine/Private/Animation/AnimNode_UseCachedPose.cpp b/Engine/Source/Runtime/Engine/Private/Animation/AnimNode_UseCachedPose.cpp
--- a/Engine/Source/Runtime/Engine/Private/Animation/AnimNode_UseCachedPose.cpp
+++ b/Engine/Source/Runtime/Engine/Private/Animation/AnimNode_UseCachedPose.cpp
@@ -28,3 +28,11 @@ void FAnimNode_UseCachedPose::Evaluate(FPoseContext& Output)
 {
 	LinkToCachingNode.Evaluate(Output);
 }
+
+void FAnimNode_UseCachedPose::GatherDebugData(FNodeDebugData& DebugData)
+{
+	FString DebugLine = DebugData.GetNodeName(this);
+	DebugData.AddDebugItem(DebugLine, true);
+
+	LinkToCachingNode.GatherDebugData(DebugData);
+}
\ No newline at end of file
diff --git a/Engine/Source/Runtime/Engine/Private/Character.cpp b/Engine/Source/Runtime/Engine/Private/Character.cpp
--- a/Engine/Source/Runtime/Engine/Private/Character.cpp
+++ b/Engine/Source/Runtime/Engine/Private/Character.cpp
@@ -7,6 +7,7 @@
 
 #include "EnginePrivate.h"
 #include "Net/UnrealNetwork.h"
+#include "DisplayDebugHelpers.h"
 
 DEFINE_LOG_CATEGORY_STATIC(LogCharacter, Log, All);
 DEFINE_LOG_CATEGORY_STATIC(LogAvatar, Log, All);
@@ -491,27 +492,18 @@ void ACharacter::BaseChange()
 	}
 }
 
-void ACharacter::DisplayDebug(UCanvas* Canvas, const TArray<FName>& DebugDisplay, float& YL, float& YPos)
+void ACharacter::DisplayDebug(UCanvas* Canvas, const FDebugDisplayInfo& DebugDisplay, float& YL, float& YPos)
 {
 	Super::DisplayDebug(Canvas, DebugDisplay, YL, YPos);
 
 	float Indent = 0.f;
 
-	class Indenter
-	{
-		float& Indent;
-
-	public:
-		Indenter(float& InIndent) : Indent(InIndent) { Indent += 4.0f; }
-		~Indenter() { Indent -= 4.0f; }
-	};
-
 	UFont* RenderFont = GEngine->GetSmallFont();
 
 	static FName NAME_Physics = FName(TEXT("Physics"));
-	if (DebugDisplay.Contains(NAME_Physics) )
+	if (DebugDisplay.IsDisplayOn(NAME_Physics) )
 	{
-		Indenter PhysicsIndent(Indent);
+		FIndenter PhysicsIndent(Indent);
 
 		FString BaseString;
 		if ( CharacterMovement == NULL || MovementBase == NULL )
@@ -536,98 +528,6 @@ void ACharacter::DisplayDebug(UCanvas* Canvas, const TArray<FName>& DebugDisplay
 		Canvas->DrawText(RenderFont, T, Indent, YPos );
 		YPos += YL;
 	}
-
-	if(Mesh && Mesh->GetAnimInstance())
-	{
-		static FName NAME_Animation = FName(TEXT("Animation"));
-		if( DebugDisplay.Contains(NAME_Animation) )
-		{
-			YPos += YL;
-
-			Canvas->DrawText(RenderFont, TEXT("Animation"), Indent, YPos );
-			YPos += YL;
-			Indenter AnimIndent(Indent);
-
-			//Display Sync Groups
-			UAnimInstance* AnimInstance = Mesh->GetAnimInstance();
-			
-			FString Heading = FString::Printf(TEXT("SyncGroups: %i"), AnimInstance->SyncGroups.Num());
-			Canvas->DrawText(RenderFont, Heading, Indent, YPos );
-			YPos += YL;
-
-			for (int32 GroupIndex = 0; GroupIndex < AnimInstance->SyncGroups.Num(); ++GroupIndex)
-			{
-				Indenter GroupIndent(Indent);
-				FAnimGroupInstance& SyncGroup = AnimInstance->SyncGroups[GroupIndex];
-
-				FString GroupLabel = FString::Printf(TEXT("Group %i - Players %i"), GroupIndex, SyncGroup.ActivePlayers.Num());
-				Canvas->DrawText(RenderFont, GroupLabel, Indent, YPos );
-				YPos += YL;
-
-				if (SyncGroup.ActivePlayers.Num() > 0)
-				{
-					const int32 GroupLeaderIndex = FMath::Max(SyncGroup.GroupLeaderIndex, 0);
-
-					for(int32 PlayerIndex = 0; PlayerIndex < SyncGroup.ActivePlayers.Num(); ++PlayerIndex)
-					{
-						Indenter PlayerIndent(Indent);
-						FAnimTickRecord& Player = SyncGroup.ActivePlayers[PlayerIndex];
-
-						Canvas->SetLinearDrawColor( (PlayerIndex == GroupLeaderIndex) ? FLinearColor::Green : FLinearColor::Yellow);
-
-						FString PlayerEntry = FString::Printf(TEXT("%i) %s W:%.3f"), PlayerIndex, *Player.SourceAsset->GetName(), Player.EffectiveBlendWeight);
-						Canvas->DrawText(RenderFont, PlayerEntry, Indent, YPos );
-						YPos += YL;
-					}
-				}
-			}
-
-			Heading = FString::Printf(TEXT("Ungrouped: %i"), AnimInstance->UngroupedActivePlayers.Num());
-			Canvas->DrawText(RenderFont, Heading, Indent, YPos );
-			YPos += YL;
-
-			Canvas->SetLinearDrawColor( FLinearColor::Yellow );
-
-			for (int32 PlayerIndex = 0; PlayerIndex < AnimInstance->UngroupedActivePlayers.Num(); ++PlayerIndex)
-			{
-				Indenter PlayerIndent(Indent);
-				const FAnimTickRecord& Player = AnimInstance->UngroupedActivePlayers[PlayerIndex];
-				FString PlayerEntry = FString::Printf(TEXT("%i) %s W:%.3f"), PlayerIndex, *Player.SourceAsset->GetName(), Player.EffectiveBlendWeight);
-				Canvas->DrawText(RenderFont, PlayerEntry, Indent, YPos );
-				YPos += YL;
-			}
-
-			Heading = FString::Printf(TEXT("Montages: %i"), AnimInstance->MontageInstances.Num());
-			Canvas->DrawText(RenderFont, Heading, Indent, YPos );
-			YPos += YL;
-
-			for (int32 MontageIndex = 0; MontageIndex < AnimInstance->MontageInstances.Num(); ++MontageIndex)
-			{
-				Indenter PlayerIndent(Indent);
-
-				Canvas->SetLinearDrawColor( (MontageIndex == (AnimInstance->MontageInstances.Num()-1)) ? FLinearColor::Green : FLinearColor::Yellow );
-
-				FAnimMontageInstance* MontageInstance = AnimInstance->MontageInstances[MontageIndex];
-
-				FString MontageEntry = FString::Printf(TEXT("%i) %s Sec: %s W:%.3f DW:%.3f"), MontageIndex, *MontageInstance->Montage->GetName(), *MontageInstance->GetCurrentSection().ToString(), MontageInstance->Weight, MontageInstance->DesiredWeight);
-				Canvas->DrawText(RenderFont, MontageEntry, Indent, YPos );
-				YPos += YL;
-			}
-
-		}
-
-		static FName NAME_SimpleBones = FName(TEXT("SimpleBones"));
-		static FName NAME_Bones = FName(TEXT("Bones"));
-
-		if( DebugDisplay.Contains(NAME_SimpleBones) )
-		{
-			Mesh->DebugDrawBones(Canvas,true);
-		}
-		else if( DebugDisplay.Contains(NAME_Bones) )
-		{
-			Mesh->DebugDrawBones(Canvas,false);
-		}
-	}
 }
 
 void ACharacter::LaunchCharacter(FVector LaunchVelocity, bool bXYOverride, bool bZOverride)
@@ -1174,4 +1074,4 @@ void ACharacter::ClientCheatGhost_Implementation()
 		CharacterMovement->bCheatFlying = true;
 		CharacterMovement->SetMovementMode(MOVE_Flying);
 	}
-}
+}
\ No newline at end of file
diff --git a/Engine/Source/Runtime/Engine/Private/CharacterMovementComponent.cpp b/Engine/Source/Runtime/Engine/Private/CharacterMovementComponent.cpp
--- a/Engine/Source/Runtime/Engine/Private/CharacterMovementComponent.cpp
+++ b/Engine/Source/Runtime/Engine/Private/CharacterMovementComponent.cpp
@@ -4288,7 +4288,7 @@ FString UCharacterMovementComponent::GetMovementName()
 	return TEXT("Unknown");
 }
 
-void UCharacterMovementComponent::DisplayDebug(UCanvas* Canvas, const TArray<FName>& DebugDisplay, float& YL, float& YPos)
+void UCharacterMovementComponent::DisplayDebug(UCanvas* Canvas, const FDebugDisplayInfo& DebugDisplay, float& YL, float& YPos)
 {
 	if ( CharacterOwner == NULL )
 	{
diff --git a/Engine/Source/Runtime/Engine/Private/Controller.cpp b/Engine/Source/Runtime/Engine/Private/Controller.cpp
--- a/Engine/Source/Runtime/Engine/Private/Controller.cpp
+++ b/Engine/Source/Runtime/Engine/Private/Controller.cpp
@@ -445,7 +445,7 @@ void AController::GetActorEyesViewPoint( FVector& out_Location, FRotator& out_Ro
 }
 
 
-void AController::DisplayDebug(UCanvas* Canvas, const TArray<FName>& DebugDisplay, float& YL, float& YPos)
+void AController::DisplayDebug(UCanvas* Canvas, const FDebugDisplayInfo& DebugDisplay, float& YL, float& YPos)
 {
 	UFont* RenderFont = GEngine->GetSmallFont();
 	if ( Pawn == NULL )
diff --git a/Engine/Source/Runtime/Engine/Private/DrawDebugHelpers.cpp b/Engine/Source/Runtime/Engine/Private/DrawDebugHelpers.cpp
--- a/Engine/Source/Runtime/Engine/Private/DrawDebugHelpers.cpp
+++ b/Engine/Source/Runtime/Engine/Private/DrawDebugHelpers.cpp
@@ -864,16 +864,21 @@ void DrawDebugFloatHistory(UWorld const & WorldRef, FDebugFloatHistory const & F
 //////////////////////////////////////////////////////////////////
 // Debug draw canvas operations
 
-void DrawDebugCanvasLine(UCanvas* Canvas, const FVector& Start, const FVector& End, const FLinearColor& LineColor)
+void DrawDebugCanvas2DLine(UCanvas* Canvas, const FVector& Start, const FVector& End, const FLinearColor& LineColor)
 {
 	FCanvasLineItem LineItem;
-	LineItem.Origin = Canvas->Project(Start);
-	LineItem.EndPos = Canvas->Project(End);
+	LineItem.Origin = Start;
+	LineItem.EndPos = End;
 	LineItem.SetColor(LineColor);
 
 	LineItem.Draw(Canvas->Canvas);
 }
 
+void DrawDebugCanvasLine(UCanvas* Canvas, const FVector& Start, const FVector& End, const FLinearColor& LineColor)
+{
+	DrawDebugCanvas2DLine(Canvas, Canvas->Project(Start), Canvas->Project(End), LineColor);
+}
+
 void DrawDebugCanvasCircle(UCanvas* Canvas, const FVector& Base, const FVector& X, const FVector& Y, FColor Color, float Radius, int32 NumSides)
 {
 	const float	AngleDelta = 2.0f * PI / NumSides;
diff --git a/Engine/Source/Runtime/Engine/Private/GameMode.cpp b/Engine/Source/Runtime/Engine/Private/GameMode.cpp
--- a/Engine/Source/Runtime/Engine/Private/GameMode.cpp
+++ b/Engine/Source/Runtime/Engine/Private/GameMode.cpp
@@ -998,7 +998,7 @@ APlayerController* AGameMode::Login(const FString& Portal, const FString& Option
 }
 
 
-void AGameMode::DisplayDebug(UCanvas* Canvas, const TArray<FName>& DebugDisplay, float& YL, float& YPos)
+void AGameMode::DisplayDebug(UCanvas* Canvas, const FDebugDisplayInfo& DebugDisplay, float& YL, float& YPos)
 {
 	Canvas->SetDrawColor(255,255,255);
 }
diff --git a/Engine/Source/Runtime/Engine/Private/HUD.cpp b/Engine/Source/Runtime/Engine/Private/HUD.cpp
--- a/Engine/Source/Runtime/Engine/Private/HUD.cpp
+++ b/Engine/Source/Runtime/Engine/Private/HUD.cpp
@@ -9,6 +9,7 @@
 #include "Net/UnrealNetwork.h"
 #include "MessageLog.h"
 #include "UObjectToken.h"
+#include "DisplayDebugHelpers.h"
 
 DEFINE_LOG_CATEGORY_STATIC(LogHUD, Log, All);
 
@@ -238,6 +239,24 @@ void AHUD::ShowDebug(FName DebugType)
 	}
 }
 
+void AHUD::ShowDebugToggleSubCategory(FName Category)
+{
+	if( Category == NAME_Reset )
+	{
+		ToggledDebugCategories.Empty();
+		SaveConfig();
+	}
+	else
+	{
+		if (0 == ToggledDebugCategories.Remove(Category))
+		{
+			ToggledDebugCategories.Add(Category);
+		}
+		SaveConfig();
+	}
+}
+
+
 bool AHUD::ShouldDisplayDebug(const FName & DebugType) const
 {
 	return bShowDebugInfo && DebugDisplay.Contains(DebugType);
@@ -245,11 +264,15 @@ bool AHUD::ShouldDisplayDebug(const FName & DebugType) const
 
 void AHUD::ShowDebugInfo(float& YL, float& YPos)
 {
-	PlayerOwner->PlayerCameraManager->ViewTarget.Target->DisplayDebug(DebugCanvas, DebugDisplay, YL, YPos);
+	FLinearColor BackgroundColor(0.f, 0.f, 0.f, 0.5f);
+	DebugCanvas->Canvas->DrawTile(0, 0, DebugCanvas->ClipX, DebugCanvas->ClipY, 0.f, 0.f, 0.f, 0.f, BackgroundColor);
+
+	FDebugDisplayInfo DisplayInfo(DebugDisplay,ToggledDebugCategories);
+	PlayerOwner->PlayerCameraManager->ViewTarget.Target->DisplayDebug(DebugCanvas, DisplayInfo, YL, YPos);
 
 	if (ShouldDisplayDebug(NAME_Game))
 	{
-		GetWorld()->GetAuthGameMode()->DisplayDebug(DebugCanvas, DebugDisplay, YL, YPos);
+		GetWorld()->GetAuthGameMode()->DisplayDebug(DebugCanvas, DisplayInfo, YL, YPos);
 	}
 }
 
diff --git a/Engine/Source/Runtime/Engine/Private/Pawn.cpp b/Engine/Source/Runtime/Engine/Private/Pawn.cpp
--- a/Engine/Source/Runtime/Engine/Private/Pawn.cpp
+++ b/Engine/Source/Runtime/Engine/Private/Pawn.cpp
@@ -14,6 +14,7 @@
 #include "EngineKismetLibraryClasses.h"
 #include "ParticleDefinitions.h"
 #include "EngineComponentClasses.h"
+#include "DisplayDebugHelpers.h"
 
 DEFINE_LOG_CATEGORY(LogDamage);
 DEFINE_LOG_CATEGORY_STATIC(LogPawn, Warning, All);
@@ -649,7 +650,7 @@ FString APawn::GetHumanReadableName() const
 }
 
 
-void APawn::DisplayDebug(class UCanvas* Canvas, const TArray<FName>& DebugDisplay, float& YL, float& YPos)
+void APawn::DisplayDebug(class UCanvas* Canvas, const FDebugDisplayInfo& DebugDisplay, float& YL, float& YPos)
 {
 // 	for (FTouchingComponentsMap::TIterator It(TouchingComponents); It; ++It)
 // 	{
@@ -678,7 +679,7 @@ void APawn::DisplayDebug(class UCanvas* Canvas, const TArray<FName>& DebugDispla
 	Canvas->SetDrawColor(255,255,255);
 
 
-	if (DebugDisplay.Contains(NAME_Camera))
+	if (DebugDisplay.IsDisplayOn(NAME_Camera))
 	{
 		Canvas->DrawText(RenderFont, FString::Printf(TEXT("BaseEyeHeight %f"), BaseEyeHeight), 4.0f, YPos);
 		YPos += YL;
diff --git a/Engine/Source/Runtime/Engine/Private/PlayerCameraManager.cpp b/Engine/Source/Runtime/Engine/Private/PlayerCameraManager.cpp
--- a/Engine/Source/Runtime/Engine/Private/PlayerCameraManager.cpp
+++ b/Engine/Source/Runtime/Engine/Private/PlayerCameraManager.cpp
@@ -899,7 +899,7 @@ void APlayerCameraManager::LimitViewYaw(FRotator& ViewRotation, float InViewYawM
 	ViewRotation.Yaw = FRotator::ClampAxis(ViewRotation.Yaw);
 }
 
-void APlayerCameraManager::DisplayDebug(class UCanvas* Canvas, const TArray<FName>& DebugDisplay, float& YL, float& YPos)
+void APlayerCameraManager::DisplayDebug(class UCanvas* Canvas, const FDebugDisplayInfo& DebugDisplay, float& YL, float& YPos)
 {
 	Canvas->SetDrawColor(255,255,255);
 
diff --git a/Engine/Source/Runtime/Engine/Private/PlayerController.cpp b/Engine/Source/Runtime/Engine/Private/PlayerController.cpp
--- a/Engine/Source/Runtime/Engine/Private/PlayerController.cpp
+++ b/Engine/Source/Runtime/Engine/Private/PlayerController.cpp
@@ -13,6 +13,7 @@
 #include "IForceFeedbackSystem.h"
 #include "Slate.h"
 #include "GameFramework/TouchInterface.h"
+#include "DisplayDebugHelpers.h"
 
 DEFINE_LOG_CATEGORY_STATIC(LogPlayerController, Log, All);
 
@@ -2542,7 +2543,7 @@ void APlayerController::ClientIgnoreLookInput_Implementation(bool bIgnore)
 }
 
 
-void APlayerController::DisplayDebug(class UCanvas* Canvas, const TArray<FName>& DebugDisplay, float& YL, float& YPos)
+void APlayerController::DisplayDebug(class UCanvas* Canvas, const FDebugDisplayInfo& DebugDisplay, float& YL, float& YPos)
 {
 	Super::DisplayDebug(Canvas, DebugDisplay, YL, YPos);
 
@@ -2551,7 +2552,7 @@ void APlayerController::DisplayDebug(class UCanvas* Canvas, const TArray<FName>&
 	Canvas->DrawText(RenderFont, FString::Printf(TEXT("STATE %s"), *GetStateName().ToString()), 4.0f, YPos );
 	YPos += YL;
 
-	if (DebugDisplay.Contains(NAME_Camera))
+	if (DebugDisplay.IsDisplayOn(NAME_Camera))
 	{
 		if (PlayerCameraManager != NULL)
 		{
@@ -2566,7 +2567,7 @@ void APlayerController::DisplayDebug(class UCanvas* Canvas, const TArray<FName>&
 			YPos += YL;
 		}
 	}
-	if ( DebugDisplay.Contains(NAME_Input) )
+	if ( DebugDisplay.IsDisplayOn(NAME_Input) )
 	{
 		TArray<UInputComponent*> InputStack;
 		BuildInputStack(InputStack);
@@ -2601,7 +2602,7 @@ void APlayerController::DisplayDebug(class UCanvas* Canvas, const TArray<FName>&
 			YPos += YL;
 		}
 	}
-	if ( DebugDisplay.Contains("ForceFeedback"))
+	if ( DebugDisplay.IsDisplayOn("ForceFeedback"))
 	{
 		Canvas->SetDrawColor(255, 255, 255);
 		Canvas->DrawText(RenderFont, FString::Printf(TEXT("Force Feedback - LL: %.2f LS: %.2f RL: %.2f RS: %.2f"), ForceFeedbackValues.LeftLarge, ForceFeedbackValues.LeftSmall, ForceFeedbackValues.RightLarge, ForceFeedbackValues.RightSmall), 4.0f, YPos);
diff --git a/Engine/Source/Runtime/Engine/Private/UserInterface/PlayerInput.cpp b/Engine/Source/Runtime/Engine/Private/UserInterface/PlayerInput.cpp
--- a/Engine/Source/Runtime/Engine/Private/UserInterface/PlayerInput.cpp
+++ b/Engine/Source/Runtime/Engine/Private/UserInterface/PlayerInput.cpp
@@ -1096,7 +1096,7 @@ float UPlayerInput::SmoothMouse(float aMouse, float DeltaTime, uint8& SampleCoun
 	return aMouse;
 }
 
-void UPlayerInput::DisplayDebug(class UCanvas* Canvas, const TArray<FName>& DebugDisplay, float& YL, float& YPos)
+void UPlayerInput::DisplayDebug(class UCanvas* Canvas, const FDebugDisplayInfo& DebugDisplay, float& YL, float& YPos)
 {
 	if (Canvas)
 	{
diff --git a/Engine/Source/Runtime/Engine/Public/DisplayDebugHelpers.h b/Engine/Source/Runtime/Engine/Public/DisplayDebugHelpers.h
new file mode 100644
--- /dev/null
+++ b/Engine/Source/Runtime/Engine/Public/DisplayDebugHelpers.h
@@ -0,0 +1,30 @@
+// Copyright 1998-2014 Epic Games, Inc. All Rights Reserved.
+
+#pragma once
+
+// Tracks what debug information we have switched on
+class FDebugDisplayInfo
+{
+private:
+
+	TArray<FName> DisplayNames;
+
+	TArray<FName> ToggledCategories;
+
+public:
+
+	FDebugDisplayInfo(const TArray<FName>& InDisplayNames, const TArray<FName>& InToggledCategories) : DisplayNames(InDisplayNames), ToggledCategories(InToggledCategories) {}
+
+	bool IsDisplayOn(FName DisplayName) const { return DisplayNames.Contains(DisplayName); }
+	bool IsCategoryToggledOn(FName Category, bool bDefaultsToOn) const { return ToggledCategories.Contains(Category) != bDefaultsToOn; }
+};
+
+// Simple object to track scope indentation
+class FIndenter
+{
+	float& Indent;
+
+public:
+	FIndenter(float& InIndent) : Indent(InIndent) { Indent += 4.0f; }
+	~FIndenter() { Indent -= 4.0f; }
+};
\ No newline at end of file
diff --git a/Engine/Source/Runtime/Engine/Public/DrawDebugHelpers.h b/Engine/Source/Runtime/Engine/Public/DrawDebugHelpers.h
--- a/Engine/Source/Runtime/Engine/Public/DrawDebugHelpers.h
+++ b/Engine/Source/Runtime/Engine/Public/DrawDebugHelpers.h
@@ -51,6 +51,17 @@ ENGINE_API void DrawDebugFloatHistory(UWorld const & WorldRef, FDebugFloatHistor
 /* Draws a 2D Histogram of size 'DrawSize' based FDebugFloatHistory struct, using DrawLocation for the location in the world, rotation will face camera of first player. */
 ENGINE_API void DrawDebugFloatHistory(UWorld const & WorldRef, FDebugFloatHistory const & FloatHistory, FVector const & DrawLocation, FVector2D const & DrawSize, FColor const & DrawColor, bool const & bPersistent = false, float const & LifeTime = -1.f, uint8 const & DepthPriority = 0);
 
+/**
+ * Draws a 2D line on a canvas
+ *
+ * @param	Canvas		The Canvas to draw on
+ * @param	Start		The start of the line in screen space
+ * @param	End			The end of the line in screen space
+ * @param	LineColor	The color to use for the line
+ */
+
+ENGINE_API void DrawDebugCanvas2DLine(UCanvas* Canvas, const FVector& Start, const FVector& End, const FLinearColor& LineColor);
+
 /**
  * Draws a line on a canvas.
  *
