Fix compiler error in Clang

diff --git a/Engine/Source/Runtime/Engine/Classes/Animation/AnimNodeBase.h b/Engine/Source/Runtime/Engine/Classes/Animation/AnimNodeBase.h
--- a/Engine/Source/Runtime/Engine/Classes/Animation/AnimNodeBase.h
+++ b/Engine/Source/Runtime/Engine/Classes/Animation/AnimNodeBase.h
@@ -211,8 +211,8 @@ public:
 		bool IsOnActiveBranch() { return AbsoluteWeight > ZERO_ANIMWEIGHT_THRESH; }
 	};
 
-	FNodeDebugData(const class UAnimInstance* AnimInstance) : AnimInstance(AnimInstance), AbsoluteWeight(1.f) {}
-	FNodeDebugData(const class UAnimInstance* AnimInstance, float AbsWeight) : AnimInstance(AnimInstance), AbsoluteWeight(AbsWeight) {}
+	FNodeDebugData(const class UAnimInstance* AnimInstance) : AbsoluteWeight(1.f), AnimInstance(AnimInstance) {}
+	FNodeDebugData(const class UAnimInstance* AnimInstance, float AbsWeight) : AbsoluteWeight(AbsWeight), AnimInstance(AnimInstance) {}
 
 	void			AddDebugItem(FString DebugData, bool bPoseSource = false);
 	FNodeDebugData&	BranchFlow(float BranchWeight);
