non-ttp followup to "Icon polish pass" changes in CL 2018777/8/9 by [REDACTED_PERSON]

#branch UE4
#proj Runtime.Engine & Editor.LandscapeEditor
#summary Added code scaling to halve size of decals using icons recently doubled in size due to Phil's changes.

#change Applying uniform relative scale of 0.5f to billboard components in AEmitter, ALandscapeGizmoActor, APhysicsThruster, AEQSTestingPawn, AAtmosphericFog, ANavLinkProxy, AWindDirectionalSource, ALandscape, ALandscapeProxy, ATriggerBase, AMatineeActor, AReflectionCapture, ATextRenderActor, APlayerStart, ANote, ANiagaraActor, ULightComponent, ANavigationObjectBase, AExponentialHeightFog, ALandscapePlaceholder, UAudioComponent & ADecalActor.
#add New properties ULightComponentBase::StaticEditorTextureScale and DynamicEditorTextureScale. Also GetEditorSpriteScale() getter to access the appropriate scale value.
#change ULightComponentBase::UpdateLightSpriteTexture() - gets the correct scale and sets the scale value on the billboard component.
#add Setting new properties StaticEditorTextureScale and DynamicEditorTextureScale in ULightComponentBase derived classes. Setting to 1.0f in class with unchanged images: USkyLightComponent. Setting to 0.5f in classes with new double-size images: USpotLightComponent, UPointLightComponent & UDirectionalLightComponent.

diff --git a/Engine/Source/Editor/LandscapeEditor/Private/ActorFactoryLandscape.cpp b/Engine/Source/Editor/LandscapeEditor/Private/ActorFactoryLandscape.cpp
--- a/Engine/Source/Editor/LandscapeEditor/Private/ActorFactoryLandscape.cpp
+++ b/Engine/Source/Editor/LandscapeEditor/Private/ActorFactoryLandscape.cpp
@@ -54,6 +54,7 @@ ALandscapePlaceholder::ALandscapePlaceholder(const class FPostConstructInitializ
 	if (SpriteComponent.IsValid())
 	{
 		SpriteComponent->Sprite = ConstructorStatics.TerrainTexture.Get();
+		SpriteComponent->RelativeScale3D = FVector(0.5f, 0.5f, 0.5f);
 		SpriteComponent->AttachParent = RootComponent;
 		SpriteComponent->RelativeLocation = FVector(0, 0, 100);
 		SpriteComponent->bAbsoluteScale = true;
diff --git a/Engine/Source/Runtime/Engine/Classes/Components/LightComponentBase.h b/Engine/Source/Runtime/Engine/Classes/Components/LightComponentBase.h
--- a/Engine/Source/Runtime/Engine/Classes/Components/LightComponentBase.h
+++ b/Engine/Source/Runtime/Engine/Classes/Components/LightComponentBase.h
@@ -74,9 +74,17 @@ class ENGINE_API ULightComponentBase : public USceneComponent
 	UPROPERTY(transient)
 	UTexture2D* StaticEditorTexture;
 
+	/** Sprite scaling for static light in the editor. */
+	UPROPERTY(transient)
+	float StaticEditorTextureScale;
+
 	/** Sprite for dynamic light in the editor. */
 	UPROPERTY(transient)
 	UTexture2D* DynamicEditorTexture;
+
+	/** Sprite scaling for dynamic light in the editor. */
+	UPROPERTY(transient)
+	float DynamicEditorTextureScale;
 #endif
 
 	virtual void Serialize(FArchive& Ar) OVERRIDE;
@@ -99,6 +107,13 @@ class ENGINE_API ULightComponentBase : public USceneComponent
 		return (Mobility == EComponentMobility::Movable) ? DynamicEditorTexture : StaticEditorTexture;
 	}
 
+	/**
+	* @return Uniform scaling factor for the sprite for the light component class
+	*/
+	float GetEditorSpriteScale() const
+	{
+		return (Mobility == EComponentMobility::Movable) ? DynamicEditorTextureScale : StaticEditorTextureScale;
+	}
 
 	/** Update the texture used on the editor sprite */
 	virtual void UpdateLightSpriteTexture();
diff --git a/Engine/Source/Runtime/Engine/Private/AI/EnvironmentQuery/EQSTestingPawn.cpp b/Engine/Source/Runtime/Engine/Private/AI/EnvironmentQuery/EQSTestingPawn.cpp
--- a/Engine/Source/Runtime/Engine/Private/AI/EnvironmentQuery/EQSTestingPawn.cpp
+++ b/Engine/Source/Runtime/Engine/Private/AI/EnvironmentQuery/EQSTestingPawn.cpp
@@ -46,6 +46,7 @@ AEQSTestingPawn::AEQSTestingPawn(const class FPostConstructInitializeProperties&
 		static FConstructorStatics ConstructorStatics;
 
 		SpriteComponent->Sprite = ConstructorStatics.TextureObject.Get();
+		SpriteComponent->RelativeScale3D = FVector(0.5f, 0.5f, 0.5f);
 		SpriteComponent->bHiddenInGame = true;
 		//SpriteComponent->Mobility = EComponentMobility::Static;
 		SpriteComponent->SpriteInfo.Category = ConstructorStatics.ID_Misc;
diff --git a/Engine/Source/Runtime/Engine/Private/AI/Navigation/NavLinkProxy.cpp b/Engine/Source/Runtime/Engine/Private/AI/Navigation/NavLinkProxy.cpp
--- a/Engine/Source/Runtime/Engine/Private/AI/Navigation/NavLinkProxy.cpp
+++ b/Engine/Source/Runtime/Engine/Private/AI/Navigation/NavLinkProxy.cpp
@@ -22,6 +22,7 @@ ANavLinkProxy::ANavLinkProxy(const class FPostConstructInitializeProperties& PCI
 		static ConstructorHelpers::FObjectFinderOptional<UTexture2D> SpriteTexture(TEXT("/Engine/EditorResources/AI/S_NavLink"));
 
 		SpriteComponent->Sprite = SpriteTexture.Get();
+		SpriteComponent->RelativeScale3D = FVector(0.5f, 0.5f, 0.5f);
 		SpriteComponent->bHiddenInGame = true;
 		SpriteComponent->bVisible = true;
 
diff --git a/Engine/Source/Runtime/Engine/Private/Atmosphere/Atmosphere.cpp b/Engine/Source/Runtime/Engine/Private/Atmosphere/Atmosphere.cpp
--- a/Engine/Source/Runtime/Engine/Private/Atmosphere/Atmosphere.cpp
+++ b/Engine/Source/Runtime/Engine/Private/Atmosphere/Atmosphere.cpp
@@ -37,6 +37,7 @@ AAtmosphericFog::AAtmosphericFog(const class FPostConstructInitializeProperties&
 		if (SpriteComponent)
 		{
 			SpriteComponent->Sprite = ConstructorStatics.FogTextureObject.Get();
+			SpriteComponent->RelativeScale3D = FVector(0.5f, 0.5f, 0.5f);
 			SpriteComponent->SpriteInfo.Category = ConstructorStatics.ID_Fog;
 			SpriteComponent->SpriteInfo.DisplayName = ConstructorStatics.NAME_Fog;
 			SpriteComponent->AttachParent = AtmosphericFogComponent;
diff --git a/Engine/Source/Runtime/Engine/Private/AudioComponent.cpp b/Engine/Source/Runtime/Engine/Private/AudioComponent.cpp
--- a/Engine/Source/Runtime/Engine/Private/AudioComponent.cpp
+++ b/Engine/Source/Runtime/Engine/Private/AudioComponent.cpp
@@ -52,6 +52,7 @@ void UAudioComponent::OnRegister()
 		SpriteComponent = ConstructObject<UBillboardComponent>(UBillboardComponent::StaticClass(), GetOwner(), NAME_None, RF_Transactional | RF_TextExportTransient);
 
 		UpdateSpriteTexture();
+		SpriteComponent->RelativeScale3D = FVector(0.5f, 0.5f, 0.5f);
 		SpriteComponent->AttachTo(this);
 		SpriteComponent->AlwaysLoadOnClient = false;
 		SpriteComponent->AlwaysLoadOnServer = false;
diff --git a/Engine/Source/Runtime/Engine/Private/DecalComponent.cpp b/Engine/Source/Runtime/Engine/Private/DecalComponent.cpp
--- a/Engine/Source/Runtime/Engine/Private/DecalComponent.cpp
+++ b/Engine/Source/Runtime/Engine/Private/DecalComponent.cpp
@@ -78,6 +78,7 @@ ADecalActor::ADecalActor(const class FPostConstructInitializeProperties& PCIP)
 		if (SpriteComponent)
 		{
 			SpriteComponent->Sprite = ConstructorStatics.DecalTexture.Get();
+			SpriteComponent->RelativeScale3D = FVector(0.5f, 0.5f, 0.5f);
 			SpriteComponent->AttachParent = Decal;
 			SpriteComponent->bIsScreenSizeScaled = true;
 			SpriteComponent->bAbsoluteScale = true;
diff --git a/Engine/Source/Runtime/Engine/Private/DirectionalLightComponent.cpp b/Engine/Source/Runtime/Engine/Private/DirectionalLightComponent.cpp
--- a/Engine/Source/Runtime/Engine/Private/DirectionalLightComponent.cpp
+++ b/Engine/Source/Runtime/Engine/Private/DirectionalLightComponent.cpp
@@ -563,7 +563,9 @@ UDirectionalLightComponent::UDirectionalLightComponent(const class FPostConstruc
 		static ConstructorHelpers::FObjectFinder<UTexture2D> DynamicTexture(TEXT("/Engine/EditorResources/LightIcons/S_LightDirectionalMove"));
 
 		StaticEditorTexture = StaticTexture.Object;
+		StaticEditorTextureScale = 0.5f;
 		DynamicEditorTexture = DynamicTexture.Object;
+		DynamicEditorTextureScale = 0.5f;
 	}
 #endif
 
diff --git a/Engine/Source/Runtime/Engine/Private/HeightFogComponent.cpp b/Engine/Source/Runtime/Engine/Private/HeightFogComponent.cpp
--- a/Engine/Source/Runtime/Engine/Private/HeightFogComponent.cpp
+++ b/Engine/Source/Runtime/Engine/Private/HeightFogComponent.cpp
@@ -184,6 +184,7 @@ AExponentialHeightFog::AExponentialHeightFog(const class FPostConstructInitializ
 		static FConstructorStatics ConstructorStatics;
 
 		SpriteComponent->Sprite = ConstructorStatics.FogTextureObject.Get();
+		SpriteComponent->RelativeScale3D = FVector(0.5f, 0.5f, 0.5f);
 		SpriteComponent->SpriteInfo.Category = ConstructorStatics.ID_Fog;
 		SpriteComponent->SpriteInfo.DisplayName = ConstructorStatics.NAME_Fog;
 		SpriteComponent->AttachParent = Component;
diff --git a/Engine/Source/Runtime/Engine/Private/Interpolation.cpp b/Engine/Source/Runtime/Engine/Private/Interpolation.cpp
--- a/Engine/Source/Runtime/Engine/Private/Interpolation.cpp
+++ b/Engine/Source/Runtime/Engine/Private/Interpolation.cpp
@@ -288,6 +288,7 @@ AMatineeActor::AMatineeActor(const class FPostConstructInitializeProperties& PCI
 		static FConstructorStatics ConstructorStatics;
 
 		SpriteComponent->Sprite = ConstructorStatics.SceneManagerObject.Get();
+		SpriteComponent->RelativeScale3D = FVector(0.5f, 0.5f, 0.5f);
 		SpriteComponent->SpriteInfo.Category = ConstructorStatics.ID_Matinee;
 		SpriteComponent->SpriteInfo.DisplayName = ConstructorStatics.NAME_Matinee;
 		SpriteComponent->AttachParent = RootComponent;
diff --git a/Engine/Source/Runtime/Engine/Private/Landscape/Landscape.cpp b/Engine/Source/Runtime/Engine/Private/Landscape/Landscape.cpp
--- a/Engine/Source/Runtime/Engine/Private/Landscape/Landscape.cpp
+++ b/Engine/Source/Runtime/Engine/Private/Landscape/Landscape.cpp
@@ -443,6 +443,7 @@ ALandscape::ALandscape(const class FPostConstructInitializeProperties& PCIP)
 		static FConstructorStatics ConstructorStatics;
 
 		SpriteComponent->Sprite = ConstructorStatics.TerrainTexture.Get();
+		SpriteComponent->RelativeScale3D = FVector(0.5f, 0.5f, 0.5f);
 		SpriteComponent->SpriteInfo.Category = ConstructorStatics.ID_Landscape;
 		SpriteComponent->SpriteInfo.DisplayName = ConstructorStatics.NAME_Landscape;
 	}
@@ -487,6 +488,7 @@ ALandscapeProxy::ALandscapeProxy(const class FPostConstructInitializeProperties&
 		static FConstructorStatics ConstructorStatics;
 
 		SpriteComponent->Sprite = ConstructorStatics.TerrainTexture.Get();
+		SpriteComponent->RelativeScale3D = FVector(0.5f, 0.5f, 0.5f);
 		SpriteComponent->AttachParent = RootComponent;
 		SpriteComponent->bAbsoluteScale = true;
 	}
diff --git a/Engine/Source/Runtime/Engine/Private/Landscape/LandscapeGizmoActor.cpp b/Engine/Source/Runtime/Engine/Private/Landscape/LandscapeGizmoActor.cpp
--- a/Engine/Source/Runtime/Engine/Private/Landscape/LandscapeGizmoActor.cpp
+++ b/Engine/Source/Runtime/Engine/Private/Landscape/LandscapeGizmoActor.cpp
@@ -356,6 +356,7 @@ ALandscapeGizmoActor::ALandscapeGizmoActor(const class FPostConstructInitializeP
 		static FConstructorStatics ConstructorStatics;
 
 		SpriteComponent->Sprite = ConstructorStatics.DecalActorIconTexture.Get();
+		SpriteComponent->RelativeScale3D = FVector(0.5f, 0.5f, 0.5f);
 		SpriteComponent->bHiddenInGame = true;
 		SpriteComponent->SpriteInfo.Category = ConstructorStatics.ID_Misc;
 		SpriteComponent->SpriteInfo.DisplayName = ConstructorStatics.NAME_Misc;
diff --git a/Engine/Source/Runtime/Engine/Private/LightComponent.cpp b/Engine/Source/Runtime/Engine/Private/LightComponent.cpp
--- a/Engine/Source/Runtime/Engine/Private/LightComponent.cpp
+++ b/Engine/Source/Runtime/Engine/Private/LightComponent.cpp
@@ -55,6 +55,9 @@ void ULightComponentBase::UpdateLightSpriteTexture()
 	if (SpriteComponent != NULL)
 	{
 		SpriteComponent->SetSprite(GetEditorSprite());
+
+		float SpriteScale = GetEditorSpriteScale();
+		SpriteComponent->RelativeScale3D = FVector(SpriteScale);
 	}
 }
 
@@ -477,6 +480,7 @@ void ULightComponent::UpdateLightSpriteTexture()
 			UTexture2D* SpriteTexture = NULL;
 			SpriteTexture = LoadObject<UTexture2D>(NULL, TEXT("/Engine/EditorResources/LightIcons/S_LightError.S_LightError"));
 			SpriteComponent->SetSprite(SpriteTexture);
+			SpriteComponent->RelativeScale3D = FVector(0.5f, 0.5f, 0.5f);
 		}
 		else
 		{
diff --git a/Engine/Source/Runtime/Engine/Private/NavigationObjectBase.cpp b/Engine/Source/Runtime/Engine/Private/NavigationObjectBase.cpp
--- a/Engine/Source/Runtime/Engine/Private/NavigationObjectBase.cpp
+++ b/Engine/Source/Runtime/Engine/Private/NavigationObjectBase.cpp
@@ -46,6 +46,7 @@ ANavigationObjectBase::ANavigationObjectBase(const class FPostConstructInitializ
 		if (GoodSprite)
 		{
 			GoodSprite->Sprite = ConstructorStatics.NavigationTextureObject.Get();
+			GoodSprite->RelativeScale3D = FVector(0.5f, 0.5f, 0.5f);
 			GoodSprite->bHiddenInGame = true;
 			GoodSprite->SpriteInfo.Category = ConstructorStatics.ID_Navigation;
 			GoodSprite->SpriteInfo.DisplayName = ConstructorStatics.NAME_Navigation;
diff --git a/Engine/Source/Runtime/Engine/Private/NiagaraActor.cpp b/Engine/Source/Runtime/Engine/Private/NiagaraActor.cpp
--- a/Engine/Source/Runtime/Engine/Private/NiagaraActor.cpp
+++ b/Engine/Source/Runtime/Engine/Private/NiagaraActor.cpp
@@ -34,6 +34,7 @@ ANiagaraActor::ANiagaraActor(const class FPostConstructInitializeProperties& PCI
 		if (SpriteComponent)
 		{
 			SpriteComponent->Sprite = ConstructorStatics.SpriteTextureObject.Get();
+			SpriteComponent->RelativeScale3D = FVector(0.5f, 0.5f, 0.5f);
 			SpriteComponent->bHiddenInGame = true;
 			SpriteComponent->bIsScreenSizeScaled = true;
 			SpriteComponent->SpriteInfo.Category = ConstructorStatics.ID_Effects;
diff --git a/Engine/Source/Runtime/Engine/Private/Note.cpp b/Engine/Source/Runtime/Engine/Private/Note.cpp
--- a/Engine/Source/Runtime/Engine/Private/Note.cpp
+++ b/Engine/Source/Runtime/Engine/Private/Note.cpp
@@ -45,6 +45,7 @@ ANote::ANote(const class FPostConstructInitializeProperties& PCIP)
 		if (SpriteComponent)
 		{
 			SpriteComponent->Sprite = ConstructorStatics.NoteTextureObject.Get();
+			SpriteComponent->RelativeScale3D = FVector(0.5f, 0.5f, 0.5f);
 			SpriteComponent->SpriteInfo.Category = ConstructorStatics.ID_Notes;
 			SpriteComponent->SpriteInfo.DisplayName = ConstructorStatics.NAME_Notes;
 			SpriteComponent->AttachParent = RootComponent;
diff --git a/Engine/Source/Runtime/Engine/Private/Particles/ParticleComponents.cpp b/Engine/Source/Runtime/Engine/Private/Particles/ParticleComponents.cpp
--- a/Engine/Source/Runtime/Engine/Private/Particles/ParticleComponents.cpp
+++ b/Engine/Source/Runtime/Engine/Private/Particles/ParticleComponents.cpp
@@ -173,6 +173,7 @@ AEmitter::AEmitter(const class FPostConstructInitializeProperties& PCIP)
 		if (SpriteComponent)
 		{
 			SpriteComponent->Sprite = ConstructorStatics.SpriteTextureObject.Get();
+			SpriteComponent->RelativeScale3D = FVector(0.5f, 0.5f, 0.5f);
 			SpriteComponent->bHiddenInGame = true;
 			SpriteComponent->bIsScreenSizeScaled = true;
 			SpriteComponent->SpriteInfo.Category = ConstructorStatics.ID_Effects;
diff --git a/Engine/Source/Runtime/Engine/Private/PhysicsEngine/PhysicsThruster.cpp b/Engine/Source/Runtime/Engine/Private/PhysicsEngine/PhysicsThruster.cpp
--- a/Engine/Source/Runtime/Engine/Private/PhysicsEngine/PhysicsThruster.cpp
+++ b/Engine/Source/Runtime/Engine/Private/PhysicsEngine/PhysicsThruster.cpp
@@ -83,6 +83,7 @@ APhysicsThruster::APhysicsThruster(const class FPostConstructInitializePropertie
 		if (SpriteComponent)
 		{
 			SpriteComponent->Sprite = ConstructorStatics.ThrusterTexture.Get();
+			SpriteComponent->RelativeScale3D = FVector(0.5f, 0.5f, 0.5f);
 			SpriteComponent->SpriteInfo.Category = ConstructorStatics.ID_Physics;
 			SpriteComponent->SpriteInfo.DisplayName = ConstructorStatics.NAME_Physics;
 			SpriteComponent->AttachParent = ThrusterComponent;
diff --git a/Engine/Source/Runtime/Engine/Private/PlayerStart.cpp b/Engine/Source/Runtime/Engine/Private/PlayerStart.cpp
--- a/Engine/Source/Runtime/Engine/Private/PlayerStart.cpp
+++ b/Engine/Source/Runtime/Engine/Private/PlayerStart.cpp
@@ -34,6 +34,7 @@ APlayerStart::APlayerStart(const class FPostConstructInitializeProperties& PCIP)
 		if (GoodSprite)
 		{
 			GoodSprite->Sprite = ConstructorStatics.PlayerStartTextureObject.Get();
+			GoodSprite->RelativeScale3D = FVector(0.5f, 0.5f, 0.5f);
 			GoodSprite->SpriteInfo.Category = ConstructorStatics.ID_PlayerStart;
 			GoodSprite->SpriteInfo.DisplayName = ConstructorStatics.NAME_PlayerStart;
 		}
diff --git a/Engine/Source/Runtime/Engine/Private/PointLightComponent.cpp b/Engine/Source/Runtime/Engine/Private/PointLightComponent.cpp
--- a/Engine/Source/Runtime/Engine/Private/PointLightComponent.cpp
+++ b/Engine/Source/Runtime/Engine/Private/PointLightComponent.cpp
@@ -106,7 +106,9 @@ UPointLightComponent::UPointLightComponent(const class FPostConstructInitializeP
 		static ConstructorHelpers::FObjectFinder<UTexture2D> DynamicTexture(TEXT("/Engine/EditorResources/LightIcons/S_LightPointMove"));
 
 		StaticEditorTexture = StaticTexture.Object;
+		StaticEditorTextureScale = 0.5f;
 		DynamicEditorTexture = DynamicTexture.Object;
+		DynamicEditorTextureScale = 0.5f;
 	}
 #endif
 
diff --git a/Engine/Source/Runtime/Engine/Private/ReflectionCaptureComponent.cpp b/Engine/Source/Runtime/Engine/Private/ReflectionCaptureComponent.cpp
--- a/Engine/Source/Runtime/Engine/Private/ReflectionCaptureComponent.cpp
+++ b/Engine/Source/Runtime/Engine/Private/ReflectionCaptureComponent.cpp
@@ -58,6 +58,7 @@ AReflectionCapture::AReflectionCapture(const class FPostConstructInitializePrope
 		static FConstructorStatics ConstructorStatics;
 
 		SpriteComponent->Sprite = ConstructorStatics.DecalTexture.Get();
+		SpriteComponent->RelativeScale3D = FVector(0.5f, 0.5f, 0.5f);
 		SpriteComponent->bHiddenInGame = true;
 		SpriteComponent->bAbsoluteScale = true;
 		SpriteComponent->BodyInstance.bEnableCollision_DEPRECATED = false;	
diff --git a/Engine/Source/Runtime/Engine/Private/SkyLightComponent.cpp b/Engine/Source/Runtime/Engine/Private/SkyLightComponent.cpp
--- a/Engine/Source/Runtime/Engine/Private/SkyLightComponent.cpp
+++ b/Engine/Source/Runtime/Engine/Private/SkyLightComponent.cpp
@@ -85,7 +85,9 @@ USkyLightComponent::USkyLightComponent(const class FPostConstructInitializePrope
 	{
 		static ConstructorHelpers::FObjectFinder<UTexture2D> StaticTexture(TEXT("/Engine/EditorResources/LightIcons/SkyLight"));
 		StaticEditorTexture = StaticTexture.Object;
+		StaticEditorTextureScale = 1.0f;
 		DynamicEditorTexture = StaticTexture.Object;
+		DynamicEditorTextureScale = 1.0f;
 	}
 #endif
 
@@ -423,21 +425,21 @@ ASkyLight::ASkyLight(const class FPostConstructInitializeProperties& PCIP)
 #if WITH_EDITORONLY_DATA
 	if (!IsRunningCommandlet())
 	{
-		// Structure to hold one-time initialization
-		struct FConstructorStatics
-		{
+	// Structure to hold one-time initialization
+	struct FConstructorStatics
+	{
 			ConstructorHelpers::FObjectFinderOptional<UTexture2D> SkyLightTextureObject;
-			FName ID_Sky;
-			FText NAME_Sky;
+		FName ID_Sky;
+		FText NAME_Sky;
 
-			FConstructorStatics()
+		FConstructorStatics()
 				: SkyLightTextureObject(TEXT("/Engine/EditorResources/LightIcons/SkyLight"))
 				, ID_Sky(TEXT("Sky"))
-				, NAME_Sky(NSLOCTEXT( "SpriteCategory", "Sky", "Sky" ))
-			{
-			}
-		};
-		static FConstructorStatics ConstructorStatics;
+			, NAME_Sky(NSLOCTEXT( "SpriteCategory", "Sky", "Sky" ))
+		{
+		}
+	};
+	static FConstructorStatics ConstructorStatics;
 
 		if (SpriteComponent)
 		{
diff --git a/Engine/Source/Runtime/Engine/Private/SpotLightComponent.cpp b/Engine/Source/Runtime/Engine/Private/SpotLightComponent.cpp
--- a/Engine/Source/Runtime/Engine/Private/SpotLightComponent.cpp
+++ b/Engine/Source/Runtime/Engine/Private/SpotLightComponent.cpp
@@ -161,7 +161,9 @@ USpotLightComponent::USpotLightComponent(const class FPostConstructInitializePro
 		static ConstructorHelpers::FObjectFinder<UTexture2D> DynamicTexture(TEXT("/Engine/EditorResources/LightIcons/S_LightSpotMove"));
 
 		StaticEditorTexture = StaticTexture.Object;
+		StaticEditorTextureScale = 0.5f;
 		DynamicEditorTexture = DynamicTexture.Object;
+		DynamicEditorTextureScale = 0.5f;
 	}
 #endif
 
diff --git a/Engine/Source/Runtime/Engine/Private/TextRenderComponent.cpp b/Engine/Source/Runtime/Engine/Private/TextRenderComponent.cpp
--- a/Engine/Source/Runtime/Engine/Private/TextRenderComponent.cpp
+++ b/Engine/Source/Runtime/Engine/Private/TextRenderComponent.cpp
@@ -26,6 +26,7 @@ ATextRenderActor::ATextRenderActor(const class FPostConstructInitializePropertie
 		static FConstructorStatics ConstructorStatics;
 
 		SpriteComponent->Sprite = ConstructorStatics.TextRenderTexture.Get();
+		SpriteComponent->RelativeScale3D = FVector(0.5f, 0.5f, 0.5f);
 		SpriteComponent->AttachParent = TextRender;
 		SpriteComponent->bIsScreenSizeScaled = true;
 		SpriteComponent->bAbsoluteScale = true;
diff --git a/Engine/Source/Runtime/Engine/Private/TriggerBase.cpp b/Engine/Source/Runtime/Engine/Private/TriggerBase.cpp
--- a/Engine/Source/Runtime/Engine/Private/TriggerBase.cpp
+++ b/Engine/Source/Runtime/Engine/Private/TriggerBase.cpp
@@ -36,6 +36,7 @@ ATriggerBase::ATriggerBase(const class FPostConstructInitializeProperties& PCIP)
 	if (SpriteComponent)
 	{
 		SpriteComponent->Sprite = ConstructorStatics.TriggerTextureObject.Get();
+		SpriteComponent->RelativeScale3D = FVector(0.5f, 0.5f, 0.5f);
 		SpriteComponent->bHiddenInGame = false;
 		SpriteComponent->AlwaysLoadOnClient = false;
 		SpriteComponent->AlwaysLoadOnServer = false;
diff --git a/Engine/Source/Runtime/Engine/Private/WindDirectionalSource.cpp b/Engine/Source/Runtime/Engine/Private/WindDirectionalSource.cpp
--- a/Engine/Source/Runtime/Engine/Private/WindDirectionalSource.cpp
+++ b/Engine/Source/Runtime/Engine/Private/WindDirectionalSource.cpp
@@ -42,6 +42,7 @@ AWindDirectionalSource::AWindDirectionalSource(const class FPostConstructInitial
 		if (SpriteComponent)
 		{
 			SpriteComponent->Sprite = ConstructorStatics.SpriteTexture.Get();
+			SpriteComponent->RelativeScale3D = FVector(0.5f, 0.5f, 0.5f);
 			SpriteComponent->SpriteInfo.Category = ConstructorStatics.ID_Wind;
 			SpriteComponent->SpriteInfo.DisplayName = ConstructorStatics.NAME_Wind;
 			SpriteComponent->AttachParent = Component;
