#summary 	CanJump() is now exposed in ACharacter to be overriden by Blueprints.
#ttp 331402 	UE4: Make custom jump behavior via CanJump() etc more extensible to code and BPs in Character
#add 	Added BlueprintNativeEvent to CanJump() and specifed CanJump_Implementation.

diff --git a/Engine/Source/Runtime/Engine/Classes/GameFramework/Character.h b/Engine/Source/Runtime/Engine/Classes/GameFramework/Character.h
--- a/Engine/Source/Runtime/Engine/Classes/GameFramework/Character.h
+++ b/Engine/Source/Runtime/Engine/Classes/GameFramework/Character.h
@@ -263,9 +263,13 @@ public:
 	UFUNCTION(BlueprintCallable, Category="Pawn|Character")
 	virtual void Jump();
 
-	/** @Return Whether the character can jump in the current state. */
-	UFUNCTION(BlueprintCallable, Category="Pawn|Character")
-	virtual bool CanJump() const;
+	/**
+	 * Default implementation returns true if the character is not crouching, has a 
+	 * valid CharacterMovementComponent and CanEverJump() returns true.
+	 * @Return Whether the character can jump in the current state. 
+	 */
+	UFUNCTION(BlueprintNativeEvent, Category="Pawn|Character")
+	bool CanJump() const;
 
 	/** Play Animation Montage on the character mesh **/
 	UFUNCTION(BlueprintCallable, Category=Animation)
diff --git a/Engine/Source/Runtime/Engine/Private/Character.cpp b/Engine/Source/Runtime/Engine/Private/Character.cpp
--- a/Engine/Source/Runtime/Engine/Private/Character.cpp
+++ b/Engine/Source/Runtime/Engine/Private/Character.cpp
@@ -187,7 +187,7 @@ void ACharacter::Landed(const FHitResult& Hit)
 	OnLanded(Hit);
 }
 
-bool ACharacter::CanJump() const
+bool ACharacter::CanJump_Implementation() const
 {
 	return !bIsCrouched && CharacterMovement && CharacterMovement->IsMovingOnGround() && CharacterMovement->CanEverJump() && !CharacterMovement->bWantsToCrouch;
 }
