#ue4 BP Editoronly: Blueprints are not handled as BlueprintType, they are not loaded outside the editor (ini switch)

diff --git a/Engine/Source/Editor/BlueprintGraph/Private/EdGraphSchema_K2.cpp b/Engine/Source/Editor/BlueprintGraph/Private/EdGraphSchema_K2.cpp
--- a/Engine/Source/Editor/BlueprintGraph/Private/EdGraphSchema_K2.cpp
+++ b/Engine/Source/Editor/BlueprintGraph/Private/EdGraphSchema_K2.cpp
@@ -543,6 +543,20 @@ bool UEdGraphSchema_K2::IsAllowableBlueprintVariableType(const class UClass* InC
 			return true;
 		}
 
+		struct FBlueprintIsNotBlueprintTypeHelper
+		{
+			bool bNotBlueprintType;
+
+			FBlueprintIsNotBlueprintTypeHelper() : bNotBlueprintType(false)
+			{
+				GConfig->GetBool(TEXT("EditoronlyBP"), TEXT("bBlueprintIsNotBlueprintType"), bNotBlueprintType, GEditorIni);
+			}
+		};
+		static const FBlueprintIsNotBlueprintTypeHelper BPHelper;
+		if (BPHelper.bNotBlueprintType && InClass->IsChildOf(UBlueprint::StaticClass()))
+		{
+			return false;
+		}
 
 		const UClass* ParentClass = InClass;
 		while(ParentClass)
diff --git a/Engine/Source/Editor/GraphEditor/Private/KismetNodes/SGraphNodeSpawnActorFromClass.cpp b/Engine/Source/Editor/GraphEditor/Private/KismetNodes/SGraphNodeSpawnActorFromClass.cpp
--- a/Engine/Source/Editor/GraphEditor/Private/KismetNodes/SGraphNodeSpawnActorFromClass.cpp
+++ b/Engine/Source/Editor/GraphEditor/Private/KismetNodes/SGraphNodeSpawnActorFromClass.cpp
@@ -45,7 +45,7 @@ class SGraphPinActorBasedClass : public SGraphPinObject
 			if(NULL != InClass)
 			{
 				const bool bActorBased = InClass->IsChildOf(AActor::StaticClass());
-				const bool bBlueprintType = InClass->GetBoolMetaDataHierarchical(FBlueprintMetadata::MD_AllowableBlueprintVariableType);
+				const bool bBlueprintType = UEdGraphSchema_K2::IsAllowableBlueprintVariableType(InClass);
 				const bool bNotAbstract = !InClass->HasAnyClassFlags(CLASS_Abstract);
 				return bActorBased && bBlueprintType && bNotAbstract;
 			}
diff --git a/Engine/Source/Runtime/Engine/Classes/Engine/Blueprint.h b/Engine/Source/Runtime/Engine/Classes/Engine/Blueprint.h
--- a/Engine/Source/Runtime/Engine/Classes/Engine/Blueprint.h
+++ b/Engine/Source/Runtime/Engine/Classes/Engine/Blueprint.h
@@ -516,6 +516,8 @@ public:
 	virtual void Serialize(FArchive& Ar) OVERRIDE;
 	virtual FString GetDesc(void) OVERRIDE;
 	virtual void TagSubobjects(EObjectFlags NewFlags) OVERRIDE;
+	virtual bool NeedsLoadForClient() const OVERRIDE;
+	virtual bool NeedsLoadForServer() const OVERRIDE;
 	// End of UObject interface
 
 	/** Get the Blueprint object that generated the supplied class */
diff --git a/Engine/Source/Runtime/Engine/Private/Blueprint.cpp b/Engine/Source/Runtime/Engine/Private/Blueprint.cpp
--- a/Engine/Source/Runtime/Engine/Private/Blueprint.cpp
+++ b/Engine/Source/Runtime/Engine/Private/Blueprint.cpp
@@ -774,6 +774,29 @@ FString UBlueprint::GetDesc(void)
 	return ResultString;
 }
 
+struct FDontLoadBlueprintOutsideEditorHelper
+{
+	bool bDontLoad;
+
+	FDontLoadBlueprintOutsideEditorHelper() : bDontLoad(false)
+	{
+		GConfig->GetBool(TEXT("EditoronlyBP"), TEXT("bDontLoadBlueprintOutsideEditor"), bDontLoad, GEditorIni);
+	}
+};
+
+bool UBlueprint::NeedsLoadForClient() const
+{
+	static const FDontLoadBlueprintOutsideEditorHelper Helper;
+	return !Helper.bDontLoad;
+}
+
+bool UBlueprint::NeedsLoadForServer() const
+{
+	static const FDontLoadBlueprintOutsideEditorHelper Helper;
+	return !Helper.bDontLoad;
+}
+
+
 void UBlueprint::TagSubobjects(EObjectFlags NewFlags)
 {
 	Super::TagSubobjects(NewFlags);
