- Add Deproject to HUD
- Add DeprojectScreenPositionToWorld to PlayerController

diff --git a/Engine/Source/Runtime/Engine/Classes/GameFramework/HUD.h b/Engine/Source/Runtime/Engine/Classes/GameFramework/HUD.h
--- a/Engine/Source/Runtime/Engine/Classes/GameFramework/HUD.h
+++ b/Engine/Source/Runtime/Engine/Classes/GameFramework/HUD.h
@@ -397,6 +397,10 @@ public:
 	UFUNCTION(BlueprintPure, Category=HUD)
 	FVector Project(FVector Location);
 	
+	/** Transforms a 2D screen location into a 3D location and direction */
+	void Deproject(float ScreenX, float ScreenY, FVector& WorldPosition, FVector& WorldDirection);
+
+
 	/**
 	 * Add a hitbox to the hud
 	 * @param Position			Coordinates of the top left of the hit box.
diff --git a/Engine/Source/Runtime/Engine/Classes/GameFramework/PlayerController.h b/Engine/Source/Runtime/Engine/Classes/GameFramework/PlayerController.h
--- a/Engine/Source/Runtime/Engine/Classes/GameFramework/PlayerController.h
+++ b/Engine/Source/Runtime/Engine/Classes/GameFramework/PlayerController.h
@@ -319,6 +319,10 @@ public:
 	UFUNCTION(BlueprintCallable, Category="Game|Player", meta=(FriendlyName="ConvertMouseLocationToWorldSpace"))
 	void DeprojectMousePositionToWorld(FVector & WorldLocation, FVector & WorldDirection) const;
 
+	/** Convert current mouse 2D position to World Space 3D position and direction. **/
+	UFUNCTION(BlueprintCallable, Category = "Game|Player", meta = (FriendlyName = "ConvertScreenLocationToWorldSpace"))
+	void DeprojectScreenPositionToWorld(float ScreenX, float ScreenY, FVector & WorldLocation, FVector & WorldDirection) const;
+
 	/**
 	  * Updates the rotation of player, based on ControlRotation after RotationInput has been applied.
 	  * This may then be modified by the PlayerCamera, and is passed to Pawn->FaceRotation().
diff --git a/Engine/Source/Runtime/Engine/Private/HUD.cpp b/Engine/Source/Runtime/Engine/Private/HUD.cpp
--- a/Engine/Source/Runtime/Engine/Private/HUD.cpp
+++ b/Engine/Source/Runtime/Engine/Private/HUD.cpp
@@ -584,6 +584,15 @@ FVector AHUD::Project(FVector Location)
 	return FVector(0,0,0);
 }
 
+void AHUD::Deproject(float ScreenX, float ScreenY, FVector& WorldPosition, FVector& WorldDirection)
+{
+	WorldPosition = WorldDirection = FVector(0,0,0);
+	if (IsCanvasValid_WarnIfNot())
+	{
+		Canvas->Deproject(FVector2D(ScreenX, ScreenY), WorldPosition, WorldDirection);
+	}
+}
+
 void AHUD::DrawRect(FLinearColor Color, float ScreenX, float ScreenY, float Width, float Height)
 {
 	if (IsCanvasValid_WarnIfNot())	
diff --git a/Engine/Source/Runtime/Engine/Private/PlayerController.cpp b/Engine/Source/Runtime/Engine/Private/PlayerController.cpp
--- a/Engine/Source/Runtime/Engine/Private/PlayerController.cpp
+++ b/Engine/Source/Runtime/Engine/Private/PlayerController.cpp
@@ -1688,33 +1688,43 @@ bool APlayerController::GetHitResultUnderFingerForObjects(ETouchIndex::Type Fing
 	return bHit;
 }
 
-void APlayerController::DeprojectMousePositionToWorld(FVector & WorldPosition, FVector & WorldDirection) const
+void APlayerController::DeprojectMousePositionToWorld(FVector & WorldLocation, FVector & WorldDirection) const
 {
 	ULocalPlayer* LocalPlayer = Cast<ULocalPlayer>(Player);
-
-	if (LocalPlayer != NULL && LocalPlayer->ViewportClient != NULL && LocalPlayer->ViewportClient->Viewport != NULL)
+	if (LocalPlayer != NULL && LocalPlayer->ViewportClient != NULL)
 	{
 		FVector2D ScreenPosition = LocalPlayer->ViewportClient->GetMousePosition();
+		DeprojectScreenPositionToWorld(ScreenPosition.X, ScreenPosition.Y, WorldLocation, WorldDirection);
+	}
+}
+
+void APlayerController::DeprojectScreenPositionToWorld(float ScreenX, float ScreenY, FVector & WorldLocation, FVector & WorldDirection) const
+{
+	ULocalPlayer* LocalPlayer = Cast<ULocalPlayer>(Player);
 
+	if (LocalPlayer != NULL && LocalPlayer->ViewportClient != NULL && LocalPlayer->ViewportClient->Viewport != NULL)
+	{
 		// Create a view family for the game viewport
-		FSceneViewFamilyContext ViewFamily( FSceneViewFamily::ConstructionValues(
+		FSceneViewFamilyContext ViewFamily(FSceneViewFamily::ConstructionValues(
 			LocalPlayer->ViewportClient->Viewport,
 			GetWorld()->Scene,
-			LocalPlayer->ViewportClient->EngineShowFlags )
-			.SetRealtimeUpdate(true) );
+			LocalPlayer->ViewportClient->EngineShowFlags)
+			.SetRealtimeUpdate(true));
 
 		// Calculate a view where the player is to update the streaming from the players start location
 		FVector ViewLocation;
 		FRotator ViewRotation;
-		FSceneView* SceneView = LocalPlayer->CalcSceneView( &ViewFamily, /*out*/ ViewLocation, /*out*/ ViewRotation, LocalPlayer->ViewportClient->Viewport );
+		FSceneView* SceneView = LocalPlayer->CalcSceneView(&ViewFamily, /*out*/ ViewLocation, /*out*/ ViewRotation, LocalPlayer->ViewportClient->Viewport);
 
 		if (SceneView)
 		{
-			SceneView->DeprojectFVector2D(ScreenPosition, WorldPosition, WorldDirection);
+			const FVector2D ScreenPosition(ScreenX, ScreenY);
+			SceneView->DeprojectFVector2D(ScreenPosition, WorldLocation, WorldDirection);
 		}
 	}
 }
 
+
 bool APlayerController::GetHitResultAtScreenPosition(const FVector2D ScreenPosition, const ECollisionChannel TraceChannel, bool bTraceComplex, FHitResult& HitResult) const
 {
 	// Early out if we clicked on a HUD hitbox
