TTP 332087

* Fixing clicking on static meshes in the orthographic views, need to return a default hitproxy from CreateHitProxies or there isn't one set when the static meshes are batched in the ortho wireframe draw.

diff --git a/Engine/Source/Runtime/Engine/Private/SkeletalMesh.cpp b/Engine/Source/Runtime/Engine/Private/SkeletalMesh.cpp
--- a/Engine/Source/Runtime/Engine/Private/SkeletalMesh.cpp
+++ b/Engine/Source/Runtime/Engine/Private/SkeletalMesh.cpp
@@ -4027,18 +4027,14 @@ HHitProxy* FSkeletalMeshSceneProxy::CreateHitProxies(UPrimitiveComponent* Compon
 					OutHitProxies.Add(ActorHitProxy);
 				}
 			}
-
-			return NULL;
 		}
 		else
 		{
 			return FPrimitiveSceneProxy::CreateHitProxies(Component, OutHitProxies);
 		}
 	}
-	else
-	{
-		return NULL;
-	}
+
+	return NULL;
 }
 #endif
 
diff --git a/Engine/Source/Runtime/Engine/Private/StaticMeshRender.cpp b/Engine/Source/Runtime/Engine/Private/StaticMeshRender.cpp
--- a/Engine/Source/Runtime/Engine/Private/StaticMeshRender.cpp
+++ b/Engine/Source/Runtime/Engine/Private/StaticMeshRender.cpp
@@ -296,8 +296,14 @@ void FStaticMeshSceneProxy::SetIndexSource(int32 LODIndex, int32 SectionIndex, F
 #if WITH_EDITOR
 HHitProxy* FStaticMeshSceneProxy::CreateHitProxies(UPrimitiveComponent* Component, TArray<TRefCountPtr<HHitProxy> >& OutHitProxies)
 {
+	// In order to be able to click on static meshes when they're batched up, we need to have catch all default
+	// hit proxy to return.
+	HHitProxy* DefaultHitProxy = FPrimitiveSceneProxy::CreateHitProxies(Component, OutHitProxies);
+
 	if ( Component->GetOwner() )
 	{
+		// Generate separate hit proxies for each sub mesh, so that we can perform hit tests against each section for applying materials
+		// to each one.
 		for ( int32 LODIndex = 0; LODIndex < RenderData->LODResources.Num(); LODIndex++ )
 		{
 			const FStaticMeshLODResources& LODModel = RenderData->LODResources[LODIndex];
@@ -326,13 +332,9 @@ HHitProxy* FStaticMeshSceneProxy::CreateHitProxies(UPrimitiveComponent* Componen
 				OutHitProxies.Add(ActorHitProxy);
 			}
 		}
-		
-		return NULL;
-	}
-	else
-	{
-		return NULL;
 	}
+	
+	return DefaultHitProxy;
 }
 #endif // WITH_EDITOR
 
