#ue4 Make sure GameKing state is correct when using PIE.

diff --git a/Engine/Source/Runtime/Engine/Classes/Engine/GameViewportClient.h b/Engine/Source/Runtime/Engine/Classes/Engine/GameViewportClient.h
--- a/Engine/Source/Runtime/Engine/Classes/Engine/GameViewportClient.h
+++ b/Engine/Source/Runtime/Engine/Classes/Engine/GameViewportClient.h
@@ -244,7 +244,7 @@ public:
 	/** Returns a relative world context for this viewport.	 */
 	virtual UWorld* GetWorld() const OVERRIDE;
 
-	void SetReferenceToWorldContext(struct FWorldContext& WorldContext);
+	virtual void SetReferenceToWorldContext(struct FWorldContext& WorldContext);
 
 public:
 	// Begin UObject Interface
