#summary Duplicating timelines or Blueprints with timelines will no longer lose curve data after re-opening the editor. External curves will no longer duplicate the curve every compile.

#change The outer for duplicated curves is now always the Blueprint or the external curve package, unless the TimelineTemplate is in a transient package.

#ttp 312759 - ROCKET: TASK: Copying a timeline into another blueprint loses data when the editor is closed and reopened
#ttp 320627 - UE4: BLUEPRINTS: When duplicating a Blueprint that contains timelines, after restarting the editor the timeline curves will be missing
#ttp 320633 - UE4: BLUEPRINTS: When using the "Create external curve" option on a timeline, a duplicate curve asset will be created each time that Blueprint is compiled
#ttp 327736 - +BP: Timeline info disappears from copied Blueprint after exiting editor
#ttp 329225 - k2: A timeline can be copied from one blueprint to another and work correctly until the editor is closed and restarted. The curves/animation tracks are listed as missing at that point.

diff --git a/Engine/Source/Editor/BlueprintGraph/Private/K2Node_Timeline.cpp b/Engine/Source/Editor/BlueprintGraph/Private/K2Node_Timeline.cpp
--- a/Engine/Source/Editor/BlueprintGraph/Private/K2Node_Timeline.cpp
+++ b/Engine/Source/Editor/BlueprintGraph/Private/K2Node_Timeline.cpp
@@ -168,7 +168,7 @@ void UK2Node_Timeline::PostPasteNode()
 				FTTFloatTrack& Track = *TrackIt;
 				if(Track.CurveFloat->GetOuter()->IsA(UBlueprint::StaticClass()))
 				{
-					Track.CurveFloat->Rename(*Template->MakeUniqueCurveName(Track.CurveFloat), Blueprint, REN_DontCreateRedirectors);
+					Track.CurveFloat->Rename(*Template->MakeUniqueCurveName(Track.CurveFloat, Track.CurveFloat->GetOuter()), Blueprint, REN_DontCreateRedirectors);
 				}
 			}
 
@@ -177,7 +177,7 @@ void UK2Node_Timeline::PostPasteNode()
 				FTTEventTrack& Track = *TrackIt;
 				if(Track.CurveKeys->GetOuter()->IsA(UBlueprint::StaticClass()))
 				{
-					Track.CurveKeys->Rename(*Template->MakeUniqueCurveName(Track.CurveKeys), Blueprint, REN_DontCreateRedirectors);
+					Track.CurveKeys->Rename(*Template->MakeUniqueCurveName(Track.CurveKeys, Track.CurveKeys->GetOuter()), Blueprint, REN_DontCreateRedirectors);
 				}
 			}
 
@@ -186,7 +186,7 @@ void UK2Node_Timeline::PostPasteNode()
 				FTTVectorTrack& Track = *TrackIt;
 				if(Track.CurveVector->GetOuter()->IsA(UBlueprint::StaticClass()))
 				{
-					Track.CurveVector->Rename(*Template->MakeUniqueCurveName(Track.CurveVector), Blueprint, REN_DontCreateRedirectors);
+					Track.CurveVector->Rename(*Template->MakeUniqueCurveName(Track.CurveVector, Track.CurveVector->GetOuter()), Blueprint, REN_DontCreateRedirectors);
 				}
 			}
 
@@ -195,7 +195,7 @@ void UK2Node_Timeline::PostPasteNode()
 				FTTLinearColorTrack& Track = *TrackIt;
 				if(Track.CurveLinearColor->GetOuter()->IsA(UBlueprint::StaticClass()))
 				{
-					Track.CurveLinearColor->Rename(*Template->MakeUniqueCurveName(Track.CurveLinearColor), Blueprint, REN_DontCreateRedirectors);
+					Track.CurveLinearColor->Rename(*Template->MakeUniqueCurveName(Track.CurveLinearColor, Track.CurveLinearColor->GetOuter()), Blueprint, REN_DontCreateRedirectors);
 				}
 			}
 		}
diff --git a/Engine/Source/Runtime/Engine/Classes/Engine/TimelineTemplate.h b/Engine/Source/Runtime/Engine/Classes/Engine/TimelineTemplate.h
--- a/Engine/Source/Runtime/Engine/Classes/Engine/TimelineTemplate.h
+++ b/Engine/Source/Runtime/Engine/Classes/Engine/TimelineTemplate.h
@@ -198,7 +198,7 @@ class UTimelineTemplate : public UObject
 	ENGINE_API FName GetTrackPropertyName(const FName TrackName) const;
 
 	/* Create a new unique name for a curve */
-	ENGINE_API static FString MakeUniqueCurveName(UObject* Obj);
+	ENGINE_API static FString MakeUniqueCurveName(UObject* Obj, UObject* InOuter);
 
 	ENGINE_API static FString TimelineTemplateNameToVariableName(FName Name);
 	ENGINE_API static FString TimelineVariableNameToTemplateName(FName Name);
diff --git a/Engine/Source/Runtime/Engine/Private/TimelineTemplate.cpp b/Engine/Source/Runtime/Engine/Private/TimelineTemplate.cpp
--- a/Engine/Source/Runtime/Engine/Private/TimelineTemplate.cpp
+++ b/Engine/Source/Runtime/Engine/Private/TimelineTemplate.cpp
@@ -181,13 +181,12 @@ void UTimelineTemplate::RemoveMetaData(const FName& Key)
 	}
 }
 
-FString UTimelineTemplate::MakeUniqueCurveName(UObject* Obj)
+FString UTimelineTemplate::MakeUniqueCurveName(UObject* Obj, UObject* InOuter)
 {
 	FString OriginalName = Obj->GetName();
 	UClass* Class = Obj->GetClass();
-	UObject* Outer = Obj->GetOuter();
 	FName TestName(*OriginalName);
-	while(StaticFindObjectFast(Class, Outer, TestName))
+	while(StaticFindObjectFast(Class, InOuter, TestName))
 	{
 		TestName = FName(*OriginalName, TestName.GetNumber()+1);
 	}
@@ -216,12 +215,27 @@ void UTimelineTemplate::PostDuplicate(bool bDuplicateForPIE)
 {
 	Super::PostDuplicate(bDuplicateForPIE);
 
+	UObject* NewCurveOuter = GetOuter();
+
+	// The outer should always be a BlueprintGeneratedClass, however the Curves get put into the Blueprint, however if they belong to the Transient package, that's where the duplicated Blueprint should be.
+	if(GetOutermost() != GetTransientPackage())
+	{
+		if(UBlueprintGeneratedClass* BPClass = Cast<UBlueprintGeneratedClass>(GetOuter()))
+		{
+			NewCurveOuter = BPClass->ClassGeneratedBy;
+		}
+	}
+
 	for(TArray<struct FTTFloatTrack>::TIterator It = FloatTracks.CreateIterator();It;++It)
 	{
 		FTTFloatTrack& Track = *It;
 		if( Track.CurveFloat != NULL )
 		{
-			Track.CurveFloat = DuplicateObject<UCurveFloat>(Track.CurveFloat, Track.CurveFloat->GetOuter(), *MakeUniqueCurveName(Track.CurveFloat)); 
+			// Do not duplicate external curves unless duplicating to a transient package
+			if(!Track.CurveFloat->GetOuter()->IsA(UPackage::StaticClass()) || GetOutermost() == GetTransientPackage())
+			{
+				Track.CurveFloat = DuplicateObject<UCurveFloat>(Track.CurveFloat, NewCurveOuter, *MakeUniqueCurveName(Track.CurveFloat, NewCurveOuter));
+			}
 		}
 		else
 		{
@@ -234,7 +248,11 @@ void UTimelineTemplate::PostDuplicate(bool bDuplicateForPIE)
 		FTTEventTrack& Track = *It;
 		if( Track.CurveKeys != NULL )
 		{
-			Track.CurveKeys = DuplicateObject<UCurveFloat>(Track.CurveKeys, Track.CurveKeys->GetOuter(), *MakeUniqueCurveName(Track.CurveKeys)); 
+			// Do not duplicate external curves unless duplicating to a transient package
+			if(!Track.CurveKeys->GetOuter()->IsA(UPackage::StaticClass()) || GetOutermost() == GetTransientPackage())
+			{
+				Track.CurveKeys = DuplicateObject<UCurveFloat>(Track.CurveKeys, NewCurveOuter, *MakeUniqueCurveName(Track.CurveKeys, NewCurveOuter));
+			}
 		}
 		else
 		{
@@ -247,7 +265,11 @@ void UTimelineTemplate::PostDuplicate(bool bDuplicateForPIE)
 		FTTVectorTrack& Track = *It;
 		if( Track.CurveVector != NULL )
 		{
-			Track.CurveVector = DuplicateObject<UCurveVector>(Track.CurveVector, Track.CurveVector->GetOuter(), *MakeUniqueCurveName(Track.CurveVector)); 
+			// Do not duplicate external curves unless duplicating to a transient package
+			if(!Track.CurveVector->GetOuter()->IsA(UPackage::StaticClass()) || GetOutermost() == GetTransientPackage())
+			{
+				Track.CurveVector = DuplicateObject<UCurveVector>(Track.CurveVector, NewCurveOuter, *MakeUniqueCurveName(Track.CurveVector, NewCurveOuter));
+			}
 		}
 		else
 		{
@@ -260,7 +282,11 @@ void UTimelineTemplate::PostDuplicate(bool bDuplicateForPIE)
 		FTTLinearColorTrack& Track = *It;
 		if( Track.CurveLinearColor != NULL )
 		{
-			Track.CurveLinearColor = DuplicateObject<UCurveLinearColor>(Track.CurveLinearColor, Track.CurveLinearColor->GetOuter(), *MakeUniqueCurveName(Track.CurveLinearColor)); 
+			// Do not duplicate external curves unless duplicating to a transient package
+			if(!Track.CurveLinearColor->GetOuter()->IsA(UPackage::StaticClass()) || GetOutermost() == GetTransientPackage())
+			{
+				Track.CurveLinearColor = DuplicateObject<UCurveLinearColor>(Track.CurveLinearColor, NewCurveOuter, *MakeUniqueCurveName(Track.CurveLinearColor, NewCurveOuter));
+			}
 		}
 		else
 		{
