Added new CanvasRenderTarget2D class!
- You can you this to create render targets that can be dynamically rendered to using a Canvas API
- This isn't used for anything in the engine yet.  It's designed for game teams to use.

#github submission by DiCoJamesTan
https://github.com/EpicGames/UnrealEngine/pull/24

MikeF: I made some changes to the original submission, mostly renaming variables and comments.  I also changed it to use the new accessors for current time and delta time.

diff --git a/Engine/Source/Runtime/Engine/Classes/Engine/CanvasRenderTarget2D.h b/Engine/Source/Runtime/Engine/Classes/Engine/CanvasRenderTarget2D.h
new file mode 100644
--- /dev/null
+++ b/Engine/Source/Runtime/Engine/Classes/Engine/CanvasRenderTarget2D.h
@@ -0,0 +1,47 @@
+// Copyright 1998-2014 Epic Games, Inc. All Rights Reserved.
+
+#pragma once
+
+#include "Canvas.h"
+#include "CanvasRenderTarget2D.generated.h"
+
+/** This delegate allows developers to bind a function to be called every time the texture is updated, so that they can use the supplied canvas to draw to the texture. */
+DECLARE_DELEGATE_OneParam( FOnCanvasRenderTargetUpdate, class UCanvas* );
+
+/**
+ * CanvasRenderTarget2D is 2D render target which exposes a Canvas interface to allow you to draw elements onto 
+ * it directly.
+ */
+UCLASS(MinimalAPI)
+class UCanvasRenderTarget2D : public UTextureRenderTarget2D
+{
+	GENERATED_UCLASS_BODY()
+
+public:
+
+	/**
+	 * Updates the the canvas render target texture's resource. This is where the render target will create or 
+	 * find a canvas object to use.  It also calls UpdateResourceImmediate() to clear the render target texture 
+	 * from the deferred rendering list, to stop the texture from being cleared the next frame. From there it
+	 * will ask the rendering thread to set up the RHI viewport. The canvas is then set up for rendering and 
+	 * then the user's update delegate is called.  The canvas is then flushed and the RHI resolves the 
+	 * texture to make it available for rendering.
+	 */
+	virtual void UpdateResource();
+
+	/**
+	 * Assigns a user delegate to be called every time the texture needs to be refreshed for rendering.  This delegate
+	 * passes along a Canvas object that you can use to draw to the texture directly!
+	 *
+	 * @param	InDelegate	The canvas update delegate to use.  You can change this at any time.
+	 */
+	void SetOnCanvasRenderTargetUpdate( const FOnCanvasRenderTargetUpdate& InDelegate )
+	{
+		OnCanvasRenderTargetUpdate = InDelegate;
+	}
+
+private:
+
+	/** Texture update delegate */
+	FOnCanvasRenderTargetUpdate OnCanvasRenderTargetUpdate;
+};
diff --git a/Engine/Source/Runtime/Engine/Private/CanvasRenderTarget2D.cpp b/Engine/Source/Runtime/Engine/Private/CanvasRenderTarget2D.cpp
new file mode 100644
--- /dev/null
+++ b/Engine/Source/Runtime/Engine/Private/CanvasRenderTarget2D.cpp
@@ -0,0 +1,73 @@
+// Copyright 1998-2014 Epic Games, Inc. All Rights Reserved.
+
+#include "EnginePrivate.h"
+#include "Engine/CanvasRenderTarget2D.h"
+
+UCanvasRenderTarget2D::UCanvasRenderTarget2D( const class FPostConstructInitializeProperties& PCIP )
+	: Super(PCIP)
+{
+	bNeedsTwoCopies = false;
+}
+
+
+void UCanvasRenderTarget2D::UpdateResource()
+{
+	// Call parent implementation
+	Super::UpdateResource();
+	
+	// Create or find the canvas object to use to render onto the texture.  Multiple canvas render target textures can share the same canvas.
+	static const FName CanvasName( TEXT( "CanvasRenderTarget2DCanvas" ) );
+	UCanvas* Canvas = (UCanvas*)StaticFindObjectFast(UCanvas::StaticClass(), GetTransientPackage(), CanvasName );
+	if (Canvas == nullptr)
+	{
+		Canvas = ConstructObject<UCanvas>(UCanvas::StaticClass(), GetTransientPackage(), CanvasName );
+		Canvas->AddToRoot();
+	}
+
+	Canvas->Init(GetSurfaceWidth(), GetSurfaceHeight(), nullptr);
+	Canvas->Update();
+
+	// Update the resource immediately to remove it from the deferred resource update list. This prevents the texture
+	// from being cleared each frame.
+	UpdateResourceImmediate();
+
+	// Enqueue the rendering command to set up the rendering canvas.
+	ENQUEUE_UNIQUE_RENDER_COMMAND_ONEPARAMETER
+	(
+		CanvasRenderTargetMakeCurrentCommand,
+		FTextureRenderTarget2DResource*,
+		TextureRenderTarget,
+		(FTextureRenderTarget2DResource*)GameThread_GetRenderTargetResource(),
+		{
+			RHISetRenderTarget(TextureRenderTarget->GetRenderTargetTexture(), FTexture2DRHIRef());
+			RHISetViewport(0, 0, 0.0f, TextureRenderTarget->GetSizeXY().X, TextureRenderTarget->GetSizeXY().Y, 1.0f);
+		}
+	);
+
+	// Create the FCanvas which does the actual rendering.
+	FCanvas RenderCanvas(GameThread_GetRenderTargetResource(), nullptr, FApp::GetCurrentTime() - GStartTime, FApp::GetDeltaTime(), FApp::GetCurrentTime() - GStartTime);
+	Canvas->Canvas = &RenderCanvas;
+
+	// Check if the render target update delegate is bound and if it is, execute the delegate with the current Canvas being used.
+	if (OnCanvasRenderTargetUpdate.IsBound())
+	{
+		OnCanvasRenderTargetUpdate.Execute(Canvas);
+	}
+
+	// Clean up and flush the rendering canvas.
+	Canvas->Canvas = nullptr;
+	RenderCanvas.Flush();
+
+	// Enqueue the rendering command to copy the freshly rendering texture resource back to the render target RHI 
+	// so that the texture is updated and available for rendering.
+	ENQUEUE_UNIQUE_RENDER_COMMAND_ONEPARAMETER
+	(
+		CanvasRenderTargetResolveCommand,
+		FTextureRenderTargetResource*,
+		RenderTargetResource,
+		GameThread_GetRenderTargetResource(),
+		{
+			RHICopyToResolveTarget(RenderTargetResource->GetRenderTargetTexture(), RenderTargetResource->TextureRHI, true, FResolveParams());
+		}
+	);
+}
