#ue4 Blueprint:: NewVariables and ImplementedInterfaces are no longer marked as "EditAnywhere"

diff --git a/Engine/Source/Runtime/Engine/Classes/Engine/Blueprint.h b/Engine/Source/Runtime/Engine/Classes/Engine/Blueprint.h
--- a/Engine/Source/Runtime/Engine/Classes/Engine/Blueprint.h
+++ b/Engine/Source/Runtime/Engine/Classes/Engine/Blueprint.h
@@ -367,11 +367,11 @@ class ENGINE_API UBlueprint : public UBlueprintCore
 	TEnumAsByte<enum EBlueprintStatus> Status;
 
 	/** Array of new variables to be added to generated class */
-	UPROPERTY(EditAnywhere, Category=Blueprint)
+	UPROPERTY()
 	TArray<struct FBPVariableDescription> NewVariables;
 
 	/** Array of info about the interfaces we implement in this blueprint */
-	UPROPERTY(EditAnywhere, Category=Blueprint, AssetRegistrySearchable)
+	UPROPERTY(AssetRegistrySearchable)
 	TArray<struct FBPInterfaceDescription> ImplementedInterfaces;
 
 	UPROPERTY(EditAnywhere, Category=Blueprint)
