#ue4
- CameraAnims working again! Includes Blueprint API.

diff --git a/Engine/Source/Developer/AssetTools/Private/AssetTypeActions/AssetTypeActions_CameraAnim.cpp b/Engine/Source/Developer/AssetTools/Private/AssetTypeActions/AssetTypeActions_CameraAnim.cpp
--- a/Engine/Source/Developer/AssetTools/Private/AssetTypeActions/AssetTypeActions_CameraAnim.cpp
+++ b/Engine/Source/Developer/AssetTools/Private/AssetTypeActions/AssetTypeActions_CameraAnim.cpp
@@ -13,11 +13,11 @@ void FAssetTypeActions_CameraAnim::CreateMatineeActorForCameraAnim(UCameraAnim*
 
 	FActorSpawnParameters ActorSpawnParameters;
 	ActorSpawnParameters.Name = InCameraAnim->GetFName();
-	PreviewMatinee = GEditor->GetEditorWorldContext().World()->SpawnActor<AMatineeActorCameraAnim>(ActorSpawnParameters);
-	check(PreviewMatinee.IsValid());
+	PreviewMatineeActor = GEditor->GetEditorWorldContext().World()->SpawnActor<AMatineeActorCameraAnim>(ActorSpawnParameters);
+	check(PreviewMatineeActor.IsValid());
 	UInterpData* NewData = ConstructObject<UInterpData>( UInterpData::StaticClass(), GetTransientPackage(), NAME_None, RF_Transactional);
-	PreviewMatinee.Get()->MatineeData = NewData;
-	PreviewMatinee.Get()->CameraAnim = InCameraAnim;
+	PreviewMatineeActor.Get()->MatineeData = NewData;
+	PreviewMatineeActor.Get()->CameraAnim = InCameraAnim;
 }
 
 void FAssetTypeActions_CameraAnim::CreateCameraActorForCameraAnim(UCameraAnim* InCameraAnim)
@@ -45,7 +45,9 @@ void FAssetTypeActions_CameraAnim::CreateCameraActorForCameraAnim(UCameraAnim* I
 
 	// copy data from the CamAnim to the CameraActor
 	check(PreviewCamera.Get()->CameraComponent);
+	PreviewCamera.Get()->PreviewedCameraAnim = InCameraAnim;
 	PreviewCamera.Get()->CameraComponent->FieldOfView = InCameraAnim->BaseFOV;
+	PreviewCamera.Get()->CameraComponent->PostProcessSettings = InCameraAnim->BasePostProcessSettings;
 }
 
 void FAssetTypeActions_CameraAnim::CreatePreviewPawnForCameraAnim(UCameraAnim* InCameraAnim)
@@ -102,21 +104,21 @@ UInterpGroup* FAssetTypeActions_CameraAnim::CreateInterpGroup(UCameraAnim* InCam
 	if (PreviewInfo.PawnInst)
 	{
 		// create InterpGroup so that we can play animation to this pawn
-		check(PreviewMatinee.Get()->MatineeData);
-		UInterpGroup* NewGroup = ConstructObject<UInterpGroup>( UInterpGroup::StaticClass(), PreviewMatinee.Get()->MatineeData, NAME_None, RF_Transient );
+		check(PreviewMatineeActor.Get()->MatineeData);
+		UInterpGroup* NewGroup = ConstructObject<UInterpGroup>(UInterpGroup::StaticClass(), PreviewMatineeActor.Get()->MatineeData, NAME_None, RF_Transient);
 		NewGroup->GroupName = FName(TEXT("Preview Pawn"));
 		NewGroup->EnsureUniqueName();
 
-		PreviewMatinee.Get()->MatineeData->InterpGroups.Add(NewGroup);
+		PreviewMatineeActor.Get()->MatineeData->InterpGroups.Add(NewGroup);
 
 		// now add group inst
-		UInterpGroupInst* NewGroupInst = ConstructObject<UInterpGroupInst>( UInterpGroupInst::StaticClass(), PreviewMatinee.Get(), NAME_None, RF_Transient );
+		UInterpGroupInst* NewGroupInst = ConstructObject<UInterpGroupInst>(UInterpGroupInst::StaticClass(), PreviewMatineeActor.Get(), NAME_None, RF_Transient);
 		// Initialise group instance, saving ref to actor it works on.
 		NewGroupInst->InitGroupInst(NewGroup, PreviewInfo.PawnInst);
-		const int32 NewGroupInstIndex = PreviewMatinee.Get()->GroupInst.Add(NewGroupInst);
+		const int32 NewGroupInstIndex = PreviewMatineeActor.Get()->GroupInst.Add(NewGroupInst);
 
 		//Link group with actor
-		PreviewMatinee.Get()->InitGroupActorForGroup(NewGroup, PreviewInfo.PawnInst);
+		PreviewMatineeActor.Get()->InitGroupActorForGroup(NewGroup, PreviewInfo.PawnInst);
 
 		// Now time to add AnimTrack so that we can play animation
 		int32 AnimTrackIndex = INDEX_NONE;
@@ -170,7 +172,7 @@ void FAssetTypeActions_CameraAnim::OpenAssetEditor( const TArray<UObject*>& InOb
 			// construct a temporary matinee actor
 			CreateMatineeActorForCameraAnim(CameraAnim);
 
-			if(GEditor->ShouldOpenMatinee(PreviewMatinee.Get()))
+			if (GEditor->ShouldOpenMatinee(PreviewMatineeActor.Get()))
 			{
 				// changed the actor type, but don't want to lose any properties from previous
 				// so duplicate from old, but with new class
@@ -183,14 +185,14 @@ void FAssetTypeActions_CameraAnim::OpenAssetEditor( const TArray<UObject*>& InOb
 				UInterpGroupCamera* NewInterpGroup = CastChecked<UInterpGroupCamera>(CameraAnim->CameraInterpGroup);
 				check(NewInterpGroup);
 
-				if (PreviewMatinee.Get()->MatineeData)
+				if (PreviewMatineeActor.Get()->MatineeData)
 				{
-					PreviewMatinee.Get()->MatineeData->SetFlags(RF_Transient);
-					PreviewMatinee.Get()->MatineeData->InterpLength = CameraAnim->AnimLength;
+					PreviewMatineeActor.Get()->MatineeData->SetFlags(RF_Transient);
+					PreviewMatineeActor.Get()->MatineeData->InterpLength = CameraAnim->AnimLength;
 
 					if (NewInterpGroup)
 					{
-						PreviewMatinee.Get()->MatineeData->InterpGroups.Add(NewInterpGroup);
+						PreviewMatineeActor.Get()->MatineeData->InterpGroups.Add(NewInterpGroup);
 					}
 				}
 
@@ -198,23 +200,23 @@ void FAssetTypeActions_CameraAnim::OpenAssetEditor( const TArray<UObject*>& InOb
 				CreateCameraActorForCameraAnim(CameraAnim);
 
 				// set up the group actor
-				PreviewMatinee.Get()->InitGroupActorForGroup(NewInterpGroup, PreviewCamera.Get());
+				PreviewMatineeActor.Get()->InitGroupActorForGroup(NewInterpGroup, PreviewCamera.Get());
 
 				// Create preview pawn
 				CreatePreviewPawnForCameraAnim(CameraAnim);
 
 				// this will create the instances for everything
-				PreviewMatinee.Get()->InitInterp();
+				PreviewMatineeActor.Get()->InitInterp();
 
 				// open matinee for this actor
-				GEditor->OpenMatinee(PreviewMatinee.Get());
+				GEditor->OpenMatinee(PreviewMatineeActor.Get());
 
 				// install our delegate so we can clean up when finished
 				GEditorModeTools().OnEditorModeChanged().AddRaw(this, &FAssetTypeActions_CameraAnim::OnMatineeEditorClosed);
 			}
 			else
 			{
-				GEditor->GetEditorWorldContext().World()->DestroyActor(PreviewMatinee.Get());
+				GEditor->GetEditorWorldContext().World()->DestroyActor(PreviewMatineeActor.Get());
 			}
 		}
 	}
@@ -232,10 +234,10 @@ void FAssetTypeActions_CameraAnim::OnMatineeEditorClosed(FEdMode* InEditorMode,
 			PreviewCamera.Reset();
 		}
 
-		if(PreviewMatinee.IsValid())
+		if (PreviewMatineeActor.IsValid())
 		{
-			GEditor->GetEditorWorldContext().World()->DestroyActor(PreviewMatinee.Get(), false, false);
-			PreviewMatinee.Reset();
+			GEditor->GetEditorWorldContext().World()->DestroyActor(PreviewMatineeActor.Get(), false, false);
+			PreviewMatineeActor.Reset();
 		}
 
 		if(PreviewPawn.IsValid())
diff --git a/Engine/Source/Developer/AssetTools/Private/AssetTypeActions/AssetTypeActions_CameraAnim.h b/Engine/Source/Developer/AssetTools/Private/AssetTypeActions/AssetTypeActions_CameraAnim.h
--- a/Engine/Source/Developer/AssetTools/Private/AssetTypeActions/AssetTypeActions_CameraAnim.h
+++ b/Engine/Source/Developer/AssetTools/Private/AssetTypeActions/AssetTypeActions_CameraAnim.h
@@ -37,7 +37,7 @@ private:
 	TWeakObjectPtr<class ACameraActor> PreviewCamera;
 
 	/** The matinee actor we will use for previewing the camera anim */
-	TWeakObjectPtr<class AMatineeActorCameraAnim> PreviewMatinee;
+	TWeakObjectPtr<class AMatineeActorCameraAnim> PreviewMatineeActor;
 
 	/** The pawn we we will use for previewing the camera anim */
 	TWeakObjectPtr<class APawn> PreviewPawn;
diff --git a/Engine/Source/Editor/Matinee/Private/MatineeTools.cpp b/Engine/Source/Editor/Matinee/Private/MatineeTools.cpp
--- a/Engine/Source/Editor/Matinee/Private/MatineeTools.cpp
+++ b/Engine/Source/Editor/Matinee/Private/MatineeTools.cpp
@@ -1759,6 +1759,13 @@ void FMatinee::SetInterpEnd(float NewInterpLength)
 	IData->EdSectionStart = FMath::Clamp(IData->EdSectionStart, 0.f, IData->InterpLength);
 	IData->EdSectionEnd = FMath::Clamp(IData->EdSectionEnd, 0.f, IData->InterpLength);
 	CurveEd->SetRegionMarker(true, IData->EdSectionStart, IData->EdSectionEnd, RegionFillColor);
+
+	// Update the CameraAnim if necessary
+	AMatineeActorCameraAnim* const CamAnimMatineeActor = CastChecked<AMatineeActorCameraAnim>(MatineeActor);
+	if ( (CamAnimMatineeActor != nullptr) && (CamAnimMatineeActor->CameraAnim != nullptr) )
+	{
+		CamAnimMatineeActor->CameraAnim->AnimLength = IData->InterpLength;
+	}
 }
 
 void FMatinee::MoveLoopMarker(float NewMarkerPos, bool bIsStart)
diff --git a/Engine/Source/Runtime/Engine/Classes/Camera/CameraActor.h b/Engine/Source/Runtime/Engine/Classes/Camera/CameraActor.h
--- a/Engine/Source/Runtime/Engine/Classes/Camera/CameraActor.h
+++ b/Engine/Source/Runtime/Engine/Classes/Camera/CameraActor.h
@@ -17,6 +17,9 @@ public:
 	UPROPERTY(Category=CameraActor, VisibleAnywhere, BlueprintReadOnly)
 	TSubobjectPtr<class UCameraComponent> CameraComponent;
 
+	/** If this CameraActor is being used to preview a CameraAnim in the editor, this is the anim being previewed. */
+	TWeakObjectPtr<class UCameraAnim> PreviewedCameraAnim;
+
 private:
 	UPROPERTY()
 	uint32 bConstrainAspectRatio_DEPRECATED:1;
@@ -37,5 +40,7 @@ public:
 	// Begin UObject interface
 	virtual void Serialize(FArchive& Ar) OVERRIDE;
 	ENGINE_API virtual void PostLoadSubobjects(FObjectInstancingGraph* OuterInstanceGraph) OVERRIDE;
+	ENGINE_API virtual void PostEditChangeProperty(FPropertyChangedEvent& PropertyChangedEvent) OVERRIDE;
 	// End UObject interface
+	
 };
diff --git a/Engine/Source/Runtime/Engine/Classes/Camera/CameraAnim.h b/Engine/Source/Runtime/Engine/Classes/Camera/CameraAnim.h
--- a/Engine/Source/Runtime/Engine/Classes/Camera/CameraAnim.h
+++ b/Engine/Source/Runtime/Engine/Classes/Camera/CameraAnim.h
@@ -35,6 +35,13 @@ class UCameraAnim : public UObject
 	UPROPERTY()
 	float BaseFOV;
 
+	/** Default PP settings to put on the animated camera. For modifying PP without keyframes. */
+	UPROPERTY()
+	FPostProcessSettings BasePostProcessSettings;
+
+	/** Default PP blend weight to put on the animated camera. For modifying PP without keyframes. */
+	UPROPERTY()
+	float BasePostProcessBlendWeight;
 
 protected:
 	// @todo document
diff --git a/Engine/Source/Runtime/Engine/Classes/Camera/CameraAnimInst.h b/Engine/Source/Runtime/Engine/Classes/Camera/CameraAnimInst.h
--- a/Engine/Source/Runtime/Engine/Classes/Camera/CameraAnimInst.h
+++ b/Engine/Source/Runtime/Engine/Classes/Camera/CameraAnimInst.h
@@ -1,15 +1,14 @@
-// Copyright 1998-2014 Epic Games, Inc. All Rights Reserved.
-
+ // Copyright 1998-2014 Epic Games, Inc. All Rights Reserved.
 
 /**
- *	CameraAnim: defines a pre-packaged animation to be played on a camera.
+ *	A CameraAnimInst is an active instance of a CameraAnim.
  */
 
 #pragma once
 #include "CameraAnimInst.generated.h"
 
-UCLASS(notplaceable)
-class UCameraAnimInst : public UObject
+UCLASS(notplaceable, BlueprintType, transient)
+class ENGINE_API UCameraAnimInst : public UObject
 {
 	GENERATED_UCLASS_BODY()
 
@@ -24,70 +23,57 @@ protected:
 
 public:
 	/** Current time for the animation */
-	UPROPERTY(transient)
 	float CurTime;
 
 protected:
 	/** True if the animation should loop, false otherwise. */
-	UPROPERTY(transient)
 	uint32 bLooping:1;
 
 public:
 	/** True if the animation has finished, false otherwise. */
-	UPROPERTY(transient)
 	uint32 bFinished:1;
 
 	/** True if it's ok for the system to auto-release this instance upon completion. */
-	UPROPERTY(transient)
 	uint32 bAutoReleaseWhenFinished:1;
 
 protected:
 	/** Time to interpolate in from zero, for smooth starts. */
-	UPROPERTY()
 	float BlendInTime;
 
 	/** Time to interpolate out to zero, for smooth finishes. */
-	UPROPERTY()
 	float BlendOutTime;
 
 	/** True if currently blending in. */
-	UPROPERTY(transient)
 	uint32 bBlendingIn:1;
 
 	/** True if currently blending out. */
-	UPROPERTY(transient)
 	uint32 bBlendingOut:1;
 
 	/** Current time for the blend-in.  I.e. how long we have been blending. */
-	UPROPERTY(transient)
 	float CurBlendInTime;
 
 	/** Current time for the blend-out.  I.e. how long we have been blending. */
-	UPROPERTY(transient)
 	float CurBlendOutTime;
 
 public:
 	/** Multiplier for playback rate.  1.0 = normal. */
-	UPROPERTY()
+	UPROPERTY(BlueprintReadWrite, Category=CameraAnimInst)
 	float PlayRate;
 
-	/** "Intensity" scalar.  This is the scale at which the anim was first played.  */
-	UPROPERTY()
+	/** "Intensity" value used to scale keyframe values. */
 	float BasePlayScale;
 
 	/** A supplemental scale factor, allowing external systems to scale this anim as necessary.  This is reset to 1.f each frame. */
-	UPROPERTY()
 	float TransientScaleModifier;
 
 	/* Number in range [0..1], controlling how much this influence this instance should have. */
-	UPROPERTY()
 	float CurrentBlendWeight;
 
 protected:
 	/** How much longer to play the anim, if a specific duration is desired.  Has no effect if 0.  */
-	UPROPERTY(transient)
 	float RemainingTime;
 
+
 public:
 	/** cached movement track from the currently playing anim so we don't have to go find it every frame */
 	UPROPERTY(transient)
@@ -97,29 +83,32 @@ public:
 	class UInterpTrackInstMove* MoveInst;
 
 	UPROPERTY()
-	TEnumAsByte<enum ECameraAnimPlaySpace> PlaySpace;
+	TEnumAsByte<ECameraAnimPlaySpace::Type> PlaySpace;
 
-	/** The user-defined space for CAPS_UserDefined */
-	UPROPERTY(transient)
+	/** The user-defined space for UserDefined PlaySpace */
 	FMatrix UserPlaySpaceMatrix;
 
 	/** Camera Anim debug variable to trace back to previous location **/
-	UPROPERTY(transient)
 	FVector LastCameraLoc;
 
+	/** transform of initial anim key, used for treating anim keys as offsets from initial key */
+	FTransform InitialCamToWorld;
+
+	/** FOV of the initial anim key, used for treating fov keys as offsets from initial key. */
+	float InitialFOV;
 
 	/**
 	 * Starts this instance playing the specified CameraAnim.
 	 *
-	 * CamAnim:		The animation that should play on this instance.
-	 * CamActor:	The  AActor  that will be modified by this animation.
-	 * InRate:		How fast to play the animation.  1.f is normal.
-	 * InScale:		How intense to play the animation.  1.f is normal.
-	 * InBlendInTime:	Time over which to linearly ramp in.
-	 * InBlendInTime:	Time over which to linearly ramp out.
-	 * bInLoop:			Whether or not to loop the animation.
-	 * bRandomStartTime:	Whether or not to choose a random time to start playing.  Only really makes sense for bLoop = true;
-	 * Duration:	optional specific playtime for this animation.  This is total time, including blends.
+	 * @param CamAnim			The animation that should play on this instance.
+	 * @param CamActor			The  AActor  that will be modified by this animation.
+	 * @param InRate			How fast to play the animation.  1.f is normal.
+	 * @param InScale			How intense to play the animation.  1.f is normal.
+	 * @param InBlendInTime		Time over which to linearly ramp in.
+	 * @param InBlendOutTime	Time over which to linearly ramp out.
+	 * @param bInLoop			Whether or not to loop the animation.
+	 * @param bRandomStartTime	Whether or not to choose a random time to start playing.  Only really makes sense for bLoop = true;
+	 * @param Duration			optional specific playtime for this animation.  This is total time, including blends.
 	 */
 	void Play(class UCameraAnim* Anim, class AActor* CamActor, float InRate, float InScale, float InBlendInTime, float InBlendOutTime, bool bInLoop, bool bRandomStartTime, float Duration = 0.f);
 	
@@ -130,13 +119,20 @@ public:
 	void AdvanceAnim(float DeltaTime, bool bJump);
 	
 	/** Stops this instance playing whatever animation it is playing. */
+	UFUNCTION(BlueprintCallable, Category = CameraAnimInst)
 	void Stop(bool bImmediate = false);
 	
 	/** Applies given scaling factor to the playing animation for the next update only. */
 	void ApplyTransientScaling(float Scalar);
 	
 	/** Sets this anim to play in an alternate playspace */
-	void SetPlaySpace(ECameraAnimPlaySpace NewSpace, FRotator UserPlaySpace = FRotator::ZeroRotator);
+	void SetPlaySpace(ECameraAnimPlaySpace::Type NewSpace, FRotator UserPlaySpace = FRotator::ZeroRotator);
+
+	/** Changes the running duration of this active anim, while maintaining playback position. */
+	UFUNCTION(BlueprintCallable, Category=CameraAnimInst)
+	void SetDuration(float NewDuration);
+
+	
 };
 
 
diff --git a/Engine/Source/Runtime/Engine/Classes/Camera/CameraComponent.h b/Engine/Source/Runtime/Engine/Classes/Camera/CameraComponent.h
--- a/Engine/Source/Runtime/Engine/Classes/Camera/CameraComponent.h
+++ b/Engine/Source/Runtime/Engine/Classes/Camera/CameraComponent.h
@@ -37,7 +37,7 @@ class UCameraComponent : public USceneComponent
 	TEnumAsByte<ECameraProjectionMode::Type> ProjectionMode;
 
 	/** Indicates if PostProcessSettings should be used when using this Camera to view through. */
-	UPROPERTY(EditAnywhere, BlueprintReadWrite, Category=CameraSettings, meta=(UIMin = "0.0", UIMax = "1.0"))
+	UPROPERTY(Interp, EditAnywhere, BlueprintReadWrite, Category=CameraSettings, meta=(UIMin = "0.0", UIMax = "1.0"))
 	float PostProcessBlendWeight;
 
 	/** Post process settings to use for this camera. Don't forget to check the properties you want to override */
diff --git a/Engine/Source/Runtime/Engine/Classes/Camera/CameraModifier_CameraShake.h b/Engine/Source/Runtime/Engine/Classes/Camera/CameraModifier_CameraShake.h
--- a/Engine/Source/Runtime/Engine/Classes/Camera/CameraModifier_CameraShake.h
+++ b/Engine/Source/Runtime/Engine/Classes/Camera/CameraModifier_CameraShake.h
@@ -1,15 +1,16 @@
 // Copyright 1998-2014 Epic Games, Inc. All Rights Reserved.
 
 /**
- * 
  * Camera modifier that provides support for code-based oscillating camera shakes.
  */
 
 #pragma once
+
+#include "CameraTypes.h"
 #include "CameraModifier_CameraShake.generated.h"
 
 USTRUCT()
-struct FCameraShakeInstance
+struct ENGINE_API FCameraShakeInstance
 {
 	GENERATED_USTRUCT_BODY()
 
@@ -56,7 +57,7 @@ struct FCameraShakeInstance
 
 	/** What space to play the shake in before applying to the camera.  Affects Anim and Oscillation both. */
 	UPROPERTY()
-	TEnumAsByte<enum ECameraAnimPlaySpace> PlaySpace;
+	TEnumAsByte<ECameraAnimPlaySpace::Type> PlaySpace;
 
 	/** Matrix defining the playspace, used when PlaySpace == CAPS_UserDefined */
 	UPROPERTY()
@@ -83,7 +84,7 @@ struct FCameraShakeInstance
 };
 
 UCLASS(dependson=UCameraShake, config=Camera)
-class UCameraModifier_CameraShake : public UCameraModifier
+class ENGINE_API UCameraModifier_CameraShake : public UCameraModifier
 {
 	GENERATED_UCLASS_BODY()
 
@@ -96,20 +97,18 @@ protected:
 	UPROPERTY(EditAnywhere, Category=CameraModifier_CameraShake)
 	float SplitScreenShakeScale;
 
-public:
-
-protected:
 	/** For situational scaling of individual shakes. */
 	virtual float GetShakeScale(FCameraShakeInstance const& ShakeInst) const;
+	virtual void ReinitShake(FCameraShakeInstance& ShakeInst, float Scale);
+
 public:
 	float InitializeOffset( const FFOscillator& Param );
-	virtual void ReinitShake(int32 ActiveShakeIdx, float Scale);
 	
 	/** Initialize camera shake structure */
-	virtual FCameraShakeInstance InitializeShake(TSubclassOf<class UCameraShake> NewShake, float Scale, ECameraAnimPlaySpace PlaySpace, FRotator UserPlaySpaceRot = FRotator::ZeroRotator);
+	virtual FCameraShakeInstance InitializeShake(TSubclassOf<class UCameraShake> NewShake, float Scale, ECameraAnimPlaySpace::Type PlaySpace, FRotator UserPlaySpaceRot = FRotator::ZeroRotator);
 	
 	/** Add a new screen shake to the list */
-	virtual void AddCameraShake(TSubclassOf<class UCameraShake> NewShake, float Scale, ECameraAnimPlaySpace PlaySpace=CAPS_CameraLocal, FRotator UserPlaySpaceRot = FRotator::ZeroRotator);
+	virtual void AddCameraShake(TSubclassOf<class UCameraShake> NewShake, float Scale, ECameraAnimPlaySpace::Type PlaySpace=ECameraAnimPlaySpace::CameraLocal, FRotator UserPlaySpaceRot = FRotator::ZeroRotator);
 	
 	virtual void RemoveCameraShake(TSubclassOf<class UCameraShake> Shake);
 	virtual void RemoveAllCameraShakes();
diff --git a/Engine/Source/Runtime/Engine/Classes/Camera/CameraShake.h b/Engine/Source/Runtime/Engine/Classes/Camera/CameraShake.h
--- a/Engine/Source/Runtime/Engine/Classes/Camera/CameraShake.h
+++ b/Engine/Source/Runtime/Engine/Classes/Camera/CameraShake.h
@@ -90,7 +90,7 @@ class ENGINE_API UCameraShake : public UObject
 
 	/** 
 	 *  true to only allow a single instance of this shake to play at any given time. 
-	 *  Subsequents attempts to play this shake will simply restart the timer.
+	 *  Subsequent attempts to play this shake will simply restart the timer.
 	 */
 	UPROPERTY(EditAnywhere, Category=CameraShake)
 	uint32 bSingleInstance:1;
diff --git a/Engine/Source/Runtime/Engine/Classes/Camera/CameraTypes.h b/Engine/Source/Runtime/Engine/Classes/Camera/CameraTypes.h
--- a/Engine/Source/Runtime/Engine/Classes/Camera/CameraTypes.h
+++ b/Engine/Source/Runtime/Engine/Classes/Camera/CameraTypes.h
@@ -16,6 +16,20 @@ namespace ECameraProjectionMode
 	};
 }
 
+UENUM()
+namespace ECameraAnimPlaySpace
+{
+	enum Type
+	{
+ 		/** This anim is applied in camera space */
+ 		CameraLocal,
+ 		/** This anim is applied in world space */
+ 		World,
+ 		/** This anim is applied in a user-specified space (defined by UserPlaySpaceMatrix) */
+ 		UserDefined,
+ 	};
+}
+
 USTRUCT()
 struct FMinimalViewInfo
 {
diff --git a/Engine/Source/Runtime/Engine/Classes/Camera/PlayerCameraManager.h b/Engine/Source/Runtime/Engine/Classes/Camera/PlayerCameraManager.h
--- a/Engine/Source/Runtime/Engine/Classes/Camera/PlayerCameraManager.h
+++ b/Engine/Source/Runtime/Engine/Classes/Camera/PlayerCameraManager.h
@@ -21,21 +21,6 @@ enum EViewTargetBlendFunction
 	VTBlend_MAX,
 };
 
-UENUM()
-enum ECameraAnimPlaySpace
-{
-	/** This anim is applied in camera space */
-	CAPS_CameraLocal,
-	/** This anim is applied in world space */
-	CAPS_World,
-	/** This anim is applied in a user-specified space (defined by UserPlaySpaceMatrix) */
-	CAPS_UserDefined,
-	CAPS_MAX,
-};
-
-/** The actors which the camera shouldn't see. Used to hide actors which the camera penetrates. */
-//var TArray<Actor> HiddenActors;
-
 /* Caching Camera, for optimization */
 USTRUCT()
 struct FCameraCacheEntry
@@ -294,11 +279,20 @@ protected:
 	UPROPERTY()
 	class UCameraAnimInst* AnimInstPool[8];    /*MAX_ACTIVE_CAMERA_ANIMS @fixmeconst*/
 
+	/** Parallel arrays */
+	TArray<struct FPostProcessSettings> PostProcessBlendCache;
+	TArray<float> PostProcessBlendCacheWeights;
+
+	void AddCachedPPBlend(struct FPostProcessSettings& PPSettings, float BlendWeight);
+	void ClearCachedPPBlends();
+
 public:
 	/** Array of anim instances that are currently playing and in-use */
 	UPROPERTY()
 	TArray<class UCameraAnimInst*> ActiveAnims;
 
+	void GetCachedPostProcessBlends(TArray<struct FPostProcessSettings> const*& OutPPSettings, TArray<float> const*& OutBlendWeigthts) const;
+	
 protected:
 	/** Array of anim instances that are not playing and available */
 	UPROPERTY()
@@ -383,13 +377,6 @@ public:
 	float ViewRollMax;
 
 
-	/**
-	 * Stop playing all instances of the indicated CameraAnim.
-	 * bImmediate: true to stop it right now, false to blend it out over BlendOutTime.
-	 */
-	UFUNCTION()
-	virtual void StopAllCameraAnims(bool bImmediate = false);
-
 	/** Blueprint hook to allow blueprints to implement custom cameras. Return true to use returned values, or false to ignore them. */
 	UFUNCTION(BlueprintImplementableEvent, BlueprintCosmetic)
 	virtual bool BlueprintUpdateCamera(AActor* CameraTarget, FVector& NewCameraLocation, FRotator& NewCameraRotation, float& NewCameraFOV);
@@ -400,19 +387,6 @@ public:
 
 	static const int32 MAX_ACTIVE_CAMERA_ANIMS = 8;
 
-protected:
-	/** Gets specified temporary CameraActor ready to update the specified Anim. */
-	void InitTempCameraActor(class ACameraActor* CamActor, class UCameraAnim* AnimToInitFor) const;
-	void ApplyAnimToCamera(class ACameraActor const* AnimatedCamActor, class UCameraAnimInst const* AnimInst, FMinimalViewInfo& InOutPOV);
-
-	/** Returns an available CameraAnimInst, or NULL if no more are available. */
-	UCameraAnimInst* AllocCameraAnimInst();
-	
-	/** Returns an available CameraAnimInst, or NULL if no more are available. */
-	void ReleaseCameraAnimInst(UCameraAnimInst* Inst);
-	
-	/** Returns first existing instance of the specified camera anim, or NULL if none exists. */
-	UCameraAnimInst* FindExistingCameraAnimInst(UCameraAnim const* Anim);
 
 	virtual void AssignViewTarget(AActor* NewTarget, FTViewTarget& VT, struct FViewTargetTransitionParams TransitionParams=FViewTargetTransitionParams());
 
@@ -564,7 +538,7 @@ public:
 	//
 
 	/** Play a camera shake */
-	virtual void PlayCameraShake(TSubclassOf<class UCameraShake> Shake, float Scale, ECameraAnimPlaySpace PlaySpace=CAPS_CameraLocal, FRotator UserPlaySpaceRot = FRotator::ZeroRotator);
+	virtual void PlayCameraShake(TSubclassOf<class UCameraShake> Shake, float Scale, enum ECameraAnimPlaySpace::Type PlaySpace=ECameraAnimPlaySpace::CameraLocal, FRotator UserPlaySpaceRot = FRotator::ZeroRotator);
 	
 	/** Stop playing a camera shake. */
 	virtual void StopCameraShake(TSubclassOf<class UCameraShake> Shake);
@@ -576,18 +550,59 @@ public:
 	//  CameraAnim support.
 	//
 	
-	/** Play the indicated CameraAnim on this camera.  Returns the CameraAnim instance, which can be stored to manipulate/stop the anim after the fact. */
-	virtual class UCameraAnimInst* PlayCameraAnim(class UCameraAnim* Anim, float Rate=1.f, float Scale=1.f, float BlendInTime = 0.f, float BlendOutTime = 0.f, bool bLoop = false, bool bRandomStartTime = false, float Duration = 0.f, bool bSingleInstance = false);
+	/**
+	 * Play the indicated CameraAnim on this camera.
+	 * 
+	 * @param Anim The animation that should play on this instance.
+	 * @param Rate				How fast to play the animation. 1.0 is normal.
+	 * @param Scale				How "intense" to play the animation. 1.0 is normal.
+	 * @param BlendInTime		Time to linearly ramp in.
+	 * @param BlendOutTime		Time to linearly ramp out.
+	 * @param bLoop				True to loop the animation if it hits the end.
+	 * @param bRandomStartTime	Whether or not to choose a random time to start playing. Useful with bLoop=true and a duration to randomize things like shakes.
+	 * @param Duration			Optional total playtime for this animation, including blends. 0 means to use animations natural duration, or infinite if looping.
+	 * @param PlaySpace			Which space to play the animation in.
+	 * @param UserPlaySpaceRot  Custom play space, used when PlaySpace is UserDefined.
+	 * @return The CameraAnim instance, which can be stored to manipulate/stop the anim after the fact.
+	 */
+	UFUNCTION(BlueprintCallable, Category="Camera Animation")
+	virtual class UCameraAnimInst* PlayCameraAnim(class UCameraAnim* Anim, float Rate=1.f, float Scale=1.f, float BlendInTime=0.f, float BlendOutTime=0.f, bool bLoop=false, bool bRandomStartTime=false, float Duration=0.f, ECameraAnimPlaySpace::Type PlaySpace=ECameraAnimPlaySpace::CameraLocal, FRotator UserPlaySpaceRot=FRotator::ZeroRotator);
 	
 	/**
 	 * Stop playing all instances of the indicated CameraAnim.
-	 * @param bImmediate if true stop it right now, false to blend it out over BlendOutTime.
+	 * @param bImmediate	True to stop it right now and ignore blend out, false to let it blend out as indicated.
 	 */
-	virtual void StopAllCameraAnimsByType(class UCameraAnim* Anim, bool bImmediate = false);
+	UFUNCTION(BlueprintCallable, Category = "Camera Animation")
+	virtual void StopAllInstancesOfCameraAnim(class UCameraAnim* Anim, bool bImmediate = false);
 	
-	/** Stops the given CameraAnim instance from playing.  The given pointer should be considered invalid after this. */
-	virtual void StopCameraAnim(class UCameraAnimInst* AnimInst, bool bImmediate = false);
+	/** 
+	 * Stops the given CameraAnimInst from playing.  The given pointer should be considered invalid after this. 
+	 * @param bImmediate	True to stop it right now and ignore blend out, false to let it blend out as indicated.
+	 */
+	UFUNCTION(BlueprintCallable, Category = "Camera Animation")
+	virtual void StopCameraAnimInst(class UCameraAnimInst* AnimInst, bool bImmediate = false);
 	
+	/**
+	 * Stop playing all CameraAnims on this CameraManager.
+	 * @param bImmediate	True to stop it right now and ignore blend out, false to let it blend out as indicated.
+	 */
+	UFUNCTION(BlueprintCallable, Category = "Camera Animation")
+	virtual void StopAllCameraAnims(bool bImmediate = false);
+
+	/** Returns first existing instance of the specified camera anim, or NULL if none exists. */
+	UCameraAnimInst* FindInstanceOfCameraAnim(class UCameraAnim const* Anim) const;
+
+protected:
+	/** Gets specified temporary CameraActor ready to update the specified Anim. */
+	void InitTempCameraActor(class ACameraActor* CamActor, class UCameraAnimInst const* AnimInstToInitFor) const;
+	void ApplyAnimToCamera(class ACameraActor const* AnimatedCamActor, class UCameraAnimInst const* AnimInst, FMinimalViewInfo& InOutPOV);
+
+	/** Returns an available CameraAnimInst, or NULL if no more are available. */
+	UCameraAnimInst* AllocCameraAnimInst();
+
+	/** Returns an available CameraAnimInst, or NULL if no more are available. */
+	void ReleaseCameraAnimInst(UCameraAnimInst* Inst);
+
 protected:
 	/** Limit the player's view pitch.  */
 	virtual void LimitViewPitch( FRotator& ViewRotation, float InViewPitchMin, float InViewPitchMax );
diff --git a/Engine/Source/Runtime/Engine/Classes/GameFramework/PlayerController.h b/Engine/Source/Runtime/Engine/Classes/GameFramework/PlayerController.h
--- a/Engine/Source/Runtime/Engine/Classes/GameFramework/PlayerController.h
+++ b/Engine/Source/Runtime/Engine/Classes/GameFramework/PlayerController.h
@@ -534,7 +534,7 @@ public:
 	 * @param CustomPlaySpace - Matrix used when Space = CAPS_UserDefined
 	 */
 	UFUNCTION(unreliable, client)
-	void ClientPlayCameraAnim(class UCameraAnim* AnimToPlay, float Scale = 0, float Rate = 0, float BlendInTime = 0, float BlendOutTime = 0, bool bLoop = false, bool bRandomStartTime = false, enum ECameraAnimPlaySpace Space = ECameraAnimPlaySpace(0), FRotator CustomPlaySpace = FRotator::ZeroRotator);
+	void ClientPlayCameraAnim(class UCameraAnim* AnimToPlay, float Scale=0, float Rate=0, float BlendInTime=0, float BlendOutTime=0, bool bLoop=false, bool bRandomStartTime=false, enum ECameraAnimPlaySpace::Type Space=ECameraAnimPlaySpace::CameraLocal, FRotator CustomPlaySpace=FRotator::ZeroRotator);
 
 	/** 
 	 * Play Camera Shake 
@@ -544,7 +544,7 @@ public:
 	 * @param UserPlaySpaceRot - Matrix used when PlaySpace = CAPS_UserDefined
 	 */
 	UFUNCTION(unreliable, client, BlueprintCallable, Category="Game|Feedback")
-	void ClientPlayCameraShake(TSubclassOf<class UCameraShake> Shake, float Scale = 0, enum ECameraAnimPlaySpace PlaySpace = CAPS_CameraLocal, FRotator UserPlaySpaceRot = FRotator::ZeroRotator);
+	void ClientPlayCameraShake(TSubclassOf<class UCameraShake> Shake, float Scale = 0, ECameraAnimPlaySpace::Type PlaySpace = ECameraAnimPlaySpace::CameraLocal, FRotator UserPlaySpaceRot = FRotator::ZeroRotator);
 
 	/**
 	 * Play sound client-side (so only the client will hear it)
diff --git a/Engine/Source/Runtime/Engine/Classes/Matinee/InterpTrackInstEvent.h b/Engine/Source/Runtime/Engine/Classes/Matinee/InterpTrackInstEvent.h
--- a/Engine/Source/Runtime/Engine/Classes/Matinee/InterpTrackInstEvent.h
+++ b/Engine/Source/Runtime/Engine/Classes/Matinee/InterpTrackInstEvent.h
@@ -4,7 +4,7 @@
 #pragma once
 #include "InterpTrackInstEvent.generated.h"
 
-UCLASS(HeaderGroup=Interpolation)
+UCLASS(HeaderGroup=Interpolation, MinimalAPI)
 class UInterpTrackInstEvent : public UInterpTrackInst
 {
 	GENERATED_UCLASS_BODY()
diff --git a/Engine/Source/Runtime/Engine/Classes/Matinee/MatineeActor.h b/Engine/Source/Runtime/Engine/Classes/Matinee/MatineeActor.h
--- a/Engine/Source/Runtime/Engine/Classes/Matinee/MatineeActor.h
+++ b/Engine/Source/Runtime/Engine/Classes/Matinee/MatineeActor.h
@@ -450,7 +450,7 @@ public:
 	
 
 	/** Try to invoke the event with the given name in the level script */
-	void NotifyEventTriggered(FName EventName, float EventTime);
+	virtual void NotifyEventTriggered(FName EventName, float EventTime);
 
 	/** Util to get the name of the function to find for the given event name */
 	ENGINE_API FName GetFunctionNameForEvent(FName EventName);
diff --git a/Engine/Source/Runtime/Engine/Private/Camera/CameraActor.cpp b/Engine/Source/Runtime/Engine/Private/Camera/CameraActor.cpp
--- a/Engine/Source/Runtime/Engine/Private/Camera/CameraActor.cpp
+++ b/Engine/Source/Runtime/Engine/Private/Camera/CameraActor.cpp
@@ -66,4 +66,16 @@ void ACameraActor::PostLoadSubobjects(FObjectInstancingGraph* OuterInstanceGraph
 	}
 }
 
+void ACameraActor::PostEditChangeProperty(FPropertyChangedEvent& PropertyChangedEvent)
+{
+	Super::PostEditChangeProperty(PropertyChangedEvent);
+
+	if (PreviewedCameraAnim.IsValid() && CameraComponent)
+	{
+		PreviewedCameraAnim->BaseFOV = CameraComponent->FieldOfView;
+		PreviewedCameraAnim->BasePostProcessSettings = CameraComponent->PostProcessSettings;
+		PreviewedCameraAnim->BasePostProcessBlendWeight = CameraComponent->PostProcessBlendWeight;
+	}
+}
+
 #undef LOCTEXT_NAMESPACE
diff --git a/Engine/Source/Runtime/Engine/Private/Camera/CameraAnimInst.cpp b/Engine/Source/Runtime/Engine/Private/Camera/CameraAnimInst.cpp
--- a/Engine/Source/Runtime/Engine/Private/Camera/CameraAnimInst.cpp
+++ b/Engine/Source/Runtime/Engine/Private/Camera/CameraAnimInst.cpp
@@ -13,7 +13,7 @@ UCameraAnimInst::UCameraAnimInst(const class FPostConstructInitializeProperties&
 	bAutoReleaseWhenFinished = true;
 	PlayRate = 1.0f;
 	TransientScaleModifier = 1.0f;
-	PlaySpace = CAPS_CameraLocal;
+	PlaySpace = ECameraAnimPlaySpace::CameraLocal;
 
 	InterpGroupInst = PCIP.CreateDefaultSubobject<UInterpGroupInst>(this, TEXT("InterpGroupInst0"));
 }
@@ -131,8 +131,26 @@ void UCameraAnimInst::Update(float NewRate, float NewScale, float NewBlendInTime
 	bFinished = false;
 }
 
+void UCameraAnimInst::SetDuration(float NewDuration)
+{
+	// setting a new duration will reset the play timer back to 0 but maintain current playback position
+	// if blending out, stop it and blend back in to reverse it so it's smooth
+	if (bBlendingOut)
+	{
+		bBlendingOut = false;
+		CurBlendOutTime = 0.f;
+
+		// stop any blendout and reverse it to a blendin
+		bBlendingIn = true;
+		CurBlendInTime = BlendInTime * (1.f - CurBlendOutTime / BlendOutTime);
+	}
+
+	RemainingTime = (NewDuration > 0.f) ? (NewDuration - BlendOutTime) : 0.f;
+	bFinished = false;
+}
+
 
-void UCameraAnimInst::Play(class UCameraAnim* Anim, class AActor* CamActor, float InRate, float InScale, float InBlendInTime, float InBlendOutTime, bool bInLooping, bool bRandomStartTime, float Duration)
+void UCameraAnimInst::Play(UCameraAnim* Anim, class AActor* CamActor, float InRate, float InScale, float InBlendInTime, float InBlendOutTime, bool bInLooping, bool bRandomStartTime, float Duration)
 {
 	if (Anim && Anim->CameraInterpGroup)
 	{
@@ -156,7 +174,6 @@ void UCameraAnimInst::Play(class UCameraAnim* Anim, class AActor* CamActor, floa
 		RemainingTime = (Duration > 0.f) ? (Duration - BlendOutTime) : 0.f;
 
 		// init the interpgroup
-
 		if ( CamActor->IsA(ACameraActor::StaticClass()) )
 		{
 			// ensure CameraActor is zeroed, so RelativeToInitial anims get proper InitialTM
@@ -175,7 +192,28 @@ void UCameraAnimInst::Play(class UCameraAnim* Anim, class AActor* CamActor, floa
 				// only 1 move track per group, so we can bail here
 				break;					
 			}
-		}	
+		}
+
+		// store initial transform so we can treat camera movements as offsets relative to the initial anim key
+		if (MoveTrack && MoveInst)
+		{
+			FRotator OutRot;
+			FVector OutLoc;
+			MoveTrack->GetKeyTransformAtTime(MoveInst, 0.f, OutLoc, OutRot);
+			InitialCamToWorld = FTransform(OutRot, OutLoc);		// @todo, store inverted since that's how we use it?
+
+			// find FOV track if it exists, else just use the fov saved in the anim
+			InitialFOV = Anim->BaseFOV;
+			for (int32 Idx = 0; Idx < InterpGroupInst->TrackInst.Num(); ++Idx)
+			{
+				UInterpTrackFloatProp* const FloatTrack = Cast<UInterpTrackFloatProp>(CamAnim->CameraInterpGroup->InterpTracks[Idx]);
+				if ( FloatTrack && (FloatTrack->PropertyName == TEXT("CameraComponent.FieldOfView")) )
+				{
+					InitialFOV = FloatTrack->EvalSub(0, 0.f);
+				}
+			}
+
+		}
 	}
 }
 
@@ -204,9 +242,9 @@ void UCameraAnimInst::ApplyTransientScaling(float Scalar)
 	TransientScaleModifier *= Scalar;
 }
 
-void UCameraAnimInst::SetPlaySpace(ECameraAnimPlaySpace NewSpace, FRotator UserPlaySpace)
+void UCameraAnimInst::SetPlaySpace(ECameraAnimPlaySpace::Type NewSpace, FRotator UserPlaySpace)
 {
 	PlaySpace = NewSpace;
-	UserPlaySpaceMatrix = (PlaySpace == CAPS_UserDefined) ? FRotationMatrix(UserPlaySpace) : FMatrix::Identity;
+	UserPlaySpaceMatrix = (PlaySpace == ECameraAnimPlaySpace::UserDefined) ? FRotationMatrix(UserPlaySpace) : FMatrix::Identity;
 }
 
diff --git a/Engine/Source/Runtime/Engine/Private/Camera/CameraModifier_CameraShake.cpp b/Engine/Source/Runtime/Engine/Private/Camera/CameraModifier_CameraShake.cpp
--- a/Engine/Source/Runtime/Engine/Private/Camera/CameraModifier_CameraShake.cpp
+++ b/Engine/Source/Runtime/Engine/Private/Camera/CameraModifier_CameraShake.cpp
@@ -129,7 +129,7 @@ void UCameraModifier_CameraShake::UpdateCameraShake(float DeltaTime, FCameraShak
 		RotOffset.Yaw = UpdateFOscillator(Shake.SourceShake->RotOscillation.Yaw, Shake.RotSinOffset.Y, DeltaTime) * OscillationScale;
 		RotOffset.Roll = UpdateFOscillator(Shake.SourceShake->RotOscillation.Roll, Shake.RotSinOffset.Z, DeltaTime) * OscillationScale;
 
-		if (Shake.PlaySpace == CAPS_CameraLocal)
+		if (Shake.PlaySpace == ECameraAnimPlaySpace::CameraLocal)
 		{
 			// the else case will handle this as well, but this is the faster, cleaner, most common code path
 
@@ -144,7 +144,7 @@ void UCameraModifier_CameraShake::UpdateCameraShake(float DeltaTime, FCameraShak
 		else
 		{
 			// find desired space
-			FMatrix const PlaySpaceToWorld = (Shake.PlaySpace == CAPS_UserDefined) ? Shake.UserPlaySpaceMatrix : FMatrix::Identity; 
+			FMatrix const PlaySpaceToWorld = (Shake.PlaySpace == ECameraAnimPlaySpace::UserDefined) ? Shake.UserPlaySpaceMatrix : FMatrix::Identity;
 
 			// apply loc offset relative to desired space
 			InOutPOV.Location += PlaySpaceToWorld.TransformVector( LocOffset );
@@ -153,7 +153,7 @@ void UCameraModifier_CameraShake::UpdateCameraShake(float DeltaTime, FCameraShak
 
 			// find transform from camera to the "play space"
 			FRotationMatrix const CamToWorld(InOutPOV.Rotation);
-			FMatrix const CameraToPlaySpace = CamToWorld * PlaySpaceToWorld.InverseSafe();	// CameraToWorld * WorldToPlaySpace
+			FMatrix const CameraToPlaySpace = CamToWorld * PlaySpaceToWorld.InverseSafe();			// CameraToWorld * WorldToPlaySpace
 
 			// find transform from anim (applied in playspace) back to camera
 			FRotationMatrix const AnimToPlaySpace(RotOffset);
@@ -226,28 +226,28 @@ float UCameraModifier_CameraShake::InitializeOffset( const FFOscillator& Param )
 	}
 }
 
-void UCameraModifier_CameraShake::ReinitShake(int32 ActiveShakeIdx, float Scale)
+void UCameraModifier_CameraShake::ReinitShake(FCameraShakeInstance& ShakeInst, float NewScale)
 {
 	if (GEngine->IsSplitScreen(CameraOwner->GetWorld()))
 	{
-		Scale *= SplitScreenShakeScale;
+		NewScale *= SplitScreenShakeScale;
 	}
-	ActiveShakes[ActiveShakeIdx].Scale = Scale;
+	ShakeInst.Scale = NewScale;
 
-	UCameraShake* SourceShake = ActiveShakes[ActiveShakeIdx].SourceShake;
+	UCameraShake* const SourceShake = ShakeInst.SourceShake;
 
 	if (SourceShake->OscillationDuration != 0.f)
 	{
-		ActiveShakes[ActiveShakeIdx].OscillatorTimeRemaining = SourceShake->OscillationDuration;
+		ShakeInst.OscillatorTimeRemaining = SourceShake->OscillationDuration;
 
-		if (ActiveShakes[ActiveShakeIdx].bBlendingOut)
+		if (ShakeInst.bBlendingOut)
 		{
-			ActiveShakes[ActiveShakeIdx].bBlendingOut = false;
-			ActiveShakes[ActiveShakeIdx].CurrentBlendOutTime = 0.f;
+			ShakeInst.bBlendingOut = false;
+			ShakeInst.CurrentBlendOutTime = 0.f;
 
 			// stop any blendout and reverse it to a blendin
-			ActiveShakes[ActiveShakeIdx].bBlendingIn = true;
-			ActiveShakes[ActiveShakeIdx].CurrentBlendInTime = ActiveShakes[ActiveShakeIdx].SourceShake->OscillationBlendInTime * (1.f - ActiveShakes[ActiveShakeIdx].CurrentBlendOutTime / ActiveShakes[ActiveShakeIdx].SourceShake->OscillationBlendOutTime);
+			ShakeInst.bBlendingIn = true;
+			ShakeInst.CurrentBlendInTime = ShakeInst.SourceShake->OscillationBlendInTime * (1.f - ShakeInst.CurrentBlendOutTime / ShakeInst.SourceShake->OscillationBlendOutTime);
 		}
 	}
 
@@ -262,25 +262,25 @@ void UCameraModifier_CameraShake::ReinitShake(int32 ActiveShakeIdx, float Scale)
 			Duration = SourceShake->RandomAnimSegmentDuration;
 		}
 
-		ActiveShakes[ActiveShakeIdx].AnimInst = CameraOwner->PlayCameraAnim(SourceShake->Anim,
-			SourceShake->AnimPlayRate, 
-			Scale,
-			SourceShake->AnimBlendInTime,
-			SourceShake->AnimBlendOutTime,
-			bLoop,
-			bRandomStart,
-			Duration,
-			true);
+		if (ShakeInst.AnimInst != nullptr)
+		{
+			ShakeInst.AnimInst->Update(SourceShake->AnimPlayRate, NewScale, SourceShake->AnimBlendInTime, SourceShake->AnimBlendOutTime, Duration);
+		}
+		else
+		{
+			ShakeInst.AnimInst = CameraOwner->PlayCameraAnim(SourceShake->Anim, SourceShake->AnimPlayRate, NewScale, SourceShake->AnimBlendInTime, SourceShake->AnimBlendOutTime, bLoop, bRandomStart, Duration);
+		}
 	}
 }
 
-FCameraShakeInstance UCameraModifier_CameraShake::InitializeShake(TSubclassOf<class UCameraShake> NewShake, float Scale, ECameraAnimPlaySpace PlaySpace, FRotator UserPlaySpaceRot)
+FCameraShakeInstance UCameraModifier_CameraShake::InitializeShake(TSubclassOf<class UCameraShake> NewShake, float Scale, ECameraAnimPlaySpace::Type PlaySpace, FRotator UserPlaySpaceRot)
 {
 	FCameraShakeInstance Inst;
 	Inst.SourceShakeName = NewShake->GetFName();
 
 	// Create a camera shake object of class NewShake
-	Inst.SourceShake = Cast<UCameraShake>(StaticConstructObject(NewShake, this));
+	UCameraShake* const SourceShake = Cast<UCameraShake>(StaticConstructObject(NewShake, this));
+	Inst.SourceShake = SourceShake;
 
 	Inst.Scale = Scale;
 	if (GEngine->IsSplitScreen(CameraOwner->GetWorld()))
@@ -289,21 +289,21 @@ FCameraShakeInstance UCameraModifier_CameraShake::InitializeShake(TSubclassOf<cl
 	}
 
 	// init oscillations
-	if ( Inst.SourceShake->OscillationDuration != 0.f )
+	if ( SourceShake->OscillationDuration != 0.f )
 	{
-		Inst.RotSinOffset.X		= InitializeOffset( Inst.SourceShake->RotOscillation.Pitch );
-		Inst.RotSinOffset.Y		= InitializeOffset( Inst.SourceShake->RotOscillation.Yaw );
-		Inst.RotSinOffset.Z		= InitializeOffset( Inst.SourceShake->RotOscillation.Roll );
+		Inst.RotSinOffset.X		= InitializeOffset( SourceShake->RotOscillation.Pitch );
+		Inst.RotSinOffset.Y		= InitializeOffset( SourceShake->RotOscillation.Yaw );
+		Inst.RotSinOffset.Z		= InitializeOffset( SourceShake->RotOscillation.Roll );
 
-		Inst.LocSinOffset.X		= InitializeOffset( Inst.SourceShake->LocOscillation.X );
-		Inst.LocSinOffset.Y		= InitializeOffset( Inst.SourceShake->LocOscillation.Y );
-		Inst.LocSinOffset.Z		= InitializeOffset( Inst.SourceShake->LocOscillation.Z );
+		Inst.LocSinOffset.X		= InitializeOffset( SourceShake->LocOscillation.X );
+		Inst.LocSinOffset.Y		= InitializeOffset( SourceShake->LocOscillation.Y );
+		Inst.LocSinOffset.Z		= InitializeOffset( SourceShake->LocOscillation.Z );
 
-		Inst.FOVSinOffset		= InitializeOffset( Inst.SourceShake->FOVOscillation );
+		Inst.FOVSinOffset		= InitializeOffset( SourceShake->FOVOscillation );
 
-		Inst.OscillatorTimeRemaining = Inst.SourceShake->OscillationDuration;
+		Inst.OscillatorTimeRemaining = SourceShake->OscillationDuration;
 
-		if (Inst.SourceShake->OscillationBlendInTime > 0.f)
+		if (SourceShake->OscillationBlendInTime > 0.f)
 		{
 			Inst.bBlendingIn = true;
 			Inst.CurrentBlendInTime = 0.f;
@@ -311,30 +311,26 @@ FCameraShakeInstance UCameraModifier_CameraShake::InitializeShake(TSubclassOf<cl
 	}
 
 	// init anims
-	if (Inst.SourceShake->Anim != NULL)
+	if (SourceShake->Anim != NULL)
 	{
 		bool bLoop = false;
 		bool bRandomStart = false;
 		float Duration = 0.f;
-		if (Inst.SourceShake->bRandomAnimSegment)
+		if (SourceShake->bRandomAnimSegment)
 		{
 			bLoop = true;
 			bRandomStart = true;
-			Duration = Inst.SourceShake->RandomAnimSegmentDuration;
+			Duration = SourceShake->RandomAnimSegmentDuration;
 		}
 
 		if (Scale > 0.f)
 		{
-			Inst.AnimInst = CameraOwner->PlayCameraAnim(Inst.SourceShake->Anim, Inst.SourceShake->AnimPlayRate, Scale, Inst.SourceShake->AnimBlendInTime, Inst.SourceShake->AnimBlendOutTime, bLoop, bRandomStart, Duration, Inst.SourceShake->bSingleInstance);
-			if (PlaySpace != CAPS_CameraLocal && Inst.AnimInst != NULL)
-			{
-				Inst.AnimInst->SetPlaySpace(PlaySpace, UserPlaySpaceRot);
-			}
+			Inst.AnimInst = CameraOwner->PlayCameraAnim(SourceShake->Anim, SourceShake->AnimPlayRate, Scale, SourceShake->AnimBlendInTime, SourceShake->AnimBlendOutTime, bLoop, bRandomStart, Duration, PlaySpace, UserPlaySpaceRot);
 		}
 	}
 
 	Inst.PlaySpace = PlaySpace;
-	if (Inst.PlaySpace == CAPS_UserDefined)
+	if (Inst.PlaySpace == ECameraAnimPlaySpace::UserDefined)
 	{
 		Inst.UserPlaySpaceMatrix = FRotationMatrix(UserPlaySpaceRot);
 	}
@@ -343,17 +339,19 @@ FCameraShakeInstance UCameraModifier_CameraShake::InitializeShake(TSubclassOf<cl
 }
 
 
-void UCameraModifier_CameraShake::AddCameraShake( TSubclassOf<class UCameraShake> NewShake, float Scale, ECameraAnimPlaySpace PlaySpace, FRotator UserPlaySpaceRot )
+void UCameraModifier_CameraShake::AddCameraShake( TSubclassOf<class UCameraShake> NewShake, float Scale, ECameraAnimPlaySpace::Type PlaySpace, FRotator UserPlaySpaceRot )
 {
 	if (NewShake != NULL)
 	{
-		if (Cast<UCameraShake>(NewShake->GetDefaultObject())->bSingleInstance)
+		UCameraShake const* const ShakeCDO = GetDefault<UCameraShake>(NewShake);
+
+		if (ShakeCDO && ShakeCDO->bSingleInstance)
 		{
-			for (int32 i = 0; i < ActiveShakes.Num(); ++i)
+			for (FCameraShakeInstance& ShakeInst : ActiveShakes)
 			{
-				if (ActiveShakes[i].SourceShakeName == NewShake->GetFName())
+				if (ShakeInst.SourceShakeName == NewShake->GetFName())
 				{
-					ReinitShake(i, Scale);
+					ReinitShake(ShakeInst, Scale);
 					return;
 				}
 			}
@@ -375,7 +373,7 @@ void UCameraModifier_CameraShake::RemoveCameraShake(TSubclassOf<class UCameraSha
 			UCameraAnimInst* AnimInst = ActiveShakes[i].AnimInst;
 			if ( (AnimInst != NULL) && !AnimInst->bFinished )
 			{
-				CameraOwner->StopCameraAnim(AnimInst, true);
+				CameraOwner->StopCameraAnimInst(AnimInst, true);
 			}
 
 			ActiveShakes.RemoveAt(i,1);
@@ -392,7 +390,7 @@ void UCameraModifier_CameraShake::RemoveAllCameraShakes()
 		UCameraAnimInst* AnimInst = ActiveShakes[Idx].AnimInst;
 		if ( (AnimInst != NULL) && !AnimInst->bFinished )
 		{
-			CameraOwner->StopCameraAnim(AnimInst, true);
+			CameraOwner->StopCameraAnimInst(AnimInst, true);
 		}
 	}
 
diff --git a/Engine/Source/Runtime/Engine/Private/LocalPlayer.cpp b/Engine/Source/Runtime/Engine/Private/LocalPlayer.cpp
--- a/Engine/Source/Runtime/Engine/Private/LocalPlayer.cpp
+++ b/Engine/Source/Runtime/Engine/Private/LocalPlayer.cpp
@@ -373,7 +373,7 @@ FSceneView* ULocalPlayer::CalcSceneView( class FSceneViewFamily* ViewFamily,
 	ViewInitOptions.WorldToMetersScale = PlayerController->GetWorldSettings()->WorldToMeters;
 	PlayerController->BuildHiddenComponentList(OutViewLocation, /*out*/ ViewInitOptions.HiddenPrimitives);
 
-	FSceneView* View = new FSceneView(ViewInitOptions);
+	FSceneView* const View = new FSceneView(ViewInitOptions);
 	
 	View->ViewLocation = OutViewLocation;
 	View->ViewRotation = OutViewRotation;
@@ -388,6 +388,19 @@ FSceneView* ULocalPlayer::CalcSceneView( class FSceneViewFamily* ViewFamily,
 	{
 		View->StartFinalPostprocessSettings(OutViewLocation);
 
+		// CameraAnim override
+		if (PlayerController && PlayerController->PlayerCameraManager)
+		{
+			TArray<FPostProcessSettings> const* CameraAnimPPSettings;
+			TArray<float> const* CameraAnimPPBlendWeights;
+			PlayerController->PlayerCameraManager->GetCachedPostProcessBlends(CameraAnimPPSettings, CameraAnimPPBlendWeights);
+			
+			for (int32 PPIdx = 0; PPIdx < CameraAnimPPBlendWeights->Num(); ++PPIdx)
+			{
+				View->OverridePostProcessSettings( (*CameraAnimPPSettings)[PPIdx], (*CameraAnimPPBlendWeights)[PPIdx]);
+			}
+		}
+
 		//	CAMERA OVERRIDE
 		//	NOTE: Matinee works through this channel
 		View->OverridePostProcessSettings(ViewInfo.PostProcessSettings, ViewInfo.PostProcessBlendWeight);
diff --git a/Engine/Source/Runtime/Engine/Private/PlayerCameraManager.cpp b/Engine/Source/Runtime/Engine/Private/PlayerCameraManager.cpp
--- a/Engine/Source/Runtime/Engine/Private/PlayerCameraManager.cpp
+++ b/Engine/Source/Runtime/Engine/Private/PlayerCameraManager.cpp
@@ -7,6 +7,7 @@
 
 DEFINE_LOG_CATEGORY_STATIC(LogPlayerCameraManager, Log, All);
 
+
 //////////////////////////////////////////////////////////////////////////
 // APlayerCameraManager
 
@@ -206,6 +207,8 @@ void APlayerCameraManager::ApplyCameraModifiers(float DeltaTime, FMinimalViewInf
 		}
 	}
 
+	ClearCachedPPBlends();
+
 	// Now apply CameraAnims
 	// these essentially behave as the highest-pri modifier.
 	for (int32 Idx = 0; Idx < ActiveAnims.Num(); ++Idx)
@@ -215,7 +218,7 @@ void APlayerCameraManager::ApplyCameraModifiers(float DeltaTime, FMinimalViewInf
 		if (!AnimInst->bFinished)
 		{
 			// clear out animated camera actor
-			InitTempCameraActor(AnimCameraActor, AnimInst->CamAnim);
+			InitTempCameraActor(AnimCameraActor, AnimInst);
 
 			// evaluate the animation at the new time
 			AnimInst->AdvanceAnim(DeltaTime, false);
@@ -225,7 +228,6 @@ void APlayerCameraManager::ApplyCameraModifiers(float DeltaTime, FMinimalViewInf
 			{
 				ApplyAnimToCamera(AnimCameraActor, AnimInst, InOutPOV);
 			}
-
 		}
 
 		// handle animations that have finished
@@ -247,13 +249,38 @@ void APlayerCameraManager::ApplyCameraModifiers(float DeltaTime, FMinimalViewInf
 	}
 }
 
-void APlayerCameraManager::ApplyAnimToCamera(class ACameraActor const* AnimatedCamActor, class UCameraAnimInst const* AnimInst, FMinimalViewInfo& InOutPOV)
+void APlayerCameraManager::AddCachedPPBlend(struct FPostProcessSettings& PPSettings, float BlendWeight)
+{
+	check(PostProcessBlendCache.Num() == PostProcessBlendCacheWeights.Num());
+	PostProcessBlendCache.Add(PPSettings);
+	PostProcessBlendCacheWeights.Add(BlendWeight);
+}
+
+void APlayerCameraManager::ClearCachedPPBlends()
+{
+	PostProcessBlendCache.Empty();
+	PostProcessBlendCacheWeights.Empty();
+}
+
+void APlayerCameraManager::GetCachedPostProcessBlends(TArray<FPostProcessSettings> const*& OutPPSettings, TArray<float> const*& OutBlendWeigthts) const
+{
+	OutPPSettings = &PostProcessBlendCache;
+	OutBlendWeigthts = &PostProcessBlendCacheWeights;
+}
+
+void APlayerCameraManager::ApplyAnimToCamera(ACameraActor const* AnimatedCamActor, UCameraAnimInst const* AnimInst, FMinimalViewInfo& InOutPOV)
 {
 	float const Scale = AnimInst->CurrentBlendWeight;
 
 	FRotationMatrix const CameraToWorld(InOutPOV.Rotation);
 
-	if (AnimInst->PlaySpace == CAPS_CameraLocal)
+	// move animated cam actor to initial-relative position
+	FTransform const AnimatedCamToWorld = AnimatedCamActor->GetTransform();
+	FTransform const AnimatedCamToInitialCam = AnimatedCamToWorld * AnimInst->InitialCamToWorld.InverseSafe();
+	ACameraActor* const MutableCamActor = const_cast<ACameraActor*>(AnimatedCamActor);
+	MutableCamActor->SetActorTransform(AnimatedCamToInitialCam);		// set it back because that's what the code below expects
+
+	if (AnimInst->PlaySpace == ECameraAnimPlaySpace::CameraLocal)
 	{
 		// the code in the else block will handle this just fine, but this path provides efficiency and simplicity for the most common case
 
@@ -270,7 +297,7 @@ void APlayerCameraManager::ApplyAnimToCamera(class ACameraActor const* AnimatedC
 		// handle playing the anim in an arbitrary space relative to the camera
 
 		// find desired space
-		FMatrix const PlaySpaceToWorld = (AnimInst->PlaySpace == CAPS_UserDefined) ? AnimInst->UserPlaySpaceMatrix : FMatrix::Identity; 
+		FMatrix const PlaySpaceToWorld = (AnimInst->PlaySpace == ECameraAnimPlaySpace::UserDefined) ? AnimInst->UserPlaySpaceMatrix : FMatrix::Identity;
 
 		// loc
 		FVector const LocalOffset = PlaySpaceToWorld.TransformVector( AnimatedCamActor->GetActorLocation()*Scale );
@@ -295,10 +322,15 @@ void APlayerCameraManager::ApplyAnimToCamera(class ACameraActor const* AnimatedC
 	}
 
 	// fov
-	ACameraActor const* const DefaultCamActor = GetDefault<ACameraActor>();
-	if (DefaultCamActor)
+	const float FOVMin = 5.f;
+	const float FOVMax = 170.f;
+	InOutPOV.FOV += AnimatedCamActor->CameraComponent->FieldOfView - AnimInst->InitialFOV;
+	InOutPOV.FOV = FMath::Clamp<float>(InOutPOV.FOV, FOVMin, FOVMax);
+
+	// postprocess
+	if (AnimatedCamActor->CameraComponent.Get()->PostProcessBlendWeight > 0.f)
 	{
-		InOutPOV.FOV += (AnimatedCamActor->CameraComponent->FieldOfView - DefaultCamActor->CameraComponent->FieldOfView) * Scale;
+		AddCachedPPBlend(AnimatedCamActor->CameraComponent.Get()->PostProcessSettings, AnimatedCamActor->CameraComponent.Get()->PostProcessBlendWeight);
 	}
 }
 
@@ -333,7 +365,7 @@ void APlayerCameraManager::ReleaseCameraAnimInst(UCameraAnimInst* Inst)
 }
 
 
-UCameraAnimInst* APlayerCameraManager::FindExistingCameraAnimInst(UCameraAnim const* Anim)
+UCameraAnimInst* APlayerCameraManager::FindInstanceOfCameraAnim(UCameraAnim const* Anim) const
 {
 	int32 const NumActiveAnims = ActiveAnims.Num();
 	for (int32 Idx=0; Idx<NumActiveAnims; Idx++)
@@ -347,28 +379,17 @@ UCameraAnimInst* APlayerCameraManager::FindExistingCameraAnimInst(UCameraAnim co
 	return NULL;
 }
 
-UCameraAnimInst* APlayerCameraManager::PlayCameraAnim(class UCameraAnim* Anim, float Rate, float Scale, float BlendInTime, float BlendOutTime, bool bLoop, bool bRandomStartTime, float Duration, bool bSingleInstance)
+UCameraAnimInst* APlayerCameraManager::PlayCameraAnim(UCameraAnim* Anim, float Rate, float Scale, float BlendInTime, float BlendOutTime, bool bLoop, bool bRandomStartTime, float Duration, ECameraAnimPlaySpace::Type PlaySpace, FRotator UserPlaySpaceRot)
 {
-	if (bSingleInstance)
-	{
-		// try to update the existing instance with the new data
-		UCameraAnimInst* ExistingInst = FindExistingCameraAnimInst(Anim);
-		if (ExistingInst)
-		{
-			ExistingInst->Update(Rate, Scale, BlendInTime, BlendOutTime, Duration);
-			return ExistingInst;
-		}
-	}
-
 	// get a new instance and play it
 	if (AnimCameraActor != NULL)
 	{
 		UCameraAnimInst* const Inst = AllocCameraAnimInst();
 		if (Inst)
 		{
-			// clear LastCameraLoc
-			Inst->LastCameraLoc = FVector::ZeroVector;
+			Inst->LastCameraLoc = FVector::ZeroVector;		// clear LastCameraLoc
 			Inst->Play(Anim, AnimCameraActor, Rate, Scale, BlendInTime, BlendOutTime, bLoop, bRandomStartTime, Duration);
+			Inst->SetPlaySpace(PlaySpace, UserPlaySpaceRot);
 			return Inst;
 		}
 	}
@@ -376,7 +397,7 @@ UCameraAnimInst* APlayerCameraManager::PlayCameraAnim(class UCameraAnim* Anim, f
 	return NULL;
 }
 
-void APlayerCameraManager::StopAllCameraAnimsByType(class UCameraAnim* Anim, bool bImmediate)
+void APlayerCameraManager::StopAllInstancesOfCameraAnim(UCameraAnim* Anim, bool bImmediate)
 {
 	// find cameraaniminst for this.
 	for (int32 Idx=0; Idx<ActiveAnims.Num(); ++Idx)
@@ -396,7 +417,7 @@ void APlayerCameraManager::StopAllCameraAnims(bool bImmediate)
 	}
 }
 
-void APlayerCameraManager::StopCameraAnim(class UCameraAnimInst* AnimInst, bool bImmediate)
+void APlayerCameraManager::StopCameraAnimInst(class UCameraAnimInst* AnimInst, bool bImmediate)
 {
 	if (AnimInst != NULL)
 	{
@@ -404,19 +425,22 @@ void APlayerCameraManager::StopCameraAnim(class UCameraAnimInst* AnimInst, bool
 	}
 }
 
-void APlayerCameraManager::InitTempCameraActor(class ACameraActor* CamActor, class UCameraAnim* AnimToInitFor) const
+
+void APlayerCameraManager::InitTempCameraActor(ACameraActor* CamActor, UCameraAnimInst const* AnimInstToInitFor) const
 {
 	if (CamActor)
 	{
 		CamActor->TeleportTo(FVector::ZeroVector, FRotator::ZeroRotator);
 
-		if (AnimToInitFor)
+		if (AnimInstToInitFor)
 		{
 			ACameraActor const* const DefaultCamActor = GetDefault<ACameraActor>();
 			if (DefaultCamActor)
 			{
 				CamActor->CameraComponent->AspectRatio = DefaultCamActor->CameraComponent->AspectRatio;
-				CamActor->CameraComponent->FieldOfView = AnimToInitFor->BaseFOV;
+				CamActor->CameraComponent->FieldOfView = AnimInstToInitFor->CamAnim->BaseFOV;
+				CamActor->CameraComponent->PostProcessSettings = AnimInstToInitFor->CamAnim->BasePostProcessSettings;
+				CamActor->CameraComponent->PostProcessBlendWeight = AnimInstToInitFor->CamAnim->BasePostProcessBlendWeight;
 			}
 		}
 	}
@@ -583,13 +607,17 @@ void APlayerCameraManager::PostInitializeComponents()
 		FreeAnims.Add(AnimInstPool[Idx]);
 	}
 
-	// spawn the two temp CameraActors used for updating CameraAnims
-	//AnimCameraActor = SpawnActor<ADynamicCameraActor>(this, Instigator, true);
+	// spawn the temp CameraActor used for updating CameraAnims
+	FActorSpawnParameters SpawnInfo;
+	SpawnInfo.Owner = this;
+	SpawnInfo.Instigator = Instigator;
+	SpawnInfo.bNoCollisionFail = true;
+	AnimCameraActor = GetWorld()->SpawnActor<ACameraActor>(SpawnInfo);
 }
 
 void APlayerCameraManager::Destroyed()
 {
-	// clean up the temp camera actors
+	// clean up the temp camera actor
 	if (AnimCameraActor)
 	{
 		AnimCameraActor->Destroy();
@@ -966,7 +994,7 @@ AEmitterCameraLensEffectBase* APlayerCameraManager::AddCameraLensEffect(TSubclas
 			SpawnInfo.Owner = PCOwner->GetViewTarget();
 			SpawnInfo.Instigator = Instigator;
 			SpawnInfo.bNoCollisionFail = true;
-			LensEffect = GetWorld()->SpawnActor<AEmitterCameraLensEffectBase>(LensEffectEmitterClass, SpawnInfo );
+			LensEffect = GetWorld()->SpawnActor<AEmitterCameraLensEffectBase>(LensEffectEmitterClass, SpawnInfo);
 			if (LensEffect != NULL)
 			{
 				FVector CamLoc;
@@ -1008,7 +1036,7 @@ void APlayerCameraManager::ClearCameraLensEffects()
 
 
 
-void APlayerCameraManager::PlayCameraShake(TSubclassOf<UCameraShake> Shake, float Scale, ECameraAnimPlaySpace PlaySpace, FRotator UserPlaySpaceRot)
+void APlayerCameraManager::PlayCameraShake(TSubclassOf<UCameraShake> Shake, float Scale, ECameraAnimPlaySpace::Type PlaySpace, FRotator UserPlaySpaceRot)
 {
 	if (Shake != NULL)
 	{
@@ -1059,7 +1087,7 @@ void APlayerCameraManager::PlayWorldCameraShake(UWorld* InWorld, TSubclassOf<cla
 				FVector CamLoc;
 				FRotator CamRot;
 				PlayerController->PlayerCameraManager->GetCameraViewPoint(CamLoc, CamRot);
-				PlayerController->ClientPlayCameraShake(Shake, ShakeScale, CAPS_UserDefined, (Epicenter - CamLoc).Rotation());
+				PlayerController->ClientPlayCameraShake(Shake, ShakeScale, ECameraAnimPlaySpace::UserDefined, (Epicenter - CamLoc).Rotation());
 			}
 			else
 			{
diff --git a/Engine/Source/Runtime/Engine/Private/PlayerController.cpp b/Engine/Source/Runtime/Engine/Private/PlayerController.cpp
--- a/Engine/Source/Runtime/Engine/Private/PlayerController.cpp
+++ b/Engine/Source/Runtime/Engine/Private/PlayerController.cpp
@@ -3195,7 +3195,7 @@ void APlayerController::ProcessForceFeedback(const float DeltaTime, const bool b
 	}
 }
 
-void APlayerController::ClientPlayCameraShake_Implementation( TSubclassOf<class UCameraShake> Shake, float Scale, ECameraAnimPlaySpace PlaySpace, FRotator UserPlaySpaceRot )
+void APlayerController::ClientPlayCameraShake_Implementation( TSubclassOf<class UCameraShake> Shake, float Scale, ECameraAnimPlaySpace::Type PlaySpace, FRotator UserPlaySpaceRot )
 {
 	if (PlayerCameraManager != NULL)
 	{
@@ -3213,15 +3213,11 @@ void APlayerController::ClientStopCameraShake_Implementation( TSubclassOf<class
 
 void APlayerController::ClientPlayCameraAnim_Implementation( UCameraAnim* AnimToPlay, float Scale, float Rate,
 						float BlendInTime, float BlendOutTime, bool bLoop,
-						bool bRandomStartTime, ECameraAnimPlaySpace Space, FRotator CustomPlaySpace )
+						bool bRandomStartTime, ECameraAnimPlaySpace::Type Space, FRotator CustomPlaySpace )
 {
 	if (PlayerCameraManager != NULL)
 	{
-		UCameraAnimInst* AnimInst = PlayerCameraManager->PlayCameraAnim(AnimToPlay, Rate, Scale, BlendInTime, BlendOutTime, bLoop, bRandomStartTime);
-		if (AnimInst != NULL && Space != CAPS_CameraLocal)
-		{
-			AnimInst->SetPlaySpace(Space, CustomPlaySpace);
-		}
+		PlayerCameraManager->PlayCameraAnim(AnimToPlay, Rate, Scale, BlendInTime, BlendOutTime, bLoop, bRandomStartTime, 0.f, Space, CustomPlaySpace);
 	}
 }
 
@@ -3229,7 +3225,7 @@ void APlayerController::ClientStopCameraAnim_Implementation(UCameraAnim* AnimToS
 {
 	if (PlayerCameraManager != NULL)
 	{
-		PlayerCameraManager->StopAllCameraAnimsByType(AnimToStop);
+		PlayerCameraManager->StopAllInstancesOfCameraAnim(AnimToStop);
 	}
 }
 
