Add sub-stepping support for vehicles

diff --git a/Engine/Source/Runtime/Engine/Private/PhysicsEngine/PhysScene.cpp b/Engine/Source/Runtime/Engine/Private/PhysicsEngine/PhysScene.cpp
--- a/Engine/Source/Runtime/Engine/Private/PhysicsEngine/PhysScene.cpp
+++ b/Engine/Source/Runtime/Engine/Private/PhysicsEngine/PhysScene.cpp
@@ -399,9 +399,14 @@ void FPhysScene::TickPhysScene(uint32 SceneType, FGraphEventRef& InOutCompletion
 
 #if WITH_PHYSX
 
-	if ( VehicleManager && SceneType == PST_Sync )
+#if WITH_SUBSTEPPING
+	if (IsSubstepping() == false)
+#endif
 	{
-		VehicleManager->Update( AveragedFrameTime[SceneType] );
+		if (VehicleManager && SceneType == PST_Sync)
+		{
+			VehicleManager->Update(AveragedFrameTime[SceneType]);
+		}
 	}
 
 #if !WITH_APEX
@@ -1065,6 +1070,10 @@ void FPhysScene::InitPhysScene(uint32 SceneType)
 	//Initialize substeppers
 	//we don't bother sub-stepping cloth
 	PhysSubSteppers[SceneType] = SceneType == PST_Cloth ? NULL : new FPhysSubstepTask(PScene);
+	if (SceneType == PST_Sync)
+	{
+		PhysSubSteppers[SceneType]->SetVehicleManager(VehicleManager);
+	}
 	
 #endif
 
@@ -1094,6 +1103,12 @@ void FPhysScene::TermPhysScene(uint32 SceneType)
 		}
 
 #if WITH_SUBSTEPPING
+
+		if (SceneType == PST_Sync && PhysSubSteppers[SceneType])
+		{
+			PhysSubSteppers[SceneType]->SetVehicleManager(NULL);
+		}
+
 		delete PhysSubSteppers[SceneType];
 		PhysSubSteppers[SceneType] = NULL;
 #endif
diff --git a/Engine/Source/Runtime/Engine/Private/PhysicsEngine/PhysSubstepTasks.cpp b/Engine/Source/Runtime/Engine/Private/PhysicsEngine/PhysSubstepTasks.cpp
--- a/Engine/Source/Runtime/Engine/Private/PhysicsEngine/PhysSubstepTasks.cpp
+++ b/Engine/Source/Runtime/Engine/Private/PhysicsEngine/PhysSubstepTasks.cpp
@@ -4,7 +4,11 @@
 
 #if WITH_SUBSTEPPING
 
+#if WITH_PHYSX
 #include "PhysXSupport.h"
+#include "../Vehicles/PhysXVehicleManager.h"
+#endif
+
 #include "PhysSubstepTasks.h"
 
 #if WITH_PHYSX
@@ -18,6 +22,7 @@ FPhysSubstepTask::FPhysSubstepTask(PxScene * GivenPScene) :
 	StepScale(0.f),
 	TotalSubTime(0.f),
 	CurrentSubStep(0),
+	VehicleManager(NULL),
 #if WITH_APEX
 	ApexScene(NULL),
 #endif
@@ -266,21 +271,27 @@ void FPhysSubstepTask::SubstepSimulationStart()
 	FDelegateGraphTask::CreateAndDispatchWhenReady(FDelegateGraphTask::FDelegate::CreateRaw(this, &FPhysSubstepTask::SubstepSimulationEnd), TEXT("ProcessPhysSubstepSimulation"), CompletionEvent);
 
 	++CurrentSubStep;
-	if (CurrentSubStep < NumSubsteps)
+
+	bool bLastSubstep = CurrentSubStep >= NumSubsteps;
+
+	if (!bLastSubstep)
 	{
-		
+
 		Alpha += StepScale;
 		TotalSubTime += SubTime;
-		SubstepInterpolation(Alpha);
-		ApexScene->simulate(SubTime, false, SubstepTask);	
-		SubstepTask->removeReference();
 	}
-	else
+
+	float DeltaTime = bLastSubstep ? (DeltaSeconds - TotalSubTime) : SubTime;
+	float Interpolation = bLastSubstep ? 1.f : Alpha;
+
+	if (VehicleManager)
 	{
-		SubstepInterpolation(1.f);
-		ApexScene->simulate(DeltaSeconds - TotalSubTime, true, SubstepTask);
-		SubstepTask->removeReference();
+		VehicleManager->Update(DeltaTime);
 	}
+
+	SubstepInterpolation(Interpolation);
+	ApexScene->simulate(DeltaTime, bLastSubstep, SubstepTask);
+	SubstepTask->removeReference();
 }
 
 void FPhysSubstepTask::SubstepSimulationEnd(ENamedThreads::Type CurrentThread, const FGraphEventRef& MyCompletionGraphEvent)
@@ -306,4 +317,10 @@ void FPhysSubstepTask::SubstepSimulationEnd(ENamedThreads::Type CurrentThread, c
 }
 #endif
 
+
+void FPhysSubstepTask::SetVehicleManager(FPhysXVehicleManager * InVehicleManager)
+{
+	VehicleManager = InVehicleManager;
+}
+
 #endif //if WITH_SUBSTEPPING
\ No newline at end of file
diff --git a/Engine/Source/Runtime/Engine/Private/PhysicsEngine/PhysSubstepTasks.h b/Engine/Source/Runtime/Engine/Private/PhysicsEngine/PhysSubstepTasks.h
--- a/Engine/Source/Runtime/Engine/Private/PhysicsEngine/PhysSubstepTasks.h
+++ b/Engine/Source/Runtime/Engine/Private/PhysicsEngine/PhysSubstepTasks.h
@@ -107,6 +107,8 @@ public:
 	void SwapBuffers();
 	float UpdateTime(float UseDelta);
 
+	void SetVehicleManager(class FPhysXVehicleManager *	InVehicleManager);
+
 #if WITH_APEX
 	void StepSimulation(NxApexScene * ApexScene, PhysXCompletionTask * Task);
 	void SubstepSimulationStart();
@@ -134,6 +136,9 @@ private:
 	uint32 CurrentSubStep;
 	FGraphEventRef CompletionEvent;
 
+	/** Vehicle scene */
+	class FPhysXVehicleManager*			VehicleManager;
+
 #if WITH_APEX
 	class NxApexScene * ApexScene;
 #endif
