#ue4
Fixed WITH_EDITOR issue with CameraAnim checkin

diff --git a/Engine/Source/Runtime/Engine/Classes/Camera/CameraActor.h b/Engine/Source/Runtime/Engine/Classes/Camera/CameraActor.h
--- a/Engine/Source/Runtime/Engine/Classes/Camera/CameraActor.h
+++ b/Engine/Source/Runtime/Engine/Classes/Camera/CameraActor.h
@@ -40,7 +40,10 @@ public:
 	// Begin UObject interface
 	virtual void Serialize(FArchive& Ar) OVERRIDE;
 	ENGINE_API virtual void PostLoadSubobjects(FObjectInstancingGraph* OuterInstanceGraph) OVERRIDE;
+
+#if WITH_EDITOR
 	ENGINE_API virtual void PostEditChangeProperty(FPropertyChangedEvent& PropertyChangedEvent) OVERRIDE;
+#endif
 	// End UObject interface
 	
 };
