#ue4
#ios
* CIS fix

diff --git a/Engine/Source/Runtime/Engine/Private/Camera/CameraActor.cpp b/Engine/Source/Runtime/Engine/Private/Camera/CameraActor.cpp
--- a/Engine/Source/Runtime/Engine/Private/Camera/CameraActor.cpp
+++ b/Engine/Source/Runtime/Engine/Private/Camera/CameraActor.cpp
@@ -66,6 +66,7 @@ void ACameraActor::PostLoadSubobjects(FObjectInstancingGraph* OuterInstanceGraph
 	}
 }
 
+#if WITH_EDITOR
 void ACameraActor::PostEditChangeProperty(FPropertyChangedEvent& PropertyChangedEvent)
 {
 	Super::PostEditChangeProperty(PropertyChangedEvent);
@@ -77,5 +78,6 @@ void ACameraActor::PostEditChangeProperty(FPropertyChangedEvent& PropertyChanged
 		PreviewedCameraAnim->BasePostProcessBlendWeight = CameraComponent->PostProcessBlendWeight;
 	}
 }
+#endif
 
 #undef LOCTEXT_NAMESPACE
