UE4: CIS fix

diff --git a/Engine/Source/Runtime/Engine/Classes/Engine/LevelStreaming.h b/Engine/Source/Runtime/Engine/Classes/Engine/LevelStreaming.h
--- a/Engine/Source/Runtime/Engine/Classes/Engine/LevelStreaming.h
+++ b/Engine/Source/Runtime/Engine/Classes/Engine/LevelStreaming.h
@@ -165,7 +165,7 @@ class ULevelStreaming : public UObject
 	}
 
 	/** Rename package name to PIE appropriate name */
-	void RenameForPIE(int PIEInstanceID);
+	ENGINE_API void RenameForPIE(int PIEInstanceID);
 
 	/**
 	 * Return whether this level should be present in memory which in turn tells the 
