UE4: SCW now launches with -buildmachine if the invoking process also launched with -buildmachine. This will prevent crash report client from being spawned on build machines.

diff --git a/Engine/Source/Runtime/Engine/Private/ShaderCompiler/ShaderCompiler.cpp b/Engine/Source/Runtime/Engine/Private/ShaderCompiler/ShaderCompiler.cpp
--- a/Engine/Source/Runtime/Engine/Private/ShaderCompiler/ShaderCompiler.cpp
+++ b/Engine/Source/Runtime/Engine/Private/ShaderCompiler/ShaderCompiler.cpp
@@ -1122,6 +1122,10 @@ uint32 FShaderCompilingManager::LaunchWorker(const FString& WorkingDirectory, ui
 	{
 		WorkerParameters += FString(TEXT(" -communicatethroughfile "));
 	}
+	if ( GIsBuildMachine )
+	{
+		WorkerParameters += FString(TEXT(" -buildmachine "));
+	}
 	WorkerParameters += FCommandLine::GetSubprocessCommandline();
 
 	// Launch the worker process
