TTP#332181
Fix incorrectly using async scene when terminating body and asyncScene is disabled

diff --git a/Engine/Source/Runtime/Engine/Private/PhysicsEngine/PhysScene.cpp b/Engine/Source/Runtime/Engine/Private/PhysicsEngine/PhysScene.cpp
--- a/Engine/Source/Runtime/Engine/Private/PhysicsEngine/PhysScene.cpp
+++ b/Engine/Source/Runtime/Engine/Private/PhysicsEngine/PhysScene.cpp
@@ -22,7 +22,7 @@ FORCEINLINE EPhysicsSceneType SceneType(const FBodyInstance * BodyInstance)
 #if WITH_PHYSX
 	//This is a helper function for dynamic actors - static actors are in both scenes
 	check(BodyInstance->GetPxRigidDynamic());
-	return BodyInstance->UseAsyncScene() ? PST_Async : PST_Sync;
+	return UPhysicsSettings::Get()->bEnableAsyncScene && BodyInstance->UseAsyncScene() ? PST_Async : PST_Sync;
 #endif
 
 	return PST_Sync;
