#UE4 Temp workaround for GL4

diff --git a/Engine/Source/Runtime/Engine/Private/GPUSort.cpp b/Engine/Source/Runtime/Engine/Private/GPUSort.cpp
--- a/Engine/Source/Runtime/Engine/Private/GPUSort.cpp
+++ b/Engine/Source/Runtime/Engine/Private/GPUSort.cpp
@@ -220,7 +220,8 @@ public:
 
 	static bool ShouldCache( EShaderPlatform Platform )
 	{
-		return IsFeatureLevelSupported(Platform, ERHIFeatureLevel::SM5);
+		//@todo-rco: Fix the shader
+		return IsFeatureLevelSupported(Platform, ERHIFeatureLevel::SM5) && !IsOpenGLPlatform(Platform);
 	}
 
 	static void ModifyCompilationEnvironment( EShaderPlatform Platform, FShaderCompilerEnvironment& OutEnvironment )
@@ -304,7 +305,8 @@ public:
 
 	static bool ShouldCache( EShaderPlatform Platform )
 	{
-		return IsFeatureLevelSupported(Platform, ERHIFeatureLevel::SM5);
+		//@todo-rco: Fix the shader
+		return IsFeatureLevelSupported(Platform, ERHIFeatureLevel::SM5) && !IsOpenGLPlatform(Platform);
 	}
 
 	static void ModifyCompilationEnvironment( EShaderPlatform Platform, FShaderCompilerEnvironment& OutEnvironment )
@@ -421,7 +423,8 @@ public:
 
 	static bool ShouldCache( EShaderPlatform Platform )
 	{
-		return IsFeatureLevelSupported(Platform, ERHIFeatureLevel::SM5);
+		//@todo-rco: Fix the shader
+		return IsFeatureLevelSupported(Platform, ERHIFeatureLevel::SM5) && !IsOpenGLPlatform(Platform);
 	}
 
 	static void ModifyCompilationEnvironment( EShaderPlatform Platform, FShaderCompilerEnvironment& OutEnvironment )
@@ -517,7 +520,8 @@ public:
 
 	static bool ShouldCache( EShaderPlatform Platform )
 	{
-		return IsFeatureLevelSupported(Platform, ERHIFeatureLevel::SM5);
+		//@todo-rco: Fix the shader
+		return IsFeatureLevelSupported(Platform, ERHIFeatureLevel::SM5) && !IsOpenGLPlatform(Platform);
 	}
 
 	static void ModifyCompilationEnvironment( EShaderPlatform Platform, FShaderCompilerEnvironment& OutEnvironment )
