Fix nonuniform scaling at runtime and bad radius of capsule

diff --git a/Engine/Source/Runtime/Engine/Private/PhysicsEngine/BodyInstance.cpp b/Engine/Source/Runtime/Engine/Private/PhysicsEngine/BodyInstance.cpp
--- a/Engine/Source/Runtime/Engine/Private/PhysicsEngine/BodyInstance.cpp
+++ b/Engine/Source/Runtime/Engine/Private/PhysicsEngine/BodyInstance.cpp
@@ -1069,8 +1069,25 @@ void FBodyInstance::Weld(FBodyInstance * TheirBody, const FTransform & RelativeT
 }
 #endif
 
+float AdjustForSmallThreshold(float NewVal, float OldVal)
+{
+	float Threshold = 0.1f;
+	float Delta = NewVal - OldVal;
+	if (Delta < 0 && FMath::Abs(NewVal) < Threshold)	//getting smaller and passed threshold so flip sign
+	{
+		return -Threshold;
+	}
+	else if (Delta > 0 && FMath::Abs(NewVal) < Threshold)	//getting bigger and passed small threshold so flip sign
+	{
+		return Threshold;
+	}
+
+	return NewVal;
+}
+
 bool FBodyInstance::UpdateBodyScale(const FVector & inScale3D)
 {
+	FVector InScale3DAdjusted = inScale3D;
 #if WITH_PHYSX
 	if (BodySetup == NULL || GetPxRigidActor() == NULL)
 	{
@@ -1090,16 +1107,36 @@ bool FBodyInstance::UpdateBodyScale(const FVector & inScale3D)
 	ensure ( !Scale3D.ContainsNaN() && !inScale3D.ContainsNaN() );
 #endif
 	FVector OldScale3D = Scale3D;
+	
+	//We don't want to go smaller than KINDA_SMALL_NUMBER
+	//But we still want to be able to cross the negative threshold to flip the scale
+	InScale3DAdjusted.X = AdjustForSmallThreshold(inScale3D.X, OldScale3D.X);
+	InScale3DAdjusted.Y = AdjustForSmallThreshold(inScale3D.Y, OldScale3D.Y);
+	InScale3DAdjusted.Z = AdjustForSmallThreshold(inScale3D.Z, OldScale3D.Z);
+	
+	if (FMath::Abs(OldScale3D.X) < 0.1f)
+	{
+		OldScale3D.X = 0.1f;
+	}
 
-	if (OldScale3D.IsNearlyZero())
+	if (FMath::Abs(OldScale3D.Y) < 0.1f)
 	{
-		// min scale
-		OldScale3D = FVector(0.1f);
+		OldScale3D.Y = 0.1f;
+	}
+	
+	if (FMath::Abs(OldScale3D.Z) < 0.1f)
+	{
+		OldScale3D.Z = 0.1f;
 	}
 
 	// Determine if applied scaling is uniform. If it isn't, only convex geometry will be copied over
-	float	NewUniScale = inScale3D.GetMin()/OldScale3D.GetMin(); // total scale includes inversing old scale and applying new scale
+	float	NewUniScale = InScale3DAdjusted.GetMin() / OldScale3D.GetMin(); // total scale includes inversing old scale and applying new scale
 	float	NewUniScaleAbs = FMath::Abs(NewUniScale);
+	FVector NewScale3D = InScale3DAdjusted;
+	NewScale3D.X /= OldScale3D.X;
+	NewScale3D.Y /= OldScale3D.Y;
+	NewScale3D.Z /= OldScale3D.Z;
+	FVector NewScale3DAbs = NewScale3D.GetAbs();
 
 	// Is the target a static actor
 	// Not used
@@ -1170,7 +1207,7 @@ bool FBodyInstance::UpdateBodyScale(const FVector & inScale3D)
 			}
 			else
 			{
-				FMessageLog("PIE").Warning(FText::Format(LOCTEXT("PhysicsInvalidScale", "Scale ''{0}'' is not valid on object '{1}'."), FText::FromString(inScale3D.ToString()), FText::FromString(GetBodyDebugName())));
+				FMessageLog("PIE").Warning(FText::Format(LOCTEXT("PhysicsInvalidScale", "Scale ''{0}'' is not valid on object '{1}'."), FText::FromString(InScale3DAdjusted.ToString()), FText::FromString(GetBodyDebugName())));
 			}
 		}
 		else if (GeomType == PxGeometryType::eBOX)
@@ -1179,10 +1216,12 @@ bool FBodyInstance::UpdateBodyScale(const FVector & inScale3D)
 			PShape->getBoxGeometry(PBoxGeom);
 			SCENE_UNLOCK_READ(PScene);
 
-			PBoxGeom.halfExtents.x *= NewUniScaleAbs;
-			PBoxGeom.halfExtents.y *= NewUniScaleAbs;
-			PBoxGeom.halfExtents.z *= NewUniScaleAbs;
-			PLocalPose.p *= NewUniScale;
+			PBoxGeom.halfExtents.x *= NewScale3DAbs.X;
+			PBoxGeom.halfExtents.y *= NewScale3DAbs.Y;
+			PBoxGeom.halfExtents.z *= NewScale3DAbs.Z;
+			PLocalPose.p.x *= NewScale3D.X;
+			PLocalPose.p.y *= NewScale3D.Y;
+			PLocalPose.p.z *= NewScale3D.Z;
 
 			if (PBoxGeom.isValid())
 			{
@@ -1193,7 +1232,7 @@ bool FBodyInstance::UpdateBodyScale(const FVector & inScale3D)
 			}
 			else
 			{
-				FMessageLog("PIE").Warning(FText::Format(LOCTEXT("PhysicsInvalidScale", "Scale ''{0}'' is not valid on object '{1}'."), FText::FromString(inScale3D.ToString()), FText::FromString(GetBodyDebugName())));
+				FMessageLog("PIE").Warning(FText::Format(LOCTEXT("PhysicsInvalidScale", "Scale ''{0}'' is not valid on object '{1}'."), FText::FromString(InScale3DAdjusted.ToString()), FText::FromString(GetBodyDebugName())));
 			}
 		}
 		else if (GeomType == PxGeometryType::eCAPSULE)
@@ -1202,9 +1241,12 @@ bool FBodyInstance::UpdateBodyScale(const FVector & inScale3D)
 			PShape->getCapsuleGeometry(PCapsuleGeom);
 			SCENE_UNLOCK_READ(PScene);
 
-			PCapsuleGeom.halfHeight *= NewUniScaleAbs;
-			PCapsuleGeom.radius *= NewUniScaleAbs;
-			PLocalPose.p *= NewUniScale;
+			PCapsuleGeom.halfHeight *= NewScale3DAbs.Z;
+			PCapsuleGeom.radius *= FMath::Max(NewScale3DAbs.X, NewScale3DAbs.Y);
+
+			PLocalPose.p.x *= NewScale3D.X;
+			PLocalPose.p.y *= NewScale3D.Y;
+			PLocalPose.p.z *= NewScale3D.Z;
 
 			if(PCapsuleGeom.isValid())
 			{
@@ -1215,7 +1257,7 @@ bool FBodyInstance::UpdateBodyScale(const FVector & inScale3D)
 			}
 			else
 			{
-				FMessageLog("PIE").Warning(FText::Format(LOCTEXT("PhysicsInvalidScale", "Scale ''{0}'' is not valid on object '{1}'."), FText::FromString(inScale3D.ToString()), FText::FromString(GetBodyDebugName())));
+				FMessageLog("PIE").Warning(FText::Format(LOCTEXT("PhysicsInvalidScale", "Scale ''{0}'' is not valid on object '{1}'."), FText::FromString(InScale3DAdjusted.ToString()), FText::FromString(GetBodyDebugName())));
 			}
 		}
 		else if (GeomType == PxGeometryType::eCONVEXMESH)
@@ -1235,7 +1277,7 @@ bool FBodyInstance::UpdateBodyScale(const FVector & inScale3D)
 				if(ConvexElem->ConvexMesh == PConvexGeom.convexMesh)
 				{
 					// Please note that this one we don't inverse old scale, but just set new one
-					FVector NewScale3D = inScale3D; 
+					FVector NewScale3D = InScale3DAdjusted;
 					FVector Scale3DAbs(FMath::Abs(NewScale3D.X), FMath::Abs(NewScale3D.Y), FMath::Abs(NewScale3D.Z)); // magnitude of scale (sign removed)
 
 					PxTransform PNewLocalPose;
@@ -1257,7 +1299,7 @@ bool FBodyInstance::UpdateBodyScale(const FVector & inScale3D)
 					}
 					else
 					{
-						FMessageLog("PIE").Warning(FText::Format(LOCTEXT("PhysicsInvalidScale", "Scale ''{0}'' is not valid on object '{1}'."), FText::FromString(inScale3D.ToString()), FText::FromString(GetBodyDebugName())));
+						FMessageLog("PIE").Warning(FText::Format(LOCTEXT("PhysicsInvalidScale", "Scale ''{0}'' is not valid on object '{1}'."), FText::FromString(InScale3DAdjusted.ToString()), FText::FromString(GetBodyDebugName())));
 					}
 					break;
 				}
@@ -1273,7 +1315,7 @@ bool FBodyInstance::UpdateBodyScale(const FVector & inScale3D)
 			if(BodySetup->TriMesh != NULL || BodySetup->TriMeshNegX != NULL)
 			{
 				// Please note that this one we don't inverse old scale, but just set new one
-				FVector NewScale3D = inScale3D; 
+				FVector NewScale3D = InScale3DAdjusted;
 				FVector Scale3DAbs(FMath::Abs(NewScale3D.X), FMath::Abs(NewScale3D.Y), FMath::Abs(NewScale3D.Z)); // magnitude of scale (sign removed)
 
 				PxTransform PNewLocalPose;
@@ -1300,7 +1342,7 @@ bool FBodyInstance::UpdateBodyScale(const FVector & inScale3D)
 					}
 					else
 					{
-						FMessageLog("PIE").Warning(FText::Format(LOCTEXT("PhysicsInvalidScale", "Scale ''{0}'' is not valid on object '{1}'."), FText::FromString(inScale3D.ToString()), FText::FromString(GetBodyDebugName())));
+						FMessageLog("PIE").Warning(FText::Format(LOCTEXT("PhysicsInvalidScale", "Scale ''{0}'' is not valid on object '{1}'."), FText::FromString(InScale3DAdjusted.ToString()), FText::FromString(GetBodyDebugName())));
 					}
 				}
 			}
@@ -1315,7 +1357,7 @@ bool FBodyInstance::UpdateBodyScale(const FVector & inScale3D)
 	// if success, overwrite old Scale3D, otherwise, just don't do it. It will have invalid scale next time
 	if ( bSuccess )
 	{
-		Scale3D = inScale3D;
+		Scale3D = InScale3DAdjusted;
 
 		// update mass if required
 		if (bUpdateMassWhenScaleChanges)
diff --git a/Engine/Source/Runtime/Engine/Private/PhysicsEngine/BodySetup.cpp b/Engine/Source/Runtime/Engine/Private/PhysicsEngine/BodySetup.cpp
--- a/Engine/Source/Runtime/Engine/Private/PhysicsEngine/BodySetup.cpp
+++ b/Engine/Source/Runtime/Engine/Private/PhysicsEngine/BodySetup.cpp
@@ -379,8 +379,8 @@ void UBodySetup::AddShapesToRigidActor(PxRigidActor* PDestActor, FVector& Scale3
 				HalfHeight = FMath::Max<float>(HalfHeight - Radius, 1.f);
 
 				PxCapsuleGeometry PCapsuleGeom;
-				PCapsuleGeom.halfHeight = (HalfHeight);
-				PCapsuleGeom.radius = (Radius * ScaleRadius);
+				PCapsuleGeom.halfHeight = HalfHeight;
+				PCapsuleGeom.radius = Radius;
 
 				if (ensure(PCapsuleGeom.isValid()))
 				{
