- Hide streaming level internal properties from details editor
- Removed unused bIsFullyStatic property

diff --git a/Engine/Source/Runtime/Engine/Classes/Engine/LevelStreaming.h b/Engine/Source/Runtime/Engine/Classes/Engine/LevelStreaming.h
--- a/Engine/Source/Runtime/Engine/Classes/Engine/LevelStreaming.h
+++ b/Engine/Source/Runtime/Engine/Classes/Engine/LevelStreaming.h
@@ -99,17 +99,13 @@ class ULevelStreaming : public UObject
 	uint32 bFailedToLoad:1;
 
 	/** Whether this level should be visible in the Editor																		*/
-	UPROPERTY(EditAnywhere, Category=LevelStreaming)
+	UPROPERTY()
 	uint32 bShouldBeVisibleInEditor:1;
 
 	/** Whether this level is locked; that is, its actors are read-only.														*/
-	UPROPERTY(EditAnywhere, Category=LevelStreaming)
+	UPROPERTY()
 	uint32 bLocked:1;
 
-	/** Whether this level is fully static - if it is, then assumptions can be made about it, ie it doesn't need to be reloaded since nothing could have changed */
-	UPROPERTY(EditAnywhere, Category=LevelStreaming)
-	uint32 bIsFullyStatic:1;
-
 	/** Whether the level should be loaded																						*/
 	UPROPERTY(Category=LevelStreaming, BlueprintReadWrite)
 	uint32 bShouldBeLoaded:1;
