- Fix typo in UParticleModuleLocationSkelVertSurface::Spawn

diff --git a/Engine/Source/Runtime/Engine/Private/Particles/ParticleModules_Location.cpp b/Engine/Source/Runtime/Engine/Private/Particles/ParticleModules_Location.cpp
--- a/Engine/Source/Runtime/Engine/Private/Particles/ParticleModules_Location.cpp
+++ b/Engine/Source/Runtime/Engine/Private/Particles/ParticleModules_Location.cpp
@@ -2117,7 +2117,7 @@ void UParticleModuleLocationSkelVertSurface::Spawn(FParticleEmitterInstance* Own
 			BoneIndex1 = BoneIndex2 = BoneIndex3 = INDEX_NONE;
 			if(!VertInfluencedByActiveBone(Owner, SourceComponent, VertIndex[0], &BoneIndex1) &&
 			   !VertInfluencedByActiveBone(Owner, SourceComponent, VertIndex[1], &BoneIndex2) && 
-			   !VertInfluencedByActiveBone(Owner, SourceComponent, VertIndex[2]) &BoneIndex3)
+			   !VertInfluencedByActiveBone(Owner, SourceComponent, VertIndex[2], &BoneIndex3))
 			{
 				SPAWN_INIT
 				{
