Improve default PlayerStart selection

diff --git a/Engine/Source/Runtime/Engine/Private/GameMode.cpp b/Engine/Source/Runtime/Engine/Private/GameMode.cpp
--- a/Engine/Source/Runtime/Engine/Private/GameMode.cpp
+++ b/Engine/Source/Runtime/Engine/Private/GameMode.cpp
@@ -331,10 +331,10 @@ AActor* AGameMode::FindPlayerStart( AController* Player, const FString& Incoming
 	// if incoming start is specified, then just use it
 	if( !IncomingName.IsEmpty() )
 	{
-		for( FActorIterator It(World); It; ++It )
+		for (TActorIterator<APlayerStart> It(World); It; ++It)
 		{
-			APlayerStart* Start = Cast<APlayerStart>(*It);
-			if ( Start && Start->PlayerStartTag == FName(*IncomingName) )
+			APlayerStart* Start = *It;
+			if (Start && Start->PlayerStartTag == FName(*IncomingName))
 			{
 				return Start;
 			}
@@ -344,7 +344,6 @@ AActor* AGameMode::FindPlayerStart( AController* Player, const FString& Incoming
 	// always pick StartSpot at start of match
 	if ( ShouldSpawnAtStartSpot(Player) )
 	{
-		// NULL CHECK MISSING!!!!!
 		return Player->StartSpot.Get();
 	}
 
@@ -1397,21 +1396,44 @@ void AGameMode::RemovePlayerStart(APlayerStart* RemovedPlayerStart)
 
 AActor* AGameMode::ChoosePlayerStart( AController* Player )
 {
-	// Find first player start
+	// Choose a player start
 	APlayerStart* FoundPlayerStart = NULL;
+	APawn* PawnToFit = DefaultPawnClass->GetDefaultObject<APawn>();
+	TArray<APlayerStart*> UnOccupiedStartPoints;
+	TArray<APlayerStart*> OccupiedStartPoints;
 	for (int32 PlayerStartIndex = 0; PlayerStartIndex < PlayerStarts.Num(); ++PlayerStartIndex)
 	{
 		APlayerStart* PlayerStart = PlayerStarts[PlayerStartIndex];
 
-		if( Cast<APlayerStartPIE>( PlayerStart ) != NULL )
+		if (Cast<APlayerStartPIE>( PlayerStart ) != NULL )
 		{
 			// Always prefer the first "Play from Here" PlayerStart, if we find one while in PIE mode
 			FoundPlayerStart = PlayerStart;
 			break;
 		}
-		else if( PlayerStart != NULL && FoundPlayerStart == NULL )
+		else if (PlayerStart != NULL)
 		{
-			FoundPlayerStart = PlayerStart;
+			FVector ActorLocation = PlayerStart->GetActorLocation();
+			const FRotator ActorRotation = PlayerStart->GetActorRotation();
+			if (!GetWorld()->EncroachingBlockingGeometry(PawnToFit, ActorLocation, ActorRotation))
+			{
+				UnOccupiedStartPoints.Add(PlayerStart);
+			}
+			else if (GetWorld()->FindTeleportSpot(PawnToFit, ActorLocation, ActorRotation))
+			{
+				OccupiedStartPoints.Add(PlayerStart);
+			}
+		}
+	}
+	if (FoundPlayerStart == NULL)
+	{
+		if (UnOccupiedStartPoints.Num() > 0)
+		{
+			FoundPlayerStart = UnOccupiedStartPoints[FMath::RandRange(0, UnOccupiedStartPoints.Num() - 1)];
+		}
+		else if (OccupiedStartPoints.Num() > 0)
+		{
+			FoundPlayerStart = OccupiedStartPoints[FMath::RandRange(0, OccupiedStartPoints.Num() - 1)];
 		}
 	}
 	return FoundPlayerStart;
diff --git a/Engine/Source/Runtime/Engine/Private/LevelActor.cpp b/Engine/Source/Runtime/Engine/Private/LevelActor.cpp
--- a/Engine/Source/Runtime/Engine/Private/LevelActor.cpp
+++ b/Engine/Source/Runtime/Engine/Private/LevelActor.cpp
@@ -633,11 +633,18 @@ bool UWorld::FindTeleportSpot(AActor* TestActor, FVector& TestLocation, FRotator
 	// now try just XY
 	if ( (Adjust.X != 0.f) || (Adjust.Y != 0.f) )
 	{
-		TestLocation.X += Adjust.X;
-		TestLocation.Y += Adjust.Y;
-		if( !EncroachingBlockingGeometry(TestActor, TestLocation, TestRotation, &Adjust) )
+		for (int i = 0; i < 8; ++i)
 		{
-			return true;
+			const FVector OriginalTestLocation = TestLocation;
+
+			TestLocation.X += (i < 4 ? Adjust.X : Adjust.Y) * (i % 2 == 0 ? 1 : -1);
+			TestLocation.Y += (i < 4 ? Adjust.Y : Adjust.X) * (i % 4 < 2 ? 1 : -1);
+			if (!EncroachingBlockingGeometry(TestActor, TestLocation, TestRotation, &Adjust))
+			{
+				return true;
+			}
+
+			TestLocation = OriginalTestLocation;
 		}
 	}
 
