Fixed NULL access of default pawn class - Not all games have a pawn.

diff --git a/Engine/Source/Runtime/Engine/Private/GameMode.cpp b/Engine/Source/Runtime/Engine/Private/GameMode.cpp
--- a/Engine/Source/Runtime/Engine/Private/GameMode.cpp
+++ b/Engine/Source/Runtime/Engine/Private/GameMode.cpp
@@ -1398,7 +1398,7 @@ AActor* AGameMode::ChoosePlayerStart( AController* Player )
 {
 	// Choose a player start
 	APlayerStart* FoundPlayerStart = NULL;
-	APawn* PawnToFit = DefaultPawnClass->GetDefaultObject<APawn>();
+	APawn* PawnToFit = DefaultPawnClass ? DefaultPawnClass->GetDefaultObject<APawn>() : NULL;
 	TArray<APlayerStart*> UnOccupiedStartPoints;
 	TArray<APlayerStart*> OccupiedStartPoints;
 	for (int32 PlayerStartIndex = 0; PlayerStartIndex < PlayerStarts.Num(); ++PlayerStartIndex)
