- Remove comment referring to UnrealScript from UCanvas::Update

diff --git a/Engine/Source/Runtime/Engine/Private/UserInterface/Canvas.cpp b/Engine/Source/Runtime/Engine/Private/UserInterface/Canvas.cpp
--- a/Engine/Source/Runtime/Engine/Private/UserInterface/Canvas.cpp
+++ b/Engine/Source/Runtime/Engine/Private/UserInterface/Canvas.cpp
@@ -1120,7 +1120,7 @@ void UCanvas::UpdateAllCanvasSafeZoneData()
 
 void UCanvas::Update()
 {
-	// Call UnrealScript to reset.
+	// Reset canvas params.
 	Reset();
 
 	// Copy size parameters from viewport.
