break vehicle tuning out of tick function - needed for substepping

diff --git a/Engine/Source/Runtime/Engine/Classes/Vehicles/WheeledVehicleMovementComponent.h b/Engine/Source/Runtime/Engine/Classes/Vehicles/WheeledVehicleMovementComponent.h
--- a/Engine/Source/Runtime/Engine/Classes/Vehicles/WheeledVehicleMovementComponent.h
+++ b/Engine/Source/Runtime/Engine/Classes/Vehicles/WheeledVehicleMovementComponent.h
@@ -265,6 +265,9 @@ class ENGINE_API UWheeledVehicleMovementComponent : public UPawnMovementComponen
 	/** Tick this vehicle sim right before input is sent to the vehicle system  */
 	virtual void TickVehicle( float DeltaTime );
 
+	/** Updates the vehicle tuning data. Used when designer tweaks values while game is running */
+	virtual void UpdateTuning();
+
 	/** Used to create any physics engine information for this component */
 	virtual void CreatePhysicsState() OVERRIDE;
 
diff --git a/Engine/Source/Runtime/Engine/Private/PhysicsEngine/PhysScene.cpp b/Engine/Source/Runtime/Engine/Private/PhysicsEngine/PhysScene.cpp
--- a/Engine/Source/Runtime/Engine/Private/PhysicsEngine/PhysScene.cpp
+++ b/Engine/Source/Runtime/Engine/Private/PhysicsEngine/PhysScene.cpp
@@ -399,11 +399,12 @@ void FPhysScene::TickPhysScene(uint32 SceneType, FGraphEventRef& InOutCompletion
 
 #if WITH_PHYSX
 
+	if (VehicleManager && SceneType == PST_Sync)
+	{
+		VehicleManager->UpdateTuning();
 #if WITH_SUBSTEPPING
-	if (IsSubstepping() == false)
+		if (IsSubstepping() == false)
 #endif
-	{
-		if (VehicleManager && SceneType == PST_Sync)
 		{
 			VehicleManager->Update(AveragedFrameTime[SceneType]);
 		}
diff --git a/Engine/Source/Runtime/Engine/Private/Vehicles/PhysXVehicleManager.cpp b/Engine/Source/Runtime/Engine/Private/Vehicles/PhysXVehicleManager.cpp
--- a/Engine/Source/Runtime/Engine/Private/Vehicles/PhysXVehicleManager.cpp
+++ b/Engine/Source/Runtime/Engine/Private/Vehicles/PhysXVehicleManager.cpp
@@ -291,6 +291,14 @@ void FPhysXVehicleManager::Update( float DeltaTime )
 #endif //PX_DEBUG_VEHICLE_ON
 }
 
+void FPhysXVehicleManager::UpdateTuning()
+{
+	for (int32 i = 0; i < Vehicles.Num(); ++i)
+	{
+		Vehicles[i]->UpdateTuning();
+	}
+}
+
 void FPhysXVehicleManager::UpdateVehicles( float DeltaTime )
 {
 	PxVehicleUpdates( DeltaTime, GetSceneGravity(), *SurfaceTirePairs, PVehicles.Num(), PVehicles.GetTypedData(), PVehiclesWheelsStates.GetTypedData());
diff --git a/Engine/Source/Runtime/Engine/Private/Vehicles/PhysXVehicleManager.h b/Engine/Source/Runtime/Engine/Private/Vehicles/PhysXVehicleManager.h
--- a/Engine/Source/Runtime/Engine/Private/Vehicles/PhysXVehicleManager.h
+++ b/Engine/Source/Runtime/Engine/Private/Vehicles/PhysXVehicleManager.h
@@ -56,6 +56,10 @@ public:
 	 * Update vehicle data before the scene simulates
 	 */
 	void Update( float DeltaTime );
+	
+	/**
+	 * Update vehicle tuning. Used when designer tweaks values while the game is running */
+	void UpdateTuning();
 
 private:
 
diff --git a/Engine/Source/Runtime/Engine/Private/Vehicles/WheeledVehicleMovementComponent.cpp b/Engine/Source/Runtime/Engine/Private/Vehicles/WheeledVehicleMovementComponent.cpp
--- a/Engine/Source/Runtime/Engine/Private/Vehicles/WheeledVehicleMovementComponent.cpp
+++ b/Engine/Source/Runtime/Engine/Private/Vehicles/WheeledVehicleMovementComponent.cpp
@@ -632,11 +632,6 @@ void UWheeledVehicleMovementComponent::DestroyWheels()
 
 void UWheeledVehicleMovementComponent::TickVehicle( float DeltaTime )
 {
-	if (VehicleSetupTag != FPhysXVehicleManager::VehicleSetupTag )
-	{
-		RecreatePhysicsState();
-	}
-
 	// movement updates and replication
 	if (PVehicle && UpdatedComponent)
 	{
@@ -655,6 +650,14 @@ void UWheeledVehicleMovementComponent::TickVehicle( float DeltaTime )
 	}
 }
 
+void UWheeledVehicleMovementComponent::UpdateTuning()
+{
+	if (VehicleSetupTag != FPhysXVehicleManager::VehicleSetupTag)
+	{
+		RecreatePhysicsState();
+	}
+}
+
 void UWheeledVehicleMovementComponent::UpdateSimulation( float DeltaTime )
 {
 }
