#ue4

- Minor tweak to comment on CanJump()

diff --git a/Engine/Source/Runtime/Engine/Classes/GameFramework/Character.h b/Engine/Source/Runtime/Engine/Classes/GameFramework/Character.h
--- a/Engine/Source/Runtime/Engine/Classes/GameFramework/Character.h
+++ b/Engine/Source/Runtime/Engine/Classes/GameFramework/Character.h
@@ -264,8 +264,8 @@ public:
 	virtual void Jump();
 
 	/**
-	 * Default implementation returns true if the character is not crouching, has a 
-	 * valid CharacterMovementComponent and CanEverJump() returns true.
+	 * Default implementation returns true if the character is on the ground and not crouching,
+	 * has a valid CharacterMovementComponent and CanEverJump() returns true.
 	 * @Return Whether the character can jump in the current state. 
 	 */
 	UFUNCTION(BlueprintNativeEvent, Category="Pawn|Character")
