More vehicle fixes:
- change tire sub-stepping back to default NVIDIA values
- fix gears never going past 1st
- fix incorrect throttle when player goes in opposite direction

diff --git a/Engine/Source/Runtime/Engine/Classes/Vehicles/WheeledVehicleMovementComponent4W.h b/Engine/Source/Runtime/Engine/Classes/Vehicles/WheeledVehicleMovementComponent4W.h
--- a/Engine/Source/Runtime/Engine/Classes/Vehicles/WheeledVehicleMovementComponent4W.h
+++ b/Engine/Source/Runtime/Engine/Classes/Vehicles/WheeledVehicleMovementComponent4W.h
@@ -192,9 +192,6 @@ protected:
 
 #endif // WITH_PHYSX
 
-	/** Compute throttle input */
-	virtual float CalcThrottleInput() OVERRIDE;
-
 	/** update simulation data: engine */
 	void UpdateEngineSetup(const FVehicleEngineData& NewEngineSetup);
 
diff --git a/Engine/Source/Runtime/Engine/Private/Vehicles/WheeledVehicleMovementComponent.cpp b/Engine/Source/Runtime/Engine/Private/Vehicles/WheeledVehicleMovementComponent.cpp
--- a/Engine/Source/Runtime/Engine/Private/Vehicles/WheeledVehicleMovementComponent.cpp
+++ b/Engine/Source/Runtime/Engine/Private/Vehicles/WheeledVehicleMovementComponent.cpp
@@ -253,7 +253,7 @@ void UWheeledVehicleMovementComponent::SetupWheels( PxVehicleWheelsSimData* PWhe
 	const PxReal LengthScale = 100.f; // Convert default from m to cm
 
 	// Control substepping
-	PWheelsSimData->setSubStepCount( 5.f * LengthScale, 5, 2 ); // forcing constant substeps until all the kinks are ironed out	
+	PWheelsSimData->setSubStepCount( 5.f * LengthScale, 3, 1 );
 	PWheelsSimData->setMinLongSlipDenominator(4.f * LengthScale); 
 
 	// Prealloc data for the sprung masses
@@ -671,15 +671,18 @@ void UWheeledVehicleMovementComponent::UpdateState( float DeltaTime )
 	if (MyOwner && MyOwner->IsLocallyControlled())
 	{
 		//Manual shifting between reverse and first gear
-		if (RawThrottleInput < 0.f && GetCurrentGear() >= 0)
+		if (FMath::Abs(GetForwardSpeed()) < WrongDirectionThreshold)	//we only shift between reverse and first if the car is slow enough. This isn't 100% correct since we really only care about engine speed, but good enough
 		{
-			SetTargetGear(-1, true);
-		}
-		else if (RawThrottleInput > 0.f && GetCurrentGear() <= 0)
-		{
-			SetTargetGear(1, true);
+			if (RawThrottleInput < 0.f && GetCurrentGear() >= 0 && GetTargetGear() >= 0)
+			{
+				SetTargetGear(-1, true);
+			}
+			else if (RawThrottleInput > 0.f && GetCurrentGear() <= 0 && GetTargetGear() <= 0)
+			{
+				SetTargetGear(1, true);
+			}
 		}
-
+		
 		SteeringInput = SteeringInputRate.InterpInputValue( DeltaTime, SteeringInput, CalcSteeringInput() );
 		ThrottleInput = ThrottleInputRate.InterpInputValue( DeltaTime, ThrottleInput, CalcThrottleInput() );
 		BrakeInput = BrakeInputRate.InterpInputValue( DeltaTime, BrakeInput, CalcBrakeInput() );
@@ -779,6 +782,12 @@ float UWheeledVehicleMovementComponent::CalcHandbrakeInput()
 
 float UWheeledVehicleMovementComponent::CalcThrottleInput()
 {
+	//If the user is changing direction we should really be braking first and not applying any gas, so wait until they've changed gears
+	if (RawThrottleInput > 0.f && GetTargetGear() < 0 || RawThrottleInput < 0.f && GetTargetGear() > 0) 
+	{
+		return 0.f;
+	}
+
 	return FMath::Abs(RawThrottleInput);
 }
 
diff --git a/Engine/Source/Runtime/Engine/Private/Vehicles/WheeledVehicleMovementComponent4W.cpp b/Engine/Source/Runtime/Engine/Private/Vehicles/WheeledVehicleMovementComponent4W.cpp
--- a/Engine/Source/Runtime/Engine/Private/Vehicles/WheeledVehicleMovementComponent4W.cpp
+++ b/Engine/Source/Runtime/Engine/Private/Vehicles/WheeledVehicleMovementComponent4W.cpp
@@ -262,11 +262,6 @@ void UWheeledVehicleMovementComponent4W::UpdateSimulation(float DeltaTime)
 
 #endif // WITH_PHYSX
 
-float UWheeledVehicleMovementComponent4W::CalcThrottleInput()
-{
-	return FMath::Abs(RawThrottleInput);
-}
-
 void UWheeledVehicleMovementComponent4W::UpdateEngineSetup(const FVehicleEngineData& NewEngineSetup)
 {
 #if WITH_PHYSX
