Make rendering component of NavigationData transient, duplicatetransient

diff --git a/Engine/Source/Runtime/Engine/Classes/AI/Navigation/NavigationData.h b/Engine/Source/Runtime/Engine/Classes/AI/Navigation/NavigationData.h
--- a/Engine/Source/Runtime/Engine/Classes/AI/Navigation/NavigationData.h
+++ b/Engine/Source/Runtime/Engine/Classes/AI/Navigation/NavigationData.h
@@ -27,7 +27,7 @@ class ENGINE_API ANavigationData : public AActor
 {
 	GENERATED_UCLASS_BODY()
 
-	UPROPERTY()
+	UPROPERTY(transient, duplicatetransient)
 	class UPrimitiveComponent* RenderingComp;
 
 	UPROPERTY()
