Mark PoseableMeshComponent ENGINE_API

diff --git a/Engine/Source/Runtime/Engine/Classes/Components/PoseableMeshComponent.h b/Engine/Source/Runtime/Engine/Classes/Components/PoseableMeshComponent.h
--- a/Engine/Source/Runtime/Engine/Classes/Components/PoseableMeshComponent.h
+++ b/Engine/Source/Runtime/Engine/Classes/Components/PoseableMeshComponent.h
@@ -23,8 +23,8 @@ namespace EBoneSpaces
 /**
  *	UPoseableMeshComponent that allows bone transforms to be driven by blueprint.
  */
-UCLASS(HeaderGroup=Component, ClassGroup=Rendering, hidecategories=(Object,Physics), config=Engine, editinlinenew, meta=(BlueprintSpawnableComponent),MinimalAPI)
-class UPoseableMeshComponent : public USkinnedMeshComponent
+UCLASS(HeaderGroup=Component, ClassGroup=Rendering, hidecategories=(Object,Physics), config=Engine, editinlinenew, meta=(BlueprintSpawnableComponent))
+class ENGINE_API UPoseableMeshComponent : public USkinnedMeshComponent
 {
 	GENERATED_UCLASS_BODY()
 
