Fix build

diff --git a/Engine/Source/Runtime/Engine/Private/Vehicles/WheeledVehicleMovementComponent.cpp b/Engine/Source/Runtime/Engine/Private/Vehicles/WheeledVehicleMovementComponent.cpp
--- a/Engine/Source/Runtime/Engine/Private/Vehicles/WheeledVehicleMovementComponent.cpp
+++ b/Engine/Source/Runtime/Engine/Private/Vehicles/WheeledVehicleMovementComponent.cpp
@@ -783,7 +783,7 @@ float UWheeledVehicleMovementComponent::CalcHandbrakeInput()
 float UWheeledVehicleMovementComponent::CalcThrottleInput()
 {
 	//If the user is changing direction we should really be braking first and not applying any gas, so wait until they've changed gears
-	if (RawThrottleInput > 0.f && GetTargetGear() < 0 || RawThrottleInput < 0.f && GetTargetGear() > 0) 
+	if ( (RawThrottleInput > 0.f && GetTargetGear() < 0) || (RawThrottleInput < 0.f && GetTargetGear() > 0)) 
 	{
 		return 0.f;
 	}
