#UE4
Missing Super::Tick() call in AGameMode::Tick()

diff --git a/Engine/Source/Runtime/Engine/Private/GameMode.cpp b/Engine/Source/Runtime/Engine/Private/GameMode.cpp
--- a/Engine/Source/Runtime/Engine/Private/GameMode.cpp
+++ b/Engine/Source/Runtime/Engine/Private/GameMode.cpp
@@ -680,6 +680,8 @@ void AGameMode::SetMatchState(FName NewState)
 
 void AGameMode::Tick(float DeltaSeconds)
 {
+	Super::Tick(DeltaSeconds);
+
 	if (GetMatchState() == MatchState::WaitingToStart)
 	{
 		// Check to see if we should start the match
