#UE4 - fixed PlayerController returning incorrect IsLocalPlayerController() on dedicated servers
- only during a brief time between Login() and PostLogin() when PlayerController was created but Role/RemoteRole wasn't set yet
- bigger change forthcoming to Login() to narrow gap above

diff --git a/Engine/Source/Runtime/Engine/Private/PlayerController.cpp b/Engine/Source/Runtime/Engine/Private/PlayerController.cpp
--- a/Engine/Source/Runtime/Engine/Private/PlayerController.cpp
+++ b/Engine/Source/Runtime/Engine/Private/PlayerController.cpp
@@ -75,12 +75,18 @@ UNetConnection* APlayerController::GetNetConnection()
 
 bool APlayerController::IsLocalController() const
 {
+	ENetMode NetMode = GetNetMode();
+	if (NetMode == NM_DedicatedServer)
+	{
+		return false;
+	}
+
 	if (Super::IsLocalController())
 	{
 		return true;
 	}
 
-	if (GetNetMode() == NM_Client)
+	if (NetMode == NM_Client)
 	{
 		// Clients only receive their own PC. We are not ROLE_AutonomousProxy until after PostInitializeComponents so we can't check that.
 		return true;
