Temporarily reverting ensure that was causing hitches and triggering for Orion (fallout from CL 2044506).

diff --git a/Engine/Source/Runtime/Engine/Private/EdGraph/EdGraph.cpp b/Engine/Source/Runtime/Engine/Private/EdGraph/EdGraph.cpp
--- a/Engine/Source/Runtime/Engine/Private/EdGraph/EdGraph.cpp
+++ b/Engine/Source/Runtime/Engine/Private/EdGraph/EdGraph.cpp
@@ -635,10 +635,11 @@ FText UEdGraphNode::GetNodeSearchTitle() const
 
 	FText LocalizedTitle = GetNodeTitle(ENodeTitleType::ListView);
 	FText NativeTitle = FText::FromString(GetNodeNativeTitle(ENodeTitleType::ListView));
-	if(FInternationalization::Get().GetCurrentCulture()->GetName() == TEXT("en_US"))
-	{
-		ensure( LocalizedTitle.CompareTo(NativeTitle) == 0 );
-	}
+//	@TODO: vet the performance implications of this
+// 	if(FInternationalization::Get().GetCurrentCulture()->GetName() == TEXT("en_US"))
+// 	{
+// 		ensure( LocalizedTitle.CompareTo(NativeTitle) == 0 );
+// 	}
 
 	FFormatNamedArguments Args;
 	Args.Add(TEXT("LocalizedTitle"), LocalizedTitle);
