Stopped editor from caching inline shaders if not needed.

diff --git a/Engine/Source/Runtime/Engine/Private/TextureDerivedData.cpp b/Engine/Source/Runtime/Engine/Private/TextureDerivedData.cpp
--- a/Engine/Source/Runtime/Engine/Private/TextureDerivedData.cpp
+++ b/Engine/Source/Runtime/Engine/Private/TextureDerivedData.cpp
@@ -1383,20 +1383,6 @@ void UTexture::UpdateCachedLODBias()
 void UTexture::BeginCachePlatformData()
 {
 	CachePlatformData(true);
-
-
-	ITargetPlatformManagerModule* TPM = GetTargetPlatformManager();
-	
-	if (TPM)
-	{
-		const TArray<ITargetPlatform*>& Platforms = TPM->GetActiveTargetPlatforms();
-		for (int32 PlatformIndex = 0; PlatformIndex < Platforms.Num(); PlatformIndex++)
-		{
-			BeginCacheForCookedPlatformData( Platforms[PlatformIndex]);
-		}
-	}
-
-
 }
 
 void UTexture::BeginCacheForCookedPlatformData( const ITargetPlatform *TargetPlatform )
