UE4: UBrush objects corresponding to the default brush are now saved with the same outer as the default brush, instead of top-level objects in world packages

diff --git a/Engine/Source/Runtime/Engine/Private/World.cpp b/Engine/Source/Runtime/Engine/Private/World.cpp
--- a/Engine/Source/Runtime/Engine/Private/World.cpp
+++ b/Engine/Source/Runtime/Engine/Private/World.cpp
@@ -706,7 +706,7 @@ void UWorld::InitWorld(const InitializationValues IVS)
 			// Spawn the default brush.
 			DefaultBrush = SpawnBrush();
 			check(DefaultBrush->BrushComponent);
-			DefaultBrush->Brush = new( GetOuter(), TEXT("Brush") )UModel( FPostConstructInitializeProperties(), DefaultBrush, 1 );
+			DefaultBrush->Brush = new( DefaultBrush->GetOuter(), TEXT("Brush") )UModel( FPostConstructInitializeProperties(), DefaultBrush, 1 );
 			DefaultBrush->BrushComponent->Brush = DefaultBrush->Brush;
 			DefaultBrush->SetNotForClientOrServer();
 			DefaultBrush->Brush->SetFlags( RF_Transactional );
@@ -734,6 +734,12 @@ void UWorld::InitWorld(const InitializationValues IVS)
 		{
 			UModel* Model = DefaultBrush->Brush;
 
+			// Ensure the DefaultBrush's model has the same outer as the default brush itself. Older packages erroneously stored this object as a top-level package
+			if (Model->GetOuter() != DefaultBrush->GetOuter())
+			{
+				Model->Rename(NULL, DefaultBrush->GetOuter(), REN_DoNotDirty | REN_DontCreateRedirectors | REN_ForceNoResetLoaders | REN_NonTransactional);
+			}
+
 			const FLightmassPrimitiveSettings DefaultSettings;
 
 			for (int32 i = 0; i < Model->LightmassSettings.Num(); ++i)
