#UE4
- gave NavigationComponent means to control default navigation query extent it's using

diff --git a/Engine/Source/Runtime/Engine/Classes/AI/Navigation/NavigationComponent.h b/Engine/Source/Runtime/Engine/Classes/AI/Navigation/NavigationComponent.h
--- a/Engine/Source/Runtime/Engine/Classes/AI/Navigation/NavigationComponent.h
+++ b/Engine/Source/Runtime/Engine/Classes/AI/Navigation/NavigationComponent.h
@@ -125,6 +125,8 @@ class ENGINE_API UNavigationComponent : public UActorComponent, public INavigati
 		}
 	}
 
+	FORCEINLINE FVector GetQueryExtent() const { return NavigationQueryExtent; }
+
 	//----------------------------------------------------------------------//
 	// misc
 	//----------------------------------------------------------------------//
@@ -172,6 +174,9 @@ protected:
 	UPROPERTY(transient)
 	class UPathFollowingComponent* PathFollowComp;
 
+	/** default query extent to use. It's a sum of navigation data's default query extent and nav agent's size */
+	mutable FVector NavigationQueryExtent;
+
 	/** Shared path - there can be multiple users of one path instance. No not modify it, unless you know what you're doing */
 	FNavPathSharedPtr Path;
 
@@ -242,6 +247,9 @@ protected:
 	/** asks NavigationSystem for NavigationData appropriate for self and caches it in MyNavData */
 	class ANavigationData* PickNavData() const;
 
+	/** called as part of PickNavData to update NavigationQueryExtent */
+	virtual void CacheNavQueryExtent() const;
+
 	/** Generates simple path in straight line without using navigation data. */
 	bool GenerateSimplePath(const FVector& GoalLocation);
 
diff --git a/Engine/Source/Runtime/Engine/Private/AI/Navigation/NavigationComponent.cpp b/Engine/Source/Runtime/Engine/Private/AI/Navigation/NavigationComponent.cpp
--- a/Engine/Source/Runtime/Engine/Private/AI/Navigation/NavigationComponent.cpp
+++ b/Engine/Source/Runtime/Engine/Private/AI/Navigation/NavigationComponent.cpp
@@ -143,7 +143,7 @@ FNavLocation UNavigationComponent::ProjectPointToNavigation(const FVector& Locat
 	
 	if (GetNavData() != NULL)
 	{
-		if (!MyNavData->ProjectPoint(Location, OutLocation, MyNavData->GetDefaultQueryExtent()))
+		if (!MyNavData->ProjectPoint(Location, OutLocation, GetQueryExtent()))
 		{
 			OutLocation = FNavLocation();
 		}
@@ -154,7 +154,7 @@ FNavLocation UNavigationComponent::ProjectPointToNavigation(const FVector& Locat
 
 bool UNavigationComponent::ProjectPointToNavigation(const FVector& Location, FNavLocation& ProjectedLocation) const
 {
-	return GetNavData() != NULL && MyNavData->ProjectPoint(Location, ProjectedLocation, MyNavData->GetDefaultQueryExtent());
+	return GetNavData() != NULL && MyNavData->ProjectPoint(Location, ProjectedLocation, GetQueryExtent());
 }
 
 bool UNavigationComponent::ProjectPointToNavigation(const FVector& Location, FVector& ProjectedLocation) const
@@ -204,7 +204,7 @@ bool UNavigationComponent::AsyncGeneratePath(const FVector& FromLocation, const
 			UE_VLOG_LOCATION(GetOwner(), GoalLocation, 30, FColor::Red, TEXT("Invalid pathing GoalLocation"));
 			if (MyNavData != NULL)
 			{
-				UE_VLOG_BOX(GetOwner(), FBox(GoalLocation-MyNavData->GetDefaultQueryExtent(),GoalLocation+MyNavData->GetDefaultQueryExtent()), FColor::Red, TEXT_EMPTY);
+				UE_VLOG_BOX(GetOwner(), FBox(GoalLocation - GetQueryExtent(), GoalLocation + GetQueryExtent()), FColor::Red, TEXT_EMPTY);
 			}
 		}
 	}
@@ -326,8 +326,7 @@ bool UNavigationComponent::GeneratePathTo(const FVector& GoalLocation, TSharedPt
 
 #if ENABLE_VISUAL_LOG
 		UE_VLOG_LOCATION(GetOwner(), GoalLocation, 34, FColor(0,127,14), TEXT("GoalLocation"));
-		const FVector ProjectionExtent = MyNavData ? MyNavData->GetDefaultQueryExtent() : FVector(DEFAULT_NAV_QUERY_EXTENT_HORIZONTAL, DEFAULT_NAV_QUERY_EXTENT_HORIZONTAL, DEFAULT_NAV_QUERY_EXTENT_VERTICAL);
-		UE_VLOG_BOX(GetOwner(), FBox(GoalLocation - ProjectionExtent, GoalLocation + ProjectionExtent), FColor::Green, TEXT_EMPTY);
+		UE_VLOG_BOX(GetOwner(), FBox(GoalLocation - GetQueryExtent(), GoalLocation + GetQueryExtent()), FColor::Green, TEXT_EMPTY);
 #endif
 
 		if (ProjectPointToNavigation(GoalLocation, NavMeshGoalLocation) == QuerySuccess)
@@ -763,9 +762,24 @@ ANavigationData* UNavigationComponent::PickNavData() const
 		}
 	}
 
+	CacheNavQueryExtent();
+
 	return MyNavData;
 }
 
+void UNavigationComponent::CacheNavQueryExtent() const
+{
+	NavigationQueryExtent = MyNavData ? MyNavData->GetDefaultQueryExtent() : FVector(DEFAULT_NAV_QUERY_EXTENT_HORIZONTAL, DEFAULT_NAV_QUERY_EXTENT_HORIZONTAL, DEFAULT_NAV_QUERY_EXTENT_VERTICAL);
+	if (MyNavAgent)
+	{
+		const FNavAgentProperties* AgentProperties = MyNavAgent->GetNavAgentProperties();
+		check(AgentProperties);
+		NavigationQueryExtent = FVector(FMath::Max(NavigationQueryExtent.X, AgentProperties->AgentRadius)
+			, FMath::Max(NavigationQueryExtent.Y, AgentProperties->AgentRadius)
+			, FMath::Max(NavigationQueryExtent.Z, AgentProperties->AgentHeight / 2));
+	}
+}
+
 void UNavigationComponent::OnNavDataRegistered(class ANavigationData* NavData)
 {
 	// re-pick regardless of whether we have MyNavData set already - the new one can be better!
