#ttp 332023 - Particles: PSC cannot be toggled visible if initially set to invisible.

#branch UE4

#proj Engine

#summary ParticleSystemComponents can now toggle visibiltiy correctly.

This was introduced quite a while back. Relevant CLs are 1368596 and 1371100.
Basically a check for "can this ever render" was made as "should this render right now".
Added a CanEverRender() function to USceneComponent.

diff --git a/Engine/Source/Runtime/Engine/Classes/Components/SceneComponent.h b/Engine/Source/Runtime/Engine/Classes/Components/SceneComponent.h
--- a/Engine/Source/Runtime/Engine/Classes/Components/SceneComponent.h
+++ b/Engine/Source/Runtime/Engine/Classes/Components/SceneComponent.h
@@ -599,6 +599,9 @@ public:
 	/** return true if it should render **/
 	bool ShouldRender() const;
 
+	/** return true if it can ever render **/
+	bool CanEverRender() const;
+
 	/** 
 	 *  Looking at various values of the component, determines if this
 	 *  component should be added to the scene
diff --git a/Engine/Source/Runtime/Engine/Private/Particles/ParticleComponents.cpp b/Engine/Source/Runtime/Engine/Private/Particles/ParticleComponents.cpp
--- a/Engine/Source/Runtime/Engine/Private/Particles/ParticleComponents.cpp
+++ b/Engine/Source/Runtime/Engine/Private/Particles/ParticleComponents.cpp
@@ -4126,7 +4126,7 @@ void UParticleSystemComponent::InitParticles()
 		if (EmitterInstances.Num() == 0)
 		{
 			const bool bDetailModeAllowsRendering	= DetailMode <= GlobalDetailMode;
-			if ( bDetailModeAllowsRendering && ( ShouldRender() ) )
+			if (bDetailModeAllowsRendering && (CanEverRender()))
 			{
 				EmitterInstances.Empty(Template->Emitters.Num());
 				for (int32 Idx = 0; Idx < Template->Emitters.Num(); Idx++)
diff --git a/Engine/Source/Runtime/Engine/Private/SceneComponent.cpp b/Engine/Source/Runtime/Engine/Private/SceneComponent.cpp
--- a/Engine/Source/Runtime/Engine/Private/SceneComponent.cpp
+++ b/Engine/Source/Runtime/Engine/Private/SceneComponent.cpp
@@ -1245,6 +1245,22 @@ bool USceneComponent::ShouldRender() const
 	return ((bInGameWorld && bShowInGame) || (!bInGameWorld && bShowInEditor)) && bVisible == true;
 }
 
+bool USceneComponent::CanEverRender() const
+{
+	AActor* Owner = GetOwner();
+	const bool bShowInEditor =
+#if WITH_EDITOR
+		(!Owner || !Owner->IsHiddenEd());
+#else
+		false;
+#endif
+
+	const bool bInGameWorld = GetWorld() && GetWorld()->IsGameWorld();
+
+	const bool bShowInGame = (!Owner || !Owner->bHidden);
+	return ((bInGameWorld && bShowInGame) || (!bInGameWorld && bShowInEditor));
+}
+
 bool USceneComponent::ShouldComponentAddToScene() const
 {
 	// If the detail mode setting allows it, add it to the scene.
