#ttp 319367 - UE4: RENDERING: Materials: translucent materials referencing the scene texture node only compile correctly when the decal mask is referenced in the opacity input

#branch UE4

#proj Engine

#summary Any non post process materials using scene textures now fail to compile as is correct.

#extra This was caused by the  bNeedsSceneTexturePostProcessInputs flag being overritten to false by an scene texture experssion that is ok being compiled later after a bad expression had already set it.
The flag now respects it's previous state.

diff --git a/Engine/Source/Runtime/Engine/Private/Materials/HLSLMaterialTranslator.h b/Engine/Source/Runtime/Engine/Private/Materials/HLSLMaterialTranslator.h
--- a/Engine/Source/Runtime/Engine/Private/Materials/HLSLMaterialTranslator.h
+++ b/Engine/Source/Runtime/Engine/Private/Materials/HLSLMaterialTranslator.h
@@ -2164,7 +2164,7 @@ protected:
 		}
 		
 		MaterialCompilationOutput.bNeedsSceneTextures = true;
-		bNeedsSceneTexturePostProcessInputs = ((InSceneTextureId >= PPI_PostProcessInput0 && InSceneTextureId <= PPI_PostProcessInput6) || InSceneTextureId == PPI_SceneColor);
+		bNeedsSceneTexturePostProcessInputs = bNeedsSceneTexturePostProcessInputs || ((InSceneTextureId >= PPI_PostProcessInput0 && InSceneTextureId <= PPI_PostProcessInput6) || InSceneTextureId == PPI_SceneColor);
 
 		ESceneTextureId SceneTextureId = (ESceneTextureId)InSceneTextureId;
 
