#ttp 3323311 - Normal map texture thumbnails appear with a reddish hue when viewed in the material editor
#proj Engine
#branch UE4
#summary reverted changes to 'GetBuildSettingsForRunningPlatform' function.

diff --git a/Engine/Source/Runtime/Engine/Private/TextureDerivedData.cpp b/Engine/Source/Runtime/Engine/Private/TextureDerivedData.cpp
--- a/Engine/Source/Runtime/Engine/Private/TextureDerivedData.cpp
+++ b/Engine/Source/Runtime/Engine/Private/TextureDerivedData.cpp
@@ -343,43 +343,36 @@ static void GetBuildSettingsForRunningPlatform(
 	FTextureBuildSettings& OutBuildSettings
 	)
 {
-	static bool bInitialized = false;
-	static TArray<FName> PlatformFormats;
-
-	if ( !bInitialized )
+	// Compress to whatever formats the active target platforms want
+	ITargetPlatformManagerModule* TPM = GetTargetPlatformManager();
+	if (TPM)
 	{
-		bInitialized = true;
-		// Compress to whatever formats the active target platforms want
-		ITargetPlatformManagerModule* TPM = GetTargetPlatformManager();
-		if (TPM)
-		{
-			ITargetPlatform* CurrentPlatform = NULL;
-			const TArray<ITargetPlatform*>& Platforms = TPM->GetActiveTargetPlatforms();
+		ITargetPlatform* CurrentPlatform = NULL;
+		const TArray<ITargetPlatform*>& Platforms = TPM->GetActiveTargetPlatforms();
 
-			check(Platforms.Num());
+		check(Platforms.Num());
 
-			CurrentPlatform = Platforms[0];
+		CurrentPlatform = Platforms[0];
 
-			for (int32 Index = 1; Index < Platforms.Num(); Index++)
+		for (int32 Index = 1; Index < Platforms.Num(); Index++)
+		{
+			if (Platforms[Index]->IsRunningPlatform())
 			{
-				if (Platforms[Index]->IsRunningPlatform())
-				{
-					CurrentPlatform = Platforms[Index];
-					break;
-				}
+				CurrentPlatform = Platforms[Index];
+				break;
 			}
-
-			check(CurrentPlatform != NULL);
-
-			CurrentPlatform->GetTextureFormats(&Texture, PlatformFormats);
 		}
-	}
 
+		check(CurrentPlatform != NULL);
 
-	// Assume there is at least one format and the first one is what we want at runtime.
-	check(PlatformFormats.Num());
-	GetTextureBuildSettings(Texture, GSystemSettings.TextureLODSettings, OutBuildSettings);
-	OutBuildSettings.TextureFormatName = PlatformFormats[0];
+		TArray<FName> PlatformFormats;
+		CurrentPlatform->GetTextureFormats(&Texture, PlatformFormats);
+
+		// Assume there is at least one format and the first one is what we want at runtime.
+		check(PlatformFormats.Num());
+		GetTextureBuildSettings(Texture, GSystemSettings.TextureLODSettings, OutBuildSettings);
+		OutBuildSettings.TextureFormatName = PlatformFormats[0];
+	}
 }
 
 /**
