Back out ensure of convex mesh != NULL.
The fix will take some time and it's not urgent so for now just go back to being a warning

diff --git a/Engine/Source/Runtime/Engine/Private/PhysicsEngine/BodySetup.cpp b/Engine/Source/Runtime/Engine/Private/PhysicsEngine/BodySetup.cpp
--- a/Engine/Source/Runtime/Engine/Private/PhysicsEngine/BodySetup.cpp
+++ b/Engine/Source/Runtime/Engine/Private/PhysicsEngine/BodySetup.cpp
@@ -408,7 +408,7 @@ void UBodySetup::AddShapesToRigidActor(PxRigidActor* PDestActor, FVector& Scale3
 		{
 			FKConvexElem* ConvexElem = &(AggGeom.ConvexElems[i]);
 
-			if(ensure(ConvexElem->ConvexMesh != NULL))
+			if(ConvexElem->ConvexMesh != NULL)
 			{
 				PxTransform PLocalPose;
 				bool bUseNegX = CalcMeshNegScaleCompensation(Scale3D, PLocalPose);
