#ue4

- Fix ACharacter::LaunchCharacter() adding base velocity before launch, then the character movement component also imparting the velocity when the jump is executed. It should only be done in the component.

- Refactor base velocity and tangential velocity code to the MovementBaseUtility namespace: GetMovementBaseVelocity(), GetMovementBaseTangentialVelocity().

diff --git a/Engine/Source/Runtime/Engine/Classes/GameFramework/Character.h b/Engine/Source/Runtime/Engine/Classes/GameFramework/Character.h
--- a/Engine/Source/Runtime/Engine/Classes/GameFramework/Character.h
+++ b/Engine/Source/Runtime/Engine/Classes/GameFramework/Character.h
@@ -80,6 +80,12 @@ namespace MovementBaseUtility
 
 	/** Remove tick dependency of BasedObjectTick on OldBase */
 	void RemoveTickDependency(FTickFunction& BasedObjectTick, class UPrimitiveComponent* OldBase);
+
+	/** Get the velocity of the given component, first checking the ComponentVelocity and falling back to the physics velocity if necessary. */
+	FVector GetMovementBaseVelocity(const class UPrimitiveComponent* MovementBase);
+
+	/** Get the tangential velocity at WorldLocation for the given component. */
+	FVector GetMovementBaseTangentialVelocity(const class UPrimitiveComponent* MovementBase, const FVector& WorldLocation);
 }
 
 /** Struct to hold relative position information from the server. */
diff --git a/Engine/Source/Runtime/Engine/Classes/GameFramework/CharacterMovementComponent.h b/Engine/Source/Runtime/Engine/Classes/GameFramework/CharacterMovementComponent.h
--- a/Engine/Source/Runtime/Engine/Classes/GameFramework/CharacterMovementComponent.h
+++ b/Engine/Source/Runtime/Engine/Classes/GameFramework/CharacterMovementComponent.h
@@ -460,7 +460,7 @@ public:
 	UPROPERTY(Category="Character Movement", EditAnywhere, BlueprintReadWrite)
 	uint32 bImpartBaseVelocityZ:1;
 
-	/** If true, impart the base actor's angular velocity when falling off it (which includes jumping).
+	/** If true, impart the base component's tangential components of angular velocity when jumping or falling off it.
 	  * Only those components of the velocity allowed by the separate component settings (bImpartBaseVelocity[XYZ]) will be applied. */
 	UPROPERTY(Category="Character Movement", EditAnywhere, BlueprintReadWrite)
 	uint32 bImpartBaseAngularVelocity:1;
diff --git a/Engine/Source/Runtime/Engine/Private/Character.cpp b/Engine/Source/Runtime/Engine/Private/Character.cpp
--- a/Engine/Source/Runtime/Engine/Private/Character.cpp
+++ b/Engine/Source/Runtime/Engine/Private/Character.cpp
@@ -363,6 +363,52 @@ namespace MovementBaseUtility
 			}
 		}
 	}
+
+	FVector GetMovementBaseVelocity(const class UPrimitiveComponent* MovementBase)
+	{
+		FVector BaseVelocity = FVector::ZeroVector;
+		if (MovementBaseUtility::IsDynamicBase(MovementBase))
+		{
+			BaseVelocity = MovementBase->GetComponentVelocity();
+			if (BaseVelocity.IsZero())
+			{
+				// Fall back to actor's Root component
+				const AActor* Owner = MovementBase->GetOwner();
+				if (Owner)
+				{
+					// Component might be moved manually (not by simulated physics or a movement component), see if the root component of the actor has a velocity.
+					BaseVelocity = MovementBase->GetOwner()->GetVelocity();
+				}				
+			}
+
+			// Fall back to physics velocity.
+			if (BaseVelocity.IsZero() && MovementBase->GetBodyInstance())
+			{
+				BaseVelocity = MovementBase->GetBodyInstance()->GetUnrealWorldVelocity();
+			}
+		}
+		
+		return BaseVelocity;
+	}
+
+	FVector GetMovementBaseTangentialVelocity(const class UPrimitiveComponent* MovementBase, const FVector& WorldLocation)
+	{
+		if (MovementBaseUtility::IsDynamicBase(MovementBase))
+		{
+			if (MovementBase->GetBodyInstance())
+			{
+				const FVector BaseAngVel = MovementBase->GetBodyInstance()->GetUnrealWorldAngularVelocity();
+				if (!BaseAngVel.IsZero())
+				{
+					const FVector RadialDistanceToBase = WorldLocation - MovementBase->GetComponentLocation();
+					const FVector TangentialVel = FMath::DegreesToRadians(BaseAngVel) ^ RadialDistanceToBase;
+					return TangentialVel;
+				}
+			}			
+		}
+		
+		return FVector::ZeroVector;
+	}
 }
 
 
@@ -537,20 +583,16 @@ void ACharacter::LaunchCharacter(FVector LaunchVelocity, bool bXYOverride, bool
 	if (CharacterMovement)
 	{
 		FVector FinalVel = LaunchVelocity;
-		if (MovementBaseUtility::IsDynamicBase(MovementBase))
-		{
-			// Is this right? What if we are based on non root component and the root component has a velocity? Will it be lost?
-			FinalVel += MovementBase->GetComponentVelocity();
-		}
+		const FVector Velocity = GetVelocity();
 
 		if (!bXYOverride)
 		{
-			FinalVel.X += GetVelocity().X;
-			FinalVel.Y += GetVelocity().Y;
+			FinalVel.X += Velocity.X;
+			FinalVel.Y += Velocity.Y;
 		}
 		if (!bZOverride)
 		{
-			FinalVel.Z += GetVelocity().Z;
+			FinalVel.Z += Velocity.Z;
 		}
 
 		CharacterMovement->Launch(FinalVel);
diff --git a/Engine/Source/Runtime/Engine/Private/CharacterMovementComponent.cpp b/Engine/Source/Runtime/Engine/Private/CharacterMovementComponent.cpp
--- a/Engine/Source/Runtime/Engine/Private/CharacterMovementComponent.cpp
+++ b/Engine/Source/Runtime/Engine/Private/CharacterMovementComponent.cpp
@@ -363,31 +363,15 @@ FVector UCharacterMovementComponent::GetImpartedMovementBaseVelocity() const
 	if (CharacterOwner)
 	{
 		UPrimitiveComponent* MovementBase = CharacterOwner->GetMovementBase();
-		if (MovementBase != NULL && (MovementBase->Mobility == EComponentMobility::Movable))
+		if (MovementBaseUtility::IsDynamicBase(MovementBase))
 		{
-			FVector BaseVelocity = MovementBase->GetComponentVelocity();
-			if (BaseVelocity.IsZero() && MovementBase->GetOwner())
-			{
-				// Component might be moved manually (not by simulated physics or a movement component), see if the root component of the actor has a velocity.
-				BaseVelocity = MovementBase->GetOwner()->GetVelocity();
-			}
-
-			// Fall back to physics velocity.
-			if (BaseVelocity.IsZero() && MovementBase->GetBodyInstance())
-			{
-				BaseVelocity = MovementBase->GetBodyInstance()->GetUnrealWorldVelocity();
-			}
-
+			FVector BaseVelocity = MovementBaseUtility::GetMovementBaseVelocity(MovementBase);
+			
 			if (bImpartBaseAngularVelocity)
 			{
-				const FVector BaseAngVel = MovementBase->GetPhysicsAngularVelocity();
-				if (!BaseAngVel.IsZero())
-				{
-					const FVector CharacterBasePosition = (UpdatedComponent->GetComponentLocation() - FVector(0.f, 0.f, CharacterOwner->CapsuleComponent->GetScaledCapsuleHalfHeight()));
-					const FVector RadialDistanceToBase = CharacterBasePosition - MovementBase->GetComponentLocation();
-					const FVector BaseLinearVelFromAngularVel = FMath::DegreesToRadians(BaseAngVel) ^ RadialDistanceToBase;
-					BaseVelocity += BaseLinearVelFromAngularVel;
-				}
+				const FVector CharacterBasePosition = (UpdatedComponent->GetComponentLocation() - FVector(0.f, 0.f, CharacterOwner->CapsuleComponent->GetScaledCapsuleHalfHeight()));
+				const FVector BaseTangentialVel = MovementBaseUtility::GetMovementBaseTangentialVelocity(MovementBase, CharacterBasePosition);
+				BaseVelocity += BaseTangentialVel;
 			}
 
 			if (bImpartBaseVelocityX)
