#ue4
- UWorld::EncroachingBlockingGeometry now handles sphere components
- fixes potential ensure with playerstart choosing code when using DefaultPawns

diff --git a/Engine/Source/Runtime/Engine/Private/LevelActor.cpp b/Engine/Source/Runtime/Engine/Private/LevelActor.cpp
--- a/Engine/Source/Runtime/Engine/Private/LevelActor.cpp
+++ b/Engine/Source/Runtime/Engine/Private/LevelActor.cpp
@@ -686,22 +686,31 @@ bool UWorld::EncroachingBlockingGeometry(AActor* TestActor, FVector TestLocation
 		static FName NAME_EncroachingBlockingGeometry = FName(TEXT("EncroachingBlockingGeometry"));
 		ECollisionChannel BlockingChannel = PrimComp->GetCollisionObjectType();
 
-		// @TODO support more than just capsules for spawning, don't yet have - add primitive component function to calculate these, use boundingbox as fallback
-		UCapsuleComponent* Capsule = Cast<UCapsuleComponent>(PrimComp);
-
-		if ( Capsule )//&& (Capsule->BodyInstance.GetPxRigidActor() == NULL) )
+		// @TODO support more than just capsules and spheres for spawning, don't yet have - add primitive component function to calculate these, use boundingbox as fallback
+		UCapsuleComponent* const Capsule = Cast<UCapsuleComponent>(PrimComp);
+		if (Capsule)
 		{
 			FCollisionQueryParams Params(NAME_EncroachingBlockingGeometry, false, TestActor);
 			bFoundBlockingHit = OverlapMulti(Overlaps, TestLocation, FQuat::Identity, BlockingChannel, FCollisionShape::MakeCapsule(FMath::Max(Capsule->GetScaledCapsuleRadius() - Epsilon, 0.1f), FMath::Max(Capsule->GetScaledCapsuleHalfHeight() - Epsilon, 0.1f)), Params);
 		}
 		else
 		{
-			if(ensureMsgf(PrimComp->IsRegistered(), TEXT("Components must be registered in order to be used in a ComponentOverlapMulti call. PriComp: %s TestActor: %s"), *PrimComp->GetName(), *TestActor->GetName()))
+			USphereComponent* const Sphere = Cast<USphereComponent>(PrimComp);
+			if (Sphere)
+			{
+				FCollisionQueryParams Params(NAME_EncroachingBlockingGeometry, false, TestActor);
+				bFoundBlockingHit = OverlapMulti(Overlaps, TestLocation, FQuat::Identity, BlockingChannel, FCollisionShape::MakeSphere(FMath::Max(Sphere->GetScaledSphereRadius() - Epsilon, 0.1f)), Params);
+			}
+			else
 			{
-				FComponentQueryParams Params(NAME_EncroachingBlockingGeometry, TestActor);
-				bFoundBlockingHit = ComponentOverlapMulti(Overlaps, PrimComp, TestLocation, TestActor->GetActorRotation(), Params);
+				if(ensureMsgf(PrimComp->IsRegistered(), TEXT("Components must be registered in order to be used in a ComponentOverlapMulti call. PriComp: %s TestActor: %s"), *PrimComp->GetName(), *TestActor->GetName()))
+				{
+					FComponentQueryParams Params(NAME_EncroachingBlockingGeometry, TestActor);
+					bFoundBlockingHit = ComponentOverlapMulti(Overlaps, PrimComp, TestLocation, TestActor->GetActorRotation(), Params);
+				}
 			}
 		}
+
 		for( int32 HitIdx = 0; HitIdx < Overlaps.Num(); HitIdx++ )
 		{
 			if ( Overlaps[HitIdx].Component.IsValid() && (Overlaps[HitIdx].Component.Get()->GetCollisionResponseToChannel(BlockingChannel) == ECR_Block) )
