https://udn.unrealengine.com/questions/204631/purpose-of-getallsocketnames.html

- Fixed to use socket name instead of name

diff --git a/Engine/Source/Runtime/Engine/Private/SkinnedMeshComponent.cpp b/Engine/Source/Runtime/Engine/Private/SkinnedMeshComponent.cpp
--- a/Engine/Source/Runtime/Engine/Private/SkinnedMeshComponent.cpp
+++ b/Engine/Source/Runtime/Engine/Private/SkinnedMeshComponent.cpp
@@ -943,7 +943,7 @@ TArray<FName> USkinnedMeshComponent::GetAllSocketNames() const
 	{
 		for( auto It=SkeletalMesh->Skeleton->Sockets.CreateConstIterator(); It; ++It )
 		{
-			SocketNames.Add( FName( *((*It)->GetName()) ) );
+			SocketNames.Add( (*It)->SocketName ) ;
 		}
 	}
 	return SocketNames;
diff --git a/Engine/Source/Runtime/Engine/Private/StaticMeshComponent.cpp b/Engine/Source/Runtime/Engine/Private/StaticMeshComponent.cpp
--- a/Engine/Source/Runtime/Engine/Private/StaticMeshComponent.cpp
+++ b/Engine/Source/Runtime/Engine/Private/StaticMeshComponent.cpp
@@ -144,7 +144,7 @@ TArray<FName> UStaticMeshComponent::GetAllSocketNames() const
 	{
 		for( auto It=StaticMesh->Sockets.CreateConstIterator(); It; ++It )
 		{
-			SocketNames.Add( FName( *((*It)->GetName()) ) );
+			SocketNames.Add( (*It)->SocketName ) ;
 		}
 	}
 	return SocketNames;
