#ue4

- Cleaned up comments on LaunchCharacter()

diff --git a/Engine/Source/Runtime/Engine/Classes/GameFramework/Character.h b/Engine/Source/Runtime/Engine/Classes/GameFramework/Character.h
--- a/Engine/Source/Runtime/Engine/Classes/GameFramework/Character.h
+++ b/Engine/Source/Runtime/Engine/Classes/GameFramework/Character.h
@@ -300,10 +300,11 @@ protected:
 
 public:
 
-	/** Launch Character with LaunchVelocity, triggers the OnLaunched event.
-	  * @PARAM LaunchVelocity is the velocity to impart to the Pawn
-	  * @PARAM bXYOverride if true replace the XY part of the Pawn's velocity instead of adding to it.
-	  * @PARAM bZOverride if true replace the Z component of the Pawn's velocity instead of adding to it.
+	/** Set a pending launch velocity on the Character. This velocity will be processed on the next CharacterMovementComponent tick,
+	  * and will set it to the "falling" state. Triggers the OnLaunched event.
+	  * @PARAM LaunchVelocity is the velocity to impart to the Character
+	  * @PARAM bXYOverride if true replace the XY part of the Character's velocity instead of adding to it.
+	  * @PARAM bZOverride if true replace the Z component of the Character's velocity instead of adding to it.
 	  */
 	UFUNCTION(BlueprintCallable, Category="Pawn|Character")
 	virtual void LaunchCharacter(FVector LaunchVelocity, bool bXYOverride, bool bZOverride);
