Prevent a crash in the texture streaming system when given an invalid view location

diff --git a/Engine/Source/Runtime/Engine/Private/ContentStreaming.cpp b/Engine/Source/Runtime/Engine/Private/ContentStreaming.cpp
--- a/Engine/Source/Runtime/Engine/Private/ContentStreaming.cpp
+++ b/Engine/Source/Runtime/Engine/Private/ContentStreaming.cpp
@@ -222,6 +222,7 @@ struct FStreamingTexture
 	{
 		Texture = InTexture;
 		WantedMips = InTexture->ResidentMips;
+		DynamicScreenSize = 0.0f;
 		MipCount = InTexture->GetNumMips();
 		STAT( PerfectWantedMips = InTexture->ResidentMips );
 		STAT( MostResidentMips = InTexture->ResidentMips );
@@ -1426,30 +1427,38 @@ void FStreamingManagerBase::SetupViewInfos( float DeltaTime )
  */
 void FStreamingManagerBase::AddViewInformation( const FVector& ViewOrigin, float ScreenSize, float FOVScreenSize, float BoostFactor/*=1.0f*/, bool bOverrideLocation/*=false*/, float Duration/*=0.0f*/, TWeakObjectPtr<AActor> InActorToBoost /*=NULL*/ )
 {
-	BoostFactor *= CVarStreamingBoost.GetValueOnGameThread();
-
-	if ( bPendingRemoveViews )
+	// Is this a reasonable location?
+	if ( FMath::Abs(ViewOrigin.X) < (1.0e+20f) && FMath::Abs(ViewOrigin.Y) < (1.0e+20f) && FMath::Abs(ViewOrigin.Z) < (1.0e+20f) )
 	{
-		bPendingRemoveViews = false;
+		BoostFactor *= CVarStreamingBoost.GetValueOnGameThread();
 
-		// Remove out-dated override views and empty the PendingViewInfos/SlaveLocation arrays to be populated again during next frame.
-		RemoveStreamingViews( RemoveStreamingViews_Normal );
-	}
+		if ( bPendingRemoveViews )
+		{
+			bPendingRemoveViews = false;
 
-	// Remove a lasting location if we're given the same location again but with 0 duration.
-	if ( Duration <= 0.0f )
-	{
-		RemoveViewInfoFromArray( LastingViewInfos, ViewOrigin );
-	}
+			// Remove out-dated override views and empty the PendingViewInfos/SlaveLocation arrays to be populated again during next frame.
+			RemoveStreamingViews( RemoveStreamingViews_Normal );
+		}
 
-	// Should we remember this location for a while?
-	if ( Duration > 0.0f )
-	{
-		AddViewInfoToArray( LastingViewInfos, ViewOrigin, ScreenSize, FOVScreenSize, BoostFactor, bOverrideLocation, Duration, InActorToBoost );
+		// Remove a lasting location if we're given the same location again but with 0 duration.
+		if ( Duration <= 0.0f )
+		{
+			RemoveViewInfoFromArray( LastingViewInfos, ViewOrigin );
+		}
+
+		// Should we remember this location for a while?
+		if ( Duration > 0.0f )
+		{
+			AddViewInfoToArray( LastingViewInfos, ViewOrigin, ScreenSize, FOVScreenSize, BoostFactor, bOverrideLocation, Duration, InActorToBoost );
+		}
+		else
+		{
+			AddViewInfoToArray( PendingViewInfos, ViewOrigin, ScreenSize, FOVScreenSize, BoostFactor, bOverrideLocation, 0.0f, InActorToBoost );
+		}
 	}
 	else
 	{
-		AddViewInfoToArray( PendingViewInfos, ViewOrigin, ScreenSize, FOVScreenSize, BoostFactor, bOverrideLocation, 0.0f, InActorToBoost );
+		int SetDebugBreakpointHere = 0;
 	}
 }
 
@@ -1502,7 +1511,7 @@ void FStreamingManagerBase::Tick( float DeltaTime, bool bProcessEverything/*=fal
 {
 	UpdateResourceStreaming( DeltaTime, bProcessEverything );
 
-	// Trigger a call to bPendingRemoveViews( Normal ) next time we add a view.
+	// Trigger a call to RemoveStreamingViews( RemoveStreamingViews_Normal ) next time a view is added.
 	bPendingRemoveViews = true;
 }
 
