#ue4
- Hid deprecated nodes from the default GetMenuEntries behavior of UEdGraphSchema::GetGraphContextActions

diff --git a/Engine/Source/Runtime/Engine/Private/EdGraph/EdGraph.cpp b/Engine/Source/Runtime/Engine/Private/EdGraph/EdGraph.cpp
--- a/Engine/Source/Runtime/Engine/Private/EdGraph/EdGraph.cpp
+++ b/Engine/Source/Runtime/Engine/Private/EdGraph/EdGraph.cpp
@@ -1305,7 +1305,7 @@ void UEdGraphSchema::GetGraphContextActions(FGraphContextMenuBuilder& ContextMen
 	for (TObjectIterator<UClass> ClassIt; ClassIt; ++ClassIt)
 	{
 		UClass* Class = *ClassIt;
-		if (Class->IsChildOf(UEdGraphNode::StaticClass()) && !Class->HasAnyClassFlags(CLASS_Abstract))
+		if (Class->IsChildOf(UEdGraphNode::StaticClass()) && !Class->HasAnyClassFlags(CLASS_Abstract | CLASS_Deprecated))
 		{
 			const UEdGraphNode* ClassCDO = Class->GetDefaultObject<UEdGraphNode>();
 
