- Removed editor icons from landscape actors

diff --git a/Engine/Source/Runtime/Engine/Classes/Landscape/LandscapeProxy.h b/Engine/Source/Runtime/Engine/Classes/Landscape/LandscapeProxy.h
--- a/Engine/Source/Runtime/Engine/Classes/Landscape/LandscapeProxy.h
+++ b/Engine/Source/Runtime/Engine/Classes/Landscape/LandscapeProxy.h
@@ -165,7 +165,7 @@ namespace ELandscapeLayerPaintingRestriction
 }
 
 UCLASS(HeaderGroup=Terrain, dependson=UEngineTypes, NotPlaceable, hidecategories=(Display, Attachment, Physics, Debug, Lighting, LOD), showcategories=(Rendering, "Utilities|Orientation"), MinimalAPI)
-class ALandscapeProxy : public AInfo, public INavRelevantActorInterface
+class ALandscapeProxy : public AActor, public INavRelevantActorInterface
 {
 	GENERATED_UCLASS_BODY()
 
diff --git a/Engine/Source/Runtime/Engine/Private/Landscape/Landscape.cpp b/Engine/Source/Runtime/Engine/Private/Landscape/Landscape.cpp
--- a/Engine/Source/Runtime/Engine/Private/Landscape/Landscape.cpp
+++ b/Engine/Source/Runtime/Engine/Private/Landscape/Landscape.cpp
@@ -423,35 +423,6 @@ void ULandscapeComponent::PostLoad()
 ALandscape::ALandscape(const class FPostConstructInitializeProperties& PCIP)
 	: Super(PCIP)
 {
-#if WITH_EDITORONLY_DATA
-	if ((SpriteComponent != NULL) && !IsRunningCommandlet())
-	{
-		// Structure to hold one-time initialization
-		struct FConstructorStatics
-		{
-			ConstructorHelpers::FObjectFinderOptional<UTexture2D> TerrainTexture;
-			FName ID_Landscape;
-			FText NAME_Landscape;
-			FConstructorStatics()
-				: TerrainTexture(TEXT("/Engine/EditorResources/S_Terrain.S_Terrain"))
-				, ID_Landscape(TEXT("Landscape"))
-				, NAME_Landscape(NSLOCTEXT( "SpriteCategory", "Landscape", "Landscape" ))
-
-			{
-			}
-		};
-		static FConstructorStatics ConstructorStatics;
-
-		SpriteComponent->Sprite = ConstructorStatics.TerrainTexture.Get();
-		SpriteComponent->RelativeScale3D = FVector(0.5f, 0.5f, 0.5f);
-		SpriteComponent->SpriteInfo.Category = ConstructorStatics.ID_Landscape;
-		SpriteComponent->SpriteInfo.DisplayName = ConstructorStatics.NAME_Landscape;
-	}
-#endif // WITH_EDITORONLY_DATA
-
-	bHidden = false;
-	StaticLightingResolution = 1.0f;
-	StreamingDistanceMultiplier = 1.0f;
 	bIsProxy = false;
 
 #if WITH_EDITORONLY_DATA
@@ -467,36 +438,21 @@ ALandscape* ALandscape::GetLandscapeActor()
 ALandscapeProxy::ALandscapeProxy(const class FPostConstructInitializeProperties& PCIP)
 	: Super(PCIP)
 {
+	PrimaryActorTick.bCanEverTick = false;
+	bReplicates = false;
+	NetUpdateFrequency = 10.0f;
+	bHidden = false;
+	bReplicateMovement = false;
+	bCanBeDamaged = false;
+		
 	TSubobjectPtr<USceneComponent> SceneComponent = PCIP.CreateDefaultSubobject<USceneComponent>(this, TEXT("RootComponent0"));
 	RootComponent = SceneComponent;
 	RootComponent->RelativeScale3D = FVector(128.0f, 128.0f, 256.0f); // Old default scale, preserved for compatibility. See ULandscapeEditorObject::NewLandscape_Scale
 	RootComponent->Mobility = EComponentMobility::Static;
 	LandscapeSectionOffset = FIntPoint::ZeroValue;
 	
-#if WITH_EDITORONLY_DATA
-	if (!IsRunningCommandlet() && (SpriteComponent != NULL))
-	{
-		// Structure to hold one-time initialization
-		struct FConstructorStatics
-		{
-			ConstructorHelpers::FObjectFinderOptional<UTexture2D> TerrainTexture;
-			FConstructorStatics()
-				: TerrainTexture(TEXT("Texture2D'/Engine/EditorResources/S_Terrain.S_Terrain'"))
-			{
-			}
-		};
-		static FConstructorStatics ConstructorStatics;
-
-		SpriteComponent->Sprite = ConstructorStatics.TerrainTexture.Get();
-		SpriteComponent->RelativeScale3D = FVector(0.5f, 0.5f, 0.5f);
-		SpriteComponent->AttachParent = RootComponent;
-		SpriteComponent->bAbsoluteScale = true;
-	}
-#endif
-
 	StaticLightingResolution = 1.0f;
 	StreamingDistanceMultiplier = 1.0f;
-	bHidden = false;
 	bIsProxy = true;
 	MaxLODLevel = -1;
 #if WITH_EDITORONLY_DATA
