#ttp 332490 - Fix potential crash refreshing sampler states when running dedicated server
#proj Engine
#branch UE4
#summary Fix crash that can occur when texture resource doesnt exist. prevent unnecessary refreshing of sampler states when running dedicated server.

diff --git a/Engine/Source/Runtime/Engine/Private/Texture2D.cpp b/Engine/Source/Runtime/Engine/Private/Texture2D.cpp
--- a/Engine/Source/Runtime/Engine/Private/Texture2D.cpp
+++ b/Engine/Source/Runtime/Engine/Private/Texture2D.cpp
@@ -1124,6 +1124,11 @@ float UTexture2D::GetGlobalMipMapLODBias()
 
 void UTexture2D::RefreshSamplerStates()
 {
+	if (Resource == nullptr)
+	{
+		return;
+	}
+
 	ENQUEUE_UNIQUE_RENDER_COMMAND_ONEPARAMETER(
 		RefreshSamplerStatesCommand,
 		FTexture2DResource*, Texture2DResource, ((FTexture2DResource*)Resource),
diff --git a/Engine/Source/Runtime/Engine/Private/UnrealEngine.cpp b/Engine/Source/Runtime/Engine/Private/UnrealEngine.cpp
--- a/Engine/Source/Runtime/Engine/Private/UnrealEngine.cpp
+++ b/Engine/Source/Runtime/Engine/Private/UnrealEngine.cpp
@@ -388,6 +388,12 @@ void SystemResolutionSinkCallback()
 */
 void RefreshSamplerStatesCallback()
 {
+	if (FApp::CanEverRender() == false)
+	{
+		// Avoid unnecessary work when running in dedicated server mode.
+		return;
+	}
+
 	bool bRefreshSamplerStates = false;
 
 	{
