#ttp 332096 - EDITOR: Missing Tooltip for Camera Anim
#ttp 332097 - EDITOR: Missing tooltip for Curve asset
#ttp 332098 - EDITOR: Missing Tooltip for Force Feedback Effect asset
#ttp 332099 - EDITOR: Missing Tooltip for Touch Interface Setup asset
#summary Added some comments for some asset types
#proj Editor
#branch UE4

diff --git a/Engine/Source/Runtime/Engine/Classes/Camera/CameraAnim.h b/Engine/Source/Runtime/Engine/Classes/Camera/CameraAnim.h
--- a/Engine/Source/Runtime/Engine/Classes/Camera/CameraAnim.h
+++ b/Engine/Source/Runtime/Engine/Classes/Camera/CameraAnim.h
@@ -1,13 +1,12 @@
 // Copyright 1998-2014 Epic Games, Inc. All Rights Reserved.
 
-
-/**
- * CameraAnim: defines a pre-packaged animation to be played on a camera.
- */
-
 #pragma once
+
 #include "CameraAnim.generated.h"
 
+/**
+ * A predefined animation to be played on a camera
+ */
 UCLASS(notplaceable, MinimalAPI)
 class UCameraAnim : public UObject
 {
diff --git a/Engine/Source/Runtime/Engine/Classes/Curves/CurveBase.h b/Engine/Source/Runtime/Engine/Classes/Curves/CurveBase.h
--- a/Engine/Source/Runtime/Engine/Classes/Curves/CurveBase.h
+++ b/Engine/Source/Runtime/Engine/Classes/Curves/CurveBase.h
@@ -400,6 +400,9 @@ public:
 
 //////////////////////////////////////////////////////////////////////////
 
+/**
+ * Defines a curve of interpolated points to evaluate over a given range
+ */
 UCLASS(abstract,MinimalAPI)
 class UCurveBase : public UObject, public FCurveOwnerInterface
 {
diff --git a/Engine/Source/Runtime/Engine/Classes/GameFramework/ForceFeedbackEffect.h b/Engine/Source/Runtime/Engine/Classes/GameFramework/ForceFeedbackEffect.h
--- a/Engine/Source/Runtime/Engine/Classes/GameFramework/ForceFeedbackEffect.h
+++ b/Engine/Source/Runtime/Engine/Classes/GameFramework/ForceFeedbackEffect.h
@@ -66,6 +66,10 @@ struct FActiveForceFeedbackEffect
 	bool Update(const float DeltaTime, struct FForceFeedbackValues& Values);
 };
 
+
+/**
+ * A predefined force-feedback effect to be played on a controller
+ */
 UCLASS(BlueprintType, headergroup=UserInterface, MinimalAPI)
 class UForceFeedbackEffect : public UObject
 {
diff --git a/Engine/Source/Runtime/Engine/Classes/GameFramework/TouchInterface.h b/Engine/Source/Runtime/Engine/Classes/GameFramework/TouchInterface.h
--- a/Engine/Source/Runtime/Engine/Classes/GameFramework/TouchInterface.h
+++ b/Engine/Source/Runtime/Engine/Classes/GameFramework/TouchInterface.h
@@ -38,6 +38,9 @@ struct FTouchInputControl
 };
 
 
+/**
+ * Defines an interface by which touch input can be controlled using any number of buttons and virtual joysticks
+ */
 UCLASS(Blueprintable, BlueprintType)
 class ENGINE_API UTouchInterface : public UObject
 {
diff --git a/Engine/Source/Runtime/Engine/Classes/Sound/SoundAttenuation.h b/Engine/Source/Runtime/Engine/Classes/Sound/SoundAttenuation.h
--- a/Engine/Source/Runtime/Engine/Classes/Sound/SoundAttenuation.h
+++ b/Engine/Source/Runtime/Engine/Classes/Sound/SoundAttenuation.h
@@ -150,7 +150,7 @@ struct TStructOpsTypeTraits<FAttenuationSettings> : public TStructOpsTypeTraitsB
 };
 
 /** 
- * Defines how a sounds changes volume with distance to the listener
+ * Defines how a sound changes volume with distance to the listener
  */
 UCLASS(BlueprintType, hidecategories=Object, editinlinenew, MinimalAPI)
 class USoundAttenuation : public UObject
