#UE4 - added a warning if IsLocalPlayerController is called during construction / destruction of PlayerController

diff --git a/Engine/Source/Runtime/Engine/Private/PlayerController.cpp b/Engine/Source/Runtime/Engine/Private/PlayerController.cpp
--- a/Engine/Source/Runtime/Engine/Private/PlayerController.cpp
+++ b/Engine/Source/Runtime/Engine/Private/PlayerController.cpp
@@ -75,6 +75,11 @@ UNetConnection* APlayerController::GetNetConnection()
 
 bool APlayerController::IsLocalController() const
 {
+	if (Player == NULL)
+	{
+		UE_LOG(LogPlayerController, Warning, TEXT("Calling IsLocalController() while Player is NULL is undefined!"));
+	}
+
 	ENetMode NetMode = GetNetMode();
 	if (NetMode == NM_DedicatedServer)
 	{
