#UE4 - misc cleanup

diff --git a/Engine/Source/Runtime/Engine/Private/World.cpp b/Engine/Source/Runtime/Engine/Private/World.cpp
--- a/Engine/Source/Runtime/Engine/Private/World.cpp
+++ b/Engine/Source/Runtime/Engine/Private/World.cpp
@@ -1312,7 +1312,7 @@ static bool IsTimeLimitExceeded( const TCHAR* CurrentTask, double StartTime, ULe
 
 #if PERF_TRACK_DETAILED_ASYNC_STATS
 
-// Variables for tracking how long eachpart of the AddToWorld process takes
+// Variables for tracking how long each part of the AddToWorld process takes
 double MoveActorTime = 0.0;
 double ShiftActorsTime = 0.0;
 double UpdateComponentsTime = 0.0;
@@ -1956,7 +1956,7 @@ void UWorld::UpdateLevelStreamingInner( UWorld* PersistentWorld, FSceneViewFamil
 		
 		// Figure out whether level should be loaded, visible and block on load if it should be loaded but currently isn't.
 		bool bShouldBeLoaded				= bHasVisibilityRequestPending || (!GEngine->bUseBackgroundLevelStreaming && !PersistentWorld->bShouldForceUnloadStreamingLevels && !StreamingLevel->bIsRequestingUnloadAndRemoval);
-		bool bShouldBeVisible				= bHasVisibilityRequestPending || PersistentWorld->bShouldForceVisibleStreamingLevels;;
+		bool bShouldBeVisible				= bHasVisibilityRequestPending || PersistentWorld->bShouldForceVisibleStreamingLevels;
 		bool bShouldBlockOnLoad				= StreamingLevel->bShouldBlockOnLoad;
 
 		// Don't update if the code requested this level object to be unloaded and removed.
