#UE4 - reordered if statement to check for UPlayer before calling IsLocalPlayerController

diff --git a/Engine/Source/Runtime/Engine/Private/PlayerCameraManager.cpp b/Engine/Source/Runtime/Engine/Private/PlayerCameraManager.cpp
--- a/Engine/Source/Runtime/Engine/Private/PlayerCameraManager.cpp
+++ b/Engine/Source/Runtime/Engine/Private/PlayerCameraManager.cpp
@@ -726,7 +726,7 @@ void APlayerCameraManager::SetDesiredColorScale(FVector NewColorScale, float Int
 
 void APlayerCameraManager::UpdateCamera(float DeltaTime)
 {
-	if ( (PCOwner->IsLocalPlayerController() && PCOwner->Player) || !bUseClientSideCameraUpdates || bDebugClientSideCamera )
+	if ((PCOwner->Player && PCOwner->IsLocalPlayerController()) || !bUseClientSideCameraUpdates || bDebugClientSideCamera)
 	{
 		DoUpdateCamera(DeltaTime);
 
