Add cache for active transforms so that we can deal with deleted actors

diff --git a/Engine/Source/Runtime/Engine/Private/PhysicsEngine/PhysScene.cpp b/Engine/Source/Runtime/Engine/Private/PhysicsEngine/PhysScene.cpp
--- a/Engine/Source/Runtime/Engine/Private/PhysicsEngine/PhysScene.cpp
+++ b/Engine/Source/Runtime/Engine/Private/PhysicsEngine/PhysScene.cpp
@@ -252,6 +252,24 @@ void FPhysScene::AddTorque(FBodyInstance * BodyInstance, const FVector & Torque)
 #endif
 }
 
+#if WITH_PHYSX
+void FPhysScene::RemoveBodyFromActiveTransforms(class PxActor * PActor, uint32 SceneType)
+{
+	if (PActor)
+	{
+		TArray<const PxActiveTransform*> & PActiveTransforms = ActiveTransforms[SceneType];
+		for (int32 TransformIndex = 0; TransformIndex < PActiveTransforms.Num(); ++TransformIndex)
+		{
+			const PxActiveTransform * PActiveTransform = PActiveTransforms[TransformIndex];
+			if (PActiveTransform && PActiveTransform->actor == PActor)
+			{
+				PActiveTransforms[TransformIndex] = NULL;	//we NULL out the pointer to indicate that this actor has been destroyed. This is needed because the transform cache doesn't get updated until the next fetchResults
+			}
+		}
+	}
+}
+#endif
+
 void FPhysScene::TermBody(FBodyInstance * BodyInstance)
 {
 #if WITH_SUBSTEPPING
@@ -261,6 +279,12 @@ void FPhysScene::TermBody(FBodyInstance * BodyInstance)
 		PhysSubStepper->RemoveBodyInstance(BodyInstance);
 	}
 #endif
+
+
+#if WITH_PHYSX
+	RemoveBodyFromActiveTransforms(BodyInstance->RigidActorSync, PST_Sync);
+	RemoveBodyFromActiveTransforms(BodyInstance->RigidActorAsync, PST_Async);
+#endif
 }
 
 #if WITH_SUBSTEPPING
@@ -539,6 +563,9 @@ void FPhysScene::ProcessPhysScene(uint32 SceneType)
 	check(ApexScene);
 	ApexScene->fetchResults( true, &OutErrorCode );
 #endif	//	#if !WITH_APEX
+
+	UpdateActiveTransforms(SceneType);
+
 	if(OutErrorCode != 0)
 	{
 		UE_LOG(LogPhysics, Log, TEXT("PHYSX FETCHRESULTS ERROR: %d"), OutErrorCode);
@@ -551,22 +578,44 @@ void FPhysScene::ProcessPhysScene(uint32 SceneType)
 	bPhysXSceneExecuting[SceneType] = false;
 }
 
-void FPhysScene::SyncComponentsToBodies(uint32 SceneType)
-{
 #if WITH_PHYSX
+void FPhysScene::UpdateActiveTransforms(uint32 SceneType)
+{
+	if (SceneType == PST_Cloth)	//cloth doesn't bother with updating components to bodies so we don't need to store any transforms
+	{
+		return;
+	}
+
 	PxScene* PScene = GetPhysXScene(SceneType);
 	check(PScene);
 	SCENE_LOCK_READ(PScene);
-
 	PxU32 NumTransforms = 0;
 	const PxActiveTransform* PActiveTransforms = PScene->getActiveTransforms(NumTransforms);
-
 	SCENE_UNLOCK_READ(PScene);
+	
+	ActiveTransforms[SceneType].Empty(NumTransforms);
 
+	for (PxU32 TransformIdx = 0; TransformIdx < NumTransforms; ++TransformIdx)
+	{
+		ActiveTransforms[SceneType].Add(&PActiveTransforms[TransformIdx]);	//it's ok to use a pointer here because the physics scene leaves the data for us until the next call to fetchTransform;
+	}
+}
+#endif
+
+void FPhysScene::SyncComponentsToBodies(uint32 SceneType)
+{
+#if WITH_PHYSX
+	PxU32 NumTransforms = ActiveTransforms[SceneType].Num();
+	const TArray<const PxActiveTransform*> & PActiveTransforms = ActiveTransforms[SceneType];
 
 	for(PxU32 TransformIdx=0; TransformIdx<NumTransforms; TransformIdx++)
 	{
-		const PxActiveTransform& PActiveTransform = PActiveTransforms[TransformIdx];
+		if (PActiveTransforms[TransformIdx] == NULL)	//it's possible to call TermBody on FBodyInstance after fetchResults, but before SyncComponentsToBodies - in this case a NULL is used to represent the stale data
+		{
+			continue;
+		}
+
+		const PxActiveTransform& PActiveTransform = *PActiveTransforms[TransformIdx];
 		PxRigidActor* RigidActor = PActiveTransform.actor->isRigidActor();
 
 #if WITH_APEX
diff --git a/Engine/Source/Runtime/Engine/Public/PhysicsPublic.h b/Engine/Source/Runtime/Engine/Public/PhysicsPublic.h
--- a/Engine/Source/Runtime/Engine/Public/PhysicsPublic.h
+++ b/Engine/Source/Runtime/Engine/Public/PhysicsPublic.h
@@ -33,6 +33,8 @@ namespace physx
 	class PxMat44;
 	class PxCpuDispatcher;
 	class PxSimulationEventCallback;
+	struct PxActiveTransform;
+	class PxActor;
 	
 #if WITH_APEX
 	namespace apex
@@ -408,6 +410,15 @@ private:
 #if WITH_PHYSX
 	/** User data wrapper passed to physx */
 	struct FPhysxUserData PhysxUserData;
+
+	/** Cache of active transforms */ //TODO: this solution is not great
+	TArray<const PxActiveTransform*> ActiveTransforms[PST_MAX];
+
+	/** Updates our cache of active transforms */
+	void UpdateActiveTransforms(uint32 SceneType);
+
+	/** When actors are destroyed they must be removed from active transforms */
+	void RemoveBodyFromActiveTransforms(PxActor * PActor, uint32 SceneType);
 #endif
 
 #if WITH_SUBSTEPPING
