#UE4 - updated Login call to take the UPlayer so that local-ness can be established when IsLocalPlayerController is undefined

diff --git a/Engine/Source/Runtime/Engine/Classes/GameFramework/GameMode.h b/Engine/Source/Runtime/Engine/Classes/GameFramework/GameMode.h
--- a/Engine/Source/Runtime/Engine/Classes/GameFramework/GameMode.h
+++ b/Engine/Source/Runtime/Engine/Classes/GameFramework/GameMode.h
@@ -386,7 +386,7 @@ public:
 	virtual void PreLogin(const FString& Options, const FString& Address, const TSharedPtr<class FUniqueNetId>& UniqueId, FString& ErrorMessage);
 
 	/** If login is successful, returns a new PlayerController to associate with this player. Login fails if ErrorMessage string is set. */
-	virtual APlayerController* Login(const FString& Portal, const FString& Options, const TSharedPtr<class FUniqueNetId>& UniqueId, FString& ErrorMessage);
+	virtual APlayerController* Login(class UPlayer* NewPlayer, const FString& Portal, const FString& Options, const TSharedPtr<class FUniqueNetId>& UniqueId, FString& ErrorMessage);
 
 	/** Called after a successful login.  This is the first place it is safe to call replicated functions on the PlayerAController. */
 	virtual void PostLogin( APlayerController* NewPlayer );
diff --git a/Engine/Source/Runtime/Engine/Private/GameMode.cpp b/Engine/Source/Runtime/Engine/Private/GameMode.cpp
--- a/Engine/Source/Runtime/Engine/Private/GameMode.cpp
+++ b/Engine/Source/Runtime/Engine/Private/GameMode.cpp
@@ -942,7 +942,7 @@ bool AGameMode::MustSpectate(APlayerController* NewPlayer)
 	return NewPlayer->PlayerState->bOnlySpectator;
 }
 
-APlayerController* AGameMode::Login(const FString& Portal, const FString& Options, const TSharedPtr<FUniqueNetId>& UniqueId, FString& ErrorMessage)
+APlayerController* AGameMode::Login(UPlayer* NewPlayer, const FString& Portal, const FString& Options, const TSharedPtr<FUniqueNetId>& UniqueId, FString& ErrorMessage)
 {
 	ErrorMessage = GameSession->ApproveLogin(Options);
 	if ( !ErrorMessage.IsEmpty() )
@@ -950,10 +950,10 @@ APlayerController* AGameMode::Login(const FString& Portal, const FString& Option
 		return NULL;
 	}
 
-	APlayerController* NewPlayer = SpawnPlayerController(FVector::ZeroVector, FRotator::ZeroRotator);
+	APlayerController* NewPlayerController = SpawnPlayerController(FVector::ZeroVector, FRotator::ZeroRotator);
 
 	// Handle spawn failure.
-	if( NewPlayer == NULL)
+	if( NewPlayerController == NULL)
 	{
 		UE_LOG(LogGameMode, Log, TEXT("Couldn't spawn player controller of class %s"), PlayerControllerClass ? *PlayerControllerClass->GetName() : TEXT("NULL"));
 		ErrorMessage = FString::Printf(TEXT("Failed to spawn player controller"));
@@ -961,7 +961,7 @@ APlayerController* AGameMode::Login(const FString& Portal, const FString& Option
 	}
 
 	// Customize incoming player based on URL options
-	InitNewPlayer(NewPlayer, UniqueId, Options);
+	InitNewPlayer(NewPlayerController, UniqueId, Options);
 
 	// Find a start spot.
 	AActor* const StartSpot = FindPlayerStart( NULL, Portal );
@@ -973,29 +973,29 @@ APlayerController* AGameMode::Login(const FString& Portal, const FString& Option
 
 	FRotator InitialControllerRot = StartSpot->GetActorRotation();
 	InitialControllerRot.Roll = 0.f;
-	NewPlayer->SetInitialLocationAndRotation(StartSpot->GetActorLocation(), InitialControllerRot);
-	NewPlayer->StartSpot = StartSpot;
+	NewPlayerController->SetInitialLocationAndRotation(StartSpot->GetActorLocation(), InitialControllerRot);
+	NewPlayerController->StartSpot = StartSpot;
 
 	// Register the player with the session
-	GameSession->RegisterPlayer(NewPlayer, UniqueId, HasOption(Options, TEXT("bIsFromInvite")));
+	GameSession->RegisterPlayer(NewPlayerController, UniqueId, HasOption(Options, TEXT("bIsFromInvite")));
 	
 	// Init player's name
 	FString InName = ParseOption( Options, TEXT("Name")).Left(20);
 	if( InName.IsEmpty() )
 	{
-		InName = FString::Printf(TEXT("%s%i"), *DefaultPlayerName, NewPlayer->PlayerState->PlayerId);
+		InName = FString::Printf(TEXT("%s%i"), *DefaultPlayerName, NewPlayerController->PlayerState->PlayerId);
 	}
-	ChangeName( NewPlayer, InName, false );
+	ChangeName( NewPlayerController, InName, false );
 
 	// Set up spectating
 	bool bSpectator = FCString::Stricmp(*ParseOption( Options, TEXT("SpectatorOnly") ), TEXT("1")) == 0;
-	if ( bSpectator || MustSpectate(NewPlayer) )
+	if ( bSpectator || MustSpectate(NewPlayerController) )
 	{
-		NewPlayer->StartSpectatingOnly();
-		return NewPlayer;
+		NewPlayerController->StartSpectatingOnly();
+		return NewPlayerController;
 	}
 
-	return NewPlayer;
+	return NewPlayerController;
 }
 
 
diff --git a/Engine/Source/Runtime/Engine/Private/LevelActor.cpp b/Engine/Source/Runtime/Engine/Private/LevelActor.cpp
--- a/Engine/Source/Runtime/Engine/Private/LevelActor.cpp
+++ b/Engine/Source/Runtime/Engine/Private/LevelActor.cpp
@@ -565,7 +565,7 @@ bool UWorld::DestroyActor( AActor* ThisActor, bool bNetForce, bool bShouldModify
 	Player spawning.
 -----------------------------------------------------------------------------*/
 
-APlayerController* UWorld::SpawnPlayActor(UPlayer* Player, ENetRole RemoteRole, const FURL& InURL, const TSharedPtr<FUniqueNetId>& UniqueId, FString& Error, uint8 InNetPlayerIndex)
+APlayerController* UWorld::SpawnPlayActor(UPlayer* NewPlayer, ENetRole RemoteRole, const FURL& InURL, const TSharedPtr<FUniqueNetId>& UniqueId, FString& Error, uint8 InNetPlayerIndex)
 {
 	Error = TEXT("");
 
@@ -580,26 +580,27 @@ APlayerController* UWorld::SpawnPlayActor(UPlayer* Player, ENetRole RemoteRole,
 	AGameMode* GameMode = GetAuthGameMode();
 
 	// Give the GameMode a chance to accept the login
-	APlayerController* const Actor = GameMode->Login(*InURL.Portal, Options, UniqueId, Error);
-	if (Actor == NULL)
+	APlayerController* const NewPlayerController = GameMode->Login(NewPlayer, *InURL.Portal, Options, UniqueId, Error);
+	if (NewPlayerController == NULL)
 	{
 		UE_LOG(LogSpawn, Warning, TEXT("Login failed: %s"), *Error);
 		return NULL;
 	}
 
+	UE_LOG(LogSpawn, Log, TEXT("%s got player %s [%s]"), *NewPlayerController->GetName(), *NewPlayer->GetName(), *UniqueId->ToString());
+
 	// Possess the newly-spawned player.
-	Actor->NetPlayerIndex = InNetPlayerIndex;
-	//UE_LOG(LogSpawn, Log, TEXT("%s got player %s"), *Actor->GetName(), *Player->GetName());
-	Actor->Role = ROLE_Authority;
-	Actor->SetReplicates(RemoteRole != ROLE_None);
+	NewPlayerController->NetPlayerIndex = InNetPlayerIndex;
+	NewPlayerController->Role = ROLE_Authority;
+	NewPlayerController->SetReplicates(RemoteRole != ROLE_None);
 	if (RemoteRole == ROLE_AutonomousProxy)
 	{
-		Actor->SetAutonomousProxy(true);
+		NewPlayerController->SetAutonomousProxy(true);
 	}
-	Actor->SetPlayer(Player);
-	GameMode->PostLogin(Actor);
+	NewPlayerController->SetPlayer(NewPlayer);
+	GameMode->PostLogin(NewPlayerController);
 
-	return Actor;
+	return NewPlayerController;
 }
 
 /*-----------------------------------------------------------------------------
