Add flag to sub-stepping related functions so that you can choose whether to sub-step or not.

diff --git a/Engine/Source/Runtime/Engine/Private/PhysicsEngine/PhysScene.cpp b/Engine/Source/Runtime/Engine/Private/PhysicsEngine/PhysScene.cpp
--- a/Engine/Source/Runtime/Engine/Private/PhysicsEngine/PhysScene.cpp
+++ b/Engine/Source/Runtime/Engine/Private/PhysicsEngine/PhysScene.cpp
@@ -160,7 +160,7 @@ bool UseSyncTime(uint32 SceneType)
 
 }
 
-void FPhysScene::SetKinematicTarget(FBodyInstance * BodyInstance, const FTransform & TargetTransform)
+void FPhysScene::SetKinematicTarget(FBodyInstance * BodyInstance, const FTransform & TargetTransform, bool bAllowSubstepping)
 {
 	TargetTransform.DiagnosticCheckNaN_All();
 
@@ -168,7 +168,7 @@ void FPhysScene::SetKinematicTarget(FBodyInstance * BodyInstance, const FTransfo
 	if (PxRigidDynamic * PRigidDynamic = BodyInstance->GetPxRigidDynamic())
 	{
 #if WITH_SUBSTEPPING
-		if (IsSubstepping())
+		if (bAllowSubstepping && IsSubstepping())
 		{
 			FPhysSubstepTask * PhysSubStepper = PhysSubSteppers[SceneType(BodyInstance)];
 			PhysSubStepper->SetKinematicTarget(BodyInstance, TargetTransform);
@@ -186,14 +186,14 @@ void FPhysScene::SetKinematicTarget(FBodyInstance * BodyInstance, const FTransfo
 #endif
 }
 
-void FPhysScene::AddForce(FBodyInstance * BodyInstance, const FVector & Force)
+void FPhysScene::AddForce(FBodyInstance * BodyInstance, const FVector & Force, bool bAllowSubstepping)
 {
 #if WITH_PHYSX
 
 	if (PxRigidDynamic * PRigidDynamic = BodyInstance->GetPxRigidDynamic())
 	{
 #if WITH_SUBSTEPPING
-		if (IsSubstepping())
+		if (bAllowSubstepping && IsSubstepping())
 		{
 			FPhysSubstepTask * PhysSubStepper = PhysSubSteppers[SceneType(BodyInstance)];
 			PhysSubStepper->AddForce(BodyInstance, Force);
@@ -208,14 +208,14 @@ void FPhysScene::AddForce(FBodyInstance * BodyInstance, const FVector & Force)
 #endif
 }
 
-void FPhysScene::AddForceAtPosition(FBodyInstance * BodyInstance, const FVector & Force, const FVector & Position)
+void FPhysScene::AddForceAtPosition(FBodyInstance * BodyInstance, const FVector & Force, const FVector & Position, bool bAllowSubstepping)
 {
 #if WITH_PHYSX
 
 	if (PxRigidDynamic * PRigidDynamic = BodyInstance->GetPxRigidDynamic())
 	{
 #if WITH_SUBSTEPPING
-		if (IsSubstepping())
+		if (bAllowSubstepping && IsSubstepping())
 		{
 			FPhysSubstepTask * PhysSubStepper = PhysSubSteppers[SceneType(BodyInstance)];
 			PhysSubStepper->AddForceAtPosition(BodyInstance, Force, Position);
@@ -230,14 +230,14 @@ void FPhysScene::AddForceAtPosition(FBodyInstance * BodyInstance, const FVector
 #endif
 }
 
-void FPhysScene::AddTorque(FBodyInstance * BodyInstance, const FVector & Torque)
+void FPhysScene::AddTorque(FBodyInstance * BodyInstance, const FVector & Torque, bool bAllowSubstepping)
 {
 #if WITH_PHYSX
 
 	if (PxRigidDynamic * PRigidDynamic = BodyInstance->GetPxRigidDynamic())
 	{
 #if WITH_SUBSTEPPING
-		if (IsSubstepping())
+		if (bAllowSubstepping && IsSubstepping())
 		{
 			FPhysSubstepTask * PhysSubStepper = PhysSubSteppers[SceneType(BodyInstance)];
 			PhysSubStepper->AddTorque(BodyInstance, Torque);
diff --git a/Engine/Source/Runtime/Engine/Public/PhysicsPublic.h b/Engine/Source/Runtime/Engine/Public/PhysicsPublic.h
--- a/Engine/Source/Runtime/Engine/Public/PhysicsPublic.h
+++ b/Engine/Source/Runtime/Engine/Public/PhysicsPublic.h
@@ -362,16 +362,16 @@ public:
 	void TermBody(FBodyInstance * BodyInstance);
 
 	/** Adds a force to a body - We need to go through scene to support substepping */
-	void AddForce(FBodyInstance * BodyInstance, const FVector & Force);
+	void AddForce(FBodyInstance * BodyInstance, const FVector & Force, bool bAllowSubstepping);
 
 	/** Adds a force to a body at a specific position - We need to go through scene to support substepping */
-	void AddForceAtPosition(FBodyInstance * BodyInstance, const FVector & Force, const FVector & Position);
+	void AddForceAtPosition(FBodyInstance * BodyInstance, const FVector & Force, const FVector & Position, bool bAllowSubstepping);
 
 	/** Adds torque to a body - We need to go through scene to support substepping */
-	void AddTorque(FBodyInstance * BodyInstance, const FVector & Torque);
+	void AddTorque(FBodyInstance * BodyInstance, const FVector & Torque, bool bAllowSubstepping);
 
 	/** Sets a Kinematic actor's target position - We need to do this here to support substepping*/
-	void SetKinematicTarget(FBodyInstance * BodyInstance, const FTransform & TargetTM);
+	void SetKinematicTarget(FBodyInstance * BodyInstance, const FTransform & TargetTM, bool bAllowSubstepping);
 
 	/** Gets the collision disable table */
 	const TMap<uint32, TMap<struct FRigidBodyIndexPair, bool> *> & GetCollisionDisableTableLookup()
