Forgot these files (optional sub-stepping interface)

diff --git a/Engine/Source/Runtime/Engine/Classes/PhysicsEngine/BodyInstance.h b/Engine/Source/Runtime/Engine/Classes/PhysicsEngine/BodyInstance.h
--- a/Engine/Source/Runtime/Engine/Classes/PhysicsEngine/BodyInstance.h
+++ b/Engine/Source/Runtime/Engine/Classes/PhysicsEngine/BodyInstance.h
@@ -346,11 +346,11 @@ public:
 		/** Force this body to sleep */
 		void PutInstanceToSleep();
 		/** Add a force to this body */
-		void AddForce(const FVector& Force);
+		void AddForce(const FVector& Force, bool bAllowSubstepping = true);
 		/** Add a force at a particular world position to this body */
-		void AddForceAtPosition(const FVector& Force, const FVector& Position);
+		void AddForceAtPosition(const FVector& Force, const FVector& Position, bool bAllowSubstepping = true);
 		/** Add a torque to this body */
-		void AddTorque(const FVector& Torque);
+		void AddTorque(const FVector& Torque, bool bAllowSubstepping = true);
 		/** Add an impulse to this body */
 		void AddImpulse(const FVector& Impulse, bool bVelChange);
 		/** Add an impulse to this body and a particular world position */
diff --git a/Engine/Source/Runtime/Engine/Private/PhysicsEngine/BodyInstance.cpp b/Engine/Source/Runtime/Engine/Private/PhysicsEngine/BodyInstance.cpp
--- a/Engine/Source/Runtime/Engine/Private/PhysicsEngine/BodyInstance.cpp
+++ b/Engine/Source/Runtime/Engine/Private/PhysicsEngine/BodyInstance.cpp
@@ -1553,7 +1553,7 @@ void FBodyInstance::SetBodyTransform(const FTransform& NewTransform, bool bTelep
 		{
 			const PxScene * PScene = PRigidDynamic->getScene();
 			FPhysScene * PhysScene = FPhysxUserData::Get<FPhysScene>(PScene->userData);
-			PhysScene->SetKinematicTarget(this, NewTransform);
+			PhysScene->SetKinematicTarget(this, NewTransform, true);
 		}
 		// Otherwise, set global pose
 		else
@@ -2059,7 +2059,7 @@ void FBodyInstance::SetMaxAngularVelocity(float NewMaxAngVel, bool bAddToCurrent
 }
 
 
-void FBodyInstance::AddForce(const FVector& Force)
+void FBodyInstance::AddForce(const FVector& Force, bool bAllowSubstepping)
 {
 #if WITH_PHYSX
 	PxRigidDynamic* PRigidDynamic = GetPxRigidDynamic();
@@ -2067,13 +2067,13 @@ void FBodyInstance::AddForce(const FVector& Force)
 	{
 		const PxScene * PScene = PRigidDynamic->getScene();
 		FPhysScene * PhysScene = FPhysxUserData::Get<FPhysScene>(PScene->userData);
-		PhysScene->AddForce(this, Force);
+		PhysScene->AddForce(this, Force, bAllowSubstepping);
 		
 	}
 #endif // WITH_PHYSX
 }
 
-void FBodyInstance::AddForceAtPosition(const FVector& Force, const FVector& Position)
+void FBodyInstance::AddForceAtPosition(const FVector& Force, const FVector& Position, bool bAllowSubstepping)
 {
 #if WITH_PHYSX
 	PxRigidDynamic* PRigidDynamic = GetPxRigidDynamic();
@@ -2081,12 +2081,12 @@ void FBodyInstance::AddForceAtPosition(const FVector& Force, const FVector& Posi
 	{
 		const PxScene * PScene = PRigidDynamic->getScene();
 		FPhysScene * PhysScene = FPhysxUserData::Get<FPhysScene>(PScene->userData);
-		PhysScene->AddForceAtPosition(this, Force, Position);
+		PhysScene->AddForceAtPosition(this, Force, Position, bAllowSubstepping);
 	}
 #endif // WITH_PHYSX
 }
 
-void FBodyInstance::AddTorque(const FVector& Torque)
+void FBodyInstance::AddTorque(const FVector& Torque, bool bAllowSubstepping)
 {
 #if WITH_PHYSX
 	PxRigidDynamic* PRigidDynamic = GetPxRigidDynamic();
@@ -2094,7 +2094,7 @@ void FBodyInstance::AddTorque(const FVector& Torque)
 	{
 		const PxScene * PScene = PRigidDynamic->getScene();
 		FPhysScene * PhysScene = FPhysxUserData::Get<FPhysScene>(PScene->userData);
-		PhysScene->AddTorque(this, Torque);
+		PhysScene->AddTorque(this, Torque, bAllowSubstepping);
 	}
 #endif // WITH_PHYSX
 }
