Add drag force to vehicles

diff --git a/Engine/Source/Runtime/Engine/Classes/Vehicles/WheeledVehicleMovementComponent.h b/Engine/Source/Runtime/Engine/Classes/Vehicles/WheeledVehicleMovementComponent.h
--- a/Engine/Source/Runtime/Engine/Classes/Vehicles/WheeledVehicleMovementComponent.h
+++ b/Engine/Source/Runtime/Engine/Classes/Vehicles/WheeledVehicleMovementComponent.h
@@ -223,6 +223,18 @@ class ENGINE_API UWheeledVehicleMovementComponent : public UPawnMovementComponen
 	UPROPERTY(EditAnywhere, Category=VehicleSetup)
 	float Mass;
 
+	/** DragCoefficient of the vehicle chassis. */
+	UPROPERTY(EditAnywhere, Category = VehicleSetup)
+	float DragCoefficient;
+
+	/** Chassis width used for drag force computation (cm)*/
+	UPROPERTY(EditAnywhere, Category = VehicleSetup)
+	float ChassisWidth;
+
+	/** Chassis height used for drag force computation (cm)*/
+	UPROPERTY(EditAnywhere, Category = VehicleSetup)
+	float ChassisHeight;
+
 	/** Override center of mass offset, makes tweaking easier [uu] */
 	UPROPERTY(EditAnywhere, Category=VehicleSetup, AdvancedDisplay)
 	FVector COMOffset;
@@ -268,6 +280,9 @@ class ENGINE_API UWheeledVehicleMovementComponent : public UPawnMovementComponen
 	/** Updates the vehicle tuning data. Used when designer tweaks values while game is running */
 	virtual void UpdateTuning();
 
+	/** Updates the forces of drag acting on the vehicle */
+	virtual void UpdateDrag( float DeltaTime );
+
 	/** Used to create any physics engine information for this component */
 	virtual void CreatePhysicsState() OVERRIDE;
 
@@ -412,6 +427,10 @@ protected:
 	UPROPERTY(EditAnywhere, Category=VehicleInput, AdvancedDisplay)
 	FVehicleInputRate SteeringInputRate;
 
+	// Debug drag magnitude last applied
+	UPROPERTY(transient)
+	float DebugDragMagnitude;
+
 	/** Compute steering input */
 	float CalcSteeringInput();
 
diff --git a/Engine/Source/Runtime/Engine/Private/Vehicles/WheeledVehicleMovementComponent.cpp b/Engine/Source/Runtime/Engine/Private/Vehicles/WheeledVehicleMovementComponent.cpp
--- a/Engine/Source/Runtime/Engine/Private/Vehicles/WheeledVehicleMovementComponent.cpp
+++ b/Engine/Source/Runtime/Engine/Private/Vehicles/WheeledVehicleMovementComponent.cpp
@@ -89,6 +89,9 @@ UWheeledVehicleMovementComponent::UWheeledVehicleMovementComponent(const class F
 	: Super(PCIP)
 {
 	Mass = 1500.0f;
+	DragCoefficient = 0.3f;
+	ChassisWidth = 180.f;
+	ChassisHeight = 140.f;
 	InertiaTensorScale = FVector( 1.0f, 0.80f, 1.0f );
 	AngErrorAccumulator = 0.0f;
 	MinNormalizedTireLoad = 0.0f;
@@ -648,6 +651,26 @@ void UWheeledVehicleMovementComponent::TickVehicle( float DeltaTime )
 	{
 		Wheels[i]->Tick(DeltaTime);
 	}
+
+	UpdateDrag(DeltaTime);
+}
+
+void UWheeledVehicleMovementComponent::UpdateDrag(float DeltaTime)
+{
+	if (PVehicle && UpdatedComponent)
+	{
+		float ForwardSpeed = GetForwardSpeed();
+		FVector GlobalForwardVector = UpdatedComponent->GetForwardVector();
+		FVector DragVector = -GlobalForwardVector;
+		float SpeedSquared = ForwardSpeed * ForwardSpeed;
+		float DragArea = ChassisHeight * ChassisWidth;
+		float AirDensity = 1.25 / (100 * 100 * 100); //kg/cm^3
+		float DragMag = 0.5f * AirDensity * SpeedSquared * DragCoefficient * DragArea;
+		DebugDragMagnitude = DragMag;
+		DragVector *= DragMag;
+		FBodyInstance * BodyInstance = UpdatedComponent->GetBodyInstance();
+		BodyInstance->AddForce(DragVector, false);
+	}
 }
 
 void UWheeledVehicleMovementComponent::UpdateTuning()
@@ -1022,12 +1045,14 @@ void UWheeledVehicleMovementComponent::DrawDebug( UCanvas* Canvas, float& YL, fl
 	// draw drive data
 	{
 		Canvas->SetDrawColor( FColor::White );
-
-		YPos += Canvas->DrawText( RenderFont, FString::Printf( TEXT("Speed: %d"), (int32)GetForwardSpeed() ), 4, YPos );
+		float forwardSpeedKmH = GetForwardSpeed() * 3600.f / 100000.f;	//convert from cm/s to km/h
+		YPos += Canvas->DrawText( RenderFont, FString::Printf( TEXT("Speed (km/h): %d"), (int32)forwardSpeedKmH  ), 4, YPos );
 		YPos += Canvas->DrawText( RenderFont, FString::Printf( TEXT("Steering: %f"), SteeringInput ), 4, YPos  );
 		YPos += Canvas->DrawText( RenderFont, FString::Printf( TEXT("Throttle: %f"), ThrottleInput ), 4, YPos );
 		YPos += Canvas->DrawText( RenderFont, FString::Printf( TEXT("Brake: %f"), BrakeInput ), 4, YPos  );
+		YPos += Canvas->DrawText(RenderFont, FString::Printf(TEXT("RPM: %f"), GetEngineRotationSpeed()), 4, YPos);
 		YPos += Canvas->DrawText(RenderFont, FString::Printf(TEXT("Gear: %d"), GetCurrentGear()), 4, YPos);
+		YPos += Canvas->DrawText(RenderFont, FString::Printf(TEXT("Drag: %.3f"), DebugDragMagnitude), 4, YPos);
 
 	}
 
