#ttp 332366: Assigning to new var Existing instead of using existing one (FStreamableManager::AsyncLoadCallback)

#chnage Removed unnecessary variable declaration

diff --git a/Engine/Source/Runtime/Engine/Private/StreamableManager.cpp b/Engine/Source/Runtime/Engine/Private/StreamableManager.cpp
--- a/Engine/Source/Runtime/Engine/Private/StreamableManager.cpp
+++ b/Engine/Source/Runtime/Engine/Private/StreamableManager.cpp
@@ -329,7 +329,7 @@ void FStreamableManager::AsyncLoadCallback(FStringAssetReference Request)
 	{
 		// hmm, maybe it was redirected by a consolidate
 		TargetName = ResolveRedirects(TargetName);
-		FStreamable* Existing = StreamableItems.FindRef(TargetName);
+		Existing = StreamableItems.FindRef(TargetName);
 	}
 	if (Existing && Existing->bAsyncLoadRequestOutstanding)
 	{
