#ue4
Fixed spurious warning when Calling AActor::SetSimulatePhysics() from a constructor

diff --git a/Engine/Source/Runtime/Engine/Private/PhysicsEngine/BodyInstance.cpp b/Engine/Source/Runtime/Engine/Private/PhysicsEngine/BodyInstance.cpp
--- a/Engine/Source/Runtime/Engine/Private/PhysicsEngine/BodyInstance.cpp
+++ b/Engine/Source/Runtime/Engine/Private/PhysicsEngine/BodyInstance.cpp
@@ -1406,9 +1406,9 @@ void FBodyInstance::SetInstanceSimulatePhysics(bool bSimulate, bool bMaintainPhy
 {
 #if WITH_PHYSX
 #if !(UE_BUILD_SHIPPING || UE_BUILD_TEST)
-	if (bSimulate)
+	if (bSimulate && OwnerComponent.IsValid())
 	{
-		PxRigidActor* PRigidActor = GetPxRigidActor();
+		PxRigidActor const* const PRigidActor = GetPxRigidActor();
 		if(PRigidActor == NULL)
 		{
 			FMessageLog("PIE").Warning(FText::Format(LOCTEXT("SimPhysNoBody", "Trying to simulate physics on ''{0}'' but no physics body."), 
