#ue4
Replaced ensure in UWorld::EncroachingBlockingGeometry with a log message

diff --git a/Engine/Source/Runtime/Engine/Private/LevelActor.cpp b/Engine/Source/Runtime/Engine/Private/LevelActor.cpp
--- a/Engine/Source/Runtime/Engine/Private/LevelActor.cpp
+++ b/Engine/Source/Runtime/Engine/Private/LevelActor.cpp
@@ -702,13 +702,15 @@ bool UWorld::EncroachingBlockingGeometry(AActor* TestActor, FVector TestLocation
 				FCollisionQueryParams Params(NAME_EncroachingBlockingGeometry, false, TestActor);
 				bFoundBlockingHit = OverlapMulti(Overlaps, TestLocation, FQuat::Identity, BlockingChannel, FCollisionShape::MakeSphere(FMath::Max(Sphere->GetScaledSphereRadius() - Epsilon, 0.1f)), Params);
 			}
+			else if (PrimComp->IsRegistered())
+			{
+				// must be registered
+				FComponentQueryParams Params(NAME_EncroachingBlockingGeometry, TestActor);
+				bFoundBlockingHit = ComponentOverlapMulti(Overlaps, PrimComp, TestLocation, TestActor->GetActorRotation(), Params);
+			}
 			else
 			{
-				if(ensureMsgf(PrimComp->IsRegistered(), TEXT("Components must be registered in order to be used in a ComponentOverlapMulti call. PriComp: %s TestActor: %s"), *PrimComp->GetName(), *TestActor->GetName()))
-				{
-					FComponentQueryParams Params(NAME_EncroachingBlockingGeometry, TestActor);
-					bFoundBlockingHit = ComponentOverlapMulti(Overlaps, PrimComp, TestLocation, TestActor->GetActorRotation(), Params);
-				}
+				UE_LOG(LogPhysics, Log, TEXT("Components must be registered in order to be used in a ComponentOverlapMulti call. PriComp: %s TestActor: %s"), *PrimComp->GetName(), *TestActor->GetName());
 			}
 		}
 
