#ue4

- Make ShouldCatchAir() more flexible by passing in the old floor and new floor rather than just the normals.
- Clean up some code in the floor adjustment calls within PhysWalking.

diff --git a/Engine/Source/Runtime/Engine/Classes/GameFramework/Character.h b/Engine/Source/Runtime/Engine/Classes/GameFramework/Character.h
--- a/Engine/Source/Runtime/Engine/Classes/GameFramework/Character.h
+++ b/Engine/Source/Runtime/Engine/Classes/GameFramework/Character.h
@@ -270,6 +270,7 @@ public:
 	virtual void Jump();
 
 	/**
+	 * Check if the character can jump in the current state.
 	 * Default implementation returns true if the character is on the ground and not crouching,
 	 * has a valid CharacterMovementComponent and CanEverJump() returns true.
 	 * @Return Whether the character can jump in the current state. 
diff --git a/Engine/Source/Runtime/Engine/Classes/GameFramework/CharacterMovementComponent.h b/Engine/Source/Runtime/Engine/Classes/GameFramework/CharacterMovementComponent.h
--- a/Engine/Source/Runtime/Engine/Classes/GameFramework/CharacterMovementComponent.h
+++ b/Engine/Source/Runtime/Engine/Classes/GameFramework/CharacterMovementComponent.h
@@ -654,10 +654,12 @@ public:
 	/** Transition from walking to falling */
 	virtual void StartFalling(int32 Iterations, float remainingTime, float timeTick, const FVector& Delta, const FVector& subLoc);
 
-	/** Whether pawn should go into falling mode when starting down a steep decline.
-	 *  @returns false to provide traditional default behavior
-	  */
-	virtual bool ShouldCatchAir(const FVector& OldFloor, const FVector& Floor);
+	/**
+	 * Whether Character should go into falling mode when walking and changing position, based on an old and new floor result (both of which are considered walkable).
+	 * Default implementation always returns false.
+	 * @return true if Character should start falling
+	 */
+	virtual bool ShouldCatchAir(const FFindFloorResult& OldFloor, const FFindFloorResult& NewFloor);
 
 	/** Adjust distance from floor, trying to maintain a slight offset from the floor when walking (based on CurrentFloor). */
 	virtual void AdjustFloorHeight();
@@ -909,9 +911,9 @@ public:
 	float GetValidPerchRadius() const;
 
 	/** Return true if the hit result should be considered a walkable surface for the character. */
+	UFUNCTION(BlueprintCallable, Category="Pawn|Components|CharacterMovement")
 	virtual bool IsWalkable(const FHitResult& Hit) const;
 
-
 	/** Get the max angle in degrees of a walkable surface for the character. */
 	FORCEINLINE float GetWalkableFloorAngle() const { return WalkableFloorAngle; }
 
diff --git a/Engine/Source/Runtime/Engine/Private/CharacterMovementComponent.cpp b/Engine/Source/Runtime/Engine/Private/CharacterMovementComponent.cpp
--- a/Engine/Source/Runtime/Engine/Private/CharacterMovementComponent.cpp
+++ b/Engine/Source/Runtime/Engine/Private/CharacterMovementComponent.cpp
@@ -2989,27 +2989,14 @@ void UCharacterMovementComponent::PhysWalking(float deltaTime, int32 Iterations)
 			// Validate the floor check
 			if (CurrentFloor.IsWalkableFloor())
 			{
-				const bool bBaseChanged = (CurrentFloor.HitResult.Component != CharacterOwner->GetMovementBase());
-				if (bBaseChanged || CurrentFloor.FloorDist > MAX_FLOOR_DIST)
+				if (ShouldCatchAir(OldFloor, CurrentFloor))
 				{
-					if ( ShouldCatchAir(OldFloor.HitResult.ImpactNormal, CurrentFloor.HitResult.ImpactNormal) )
-					{
-						StartFalling(Iterations, remainingTime, timeTick, Delta, subLoc);
-						return;
-					}
-					else
-					{
-						AdjustFloorHeight();
-						if (bBaseChanged)
-						{
-							SetBase(CurrentFloor.HitResult.Component.Get());
-						}
-					}
-				}
-				else if (CurrentFloor.FloorDist < MIN_FLOOR_DIST)
-				{
-					AdjustFloorHeight();
+					StartFalling(Iterations, remainingTime, timeTick, Delta, subLoc);
+					return;
 				}
+
+				AdjustFloorHeight();
+				SetBase(CurrentFloor.HitResult.Component.Get());
 			}
 			else if (CurrentFloor.HitResult.bStartPenetrating && remainingTime <= 0.f)
 			{
@@ -3064,7 +3051,7 @@ void UCharacterMovementComponent::PhysCustom(float deltaTime, int32 Iterations)
 }
 
 
-bool UCharacterMovementComponent::ShouldCatchAir(const FVector& OldFloor, const FVector& Floor)
+bool UCharacterMovementComponent::ShouldCatchAir(const FFindFloorResult& OldFloor, const FFindFloorResult& NewFloor)
 {
 	return false;
 }
