Fix inequality error in mouse smoothing code introduced in ... 2005!

diff --git a/Engine/Source/Runtime/Engine/Private/UserInterface/PlayerInput.cpp b/Engine/Source/Runtime/Engine/Private/UserInterface/PlayerInput.cpp
--- a/Engine/Source/Runtime/Engine/Private/UserInterface/PlayerInput.cpp
+++ b/Engine/Source/Runtime/Engine/Private/UserInterface/PlayerInput.cpp
@@ -756,7 +756,7 @@ void UPlayerInput::ProcessInputStack(const TArray<UInputComponent*>& InputCompon
 		{
 		 	// calculate sampling time
 		 	// make sure not first non-zero sample
-		 	if ( SmoothedMouse[0] > 0 )
+		 	if ( SmoothedMouse[0] != 0 )
 		 	{
 		 		// not first non-zero
 		 		MouseSamplingTotal += DeltaTime;
