#ue4

- Change IsValidLandingSpot() to not always return true when in penetration. Rather it is better to let FindFloor() handle the penetration more intelligently (using a shorter and slightly more narrow capsule when appropriate).
- Fixed an issue in ComputeFloorDist where the fallback sweep (in the case of penetration or rejected adjacent hits) was using the old HitResult if the sweep didn't hit anything.

Fixes TTP 332761: users can climb BSP walls in some cases

diff --git a/Engine/Source/Runtime/Engine/Private/CharacterMovementComponent.cpp b/Engine/Source/Runtime/Engine/Private/CharacterMovementComponent.cpp
--- a/Engine/Source/Runtime/Engine/Private/CharacterMovementComponent.cpp
+++ b/Engine/Source/Runtime/Engine/Private/CharacterMovementComponent.cpp
@@ -3637,6 +3637,7 @@ void UCharacterMovementComponent::ComputeFloorDist(const FVector& CapsuleLocatio
 				TraceDist = SweepDistance + ShrinkHeight;
 				CapsuleShape.Capsule.Radius = FMath::Max(0.f, CapsuleShape.Capsule.Radius - SWEEP_EDGE_REJECT_DISTANCE - KINDA_SMALL_NUMBER);
 				CapsuleShape.Capsule.HalfHeight = FMath::Max(PawnHalfHeight - ShrinkHeight, 0.1f);
+				Hit.Reset(1.f, false);
 				bBlockingHit = GetWorld()->SweepSingle(Hit, CapsuleLocation, CapsuleLocation + FVector(0.f,0.f,-TraceDist), FQuat::Identity, CollisionChannel, CapsuleShape, QueryParams, ResponseParam);
 			}
 
@@ -3817,26 +3818,24 @@ bool UCharacterMovementComponent::IsValidLandingSpot(const FVector& CapsuleLocat
 		return false;
 	}
 
-	// This can happen when landing on upward moving geometry
-	if (Hit.bStartPenetrating)
+	// Skip some checks if penetrating. Penetration will be handled by the FindFloor call (using a smaller capsule)
+	if (!Hit.bStartPenetrating)
 	{
-		return true;
-	}
+		float PawnRadius, PawnHalfHeight;
+		CharacterOwner->CapsuleComponent->GetScaledCapsuleSize(PawnRadius, PawnHalfHeight);
 
-	float PawnRadius, PawnHalfHeight;
-	CharacterOwner->CapsuleComponent->GetScaledCapsuleSize(PawnRadius, PawnHalfHeight);
-
-	// Reject hits that are above our lower hemisphere (can happen when sliding down a vertical surface).
-	const float LowerHemisphereZ = Hit.Location.Z - PawnHalfHeight + PawnRadius;
-	if (Hit.ImpactPoint.Z >= LowerHemisphereZ)
-	{
-		return false;
-	}
+		// Reject hits that are above our lower hemisphere (can happen when sliding down a vertical surface).
+		const float LowerHemisphereZ = Hit.Location.Z - PawnHalfHeight + PawnRadius;
+		if (Hit.ImpactPoint.Z >= LowerHemisphereZ)
+		{
+			return false;
+		}
 
-	// Reject hits that are barely on the cusp of the radius of the capsule
-	if (!IsWithinEdgeTolerance(Hit.Location, Hit.ImpactPoint, PawnRadius))
-	{
-		return false;
+		// Reject hits that are barely on the cusp of the radius of the capsule
+		if (!IsWithinEdgeTolerance(Hit.Location, Hit.ImpactPoint, PawnRadius))
+		{
+			return false;
+		}
 	}
 
 	FFindFloorResult FloorResult;
