Fixed: TTP#332387 LIVE: EDITOR: CRASH: Editor crashes when 2 client game is started in PIE on ShooterEntry.umap

diff --git a/Engine/Source/Editor/UnrealEd/Private/Editor.cpp b/Engine/Source/Editor/UnrealEd/Private/Editor.cpp
--- a/Engine/Source/Editor/UnrealEd/Private/Editor.cpp
+++ b/Engine/Source/Editor/UnrealEd/Private/Editor.cpp
@@ -6301,12 +6301,18 @@ void UEditorEngine::OnLevelRemovedFromWorld(ULevel* InLevel, UWorld* InWorld)
 			}
 		}
 	}
+	else
+	{
+		// UEngine::LoadMap broadcast this event with InLevel==NULL, before cleaning up the world
+		// Reset transactions buffer, to ensure that there are no references to a world which is about to be destroyed
+		ResetTransaction( NSLOCTEXT("UnrealEd", "LoadMapTransReset", "Loading a New Map") );
+	}
 }
 
 void UEditorEngine::OnGCStreamedOutLevels()
 {
 	// Reset transaction buffer because it may hold references to streamed out levels
-	ResetTransaction( NSLOCTEXT("UnrealEd", "GCStreamedOutLevels", "Garabage Collected Streamed Out Levels") );
+	ResetTransaction( NSLOCTEXT("UnrealEd", "GCStreamedOutLevelsTransReset", "GC Streaming Levels") );
 }
 
 void UEditorEngine::UpdateRecentlyLoadedProjectFiles()
diff --git a/Engine/Source/Runtime/Engine/Private/UnrealEngine.cpp b/Engine/Source/Runtime/Engine/Private/UnrealEngine.cpp
--- a/Engine/Source/Runtime/Engine/Private/UnrealEngine.cpp
+++ b/Engine/Source/Runtime/Engine/Private/UnrealEngine.cpp
@@ -9894,7 +9894,38 @@ bool UEngine::LoadMap( FWorldContext& WorldContext, FURL URL, class UPendingNetG
 			FLazyObjectPtr::ResetPIEFixups();
 
 			NewWorld = UWorld::DuplicateWorldForPIE(SourceWorldPackage, NULL);
-			WorldPackage = CastChecked<UPackage>(NewWorld->GetOuter());
+			if (NewWorld == nullptr) 
+			{
+				// Load map from the disk in case editor does not have it
+				const FString PIEPackageName = *UWorld::ConvertToPIEPackageName(SourceWorldPackage, WorldContext.PIEInstance);
+
+				// Set the world type in the static map, so that UWorld::PostLoad can set the world type
+				UWorld::WorldTypePreLoadMap.FindOrAdd( FName(*PIEPackageName) ) = WorldContext.WorldType;
+
+				WorldPackage = LoadPackage(CreatePackage(NULL, *PIEPackageName), *SourceWorldPackage, LOAD_None);
+				if (WorldPackage == nullptr)
+				{
+					Error = FString::Printf(TEXT("Failed to load package '%s' while in PIE"), *SourceWorldPackage);
+					return false;
+				}
+
+				NewWorld = UWorld::FindWorldInPackage(WorldPackage);
+				check(NewWorld);
+#if WITH_EDITOR
+				WorldPackage->PIEInstanceID = WorldContext.PIEInstance;
+#endif			
+				// Rename streaming levels to PIE
+				for (auto StreamingLevel : NewWorld->StreamingLevels)
+				{
+					StreamingLevel->RenameForPIE(WorldContext.PIEInstance);
+				}
+			}
+			else
+			{
+				WorldPackage = CastChecked<UPackage>(NewWorld->GetOuter());
+			}
+
+			NewWorld->StreamingLevelsPrefix = UWorld::BuildPIEPackagePrefix(WorldContext.PIEInstance);
 			GIsPlayInEditorWorld = true;
 		}
 	}
@@ -9916,50 +9947,52 @@ bool UEngine::LoadMap( FWorldContext& WorldContext, FURL URL, class UPendingNetG
 			WorldPackage = LoadPackage(MapOuter, *URL.Map, LOAD_None);
 		}
 
-	if( WorldPackage == NULL )
-	{
-		// it is now the responsibility of the caller to deal with a NULL return value and alert the user if necessary
-		Error = FString::Printf(TEXT("Failed to load package '%s'"), *URL.Map);
-		return false;
-	}
+		if( WorldPackage == NULL )
+		{
+			// it is now the responsibility of the caller to deal with a NULL return value and alert the user if necessary
+			Error = FString::Printf(TEXT("Failed to load package '%s'"), *URL.Map);
+			return false;
+		}
 		
 		FScopeCycleCounterUObject MapScope(WorldPackage);
 
-	if( FPlatformProperties::RequiresCookedData() && GUseSeekFreeLoading && !(WorldPackage->PackageFlags & PKG_DisallowLazyLoading) )
-	{
-		UE_LOG(LogLoad, Fatal,TEXT("Map '%s' has not been cooked correctly! Most likely stale version on the XDK."),*WorldPackage->GetName());
-	}
+		if( FPlatformProperties::RequiresCookedData() && GUseSeekFreeLoading && !(WorldPackage->PackageFlags & PKG_DisallowLazyLoading) )
+		{
+			UE_LOG(LogLoad, Fatal,TEXT("Map '%s' has not been cooked correctly! Most likely stale version on the XDK."),*WorldPackage->GetName());
+		}
 
 		// Find the newly loaded world.
-	NewWorld = UWorld::FindWorldInPackage(WorldPackage);
-	check(NewWorld);
+		NewWorld = UWorld::FindWorldInPackage(WorldPackage);
+		check(NewWorld);
 
-	if (WorldContext.WorldType == EWorldType::PIE)
-	{
-		// If we are a PIE world and the world we just found is already initialized, then we're probably reloading the editor world and we
-		// need to create a PIE world by duplication instead
-		if (NewWorld->bIsWorldInitialized)
-		{
-			NewWorld = CreatePIEWorldByDuplication(WorldContext, NewWorld, URL.Map);
-			// CreatePIEWorldByDuplication clears GIsPlayInEditorWorld so set it again
-			GIsPlayInEditorWorld = true;
-		}
-		// Otherwise we are probably loading new map while in PIE, so we need to rename world package and all streaming levels
-		else if (Pending == NULL)
+		if (WorldContext.WorldType == EWorldType::PIE)
 		{
+			// If we are a PIE world and the world we just found is already initialized, then we're probably reloading the editor world and we
+			// need to create a PIE world by duplication instead
+			if (NewWorld->bIsWorldInitialized)
+			{
+				NewWorld = CreatePIEWorldByDuplication(WorldContext, NewWorld, URL.Map);
+				// CreatePIEWorldByDuplication clears GIsPlayInEditorWorld so set it again
+				GIsPlayInEditorWorld = true;
+			}
+			// Otherwise we are probably loading new map while in PIE, so we need to rename world package and all streaming levels
+			else if (Pending == NULL)
+			{
 #if WITH_EDITOR
-			WorldPackage->PIEInstanceID = WorldContext.PIEInstance;
+				WorldPackage->PIEInstanceID = WorldContext.PIEInstance;
 #endif				
-			const FString PIEPackageName = *UWorld::ConvertToPIEPackageName(WorldPackage->GetName(), WorldContext.PIEInstance);
+				const FString PIEPackageName = *UWorld::ConvertToPIEPackageName(WorldPackage->GetName(), WorldContext.PIEInstance);
 			
-			WorldPackage->Rename(*PIEPackageName);
-			for (int32 StreamingLevelIdx = 0; StreamingLevelIdx < NewWorld->StreamingLevels.Num(); StreamingLevelIdx++)
-			{
-				NewWorld->StreamingLevels[StreamingLevelIdx]->RenameForPIE(WorldContext.PIEInstance);
+				WorldPackage->Rename(*PIEPackageName);
+				for (int32 StreamingLevelIdx = 0; StreamingLevelIdx < NewWorld->StreamingLevels.Num(); StreamingLevelIdx++)
+				{
+					NewWorld->StreamingLevels[StreamingLevelIdx]->RenameForPIE(WorldContext.PIEInstance);
+				}
+				
+				NewWorld->StreamingLevelsPrefix = UWorld::BuildPIEPackagePrefix(WorldContext.PIEInstance);
 			}
 		}
 	}
-	}
 
 	GWorld = NewWorld;
 
