#ttp 300819 - EDITOR: SKEL: SkeletalMesh LODs not changing in level
#branch UE4
#summary Enabled tick in editor for skeletal meshes. The pose/animation updates are for the most part skipped as the world isn't playing but we get the LOD updates and bone transform updates (required for LOD updates). Changed recreation after BP compile to correctly initialise the skeletal mesh component for a ticking mesh.

diff --git a/Engine/Source/Editor/UnrealEd/Private/Kismet2/KismetReinstanceUtilities.cpp b/Engine/Source/Editor/UnrealEd/Private/Kismet2/KismetReinstanceUtilities.cpp
--- a/Engine/Source/Editor/UnrealEd/Private/Kismet2/KismetReinstanceUtilities.cpp
+++ b/Engine/Source/Editor/UnrealEd/Private/Kismet2/KismetReinstanceUtilities.cpp
@@ -385,7 +385,12 @@ void FBlueprintCompileReinstancer::ReplaceInstancesOfClass(UClass* OldClass, UCl
 
 				if (UAnimInstance* AnimTree = Cast<UAnimInstance>(NewObject))
 				{
-					AnimTree->InitializeAnimation();
+					// Initialising the anim instance isn't enough to correctly set up the skeletal mesh again in a
+					// paused world, need to initialise the skeletal mesh component that contains the anim instance.
+					if(USkeletalMeshComponent* SkelComponent = Cast<USkeletalMeshComponent>(AnimTree->GetOuter()))
+					{
+						SkelComponent->InitAnim(true);
+					}
 				}
 
 				OldObject->RemoveFromRoot();
diff --git a/Engine/Source/Runtime/Engine/Private/SkeletalMeshComponent.cpp b/Engine/Source/Runtime/Engine/Private/SkeletalMeshComponent.cpp
--- a/Engine/Source/Runtime/Engine/Private/SkeletalMeshComponent.cpp
+++ b/Engine/Source/Runtime/Engine/Private/SkeletalMeshComponent.cpp
@@ -63,6 +63,8 @@ USkeletalMeshComponent::USkeletalMeshComponent(const class FPostConstructInitial
 	PrimaryComponentTick.TickGroup = TG_ParallelAnimWork;
 	PrimaryComponentTick.bRunOnAnyThread = true;	
 #endif
+
+	bTickInEditor = true;
 }
 
 
@@ -446,6 +448,15 @@ bool USkeletalMeshComponent::UpdateLODStatus()
 
 bool USkeletalMeshComponent::ShouldUpdateTransform(bool bLODHasChanged) const
 {
+#if WITH_EDITOR
+	// If we're in an editor world (Non running, WorldType will be PIE when simulating or in PIE) then we only want transform updates on LOD changes as the
+	// animation isn't running so it would just waste CPU time
+	if(GetWorld()->WorldType == EWorldType::Editor && !bLODHasChanged)
+	{
+		return false;
+	}
+#endif
+
 	// If forcing RefPose we can skip updating the skeleton for perf, except if it's using MorphTargets.
 	const bool bSkipBecauseOfRefPose = bForceRefpose && bOldForceRefPose && (MorphTargetCurves.Num() == 0) && ((AnimScriptInstance)? AnimScriptInstance->MorphTargetCurves.Num() == 0 : true);
 
