- Added a new cloth collision type which enables to interact between the character and other clothing objects, like a curtain.
- Fixed potential memory leaks of clothing resources

diff --git a/Engine/Source/Runtime/Engine/Classes/Components/SkeletalMeshComponent.h b/Engine/Source/Runtime/Engine/Classes/Components/SkeletalMeshComponent.h
--- a/Engine/Source/Runtime/Engine/Classes/Components/SkeletalMeshComponent.h
+++ b/Engine/Source/Runtime/Engine/Classes/Components/SkeletalMeshComponent.h
@@ -95,13 +95,17 @@ public:
 /** to interact with environments */
 struct FApexClothCollisionInfo
 {
-	enum CollisionInfoState
+	enum OverlappedComponentType
 	{
-		CIS_INVALID,
-		CIS_ADDED,
-		CIS_RETAINED
+		// static objects
+		OCT_STATIC,
+		// clothing objects
+		OCT_CLOTH,
+		OCT_MAX
 	};
 
+	OverlappedComponentType OverlapCompType;
+	// to verify validation of collision info
 	uint32 Revision;
 	// ClothingCollisions will be all released when clothing doesn't intersect with this component anymore
 	TArray<physx::apex::NxClothingCollision*> ClothingCollisions;			
@@ -326,7 +330,6 @@ public:
 	 */
 	UPROPERTY(EditAnywhere, BlueprintReadWrite, Category=Clothing)
 	float TeleportRotationThreshold;
-	
 
 	/** Draw the APEX Clothing Normals on clothing sections. */
 	uint32 bDisplayClothingNormals:1;
@@ -600,6 +603,7 @@ public:
 	void GetUpdateClothSimulationData(TArray<FClothSimulData>& OutClothSimData);
 	void ApplyWindForCloth(FClothingActor& ClothingActor);
 	void RemoveAllClothingActors();
+	void ReleaseAllClothingResources();
 
 	bool IsValidClothingActor(int32 ActorIndex) const;
 	/** Draws APEX Clothing simulated normals on cloth meshes **/
@@ -927,23 +931,28 @@ public:
 	/** create new collisions when newly added  */
 	FApexClothCollisionInfo* CreateNewClothingCollsions(UPrimitiveComponent* PrimitiveComponent);
 
-
-	void RemoveAllOverappedComponentMap();
+	void RemoveAllOverlappedComponentMap();
+	/** for non-static collisions which need to be updated every tick */ 
+	void UpdateOverlappedComponent(UPrimitiveComponent* PrimComp, FApexClothCollisionInfo* Info);
 
 	void ReleaseAllParentCollisions();
 	void ReleaseAllChildrenCollisions();
 
 	void ProcessClothCollisionWithEnvironment();
-
-	void CopyCollisionsToChildren();
-	// copy children's collisions to parent, where parent means this component
-	void CopyChildrenCollisionsToParent();
+	/** copy parent's cloth collisions to attached children, where parent means this component */
+	void CopyClothCollisionsToChildren();
+	/** copy children's cloth collisions to parent, where parent means this component */
+	void CopyChildrenClothCollisionsToParent();
 
 	/**
-  	 * Get collision data only for collision with clothes.
+  	 * Get collision data from a static mesh only for collision with clothes.
 	 * Returns false when failed to get cloth collision data
 	*/
-	bool GetClothCollisionData(UPrimitiveComponent* PrimComp, TArray<FClothCollisionPrimitive>& ClothCollisionPrimitives);
+	bool GetClothCollisionDataFromStaticMesh(UPrimitiveComponent* PrimComp, TArray<FClothCollisionPrimitive>& ClothCollisionPrimitives);
+	/** find if this component has collisions for clothing and return the results calculated by bone transforms */
+	void FindClothCollisions(TArray<FApexClothCollisionVolumeData>& OutCollisions);
+	/** create Apex clothing collisions from input collision info and add them to clothing actors */
+	void CreateInternalClothCollisions(TArray<FApexClothCollisionVolumeData>& InCollisions, TArray<physx::apex::NxClothingCollision*>& OutCollisions);
 
 #endif // WITH_CLOTH_COLLISION_DETECTION
 
diff --git a/Engine/Source/Runtime/Engine/Private/SkeletalMeshComponent.cpp b/Engine/Source/Runtime/Engine/Private/SkeletalMeshComponent.cpp
--- a/Engine/Source/Runtime/Engine/Private/SkeletalMeshComponent.cpp
+++ b/Engine/Source/Runtime/Engine/Private/SkeletalMeshComponent.cpp
@@ -132,7 +132,7 @@ void USkeletalMeshComponent::OnUnregister()
 {
 #if WITH_APEX_CLOTHING
 	//clothing actors will be re-created in TickClothing
-	RemoveAllClothingActors();
+	ReleaseAllClothingResources();
 #endif// #if WITH_APEX_CLOTHING
 
 	Super::OnUnregister();
@@ -207,7 +207,7 @@ void USkeletalMeshComponent::CreateRenderState_Concurrent()
 
 #if WITH_APEX_CLOTHING
 	//clothing actors will be re-created in TickClothing
-	RemoveAllClothingActors();
+	ReleaseAllClothingResources();
 #endif// #if WITH_APEX_CLOTHING
 
 	Super::CreateRenderState_Concurrent();
diff --git a/Engine/Source/Runtime/Engine/Private/SkeletalMeshComponentPhysics.cpp b/Engine/Source/Runtime/Engine/Private/SkeletalMeshComponentPhysics.cpp
--- a/Engine/Source/Runtime/Engine/Private/SkeletalMeshComponentPhysics.cpp
+++ b/Engine/Source/Runtime/Engine/Private/SkeletalMeshComponentPhysics.cpp
@@ -766,17 +766,6 @@ void USkeletalMeshComponent::TermArticulated()
 		Aggregate = NULL;
 	}
 #endif //WITH_PHYSX
-
-#if WITH_CLOTH_COLLISION_DETECTION
-	ReleaseAllParentCollisions();
-	ReleaseAllChildrenCollisions();
-//should be called before removing clothing actors
-	RemoveAllOverappedComponentMap();
-#endif // #if WITH_CLOTH_COLLISION_DETECTION
-
-#if WITH_APEX_CLOTHING
-	RemoveAllClothingActors();
-#endif// #if WITH_APEX_CLOTHING
 }
 
 void USkeletalMeshComponent::TermBodiesBelow(FName ParentBoneName)
@@ -1902,7 +1891,7 @@ bool USkeletalMeshComponent::CreateClothingActor(int32 AssetIndex, TSharedPtr<FC
 
 	FVector ScaleVector = ComponentToWorld.GetScale3D();
 
-	//support only uniform scale
+	// support only uniform scale
 	verify(NxParameterized::setParamF32(*ActorDesc, "actorScale", ScaleVector.X));
 
 	bool bUseInternalBoneOrder = true;
@@ -2051,9 +2040,23 @@ void USkeletalMeshComponent::RemoveAllClothingActors()
 	ClothingActors.Empty();
 }
 
+void USkeletalMeshComponent::ReleaseAllClothingResources()
+{
+#if WITH_CLOTH_COLLISION_DETECTION
+	ReleaseAllParentCollisions();
+	ReleaseAllChildrenCollisions();
+	// should be called before removing clothing actors
+	RemoveAllOverlappedComponentMap();
+#endif // #if WITH_CLOTH_COLLISION_DETECTION
+
+#if WITH_APEX_CLOTHING
+	RemoveAllClothingActors();
+#endif// #if WITH_APEX_CLOTHING
+}
+
 void USkeletalMeshComponent::ApplyWindForCloth(FClothingActor& ClothingActor)
 {
-	//to convert from normalized value( usually 0.0 to 1.0 ) to Apex clohting wind value
+	//to convert from normalized value( usually 0.0 to 1.0 ) to Apex clothing wind value
 	const float WindUnitAmout = 2500.0f;
 
 	NxClothingActor* ApexClothingActor = ClothingActor.ApexClothingActor;
@@ -2162,50 +2165,84 @@ void USkeletalMeshComponent::DrawDebugClothCollisions()
 {
 	FColor Colors[6] = { FColor::Red, FColor::Green, FColor::Blue, FColor::Cyan, FColor::Yellow, FColor::Magenta };
 
-	for( auto It = ClothOverlappedComponentsMap.CreateConstIterator(); It; ++It )
+	for (auto It = ClothOverlappedComponentsMap.CreateConstIterator(); It; ++It)
 	{
 		TWeakObjectPtr<UPrimitiveComponent> PrimComp = It.Key();
 
 		TArray<FClothCollisionPrimitive> CollisionPrims;
-		if(!GetClothCollisionData(PrimComp.Get(), CollisionPrims))
+		ECollisionChannel Channel = PrimComp->GetCollisionObjectType();
+		if (Channel == ECollisionChannel::ECC_WorldStatic)
 		{
-			return;
-		}
+			if (!GetClothCollisionDataFromStaticMesh(PrimComp.Get(), CollisionPrims))
+			{
+				return;
+			}
 
-		for(int32 PrimIndex=0; PrimIndex < CollisionPrims.Num(); PrimIndex++)
-		{
-			switch(CollisionPrims[PrimIndex].PrimType)
+			for (int32 PrimIndex = 0; PrimIndex < CollisionPrims.Num(); PrimIndex++)
 			{
-			case FClothCollisionPrimitive::SPHERE:
-				DrawDebugSphere(GetWorld(), CollisionPrims[PrimIndex].Origin, CollisionPrims[PrimIndex].Radius, 10, FColor::Red);
-				break;
+				switch (CollisionPrims[PrimIndex].PrimType)
+				{
+				case FClothCollisionPrimitive::SPHERE:
+					DrawDebugSphere(GetWorld(), CollisionPrims[PrimIndex].Origin, CollisionPrims[PrimIndex].Radius, 10, FColor::Red);
+					break;
 
-			case FClothCollisionPrimitive::CAPSULE:
+				case FClothCollisionPrimitive::CAPSULE:
 				{
 					FVector DiffVec = CollisionPrims[PrimIndex].SpherePos2 - CollisionPrims[PrimIndex].SpherePos1;
 					float HalfHeight = DiffVec.Size() * 0.5f;
 					FQuat Rotation = PrimComp->ComponentToWorld.GetRotation();
-					DrawDebugCapsule(GetWorld(), CollisionPrims[PrimIndex].Origin, HalfHeight, CollisionPrims[PrimIndex].Radius, Rotation, FColor::Red );
+					DrawDebugCapsule(GetWorld(), CollisionPrims[PrimIndex].Origin, HalfHeight, CollisionPrims[PrimIndex].Radius, Rotation, FColor::Red);
 				}
 
-				break;
-			case FClothCollisionPrimitive::CONVEX:	
-				DrawDebugConvexFromPlanes(CollisionPrims[PrimIndex], Colors[PrimIndex%6]);
-				break;
+					break;
+				case FClothCollisionPrimitive::CONVEX:
+					DrawDebugConvexFromPlanes(CollisionPrims[PrimIndex], Colors[PrimIndex % 6]);
+					break;
+				}
+
+				//draw a bounding box for double checking
+				DrawDebugBox(GetWorld(), PrimComp->Bounds.Origin, PrimComp->Bounds.BoxExtent, FColor::Red);
 			}
+		}
+		else if (Channel == ECollisionChannel::ECC_PhysicsBody) // supports an interaction between the character and other clothing such as a curtain
+		{
+			bool bCreatedCollisions = false;
+			// if a component is a skeletal mesh component with clothing, add my collisions to the component
+			USkeletalMeshComponent* SkelMeshComp = Cast<USkeletalMeshComponent>(PrimComp.Get());
+			if (SkelMeshComp && SkelMeshComp->SkeletalMesh)
+			{
+				if (SkelMeshComp->ClothingActors.Num() > 0)
+				{
+					TArray<FApexClothCollisionVolumeData> NewCollisions;
+					FindClothCollisions(NewCollisions);
 
-			//draw a bounding box for double checking
-			DrawDebugBox(GetWorld(), PrimComp->Bounds.Origin, PrimComp->Bounds.BoxExtent, FColor::Red );
+					for (int32 ColIdx = 0; ColIdx < NewCollisions.Num(); ColIdx++)
+					{
+						if (NewCollisions[ColIdx].IsCapsule())
+						{
+							FVector Origin = NewCollisions[ColIdx].LocalPose.GetOrigin();
+							// apex uses y-axis as the up-axis of capsule
+							FVector UpAxis = NewCollisions[ColIdx].LocalPose.GetScaledAxis(EAxis::Y);
+							float Radius = NewCollisions[ColIdx].CapsuleRadius*UpAxis.Size();
+							float HalfHeight = NewCollisions[ColIdx].CapsuleHeight*0.5f;
+
+							FMatrix Pose = NewCollisions[ColIdx].LocalPose;
+							FMatrix RotMat(Pose.GetScaledAxis(EAxis::X), Pose.GetScaledAxis(EAxis::Z), Pose.GetScaledAxis(EAxis::Y), FVector(0,0,0));
+							FQuat Rotation = RotMat.ToQuat();
+							DrawDebugCapsule(GetWorld(), Origin, HalfHeight, Radius, Rotation, FColor::Red);
+						}
+					}
+				}
+			}
 		}
 	}
-
 	//draw this skeletal mesh component's bounding box
 
 	DrawDebugBox(GetWorld(), Bounds.Origin, Bounds.BoxExtent, FColor::Red);
 
 }
 
-bool USkeletalMeshComponent::GetClothCollisionData(UPrimitiveComponent* PrimComp, TArray<FClothCollisionPrimitive>& ClothCollisionPrimitives)
+bool USkeletalMeshComponent::GetClothCollisionDataFromStaticMesh(UPrimitiveComponent* PrimComp, TArray<FClothCollisionPrimitive>& ClothCollisionPrimitives)
 {
 	//make sure Num of collisions should be 0 in the case this returns false
 	ClothCollisionPrimitives.Empty();
@@ -2215,14 +2252,14 @@ bool USkeletalMeshComponent::GetClothCollisionData(UPrimitiveComponent* PrimComp
 		return false;
 	}
 
-	if(!PrimComp->BodyInstance.IsValidBodyInstance())
+	ECollisionChannel Channel = PrimComp->GetCollisionObjectType();
+	// accept only a static mesh
+	if (Channel != ECollisionChannel::ECC_WorldStatic)
 	{
 		return false;
 	}
 
-	ECollisionChannel Channel = PrimComp->GetCollisionObjectType();
-	//now supports for interacting only between clothes and static objects
-	if(Channel != ECollisionChannel::ECC_WorldStatic)
+	if (!PrimComp->BodyInstance.IsValidBodyInstance())
 	{
 		return false;
 	}
@@ -2354,25 +2391,23 @@ bool USkeletalMeshComponent::GetClothCollisionData(UPrimitiveComponent* PrimComp
 			}
 			break;
 		}
-	}
-
+	}		
 	return true;
 }
 
-static void FindClothCollisions(USkeletalMeshComponent* InSkelMeshComp, TArray<FApexClothCollisionVolumeData>& OutCollisions)
+void USkeletalMeshComponent::FindClothCollisions(TArray<FApexClothCollisionVolumeData>& OutCollisions)
 {
-	if(!InSkelMeshComp || !InSkelMeshComp->SkeletalMesh)
+	if(!SkeletalMesh)
 	{
 		return;
 	}
 
-	USkeletalMesh* SkelMesh = InSkelMeshComp->SkeletalMesh;
-	int32 NumAssets = SkelMesh->ClothingAssets.Num();
+	int32 NumAssets = SkeletalMesh->ClothingAssets.Num();
 
 	// find all new collisions passing to children
 	for(int32 AssetIdx=0; AssetIdx < NumAssets; AssetIdx++)
 	{
-		FClothingAssetData& Asset = SkelMesh->ClothingAssets[AssetIdx];
+		FClothingAssetData& Asset = SkeletalMesh->ClothingAssets[AssetIdx];
 		int32 NumCollisions = Asset.ClothCollisionVolumes.Num();
 
 		for(int32 ColIdx=0; ColIdx < NumCollisions; ColIdx++)
@@ -2386,14 +2421,14 @@ static void FindClothCollisions(USkeletalMeshComponent* InSkelMeshComp, TArray<F
 
 			FName BoneName = Asset.ApexClothingAsset->GetConvertedBoneName(Collision.BoneIndex);
 
-			int32 BoneIndex = InSkelMeshComp->GetBoneIndex(BoneName);
+			int32 BoneIndex = GetBoneIndex(BoneName);
 
 			if(BoneIndex < 0)
 			{
 				continue;
 			}
 
-			FMatrix BoneMat = InSkelMeshComp->GetBoneMatrix(BoneIndex);
+			FMatrix BoneMat = GetBoneMatrix(BoneIndex);
 
 			FMatrix LocalToWorld = Collision.LocalPose * BoneMat;
 
@@ -2408,24 +2443,24 @@ static void FindClothCollisions(USkeletalMeshComponent* InSkelMeshComp, TArray<F
 	}
 }
 
-static void CreateAndAddCollisions(USkeletalMeshComponent* InSkelMeshComp, TArray<FApexClothCollisionVolumeData>& InCollisions, TArray<physx::apex::NxClothingCollision*>& OutCollisions)
+void USkeletalMeshComponent::CreateInternalClothCollisions(TArray<FApexClothCollisionVolumeData>& InCollisions, TArray<physx::apex::NxClothingCollision*>& OutCollisions)
 {
 	int32 NumCollisions = InCollisions.Num();
 
 	const int32 MaxNumCapsules = 16;
 	// because sphere number can not be larger than 32 and 1 capsule takes 2 spheres 
 	NumCollisions = FMath::Min(NumCollisions, MaxNumCapsules);
-	int32 NumActors = InSkelMeshComp->ClothingActors.Num();
+	int32 NumActors = ClothingActors.Num();
 
 	for (int32 ActorIdx = 0; ActorIdx < NumActors; ActorIdx++)
 	{
-		if (!InSkelMeshComp->IsValidClothingActor(ActorIdx))
+		if (!IsValidClothingActor(ActorIdx))
 		{
 			continue;
 		}
 
-		NxClothingActor* Actor = InSkelMeshComp->ClothingActors[ActorIdx].ApexClothingActor;
-		int32 NumCurrentCapsules = InSkelMeshComp->SkeletalMesh->ClothingAssets[ActorIdx].ClothCollisionVolumes.Num(); // # of capsules 
+		NxClothingActor* Actor = ClothingActors[ActorIdx].ApexClothingActor;
+		int32 NumCurrentCapsules = SkeletalMesh->ClothingAssets[ActorIdx].ClothCollisionVolumes.Num(); // # of capsules 
 
 		for (int32 ColIdx = 0; ColIdx < NumCollisions; ColIdx++)
 		{
@@ -2454,7 +2489,7 @@ static void CreateAndAddCollisions(USkeletalMeshComponent* InSkelMeshComp, TArra
 	}
 }
 
-void USkeletalMeshComponent::CopyCollisionsToChildren()
+void USkeletalMeshComponent::CopyClothCollisionsToChildren()
 {
 	// 3 steps
 	// 1. release all previous parent collisions
@@ -2488,11 +2523,11 @@ void USkeletalMeshComponent::CopyCollisionsToChildren()
 
 	TArray<FApexClothCollisionVolumeData> NewCollisions;
 
-	FindClothCollisions(this, NewCollisions);
+	FindClothCollisions(NewCollisions);
 
 	for(int32 ChildIdx=0; ChildIdx < NumClothChildren; ChildIdx++)
 	{
-		CreateAndAddCollisions(ClothChildren[ChildIdx], NewCollisions, ClothChildren[ChildIdx]->ParentCollisions);
+		ClothChildren[ChildIdx]->CreateInternalClothCollisions(NewCollisions, ClothChildren[ChildIdx]->ParentCollisions);
 	}
 }
 
@@ -2507,7 +2542,7 @@ void USkeletalMeshComponent::ReleaseAllChildrenCollisions()
 }
 
 // children's collisions can affect to parent's cloth reversely
-void USkeletalMeshComponent::CopyChildrenCollisionsToParent()
+void USkeletalMeshComponent::CopyChildrenClothCollisionsToParent()
 {
 	// 3 steps
 	// 1. release all previous children collisions
@@ -2526,11 +2561,11 @@ void USkeletalMeshComponent::CopyChildrenCollisionsToParent()
 		USkeletalMeshComponent* pChild = Cast<USkeletalMeshComponent>(AttachChildren[i]);
 		if(pChild)
 		{
-			FindClothCollisions(pChild, NewCollisions);
+			pChild->FindClothCollisions(NewCollisions);
 		}
 	}
 
-	CreateAndAddCollisions(this, NewCollisions, ChildrenCollisions);
+	CreateInternalClothCollisions(NewCollisions, ChildrenCollisions);
 }
 
 void USkeletalMeshComponent::ReleaseClothingCollision(NxClothingCollision* Collision)
@@ -2568,81 +2603,123 @@ FApexClothCollisionInfo* USkeletalMeshComponent::CreateNewClothingCollsions(UPri
 	FApexClothCollisionInfo NewInfo;
 
 	TArray<FClothCollisionPrimitive> CollisionPrims;
-	if(!GetClothCollisionData(PrimitiveComponent, CollisionPrims))
-	{
-		return NULL;
-	}
 
-	int32 NumActors = ClothingActors.Num();
+	ECollisionChannel Channel = PrimitiveComponent->GetCollisionObjectType();
 
-	for(int32 ActorIdx=0; ActorIdx < NumActors; ActorIdx++)
+	// crate new clothing collisions for a static mesh
+	if (Channel == ECollisionChannel::ECC_WorldStatic)
 	{
-		NxClothingActor* Actor = ClothingActors[ActorIdx].ApexClothingActor;
+		if (!GetClothCollisionDataFromStaticMesh(PrimitiveComponent, CollisionPrims))
+		{
+			return NULL;
+		}
 
-		if(Actor)
-		{	
-			for(int32 PrimIndex=0; PrimIndex < CollisionPrims.Num(); PrimIndex++)
-			{
-				NxClothingCollision* ClothCol = NULL;
+		NewInfo.OverlapCompType = FApexClothCollisionInfo::OCT_STATIC;
 
-				switch(CollisionPrims[PrimIndex].PrimType)
+		int32 NumActors = ClothingActors.Num();
+
+		for (int32 ActorIdx = 0; ActorIdx < NumActors; ActorIdx++)
+		{
+			NxClothingActor* Actor = ClothingActors[ActorIdx].ApexClothingActor;
+
+			if (Actor)
+			{
+				for (int32 PrimIndex = 0; PrimIndex < CollisionPrims.Num(); PrimIndex++)
 				{
-				case FClothCollisionPrimitive::SPHERE:
-					ClothCol = Actor->createCollisionSphere(U2PVector(CollisionPrims[PrimIndex].Origin), CollisionPrims[PrimIndex].Radius);
-					if(ClothCol)
+					NxClothingCollision* ClothCol = NULL;
+
+					switch (CollisionPrims[PrimIndex].PrimType)
 					{
-						NewInfo.ClothingCollisions.Add(ClothCol);
-					}
-					break;
-				case FClothCollisionPrimitive::CAPSULE:
+					case FClothCollisionPrimitive::SPHERE:
+						ClothCol = Actor->createCollisionSphere(U2PVector(CollisionPrims[PrimIndex].Origin), CollisionPrims[PrimIndex].Radius);
+						if (ClothCol)
+						{
+							NewInfo.ClothingCollisions.Add(ClothCol);
+						}
+						break;
+					case FClothCollisionPrimitive::CAPSULE:
 					{
 						float Radius = CollisionPrims[PrimIndex].Radius;
 						NxClothingSphere* Sphere1 = Actor->createCollisionSphere(U2PVector(CollisionPrims[PrimIndex].SpherePos1), Radius);
 						NxClothingSphere* Sphere2 = Actor->createCollisionSphere(U2PVector(CollisionPrims[PrimIndex].SpherePos2), Radius);
 
 						ClothCol = Actor->createCollisionCapsule(*Sphere1, *Sphere2);
-						if(ClothCol)
+						if (ClothCol)
 						{
 							NewInfo.ClothingCollisions.Add(ClothCol);
 						}
 					}
+						break;
+					case FClothCollisionPrimitive::CONVEX:
 
-					break;
-				case FClothCollisionPrimitive::CONVEX:
+						int32 NumPlanes = CollisionPrims[PrimIndex].ConvexPlanes.Num();
 
-					int32 NumPlanes = CollisionPrims[PrimIndex].ConvexPlanes.Num();
+						TArray<NxClothingPlane*> ClothingPlanes;
 
-					TArray<NxClothingPlane*> ClothingPlanes;
+						//can not exceed 32 planes
+						NumPlanes = FMath::Min(NumPlanes, 32);
 
-					//can not exceed 32 planes
-					NumPlanes = FMath::Min(NumPlanes, 32);
+						ClothingPlanes.AddUninitialized(NumPlanes);
 
-					ClothingPlanes.AddUninitialized(NumPlanes);
+						for (int32 PlaneIdx = 0; PlaneIdx < NumPlanes; PlaneIdx++)
+						{
+							PxPlane PPlane = U2PPlane(CollisionPrims[PrimIndex].ConvexPlanes[PlaneIdx]);
 
-					for(int32 PlaneIdx=0; PlaneIdx < NumPlanes; PlaneIdx++)
-					{
-						PxPlane PPlane = U2PPlane(CollisionPrims[PrimIndex].ConvexPlanes[PlaneIdx]);
+							ClothingPlanes[PlaneIdx] = Actor->createCollisionPlane(PPlane);
+						}
 
-						ClothingPlanes[PlaneIdx] = Actor->createCollisionPlane(PPlane);
-					}
+						ClothCol = Actor->createCollisionConvex(ClothingPlanes.GetTypedData(), ClothingPlanes.Num());
 
-					ClothCol = Actor->createCollisionConvex(ClothingPlanes.GetTypedData(), ClothingPlanes.Num());
+						if (ClothCol)
+						{
+							NewInfo.ClothingCollisions.Add(ClothCol);
+						}
 
-					if(ClothCol)
-					{
-						NewInfo.ClothingCollisions.Add(ClothCol);
+						break;
 					}
+				}
+			}
+		}
+	}
+	else if (Channel == ECollisionChannel::ECC_PhysicsBody) // creates new clothing collisions for clothing objects
+	{
+		bool bCreatedCollisions = false;
+		// if a component is a skeletal mesh component with clothing, add my collisions to the component
+		USkeletalMeshComponent* SkelMeshComp = Cast<USkeletalMeshComponent>(PrimitiveComponent);
 
-					break;
+		if (SkelMeshComp && SkelMeshComp->SkeletalMesh)
+		{
+			// if a casted skeletal mesh component is myself, then skip
+			if (SkelMeshComp == this)
+			{
+				return NULL;
+			}
+
+			if (SkelMeshComp->ClothingActors.Num() > 0)
+			{
+				TArray<FApexClothCollisionVolumeData> NewCollisions;
+				// get collision infos computed by current bone transforms
+				FindClothCollisions(NewCollisions);
+				
+				if (NewCollisions.Num() > 0)
+				{
+					NewInfo.OverlapCompType = FApexClothCollisionInfo::OCT_CLOTH;
+					SkelMeshComp->CreateInternalClothCollisions(NewCollisions, NewInfo.ClothingCollisions);
+					bCreatedCollisions = true;
 				}
 			}
 		}
+
+		if (bCreatedCollisions == false)
+		{
+			return NULL;
+		}
 	}
 
 	return &ClothOverlappedComponentsMap.Add(PrimitiveComponent, NewInfo);
 }
 
-void USkeletalMeshComponent::RemoveAllOverappedComponentMap()
+void USkeletalMeshComponent::RemoveAllOverlappedComponentMap()
 {
 	for( auto It = ClothOverlappedComponentsMap.CreateConstIterator(); It; ++It )
 	{
@@ -2669,6 +2746,32 @@ void USkeletalMeshComponent::ReleaseAllParentCollisions()
 	ParentCollisions.Empty();
 }
 
+void USkeletalMeshComponent::UpdateOverlappedComponent(UPrimitiveComponent* PrimComp, FApexClothCollisionInfo* Info)
+{
+	if (Info->OverlapCompType == FApexClothCollisionInfo::OCT_CLOTH)
+	{
+		int32 NumCollisions = Info->ClothingCollisions.Num();
+		for (int32 i = 0; i < NumCollisions; i++)
+		{
+			ReleaseClothingCollision(Info->ClothingCollisions[i]);
+		}
+
+		Info->ClothingCollisions.Empty(NumCollisions);
+
+		TArray<FApexClothCollisionVolumeData> NewCollisions;
+		FindClothCollisions(NewCollisions);
+
+		if (NewCollisions.Num() > 0)
+		{
+			USkeletalMeshComponent *SkelMeshComp = Cast<USkeletalMeshComponent>(PrimComp);
+			if (SkelMeshComp)
+			{
+				SkelMeshComp->CreateInternalClothCollisions(NewCollisions, Info->ClothingCollisions);
+			}
+		}
+	}
+}
+
 void USkeletalMeshComponent::ProcessClothCollisionWithEnvironment()
 {
 	// don't handle collision detection if this component is in editor
@@ -2685,6 +2788,8 @@ void USkeletalMeshComponent::ProcessClothCollisionWithEnvironment()
 	FCollisionObjectQueryParams ObjectParams;
 
 	ObjectParams.AddObjectTypesToQuery(ECollisionChannel::ECC_WorldStatic);
+	// to collide with other clothing objects
+	ObjectParams.AddObjectTypesToQuery(ECollisionChannel::ECC_PhysicsBody);
 	
 	static FName ClothOverlapComponentsName(TEXT("ClothOverlapComponents"));
 	FCollisionQueryParams Params(ClothOverlapComponentsName, false);
@@ -2708,6 +2813,8 @@ void USkeletalMeshComponent::ProcessClothCollisionWithEnvironment()
 			if(Info)
 			{
 				Info->Revision = ClothingCollisionRevision;
+				// if this component needs update every tick 
+				UpdateOverlappedComponent(Component.Get(), Info);
 			}
 		}
 	}
@@ -2742,7 +2849,7 @@ void USkeletalMeshComponent::UpdateClothTransform()
 
 	if(bCollideWithAttachedChildren)
 	{
-		CopyCollisionsToChildren();
+		CopyClothCollisionsToChildren();
 	}
 
 	//check the environment when only transform is updated
@@ -2863,8 +2970,8 @@ void USkeletalMeshComponent::UpdateClothState(float DeltaTime)
 
 	if(bCollideWithAttachedChildren)
 	{
-		CopyCollisionsToChildren();
-		CopyChildrenCollisionsToParent();
+		CopyClothCollisionsToChildren();
+		CopyChildrenClothCollisionsToParent();
 	}
 #endif // WITH_CLOTH_COLLISION_DETECTION
 
