No black background for "show debug" if we are drawing bones

diff --git a/Engine/Source/Runtime/Engine/Private/HUD.cpp b/Engine/Source/Runtime/Engine/Private/HUD.cpp
--- a/Engine/Source/Runtime/Engine/Private/HUD.cpp
+++ b/Engine/Source/Runtime/Engine/Private/HUD.cpp
@@ -264,8 +264,11 @@ bool AHUD::ShouldDisplayDebug(const FName & DebugType) const
 
 void AHUD::ShowDebugInfo(float& YL, float& YPos)
 {
-	FLinearColor BackgroundColor(0.f, 0.f, 0.f, 0.5f);
-	DebugCanvas->Canvas->DrawTile(0, 0, DebugCanvas->ClipX, DebugCanvas->ClipY, 0.f, 0.f, 0.f, 0.f, BackgroundColor);
+	if (!DebugDisplay.Contains(TEXT("Bones")))
+	{
+		FLinearColor BackgroundColor(0.f, 0.f, 0.f, 0.5f);
+		DebugCanvas->Canvas->DrawTile(0, 0, DebugCanvas->ClipX, DebugCanvas->ClipY, 0.f, 0.f, 0.f, 0.f, BackgroundColor);
+	}
 
 	FDebugDisplayInfo DisplayInfo(DebugDisplay,ToggledDebugCategories);
 	PlayerOwner->PlayerCameraManager->ViewTarget.Target->DisplayDebug(DebugCanvas, DisplayInfo, YL, YPos);
