Ensure that maps being loaded for PIE are flagged correctly from the beginning.

diff --git a/Engine/Source/Editor/UnrealEd/Private/PlayLevel.cpp b/Engine/Source/Editor/UnrealEd/Private/PlayLevel.cpp
--- a/Engine/Source/Editor/UnrealEd/Private/PlayLevel.cpp
+++ b/Engine/Source/Editor/UnrealEd/Private/PlayLevel.cpp
@@ -2744,11 +2744,9 @@ UWorld* UEditorEngine::CreatePIEWorldBySavingToTemp(FWorldContext &WorldContext,
 	UWorld::WorldTypePreLoadMap.FindOrAdd(FName(*SavedMapNames[0])) = EWorldType::PIE;
 
 	// Load the package we saved
-	UPackage* EditorLevelPackage = LoadPackage(NULL, *SavedMapNames[0], LOAD_None);
+	UPackage* EditorLevelPackage = LoadPackage(NULL, *SavedMapNames[0], LOAD_PackageForPIE);
 	if( EditorLevelPackage )
 	{
-		EditorLevelPackage->PackageFlags |= PKG_PlayInEditor;
-
 		// Find world object and use its PersistentLevel pointer.
 		LoadedWorld = UWorld::FindWorldInPackage(EditorLevelPackage);
 
diff --git a/Engine/Source/Runtime/CoreUObject/Private/UObject/Linker.cpp b/Engine/Source/Runtime/CoreUObject/Private/UObject/Linker.cpp
--- a/Engine/Source/Runtime/CoreUObject/Private/UObject/Linker.cpp
+++ b/Engine/Source/Runtime/CoreUObject/Private/UObject/Linker.cpp
@@ -654,6 +654,10 @@ ULinkerLoad* GetPackageLinker
 				
 			// Create the package with the provided long package name.
 			UPackage* FilenamePkg = CreatePackage( NULL, *PackageName );
+			if (LoadFlags & LOAD_PackageForPIE)
+			{
+				FilenamePkg->PackageFlags |= PKG_PlayInEditor;
+			}
 
 			// If no package specified, use package from file.
 			if( InOuter==NULL )
diff --git a/Engine/Source/Runtime/CoreUObject/Public/UObject/ObjectBase.h b/Engine/Source/Runtime/CoreUObject/Public/UObject/ObjectBase.h
--- a/Engine/Source/Runtime/CoreUObject/Public/UObject/ObjectBase.h
+++ b/Engine/Source/Runtime/CoreUObject/Public/UObject/ObjectBase.h
@@ -43,6 +43,7 @@ enum ELoadFlags
 	LOAD_NoRedirects				= 0x00010000,	// Never follow redirects when loading objects; redirected loads will fail
 	LOAD_ForDiff					= 0x00020000,	// Loading for diffing.
 	LOAD_NoSeekFreeLinkerDetatch	= 0x00040000,	// Do not detach linkers for this package when seek-free loading
+	LOAD_PackageForPIE				= 0x00080000,   // This package is being loaded for PIE, it must be flagged as such immediately
 };
 
 //
diff --git a/Engine/Source/Runtime/Engine/Private/UnrealEngine.cpp b/Engine/Source/Runtime/Engine/Private/UnrealEngine.cpp
--- a/Engine/Source/Runtime/Engine/Private/UnrealEngine.cpp
+++ b/Engine/Source/Runtime/Engine/Private/UnrealEngine.cpp
@@ -9944,7 +9944,7 @@ bool UEngine::LoadMap( FWorldContext& WorldContext, FURL URL, class UPendingNetG
 		// If the level isn't already in memory, load level from disk
 		if (WorldPackage == NULL)
 		{
-			WorldPackage = LoadPackage(MapOuter, *URL.Map, LOAD_None);
+			WorldPackage = LoadPackage(MapOuter, *URL.Map, (WorldContext.WorldType == EWorldType::PIE ? LOAD_PackageForPIE : LOAD_None));
 		}
 
 		if( WorldPackage == NULL )
@@ -10003,7 +10003,7 @@ bool UEngine::LoadMap( FWorldContext& WorldContext, FURL URL, class UPendingNetG
 	// Also, dont add to root when in PIE, since PIE doesn't remove world from root
 	if (WorldContext.WorldType == EWorldType::PIE)
 	{
-		CastChecked<UPackage>(WorldContext.World()->GetOutermost())->PackageFlags |= PKG_PlayInEditor;
+		check((CastChecked<UPackage>(WorldContext.World()->GetOutermost())->PackageFlags & PKG_PlayInEditor) == PKG_PlayInEditor);
 		WorldContext.World()->ClearFlags(RF_Standalone);
 	}
 	else
