Expose enable FOV smoothing and the scale to use to input project settings

diff --git a/Engine/Config/BaseInput.ini b/Engine/Config/BaseInput.ini
--- a/Engine/Config/BaseInput.ini
+++ b/Engine/Config/BaseInput.ini
@@ -1,5 +1,7 @@
 [/Script/Engine.InputSettings]
 bEnableMouseSmoothing=true
+bEnableFOVScaling=true
+FOVScale=0.01111
 DoubleClickTime=0.2f
 
 DefaultTouchInterface=/Engine/MobileResources/HUD/DefaultVirtualJoysticks.DefaultVirtualJoysticks
diff --git a/Engine/Source/Runtime/Engine/Classes/GameFramework/InputSettings.h b/Engine/Source/Runtime/Engine/Classes/GameFramework/InputSettings.h
--- a/Engine/Source/Runtime/Engine/Classes/GameFramework/InputSettings.h
+++ b/Engine/Source/Runtime/Engine/Classes/GameFramework/InputSettings.h
@@ -19,9 +19,17 @@ class ENGINE_API UInputSettings : public UObject
 
 	// Mouse smoothing control
 	UPROPERTY(config, EditAnywhere, Category="MouseProperties", AdvancedDisplay)
-	uint32 bEnableMouseSmoothing:1;    /** if true, mouse smoothing is enabled */
+	uint32 bEnableMouseSmoothing:1;
 
-	/** If a button is pressed twice in this amount of time, it's considered a "double click" */
+	// Scale the mouse based on the player camera manager's field of view
+	UPROPERTY(config, EditAnywhere, Category="MouseProperties", AdvancedDisplay)
+	uint32 bEnableFOVScaling:1;
+
+	// The scaling value to multiply the field of view by
+	UPROPERTY(config, EditAnywhere, Category="MouseProperties", AdvancedDisplay, meta=(editcondition="bEnableFOVScaling"))
+	float FOVScale;
+
+	/** If a key is pressed twice in this amount of time, it's considered a "double click" */
 	UPROPERTY(config, EditAnywhere, Category="MouseProperties", AdvancedDisplay)
 	float DoubleClickTime;
 
diff --git a/Engine/Source/Runtime/Engine/Classes/GameFramework/PlayerInput.h b/Engine/Source/Runtime/Engine/Classes/GameFramework/PlayerInput.h
--- a/Engine/Source/Runtime/Engine/Classes/GameFramework/PlayerInput.h
+++ b/Engine/Source/Runtime/Engine/Classes/GameFramework/PlayerInput.h
@@ -234,10 +234,6 @@ public:
 	// @todo: We have 10 touches to match the number of Touch* entries in EKeys (not easy to make this an enum or define or anything)
 	FVector Touches[EKeys::NUM_TOUCH_KEYS];
 
-	// Zoom Scaling
-	UPROPERTY()
-	uint32 bEnableFOVScaling:1;
-
 	// Mouse smoothing sample data
 	float ZeroTime[2];    /** How long received mouse movement has been zero. */
 	float SmoothedMouse[2];    /** Current average mouse movement/sample */
diff --git a/Engine/Source/Runtime/Engine/Private/UserInterface/PlayerInput.cpp b/Engine/Source/Runtime/Engine/Private/UserInterface/PlayerInput.cpp
--- a/Engine/Source/Runtime/Engine/Private/UserInterface/PlayerInput.cpp
+++ b/Engine/Source/Runtime/Engine/Private/UserInterface/PlayerInput.cpp
@@ -30,7 +30,6 @@ UPlayerInput::UPlayerInput(const class FPostConstructInitializeProperties& PCIP)
 	SetFlags(RF_Transactional);
 	MouseSamplingTotal = +0.0083f;
 	MouseSamples = 1;
-	bEnableFOVScaling = true;
 }
 
 void UPlayerInput::PostInitProperties()
@@ -1277,9 +1276,11 @@ float UPlayerInput::MassageAxisInput(FKey Key, float RawValue, float DeltaTime)
 	// special handling for mouse input
 	if ( (Key == EKeys::MouseX) || (Key == EKeys::MouseY) )
 	{
+		const UInputSettings* DefaultInputSettings = GetDefault<UInputSettings>();
+
 		// Take FOV into account (lower FOV == less sensitivity).
 		APlayerController const* const PlayerController = GetOuterAPlayerController();
-		float const FOVScale = (bEnableFOVScaling && PlayerController->PlayerCameraManager) ? (0.01111f*PlayerController->PlayerCameraManager->GetFOVAngle()) : 1.0f;
+		float const FOVScale = (DefaultInputSettings->bEnableFOVScaling && PlayerController->PlayerCameraManager) ? (DefaultInputSettings->FOVScale*PlayerController->PlayerCameraManager->GetFOVAngle()) : 1.0f;
 		NewVal *= FOVScale;
 
 		// debug
@@ -1293,7 +1294,7 @@ float UPlayerInput::MassageAxisInput(FKey Key, float RawValue, float DeltaTime)
 		}
 
 		// mouse smoothing 
-		if (GetDefault<UInputSettings>()->bEnableMouseSmoothing)
+		if (DefaultInputSettings->bEnableMouseSmoothing)
 		{
 			FKeyState* const KeyState = KeyStateMap.Find(Key);
 			if (KeyState)
