Prevent recursive attempts to stop a sound if the right things are done in the OnAudioFinished callbacks

diff --git a/Engine/Source/Runtime/Engine/Private/AudioComponent.cpp b/Engine/Source/Runtime/Engine/Private/AudioComponent.cpp
--- a/Engine/Source/Runtime/Engine/Private/AudioComponent.cpp
+++ b/Engine/Source/Runtime/Engine/Private/AudioComponent.cpp
@@ -278,7 +278,10 @@ void UAudioComponent::Stop()
 
 void UAudioComponent::PlaybackCompleted(bool bFailedToStart)
 {
-	if (!bFailedToStart && GetWorld() != NULL && (OnAudioFinished.IsBound() || OnAudioFinishedNative.IsBound()) )
+	// Mark inactive before calling destroy to avoid recursion
+	bIsActive = false;
+
+	if (!bFailedToStart && GetWorld() != NULL && (OnAudioFinished.IsBound() || OnAudioFinishedNative.IsBound()))
 	{
 		INC_DWORD_STAT( STAT_AudioFinishedDelegatesCalled );
 		SCOPE_CYCLE_COUNTER( STAT_AudioFinishedDelegates );
@@ -287,9 +290,6 @@ void UAudioComponent::PlaybackCompleted(bool bFailedToStart)
 		OnAudioFinishedNative.Broadcast(this);
 	}
 
-	// Mark inactive before calling destroy to avoid recursion
-	bIsActive = false;
-
 	// Auto destruction is handled via marking object for deletion.
 	if( bAutoDestroy )
 	{
