ENGINE_API for ExponentialHeighFogComponent

diff --git a/Engine/Source/Runtime/Engine/Classes/Components/ExponentialHeightFogComponent.h b/Engine/Source/Runtime/Engine/Classes/Components/ExponentialHeightFogComponent.h
--- a/Engine/Source/Runtime/Engine/Classes/Components/ExponentialHeightFogComponent.h
+++ b/Engine/Source/Runtime/Engine/Classes/Components/ExponentialHeightFogComponent.h
@@ -7,8 +7,8 @@
 /**
  *	Used to create fogging effects such as clouds but with a density that is related to the height of the fog.
  */
-UCLASS(HeaderGroup=Component, ClassGroup=Rendering, collapsecategories, hidecategories=(Object, Mobility), editinlinenew, meta=(BlueprintSpawnableComponent),MinimalAPI)
-class UExponentialHeightFogComponent : public USceneComponent
+UCLASS(HeaderGroup=Component, ClassGroup=Rendering, collapsecategories, hidecategories=(Object, Mobility), editinlinenew, meta=(BlueprintSpawnableComponent))
+class ENGINE_API UExponentialHeightFogComponent : public USceneComponent
 {
 	GENERATED_UCLASS_BODY()
 
