Orion - NaN Tracking issue
- It seems montage instance weight becomes 0, but still active depending on tick, causing it to sneak in NaN there since it doesn't have any weight, it depends on what the value it has.

diff --git a/Engine/Source/Runtime/Engine/Private/Animation/AnimInstance.cpp b/Engine/Source/Runtime/Engine/Private/Animation/AnimInstance.cpp
--- a/Engine/Source/Runtime/Engine/Private/Animation/AnimInstance.cpp
+++ b/Engine/Source/Runtime/Engine/Private/Animation/AnimInstance.cpp
@@ -1215,7 +1215,7 @@ void UAnimInstance::SlotEvaluatePose(FName SlotNodeName, const FA2Pose & SourceP
 			FAnimMontageInstance * MontageInstance = (*Iter);
 
 			// @todo make this struct?
-			if (MontageInstance->IsValid() && MontageInstance->Montage->IsValidSlot(SlotNodeName))
+			if (MontageInstance->IsValid() && MontageInstance->Montage->IsValidSlot(SlotNodeName) && MontageInstance->Weight>ZERO_ANIMWEIGHT_THRESH)
 			{
 				int32 NewIndex = MontageWeights.AddUninitialized(1);
 				MontagePoses.AddZeroed(1);
