UE4: Consolidated duplicate implementations of ConvertIntToWindowMode in GameUserSettings.cpp and UnrealEngine.cpp into a single function in the EWindowMode namespace. This removes the need to use "extern" imports and fixes non-unity CIS.

diff --git a/Engine/Source/Runtime/Core/Public/GenericPlatform/GenericWindow.h b/Engine/Source/Runtime/Core/Public/GenericPlatform/GenericWindow.h
--- a/Engine/Source/Runtime/Core/Public/GenericPlatform/GenericWindow.h
+++ b/Engine/Source/Runtime/Core/Public/GenericPlatform/GenericWindow.h
@@ -18,6 +18,25 @@ namespace EWindowMode
 		/** The window has a border and may not take up the entire screen area */
 		Windowed
 	};
+
+	static inline Type ConvertIntToWindowMode(int32 InWindowMode)
+	{
+		Type WindowMode = Windowed;
+		switch (InWindowMode)
+		{
+			case 0:
+				WindowMode = Fullscreen;
+				break;
+			case 1:
+				WindowMode = WindowedFullscreen;
+				break;
+			case 2:
+			default:
+				WindowMode = Windowed;
+				break;
+		}
+		return WindowMode;
+	}
 }
 
 class CORE_API FGenericWindow
diff --git a/Engine/Source/Runtime/Engine/Private/GameEngine.cpp b/Engine/Source/Runtime/Engine/Private/GameEngine.cpp
--- a/Engine/Source/Runtime/Engine/Private/GameEngine.cpp
+++ b/Engine/Source/Runtime/Engine/Private/GameEngine.cpp
@@ -35,8 +35,6 @@ DEFINE_LOG_CATEGORY_STATIC(LogEngine, Log, All);
 
 ENGINE_API bool GDisallowNetworkTravel = false;
 
-extern EWindowMode::Type ConvertIntToWindowMode(int32 InWindowMode);
-
 /** Benchmark results to the log */
 static void RunSynthBenchmark(const TArray<FString>& Args)
 {
@@ -87,7 +85,7 @@ int32 GetBoundFullScreenModeCVar()
 // depending on WindowMode and the console variable r.FullScreenMode
 EWindowMode::Type GetWindowModeType(EWindowMode::Type WindowMode)
 {
-	EWindowMode::Type CurrentWindowMode = ConvertIntToWindowMode(GetBoundFullScreenModeCVar());
+	EWindowMode::Type CurrentWindowMode = EWindowMode::ConvertIntToWindowMode(GetBoundFullScreenModeCVar());
 
 	if (FPlatformProperties::SupportsWindowedMode())
 	{
diff --git a/Engine/Source/Runtime/Engine/Private/GameUserSettings.cpp b/Engine/Source/Runtime/Engine/Private/GameUserSettings.cpp
--- a/Engine/Source/Runtime/Engine/Private/GameUserSettings.cpp
+++ b/Engine/Source/Runtime/Engine/Private/GameUserSettings.cpp
@@ -36,33 +36,14 @@ void UGameUserSettings::SetScreenResolution( FIntPoint Resolution )
 	ResolutionSizeY = Resolution.Y;
 }
 
-inline EWindowMode::Type ConvertIntToWindowMode(int32 InWindowMode)
-{
-	EWindowMode::Type WindowMode = EWindowMode::Windowed;
-	switch ( InWindowMode )
-	{
-	case 0:
-		WindowMode = EWindowMode::Fullscreen;
-		break;
-	case 1:
-		WindowMode = EWindowMode::WindowedFullscreen;
-		break;
-	case 2:
-	default:
-		WindowMode = EWindowMode::Windowed;
-		break;
-	}
-	return WindowMode;
-}
-
 EWindowMode::Type UGameUserSettings::GetFullscreenMode() const
 {
-	return ConvertIntToWindowMode(FullscreenMode);
+	return EWindowMode::ConvertIntToWindowMode(FullscreenMode);
 }
 
 EWindowMode::Type UGameUserSettings::GetLastConfirmedFullscreenMode() const
 {
-	return ConvertIntToWindowMode(LastConfirmedFullscreenMode);
+	return EWindowMode::ConvertIntToWindowMode(LastConfirmedFullscreenMode);
 }
 
 void UGameUserSettings::SetFullscreenMode( EWindowMode::Type InFullscreenMode )
@@ -303,7 +284,7 @@ void UGameUserSettings::PreloadResolutionSettings()
 
 		int32 WindowModeInt = (int32)WindowMode;
 		GConfig->GetInt(*GameUserSettingsCategory, TEXT("FullscreenMode"), WindowModeInt, GGameUserSettingsIni);
-		WindowMode = ConvertIntToWindowMode(WindowModeInt);
+		WindowMode = EWindowMode::ConvertIntToWindowMode(WindowModeInt);
 
 		GConfig->GetInt(*GameUserSettingsCategory, TEXT("ResolutionSizeX"), ResolutionX, GGameUserSettingsIni);
 		GConfig->GetInt(*GameUserSettingsCategory, TEXT("ResolutionSizeY"), ResolutionY, GGameUserSettingsIni);
diff --git a/Engine/Source/Runtime/Engine/Private/UnrealEngine.cpp b/Engine/Source/Runtime/Engine/Private/UnrealEngine.cpp
--- a/Engine/Source/Runtime/Engine/Private/UnrealEngine.cpp
+++ b/Engine/Source/Runtime/Engine/Private/UnrealEngine.cpp
@@ -358,25 +358,6 @@ void ScalabilityCVarsSinkCallback()
 	}
 }
 
-inline EWindowMode::Type ConvertIntToWindowMode(int32 InWindowMode)
-{
-	EWindowMode::Type WindowMode = EWindowMode::Windowed;
-	switch (InWindowMode)
-	{
-	case 0:
-		WindowMode = EWindowMode::Fullscreen;
-		break;
-	case 1:
-		WindowMode = EWindowMode::WindowedFullscreen;
-		break;
-	case 2:
-	default:
-		WindowMode = EWindowMode::Windowed;
-		break;
-	}
-	return WindowMode;
-}
-
 void SystemResolutionSinkCallback()
 {
 	auto ResString = CVarSystemResolution->GetString();
@@ -386,7 +367,7 @@ void SystemResolutionSinkCallback()
 	
 	if (FParse::Resolution(*ResString, ResX, ResY, WindowModeInt))
 	{
-		EWindowMode::Type WindowMode = ConvertIntToWindowMode(WindowModeInt);
+		EWindowMode::Type WindowMode = EWindowMode::ConvertIntToWindowMode(WindowModeInt);
 
 		// TODO: This isn't correct, as we also need to compare the required fullscreen mode to the existing one.
 		if( GSystemResolution.ResX != ResX ||
