Cleanup FindActiveWaveInstance

diff --git a/Engine/Source/Runtime/Engine/Classes/Sound/SoundNode.h b/Engine/Source/Runtime/Engine/Classes/Sound/SoundNode.h
--- a/Engine/Source/Runtime/Engine/Classes/Sound/SoundNode.h
+++ b/Engine/Source/Runtime/Engine/Classes/Sound/SoundNode.h
@@ -168,5 +168,11 @@ public:
 		checkf(ChildIndex < MAX_ALLOWED_CHILD_NODES, TEXT("Too many children (%d) in SoundCue '%s'"), ChildIndex, *CastChecked<USoundCue>(ChildNode->GetOuter())->GetFullName());
 		return ((ParentWaveInstanceHash << ChildIndex) ^ (UPTRINT)ChildNode);
 	}
+
+	static FORCEINLINE UPTRINT GetNodeWaveInstanceHash(const UPTRINT ParentWaveInstanceHash, const UPTRINT ChildNodeHash, const uint32 ChildIndex)
+	{
+		checkf(ChildIndex < MAX_ALLOWED_CHILD_NODES, TEXT("Too many children (%d) in SoundCue"), ChildIndex);
+		return ((ParentWaveInstanceHash << ChildIndex) ^ ChildNodeHash);
+	}
 };
 
diff --git a/Engine/Source/Runtime/Engine/Private/ActiveSound.cpp b/Engine/Source/Runtime/Engine/Private/ActiveSound.cpp
--- a/Engine/Source/Runtime/Engine/Private/ActiveSound.cpp
+++ b/Engine/Source/Runtime/Engine/Private/ActiveSound.cpp
@@ -73,9 +73,9 @@ void FActiveSound::AddReferencedObjects( FReferenceCollector& Collector)
 #if WITH_EDITOR
 	if( GIsEditor )
 	{
-		for( int32 Index = 0; Index < WaveInstances.Num(); ++Index )
+		for (auto WaveInstanceIt(WaveInstances.CreateConstIterator()); WaveInstanceIt; ++WaveInstanceIt)
 		{
-			FWaveInstance* WaveInstance = WaveInstances[ Index ];
+			FWaveInstance* WaveInstance = WaveInstanceIt.Value();
 			// Avoid recursing back to the wave instance that sourced this active sound
 			if( WaveInstance )
 			{
@@ -223,9 +223,9 @@ void FActiveSound::Stop(FAudioDevice* AudioDevice)
 		Sound->CurrentPlayCount = FMath::Max( Sound->CurrentPlayCount - 1, 0 );
 	}
 
-	for( int32 Index = 0; Index < WaveInstances.Num(); ++Index )
+	for (auto WaveInstanceIt(WaveInstances.CreateIterator()); WaveInstanceIt; ++WaveInstanceIt)
 	{
-		FWaveInstance* WaveInstance = WaveInstances[Index];
+		FWaveInstance*& WaveInstance = WaveInstanceIt.Value();
 
 		// Stop the owning sound source
 		FSoundSource* Source = AudioDevice->WaveInstanceSourceMap.FindRef( WaveInstance );
@@ -241,7 +241,7 @@ void FActiveSound::Stop(FAudioDevice* AudioDevice)
 		delete WaveInstance;
 
 		// Null the entry out temporarily as later Stop calls could try to access this structure
-		WaveInstances[Index] = NULL;
+		WaveInstance = NULL;
 	}
 	WaveInstances.Empty();
 
@@ -254,18 +254,10 @@ void FActiveSound::Stop(FAudioDevice* AudioDevice)
 	AudioDevice->RemoveActiveSound(this);
 }
 
-FWaveInstance* FActiveSound::FindWaveInstance( const USoundWave* SoundWave, const UPTRINT WaveInstanceHash )
+FWaveInstance* FActiveSound::FindWaveInstance( const UPTRINT WaveInstanceHash )
 {
-	for( int32 WaveIndex = 0; WaveIndex < WaveInstances.Num(); ++WaveIndex )
-	{
-		FWaveInstance* ExistingWaveInstance = WaveInstances[ WaveIndex ];
-		if( ExistingWaveInstance && ExistingWaveInstance->WaveData == SoundWave && ExistingWaveInstance->WaveInstanceHash == WaveInstanceHash )
-		{
-			return ExistingWaveInstance;
-		}
-	}
-
-	return NULL;
+	FWaveInstance** WaveInstance = WaveInstances.Find(WaveInstanceHash);
+	return (WaveInstance ? *WaveInstance : NULL);
 }
 
 void FActiveSound::UpdateAdjustVolumeMultiplier( const float DeltaTime )
@@ -297,6 +289,13 @@ void FActiveSound::CheckOcclusion( const FVector ListenerLocation, const FVector
 	}
 }
 
+const TCHAR* GetAWaveName(TMap<UPTRINT, struct FWaveInstance*> WaveInstances)
+{
+	TArray<FWaveInstance*> WaveInstanceArray;
+	WaveInstances.GenerateValueArray(WaveInstanceArray);
+	return *WaveInstanceArray[0]->WaveData->GetName();
+}
+
 void FActiveSound::HandleInteriorVolumes( const FListener& Listener, FSoundParseParameters& ParseParams )
 {
 	// Get the settings of the ambient sound
@@ -334,7 +333,7 @@ void FActiveSound::HandleInteriorVolumes( const FListener& Listener, FSoundParse
 		ParseParams.HighFrequencyGain *= CurrentInteriorLPF;
 
 		UE_LOG(LogAudio, Verbose, TEXT( "Ambient in same volume. Volume *= %g LPF *= %g (%s)" ),
-			CurrentInteriorVolume, CurrentInteriorLPF, ( WaveInstances.Num() > 0 ) ? *WaveInstances[ 0 ]->WaveData->GetName() : TEXT( "NULL" ) );
+			CurrentInteriorVolume, CurrentInteriorLPF, ( WaveInstances.Num() > 0 ) ? GetAWaveName(WaveInstances) : TEXT( "NULL" ) );
 	}
 	else
 	{
@@ -349,7 +348,7 @@ void FActiveSound::HandleInteriorVolumes( const FListener& Listener, FSoundParse
 			ParseParams.HighFrequencyGain *= CurrentInteriorLPF;
 
 			UE_LOG(LogAudio, Verbose, TEXT( "Ambient in diff volume, ambient outside. Volume *= %g LPF *= %g (%s)" ),
-				CurrentInteriorVolume, CurrentInteriorLPF, ( WaveInstances.Num() > 0 ) ? *WaveInstances[ 0 ]->WaveData->GetName() : TEXT( "NULL" ) );
+				CurrentInteriorVolume, CurrentInteriorLPF, ( WaveInstances.Num() > 0 ) ? GetAWaveName(WaveInstances) : TEXT( "NULL" ) );
 		}
 		else
 		{
@@ -363,7 +362,7 @@ void FActiveSound::HandleInteriorVolumes( const FListener& Listener, FSoundParse
 			ParseParams.HighFrequencyGain *= CurrentInteriorLPF*CurrentExteriorLPF;
 
 			UE_LOG(LogAudio, Verbose, TEXT( "Ambient in diff volume, ambient inside. Volume *= %g LPF *= %g (%s)" ),
-				CurrentInteriorVolume*CurrentExteriorVolume, CurrentInteriorLPF*CurrentExteriorLPF, ( WaveInstances.Num() > 0 ) ? *WaveInstances[ 0 ]->WaveData->GetName() : TEXT( "NULL" ) );
+				CurrentInteriorVolume*CurrentExteriorVolume, CurrentInteriorLPF*CurrentExteriorLPF, ( WaveInstances.Num() > 0 ) ? GetAWaveName(WaveInstances) : TEXT( "NULL" ) );
 		}
 	}
 }
diff --git a/Engine/Source/Runtime/Engine/Private/SoundNodeLooping.cpp b/Engine/Source/Runtime/Engine/Private/SoundNodeLooping.cpp
--- a/Engine/Source/Runtime/Engine/Private/SoundNodeLooping.cpp
+++ b/Engine/Source/Runtime/Engine/Private/SoundNodeLooping.cpp
@@ -75,7 +75,7 @@ bool USoundNodeLooping::NotifyWaveInstanceFinished( FWaveInstance* InWaveInstanc
 		}
 		else if (ResetNode->IsA<USoundNodeWavePlayer>())
 		{
-			FWaveInstance* WaveInstance = ActiveSound->FindWaveInstance(CastChecked<USoundNodeWavePlayer>(ResetNode)->SoundWave, NodeHashPair.NodeWaveInstanceHash);
+			FWaveInstance* WaveInstance = ActiveSound->FindWaveInstance(NodeHashPair.NodeWaveInstanceHash);
 			if (WaveInstance)
 			{
 				WaveInstance->bAlreadyNotifiedHook = true;
diff --git a/Engine/Source/Runtime/Engine/Private/SoundNodeWaveParam.cpp b/Engine/Source/Runtime/Engine/Private/SoundNodeWaveParam.cpp
--- a/Engine/Source/Runtime/Engine/Private/SoundNodeWaveParam.cpp
+++ b/Engine/Source/Runtime/Engine/Private/SoundNodeWaveParam.cpp
@@ -24,7 +24,7 @@ void USoundNodeWaveParam::ParseNodes( FAudioDevice* AudioDevice, const UPTRINT N
 	ActiveSound.GetWaveParameter( WaveParameterName, NewWave );
 	if( NewWave != NULL )
 	{
-		NewWave->Parse( AudioDevice, NodeWaveInstanceHash, ActiveSound, ParseParams, WaveInstances );
+		NewWave->Parse( AudioDevice, GetNodeWaveInstanceHash(NodeWaveInstanceHash, (UPTRINT)NewWave, 0), ActiveSound, ParseParams, WaveInstances );
 	}
 	else
 	{
diff --git a/Engine/Source/Runtime/Engine/Private/SoundWave.cpp b/Engine/Source/Runtime/Engine/Private/SoundWave.cpp
--- a/Engine/Source/Runtime/Engine/Private/SoundWave.cpp
+++ b/Engine/Source/Runtime/Engine/Private/SoundWave.cpp
@@ -359,7 +359,7 @@ FWaveInstance* USoundWave::HandleStart( FActiveSound& ActiveSound, const UPTRINT
 	// Create a new wave instance and associate with the ActiveSound
 	FWaveInstance* WaveInstance = new FWaveInstance( &ActiveSound );
 	WaveInstance->WaveInstanceHash = WaveInstanceHash;
-	ActiveSound.WaveInstances.Add( WaveInstance );
+	ActiveSound.WaveInstances.Add( WaveInstanceHash, WaveInstance );
 
 	// Add in the subtitle if they exist
 	if (ActiveSound.bHandleSubtitles && Subtitles.Num() > 0)
@@ -400,7 +400,7 @@ void USoundWave::FinishDestroy()
 
 void USoundWave::Parse( FAudioDevice* AudioDevice, const UPTRINT NodeWaveInstanceHash, FActiveSound& ActiveSound, const FSoundParseParameters& ParseParams, TArray<FWaveInstance*>& WaveInstances )
 {
-	FWaveInstance* WaveInstance = ActiveSound.FindWaveInstance( this, NodeWaveInstanceHash);
+	FWaveInstance* WaveInstance = ActiveSound.FindWaveInstance(NodeWaveInstanceHash);
 
 	// Create a new WaveInstance if this SoundWave doesn't already have one associated with it.
 	if( WaveInstance == NULL )
diff --git a/Engine/Source/Runtime/Engine/Public/ActiveSound.h b/Engine/Source/Runtime/Engine/Public/ActiveSound.h
--- a/Engine/Source/Runtime/Engine/Public/ActiveSound.h
+++ b/Engine/Source/Runtime/Engine/Public/ActiveSound.h
@@ -115,7 +115,7 @@ public:
 	float CurrentInteriorVolume;
 	float CurrentInteriorLPF;
 
-	TArray<struct FWaveInstance*> WaveInstances;
+	TMap<UPTRINT, struct FWaveInstance*> WaveInstances;
 
 	TMap<UPTRINT,uint32> SoundNodeOffsetMap;
 	TArray<uint8> SoundNodeData;
@@ -130,7 +130,7 @@ public:
 	/** 
 	 * Find an existing waveinstance attached to this audio component (if any)
 	 */
-	FWaveInstance* FindWaveInstance( const USoundWave* SoundWave, const UPTRINT WaveInstanceHash );
+	FWaveInstance* FindWaveInstance(const UPTRINT WaveInstanceHash);
 
 	/** 
 	 * Check whether to apply the radio filter
