Added mark dirty when skeleton hierarchy changes

diff --git a/Engine/Source/Runtime/Engine/Private/Animation/Skeleton.cpp b/Engine/Source/Runtime/Engine/Private/Animation/Skeleton.cpp
--- a/Engine/Source/Runtime/Engine/Private/Animation/Skeleton.cpp
+++ b/Engine/Source/Runtime/Engine/Private/Animation/Skeleton.cpp
@@ -509,6 +509,7 @@ bool USkeleton::MergeBonesToBoneTree(USkeletalMesh * InSkeletalMesh, const TArra
 
 					ReferenceSkeleton.Add(NewMeshBoneInfo, InSkeletalMesh->RefSkeleton.GetRefBonePose()[MeshBoneIndex]);
 					BoneTree.AddZeroed(1);
+					MarkPackageDirty();
 				}
 			}
 
