Slot node is the #1 cause of NaN cause in animation system, so adding NaN check there.

diff --git a/Engine/Source/Runtime/Engine/Private/Animation/AnimNode_Slot.cpp b/Engine/Source/Runtime/Engine/Private/Animation/AnimNode_Slot.cpp
--- a/Engine/Source/Runtime/Engine/Private/Animation/AnimNode_Slot.cpp
+++ b/Engine/Source/Runtime/Engine/Private/Animation/AnimNode_Slot.cpp
@@ -51,6 +51,9 @@ void FAnimNode_Slot::Evaluate(FPoseContext& Output)
 	}
 
 	Output.AnimInstance->SlotEvaluatePose(SlotName, SourceContext.Pose, Output.Pose, SlotNodeWeight);
+
+	checkSlow(!Output.ContainsNaN());
+	checkSlow(Output.IsNormalized());
 }
 
 
