Better ENGINE_API exposure for Sound classes

diff --git a/Engine/Source/Editor/UnrealEd/Classes/Factories/Factory.h b/Engine/Source/Editor/UnrealEd/Classes/Factories/Factory.h
--- a/Engine/Source/Editor/UnrealEd/Classes/Factories/Factory.h
+++ b/Engine/Source/Editor/UnrealEd/Classes/Factories/Factory.h
@@ -55,7 +55,7 @@ class UFactory : public UObject
 
 	/** For interactive object imports, this value indicates whether the user wants objects to be automatically
 		overwritten (See EAppReturnType), or -1 if the user should be prompted. */
-	static int32 OverwriteYesOrNoToAllState;
+	UNREALED_API static int32 OverwriteYesOrNoToAllState;
 
 	/** If this value is true, warning messages will be shown once for all objects being imported at the same time.  
 		This value will be reset to false each time a new import operation is started. */
diff --git a/Engine/Source/Runtime/Engine/Classes/Sound/SoundBase.h b/Engine/Source/Runtime/Engine/Classes/Sound/SoundBase.h
--- a/Engine/Source/Runtime/Engine/Classes/Sound/SoundBase.h
+++ b/Engine/Source/Runtime/Engine/Classes/Sound/SoundBase.h
@@ -25,8 +25,8 @@ namespace EMaxConcurrentResolutionRule
 	};
 }
 
-UCLASS(config=Engine, hidecategories=Object, abstract, editinlinenew, MinimalAPI, BlueprintType)
-class USoundBase : public UObject
+UCLASS(config=Engine, hidecategories=Object, abstract, editinlinenew, BlueprintType)
+class ENGINE_API USoundBase : public UObject
 {
 	GENERATED_UCLASS_BODY()
 
@@ -75,12 +75,12 @@ public:
 	/**
 	 * Checks to see if a location is audible
 	 */
-	ENGINE_API bool IsAudible( const FVector& SourceLocation, const FVector& ListenerLocation, AActor* SourceActor, bool& bIsOccluded, bool bCheckOcclusion );
+	bool IsAudible( const FVector& SourceLocation, const FVector& ListenerLocation, AActor* SourceActor, bool& bIsOccluded, bool bCheckOcclusion );
 
 	/** 
 	 * Does a simple range check to all listeners to test hearability
 	 */
-	ENGINE_API bool IsAudibleSimple( const FVector Location );
+	bool IsAudibleSimple( const FVector Location );
 
 	/** 
 	 * Returns the farthest distance at which the sound could be heard
@@ -98,6 +98,6 @@ public:
 	/** 
 	 * Parses the Sound to generate the WaveInstances to play
 	 */
-	virtual void Parse( class FAudioDevice* AudioDevice, const UPTRINT NodeWaveInstanceHash, FActiveSound& AudioComponent, const FSoundParseParameters& ParseParams, TArray<FWaveInstance*>& WaveInstances ) { }
+	virtual void Parse( class FAudioDevice* AudioDevice, const UPTRINT NodeWaveInstanceHash, FActiveSound& ActiveSound, const FSoundParseParameters& ParseParams, TArray<FWaveInstance*>& WaveInstances ) { }
 };
 
diff --git a/Engine/Source/Runtime/Engine/Classes/Sound/SoundWave.h b/Engine/Source/Runtime/Engine/Classes/Sound/SoundWave.h
--- a/Engine/Source/Runtime/Engine/Classes/Sound/SoundWave.h
+++ b/Engine/Source/Runtime/Engine/Classes/Sound/SoundWave.h
@@ -25,8 +25,8 @@ enum EDecompressionType
 	DTYPE_MAX,
 };
 
-UCLASS(hidecategories=Object, editinlinenew, MinimalAPI, BlueprintType, dependson=USoundGroups)
-class USoundWave : public USoundBase
+UCLASS(hidecategories=Object, editinlinenew, BlueprintType, dependson=USoundGroups)
+class ENGINE_API USoundWave : public USoundBase
 {
 	GENERATED_UCLASS_BODY()
 
@@ -198,12 +198,12 @@ public:
 	/**
 	 * Frees up all the resources allocated in this class
 	 */
-	ENGINE_API void FreeResources();
+	void FreeResources();
 
 	/** 
 	 * Copy the compressed audio data from the bulk data
 	 */
-	virtual ENGINE_API void InitAudioResource( FByteBulkData& CompressedData );
+	virtual void InitAudioResource( FByteBulkData& CompressedData );
 
 	/** 
 	 * Copy the compressed audio data from derived data cache
@@ -211,12 +211,12 @@ public:
 	 * @param Format to get the compressed audio in
 	 * @return true if the resource has been successfully initialized or it was already initialized.
 	 */
-	virtual ENGINE_API bool InitAudioResource(FName Format);
+	virtual bool InitAudioResource(FName Format);
 
 	/** 
 	 * Remove the compressed audio data associated with the passed in wave
 	 */
-	ENGINE_API void RemoveAudioResource();
+	void RemoveAudioResource();
 
 	/** 
 	 * Prints the subtitle associated with the SoundWave to the console
@@ -258,12 +258,12 @@ public:
 	 * @param Format	format of compressed data
 	 * @return	compressed data, if it could be obtained
 	 */ 
-	virtual ENGINE_API FByteBulkData* GetCompressedData(FName Format);
+	virtual FByteBulkData* GetCompressedData(FName Format);
 
 	/** 
 	 * Change the guid and flush all compressed data
 	 */ 
-	ENGINE_API void InvalidateCompressedData();
+	void InvalidateCompressedData();
 
 
 };
diff --git a/Engine/Source/Runtime/Engine/Classes/Sound/SoundWaveStreaming.h b/Engine/Source/Runtime/Engine/Classes/Sound/SoundWaveStreaming.h
--- a/Engine/Source/Runtime/Engine/Classes/Sound/SoundWaveStreaming.h
+++ b/Engine/Source/Runtime/Engine/Classes/Sound/SoundWaveStreaming.h
@@ -10,8 +10,8 @@
 
 DECLARE_DELEGATE_TwoParams( FOnSoundWaveStreamingUnderflow, class USoundWaveStreaming*, int32 );
 
-UCLASS(MinimalAPI)
-class USoundWaveStreaming : public USoundWave
+UCLASS()
+class ENGINE_API USoundWaveStreaming : public USoundWave
 {
 	GENERATED_UCLASS_BODY()
 
@@ -36,13 +36,13 @@ public:
 
 
 	/** Add data to the FIFO that feeds the audio device. */
-	void ENGINE_API QueueAudio(const uint8* AudioData, const int32 BufferSize);
+	void QueueAudio(const uint8* AudioData, const int32 BufferSize);
 
 	/** Remove all queued data from the FIFO. This is only necessary if you want to start over, or GeneratePCMData() isn't going to be called, since that will eventually drain it. */
 	void ResetAudio();
 
 	/** Query bytes queued for playback */
-	int32 ENGINE_API GetAvailableAudioByteCount();
+	int32 GetAvailableAudioByteCount();
 
 	/** Called when GeneratePCMData is called but not enough data is available. Allows more data to be added, and will try again */
 	FOnSoundWaveStreamingUnderflow OnSoundWaveStreamingUnderflow;
diff --git a/Engine/Source/Runtime/Engine/Public/Audio.h b/Engine/Source/Runtime/Engine/Public/Audio.h
--- a/Engine/Source/Runtime/Engine/Public/Audio.h
+++ b/Engine/Source/Runtime/Engine/Public/Audio.h
@@ -147,7 +147,7 @@ private:
  * Structure encapsulating all information required to play a USoundWave on a channel/source. This is required
  * as a single USoundWave object can be used in multiple active cues or multiple times in the same cue.
  */
-struct FWaveInstance
+struct ENGINE_API FWaveInstance
 {
 	/** Static helper to create good unique type hashes */
 	static uint32 TypeHashCounter;
@@ -244,7 +244,7 @@ struct FWaveInstance
 	/**
 	 * Notifies the wave instance that the current playback buffer has finished.
 	 */
-	ENGINE_API void NotifyFinished( const bool bStopped = false );
+	void NotifyFinished( const bool bStopped = false );
 
 	/**
 	 * Friend archive function used for serialization.
@@ -254,10 +254,10 @@ struct FWaveInstance
 	/**
 	 * Function used by the GC.
 	 */
-	ENGINE_API void AddReferencedObjects( FReferenceCollector& Collector );
+	void AddReferencedObjects( FReferenceCollector& Collector );
 
 	/** Returns the actual volume the wave instance will play at */
-	ENGINE_API float GetActualVolume() const;
+	float GetActualVolume() const;
 };
 
 inline uint32 GetTypeHash( FWaveInstance* A ) { return A->TypeHash; }
