UE4: Fix for compiling WITH_RECAST = 0

diff --git a/Engine/Source/Runtime/Engine/Classes/AI/Navigation/RecastNavMesh.h b/Engine/Source/Runtime/Engine/Classes/AI/Navigation/RecastNavMesh.h
--- a/Engine/Source/Runtime/Engine/Classes/AI/Navigation/RecastNavMesh.h
+++ b/Engine/Source/Runtime/Engine/Classes/AI/Navigation/RecastNavMesh.h
@@ -877,9 +877,11 @@ private:
 #endif // WITH_RECAST
 };
 
+#if WITH_RECAST
 FORCEINLINE
 bool ARecastNavMesh::NavMeshRaycast(const ANavigationData* Self, const FVector& RayStart, const FVector& RayEnd, FVector& HitLocation, TSharedPtr<const FNavigationQueryFilter> QueryFilter, const UObject* Querier)
 {
 	FRaycastResult Result;
 	return NavMeshRaycast(Self, RayStart, RayEnd, HitLocation, QueryFilter, Querier, Result);
 }
+#endif // WITH_RECAST
\ No newline at end of file
