TTP #332810: Adding the camera-component's sub-component proxies to the CDO (in the constructor).

diff --git a/Engine/Source/Editor/UnrealEd/Private/Kismet2/KismetReinstanceUtilities.cpp b/Engine/Source/Editor/UnrealEd/Private/Kismet2/KismetReinstanceUtilities.cpp
--- a/Engine/Source/Editor/UnrealEd/Private/Kismet2/KismetReinstanceUtilities.cpp
+++ b/Engine/Source/Editor/UnrealEd/Private/Kismet2/KismetReinstanceUtilities.cpp
@@ -277,13 +277,17 @@ void FBlueprintCompileReinstancer::ReplaceInstancesOfClass(UClass* OldClass, UCl
 
 						AttachSocketName = OldRootComponent->AttachSocketName;
 
-						//detach it to remove any scaling 
+						//detach it to remove any scaling
 						OldRootComponent->DetachFromParent(true);
 					}
 
 					for (auto AttachIt = OldRootComponent->AttachChildren.CreateConstIterator(); AttachIt; ++AttachIt)
 					{
-						OldAttachChildren.Add((*AttachIt)->GetOwner());
+						USceneComponent* Child = *AttachIt;
+						if (ensure(Child != nullptr))
+						{
+							OldAttachChildren.Add(Child->GetOwner());
+						}
 					}
 
 					// Save off transform
diff --git a/Engine/Source/Runtime/Engine/Classes/Camera/CameraComponent.h b/Engine/Source/Runtime/Engine/Classes/Camera/CameraComponent.h
--- a/Engine/Source/Runtime/Engine/Classes/Camera/CameraComponent.h
+++ b/Engine/Source/Runtime/Engine/Classes/Camera/CameraComponent.h
@@ -66,14 +66,14 @@ protected:
 #if WITH_EDITORONLY_DATA
 	// The frustum component used to show visually where the camera field of view is
 	UPROPERTY(transient)
-	class UDrawFrustumComponent* DrawFrustum;
+	TSubobjectPtr<class UDrawFrustumComponent> DrawFrustum;
 
 	UPROPERTY(transient)
 	class UStaticMesh* CameraMesh;
 
 	// The camera mesh to show visually where the camera is placed
 	UPROPERTY(transient)
-	class UStaticMeshComponent* ProxyMeshComponent;
+	TSubobjectPtr<class UStaticMeshComponent> ProxyMeshComponent;
 
 public:
 	// Refreshes the visual components to match the component state
diff --git a/Engine/Source/Runtime/Engine/Private/Camera/CameraComponent.cpp b/Engine/Source/Runtime/Engine/Private/Camera/CameraComponent.cpp
--- a/Engine/Source/Runtime/Engine/Private/Camera/CameraComponent.cpp
+++ b/Engine/Source/Runtime/Engine/Private/Camera/CameraComponent.cpp
@@ -19,6 +19,21 @@ UCameraComponent::UCameraComponent(const class FPostConstructInitializePropertie
 		static ConstructorHelpers::FObjectFinder<UStaticMesh> EditorCameraMesh(TEXT("/Engine/EditorMeshes/MatineeCam_SM"));
 		CameraMesh = EditorCameraMesh.Object;
 	}
+
+	ProxyMeshComponent = PCIP.CreateEditorOnlyDefaultSubobject<UStaticMeshComponent>(GetOuter(), TEXT("CameraProxy"), true);
+	ProxyMeshComponent->AttachParent = this;
+	ProxyMeshComponent->StaticMesh = CameraMesh;
+	ProxyMeshComponent->SetCollisionProfileName(UCollisionProfile::NoCollision_ProfileName);
+	ProxyMeshComponent->bHiddenInGame = true;
+	ProxyMeshComponent->CastShadow = false;
+	ProxyMeshComponent->PostPhysicsComponentTick.bCanEverTick = false;
+	ProxyMeshComponent->bCreatedByConstructionScript = bCreatedByConstructionScript;
+
+	DrawFrustum = PCIP.CreateEditorOnlyDefaultSubobject<UDrawFrustumComponent>(GetOuter(), TEXT("FrustumVisualizer"), true);
+	DrawFrustum->AttachParent = this;
+	DrawFrustum->AlwaysLoadOnClient = false;
+	DrawFrustum->AlwaysLoadOnServer = false;
+	DrawFrustum->bCreatedByConstructionScript = bCreatedByConstructionScript;
 #endif
 
 	FieldOfView = 90.0f;
@@ -33,29 +48,6 @@ UCameraComponent::UCameraComponent(const class FPostConstructInitializePropertie
 void UCameraComponent::OnRegister()
 {
 #if WITH_EDITORONLY_DATA
-	if (ProxyMeshComponent == NULL)
-	{
-		ProxyMeshComponent = ConstructObject<UStaticMeshComponent>(UStaticMeshComponent::StaticClass(), GetOuter(), NAME_None, RF_Transactional);
-		ProxyMeshComponent->AttachTo(this);
-		ProxyMeshComponent->StaticMesh = CameraMesh;
-		ProxyMeshComponent->SetCollisionProfileName(UCollisionProfile::NoCollision_ProfileName);
-		ProxyMeshComponent->bHiddenInGame = true;
-		ProxyMeshComponent->CastShadow = false;
-		ProxyMeshComponent->PostPhysicsComponentTick.bCanEverTick = false;
-		ProxyMeshComponent->bCreatedByConstructionScript = bCreatedByConstructionScript;
-		ProxyMeshComponent->RegisterComponentWithWorld(GetWorld());
-	}
-
-	if (DrawFrustum == NULL)
-	{
-		DrawFrustum = ConstructObject<UDrawFrustumComponent>(UDrawFrustumComponent::StaticClass(), GetOuter(), NAME_None, RF_Transactional);
-		DrawFrustum->AttachTo(this);
-		DrawFrustum->AlwaysLoadOnClient = false;
-		DrawFrustum->AlwaysLoadOnServer = false;
-		DrawFrustum->bCreatedByConstructionScript = bCreatedByConstructionScript;
-		DrawFrustum->RegisterComponentWithWorld(GetWorld());
-	}
-
 	RefreshVisualRepresentation();
 #endif
 
