Improve tooltips for actor/component touch/click events

diff --git a/Engine/Source/Runtime/Engine/Classes/Components/PrimitiveComponent.h b/Engine/Source/Runtime/Engine/Classes/Components/PrimitiveComponent.h
--- a/Engine/Source/Runtime/Engine/Classes/Components/PrimitiveComponent.h
+++ b/Engine/Source/Runtime/Engine/Classes/Components/PrimitiveComponent.h
@@ -628,35 +628,35 @@ public:
 	UPROPERTY(BlueprintAssignable, Category="Collision")
 	FComponentEndOverlapSignature OnComponentEndOverlap;
 
-	/** Called when the mouse cursor is moved over this component when this mouse interface is enabled and mouse over events are enabled */
+	/** Called when the mouse cursor is moved over this component and mouse over events are enabled in the player controller */
 	UPROPERTY(BlueprintAssignable, Category="Input|Mouse Input")
 	FComponentBeginCursorOverSignature OnBeginCursorOver;
 		 
-	/** Called when the mouse cursor is moved off this component when this mouse interface is enabled and mouse over events are enabled */
+	/** Called when the mouse cursor is moved off this component and mouse over events are enabled in the player controller */
 	UPROPERTY(BlueprintAssignable, Category="Input|Mouse Input")
 	FComponentEndCursorOverSignature OnEndCursorOver;
 
-	/** Called when the left mouse button is clicked while the mouse is over this component when the mouse interface is enabled and click events are enabled */
+	/** Called when the left mouse button is clicked while the mouse is over this component and click events are enabled in the player controller */
 	UPROPERTY(BlueprintAssignable, Category="Input|Mouse Input")
 	FComponentOnClickedSignature OnClicked;
 
-	/** Called when the left mouse button is released while the mouse is over this component when the mouse interface is enabled and click events are enabled */
+	/** Called when the left mouse button is released while the mouse is over this component click events are enabled in the player controller */
 	UPROPERTY(BlueprintAssignable, Category="Input|Mouse Input")
 	FComponentOnReleasedSignature OnReleased;
 		 
-	/** Called when a touch input is received over this component when and click events are enabled */
+	/** Called when a touch input is received over this component when touch events are enabled in the player controller */
 	UPROPERTY(BlueprintAssignable, Category="Input|Touch Input")
 	FComponentOnInputTouchBeginSignature OnInputTouchBegin;
 
-	/** Called when a touch input is received over this component when and click events are enabled */
+	/** Called when a touch input is released over this component when touch events are enabled in the player controller */
 	UPROPERTY(BlueprintAssignable, Category="Input|Touch Input")
 	FComponentOnInputTouchEndSignature OnInputTouchEnd;
 
-	/** Called when a finger is moved over this component when this mouse interface is enabled and touch over events are enabled */
+	/** Called when a finger is moved over this component when touch over events are enabled in the player controller */
 	UPROPERTY(BlueprintAssignable, Category="Input|Touch Input")
 	FComponentBeginTouchOverSignature OnInputTouchEnter;
 
-	/** Called when a finger is moved off this component when this mouse interface is enabled and touch over events are enabled */
+	/** Called when a finger is moved off this component when touch over events are enabled in the player controller */
 	UPROPERTY(BlueprintAssignable, Category="Input|Touch Input")
 	FComponentEndTouchOverSignature OnInputTouchLeave;
 
diff --git a/Engine/Source/Runtime/Engine/Classes/GameFramework/Actor.h b/Engine/Source/Runtime/Engine/Classes/GameFramework/Actor.h
--- a/Engine/Source/Runtime/Engine/Classes/GameFramework/Actor.h
+++ b/Engine/Source/Runtime/Engine/Classes/GameFramework/Actor.h
@@ -511,35 +511,35 @@ public:
 	UPROPERTY(BlueprintAssignable, Category="Collision")
 	FActorEndOverlapSignature OnActorEndOverlap;
 
-	/** Called when the mouse cursor is moved over this actor when this mouse interface is enabled and mouse over events are enabled */
+	/** Called when the mouse cursor is moved over this actor and mouse over events are enabled in the player controller */
 	UPROPERTY(BlueprintAssignable, Category="Input|Mouse Input")
 	FActorBeginCursorOverSignature OnBeginCursorOver;
 
-	/** Called when the mouse cursor is moved off this actor when this mouse interface is enabled and mouse over events are enabled */
+	/** Called when the mouse cursor is moved off this actor and mouse over events are enabled in the player controller */
 	UPROPERTY(BlueprintAssignable, Category="Input|Mouse Input")
 	FActorEndCursorOverSignature OnEndCursorOver;
 
-	/** Called when the left mouse button is clicked while the mouse is over this actor when the mouse interface is enabled and click events are enabled */
+	/** Called when the left mouse button is clicked while the mouse is over this actor and click events are enabled in the player controller */
 	UPROPERTY(BlueprintAssignable, Category="Input|Mouse Input")
 	FActorOnClickedSignature OnClicked;
 
-	/** Called when the left mouse button is released while the mouse is over this actor when the mouse interface is enabled and click events are enabled */
+	/** Called when the left mouse button is released while the mouse is over this actor and click events are enabled in the player controller */
 	UPROPERTY(BlueprintAssignable, Category="Input|Mouse Input")
 	FActorOnReleasedSignature OnReleased;
 
-	/** Called when a touch input is received over this actor when and click events are enabled */
+	/** Called when a touch input is received over this actor when touch events are enabled in the player controller */
 	UPROPERTY(BlueprintAssignable, Category="Input|Touch Input")
 	FActorOnInputTouchBeginSignature OnInputTouchBegin;
 		
-	/** Called when a touch input is received over this component when and click events are enabled */
+	/** Called when a touch input is received over this component when touch events are enabled in the player controller */
 	UPROPERTY(BlueprintAssignable, Category="Input|Touch Input")
 	FActorOnInputTouchEndSignature OnInputTouchEnd;
 
-	/** Called when a finger is moved over this actor when this mouse interface is enabled and touch over events are enabled */
+	/** Called when a finger is moved over this actor when touch over events are enabled in the player controller */
 	UPROPERTY(BlueprintAssignable, Category="Input|Touch Input")
 	FActorBeginTouchOverSignature OnInputTouchEnter;
 
-	/** Called when a finger is moved off this actor when this mouse interface is enabled and touch over events are enabled */
+	/** Called when a finger is moved off this actor when touch over events are enabled in the player controller */
 	UPROPERTY(BlueprintAssignable, Category="Input|Touch Input")
 	FActorEndTouchOverSignature OnInputTouchLeave;
 
diff --git a/Engine/Source/Runtime/Engine/Classes/GameFramework/PlayerController.h b/Engine/Source/Runtime/Engine/Classes/GameFramework/PlayerController.h
--- a/Engine/Source/Runtime/Engine/Classes/GameFramework/PlayerController.h
+++ b/Engine/Source/Runtime/Engine/Classes/GameFramework/PlayerController.h
@@ -163,11 +163,11 @@ class ENGINE_API APlayerController : public AController
 	UPROPERTY(EditAnywhere, BlueprintReadWrite, Category=MouseInterface)
 	uint32 bShowMouseCursor:1;
 
-	/** Whether actor/component click and touch events should be generated. */
+	/** Whether actor/component click events should be generated. */
 	UPROPERTY(EditAnywhere, BlueprintReadWrite, Category=MouseInterface)
 	uint32 bEnableClickEvents:1;
 
-	/** Whether actor/component click and touch events should be generated. */
+	/** Whether actor/component touch events should be generated. */
 	UPROPERTY(EditAnywhere, BlueprintReadWrite, Category=MouseInterface)
 	uint32 bEnableTouchEvents:1;
 
