Remove unused SoundWave.bLoopable flag

diff --git a/Engine/Source/Developer/TargetPlatform/Public/Interfaces/IAudioFormat.h b/Engine/Source/Developer/TargetPlatform/Public/Interfaces/IAudioFormat.h
--- a/Engine/Source/Developer/TargetPlatform/Public/Interfaces/IAudioFormat.h
+++ b/Engine/Source/Developer/TargetPlatform/Public/Interfaces/IAudioFormat.h
@@ -18,11 +18,6 @@ struct FSoundQualityInfo
 	 */
 	int32 Quality;
 
-	/**
-	 * Holds a flag indicating Whether to do the additional processing required for looping sounds.
-	 */
-	int32 bLoopableSound;
-
 	/**
 	 * Holds the number of distinct audio channels.
 	 */
diff --git a/Engine/Source/Editor/UnrealEd/Private/Commandlets/UE4AudioRefactorMigrationCommandlet.cpp b/Engine/Source/Editor/UnrealEd/Private/Commandlets/UE4AudioRefactorMigrationCommandlet.cpp
--- a/Engine/Source/Editor/UnrealEd/Private/Commandlets/UE4AudioRefactorMigrationCommandlet.cpp
+++ b/Engine/Source/Editor/UnrealEd/Private/Commandlets/UE4AudioRefactorMigrationCommandlet.cpp
@@ -507,7 +507,6 @@ void ReplaceSoundNodeWaveWithSoundWave(const FAssetRegistryModule& AssetRegistry
 			USoundWave* SoundWave = ConstructObject<USoundWave>(USoundWave::StaticClass(), Package, SoundName, RF_Public | RF_Standalone);
 
 			SoundWave->CompressionQuality = SoundNodeWave->CompressionQuality;
-			SoundWave->bLoopableSound = SoundNodeWave->bLoopingSound;
 			SoundWave->SpokenText = SoundNodeWave->SpokenText;
 			SoundWave->Volume = SoundNodeWave->Volume;
 			SoundWave->Pitch = SoundNodeWave->Pitch;
diff --git a/Engine/Source/Editor/UnrealEd/Private/Editor.cpp b/Engine/Source/Editor/UnrealEd/Private/Editor.cpp
--- a/Engine/Source/Editor/UnrealEd/Private/Editor.cpp
+++ b/Engine/Source/Editor/UnrealEd/Private/Editor.cpp
@@ -2803,7 +2803,6 @@ void SoundWaveQualityPreview( USoundWave* SoundWave, FPreviewInfo* PreviewInfo )
 	QualityInfo.NumChannels = *WaveInfo.pChannels;
 	QualityInfo.SampleRate = SoundWave->SampleRate;
 	QualityInfo.SampleDataSize = WaveInfo.SampleDataSize;
-	QualityInfo.bLoopableSound = SoundWave->bLoopableSound;
 	QualityInfo.DebugName = SoundWave->GetFullName();
 
 	// PCM -> Vorbis -> PCM 
diff --git a/Engine/Source/Runtime/Engine/Classes/Sound/SoundWave.h b/Engine/Source/Runtime/Engine/Classes/Sound/SoundWave.h
--- a/Engine/Source/Runtime/Engine/Classes/Sound/SoundWave.h
+++ b/Engine/Source/Runtime/Engine/Classes/Sound/SoundWave.h
@@ -34,12 +34,8 @@ class ENGINE_API USoundWave : public USoundBase
 	UPROPERTY(EditAnywhere, Category=Compression, AssetRegistrySearchable)
 	int32 CompressionQuality;
 
-	/** If set, the compressor does everything required to make this a seamlessly looping sound. */
-	UPROPERTY(EditAnywhere, Category=Compression )
-	uint32 bLoopableSound:1;
-
 	/** If set, when played directly (not through a sound cue) the wave will be played looping. */
-	UPROPERTY(EditAnywhere, Category=Compression )
+	UPROPERTY(EditAnywhere, Category=SoundWave )
 	uint32 bLooping:1;
 
 	/** Set to true for programmatically-generated, streamed audio. */
diff --git a/Engine/Source/Runtime/Engine/Private/AudioDerivedData.cpp b/Engine/Source/Runtime/Engine/Private/AudioDerivedData.cpp
--- a/Engine/Source/Runtime/Engine/Private/AudioDerivedData.cpp
+++ b/Engine/Source/Runtime/Engine/Private/AudioDerivedData.cpp
@@ -82,7 +82,6 @@ static void CookSimpleWave(USoundWave* SoundWave, FName FormatName, const IAudio
 		QualityInfo.NumChannels = *WaveInfo.pChannels;
 		QualityInfo.SampleRate = *WaveInfo.pSamplesPerSec;
 		QualityInfo.SampleDataSize = Input.Num();
-		QualityInfo.bLoopableSound = SoundWave->bLoopableSound;
 		QualityInfo.DebugName = SoundWave->GetFullName();
 
 		// Cook the data.
@@ -172,7 +171,6 @@ void CookSurroundWave( USoundWave* SoundWave, FName FormatName, const IAudioForm
 			QualityInfo.NumChannels = ChannelCount;
 			QualityInfo.SampleRate = *WaveInfo.pSamplesPerSec;
 			QualityInfo.SampleDataSize = SampleDataSize;
-			QualityInfo.bLoopableSound = SoundWave->bLoopableSound;
 			QualityInfo.DebugName = SoundWave->GetFullName();
 			//@todo tighten up the checking for empty results here
 			if(Format.CookSurround(FormatName, SourceBuffers, QualityInfo, Output)) 
diff --git a/Engine/Source/Runtime/Engine/Private/SoundWave.cpp b/Engine/Source/Runtime/Engine/Private/SoundWave.cpp
--- a/Engine/Source/Runtime/Engine/Private/SoundWave.cpp
+++ b/Engine/Source/Runtime/Engine/Private/SoundWave.cpp
@@ -13,7 +13,6 @@ USoundWave::USoundWave(const class FPostConstructInitializeProperties& PCIP)
 	Volume = 1.0;
 	Pitch = 1.0;
 	CompressionQuality = 40;
-	bLoopableSound = true;
 }
 
 SIZE_T USoundWave::GetResourceSize(EResourceSizeMode::Type Mode)
@@ -314,14 +313,10 @@ void USoundWave::PostEditChangeProperty(FPropertyChangedEvent& PropertyChangedEv
 	Super::PostEditChangeProperty(PropertyChangedEvent);
 
 	static FName CompressionQualityFName = FName( TEXT( "CompressionQuality" ) );
-	static FName LoopableSoundFName = FName( TEXT( "bLoopableSound" ) );
 
 	UProperty* PropertyThatChanged = PropertyChangedEvent.Property;
 	// Regenerate on save any compressed sound formats
-	if( PropertyThatChanged &&
-		( PropertyThatChanged->GetFName() == CompressionQualityFName
-		|| PropertyThatChanged->GetFName() == LoopableSoundFName
-		) )
+	if ( PropertyThatChanged && PropertyThatChanged->GetFName() == CompressionQualityFName )
 	{
 		InvalidateCompressedData();
 		MarkPackageDirty();
