Fixing a crash in CameraComponent when running the editor with -game (fallout from CL 2050237).

diff --git a/Engine/Source/Runtime/Engine/Private/Camera/CameraComponent.cpp b/Engine/Source/Runtime/Engine/Private/Camera/CameraComponent.cpp
--- a/Engine/Source/Runtime/Engine/Private/Camera/CameraComponent.cpp
+++ b/Engine/Source/Runtime/Engine/Private/Camera/CameraComponent.cpp
@@ -21,19 +21,25 @@ UCameraComponent::UCameraComponent(const class FPostConstructInitializePropertie
 	}
 
 	ProxyMeshComponent = PCIP.CreateEditorOnlyDefaultSubobject<UStaticMeshComponent>(GetOuter(), TEXT("CameraProxy"), true);
-	ProxyMeshComponent->AttachParent = this;
-	ProxyMeshComponent->StaticMesh = CameraMesh;
-	ProxyMeshComponent->SetCollisionProfileName(UCollisionProfile::NoCollision_ProfileName);
-	ProxyMeshComponent->bHiddenInGame = true;
-	ProxyMeshComponent->CastShadow = false;
-	ProxyMeshComponent->PostPhysicsComponentTick.bCanEverTick = false;
-	ProxyMeshComponent->bCreatedByConstructionScript = bCreatedByConstructionScript;
+	if (ProxyMeshComponent)
+	{
+		ProxyMeshComponent->AttachParent = this;
+		ProxyMeshComponent->StaticMesh = CameraMesh;
+		ProxyMeshComponent->SetCollisionProfileName(UCollisionProfile::NoCollision_ProfileName);
+		ProxyMeshComponent->bHiddenInGame = true;
+		ProxyMeshComponent->CastShadow = false;
+		ProxyMeshComponent->PostPhysicsComponentTick.bCanEverTick = false;
+		ProxyMeshComponent->bCreatedByConstructionScript = bCreatedByConstructionScript;
+	}
 
 	DrawFrustum = PCIP.CreateEditorOnlyDefaultSubobject<UDrawFrustumComponent>(GetOuter(), TEXT("FrustumVisualizer"), true);
-	DrawFrustum->AttachParent = this;
-	DrawFrustum->AlwaysLoadOnClient = false;
-	DrawFrustum->AlwaysLoadOnServer = false;
-	DrawFrustum->bCreatedByConstructionScript = bCreatedByConstructionScript;
+	if (DrawFrustum)
+	{
+		DrawFrustum->AttachParent = this;
+		DrawFrustum->AlwaysLoadOnClient = false;
+		DrawFrustum->AlwaysLoadOnServer = false;
+		DrawFrustum->bCreatedByConstructionScript = bCreatedByConstructionScript;
+	}
 #endif
 
 	FieldOfView = 90.0f;
