331222  CRITICAL: Regression: ContentExamples: MorphTarget material pops instead of gradual change

- now when we apply morphtarget normal, we send out how much alpha of delta normal we're applying to base normal.

diff --git a/Engine/Shaders/GpuSkinVertexFactory.usf b/Engine/Shaders/GpuSkinVertexFactory.usf
--- a/Engine/Shaders/GpuSkinVertexFactory.usf
+++ b/Engine/Shaders/GpuSkinVertexFactory.usf
@@ -91,7 +91,7 @@ struct FVertexFactoryInput
 	/** added to the Position */
 	float3	DeltaPosition	: ATTRIBUTE9;	//POSITION1;
 	/** added to the TangentZ and then used to derive new TangentX,TangentY */
-	half3	DeltaTangentZ	: ATTRIBUTE10;	//NORMAL1;
+	half4	DeltaTangentZ	: ATTRIBUTE10;	//NORMAL1;
 #elif GPUSKIN_APEX_CLOTH // exclusive with GPUSKIN_MORPH_BLEND
 	// APEX cloth mesh-mesh mapping data
 	// Barycentric Coordinate Data
@@ -392,7 +392,7 @@ float3x3 SkinTangents( FVertexFactoryInput Input, FVertexFactoryIntermediates In
 	float3 DeltaTangentZ = TangentBias(Input.DeltaTangentZ)*2;
 
 	// calc new normal by offseting it with the delta
-	Tangents[1] = normalize( float3(TangentZ.x,TangentZ.y,TangentZ.z) + DeltaTangentZ );
+	Tangents[1] = normalize( TangentZ.xyz + Input.DeltaTangentZ.w * DeltaTangentZ);
 	// derive the new tangent by orthonormalizing the new normal against
 	// the base tangent vector (assuming these are normalized)
 	Tangents[0] = normalize( TangentX - (dot(TangentX,Tangents[1]) * Tangents[1]) );
diff --git a/Engine/Source/Runtime/Engine/Private/SkeletalRenderGPUSkin.cpp b/Engine/Source/Runtime/Engine/Private/SkeletalRenderGPUSkin.cpp
--- a/Engine/Source/Runtime/Engine/Private/SkeletalRenderGPUSkin.cpp
+++ b/Engine/Source/Runtime/Engine/Private/SkeletalRenderGPUSkin.cpp
@@ -61,7 +61,7 @@ void FMorphVertexBuffer::InitDynamicRHI()
 	for (uint32 VertIndex=0; VertIndex < LodModel.NumVertices; ++VertIndex)
 	{
 		Buffer[VertIndex].DeltaPosition = FVector::ZeroVector;
-		Buffer[VertIndex].DeltaTangentZ = FVector::ZeroVector;
+		Buffer[VertIndex].DeltaTangentZ = FPackedNormal::ZeroNormal;
 	}
 
 	// Unlock the buffer.
@@ -347,8 +347,8 @@ void FSkeletalMeshObjectGPUSkin::FSkeletalMeshObjectLOD::UpdateMorphVertexBuffer
 				AccumulatedWeightArray.AddUninitialized(vertsToAdd);
 			}
 
-			FMemory::Memzero(DeltaTangentZAccumulationArray.GetData(), sizeof(FVector)*DeltaTangentZAccumulationArray.Num());
-			FMemory::Memzero(AccumulatedWeightArray.GetData(), sizeof(float)*AccumulatedWeightArray.Num());
+			FMemory::Memzero(DeltaTangentZAccumulationArray.GetData(), sizeof(FVector)*LodModel.NumVertices);
+			FMemory::Memzero(AccumulatedWeightArray.GetData(), sizeof(float)*LodModel.NumVertices);
 
 			// PackedNormals will be wrong init with 0, but they'll be overwritten later
 			FMemory::Memzero(&Buffer[0], sizeof(FMorphGPUSkinVertex)*LodModel.NumVertices);
@@ -380,7 +380,6 @@ void FSkeletalMeshObjectGPUSkin::FSkeletalMeshObjectLOD::UpdateMorphVertexBuffer
 					FMorphGPUSkinVertex& DestVertex = Buffer[MorphVertex.SourceIdx];
 
 					DestVertex.DeltaPosition += MorphVertex.PositionDelta * VertAnim.Weight;
-
 					// add to accumulated tangent Z
 					DeltaTangentZAccumulationArray[MorphVertex.SourceIdx] += MorphVertex.TangentZDelta * VertAnim.Weight;
 					// accumulate the weight so we can normalized it later
@@ -396,16 +395,20 @@ void FSkeletalMeshObjectGPUSkin::FSkeletalMeshObjectLOD::UpdateMorphVertexBuffer
 				FMorphGPUSkinVertex& DestVertex = Buffer[iVertex];
 				float AccumulatedWeight = AccumulatedWeightArray[iVertex];
 
-				if (!FMath::IsNearlyZero(AccumulatedWeight) )
-				{
+ 				if (AccumulatedWeight > MinVertexAnimBlendWeight)
+ 				{
 					// when copy back, make sure to normalize by accumulated weight
 					// since delta diff of normal is (-2, 2), we divide by 2 to packed into packed normal
 					// when we unpack, we'll apply *2. 
-					DestVertex.DeltaTangentZ = (DeltaTangentZAccumulationArray[iVertex]/AccumulatedWeight)/2;
+					DestVertex.DeltaTangentZ = (DeltaTangentZAccumulationArray[iVertex] / AccumulatedWeight)/2;
+					// we now add W as how much alpha of DeltaTangentZ we're apply to the original tangent
+					DestVertex.DeltaTangentZ.Vector.W = FMath::Min(1.0f, AccumulatedWeight) * 255.9999f;
 				}
 				else
 				{
 					DestVertex.DeltaTangentZ = FPackedNormal::ZeroNormal;
+					// we now add W as how much alpha of DeltaTangentZ we're apply to the original tangent					
+					DestVertex.DeltaTangentZ.Vector.W = 0;
 				}
 			}
 		} // ApplyDelta
