#ue4
Fixed issue with potential INFs and NaNs in physics constraints

diff --git a/Engine/Source/Runtime/Engine/Private/PhysicsEngine/PhysicsConstraintComponent.cpp b/Engine/Source/Runtime/Engine/Private/PhysicsEngine/PhysicsConstraintComponent.cpp
--- a/Engine/Source/Runtime/Engine/Private/PhysicsEngine/PhysicsConstraintComponent.cpp
+++ b/Engine/Source/Runtime/Engine/Private/PhysicsEngine/PhysicsConstraintComponent.cpp
@@ -306,7 +306,7 @@ void UPhysicsConstraintComponent::PostLoad()
 			TotalMass += Primitive2->CalculateMass(BoneName);
 		}
 
-		if (NumDynamic > 0)	//we don't support cases where both constrained bodies are static or NULL, but add this anyway to avoid crash
+		if ( (NumDynamic > 0) && (TotalMass > 0) )	//we don't support cases where both constrained bodies are static or NULL, but add this anyway to avoid crash
 		{
 			float AverageMass = TotalMass / NumDynamic;
 
