Added temporary backwards compatibility to previous reflection capture DDC version

diff --git a/Engine/Source/Runtime/Engine/Classes/Components/ReflectionCaptureComponent.h b/Engine/Source/Runtime/Engine/Classes/Components/ReflectionCaptureComponent.h
--- a/Engine/Source/Runtime/Engine/Classes/Components/ReflectionCaptureComponent.h
+++ b/Engine/Source/Runtime/Engine/Classes/Components/ReflectionCaptureComponent.h
@@ -31,6 +31,7 @@ public:
 
 	/** Constructs a key string for the DDC that uniquely identifies a FReflectionCaptureFullHDRDerivedData. */
 	static FString GetDDCKeyString(const FGuid& StateId);
+	static FString GetLegacyDDCKeyString(const FGuid& StateId);
 
 private:
 
diff --git a/Engine/Source/Runtime/Engine/Private/ReflectionCaptureComponent.cpp b/Engine/Source/Runtime/Engine/Private/ReflectionCaptureComponent.cpp
--- a/Engine/Source/Runtime/Engine/Private/ReflectionCaptureComponent.cpp
+++ b/Engine/Source/Runtime/Engine/Private/ReflectionCaptureComponent.cpp
@@ -204,6 +204,14 @@ FString FReflectionCaptureFullHDRDerivedData::GetDDCKeyString(const FGuid& State
 	return FDerivedDataCacheInterface::BuildCacheKey(TEXT("REFL_FULL"), *ReflectionCaptureDDCVer.ToString(), *StateId.ToString());
 }
 
+// For temporary backwards compatibility
+#define LEGACY_REFLECTIONCAPTURE_FULL_DERIVEDDATA_VER	TEXT("82f35deee55a5ec9956d4897f3d46e5a")
+
+FString FReflectionCaptureFullHDRDerivedData::GetLegacyDDCKeyString(const FGuid& StateId)
+{
+	return FDerivedDataCacheInterface::BuildCacheKey(TEXT("REFL_FULL"), LEGACY_REFLECTIONCAPTURE_FULL_DERIVEDDATA_VER, *StateId.ToString());
+}
+
 FReflectionCaptureFullHDRDerivedData::~FReflectionCaptureFullHDRDerivedData()
 {
 	DEC_MEMORY_STAT_BY(STAT_ReflectionCaptureMemory, CompressedCapturedData.GetAllocatedSize());
@@ -963,14 +971,19 @@ void UReflectionCaptureComponent::PostLoad()
 			// Attempt to load the full HDR data from the DDC
 			if (!GetDerivedDataCacheRef().GetSynchronous(*FReflectionCaptureFullHDRDerivedData::GetDDCKeyString(StateId), FullHDRDerivedData->CompressedCapturedData))
 			{
-				bDerivedDataDirty = true;
-				delete FullHDRDerivedData;
-				FullHDRDerivedData = NULL;
-
-				if (!FApp::CanEverRender())
+				// Try to load the previous version
+				// This is a temporary measure for backwards compatibility of DDC
+				if (!GetDerivedDataCacheRef().GetSynchronous(*FReflectionCaptureFullHDRDerivedData::GetLegacyDDCKeyString(StateId), FullHDRDerivedData->CompressedCapturedData))
 				{
-					// Warn, especially when running the DDC commandlet to build a DDC for the binary version of UE4.
-					UE_LOG(LogMaterial, Warning, TEXT("Reflection capture was loaded without any valid capture data and will be black.  This can happen if the DDC was not up to date during cooking.  Load the map in the editor once before cooking to fix.  %s."), *GetFullName());
+					bDerivedDataDirty = true;
+					delete FullHDRDerivedData;
+					FullHDRDerivedData = NULL;
+
+					if (!FApp::CanEverRender())
+					{
+						// Warn, especially when running the DDC commandlet to build a DDC for the binary version of UE4.
+						UE_LOG(LogMaterial, Warning, TEXT("Reflection capture was loaded without any valid capture data and will be black.  This can happen if the DDC was not up to date during cooking.  Load the map in the editor once before cooking to fix.  %s."), *GetFullName());
+					}
 				}
 			}
 		}
