Fixed scale gets lost during import

diff --git a/Engine/Source/Runtime/Engine/Private/Animation/AnimSequence.cpp b/Engine/Source/Runtime/Engine/Private/Animation/AnimSequence.cpp
--- a/Engine/Source/Runtime/Engine/Private/Animation/AnimSequence.cpp
+++ b/Engine/Source/Runtime/Engine/Private/Animation/AnimSequence.cpp
@@ -1225,6 +1225,11 @@ bool UAnimSequence::CompressRawAnimData(float MaxPosDiff, float MaxAngleDiff)
 					--I;
 					bRemovedKeys = true;
 				}
+				else
+				{
+					// if scale key exists, see if we can just empty it
+					bComopressScaleKeys |= (RawData.ScaleKeys.Num() > 1 || RawData.ScaleKeys[0].Equals(FVector(1.f)) == false);
+				}
 			}
 			else if ( RawData.ScaleKeys.Num() > 0 )
 			{
