#ue4
- Formatting fixes in UKismetSystemLibrary

diff --git a/Engine/Source/Runtime/Engine/Private/KismetSystemLibrary.cpp b/Engine/Source/Runtime/Engine/Private/KismetSystemLibrary.cpp
--- a/Engine/Source/Runtime/Engine/Private/KismetSystemLibrary.cpp
+++ b/Engine/Source/Runtime/Engine/Private/KismetSystemLibrary.cpp
@@ -2217,25 +2217,17 @@ void UKismetSystemLibrary::CollectGarbage()
 	GEngine->DeferredCommands.Add(TEXT("obj gc"));
 }
 
-#if PLATFORM_IOS
-extern CORE_API void IOSShowLeaderboardUI(const FString& CategoryName);
-#endif
-
 void UKismetSystemLibrary::EXPERIMENTAL_ShowAdBanner(bool bShowOnBottomOfScreen)
 {
-	IAdvertisingProvider * Provider = FAdvertising::Get().GetDefaultProvider();
-
-	if ( Provider != NULL )
+	if (IAdvertisingProvider* Provider = FAdvertising::Get().GetDefaultProvider())
 	{
-		Provider->ShowAdBanner( bShowOnBottomOfScreen );
+		Provider->ShowAdBanner(bShowOnBottomOfScreen);
 	}
 }
 
 void UKismetSystemLibrary::EXPERIMENTAL_HideAdBanner()
 {
-	IAdvertisingProvider * Provider = FAdvertising::Get().GetDefaultProvider();
-
-	if ( Provider != NULL )
+	if (IAdvertisingProvider* Provider = FAdvertising::Get().GetDefaultProvider())
 	{
 		Provider->HideAdBanner();
 	}
@@ -2243,9 +2235,7 @@ void UKismetSystemLibrary::EXPERIMENTAL_HideAdBanner()
 
 void UKismetSystemLibrary::EXPERIMENTAL_CloseAdBanner()
 {
-	IAdvertisingProvider * Provider = FAdvertising::Get().GetDefaultProvider();
-
-	if ( Provider != NULL )
+	if (IAdvertisingProvider* Provider = FAdvertising::Get().GetDefaultProvider())
 	{
 		Provider->CloseAdBanner();
 	}
@@ -2254,6 +2244,7 @@ void UKismetSystemLibrary::EXPERIMENTAL_CloseAdBanner()
 void UKismetSystemLibrary::EXPERIMENTAL_ShowGameCenterLeaderboard(const FString& CategoryName)
 {
 #if PLATFORM_IOS
+	extern CORE_API void IOSShowLeaderboardUI(const FString& CategoryName);
 	IOSShowLeaderboardUI(CategoryName);
 #elif PLATFORM_ANDROID
 	UE_LOG(LogTemp, Log, TEXT("Showing leaderboard!"));
