#UE4
- TTP#332824, not setting up AIController's navigation components if possesed pawn doesn't have a movement component

diff --git a/Engine/Source/Runtime/Engine/Private/AIController.cpp b/Engine/Source/Runtime/Engine/Private/AIController.cpp
--- a/Engine/Source/Runtime/Engine/Private/AIController.cpp
+++ b/Engine/Source/Runtime/Engine/Private/AIController.cpp
@@ -379,7 +379,12 @@ void AAIController::Possess(APawn* InPawn)
 		return;
 	}
 	
-	UpdateNavigationComponents();
+	// no point in doing navigation setup if pawn has no movement component
+	const UPawnMovementComponent* MovementComp = InPawn->GetMovementComponent();
+	if (MovementComp != NULL)
+	{
+		UpdateNavigationComponents();
+	}
 
 	if (bWantsPlayerState)
 	{
