#ue4 Crash Fix, when post load is called on trashed SCS.

diff --git a/Engine/Source/Runtime/Engine/Private/SimpleConstructionScript.cpp b/Engine/Source/Runtime/Engine/Private/SimpleConstructionScript.cpp
--- a/Engine/Source/Runtime/Engine/Private/SimpleConstructionScript.cpp
+++ b/Engine/Source/Runtime/Engine/Private/SimpleConstructionScript.cpp
@@ -136,7 +136,12 @@ void USimpleConstructionScript::PostLoad()
 #if WITH_EDITOR
 	// Get the Blueprint that owns the SCS
 	UBlueprint* Blueprint = GetBlueprint();
-	check(Blueprint != NULL);
+	if (!Blueprint)
+	{
+		// sometimes the PostLoad can be called, after the object was trashed, we dont want this
+		UE_LOG(LogBlueprint, Warning, TEXT("USimpleConstructionScript::PostLoad() '%s' cannot find its owner blueprint"), *GetName());
+		return;
+	}
 
 	// Fix up any uninitialized category names
 	for (NodeIndex=0; NodeIndex < Nodes.Num(); ++NodeIndex)
