fix bad assert when sub-stepping a kinematic actor changes from kinematic to simulated

diff --git a/Engine/Source/Runtime/Engine/Private/PhysicsEngine/PhysSubstepTasks.cpp b/Engine/Source/Runtime/Engine/Private/PhysicsEngine/PhysSubstepTasks.cpp
--- a/Engine/Source/Runtime/Engine/Private/PhysicsEngine/PhysSubstepTasks.cpp
+++ b/Engine/Source/Runtime/Engine/Private/PhysicsEngine/PhysSubstepTasks.cpp
@@ -170,19 +170,21 @@ void FPhysSubstepTask::InterpolateKinematicActor(const FPhysTarget & PhysTarget,
 	/** Interpolate kinematic actors */
 	if (PhysTarget.bKinematicTarget)
 	{
-		//We should only be interpolating kinematic actors
-		check(!IsRigidDynamicNonKinematic(PRigidDynamic));
-		const FKinematicTarget & KinematicTarget = PhysTarget.KinematicTarget;
-		const FTransform & TargetTM = KinematicTarget.TargetTM;
-		const FTransform & StartTM = KinematicTarget.OriginalTM;
-		FTransform InterTM = FTransform::Identity;
-
-		InterTM.SetLocation(FMath::Lerp(StartTM.GetLocation(), TargetTM.GetLocation(), Alpha));
-		InterTM.SetRotation(FMath::Lerp(StartTM.GetRotation(), TargetTM.GetRotation(), Alpha));
-
-		const PxTransform PNewPose = U2PTransform(InterTM);
-		check(PNewPose.isValid());
-		PRigidDynamic->setKinematicTarget(PNewPose);
+		//It's possible that the actor is no longer kinematic and is now simulating. In that case do nothing
+		if (!IsRigidDynamicNonKinematic(PRigidDynamic))
+		{
+			const FKinematicTarget & KinematicTarget = PhysTarget.KinematicTarget;
+			const FTransform & TargetTM = KinematicTarget.TargetTM;
+			const FTransform & StartTM = KinematicTarget.OriginalTM;
+			FTransform InterTM = FTransform::Identity;
+
+			InterTM.SetLocation(FMath::Lerp(StartTM.GetLocation(), TargetTM.GetLocation(), Alpha));
+			InterTM.SetRotation(FMath::Lerp(StartTM.GetRotation(), TargetTM.GetRotation(), Alpha));
+
+			const PxTransform PNewPose = U2PTransform(InterTM);
+			check(PNewPose.isValid());
+			PRigidDynamic->setKinematicTarget(PNewPose);
+		}
 	}
 #endif
 }
