#UE4
The ReplicatedComponents array should be serialized as cooked builds won't re-run construction scripts

diff --git a/Engine/Source/Runtime/Engine/Classes/GameFramework/Actor.h b/Engine/Source/Runtime/Engine/Classes/GameFramework/Actor.h
--- a/Engine/Source/Runtime/Engine/Classes/GameFramework/Actor.h
+++ b/Engine/Source/Runtime/Engine/Classes/GameFramework/Actor.h
@@ -547,7 +547,8 @@ public:
 	UPROPERTY(BlueprintAssignable, Category="Collision")
 	FActorHitSignature OnActorHit;
 
-	/** List of replicated components. Built at runtime */
+	/** List of replicated components, cooked builds won't rebuild this array */
+	UPROPERTY()
 	TArray< TWeakObjectPtr<class UActorComponent> >	ReplicatedComponents;
 
 	/** Called on client when Instigator is replicated (only if bReplicateInstigator is true) */
