#UE4:
* removed unused console feature - for upcoming fixes

diff --git a/Engine/Source/Runtime/Engine/Classes/GameFramework/InputSettings.h b/Engine/Source/Runtime/Engine/Classes/GameFramework/InputSettings.h
--- a/Engine/Source/Runtime/Engine/Classes/GameFramework/InputSettings.h
+++ b/Engine/Source/Runtime/Engine/Classes/GameFramework/InputSettings.h
@@ -53,10 +53,6 @@ class ENGINE_API UInputSettings : public UObject
 	UPROPERTY(config, EditAnywhere, Category="Console")
 	FKey ConsoleKey;
 
-	/** Should the user be required to hold ctrl to use the up/down arrows when navigating auto-complete */
-	UPROPERTY(config, EditAnywhere, Category="Console", AdvancedDisplay)
-	uint32 bRequireCtrlToNavigateAutoComplete : 1;
-
 	// UObject interface
 #if WITH_EDITOR
 	virtual void PostEditChangeChainProperty(struct FPropertyChangedChainEvent& PropertyChangedEvent) OVERRIDE;
diff --git a/Engine/Source/Runtime/Engine/Private/UserInterface/Console.cpp b/Engine/Source/Runtime/Engine/Private/UserInterface/Console.cpp
--- a/Engine/Source/Runtime/Engine/Private/UserInterface/Console.cpp
+++ b/Engine/Source/Runtime/Engine/Private/UserInterface/Console.cpp
@@ -631,8 +631,7 @@ bool UConsole::InputKey_Typing( int32 ControllerId, FKey Key, EInputEvent Event,
 	}
 	else if( Key == EKeys::Up )
 	{
-		const bool bRequireCtrlToNavigateAutoComplete = GetDefault<UInputSettings>()->bRequireCtrlToNavigateAutoComplete;
-		if (!bNavigatingHistory && ((bRequireCtrlToNavigateAutoComplete && bCtrl) || (!bRequireCtrlToNavigateAutoComplete && !bCtrl && AutoCompleteIndices.Num() > 1)))
+		if (!bNavigatingHistory && !bCtrl && AutoCompleteIndices.Num() > 1)
 		{
 			if (++AutoCompleteIndex == AutoCompleteIndices.Num())
 			{
@@ -660,8 +659,7 @@ bool UConsole::InputKey_Typing( int32 ControllerId, FKey Key, EInputEvent Event,
 	}
 	else if( Key == EKeys::Down )
 	{
-		const bool bRequireCtrlToNavigateAutoComplete = GetDefault<UInputSettings>()->bRequireCtrlToNavigateAutoComplete;
-		if (!bNavigatingHistory && ((bRequireCtrlToNavigateAutoComplete && bCtrl) || (!bRequireCtrlToNavigateAutoComplete && !bCtrl && AutoCompleteIndices.Num() > 1)))
+		if (!bNavigatingHistory && !bCtrl && AutoCompleteIndices.Num() > 1)
 		{
 			if (--AutoCompleteIndex < 0)
 			{
