#UE:
* fixed function of game console displaying autocompletion without the need for a description (cvars and some others)

diff --git a/Engine/Source/Runtime/Engine/Private/UserInterface/Console.cpp b/Engine/Source/Runtime/Engine/Private/UserInterface/Console.cpp
--- a/Engine/Source/Runtime/Engine/Private/UserInterface/Console.cpp
+++ b/Engine/Source/Runtime/Engine/Private/UserInterface/Console.cpp
@@ -836,7 +836,15 @@ void UConsole::PostRender_Console_Typing(UCanvas* Canvas)
 		ConsoleTile.Texture = DefaultTexture_White->Resource;
 		for (int32 MatchIdx = 0; MatchIdx < 10; MatchIdx++)
 		{
-			OutStr = AutoCompleteList[AutoCompleteIndices[Idx]].Desc;
+			const FAutoCompleteCommand &Cmd = AutoCompleteList[AutoCompleteIndices[Idx]];
+			OutStr = Cmd.Desc;
+
+			if(OutStr.IsEmpty())
+			{
+				// no Description means we display the Command directly, without that the line would be empty (happens for ConsoleVariables and some ManualAutoCompleteList)
+				OutStr = Cmd.Command;
+			}
+
 			Canvas->StrLen(Font, OutStr, info_xl, info_yl);
 			y -= info_yl - yl;
 
@@ -1211,12 +1219,22 @@ void UConsole::PostRender_Console_Open(UCanvas* Canvas)
 		for (int32 MatchIdx = 0; MatchIdx < AutoCompleteIndices.Num() && MatchIdx < 10; MatchIdx++)
 		{
 			idx = AutoCompleteIndices[MatchIdx];
-			Canvas->StrLen(Font, AutoCompleteList[idx].Desc, info_xl, info_yl);
+
+			const FAutoCompleteCommand &Cmd = AutoCompleteList[idx];
+			FString Str = Cmd.Desc;
+
+			if(Str.IsEmpty())
+			{
+				// no Description means we display the Command directly, without that the line would be empty (happens for ConsoleVariables and some ManualAutoCompleteList)
+				Str = Cmd.Command;
+			}
+
+			Canvas->StrLen(Font, Str, info_xl, info_yl);
 			ConsoleTile.Size = FVector2D( info_xl, info_yl );
 			Canvas->DrawItem( ConsoleTile, LeftPos + xl, TopPos + y );
 
 			ConsoleText.SetColor( ConsoleDefs::AutocompleteSuggestionColor );
-			ConsoleText.Text = FText::FromString( AutoCompleteList[idx].Desc );
+			ConsoleText.Text = FText::FromString( Str );
 			Canvas->DrawItem( ConsoleText, LeftPos + xl, TopPos + y );
 			y += info_yl;
 		}
