#UE4:
* Improved auto completion in game console (when in typing mode, open console mode not yet), no cycle scrolling, window cursor like in editor

diff --git a/Engine/Source/Runtime/Engine/Classes/Engine/Console.h b/Engine/Source/Runtime/Engine/Classes/Engine/Console.h
--- a/Engine/Source/Runtime/Engine/Classes/Engine/Console.h
+++ b/Engine/Source/Runtime/Engine/Classes/Engine/Console.h
@@ -136,6 +136,8 @@ class ENGINE_API UConsole : public UObject, public FOutputDevice
 	/** Currently selected auto complete index */
 	int32 AutoCompleteIndex;
 
+	int32 AutoCompleteCursor;
+
 	/** Do we need to rebuild auto complete? */
 	uint32 bIsRuntimeAutoCompleteUpToDate:1;
 
diff --git a/Engine/Source/Runtime/Engine/Private/UserInterface/Console.cpp b/Engine/Source/Runtime/Engine/Private/UserInterface/Console.cpp
--- a/Engine/Source/Runtime/Engine/Private/UserInterface/Console.cpp
+++ b/Engine/Source/Runtime/Engine/Private/UserInterface/Console.cpp
@@ -10,6 +10,7 @@
 #include "Slate.h"
 #include "DefaultValueHelper.h"
 
+static const uint32 MAX_AUTOCOMPLETION_LINES = 20;
 
 static FName NAME_Typing = FName(TEXT("Typing"));
 static FName NAME_Open = FName(TEXT("Open"));
@@ -251,6 +252,7 @@ void UConsole::UpdateCompleteIndices()
 	}
 	bNavigatingHistory = false;
 	AutoCompleteIndex = 0;
+	AutoCompleteCursor = 0;
 	AutoCompleteIndices.Empty();
 	FAutoCompleteNode *Node = &AutoCompleteTree;
 	FString LowerTypedStr = TypedStr.ToLower();
@@ -583,7 +585,9 @@ bool UConsole::InputKey_Typing( int32 ControllerId, FKey Key, EInputEvent Event,
 	{
 		if (AutoCompleteIndices.Num() > 0 && !bAutoCompleteLocked)
 		{
-			TypedStr = AutoCompleteList[AutoCompleteIndices[AutoCompleteIndex]].Command;
+			const FAutoCompleteCommand& Cmd = AutoCompleteList[AutoCompleteIndices[AutoCompleteIndex + AutoCompleteCursor]];
+
+			TypedStr = Cmd.Command;
 			SetCursorPos(TypedStr.Len());
 			bAutoCompleteLocked = true;
 		}
@@ -633,9 +637,21 @@ bool UConsole::InputKey_Typing( int32 ControllerId, FKey Key, EInputEvent Event,
 	{
 		if (!bNavigatingHistory && !bCtrl && AutoCompleteIndices.Num() > 1)
 		{
-			if (++AutoCompleteIndex == AutoCompleteIndices.Num())
+			if(AutoCompleteCursor + 1 < FMath::Min((int32)MAX_AUTOCOMPLETION_LINES, AutoCompleteIndices.Num()))
+			{
+				// move cursor within displayed region
+				++AutoCompleteCursor;
+			}
+			else 
 			{
-				AutoCompleteIndex = 0;
+				// can be negative
+				int32 ScrollRegionSize = AutoCompleteIndices.Num() - (int32)MAX_AUTOCOMPLETION_LINES;
+
+				// can we scroll?
+				if(AutoCompleteIndex + 1 < ScrollRegionSize)
+				{
+					++AutoCompleteIndex;
+				}
 			}
 		}
 		else
@@ -661,11 +677,20 @@ bool UConsole::InputKey_Typing( int32 ControllerId, FKey Key, EInputEvent Event,
 	{
 		if (!bNavigatingHistory && !bCtrl && AutoCompleteIndices.Num() > 1)
 		{
-			if (--AutoCompleteIndex < 0)
+			if(AutoCompleteCursor > 0)
 			{
-				AutoCompleteIndex = AutoCompleteIndices.Num() - 1;
+				// move cursor within displayed region
+				--AutoCompleteCursor;
+			}
+			else
+			{
+				// can we scroll?
+				if(AutoCompleteIndex > 0)
+				{
+					--AutoCompleteIndex;
+				}
+				bAutoCompleteLocked = false;
 			}
-			bAutoCompleteLocked = false;
 		}
 		else
 		if ( HistoryBot >= 0 )
@@ -822,7 +847,8 @@ void UConsole::PostRender_Console_Typing(UCanvas* Canvas)
 		Canvas->DrawItem( ConsoleText, LeftPos+xl, ClipY-3.0f-yl );
 		Canvas->StrLen(Font, *ConsoleDefs::LeadingInputText, xl, yl );
 
-		int32 StartIdx = AutoCompleteIndex - 5;
+		int32 StartIdx = AutoCompleteIndex;
+
 		if (StartIdx < 0)
 		{
 			StartIdx = FMath::Max(0,AutoCompleteIndices.Num() + StartIdx);
@@ -834,8 +860,13 @@ void UConsole::PostRender_Console_Typing(UCanvas* Canvas)
 		// Set the background color/texture of the auto-complete section
 		ConsoleTile.SetColor( ConsoleDefs::AutocompleteBackgroundColor );
 		ConsoleTile.Texture = DefaultTexture_White->Resource;
-		for (int32 MatchIdx = 0; MatchIdx < 10; MatchIdx++)
+		for (int32 MatchIdx = 0; MatchIdx < MAX_AUTOCOMPLETION_LINES; MatchIdx++)
 		{
+			if(Idx >= AutoCompleteIndices.Num())
+			{
+				break;
+			}
+
 			const FAutoCompleteCommand &Cmd = AutoCompleteList[AutoCompleteIndices[Idx]];
 			OutStr = Cmd.Desc;
 
@@ -850,7 +881,7 @@ void UConsole::PostRender_Console_Typing(UCanvas* Canvas)
 
 			ConsoleTile.Size = FVector2D( info_xl, info_yl );
 			Canvas->DrawItem( ConsoleTile, LeftPos + xl, y );
-			if (Idx == AutoCompleteIndex)
+			if (Idx == AutoCompleteIndex + AutoCompleteCursor)
 			{
 				ConsoleText.SetColor( ConsoleDefs::AutocompleteSuggestionColor );
 			}
@@ -860,21 +891,14 @@ void UConsole::PostRender_Console_Typing(UCanvas* Canvas)
 			}
 			ConsoleText.Text = FText::FromString( OutStr );
 			Canvas->DrawItem( ConsoleText, LeftPos + xl, y );
-			if (++Idx >= AutoCompleteIndices.Num())
-			{
-				Idx = 0;
-			}
+
+			++Idx;
 			y -= yl;
-			// break out if we loop on lists < 10
-			if (Idx == StartIdx)
-			{
-				break;
-			}
 		}
 		// Display a message if there were more matches
-		if (AutoCompleteIndices.Num() >= 10)
+		if (AutoCompleteIndices.Num() >= MAX_AUTOCOMPLETION_LINES)
 		{
-			OutStr = FString::Printf(TEXT("[%i more matches]"), (AutoCompleteIndices.Num() - 10 + 1));
+			OutStr = FString::Printf(TEXT("[%i more matches]"), (AutoCompleteIndices.Num() - MAX_AUTOCOMPLETION_LINES + 1));
 			Canvas->StrLen(Font, OutStr, info_xl, info_yl);
 			ConsoleTile.Size = FVector2D( info_xl, info_yl );
 			Canvas->DrawItem( ConsoleTile, LeftPos + xl, y );
