#UE4:
* disable closing of game console when using tab (autocompletion key)

diff --git a/Engine/Source/Runtime/Engine/Private/UserInterface/Console.cpp b/Engine/Source/Runtime/Engine/Private/UserInterface/Console.cpp
--- a/Engine/Source/Runtime/Engine/Private/UserInterface/Console.cpp
+++ b/Engine/Source/Runtime/Engine/Private/UserInterface/Console.cpp
@@ -992,11 +992,6 @@ bool UConsole::InputKey_Open( int32 ControllerId, FKey Key, EInputEvent Event, f
 			SetCursorPos(TypedStr.Len());
 			bAutoCompleteLocked = true;
 		}
-		else
-		{
-			FakeGotoState(NAME_None);
-			bCaptureKeyInput = true;
-		}
 		return true;
 	}
 	else if( Key == EKeys::Enter && Event == IE_Released )
