#ue4
- Exposed ShowPlatformSpecificAchievementsScreen and ShowPlatformSpecificLeaderboardScreen for use on iOS and Android

diff --git a/Engine/Config/BaseEngine.ini b/Engine/Config/BaseEngine.ini
--- a/Engine/Config/BaseEngine.ini
+++ b/Engine/Config/BaseEngine.ini
@@ -845,6 +845,9 @@ bUseStreamingPause=false
 +K2FieldRedirects=(OldFieldName="GameplayStatics.PopSoundMode",NewFieldName="GameplayStatics.PopSoundMixModifier")
 +K2FieldRedirects=(OldFieldName="GameplayStatics.ClearSoundMode",NewFieldName="GameplayStatics.ClearSoundMixModifiers")
 
++K2FieldRedirects=(OldFieldName="KismetSystemLibrary.EXPERIMENTAL_ShowGameCenterLeaderboard",NewFieldName="KismetSystemLibrary.ShowPlatformSpecificLeaderboardScreen")
+
+
 +K2ParamRedirects=(NodeName="/Script/BlueprintGraph.K2Node_CallFunction", OldParamName="SceneComponent.SetVisibility.bSetEnabled", NewParamName="bNewVisibility")
 +K2ParamRedirects=(NodeName="/Script/BlueprintGraph.K2Node_CallFunction", OldParamName="SceneComponent.SetVisibility.bInEnabled", NewParamName="bNewVisibility")
 +K2ParamRedirects=(NodeName="/Script/BlueprintGraph.K2Node_CallFunction", OldParamName="ParticleSystemComponent.SetActive.bNowActive", NewParamName="bNewActive")
diff --git a/Engine/Source/Runtime/Engine/Classes/Kismet/KismetSystemLibrary.h b/Engine/Source/Runtime/Engine/Classes/Kismet/KismetSystemLibrary.h
--- a/Engine/Source/Runtime/Engine/Classes/Kismet/KismetSystemLibrary.h
+++ b/Engine/Source/Runtime/Engine/Classes/Kismet/KismetSystemLibrary.h
@@ -997,8 +997,14 @@ class UKismetSystemLibrary : public UBlueprintFunctionLibrary
 	static void EXPERIMENTAL_CloseAdBanner();
 
 	/**
-	 * Displays the built-in iOS leaderboard GUI (iOS only; this function will be renamed or moved in a future release)
-	*/
+	 * Displays the built-in leaderboard GUI (iOS and Android only; this function may be renamed or moved in a future release)
+	 */
+	UFUNCTION(BlueprintCallable, Category = "Utilities|Platform")
+	static void ShowPlatformSpecificLeaderboardScreen(const FString& CategoryName);
+
+	/**
+	 * Displays the built-in achievements GUI (iOS and Android only; this function may be renamed or moved in a future release)
+	 */
 	UFUNCTION(BlueprintCallable, Category = "Utilities|Platform")
-	static void EXPERIMENTAL_ShowGameCenterLeaderboard(const FString& CategoryName);
+	static void ShowPlatformSpecificAchievementsScreen();
 };
diff --git a/Engine/Source/Runtime/Engine/Private/KismetSystemLibrary.cpp b/Engine/Source/Runtime/Engine/Private/KismetSystemLibrary.cpp
--- a/Engine/Source/Runtime/Engine/Private/KismetSystemLibrary.cpp
+++ b/Engine/Source/Runtime/Engine/Private/KismetSystemLibrary.cpp
@@ -2241,7 +2241,7 @@ void UKismetSystemLibrary::EXPERIMENTAL_CloseAdBanner()
 	}
 }
 
-void UKismetSystemLibrary::EXPERIMENTAL_ShowGameCenterLeaderboard(const FString& CategoryName)
+void UKismetSystemLibrary::ShowPlatformSpecificLeaderboardScreen(const FString& CategoryName)
 {
 #if PLATFORM_IOS
 	extern CORE_API void IOSShowLeaderboardUI(const FString& CategoryName);
@@ -2253,4 +2253,15 @@ void UKismetSystemLibrary::EXPERIMENTAL_ShowGameCenterLeaderboard(const FString&
 	extern FString GetLeaderboardID(const FString&);
 	AndroidThunkCpp_ShowLeaderboard(GetLeaderboardID(CategoryName));
 #endif
-}
\ No newline at end of file
+}
+
+void UKismetSystemLibrary::ShowPlatformSpecificAchievementsScreen()
+{
+#if PLATFORM_IOS
+	extern CORE_API void IOSShowAchievementsUI();
+	IOSShowAchievementsUI();
+#elif PLATFORM_ANDROID
+	extern void AndroidThunkCpp_ShowAchievements();
+	AndroidThunkCpp_ShowAchievements();
+#endif
+}
