Fix CullDistanceVolume

diff --git a/Engine/Source/Runtime/Engine/Private/CullDistanceVolume.cpp b/Engine/Source/Runtime/Engine/Private/CullDistanceVolume.cpp
--- a/Engine/Source/Runtime/Engine/Private/CullDistanceVolume.cpp
+++ b/Engine/Source/Runtime/Engine/Private/CullDistanceVolume.cpp
@@ -74,47 +74,41 @@ void ACullDistanceVolume::GetPrimitiveMaxDrawDistances(TMap<UPrimitiveComponent*
 	// Nothing to do if there is no brush component or no cull distances are set
 	if( BrushComponent && CullDistances.Num() > 0 && bEnabled )
 	{
-		// Test center of mesh bounds to see whether it is encompassed by volume
-		// and propagate cull distance if it is.
-		for( FActorIterator It(GetWorld()); It; ++It )
+		for (auto It(OutCullDistances.CreateIterator()); It; ++It)
 		{
-			AActor* Owner = *It;
-			UPrimitiveComponent* PrimitiveComponent = Owner->FindComponentByClass<UPrimitiveComponent>();
-			if (PrimitiveComponent)
+			UPrimitiveComponent* PrimitiveComponent = It.Key();
+
+			// Check whether primitive can be affected by cull distance volumes.
+			if( ACullDistanceVolume::CanBeAffectedByVolumes( PrimitiveComponent ) )
 			{
-				// Check whether primitive can be affected by cull distance volumes.
-				if( Owner && ACullDistanceVolume::CanBeAffectedByVolumes( PrimitiveComponent ) )
-				{
-					// Check whether primitive supports cull distance volumes and its center point is being encompassed by this volume.
-					if( EncompassesPoint( Owner->GetActorLocation() ) )
-					{		
-						// Find best match in CullDistances array.
-						float PrimitiveSize			= PrimitiveComponent->Bounds.SphereRadius * 2;
-						float CurrentError			= FLT_MAX;
-						float CurrentCullDistance	= 0;
-						for( int32 CullDistanceIndex=0; CullDistanceIndex<CullDistances.Num(); CullDistanceIndex++ )
+				// Check whether primitive supports cull distance volumes and its center point is being encompassed by this volume.
+				if( EncompassesPoint( PrimitiveComponent->GetComponentLocation() ) )
+				{		
+					// Find best match in CullDistances array.
+					float PrimitiveSize			= PrimitiveComponent->Bounds.SphereRadius * 2;
+					float CurrentError			= FLT_MAX;
+					float CurrentCullDistance	= 0;
+					for( int32 CullDistanceIndex=0; CullDistanceIndex<CullDistances.Num(); CullDistanceIndex++ )
+					{
+						const FCullDistanceSizePair& CullDistancePair = CullDistances[CullDistanceIndex];
+						if( FMath::Abs( PrimitiveSize - CullDistancePair.Size ) < CurrentError )
 						{
-							const FCullDistanceSizePair& CullDistancePair = CullDistances[CullDistanceIndex];
-							if( FMath::Abs( PrimitiveSize - CullDistancePair.Size ) < CurrentError )
-							{
-								CurrentError		= FMath::Abs( PrimitiveSize - CullDistancePair.Size );
-								CurrentCullDistance = CullDistancePair.CullDistance;
-							}
+							CurrentError		= FMath::Abs( PrimitiveSize - CullDistancePair.Size );
+							CurrentCullDistance = CullDistancePair.CullDistance;
 						}
+					}
 
-						float* CurrentDistPtr = OutCullDistances.Find(PrimitiveComponent);
-						check(CurrentDistPtr);
+					float& CullDistance = It.Value();
 
-						// LD or other volume specified cull distance, use minimum of current and one used for this volume.
-						if( *CurrentDistPtr > 0 )
-						{
-							OutCullDistances.Add(PrimitiveComponent, FMath::Min( *CurrentDistPtr, CurrentCullDistance ));
-						}
-						// LD didn't specify cull distance, use current setting directly.
-						else
-						{
-							OutCullDistances.Add(PrimitiveComponent, CurrentCullDistance);
-						}
+					// LD or other volume specified cull distance, use minimum of current and one used for this volume.
+					if (CullDistance > 0)
+					{
+						CullDistance = FMath::Min(CullDistance, CurrentCullDistance);
+					}
+					// LD didn't specify cull distance, use current setting directly.
+					else
+					{
+						CullDistance = CurrentCullDistance;
 					}
 				}
 			}
diff --git a/Engine/Source/Runtime/Engine/Private/World.cpp b/Engine/Source/Runtime/Engine/Private/World.cpp
--- a/Engine/Source/Runtime/Engine/Private/World.cpp
+++ b/Engine/Source/Runtime/Engine/Private/World.cpp
@@ -1058,26 +1058,31 @@ void UWorld::UpdateCullDistanceVolumes()
 	{
 		SCOPE_SECONDS_COUNTER(Duration);
 
+		TArray<ACullDistanceVolume*> CullDistanceVolumes;
+
 		// Establish base line of LD specified cull distances.
 		for( FActorIterator It(this); It; ++It )
 		{
-			auto PrimitiveComponent = It->FindComponentByClass<UPrimitiveComponent>();
-			if (PrimitiveComponent)
+			TArray<UPrimitiveComponent*> PrimitiveComponents;
+			It->GetComponents(PrimitiveComponents);
+			for (UPrimitiveComponent* PrimitiveComponent : PrimitiveComponents)
 			{
 				CompToNewMaxDrawMap.Add(PrimitiveComponent, PrimitiveComponent->LDMaxDrawDistance);
 			}
-		}
 
-		// Iterate over all cull distance volumes and get new cull distances.
-		for( FActorIterator It(this); It; ++It )
-		{
 			ACullDistanceVolume* CullDistanceVolume = Cast<ACullDistanceVolume>(*It);
-			if( CullDistanceVolume )
+			if (CullDistanceVolume)
 			{
-				CullDistanceVolume->GetPrimitiveMaxDrawDistances(CompToNewMaxDrawMap);
+				CullDistanceVolumes.Add(CullDistanceVolume);
 			}
 		}
 
+		// Iterate over all cull distance volumes and get new cull distances.
+		for (ACullDistanceVolume* CullDistanceVolume : CullDistanceVolumes)
+		{
+			CullDistanceVolume->GetPrimitiveMaxDrawDistances(CompToNewMaxDrawMap);
+		}
+
 		// Finally, go over all primitives, and see if they need to change.
 		// Only if they do do we reregister them, as thats slow.
 		for ( TMap<UPrimitiveComponent*,float>::TIterator It(CompToNewMaxDrawMap); It; ++It )
