L10N: Fixing localization errors. Missing #undefines, conflicted loc keys and comments being incorrectly detected as expected LOCTEXT macros.

diff --git a/Engine/Source/Developer/SettingsEditor/Private/Widgets/SSettingsEditorCheckoutNotice.cpp b/Engine/Source/Developer/SettingsEditor/Private/Widgets/SSettingsEditorCheckoutNotice.cpp
--- a/Engine/Source/Developer/SettingsEditor/Private/Widgets/SSettingsEditorCheckoutNotice.cpp
+++ b/Engine/Source/Developer/SettingsEditor/Private/Widgets/SSettingsEditorCheckoutNotice.cpp
@@ -122,3 +122,5 @@ FText SSettingsEditorCheckoutNotice::HandleSccUnavailableTextBlockText( ) const
 
 	return FText::GetEmpty();
 }
+
+#undef LOCTEXT_NAMESPACE
\ No newline at end of file
diff --git a/Engine/Source/Runtime/Engine/Private/AI/EnvironmentQuery/Generators/EnvQueryGenerator_OnCircle.cpp b/Engine/Source/Runtime/Engine/Private/AI/EnvironmentQuery/Generators/EnvQueryGenerator_OnCircle.cpp
--- a/Engine/Source/Runtime/Engine/Private/AI/EnvironmentQuery/Generators/EnvQueryGenerator_OnCircle.cpp
+++ b/Engine/Source/Runtime/Engine/Private/AI/EnvironmentQuery/Generators/EnvQueryGenerator_OnCircle.cpp
@@ -258,7 +258,7 @@ FText UEnvQueryGenerator_OnCircle::GetDescriptionDetails() const
 		FFormatNamedArguments ProjArgs;
 		ProjArgs.Add(TEXT("Description"), Desc);
 		ProjArgs.Add(TEXT("ProjectionDescription"), ProjDesc);
-		Desc = FText::Format(LOCTEXT("DescriptionWithProjection", "{Description}, {ProjectionDescription}"), ProjArgs);
+		Desc = FText::Format(LOCTEXT("OnCircle_DescriptionWithProjection", "{Description}, {ProjectionDescription}"), ProjArgs);
 	}
 
 	return Desc;
diff --git a/Engine/Source/Runtime/Engine/Private/AI/EnvironmentQuery/Generators/EnvQueryGenerator_SimpleGrid.cpp b/Engine/Source/Runtime/Engine/Private/AI/EnvironmentQuery/Generators/EnvQueryGenerator_SimpleGrid.cpp
--- a/Engine/Source/Runtime/Engine/Private/AI/EnvironmentQuery/Generators/EnvQueryGenerator_SimpleGrid.cpp
+++ b/Engine/Source/Runtime/Engine/Private/AI/EnvironmentQuery/Generators/EnvQueryGenerator_SimpleGrid.cpp
@@ -101,7 +101,7 @@ FText UEnvQueryGenerator_SimpleGrid::GetDescriptionDetails() const
 		FFormatNamedArguments ProjArgs;
 		ProjArgs.Add(TEXT("Description"), Desc);
 		ProjArgs.Add(TEXT("ProjectionDescription"), ProjDesc);
-		Desc = FText::Format(LOCTEXT("DescriptionWithProjection", "{Description}, {ProjectionDescription}"), ProjArgs);
+		Desc = FText::Format(LOCTEXT("SimpleGridDescriptionWithProjection", "{Description}, {ProjectionDescription}"), ProjArgs);
 	}
 
 	return Desc;
diff --git a/Engine/Source/Runtime/Engine/Public/ShowFlagsValues.inl b/Engine/Source/Runtime/Engine/Public/ShowFlagsValues.inl
--- a/Engine/Source/Runtime/Engine/Public/ShowFlagsValues.inl
+++ b/Engine/Source/Runtime/Engine/Public/ShowFlagsValues.inl
@@ -1,9 +1,9 @@
 // usage
 //
 // General purpose showflag (always variable):
-// SHOWFLAG_ALWAYS_ACCESSIBLE( <showflag name>, <showflag group>, <LOCTEXT stuff>)
+// SHOWFLAG_ALWAYS_ACCESSIBLE( <showflag name>, <showflag group>, <Localized TEXT stuff>)
 // Fixed in shipping builds:
-// SHOWFLAG_FIXED_IN_SHIPPING( <showflag name>, <fixed bool>, <showflag group>, <LOCTEXT stuff>)
+// SHOWFLAG_FIXED_IN_SHIPPING( <showflag name>, <fixed bool>, <showflag group>, <Localized TEXT stuff>)
 
 #ifndef SHOWFLAG_ALWAYS_ACCESSIBLE
 #error SHOWFLAG_ALWAYS_ACCESSIBLE macro is undefined.
@@ -259,7 +259,7 @@ SHOWFLAG_ALWAYS_ACCESSIBLE(PrecomputedVisibilityCells, SFG_Visualize, LOCTEXT("P
 /** Visualize volume lighting samples used for GI on dynamic objects */
 SHOWFLAG_ALWAYS_ACCESSIBLE(VolumeLightingSamples, SFG_Visualize, LOCTEXT("VolumeLightingSamplesSF", "Volume Lighting Samples"))
 /** needed for VMI_LpvLightingViewMode, Whether to show only LPV lighting*/
-SHOWFLAG_ALWAYS_ACCESSIBLE(LpvLightingOnly, SFG_Hidden, LOCTEXT("VisualizeLPVSF", "Visualize LPV"))
+SHOWFLAG_ALWAYS_ACCESSIBLE(LpvLightingOnly, SFG_Hidden, LOCTEXT("VisualizeLPVSF_ViewMode", "Visualize LPV"))
 
 #undef SHOWFLAG_ALWAYS_ACCESSIBLE
 #undef SHOWFLAG_FIXED_IN_SHIPPING
diff --git a/Engine/Source/Runtime/UMG/Private/WidgetBlueprint.cpp b/Engine/Source/Runtime/UMG/Private/WidgetBlueprint.cpp
--- a/Engine/Source/Runtime/UMG/Private/WidgetBlueprint.cpp
+++ b/Engine/Source/Runtime/UMG/Private/WidgetBlueprint.cpp
@@ -35,4 +35,6 @@ bool UWidgetBlueprint::ValidateGeneratedClass(const UClass* InClass)
 	//}
 
 	return Result;
-}
\ No newline at end of file
+}
+
+#undef LOCTEXT_NAMESPACE 
\ No newline at end of file
diff --git a/Engine/Source/Runtime/UMG/Private/WidgetBlueprintGeneratedClass.cpp b/Engine/Source/Runtime/UMG/Private/WidgetBlueprintGeneratedClass.cpp
--- a/Engine/Source/Runtime/UMG/Private/WidgetBlueprintGeneratedClass.cpp
+++ b/Engine/Source/Runtime/UMG/Private/WidgetBlueprintGeneratedClass.cpp
@@ -72,4 +72,6 @@ void UWidgetBlueprintGeneratedClass::CreateComponentsForActor(AActor* Actor) con
 
 		Widget->RegisterComponent();
 	}
-}
\ No newline at end of file
+}
+
+#undef LOCTEXT_NAMESPACE
\ No newline at end of file
