making GetOwner() exposed to blueprints

diff --git a/Engine/Source/Runtime/Engine/Classes/GameFramework/Actor.h b/Engine/Source/Runtime/Engine/Classes/GameFramework/Actor.h
--- a/Engine/Source/Runtime/Engine/Classes/GameFramework/Actor.h
+++ b/Engine/Source/Runtime/Engine/Classes/GameFramework/Actor.h
@@ -1389,7 +1389,8 @@ public:
 	/** Set the owner for this Actor */
 	void SetOwner( AActor* NewOwner );
 
-	AActor* GetOwner() const { return Owner; }
+	UFUNCTION(BlueprintCallable, Category=Actor)
+	AActor* GetOwner() const;
 
 	/**
 	 * This will check to see if the Actor is still in the world.  It will check things like
diff --git a/Engine/Source/Runtime/Engine/Private/Actor.cpp b/Engine/Source/Runtime/Engine/Private/Actor.cpp
--- a/Engine/Source/Runtime/Engine/Private/Actor.cpp
+++ b/Engine/Source/Runtime/Engine/Private/Actor.cpp
@@ -1083,6 +1083,11 @@ void AActor::SetOwner( AActor *NewOwner )
 	}
 }
 
+AActor* AActor::GetOwner() const
+{ 
+	return Owner; 
+}
+
 void AActor::AttachRootComponentTo(USceneComponent* InParent, FName InSocketName, EAttachLocation::Type AttachLocationType /*= EAttachLocation::KeepRelativeOffset */)
 {
 	if(RootComponent && InParent)
