#UE4:
* fixed several issues in game console, unified code path, fixed auto completion when console is open, draw window rect around autocompletion entries to be more readable

diff --git a/Engine/Source/Runtime/Engine/Classes/Engine/Console.h b/Engine/Source/Runtime/Engine/Classes/Engine/Console.h
--- a/Engine/Source/Runtime/Engine/Classes/Engine/Console.h
+++ b/Engine/Source/Runtime/Engine/Classes/Engine/Console.h
@@ -216,19 +216,8 @@ class ENGINE_API UConsole : public UObject, public FOutputDevice
 	virtual bool InputChar_Typing( int32 ControllerId, const FString& Unicode );
 	
 	/**
-	 * Process an input key event routed through unrealscript from another object. This method is assigned as the value for the
-	 * OnRecievedNativeInputKey delegate so that native input events are routed to this unrealscript function.
-	 *
-	 * @param	ControllerId	the controller that generated this input key event
-	 * @param	Key				the key which for an event occurred for (EKeys::Up, EKeys::Down, etc.)
-	 * @param	EventType		the type of event which occurred (pressed, released, etc.)
-	 * @param	AmountDepressed	for analog keys, the depression percent.
-	 *
-	 * @return	true to consume the key event, false to pass it on.
+	 * perform rendering of the console on the canvas
 	 */
-	virtual bool InputKey_Typing( int32 ControllerId, FKey Key, EInputEvent Event, float AmountDepressed = 1.f, bool bGamepad = false );
-	
-	/** perform rendering of the console on the canvas */
 	virtual void PostRender_Console_Typing(class UCanvas* Canvas);
 	
 	/** Perform actions on transition to Typing state */
@@ -255,7 +244,9 @@ class ENGINE_API UConsole : public UObject, public FOutputDevice
 	 */
 	virtual bool InputKey_Open( int32 ControllerId, FKey Key, EInputEvent Event, float AmountDepressed = 1.f, bool bGamepad = false );
 	
-	/** perform rendering of the console on the canvas */
+	/** 
+	 * perform rendering of the console on the canvas
+	 */
 	virtual void PostRender_Console_Open(class UCanvas* Canvas);
 
 
@@ -267,19 +258,6 @@ class ENGINE_API UConsole : public UObject, public FOutputDevice
 	virtual bool InputAxis(int32 ControllerId, FKey Key, float Delta, float DeltaTime, int32 NumSamples=1, bool bGamepad=false) { return false; };
 	virtual bool InputTouch(int32 ControllerId, uint32 Handle, ETouchType::Type Type, const FVector2D& TouchLocation, FDateTime DeviceTimestamp, uint32 TouchpadIndex) { return false; }
 
-	/**
-	 * Process an input key event routed through unrealscript from another object. This method is assigned as the value for the
-	 * OnRecievedNativeInputKey delegate so that native input events are routed to this unrealscript function.
-	 *
-	 * @param	ControllerId	the controller that generated this input key event
-	 * @param	Key				the name of the key which an event occured for (KEY_Up, KEY_Down, etc.)
-	 * @param	EventType		the type of event which occured (pressed, released, etc.)
-	 * @param	AmountDepressed	for analog keys, the depression percent.
-	 *
-	 * @return	true to consume the key event, false to pass it on.
-	 */
-	virtual bool InputKey_Global( int32 ControllerId, FKey Key, EInputEvent Event, float AmountDepressed = 1.f, bool bGamepad = false );
-
 	/** render to the canvas based on the console state */
 	virtual void PostRender_Console(class UCanvas* Canvas);
 
@@ -290,6 +268,8 @@ class ENGINE_API UConsole : public UObject, public FOutputDevice
 
 private:
 
+	bool InputKey_InputLine( int32 ControllerId, FKey Key, EInputEvent Event, float AmountDepressed = 1.f, bool bGamepad = false );
+
 	// interface FOutputDevice
 	virtual void Serialize( const TCHAR* V, ELogVerbosity::Type Verbosity, const class FName& Category );
 	// expose the base class other Serialize function (clang will give "error : 'UConsole::Serialize' hides overloaded virtual function" without this)
@@ -301,6 +281,8 @@ private:
 	 */
 	void OutputTextLine(const FString& Text);
 	
+	void PostRender_InputLine(class UCanvas* Canvas, FIntPoint UserInputLinePos);
+
 	/**
 	* Searches console command history and removes any entries matching the specified command.
 	* @param Command - The command to search for and purge from the history.
diff --git a/Engine/Source/Runtime/Engine/Private/UserInterface/Console.cpp b/Engine/Source/Runtime/Engine/Private/UserInterface/Console.cpp
--- a/Engine/Source/Runtime/Engine/Private/UserInterface/Console.cpp
+++ b/Engine/Source/Runtime/Engine/Private/UserInterface/Console.cpp
@@ -533,8 +533,7 @@ bool UConsole::InputChar_Typing( int32 ControllerId, const FString& Unicode )
 	return true;
 }
 
-
-bool UConsole::InputKey_Typing( int32 ControllerId, FKey Key, EInputEvent Event, float AmountDepressed, bool bGamepad)
+bool UConsole::InputKey_InputLine( int32 ControllerId, FKey Key, EInputEvent Event, float AmountDepressed, bool bGamepad)
 {
 	//`log(`location@`showvar(Key));
 
@@ -579,6 +578,11 @@ bool UConsole::InputKey_Typing( int32 ControllerId, FKey Key, EInputEvent Event,
 			FakeGotoState(NAME_None);
 			bCaptureKeyInput = true;
 		}
+		else if (ConsoleState == NAME_None)
+		{
+			FakeGotoState(NAME_Typing);
+			bCaptureKeyInput = true;
+		}
 		return true;
 	}
 	else if (Key == EKeys::Tab && Event == IE_Pressed)
@@ -619,10 +623,6 @@ bool UConsole::InputKey_Typing( int32 ControllerId, FKey Key, EInputEvent Event,
 
 		return true;
 	}
-	else if( InputKey_Global( ControllerId, Key, Event, AmountDepressed, bGamepad ) )
-	{
-		return true;
-	}
 	else if( Event != IE_Pressed && Event != IE_Repeat )
 	{
 		if( !bGamepad )
@@ -774,7 +774,7 @@ bool UConsole::InputKey_Typing( int32 ControllerId, FKey Key, EInputEvent Event,
 		return true;
 	}
 
-	return true;
+	return false;
 }
 
 
@@ -822,102 +822,8 @@ void UConsole::PostRender_Console_Typing(UCanvas* Canvas)
 		ClipY = ClipY * 0.60;
 		LeftPos = 150;
 	}
-	
-	// Background
-	FCanvasTileItem ConsoleTile( FVector2D( LeftPos,ClipY-6-yl ), DefaultTexture_Black->Resource, FVector2D( ClipX, yl+6 ), FVector2D( 0.0f, 0.0f), FVector2D( 1.0f, 1.0f ), FLinearColor::White );
-	Canvas->DrawItem( ConsoleTile );
-
-	// Separator line
-	ConsoleTile.SetColor( ConsoleDefs::BorderColor );
-	ConsoleTile.Texture = DefaultTexture_White->Resource;
-	ConsoleTile.Size = FVector2D( ClipX, 2.0f );
-	Canvas->DrawItem( ConsoleTile );
-
-	// Currently typed string
-	FText Str = FText::FromString( OutStr );
-	FCanvasTextItem ConsoleText( FVector2D( LeftPos,ClipY-3-yl ), Str , GEngine->GetSmallFont(), ConsoleDefs::InputTextColor );		
-	Canvas->DrawItem( ConsoleText );
-
-	// draw the remaining text for matching auto-complete
-	if (AutoCompleteIndices.Num() > 0)
-	{
-		int32 Idx = AutoCompleteIndices[AutoCompleteIndex];
-		ConsoleText.SetColor( ConsoleDefs::AutocompletePartialSuggestionColor );
-		ConsoleText.Text = FText::FromString( AutoCompleteList[Idx].Command.Right(AutoCompleteList[Idx].Command.Len() - TypedStr.Len()) );
-		Canvas->DrawItem( ConsoleText, LeftPos+xl, ClipY-3.0f-yl );
-		Canvas->StrLen(Font, *ConsoleDefs::LeadingInputText, xl, yl );
-
-		int32 StartIdx = AutoCompleteIndex;
-
-		if (StartIdx < 0)
-		{
-			StartIdx = FMath::Max(0,AutoCompleteIndices.Num() + StartIdx);
-		}
-		Idx = StartIdx;
-		float y = ClipY - 6.0f - (yl * 2.0f);
-		float info_xl, info_yl;
-
-		// Set the background color/texture of the auto-complete section
-		ConsoleTile.SetColor( ConsoleDefs::AutocompleteBackgroundColor );
-		ConsoleTile.Texture = DefaultTexture_White->Resource;
-		for (int32 MatchIdx = 0; MatchIdx < MAX_AUTOCOMPLETION_LINES; MatchIdx++)
-		{
-			if(Idx >= AutoCompleteIndices.Num())
-			{
-				break;
-			}
-
-			const FAutoCompleteCommand &Cmd = AutoCompleteList[AutoCompleteIndices[Idx]];
-			OutStr = Cmd.Desc;
-
-			if(OutStr.IsEmpty())
-			{
-				// no Description means we display the Command directly, without that the line would be empty (happens for ConsoleVariables and some ManualAutoCompleteList)
-				OutStr = Cmd.Command;
-			}
-
-			Canvas->StrLen(Font, OutStr, info_xl, info_yl);
-			y -= info_yl - yl;
-
-			ConsoleTile.Size = FVector2D( info_xl, info_yl );
-			Canvas->DrawItem( ConsoleTile, LeftPos + xl, y );
-			if (Idx == AutoCompleteIndex + AutoCompleteCursor)
-			{
-				ConsoleText.SetColor( ConsoleDefs::AutocompleteSuggestionColor );
-			}
-			else
-			{
-				ConsoleText.SetColor( ConsoleDefs::AutocompletePartialSuggestionColor );
-			}
-			ConsoleText.Text = FText::FromString( OutStr );
-			Canvas->DrawItem( ConsoleText, LeftPos + xl, y );
-
-			++Idx;
-			y -= yl;
-		}
-		// Display a message if there were more matches
-		if (AutoCompleteIndices.Num() >= MAX_AUTOCOMPLETION_LINES)
-		{
-			OutStr = FString::Printf(TEXT("[%i more matches]"), (AutoCompleteIndices.Num() - MAX_AUTOCOMPLETION_LINES + 1));
-			Canvas->StrLen(Font, OutStr, info_xl, info_yl);
-			ConsoleTile.Size = FVector2D( info_xl, info_yl );
-			Canvas->DrawItem( ConsoleTile, LeftPos + xl, y );
-			ConsoleText.SetColor( ConsoleDefs::AutocompleteSuggestionColor );
-			ConsoleText.Text = FText::FromString( OutStr );
-			Canvas->DrawItem( ConsoleText, LeftPos + xl, y );
-		}
-	}
-
-	// determine the cursor position
-	OutStr = FString::Printf(TEXT("%s%s"), *ConsoleDefs::LeadingInputText, *TypedStr.Left(TypedStrPos));
-	Canvas->StrLen(Font, OutStr,xl,yl);
-
-	// draw the cursor
-	ConsoleText.SetColor( ConsoleDefs::CursorColor );
-	OutStr = FString::Printf( TEXT("_") );
-	ConsoleText.Text = FText::FromString( OutStr );
-	Canvas->DrawItem( ConsoleText, LeftPos + xl, ClipY-1.0f-yl );
 
+	PostRender_InputLine(Canvas, FIntPoint(LeftPos, ClipY));
 }
 
 void UConsole::BeginState_Typing(FName PreviousStateName)
@@ -951,195 +857,7 @@ bool UConsole::InputChar_Open( int32 ControllerId, const FString& Unicode )
 
 bool UConsole::InputKey_Open( int32 ControllerId, FKey Key, EInputEvent Event, float AmountDepressed, bool bGamepad)
 {
-	if (Event == IE_Pressed)
-	{
-		bCaptureKeyInput=false;
-	}
-
-	if (ProcessControlKey(Key, Event))
-	{
-		return true;
-	}
-	else if( bGamepad )
-	{
-		return false;
-	}
-	else if( Key == EKeys::Escape && Event == IE_Released )
-	{
-		if( !TypedStr.IsEmpty() )
-		{
-			SetInputText(TEXT(""));
-			SetCursorPos(0);
-			HistoryCur = HistoryTop;
-			return true;
-		}
-		else
-		{
-			FakeGotoState( NAME_None );
-		}
-	}
-	else if( Key==GetDefault<UInputSettings>()->ConsoleKey && Event == IE_Pressed )
-	{
-		FakeGotoState(NAME_None);
-		bCaptureKeyInput = true;
-		return true;
-	}
-	else if (Key == EKeys::Tab && Event == IE_Pressed)
-	{
-		if (AutoCompleteIndices.Num() > 0 && !bAutoCompleteLocked)
-		{
-			TypedStr = AutoCompleteList[AutoCompleteIndices[0]].Command;
-			SetCursorPos(TypedStr.Len());
-			bAutoCompleteLocked = true;
-		}
-		return true;
-	}
-	else if( Key == EKeys::Enter && Event == IE_Released )
-	{
-		if( !TypedStr.IsEmpty() )
-		{
-			// Make a local copy of the string.
-			FString Temp=TypedStr;
-			SetInputText("");
-			SetCursorPos(0);
-
-			if (FCString::Stricmp(*Temp, TEXT("cls")) == 0)
-			{
-				ClearOutput();
-			}
-			else
-			{
-				ConsoleCommand(Temp);
-			}
-
-			UpdateCompleteIndices();
-		}
-
-		return true;
-	}
-	else if( InputKey_Global( ControllerId, Key, Event, AmountDepressed, bGamepad ) )
-	{
-		return true;
-	}
-	else if( Event != IE_Pressed && Event != IE_Repeat )
-	{
-		if( !bGamepad )
-		{
-			return	Key != EKeys::LeftMouseButton
-				&&	Key != EKeys::MiddleMouseButton
-				&&	Key != EKeys::RightMouseButton;
-		}
-		return false;
-	}
-	else if( Key== EKeys::Up )
-	{
-
-		if (!bCtrl)
-		{
-
-			if ( HistoryBot >= 0 )
-			{
-				if (HistoryCur == HistoryBot)
-					HistoryCur = HistoryTop;
-				else
-				{
-					HistoryCur--;
-					if (HistoryCur<0)
-						HistoryCur = MAX_HISTORY_ENTRIES-1;
-				}
-
-				SetInputText(History[HistoryCur]);
-				SetCursorPos(History[HistoryCur].Len());
-			}
-		}
-		else
-		{
-			if (SBPos<Scrollback.Num()-1)
-			{
-				SBPos++;
-
-				if (SBPos>=Scrollback.Num())
-				  SBPos = Scrollback.Num()-1;
-			}
-		}
-		return true;
-	}
-	else if( Key == EKeys::Down )
-	{
-		if (!bCtrl)
-		{
-			if ( HistoryBot >= 0 )
-			{
-				if (HistoryCur == HistoryTop)
-					HistoryCur = HistoryBot;
-				else
-					HistoryCur = (HistoryCur+1) % MAX_HISTORY_ENTRIES;
-
-				SetInputText(History[HistoryCur]);
-				SetCursorPos(History[HistoryCur].Len());
-			}
-		}
-		else
-		{
-			if (SBPos>0)
-			{
-				SBPos--;
-				if (SBPos<0)
-					SBPos = 0;
-			}
-
-		}
-		return true;
-	}
-	else if( Key==EKeys::BackSpace )
-	{
-		if( TypedStrPos>0 )
-		{
-			SetInputText(FString::Printf(TEXT("%s%s"), *TypedStr.Left(TypedStrPos-1), *TypedStr.Right( TypedStr.Len() - TypedStrPos)));
-			SetCursorPos(TypedStrPos-1);
-			// unlock auto-complete (@todo - track the lock position so we don't bother unlocking under bogus cases)
-			bAutoCompleteLocked = false;
-		}
-
-		return true;
-	}
-	else if ( Key==EKeys::Delete )
-	{
-		if ( TypedStrPos < TypedStr.Len() )
-		{
-			SetInputText(FString::Printf(TEXT("%s%s"), *TypedStr.Left(TypedStrPos), *TypedStr.Right( TypedStr.Len() - TypedStrPos - 1)));
-		}
-		return true;
-	}
-	else if ( Key==EKeys::Left )
-	{
-		SetCursorPos(FMath::Max(0, TypedStrPos - 1));
-		return true;
-	}
-	else if ( Key==EKeys::Right )
-	{
-		SetCursorPos(FMath::Min(TypedStr.Len(), TypedStrPos + 1));
-		return true;
-	}
-	else if (bCtrl && Key==EKeys::Home)
-	{
-		SBPos = Scrollback.Num()-1;
-	}
-	else if ( Key==EKeys::Home )
-	{
-		SetCursorPos(0);
-		return true;
-	}
-	else if (bCtrl && Key==EKeys::End)
-	{
-		SBPos=0;
-	}
-	else if ( Key==EKeys::End )
-	{
-		SetCursorPos(TypedStr.Len());
-		return true;
-	}
-	else if ( Key==EKeys::PageUp || Key == EKeys::MouseScrollUp )
+	if ( Key==EKeys::PageUp || Key == EKeys::MouseScrollUp )
 	{
 		if (SBPos<Scrollback.Num()-1)
 		{
@@ -1149,7 +867,7 @@ bool UConsole::InputKey_Open( int32 ControllerId, FKey Key, EInputEvent Event, f
 				SBPos++;
 
 			if (SBPos>=Scrollback.Num())
-			  SBPos = Scrollback.Num()-1;
+				SBPos = Scrollback.Num()-1;
 		}
 
 		return true;
@@ -1170,8 +888,7 @@ bool UConsole::InputKey_Open( int32 ControllerId, FKey Key, EInputEvent Event, f
 		return true;
 	}
 
-
-	return true;
+	return false;
 }
 
 void UConsole::PostRender_Console_Open(UCanvas* Canvas)
@@ -1196,117 +913,58 @@ void UConsole::PostRender_Console_Open(UCanvas* Canvas)
 		LeftPos = 150;
 	}
 
-	// Background
-	FCanvasTileItem ConsoleTile( FVector2D( 0.0f, 0.0f ), DefaultTexture_Black->Resource, FVector2D( ClipX, Height+TopPos ), FVector2D( 0.0f, 0.0f), FVector2D( 1.0f, 1.0f ), FLinearColor::White );
-	Canvas->DrawItem( ConsoleTile );
-
-	// now render the typing region
-	FString OutStr = FString::Printf(TEXT("%s%s"), *ConsoleDefs::LeadingInputText, *TypedStr);
-
 	UFont* Font = GEngine->GetSmallFont();
+
 	// determine the height of the text
 	float xl, yl;
-	Canvas->StrLen(Font, OutStr,xl,yl);
-
-	// draw the top typing region border up + 12 pixels of buffer (for the borders and some space)
-	ConsoleTile.SetColor( ConsoleDefs::BorderColor );
-	ConsoleTile.Size = FVector2D( ClipX, 2.0f );
-	ConsoleTile.Texture = DefaultTexture_White->Resource;
-	Canvas->DrawItem( ConsoleTile, LeftPos, TopPos+Height-12-yl );		
-
-	// render the text that is being typed
-	// Currently typed string
-	FText Str = FText::FromString( OutStr );
-	FCanvasTextItem ConsoleText( FVector2D( LeftPos,TopPos+Height-5-yl ), Str , GEngine->GetSmallFont(), ConsoleDefs::InputTextColor );		
-	Canvas->DrawItem( ConsoleText );
-
-	// draw the remaining text for matching auto-complete
-	float y = 0.f;
-	if (AutoCompleteIndices.Num() > 0)
-	{
-		int32 idx = AutoCompleteIndices[0];
-		ConsoleText.SetColor( ConsoleDefs::AutocompletePartialSuggestionColor );
-		ConsoleText.Text = FText::FromString( AutoCompleteList[idx].Command.Right(AutoCompleteList[idx].Command.Len() - TypedStr.Len()) );
-		Canvas->DrawItem( ConsoleText, LeftPos+xl, TopPos+Height-5.0f-yl );
-
-		float info_xl, info_yl;
-		Canvas->StrLen(Font, *ConsoleDefs::LeadingInputText, xl, yl);
-		y = Height + 5.0f;
-
-		ConsoleTile.SetColor( ConsoleDefs::AutocompleteBackgroundColor );
-		ConsoleTile.Texture = DefaultTexture_White->Resource;
-		for (int32 MatchIdx = 0; MatchIdx < AutoCompleteIndices.Num() && MatchIdx < 10; MatchIdx++)
-		{
-			idx = AutoCompleteIndices[MatchIdx];
-
-			const FAutoCompleteCommand &Cmd = AutoCompleteList[idx];
-			FString Str = Cmd.Desc;
-
-			if(Str.IsEmpty())
-			{
-				// no Description means we display the Command directly, without that the line would be empty (happens for ConsoleVariables and some ManualAutoCompleteList)
-				Str = Cmd.Command;
-			}
+	Canvas->StrLen(Font, TEXT("M"),xl,yl);
 
-			Canvas->StrLen(Font, Str, info_xl, info_yl);
-			ConsoleTile.Size = FVector2D( info_xl, info_yl );
-			Canvas->DrawItem( ConsoleTile, LeftPos + xl, TopPos + y );
-
-			ConsoleText.SetColor( ConsoleDefs::AutocompleteSuggestionColor );
-			ConsoleText.Text = FText::FromString( Str );
-			Canvas->DrawItem( ConsoleText, LeftPos + xl, TopPos + y );
-			y += info_yl;
-		}
-	}
-
-	OutStr = FString::Printf(TEXT("%s%s"), *ConsoleDefs::LeadingInputText, *TypedStr.Left(TypedStrPos));
-
-	// Get position for the cursor
-	Canvas->StrLen(Font, OutStr,xl,yl);
-
-	// draw the cursor
-	ConsoleText.SetColor( ConsoleDefs::CursorColor );
-	OutStr = FString::Printf(TEXT("_"), *ConsoleDefs::LeadingInputText, *TypedStr.Left(TypedStrPos));
-	ConsoleText.Text = FText::FromString( OutStr );
-	Canvas->DrawItem( ConsoleText, LeftPos+xl,TopPos+Height-3.0f-yl );
+	// Background
+	FCanvasTileItem ConsoleTile( FVector2D( 0.0f, 0.0f ), DefaultTexture_Black->Resource, FVector2D( ClipX, Height + TopPos - yl), FVector2D( 0.0f, 0.0f), FVector2D( 1.0f, 1.0f ), FLinearColor::White );
+	Canvas->DrawItem( ConsoleTile );
 
 	// figure out which element of the scrollback buffer to should appear first (at the top of the screen)
 	int32 idx = SBHead - SBPos;
-	y = Height-16-(yl*2);
-
-	if (Scrollback.Num()==0)
-		return;
-
-	// change the text color to white
-	ConsoleText.SetColor( FLinearColor::White );
+	
+	float y = Height - yl;
 
-	// while we have enough room to draw another line and there are more lines to draw
-	while (y>yl && idx>=0)
+	if(Scrollback.Num())
 	{
-		float PenX;
-		float PenY;
-		float PenZ = 0.1f;
-		PenX = LeftPos;
-		PenY = TopPos+y;
+		FCanvasTextItem ConsoleText( FVector2D( LeftPos,TopPos+Height-5-yl ), FText::FromString(TEXT("")), GEngine->GetSmallFont(), ConsoleDefs::InputTextColor );		
+
+		// change the text color to white
+		ConsoleText.SetColor( FLinearColor::White );
 
-		// adjust the location for any word wrapping due to long text lines
-		if (idx < Scrollback.Num())
+		// while we have enough room to draw another line and there are more lines to draw
+		while (y > -yl && idx>=0)
 		{
-			float ScrollLineXL, ScrollLineYL;
-			Canvas->StrLen(Font, Scrollback[idx], ScrollLineXL, ScrollLineYL);
-			if (ScrollLineYL > yl)
+			float PenX;
+			float PenY;
+			float PenZ = 0.1f;
+			PenX = LeftPos;
+			PenY = TopPos+y;
+
+			// adjust the location for any word wrapping due to long text lines
+			if (idx < Scrollback.Num())
 			{
-				y -= (ScrollLineYL - yl);
-				PenX = LeftPos;
-				PenY = TopPos+y;				
-			}
+				float ScrollLineXL, ScrollLineYL;
+				Canvas->StrLen(Font, Scrollback[idx], ScrollLineXL, ScrollLineYL);
+				if (ScrollLineYL > yl)
+				{
+					y -= (ScrollLineYL - yl);
+					PenX = LeftPos;
+					PenY = TopPos+y;				
+				}
 
-			ConsoleText.Text = FText::FromString( Scrollback[idx] );
-			Canvas->DrawItem( ConsoleText, PenX, PenY );
+				ConsoleText.Text = FText::FromString( Scrollback[idx] );
+				Canvas->DrawItem( ConsoleText, PenX, PenY );
+			}
+			idx--;
+			y-=yl;
 		}
-		idx--;
-		y-=yl;
 	}
+
+	PostRender_InputLine(Canvas, FIntPoint(LeftPos, TopPos + Height + 6));
 }
 
 void UConsole::BeginState_Open(FName PreviousStateName)
@@ -1337,49 +995,191 @@ bool UConsole::InputChar( int32 ControllerId, const FString& Unicode )
 	return bCaptureKeyInput;
 }
 
-
-bool UConsole::InputKey_Global( int32 ControllerId, FKey Key, EInputEvent Event, float AmountDepressed, bool bGamepad)
+bool UConsole::InputKey( int32 ControllerId, FKey Key, EInputEvent Event, float AmountDepressed, bool bGamepad)
 {
-	if ( Event == IE_Pressed )
-	{
-		bCaptureKeyInput = false;
+	bool bWasConsumed = InputKey_InputLine(ControllerId,Key,Event,AmountDepressed,bGamepad);
 
-		if ( Key == GetDefault<UInputSettings>()->ConsoleKey )
+	if(!bWasConsumed)
+	{
+		if (ConsoleState == NAME_Typing)
 		{
-			FakeGotoState(NAME_Typing);
-			bCaptureKeyInput = true;
+			// if the console is open we don't want any other one to consume the input
+			return true;
+		}
+		if (ConsoleState == NAME_Open)
+		{
+			bWasConsumed = InputKey_Open(ControllerId,Key,Event,AmountDepressed,bGamepad);
+			// if the console is open we don't want any other one to consume the input
 			return true;
 		}
 	}
 
-	return bCaptureKeyInput;
+	return bWasConsumed;
 }
 
 
-bool UConsole::InputKey( int32 ControllerId, FKey Key, EInputEvent Event, float AmountDepressed, bool bGamepad)
+void UConsole::PostRender_Console(UCanvas* Canvas)
 {
 	if (ConsoleState == NAME_Typing)
 	{
-		return InputKey_Typing(ControllerId,Key,Event,AmountDepressed,bGamepad);
+		PostRender_Console_Typing(Canvas);
 	}
-	if (ConsoleState == NAME_Open)
+	else if (ConsoleState == NAME_Open)
 	{
-		return InputKey_Open(ControllerId,Key,Event,AmountDepressed,bGamepad);
+		PostRender_Console_Open(Canvas);
 	}
-	return InputKey_Global(ControllerId,Key,Event,AmountDepressed,bGamepad);
 }
 
-
-void UConsole::PostRender_Console(UCanvas* Canvas)
+void UConsole::PostRender_InputLine(UCanvas* Canvas, FIntPoint UserInputLinePos)
 {
-	if (ConsoleState == NAME_Typing)
+	float xl, yl;
+
+	// use the smallest font
+	UFont* Font = GEngine->GetSmallFont();
+	// determine the position for the cursor
+	FString OutStr = FString::Printf(TEXT("%s%s"), *ConsoleDefs::LeadingInputText, *TypedStr);
+	Canvas->StrLen(Font, OutStr,xl,yl);
+
+	float ClipX = Canvas->ClipX;
+	float ClipY = Canvas->ClipY;
+
+	if (GEngine->IsConsoleBuild())
 	{
-		PostRender_Console_Typing(Canvas);
+		ClipX	-= 64;
+		ClipY	-= 32;
 	}
-	else if (ConsoleState == NAME_Open)
+
+	if (GEngine->IsStereoscopic3D())
 	{
-		PostRender_Console_Open(Canvas);
+		ClipX -= 150;
+		ClipY = ClipY * 0.60;
 	}
+
+	// Background
+	FCanvasTileItem ConsoleTile( FVector2D( UserInputLinePos.X,UserInputLinePos.Y-6-yl ), DefaultTexture_Black->Resource, FVector2D( ClipX, yl+6 ), FVector2D( 0.0f, 0.0f), FVector2D( 1.0f, 1.0f ), FLinearColor::White );
+	Canvas->DrawItem( ConsoleTile );
+
+	// Separator line
+	ConsoleTile.SetColor( ConsoleDefs::BorderColor );
+	ConsoleTile.Texture = DefaultTexture_White->Resource;
+	ConsoleTile.Size = FVector2D( ClipX, 2.0f );
+	Canvas->DrawItem( ConsoleTile );
+
+	// Currently typed string
+	FText Str = FText::FromString( OutStr );
+	FCanvasTextItem ConsoleText( FVector2D( UserInputLinePos.X,UserInputLinePos.Y-3-yl ), Str , GEngine->GetSmallFont(), ConsoleDefs::InputTextColor );		
+	Canvas->DrawItem( ConsoleText );
+
+	// draw the remaining text for matching auto-complete
+	if (AutoCompleteIndices.Num() > 0)
+	{
+		int32 Idx = AutoCompleteIndices[AutoCompleteIndex];
+		ConsoleText.SetColor( ConsoleDefs::AutocompletePartialSuggestionColor );
+		ConsoleText.Text = FText::FromString( AutoCompleteList[Idx].Command.Right(AutoCompleteList[Idx].Command.Len() - TypedStr.Len()) );
+		Canvas->DrawItem( ConsoleText, UserInputLinePos.X+xl, UserInputLinePos.Y-3.0f-yl );
+		Canvas->StrLen(Font, *ConsoleDefs::LeadingInputText, xl, yl );
+
+		int32 StartIdx = AutoCompleteIndex;
+
+		if (StartIdx < 0)
+		{
+			StartIdx = FMath::Max(0,AutoCompleteIndices.Num() + StartIdx);
+		}
+		Idx = StartIdx;
+		float y = UserInputLinePos.Y - 6.0f - (yl * 2.0f);
+
+		// Set the background color/texture of the auto-complete section
+		ConsoleTile.SetColor( ConsoleDefs::AutocompleteBackgroundColor );
+		ConsoleTile.Texture = DefaultTexture_White->Resource;
+
+		// wasteful memory allocations but when typing in a console command this is fine
+		TArray<FString> AutoCompleteElements;
+		// to avoid memory many allocations
+		AutoCompleteElements.Empty(MAX_AUTOCOMPLETION_LINES + 1);
+
+		for (int32 MatchIdx = 0; MatchIdx < MAX_AUTOCOMPLETION_LINES && MatchIdx < AutoCompleteIndices.Num(); MatchIdx++)
+		{
+			const FAutoCompleteCommand &Cmd = AutoCompleteList[AutoCompleteIndices[Idx + MatchIdx]];
+			OutStr = Cmd.Desc;
+
+			if(OutStr.IsEmpty())
+			{
+				// no Description means we display the Command directly, without that the line would be empty (happens for ConsoleVariables and some ManualAutoCompleteList)
+				OutStr = Cmd.Command;
+			}
+
+			AutoCompleteElements.Add(OutStr);
+		}
+
+		// Display a message if there were more matches
+		if (AutoCompleteIndices.Num() >= MAX_AUTOCOMPLETION_LINES)
+		{
+			OutStr = FString::Printf(TEXT("[%i more matches]"), (AutoCompleteIndices.Num() - MAX_AUTOCOMPLETION_LINES + 1));
+			AutoCompleteElements.Add(OutStr);
+		}
+
+		// background rectangle behind auto completion
+		{
+			float MaxWidth = 0;
+			float MaxHeight = 0;
+
+			for(int32 i = 0, Num = AutoCompleteElements.Num(); i < Num; ++i)
+			{
+				const FString& AutoCompleteElement = AutoCompleteElements[i];
+
+				float info_xl, info_yl;
+
+				Canvas->StrLen(Font, AutoCompleteElement, info_xl, info_yl);
+
+				MaxWidth = FMath::Max(MaxWidth, info_xl);
+				MaxHeight += yl;
+			}
+
+			int32 Border = 4;
+
+			// white border
+			ConsoleTile.Size = FVector2D( MaxWidth + 2 * Border, MaxHeight + 2 * Border );
+			ConsoleTile.SetColor( ConsoleDefs::BorderColor );
+			Canvas->DrawItem( ConsoleTile, UserInputLinePos.X + xl - Border, y + yl - MaxHeight - Border );
+
+			--Border;
+
+			// dark inner part
+			ConsoleTile.Size = FVector2D( MaxWidth + 2 * Border, MaxHeight + 2 * Border );
+			ConsoleTile.SetColor( ConsoleDefs::AutocompleteBackgroundColor );
+			Canvas->DrawItem( ConsoleTile, UserInputLinePos.X + xl - Border, y + yl - MaxHeight - Border );
+		}
+
+		// auto completion elements
+		for(int32 i = 0, Num = AutoCompleteElements.Num(); i < Num; ++i)
+		{
+			const FString& AutoCompleteElement = AutoCompleteElements[i];
+
+			if (i == AutoCompleteCursor									// cursor line is highlighted
+				|| (Num >= MAX_AUTOCOMPLETION_LINES && i == Num - 1))	// e.g. [%i more matches]
+			{
+				ConsoleText.SetColor( ConsoleDefs::AutocompleteSuggestionColor );
+			}
+			else
+			{
+				ConsoleText.SetColor( ConsoleDefs::AutocompletePartialSuggestionColor );
+			}
+			ConsoleText.Text = FText::FromString(AutoCompleteElement);
+			Canvas->DrawItem( ConsoleText, UserInputLinePos.X + xl, y );
+			y -= yl;
+		}
+	}
+
+	// determine the cursor position
+	OutStr = FString::Printf(TEXT("%s%s"), *ConsoleDefs::LeadingInputText, *TypedStr.Left(TypedStrPos));
+	Canvas->StrLen(Font, OutStr,xl,yl);
+
+	// draw the cursor
+	ConsoleText.SetColor( ConsoleDefs::CursorColor );
+	OutStr = FString::Printf( TEXT("_") );
+	ConsoleText.Text = FText::FromString( OutStr );
+	Canvas->DrawItem( ConsoleText, UserInputLinePos.X + xl, UserInputLinePos.Y-1.0f-yl );
+
 }
 
 bool UConsole::ConsoleActive() const
