Expose tire slip into skid audio trigger

diff --git a/Engine/Source/Runtime/Engine/Classes/Vehicles/WheeledVehicleMovementComponent.h b/Engine/Source/Runtime/Engine/Classes/Vehicles/WheeledVehicleMovementComponent.h
--- a/Engine/Source/Runtime/Engine/Classes/Vehicles/WheeledVehicleMovementComponent.h
+++ b/Engine/Source/Runtime/Engine/Classes/Vehicles/WheeledVehicleMovementComponent.h
@@ -259,6 +259,8 @@ class ENGINE_API UWheeledVehicleMovementComponent : public UPawnMovementComponen
 	// Used to recreate the physics if the blueprint changes.
 	uint32 VehicleSetupTag;
 
+	bool CheckSlipThreshold(float AbsLongSlipThreshold, float AbsLatSlipThreshold) const;
+
 #if WITH_PHYSX
 
 	// The instanced PhysX vehicle
diff --git a/Engine/Source/Runtime/Engine/Private/Vehicles/WheeledVehicleMovementComponent.cpp b/Engine/Source/Runtime/Engine/Private/Vehicles/WheeledVehicleMovementComponent.cpp
--- a/Engine/Source/Runtime/Engine/Private/Vehicles/WheeledVehicleMovementComponent.cpp
+++ b/Engine/Source/Runtime/Engine/Private/Vehicles/WheeledVehicleMovementComponent.cpp
@@ -1030,6 +1030,40 @@ void DrawTelemetryGraph( uint32 Channel, const PxVehicleGraph& PGraph, UCanvas*
 	OutX = FMath::Max(XL,GraphWidth);
 }
 
+bool UWheeledVehicleMovementComponent::CheckSlipThreshold(float AbsLongSlipThreshold, float AbsLatSlipThreshold) const
+{
+	if (PVehicle == NULL)
+	{
+		return false;
+	}
+
+	FPhysXVehicleManager* MyVehicleManager = World->GetPhysicsScene()->GetVehicleManager();
+	PxWheelQueryResult * WheelsStates = MyVehicleManager->GetWheelsStates(this);
+	check(WheelsStates);
+
+	PxReal MaxLongSlip = 0.f;
+	PxReal MaxLatSlip = 0.f;
+
+	// draw wheel data
+	for (uint32 w = 0; w < PVehicle->mWheelsSimData.getNbWheels(); ++w)
+	{
+		const PxReal AbsLongSlip = FMath::Abs(WheelsStates[w].longitudinalSlip);
+		const PxReal AbsLatSlip = FMath::Abs(WheelsStates[w].lateralSlip);
+
+		if (AbsLongSlip > AbsLongSlipThreshold)
+		{
+			return true;
+		}
+
+		if (AbsLatSlip > AbsLatSlipThreshold)
+		{
+			return true;
+		}
+	}
+
+	return false;
+}
+
 void UWheeledVehicleMovementComponent::DrawDebug( UCanvas* Canvas, float& YL, float& YPos )
 {
 	if ( PVehicle == NULL )
