#UE4:
* fixed player input when game console was opened

diff --git a/Engine/Source/Runtime/Engine/Private/UserInterface/Console.cpp b/Engine/Source/Runtime/Engine/Private/UserInterface/Console.cpp
--- a/Engine/Source/Runtime/Engine/Private/UserInterface/Console.cpp
+++ b/Engine/Source/Runtime/Engine/Private/UserInterface/Console.cpp
@@ -542,31 +542,8 @@ bool UConsole::InputKey_InputLine( int32 ControllerId, FKey Key, EInputEvent Eve
 		bCaptureKeyInput = false;
 	}
 
-	if (ProcessControlKey(Key, Event))
-	{
-		return true;
-	}
-	else if( bGamepad )
-	{
-		return false;
-	}
-	else if( Key == EKeys::Escape && Event == IE_Released )
-	{
-		if( !TypedStr.IsEmpty() )
-		{
-			SetInputText("");
-			SetCursorPos(0);
-			HistoryCur = HistoryTop;
-			return true;
-		}
-		else
-		{
-			FakeGotoState( NAME_None );
-		}
-
-		return true;
-	}
-	else if ( Key==GetDefault<UInputSettings>()->ConsoleKey && Event == IE_Pressed )
+	// cycle between console states
+	if ( Key==GetDefault<UInputSettings>()->ConsoleKey && Event == IE_Pressed )
 	{
 		if (ConsoleState == NAME_Typing)
 		{
@@ -585,193 +562,222 @@ bool UConsole::InputKey_InputLine( int32 ControllerId, FKey Key, EInputEvent Eve
 		}
 		return true;
 	}
-	else if (Key == EKeys::Tab && Event == IE_Pressed)
-	{
-		if (AutoCompleteIndices.Num() > 0 && !bAutoCompleteLocked)
-		{
-			const FAutoCompleteCommand& Cmd = AutoCompleteList[AutoCompleteIndices[AutoCompleteIndex + AutoCompleteCursor]];
 
-			TypedStr = Cmd.Command;
-			SetCursorPos(TypedStr.Len());
-			bAutoCompleteLocked = true;
-		}
-		return true;
-	}
-	else if( Key == EKeys::Enter && Event == IE_Released )
+	// if user input is open
+	if(ConsoleState != NAME_None)
 	{
-		if( !TypedStr.IsEmpty() )
+		if (ProcessControlKey(Key, Event))
 		{
-			// Make a local copy of the string.
-			FString Temp=TypedStr;
-
-			SetInputText(TEXT(""));
-			SetCursorPos(0);
-
-			ConsoleCommand(Temp);
-
-			//OutputText( Localize("Errors","Exec","Core") );
-
-			OutputText( TEXT("") );
-			FakeGotoState(NAME_None);
-
-			UpdateCompleteIndices();
+			return true;
 		}
-		else
+		else if( bGamepad )
 		{
-			FakeGotoState(NAME_None);
+			return false;
 		}
-
-		return true;
-	}
-	else if( Event != IE_Pressed && Event != IE_Repeat )
-	{
-		if( !bGamepad )
-		{
-			return	Key != EKeys::LeftMouseButton
-				&&	Key != EKeys::MiddleMouseButton
-				&&	Key != EKeys::RightMouseButton;
-		}
-		return false;
-	}
-	else if( Key == EKeys::Up )
-	{
-		if (!bNavigatingHistory && !bCtrl && AutoCompleteIndices.Num() > 1)
+		else if( Key == EKeys::Escape && Event == IE_Released )
 		{
-			if(AutoCompleteCursor + 1 < FMath::Min((int32)MAX_AUTOCOMPLETION_LINES, AutoCompleteIndices.Num()))
+			if( !TypedStr.IsEmpty() )
+			{
+				SetInputText("");
+				SetCursorPos(0);
+				HistoryCur = HistoryTop;
+				return true;
+			}
+			else
 			{
-				// move cursor within displayed region
-				++AutoCompleteCursor;
+				FakeGotoState( NAME_None );
 			}
-			else 
+
+			return true;
+		}
+		else if (Key == EKeys::Tab && Event == IE_Pressed)
+		{
+			if (AutoCompleteIndices.Num() > 0 && !bAutoCompleteLocked)
 			{
-				// can be negative
-				int32 ScrollRegionSize = AutoCompleteIndices.Num() - (int32)MAX_AUTOCOMPLETION_LINES;
+				const FAutoCompleteCommand& Cmd = AutoCompleteList[AutoCompleteIndices[AutoCompleteIndex + AutoCompleteCursor]];
 
-				// can we scroll?
-				if(AutoCompleteIndex + 1 < ScrollRegionSize)
-				{
-					++AutoCompleteIndex;
-				}
+				TypedStr = Cmd.Command;
+				SetCursorPos(TypedStr.Len());
+				bAutoCompleteLocked = true;
 			}
+			return true;
 		}
-		else
-		if ( HistoryBot >= 0 )
+		else if( Key == EKeys::Enter && Event == IE_Released )
 		{
-			if (HistoryCur == HistoryBot)
-				HistoryCur = HistoryTop;
+			if( !TypedStr.IsEmpty() )
+			{
+				// Make a local copy of the string.
+				FString Temp=TypedStr;
+
+				SetInputText(TEXT(""));
+				SetCursorPos(0);
+
+				ConsoleCommand(Temp);
+
+				//OutputText( Localize("Errors","Exec","Core") );
+
+				OutputText( TEXT("") );
+				FakeGotoState(NAME_None);
+
+				UpdateCompleteIndices();
+			}
 			else
 			{
-				HistoryCur--;
-				if (HistoryCur<0)
-					HistoryCur = MAX_HISTORY_ENTRIES-1;
+				FakeGotoState(NAME_None);
 			}
 
-			SetInputText(History[HistoryCur]);
-			SetCursorPos(History[HistoryCur].Len());
-			UpdateCompleteIndices();
-			bNavigatingHistory = true;
+			return true;
 		}
-		return true;
-	}
-	else if( Key == EKeys::Down )
-	{
-		if (!bNavigatingHistory && !bCtrl && AutoCompleteIndices.Num() > 1)
+		else if( Event != IE_Pressed && Event != IE_Repeat )
 		{
-			if(AutoCompleteCursor > 0)
+			if( !bGamepad )
 			{
-				// move cursor within displayed region
-				--AutoCompleteCursor;
+				return	Key != EKeys::LeftMouseButton
+					&&	Key != EKeys::MiddleMouseButton
+					&&	Key != EKeys::RightMouseButton;
 			}
-			else
+			return false;
+		}
+		else if( Key == EKeys::Up )
+		{
+			if (!bNavigatingHistory && !bCtrl && AutoCompleteIndices.Num() > 1)
 			{
-				// can we scroll?
-				if(AutoCompleteIndex > 0)
+				if(AutoCompleteCursor + 1 < FMath::Min((int32)MAX_AUTOCOMPLETION_LINES, AutoCompleteIndices.Num()))
 				{
-					--AutoCompleteIndex;
+					// move cursor within displayed region
+					++AutoCompleteCursor;
+				}
+				else 
+				{
+					// can be negative
+					int32 ScrollRegionSize = AutoCompleteIndices.Num() - (int32)MAX_AUTOCOMPLETION_LINES;
+
+					// can we scroll?
+					if(AutoCompleteIndex + 1 < ScrollRegionSize)
+					{
+						++AutoCompleteIndex;
+					}
 				}
-				bAutoCompleteLocked = false;
 			}
+			else
+				if ( HistoryBot >= 0 )
+				{
+					if (HistoryCur == HistoryBot)
+						HistoryCur = HistoryTop;
+					else
+					{
+						HistoryCur--;
+						if (HistoryCur<0)
+							HistoryCur = MAX_HISTORY_ENTRIES-1;
+					}
+
+					SetInputText(History[HistoryCur]);
+					SetCursorPos(History[HistoryCur].Len());
+					UpdateCompleteIndices();
+					bNavigatingHistory = true;
+				}
+				return true;
 		}
-		else
-		if ( HistoryBot >= 0 )
+		else if( Key == EKeys::Down )
 		{
-			if (HistoryCur == HistoryTop)
-				HistoryCur = HistoryBot;
+			if (!bNavigatingHistory && !bCtrl && AutoCompleteIndices.Num() > 1)
+			{
+				if(AutoCompleteCursor > 0)
+				{
+					// move cursor within displayed region
+					--AutoCompleteCursor;
+				}
+				else
+				{
+					// can we scroll?
+					if(AutoCompleteIndex > 0)
+					{
+						--AutoCompleteIndex;
+					}
+					bAutoCompleteLocked = false;
+				}
+			}
 			else
-				HistoryCur = (HistoryCur+1) % MAX_HISTORY_ENTRIES;
+				if ( HistoryBot >= 0 )
+				{
+					if (HistoryCur == HistoryTop)
+						HistoryCur = HistoryBot;
+					else
+						HistoryCur = (HistoryCur+1) % MAX_HISTORY_ENTRIES;
+
+					SetInputText(History[HistoryCur]);
+					SetCursorPos(History[HistoryCur].Len());
+					UpdateCompleteIndices();
+					bNavigatingHistory = true;
+				}
 
-			SetInputText(History[HistoryCur]);
-			SetCursorPos(History[HistoryCur].Len());
-			UpdateCompleteIndices();
-			bNavigatingHistory = true;
 		}
-
-	}
-	else if( Key==EKeys::BackSpace )
-	{
-		if( TypedStrPos>0 )
+		else if( Key==EKeys::BackSpace )
 		{
-			SetInputText(FString::Printf(TEXT("%s%s"), *TypedStr.Left(TypedStrPos-1), *TypedStr.Right( TypedStr.Len() - TypedStrPos)));
-			SetCursorPos(TypedStrPos-1);
-			// unlock auto-complete (@todo - track the lock position so we don't bother unlocking under bogus cases)
-			bAutoCompleteLocked = false;
-		}
+			if( TypedStrPos>0 )
+			{
+				SetInputText(FString::Printf(TEXT("%s%s"), *TypedStr.Left(TypedStrPos-1), *TypedStr.Right( TypedStr.Len() - TypedStrPos)));
+				SetCursorPos(TypedStrPos-1);
+				// unlock auto-complete (@todo - track the lock position so we don't bother unlocking under bogus cases)
+				bAutoCompleteLocked = false;
+			}
 
-		return true;
-	}
-	else if ( Key== EKeys::Delete )
-	{
-		if ( TypedStrPos < TypedStr.Len() )
-		{
-			SetInputText(FString::Printf(TEXT("%s%s"), *TypedStr.Left(TypedStrPos), *TypedStr.Right( TypedStr.Len() - TypedStrPos - 1)));
+			return true;
 		}
-		return true;
-	}
-	else if ( Key==EKeys::Left )
-	{
-		if (bCtrl)
+		else if ( Key== EKeys::Delete )
 		{
-			// find the nearest '.' or ' '
-			int32 NewPos = FMath::Max(TypedStr.Find(TEXT("."), ESearchCase::CaseSensitive, ESearchDir::FromEnd, TypedStrPos),TypedStr.Find(TEXT(" "),ESearchCase::CaseSensitive, ESearchDir::FromEnd,TypedStrPos));
-			SetCursorPos(FMath::Max(0,NewPos));
+			if ( TypedStrPos < TypedStr.Len() )
+			{
+				SetInputText(FString::Printf(TEXT("%s%s"), *TypedStr.Left(TypedStrPos), *TypedStr.Right( TypedStr.Len() - TypedStrPos - 1)));
+			}
+			return true;
 		}
-		else
+		else if ( Key==EKeys::Left )
 		{
-			SetCursorPos(FMath::Max(0, TypedStrPos - 1));
+			if (bCtrl)
+			{
+				// find the nearest '.' or ' '
+				int32 NewPos = FMath::Max(TypedStr.Find(TEXT("."), ESearchCase::CaseSensitive, ESearchDir::FromEnd, TypedStrPos),TypedStr.Find(TEXT(" "),ESearchCase::CaseSensitive, ESearchDir::FromEnd,TypedStrPos));
+				SetCursorPos(FMath::Max(0,NewPos));
+			}
+			else
+			{
+				SetCursorPos(FMath::Max(0, TypedStrPos - 1));
+			}
+			return true;
 		}
-		return true;
-	}
-	else if ( Key==EKeys::Right )
-	{
-		if (bCtrl)
+		else if ( Key==EKeys::Right )
 		{
-			// find the nearest '.' or ' '
-			int32 SpacePos = TypedStr.Find(TEXT(" "));
-			int32 PeriodPos = TypedStr.Find(TEXT("."));
-			// pick the closest valid index
-			int32 NewPos = SpacePos < 0 ? PeriodPos : (PeriodPos < 0 ? SpacePos : FMath::Min(SpacePos,PeriodPos));
-			// jump to end if nothing in between
-			if (NewPos == INDEX_NONE)
+			if (bCtrl)
+			{
+				// find the nearest '.' or ' '
+				int32 SpacePos = TypedStr.Find(TEXT(" "));
+				int32 PeriodPos = TypedStr.Find(TEXT("."));
+				// pick the closest valid index
+				int32 NewPos = SpacePos < 0 ? PeriodPos : (PeriodPos < 0 ? SpacePos : FMath::Min(SpacePos,PeriodPos));
+				// jump to end if nothing in between
+				if (NewPos == INDEX_NONE)
+				{
+					NewPos = TypedStr.Len();
+				}
+				SetCursorPos(FMath::Min(TypedStr.Len(),FMath::Max(TypedStrPos,NewPos)));
+			}
+			else
 			{
-				NewPos = TypedStr.Len();
+				SetCursorPos(FMath::Min(TypedStr.Len(), TypedStrPos + 1));
 			}
-			SetCursorPos(FMath::Min(TypedStr.Len(),FMath::Max(TypedStrPos,NewPos)));
+			return true;
 		}
-		else
+		else if ( Key==EKeys::Home )
 		{
-			SetCursorPos(FMath::Min(TypedStr.Len(), TypedStrPos + 1));
+			SetCursorPos(0);
+			return true;
+		}
+		else if ( Key==EKeys::End )
+		{
+			SetCursorPos(TypedStr.Len());
+			return true;
 		}
-		return true;
-	}
-	else if ( Key==EKeys::Home )
-	{
-		SetCursorPos(0);
-		return true;
-	}
-	else if ( Key==EKeys::End )
-	{
-		SetCursorPos(TypedStr.Len());
-		return true;
 	}
 
 	return false;
