expose landing sound effect for vehicles

diff --git a/Engine/Source/Runtime/Engine/Classes/Vehicles/WheeledVehicleMovementComponent.h b/Engine/Source/Runtime/Engine/Classes/Vehicles/WheeledVehicleMovementComponent.h
--- a/Engine/Source/Runtime/Engine/Classes/Vehicles/WheeledVehicleMovementComponent.h
+++ b/Engine/Source/Runtime/Engine/Classes/Vehicles/WheeledVehicleMovementComponent.h
@@ -260,6 +260,7 @@ class ENGINE_API UWheeledVehicleMovementComponent : public UPawnMovementComponen
 	uint32 VehicleSetupTag;
 
 	bool CheckSlipThreshold(float AbsLongSlipThreshold, float AbsLatSlipThreshold) const;
+	float GetMaxSpringForce() const;
 
 #if WITH_PHYSX
 
diff --git a/Engine/Source/Runtime/Engine/Private/Vehicles/WheeledVehicleMovementComponent.cpp b/Engine/Source/Runtime/Engine/Private/Vehicles/WheeledVehicleMovementComponent.cpp
--- a/Engine/Source/Runtime/Engine/Private/Vehicles/WheeledVehicleMovementComponent.cpp
+++ b/Engine/Source/Runtime/Engine/Private/Vehicles/WheeledVehicleMovementComponent.cpp
@@ -1064,6 +1064,29 @@ bool UWheeledVehicleMovementComponent::CheckSlipThreshold(float AbsLongSlipThres
 	return false;
 }
 
+float UWheeledVehicleMovementComponent::GetMaxSpringForce() const
+{
+	if (PVehicle == NULL)
+	{
+		return false;
+	}
+
+	FPhysXVehicleManager* MyVehicleManager = World->GetPhysicsScene()->GetVehicleManager();
+	PxWheelQueryResult * WheelsStates = MyVehicleManager->GetWheelsStates(this);
+	check(WheelsStates);
+
+	PxReal MaxSpringCompression = 0.f;
+
+	// draw wheel data
+	for (uint32 w = 0; w < PVehicle->mWheelsSimData.getNbWheels(); ++w)
+	{
+		MaxSpringCompression = WheelsStates[w].suspSpringForce > MaxSpringCompression ? WheelsStates[w].suspSpringForce : MaxSpringCompression;
+	}
+
+	return MaxSpringCompression;
+
+}
+
 void UWheeledVehicleMovementComponent::DrawDebug( UCanvas* Canvas, float& YL, float& YPos )
 {
 	if ( PVehicle == NULL )
