#UE4 - typo

diff --git a/Engine/Source/Runtime/Engine/Private/World.cpp b/Engine/Source/Runtime/Engine/Private/World.cpp
--- a/Engine/Source/Runtime/Engine/Private/World.cpp
+++ b/Engine/Source/Runtime/Engine/Private/World.cpp
@@ -4704,7 +4704,7 @@ void UWorld::SetSelectedLevels( const TArray<class ULevel*>& InLevels )
 #endif // WITH_EDITOR
 
 /**
- * Jumps the server to new level.  If bAbsolute is true and we are using seemless traveling, we
+ * Jumps the server to new level.  If bAbsolute is true and we are using seamless traveling, we
  * will do an absolute travel (URL will be flushed).
  *
  * @param URL the URL that we are traveling to
