Owned Components (components that have an Actor in their Outer hierarchy) can now have Outers other than the Actor itself

diff --git a/Engine/Source/Runtime/Engine/Classes/Components/ActorComponent.h b/Engine/Source/Runtime/Engine/Classes/Components/ActorComponent.h
--- a/Engine/Source/Runtime/Engine/Classes/Components/ActorComponent.h
+++ b/Engine/Source/Runtime/Engine/Classes/Components/ActorComponent.h
@@ -417,6 +417,7 @@ public:
 	virtual int32 GetFunctionCallspace( UFunction* Function, void* Parameters, FFrame* Stack ) OVERRIDE;
 	virtual bool CallRemoteFunction( UFunction* Function, void* Parameters, FOutParmRec* OutParms, FFrame* Stack ) OVERRIDE;
 	virtual void PostInitProperties() OVERRIDE;
+	virtual void PostRename(UObject* OldOuter, const FName OldName) OVERRIDE;
 #if WITH_EDITOR
 	virtual void PreEditChange(UProperty* PropertyThatWillChange) OVERRIDE;
 	virtual void PostEditChangeProperty(FPropertyChangedEvent& PropertyChangedEvent) OVERRIDE;
diff --git a/Engine/Source/Runtime/Engine/Classes/GameFramework/Actor.h b/Engine/Source/Runtime/Engine/Classes/GameFramework/Actor.h
--- a/Engine/Source/Runtime/Engine/Classes/GameFramework/Actor.h
+++ b/Engine/Source/Runtime/Engine/Classes/GameFramework/Actor.h
@@ -1891,7 +1891,7 @@ public:
 	bool IsNetStartupActor() const;
 
 	/** Searches components array and returns first encountered array by type. */
-	virtual UActorComponent* FindComponentByClass(const UClass* Class) const;
+	virtual UActorComponent* FindComponentByClass(const TSubclassOf<UActorComponent> ComponentClass) const;
 	
 	/** Templatized version for syntactic nicety. */
 	template<class T>
@@ -1908,14 +1908,12 @@ public:
 		static_assert(CanConvertPointerFromTo<T, UActorComponent>::Result, "'T' template parameter to GetComponents must be derived from ActorComponent");
 
 		SCOPE_CYCLE_COUNTER(STAT_GetComponentsTime);
-		TArray<UObject*> ChildObjects;
-		GetObjectsWithOuter(this, ChildObjects, false, RF_PendingKill);
 
-		OutComponents.Reset(ChildObjects.Num());
+		OutComponents.Reset(OwnedComponents.Num());
 
-		for (UObject* Child : ChildObjects)
+		for (UActorComponent* OwnedComponent : OwnedComponents)
 		{
-			T* const Component = Cast<T>(Child);
+			T* Component = Cast<T>(OwnedComponent);
 			if (Component)
 			{
 				OutComponents.Add(Component);
@@ -1923,6 +1921,33 @@ public:
 		}
 	}
 
+	// UActorComponent specialization to avoid unnecessary casts
+	template<>
+	void GetComponents(TArray<UActorComponent*>& OutComponents) const
+	{
+		SCOPE_CYCLE_COUNTER(STAT_GetComponentsTime);
+
+		OutComponents.Reset(OwnedComponents.Num());
+
+		for (UActorComponent* Component : OwnedComponents)
+		{
+			if (Component)
+			{
+				OutComponents.Add(Component);
+			}
+		}
+	}
+
+	void AddOwnedComponent(UActorComponent* Component);
+	void RemoveOwnedComponent(UActorComponent* Component);
+#if DO_CHECK
+	bool OwnsComponent(UActorComponent* Component) const;
+#endif
+
+private:
+	UPROPERTY(transient, duplicatetransient)
+	TArray<UActorComponent*> OwnedComponents;
+
 public:
 	//=============================================================================
 	// Navigation related functions
diff --git a/Engine/Source/Runtime/Engine/Private/Actor.cpp b/Engine/Source/Runtime/Engine/Private/Actor.cpp
--- a/Engine/Source/Runtime/Engine/Private/Actor.cpp
+++ b/Engine/Source/Runtime/Engine/Private/Actor.cpp
@@ -253,6 +253,25 @@ void AActor::PostInitProperties()
 	Super::PostInitProperties();
 	RegisterAllActorTickFunctions(true,false); // component will be handled when they are registered
 	RemoteRole = (bReplicates ? ROLE_SimulatedProxy : ROLE_None);
+
+	// Make sure the OwnedComponents list correct.  
+	// Under some circumstances sub-object instancing can result in bogus/duplicate entries.
+	// This is not necessary for CDOs
+	if (!HasAnyFlags(RF_ClassDefaultObject))
+	{
+		TArray<UObject*> Children;
+		OwnedComponents.Empty();
+		GetObjectsWithOuter(this, Children, true, RF_PendingKill);
+
+		for (UObject* Child : Children)
+		{
+			UActorComponent* Component = Cast<UActorComponent>(Child);
+			if (Component)
+			{
+				OwnedComponents.Add(Component);
+			}
+		}
+	}
 }
 
 UWorld* AActor::GetWorld() const
@@ -1924,20 +1943,41 @@ enum ECollisionResponse AActor::GetComponentsCollisionResponseToChannel(enum ECo
 	return OutResponse;
 };
 
-UActorComponent* AActor::FindComponentByClass(const UClass* Class) const
+void AActor::AddOwnedComponent(UActorComponent* Component)
+{
+	OwnedComponents.AddUnique(Component);
+}
+
+void AActor::RemoveOwnedComponent(UActorComponent* Component)
 {
-	TArray<UObject*> ChildObjects;
-	GetObjectsWithOuter(this, ChildObjects, false, RF_PendingKill);
+	if (OwnedComponents.RemoveSwap(Component) == 0)
+	{
+		// If we didn't remove something we expected to then it probably got NULL through the
+		// property system so take the time to pull them out now
+		OwnedComponents.RemoveSwap(NULL);
+	}
+}
 
-	for (UObject* Child : ChildObjects)
+#if DO_CHECK
+bool AActor::OwnsComponent(UActorComponent* Component) const
+{
+	return OwnedComponents.Contains(Component);
+}
+#endif
+
+UActorComponent* AActor::FindComponentByClass(const TSubclassOf<UActorComponent> ComponentClass) const
+{
+	UActorComponent* FoundComponent = NULL;
+	for (UActorComponent* Component : OwnedComponents)
 	{
-		if ( Child->IsA(Class) )
+		if (Component && Component->IsA(ComponentClass))
 		{
-			return (UActorComponent*) Child;
+			FoundComponent = Component;
+			break;
 		}
 	}
 
-	return NULL;
+	return FoundComponent;
 }
 
 void AActor::DisableComponentsSimulatePhysics()
@@ -2863,7 +2903,7 @@ void AActor::RegisterAllComponents()
 	GetComponents(SceneComponents);
 
 	for(int32 CompIdx=0; CompIdx < SceneComponents.Num(); CompIdx++)
-		{
+	{
 		USceneComponent* SceneComp = SceneComponents[CompIdx];
 #if !(UE_BUILD_SHIPPING || UE_BUILD_TEST)
 			AActor* CompOwner = SceneComp->GetOwner();
diff --git a/Engine/Source/Runtime/Engine/Private/ActorComponent.cpp b/Engine/Source/Runtime/Engine/Private/ActorComponent.cpp
--- a/Engine/Source/Runtime/Engine/Private/ActorComponent.cpp
+++ b/Engine/Source/Runtime/Engine/Private/ActorComponent.cpp
@@ -116,7 +116,31 @@ void UActorComponent::PostInitProperties()
 	Super::PostInitProperties();
 
 	AActor* Owner = GetOwner();
-	ensureMsgf(Owner == NULL || Owner == GetOuter(), TEXT("Component %s is owned by %s by has an Outer of %s."), *GetName(), *Owner->GetName(), *GetOuter()->GetName());
+	if (Owner)
+	{
+		Owner->AddOwnedComponent(this);
+	}
+}
+
+void UActorComponent::PostRename(UObject* OldOuter, const FName OldName)
+{
+	if (OldOuter != GetOuter())
+	{
+		AActor* Owner = GetOwner();
+		AActor* OldOwner = (OldOuter->IsA<AActor>() ? CastChecked<AActor>(OldOuter) : OldOuter->GetTypedOuter<AActor>());
+
+		if (Owner != OldOwner)
+		{
+			if (OldOwner)
+			{
+				OldOwner->RemoveOwnedComponent(this);
+			}
+			if (Owner)
+			{
+				Owner->AddOwnedComponent(this);
+			}
+		}
+	}
 }
 
 #if WITH_EDITOR
@@ -546,7 +570,7 @@ void UActorComponent::RegisterComponentWithWorld(UWorld* InWorld)
 	// Follow outer chain to see if we can find an Actor
 	AActor* Owner = GetOwner();
 
-	ensureMsgf(Owner == NULL || Owner == GetOuter(), TEXT("Component %s is owned by %s by has an Outer of %s."), *GetName(), *Owner->GetName(), *GetOuter()->GetName());
+	checkSlow(Owner == NULL || Owner->OwnsComponent(this));
 
 	if ((Owner != NULL) && Owner->GetClass()->HasAnyClassFlags(CLASS_NewerVersionExists))
 	{
@@ -624,6 +648,7 @@ void UActorComponent::DestroyComponent()
 	if(Owner != NULL)
 	{
 		Owner->SerializedComponents.Remove(this);
+		Owner->RemoveOwnedComponent(this);
 		if (Owner->GetRootComponent() == this)
 		{
 			Owner->SetRootComponent(NULL);
