#UE4 - added generic TravelToSession function to AGameSession

diff --git a/Engine/Source/Runtime/Engine/Classes/GameFramework/GameSession.h b/Engine/Source/Runtime/Engine/Classes/GameFramework/GameSession.h
--- a/Engine/Source/Runtime/Engine/Classes/GameFramework/GameSession.h
+++ b/Engine/Source/Runtime/Engine/Classes/GameFramework/GameSession.h
@@ -29,10 +29,6 @@ class ENGINE_API AGameSession : public AInfo
 	UPROPERTY(globalconfig)
 	bool bRequiresPushToTalk;
 
-	/** The options to apply for dedicated server when it starts to register */
-	UPROPERTY()
-	FString ServerOptions;
-
 	/** SessionName local copy from PlayerState class.  should really be define in this class, but need to address replication issues */
 	UPROPERTY()
 	FName SessionName;
@@ -130,6 +126,16 @@ class ENGINE_API AGameSession : public AInfo
 	 */
 	virtual void UpdateSessionJoinability(FName SessionName, bool bPublicSearchable, bool bAllowInvites, bool bJoinViaPresence, bool bJoinViaPresenceFriendsOnly);
 
+	/**
+	 * Travel to a session URL (as client) for a given session
+	 *
+	 * @param ControllerId controller initiating the session travel
+	 * @param SessionName name of session to travel to
+	 *
+	 * @return true if successful, false otherwise
+	 */
+	virtual bool TravelToSession(int32 ControllerId, FName SessionName);
+
     /**
      * Does the session require push to talk
      * @return true if a push to talk keybinding is required or if voice is always enabled
diff --git a/Engine/Source/Runtime/Engine/Private/GameSession.cpp b/Engine/Source/Runtime/Engine/Private/GameSession.cpp
--- a/Engine/Source/Runtime/Engine/Private/GameSession.cpp
+++ b/Engine/Source/Runtime/Engine/Private/GameSession.cpp
@@ -264,6 +264,32 @@ void AGameSession::ReturnToMainMenuHost()
 	}
 }
 
+bool AGameSession::TravelToSession(int32 ControllerId, FName SessionName)
+{
+	UWorld* World = GetWorld();
+	IOnlineSubsystem* OnlineSub = Online::GetSubsystem(World);
+	if (OnlineSub)
+	{
+		FString URL;
+		IOnlineSessionPtr SessionInt = OnlineSub->GetSessionInterface();
+		if (SessionInt.IsValid() && SessionInt->GetResolvedConnectString(SessionName, URL))
+		{
+			APlayerController* PC = UGameplayStatics::GetPlayerController(World, ControllerId);
+			if (PC)
+			{
+				PC->ClientTravel(URL, TRAVEL_Absolute);
+				return true;
+			}
+		}
+		else
+		{
+			UE_LOG(LogGameSession, Warning, TEXT("Failed to resolve session connect string for %s"), *SessionName.ToString());
+		}
+	}
+
+	return false;
+}
+
 void AGameSession::PostSeamlessTravel()
 {
 }
