#UE4 - Merge from John Barrett's user branch. An issue was fixed that could cause SeamlessTravelCount to be incremented at the wrong time for the client, denying the client the ability to travel.

diff --git a/Engine/Source/Runtime/Engine/Classes/GameFramework/PlayerController.h b/Engine/Source/Runtime/Engine/Classes/GameFramework/PlayerController.h
--- a/Engine/Source/Runtime/Engine/Classes/GameFramework/PlayerController.h
+++ b/Engine/Source/Runtime/Engine/Classes/GameFramework/PlayerController.h
@@ -681,17 +681,32 @@ public:
 	void ClientStopForceFeedback(class UForceFeedbackEffect* ForceFeedbackEffect, FName Tag);
 
 	/**
-	 * Travel to a different map or IP address.  Calls the PreClientTravel event before doing anything.
-	 * @param	URL				a string containing the mapname (or IP address) to travel to, along with option key/value pairs
-	 * @param	TravelType 		specifies whether the client should append URL options used in previous travels; if true is specified
+	 * Travel to a different map or IP address. Calls the PreClientTravel event before doing anything.
+	 * NOTE: This is implemented as a locally executed wrapper for ClientTravelInternal, to avoid API compatability breakage
+	 *
+	 * @param URL				A string containing the mapname (or IP address) to travel to, along with option key/value pairs
+	 * @param TravelType 		specifies whether the client should append URL options used in previous travels; if true is specified
 	 *							for the bSeamlesss parameter, this value must be TRAVEL_Relative.
-	 * @param	bSeamless		indicates whether to use seamless travel (requires TravelType of TRAVEL_Relative)
-	 * @param	MapPackageGuid	the GUID of the map package to travel to - this is used to find the file when it has been autodownloaded,
+	 * @param bSeamless			Indicates whether to use seamless travel (requires TravelType of TRAVEL_Relative)
+	 * @param MapPackageGuid	The GUID of the map package to travel to - this is used to find the file when it has been autodownloaded,
 	 * 							so it is only needed for clients
 	 */
-	UFUNCTION(reliable, client)
+	UFUNCTION()
 	void ClientTravel(const FString& URL, enum ETravelType TravelType, bool bSeamless = false, FGuid MapPackageGuid = FGuid());
 
+	/**
+	 * Internal clientside implementation of ClientTravel - use ClientTravel to call this
+	 *
+	 * @param URL				A string containing the mapname (or IP address) to travel to, along with option key/value pairs
+	 * @param TravelType 		specifies whether the client should append URL options used in previous travels; if true is specified
+	 *							for the bSeamlesss parameter, this value must be TRAVEL_Relative.
+	 * @param bSeamless			Indicates whether to use seamless travel (requires TravelType of TRAVEL_Relative)
+	 * @param MapPackageGuid	The GUID of the map package to travel to - this is used to find the file when it has been autodownloaded,
+	 * 							so it is only needed for clients
+	 */
+	UFUNCTION(reliable, client)
+	void ClientTravelInternal(const FString& URL, enum ETravelType TravelType, bool bSeamless = false, FGuid MapPackageGuid = FGuid());
+
 	/**
 	 * Replicated Update streaming status
 	 * @param PackageName - Name of the level package name used for loading.
@@ -1288,5 +1303,5 @@ public:
 
 	/** The value of SeamlessTravelCount, upon the last call to GameMode::HandleSeamlessTravelPlayer; used to detect seamless travel */
 	UPROPERTY()
-	uint16		LastSeamlessTravelCount;
+	uint16		LastCompletedSeamlessTravelCount;
 };
diff --git a/Engine/Source/Runtime/Engine/Private/GameMode.cpp b/Engine/Source/Runtime/Engine/Private/GameMode.cpp
--- a/Engine/Source/Runtime/Engine/Private/GameMode.cpp
+++ b/Engine/Source/Runtime/Engine/Private/GameMode.cpp
@@ -801,7 +801,7 @@ void AGameMode::HandleSeamlessTravelPlayer(AController*& C)
 	if (PC != NULL)
 	{
 		// Track the last completed seamless travel for the player
-		PC->LastSeamlessTravelCount = PC->SeamlessTravelCount;
+		PC->LastCompletedSeamlessTravelCount = PC->SeamlessTravelCount;
 
 		PC->CleanUpAudioComponents();
 
diff --git a/Engine/Source/Runtime/Engine/Private/PlayerController.cpp b/Engine/Source/Runtime/Engine/Private/PlayerController.cpp
--- a/Engine/Source/Runtime/Engine/Private/PlayerController.cpp
+++ b/Engine/Source/Runtime/Engine/Private/PlayerController.cpp
@@ -396,39 +396,27 @@ void APlayerController::ServerNotifyLoadedWorld_Implementation(FName WorldPackag
 
 	UWorld* CurWorld = GetWorld();
 
-	// Only valid for calling, for PC's that have seamlessly traveled
-	if (CurWorld != NULL && CurWorld->IsServer() && SeamlessTravelCount > 0)
+	// Only valid for calling, for PC's in the process of seamless traveling
+	// NOTE: SeamlessTravelCount tracks client seamless travel, through the serverside gameplay code; this should not be replaced.
+	if (CurWorld != NULL && CurWorld->IsServer() && SeamlessTravelCount > 0 && LastCompletedSeamlessTravelCount < SeamlessTravelCount)
 	{
-		// Ensure this PC has not already completed seamless travel
-		if (LastSeamlessTravelCount == SeamlessTravelCount)
-		{
-			// Client sent too many ServerNotifyLoadedWorld RPC's, this is very suspicious, the client is most likely hacking
-			UE_LOG(LogPlayerController, Warning, TEXT("APlayerController::ServerNotifyLoadedWorld_Implementation: Client attempting to travel to same world twice!"));
-
-			// FIXME: Close the connection?  We need to resolve what to do when we know for certain a client is cheating.
+		// Update our info on what world the client is in
+		UNetConnection* const Connection = Cast<UNetConnection>(Player);
 
-			return;	
-		}
-		else
+		if (Connection != NULL)
 		{
-			// Update our info on what world the client is in
-			UNetConnection* const Connection = Cast<UNetConnection>(Player);
-
-			if (Connection != NULL)
-			{
-				Connection->ClientWorldPackageName = WorldPackageName;
-				Connection->PackageMap->SetLocked(false);
-			}
+			Connection->ClientWorldPackageName = WorldPackageName;
+			Connection->PackageMap->SetLocked(false);
+		}
 
-			// if both the server and this client have completed the transition, handle it
-			FSeamlessTravelHandler& SeamlessTravelHandler = GEngine->SeamlessTravelHandlerForWorld(CurWorld);
-			AGameMode* CurGameMode = CurWorld->GetAuthGameMode();
+		// if both the server and this client have completed the transition, handle it
+		FSeamlessTravelHandler& SeamlessTravelHandler = GEngine->SeamlessTravelHandlerForWorld(CurWorld);
+		AGameMode* CurGameMode = CurWorld->GetAuthGameMode();
 
-			if (!SeamlessTravelHandler.IsInTransition() && WorldPackageName == CurWorld->GetOutermost()->GetFName() && CurGameMode != NULL)
-			{
-				AController* TravelPlayer = this;
-				CurGameMode->HandleSeamlessTravelPlayer(TravelPlayer);
-			}
+		if (!SeamlessTravelHandler.IsInTransition() && WorldPackageName == CurWorld->GetOutermost()->GetFName() && CurGameMode != NULL)
+		{
+			AController* TravelPlayer = this;
+			CurGameMode->HandleSeamlessTravelPlayer(TravelPlayer);
 		}
 	}
 }
@@ -3456,7 +3444,19 @@ bool APlayerController::IsNetRelevantFor(APlayerController* RealViewer, AActor*
 	return ( this==RealViewer );
 }
 
-void APlayerController::ClientTravel_Implementation(const FString& URL, ETravelType TravelType, bool bSeamless, FGuid MapPackageGuid)
+void APlayerController::ClientTravel(const FString& URL, ETravelType TravelType, bool bSeamless, FGuid MapPackageGuid)
+{
+	// Keep track of seamless travel serverside
+	if (bSeamless && TravelType == TRAVEL_Relative)
+	{
+		SeamlessTravelCount++;
+	}
+
+	// Now pass on to the RPC
+	ClientTravelInternal(URL, TravelType, bSeamless, MapPackageGuid);
+}
+
+void APlayerController::ClientTravelInternal_Implementation(const FString& URL, ETravelType TravelType, bool bSeamless, FGuid MapPackageGuid)
 {
 	UWorld* World = GetWorld();
 
diff --git a/Engine/Source/Runtime/Engine/Private/World.cpp b/Engine/Source/Runtime/Engine/Private/World.cpp
--- a/Engine/Source/Runtime/Engine/Private/World.cpp
+++ b/Engine/Source/Runtime/Engine/Private/World.cpp
@@ -3458,9 +3458,6 @@ void UWorld::NotifyControlMessage(UNetConnection* Connection, uint8 MessageType,
 							FString NewLevelName = GetOutermost()->GetName();
 							UE_LOG(LogNet, Log, TEXT("Client joined but was sent to another level. Asking client to travel to: '%s'"), *NewLevelName);
 							LevelName = NewLevelName;
-
-							// Incrementing SeamlessTravelCount so that the following ServerNotifyLoadedWorld RPC will be able to call HandleSeamlessTravelPlayer
-							Connection->PlayerController->SeamlessTravelCount++;
 						}
 						if (LevelName != TEXT(""))
 						{
@@ -4176,11 +4173,6 @@ UWorld* FSeamlessTravelHandler::Tick()
 					if (Cast<AController>(TheActor))
 					{
 						LoadedWorld->AddController(Cast<AController>(TheActor) );
-
-						if (Cast<APlayerController>(TheActor))
-						{
-							Cast<APlayerController>(TheActor)->SeamlessTravelCount++;
-						}
 					}
 					else if (Cast<APawn>(TheActor))
 					{
