Missed files for CL# 2053360

diff --git a/Engine/Source/Runtime/Engine/Classes/GameFramework/Character.h b/Engine/Source/Runtime/Engine/Classes/GameFramework/Character.h
--- a/Engine/Source/Runtime/Engine/Classes/GameFramework/Character.h
+++ b/Engine/Source/Runtime/Engine/Classes/GameFramework/Character.h
@@ -239,7 +239,7 @@ public:
 	virtual void PostNetReceiveLocation() OVERRIDE;
 	virtual void GetSimpleCollisionCylinder(float& CollisionRadius, float& CollisionHalfHeight) const OVERRIDE;
 	virtual void ApplyWorldOffset(const FVector& InOffset, bool bWorldShift) OVERRIDE;
-	virtual UActorComponent* FindComponentByClass(const UClass* Class) const OVERRIDE;
+	virtual UActorComponent* FindComponentByClass(const TSubclassOf<UActorComponent> ComponentClass) const OVERRIDE;
 	// End AActor Interface
 
 	template<class T>
diff --git a/Engine/Source/Runtime/Engine/Private/Character.cpp b/Engine/Source/Runtime/Engine/Private/Character.cpp
--- a/Engine/Source/Runtime/Engine/Private/Character.cpp
+++ b/Engine/Source/Runtime/Engine/Private/Character.cpp
@@ -170,15 +170,15 @@ float ACharacter::GetDefaultHalfHeight() const
 }
 
 
-UActorComponent* ACharacter::FindComponentByClass(const UClass* Class) const
+UActorComponent* ACharacter::FindComponentByClass(const TSubclassOf<UActorComponent> ComponentClass) const
 {
 	// If the character has a Mesh, treat it as the first 'hit' when finding components
-	if (Class && Mesh && Mesh->IsA(Class))
-		{
-			return Mesh;
-		}
+	if (Mesh && Mesh->IsA(ComponentClass))
+	{
+		return Mesh;
+	}
 
-	return Super::FindComponentByClass(Class);
+	return Super::FindComponentByClass(ComponentClass);
 }
 
 
