#ue4 FLatentActionManager::RemoveActionsForObject function added.

diff --git a/Engine/Source/Runtime/Engine/Classes/Engine/LatentActionManager.h b/Engine/Source/Runtime/Engine/Classes/Engine/LatentActionManager.h
--- a/Engine/Source/Runtime/Engine/Classes/Engine/LatentActionManager.h
+++ b/Engine/Source/Runtime/Engine/Classes/Engine/LatentActionManager.h
@@ -105,6 +105,15 @@ public:
 		return FindExistingActionWithPredicate<ActionType, FFirstItemPredicate>(InActionObject, UUID, FFirstItemPredicate());
 	}
 
+	/** 
+	 * Removes all actions for given object. 
+	 * It the latent actions are being currently handled (so the function is called inside a ProcessLatentActions functions scope) 
+	 * there is no guarantee, that the action will be removed before its execution.
+	 *
+	 * @param		InOject		Specific object
+	 */	
+	void RemoveActionsForObject(TWeakObjectPtr<UObject> InObject);
+
 	/** 
 	 * Adds a new action to the action list under a given UUID 
 	 */
@@ -166,5 +175,8 @@ protected:
 protected:
 	/** list of objects we have processed the latent action list for this frame. */	
 	TSet<UObject*> ProcessedThisFrame;
+
+	/** For listed objects all actions will be removed */
+	TArray<TWeakObjectPtr<UObject>> ObjectsToRemove;
 };
 
diff --git a/Engine/Source/Runtime/Engine/Private/LatentActionManager.cpp b/Engine/Source/Runtime/Engine/Private/LatentActionManager.cpp
--- a/Engine/Source/Runtime/Engine/Private/LatentActionManager.cpp
+++ b/Engine/Source/Runtime/Engine/Private/LatentActionManager.cpp
@@ -16,8 +16,34 @@ FString FPendingLatentAction::GetDescription() const
 /////////////////////////////////////////////////////
 // FLatentActionManager
 
+void FLatentActionManager::RemoveActionsForObject(TWeakObjectPtr<UObject> InObject)
+{
+	ObjectsToRemove.Add(InObject);
+}
+
 void FLatentActionManager::ProcessLatentActions(UObject* InObject, float DeltaTime)
 {
+	if (ObjectsToRemove.Num())
+	{
+		for (const auto Key : ObjectsToRemove)
+		{
+			FActionList* const ObjectActionList = ObjectToActionListMap.Find(Key);
+			if (ObjectActionList)
+			{
+				for (TMultiMap<int32, FPendingLatentAction*>::TConstIterator It(*ObjectActionList); It; ++It)
+				{
+					FPendingLatentAction* Action = It.Value();
+					Action->NotifyActionAborted();
+					delete Action;
+				}
+				ObjectActionList->Empty();
+
+				ObjectToActionListMap.Remove(Key); //TODO: we should be able to obtain an Iterator from Map::Find function. "Iterator.RemoveCurrent();" should be called here.
+			}
+		}
+		ObjectsToRemove.Empty();
+	}
+
 	//@TODO: K2: Very inefficient code right now
 	if (InObject != NULL)
 	{
diff --git a/Engine/Source/Runtime/Engine/Public/LatentActions.h b/Engine/Source/Runtime/Engine/Public/LatentActions.h
--- a/Engine/Source/Runtime/Engine/Public/LatentActions.h
+++ b/Engine/Source/Runtime/Engine/Public/LatentActions.h
@@ -89,6 +89,8 @@ public:
 	// This is only called when the object goes away before the action is finished; perform normal
 	// cleanup when responding that the action is completed in UpdateOperation
 	virtual void NotifyObjectDestroyed() {}
+
+	virtual void NotifyActionAborted() {}
 #if WITH_EDITOR
 	// Returns a human readable description of the latent operation's current state
 	virtual FString GetDescription() const;
