Fix convex elements rendering incorrectly when they have local transform

diff --git a/Engine/Source/Runtime/Engine/Private/PhysicsEngine/PhysDrawing.cpp b/Engine/Source/Runtime/Engine/Private/PhysicsEngine/PhysDrawing.cpp
--- a/Engine/Source/Runtime/Engine/Private/PhysicsEngine/PhysDrawing.cpp
+++ b/Engine/Source/Runtime/Engine/Private/PhysicsEngine/PhysDrawing.cpp
@@ -521,7 +521,11 @@ void FKAggregateGeom::DrawAggGeom(FPrimitiveDrawInterface* PDI, const FTransform
 			for(int32 i=0; i<ConvexElems.Num(); i++)
 			{
 				FColor ConvexColor = bPerHullColor ? DebugUtilColor[i%NUM_DEBUG_UTIL_COLORS] : Color;
-				ConvexElems[i].DrawElemWire(PDI, ParentTM, Scale3D, ConvexColor);
+				FTransform ElemTM = ConvexElems[i].GetTransform();
+				ElemTM.ScaleTranslation(Scale3D);
+				ElemTM *= ParentTM;
+				ConvexElems[i].DrawElemWire(PDI, ElemTM, Scale3D, ConvexColor);
+
 			}
 		}
 	}
