Removed 'template<>' declaration so specialization resolves to an overloaded function instead of a class-scope explicit specialization which will not compile under GCC (due to being non-standard)

diff --git a/Engine/Source/Runtime/Engine/Classes/GameFramework/Actor.h b/Engine/Source/Runtime/Engine/Classes/GameFramework/Actor.h
--- a/Engine/Source/Runtime/Engine/Classes/GameFramework/Actor.h
+++ b/Engine/Source/Runtime/Engine/Classes/GameFramework/Actor.h
@@ -1922,7 +1922,6 @@ public:
 	}
 
 	// UActorComponent specialization to avoid unnecessary casts
-	template<>
 	void GetComponents(TArray<UActorComponent*>& OutComponents) const
 	{
 		SCOPE_CYCLE_COUNTER(STAT_GetComponentsTime);
