318963  Editor: Copy or duplicate of a Staticmesh with 'Custom' collision settings will revert to 'Block All'
312511  COLLISION: When duplicating a Static mesh with a Custom collision Preset, duplicate will lose custom settings
332062  LIVE: EDITOR: Custom Collision gets reset to BlockAll when making copies of static meshes

diff --git a/Engine/Source/Editor/DetailCustomizations/Private/BodyInstanceCustomization.cpp b/Engine/Source/Editor/DetailCustomizations/Private/BodyInstanceCustomization.cpp
--- a/Engine/Source/Editor/DetailCustomizations/Private/BodyInstanceCustomization.cpp
+++ b/Engine/Source/Editor/DetailCustomizations/Private/BodyInstanceCustomization.cpp
@@ -61,13 +61,13 @@ void FBodyInstanceCustomization::CustomizeStructChildren( TSharedRef<class IProp
 	// need to find profile name
 	FName ProfileName;
 	TSharedPtr< FString > DisplayName = CollisionProfileComboList[0];
-	//bDisplayAdvancedCollisionSettings = true;
+	bool bDisplayAdvancedCollisionSettings = true;
 
 	// if I have valid profile name
-	if (CollisionProfileNameHandle->GetValue(ProfileName) ==  FPropertyAccess::Result::Success && ProfileName!=NAME_None)
+	if (CollisionProfileNameHandle->GetValue(ProfileName) ==  FPropertyAccess::Result::Success && FBodyInstance::IsValidCollisionProfileName(ProfileName) )
 	{
 		DisplayName = GetProfileString(ProfileName);
-		//bDisplayAdvancedCollisionSettings = false;
+		bDisplayAdvancedCollisionSettings = false;
 	}
 
 	const FString PresetsDocLink = TEXT("Shared/Collision");
@@ -123,6 +123,7 @@ void FBodyInstanceCustomization::CustomizeStructChildren( TSharedRef<class IProp
 		]
 	];
 
+	CollisionGroup.ToggleExpansion(bDisplayAdvancedCollisionSettings);
 	// now create custom set up
 	CreateCustomCollisionSetup( StructPropertyHandle, CollisionGroup );
 }
@@ -690,7 +691,7 @@ void FBodyInstanceCustomization::OnCollisionProfileChanged( TSharedPtr<FString>
 			CollisionGroup->ToggleExpansion( true );
 		}
 		// if none of them found, clear it
-		FName Name=NAME_None;
+		FName Name=UCollisionProfile::CustomCollisionProfileName;
 		ensure ( CollisionProfileNameHandle->SetValue(Name) ==  FPropertyAccess::Result::Success );
 	}
 }
@@ -700,7 +701,7 @@ void FBodyInstanceCustomization::UpdateCollisionProfile()
 	FName ProfileName;
 
 	// if I have valid profile name
-	if (CollisionProfileNameHandle->GetValue(ProfileName) ==  FPropertyAccess::Result::Success && ProfileName!=NAME_None)
+	if (CollisionProfileNameHandle->GetValue(ProfileName) ==  FPropertyAccess::Result::Success && FBodyInstance::IsValidCollisionProfileName(ProfileName) )
 	{
 		int32 NumProfiles = CollisionProfile->GetNumOfProfiles();
 		for (int32 ProfileId = 0; ProfileId < NumProfiles; ++ProfileId)
@@ -788,7 +789,7 @@ EVisibility FBodyInstanceCustomization::ShouldShowResetToDefaultResponse(int32 V
 bool FBodyInstanceCustomization::ShouldEnableCustomCollisionSetup() const
 {
 	FName ProfileName;
-	if (CollisionProfileNameHandle->GetValue(ProfileName) == FPropertyAccess::Result::Success && ProfileName == NAME_None)
+	if (CollisionProfileNameHandle->GetValue(ProfileName) == FPropertyAccess::Result::Success && FBodyInstance::IsValidCollisionProfileName(ProfileName) == false)
 	{
 		return true;
 	}
diff --git a/Engine/Source/Editor/LandscapeEditor/Private/LandscapeEdMode.cpp b/Engine/Source/Editor/LandscapeEditor/Private/LandscapeEdMode.cpp
--- a/Engine/Source/Editor/LandscapeEditor/Private/LandscapeEdMode.cpp
+++ b/Engine/Source/Editor/LandscapeEditor/Private/LandscapeEdMode.cpp
@@ -2560,7 +2560,7 @@ ALandscape* FEdModeLandscape::ChangeComponentSetting(int32 NumComponentsX, int32
 			Landscape->LightmassSettings           = OldLandscapeProxy->LightmassSettings;
 			Landscape->CollisionThickness          = OldLandscapeProxy->CollisionThickness;
 			Landscape->BodyInstance.SetCollisionProfileName(OldLandscapeProxy->BodyInstance.GetCollisionProfileName());
-			if (Landscape->BodyInstance.GetCollisionProfileName() == NAME_None)
+			if (Landscape->BodyInstance.DoesUseCollisionProfile() == false)
 			{
 				Landscape->BodyInstance.SetCollisionEnabled(OldLandscapeProxy->BodyInstance.GetCollisionEnabled());
 				Landscape->BodyInstance.SetObjectType(OldLandscapeProxy->BodyInstance.GetObjectType());
diff --git a/Engine/Source/Runtime/Core/Public/UObject/ObjectVersion.h b/Engine/Source/Runtime/Core/Public/UObject/ObjectVersion.h
--- a/Engine/Source/Runtime/Core/Public/UObject/ObjectVersion.h
+++ b/Engine/Source/Runtime/Core/Public/UObject/ObjectVersion.h
@@ -516,7 +516,8 @@ enum EUnrealEngineObjectUE4Version
 	VER_UE4_FOLIAGE_MOVABLE_MOBILITY,
 	// Undo BreakMaterialAttributes changes as it broke old content
 	VER_UE4_UNDO_BREAK_MATERIALATTRIBUTES_CHANGE,
-
+	// Now Default custom profile name isn't NONE anymore due to copy/paste not working properly with it
+	VER_UE4_ADD_CUSTOMPROFILENAME_CHANGE, 
 	// -----<new versions can be added before this line>-------------------------------------------------
 	// - this needs to be the last line (see note below)
 	VER_UE4_AUTOMATIC_VERSION_PLUS_ONE,
diff --git a/Engine/Source/Runtime/Engine/Classes/Engine/CollisionProfile.h b/Engine/Source/Runtime/Engine/Classes/Engine/CollisionProfile.h
--- a/Engine/Source/Runtime/Engine/Classes/Engine/CollisionProfile.h
+++ b/Engine/Source/Runtime/Engine/Classes/Engine/CollisionProfile.h
@@ -169,6 +169,9 @@ public:
 
 	/** Convert ObjectType or TraceType to CollisionChannel */
 	ECollisionChannel ConvertToCollisionChannel(bool TraceType, int32 Index) const;
+
+	/* custom collision profile name that you can modify what you'd like */
+	ENGINE_API static FName CustomCollisionProfileName;
 private:
 
 	/** 
diff --git a/Engine/Source/Runtime/Engine/Classes/PhysicsEngine/BodyInstance.h b/Engine/Source/Runtime/Engine/Classes/PhysicsEngine/BodyInstance.h
--- a/Engine/Source/Runtime/Engine/Classes/PhysicsEngine/BodyInstance.h
+++ b/Engine/Source/Runtime/Engine/Classes/PhysicsEngine/BodyInstance.h
@@ -220,40 +220,7 @@ public:
 		TWeakObjectPtr<class UBodySetup>				BodySetup;
 
 		/** Constructor **/
-		FBodyInstance()
-		: InstanceBodyIndex(INDEX_NONE)
-		, Scale3D(1.0)
-		, SceneIndexSync(0)
-		, SceneIndexAsync(0)
-		, bEnableCollision_DEPRECATED(true)
-		, CollisionEnabled(ECollisionEnabled::QueryAndPhysics)
-		, ObjectType(ECC_WorldStatic)
-		, bUseCCD(false)
-		, bNotifyRigidBodyCollision(false)
-		, bSimulatePhysics(false)
-		, bStartAwake(true)
-		, bEnableGravity(true)
-		, bUseAsyncScene(false)
-		, bUpdateMassWhenScaleChanges(false)
-		, bOverrideWalkableSlopeOnInstance(false)
-		, PhysMaterialOverride(NULL)
-		, COMNudge(ForceInit)
-		, SleepFamily(SF_Normal)
-		, MassScale(1.f)
-		, AngularDamping(0.0)
-		, LinearDamping(0.01)
-		, MaxAngularVelocity(400.f)
-		, PhysicsBlendWeight(0.f)
-		, PositionSolverIterationCount(8)
-		, VelocitySolverIterationCount(1)
-#if WITH_PHYSX
-		, RigidActorSync(NULL)
-		, RigidActorAsync(NULL)
-		, BodyAggregate(NULL)
-		, PhysxUserData(this)
-#endif	//WITH_PHYSX
-		{
-		}
+		FBodyInstance();
 
 		// BodyInstance interface
 
@@ -432,6 +399,9 @@ public:
 		/** Get the current collision profile assigned to this body */
 		FName GetCollisionProfileName() const;
 
+		/** return true if it uses Collision Profile System. False otherwise*/
+		bool DoesUseCollisionProfile() const;
+
 		/** Update instance's mass properties (mass, inertia and center-of-mass offset) based on MassScale, InstanceMassScale and COMNudge. */
 		void UpdateMassProperties();
 
@@ -546,6 +516,11 @@ private:
 		 * for example, they would like to re-define CollisionEnabled or ObjectType or ResponseChannels
 		 */
 		void InvalidateCollisionProfileName();
+		
+		/**
+		 * Return true if the collision profile name is valid
+		 */
+		static bool IsValidCollisionProfileName(FName InCollisionProfileName);
 
 		friend class UCollisionProfile;
 		friend class FBodyInstanceCustomization;
diff --git a/Engine/Source/Runtime/Engine/Private/Collision/CollisionProfile.cpp b/Engine/Source/Runtime/Engine/Private/Collision/CollisionProfile.cpp
--- a/Engine/Source/Runtime/Engine/Private/Collision/CollisionProfile.cpp
+++ b/Engine/Source/Runtime/Engine/Private/Collision/CollisionProfile.cpp
@@ -8,6 +8,9 @@ DEFINE_LOG_CATEGORY_STATIC(LogCollisionProfile, Warning, All)
 #define MAX_CUSTOMIZABLE_COLLISIONCHANNEL	ECC_GameTraceChannel18
 #define IS_VALID_COLLISIONCHANNEL(x) ((x)> ECC_Destructible && (x) < ECC_OverlapAll_Deprecated)
 
+// do not chnage this name. This value is serialize to other object, if you change, it will mess up serialization, you'll need to fix up name by versioning
+FName UCollisionProfile::CustomCollisionProfileName = FName(TEXT("Custom"));
+
 //////////////////////////////////////////////////////////////////////////
 FCollisionResponseTemplate::FCollisionResponseTemplate()
 	: CollisionEnabled(ECollisionEnabled::NoCollision)
@@ -67,7 +70,8 @@ bool UCollisionProfile::GetProfileTemplate(FName ProfileName, struct FCollisionR
 
 bool UCollisionProfile::CheckRedirect(FName ProfileName, FBodyInstance & BodyInstance, struct FCollisionResponseTemplate & Template) const
 {
-	if ( ProfileName != NAME_None )
+	// make sure we're not setting invalid collision profile name
+	if ( FBodyInstance::IsValidCollisionProfileName(ProfileName) )
 	{
 		const FName * NewName = ProfileRedirectsMap.Find(ProfileName);
 		if ( NewName )
@@ -205,6 +209,16 @@ void UCollisionProfile::LoadProfileConfig(bool bForceInit)
 	// read "EngineTraceChanne" and "GameTraceChanne" and set meta data
 	FConfigSection* Configs = GConfig->GetSectionPrivate( TEXT("/Script/Engine.CollisionProfile"), false, true, GEngineIni );
 
+	// before any op, verify if profiles contains invalid name - such as Custom profile name - remove all of them
+	for (auto Iter=Profiles.CreateConstIterator(); Iter; ++Iter)
+	{
+		// make sure it doens't have any 
+		if (Iter->Name == CustomCollisionProfileName)
+		{
+			UE_LOG(LogCollisionProfile, Error, TEXT("Profiles contain invalid name : %s is reserved for internal use"), *CustomCollisionProfileName.ToString());
+			Profiles.RemoveAt(Iter.GetIndex());
+		}
+	}
 	// 1. First loads all meta data for custom channels
 	// this can be used in #2, so had to fix up first
 
diff --git a/Engine/Source/Runtime/Engine/Private/PhysicsEngine/BodyInstance.cpp b/Engine/Source/Runtime/Engine/Private/PhysicsEngine/BodyInstance.cpp
--- a/Engine/Source/Runtime/Engine/Private/PhysicsEngine/BodyInstance.cpp
+++ b/Engine/Source/Runtime/Engine/Private/PhysicsEngine/BodyInstance.cpp
@@ -182,6 +182,43 @@ void FCollisionResponse::UpdateResponseContainerFromArray()
 	}
 }
 ////////////////////////////////////////////////////////////////////////////
+
+FBodyInstance::FBodyInstance()
+: InstanceBodyIndex(INDEX_NONE)
+, Scale3D(1.0)
+, SceneIndexSync(0)
+, SceneIndexAsync(0)
+, bEnableCollision_DEPRECATED(true)
+, CollisionEnabled(ECollisionEnabled::QueryAndPhysics)
+, CollisionProfileName(UCollisionProfile::CustomCollisionProfileName)
+, ObjectType(ECC_WorldStatic)
+, bUseCCD(false)
+, bNotifyRigidBodyCollision(false)
+, bSimulatePhysics(false)
+, bStartAwake(true)
+, bEnableGravity(true)
+, bUseAsyncScene(false)
+, bUpdateMassWhenScaleChanges(false)
+, bOverrideWalkableSlopeOnInstance(false)
+, PhysMaterialOverride(NULL)
+, COMNudge(ForceInit)
+, SleepFamily(SF_Normal)
+, MassScale(1.f)
+, AngularDamping(0.0)
+, LinearDamping(0.01)
+, MaxAngularVelocity(400.f)
+, PhysicsBlendWeight(0.f)
+, PositionSolverIterationCount(8)
+, VelocitySolverIterationCount(1)
+#if WITH_PHYSX
+, RigidActorSync(NULL)
+, RigidActorAsync(NULL)
+, BodyAggregate(NULL)
+, PhysxUserData(this)
+#endif	//WITH_PHYSX
+{
+}
+
 FArchive& operator<<(FArchive& Ar,FBodyInstance& BodyInst)
 {
 	if (!Ar.IsLoading() && !Ar.IsSaving())
@@ -337,10 +374,7 @@ void FBodyInstance::UpdatePhysicalMaterials()
 
 void FBodyInstance::InvalidateCollisionProfileName()
 {
-	if ( CollisionProfileName != NAME_None )
-	{
-		CollisionProfileName = NAME_None;
-	}
+	CollisionProfileName = UCollisionProfile::CustomCollisionProfileName;
 }
 
 ECollisionResponse FBodyInstance::GetResponseToChannel(ECollisionChannel Channel) const
@@ -403,6 +437,10 @@ FName FBodyInstance::GetCollisionProfileName() const
 	return CollisionProfileName;
 }
 
+bool FBodyInstance::DoesUseCollisionProfile() const
+{
+	return IsValidCollisionProfileName(CollisionProfileName);
+}
 
 void FBodyInstance::SetCollisionEnabled(ECollisionEnabled::Type NewType, bool bUpdatePhysicsFilterData)
 {
@@ -2479,6 +2517,11 @@ bool FBodyInstance::Overlap(const PxGeometry& PGeom, const PxTransform&  ShapePo
 #endif //WITH_PHYSX
 
 
+bool FBodyInstance::IsValidCollisionProfileName(FName InCollisionProfileName)
+{
+	return (InCollisionProfileName != NAME_None && InCollisionProfileName != UCollisionProfile::CustomCollisionProfileName);
+}
+
 void FBodyInstance::LoadProfileData(bool bVerifyProfile)
 {
 	if ( bVerifyProfile )
@@ -2488,7 +2531,7 @@ void FBodyInstance::LoadProfileData(bool bVerifyProfile)
 		// if same, then keep the profile name
 		// if not same, that means it has been modified from default
 		// leave it as it is, and clear profile name
-		if ( CollisionProfileName != NAME_None )
+		if ( IsValidCollisionProfileName(CollisionProfileName) )
 		{
 			FCollisionResponseTemplate Template;
 			if ( UCollisionProfile::Get()->GetProfileTemplate(CollisionProfileName, Template) ) 
@@ -2497,31 +2540,31 @@ void FBodyInstance::LoadProfileData(bool bVerifyProfile)
 				// so that means it will have valid ResponsetoChannels value, so this is okay to access. 
 				if (Template.IsEqual(CollisionEnabled, ObjectType, CollisionResponses.GetResponseContainer()) == false)
 				{
-					CollisionProfileName = NAME_None;
+					InvalidateCollisionProfileName(); 
 				}
 			}
 			else
 			{
 				UE_LOG(LogPhysics, Warning, TEXT("COLLISION PROFILE [%s] is not found"), *CollisionProfileName.ToString());
 				// if not nothing to do
-				CollisionProfileName = NAME_None;
+				InvalidateCollisionProfileName(); 
 			}
 		}
 
 	}
 	else
 	{
-		if ( CollisionProfileName != NAME_None )
+		if ( IsValidCollisionProfileName(CollisionProfileName) )
 		{
 			if ( UCollisionProfile::Get()->ReadConfig(CollisionProfileName, *this) == false)
 			{
 				// clear the name
-				CollisionProfileName = NAME_None;
+				InvalidateCollisionProfileName();
 			}
 		}
 
 		// no profile, so it just needs to update container from array data
-		if ( CollisionProfileName == NAME_None)
+		if ( DoesUseCollisionProfile() == false )
 		{
 			CollisionResponses.UpdateResponseContainerFromArray();
 		}
@@ -2553,6 +2596,15 @@ void FBodyInstance::FixupData(class UObject * Loader)
 	check (Loader);
 
 	int32 const UE4Version = Loader->GetLinkerUE4Version();
+
+	if (UE4Version < VER_UE4_ADD_CUSTOMPROFILENAME_CHANGE)
+	{
+		if (CollisionProfileName == NAME_None)
+		{
+			CollisionProfileName = UCollisionProfile::CustomCollisionProfileName;
+		}
+	}
+
 	if(UE4Version < VER_UE4_CHANGE_BENABLECOLLISION_TO_COLLISIONENABLED)
 	{
 		UpdateFromDeprecatedEnableCollision();
@@ -2570,7 +2622,7 @@ void FBodyInstance::FixupData(class UObject * Loader)
 	LoadProfileData(bNeedToVerifyProfile);
 
 	// if profile isn't set, then fix up channel responses
-	if( CollisionProfileName == NAME_None )
+	if( CollisionProfileName == UCollisionProfile::CustomCollisionProfileName ) 
 	{
 		if (UE4Version >= VER_UE4_SAVE_COLLISIONRESPONSE_PER_CHANNEL)
 		{
