- Fixed UParticleModuleKillBox and particle system velocity calculations during world orign shifting

diff --git a/Engine/Source/Runtime/Engine/Private/Particles/ParticleComponents.cpp b/Engine/Source/Runtime/Engine/Private/Particles/ParticleComponents.cpp
--- a/Engine/Source/Runtime/Engine/Private/Particles/ParticleComponents.cpp
+++ b/Engine/Source/Runtime/Engine/Private/Particles/ParticleComponents.cpp
@@ -4652,6 +4652,8 @@ void UParticleSystemComponent::Deactivate()
 void UParticleSystemComponent::ApplyWorldOffset(const FVector& InOffset, bool bWorldShift)
 {
 	Super::ApplyWorldOffset(InOffset, bWorldShift);
+
+	OldPosition+= InOffset;
 	
 	for (auto It = EmitterInstances.CreateIterator(); It; ++It)
 	{
diff --git a/Engine/Source/Runtime/Engine/Private/Particles/ParticleModules.cpp b/Engine/Source/Runtime/Engine/Private/Particles/ParticleModules.cpp
--- a/Engine/Source/Runtime/Engine/Private/Particles/ParticleModules.cpp
+++ b/Engine/Source/Runtime/Engine/Private/Particles/ParticleModules.cpp
@@ -3136,7 +3136,7 @@ void UParticleModuleKillBox::Update(FParticleEmitterInstance* Owner, int32 Offse
 
 	BEGIN_UPDATE_LOOP;
 	{
-		FVector Position = Particle.Location;
+		FVector Position = Particle.Location + Owner->PositionOffsetThisTick;
 
 		if (LODLevel->RequiredModule->bUseLocalSpace)
 		{
