Added CreateInstance function to ULevelStreaming and exposed it to blueprints.
This change allows to create a copy of existing level streaming object with unique name, so user will be able to stream-in/out multiple instances of a particular map package.

diff --git a/Engine/Source/Runtime/Engine/Classes/Engine/LevelStreaming.h b/Engine/Source/Runtime/Engine/Classes/Engine/LevelStreaming.h
--- a/Engine/Source/Runtime/Engine/Classes/Engine/LevelStreaming.h
+++ b/Engine/Source/Runtime/Engine/Classes/Engine/LevelStreaming.h
@@ -233,6 +233,10 @@ class ULevelStreaming : public UObject
 	UFUNCTION(BlueprintCallable, Category="Game")
 	bool IsLevelLoaded() const;
 
+	/** Creates a new instance of this streaming level with a provided unique instance name */
+	UFUNCTION(BlueprintCallable, Category="Game")
+	ULevelStreaming* CreateInstance(FString UniqueInstanceName);
+
 	//==============================================================================================
 	// Delegates
 	
diff --git a/Engine/Source/Runtime/Engine/Private/LevelStreaming.cpp b/Engine/Source/Runtime/Engine/Private/LevelStreaming.cpp
--- a/Engine/Source/Runtime/Engine/Private/LevelStreaming.cpp
+++ b/Engine/Source/Runtime/Engine/Private/LevelStreaming.cpp
@@ -505,6 +505,44 @@ bool ULevelStreaming::IsLevelLoaded() const
 	return LoadedLevel != NULL;
 }
 
+ULevelStreaming* ULevelStreaming::CreateInstance(FString InstanceUniqueName)
+{
+	ULevelStreaming* StreamingLevelInstance = nullptr;
+	
+	UWorld* InWorld = GetWorld();
+	if (InWorld)
+	{
+		// Create instance long package name 
+		FString InstanceShortPackageName = InWorld->StreamingLevelsPrefix + FPackageName::GetShortName(InstanceUniqueName);
+		FString InstancePackagePath = FPackageName::GetLongPackagePath(PackageName.ToString()) + TEXT("/");
+		FName	InstanceUniquePackageName = FName(*(InstancePackagePath + InstanceShortPackageName));
+
+		// check if instance name is unique among existing streaming level objects
+		const bool bUniqueName = (InWorld->StreamingLevels.FindMatch(ULevelStreaming::FPackageNameMatcher(InstanceUniquePackageName)) == INDEX_NONE);
+				
+		if (bUniqueName)
+		{
+			StreamingLevelInstance = Cast<ULevelStreaming>(StaticConstructObject(GetClass(), InWorld, NAME_None, RF_Transient, NULL));
+			// new level streaming instance will load the same map package as this object
+			StreamingLevelInstance->PackageNameToLoad = (PackageNameToLoad == NAME_None ? PackageName : PackageNameToLoad);
+			// under a provided unique name
+			StreamingLevelInstance->PackageName = InstanceUniquePackageName;
+			StreamingLevelInstance->bShouldBeLoaded = false;
+			StreamingLevelInstance->bShouldBeVisible = false;
+			StreamingLevelInstance->LevelTransform = LevelTransform;
+
+			// add a new instance to streaming level list
+			InWorld->StreamingLevels.Add(StreamingLevelInstance);
+		}
+		else
+		{
+			UE_LOG(LogStreaming, Warning, TEXT("Provided streaming level instance name is not unique: %s"), *InstanceUniquePackageName.ToString());
+		}
+	}
+	
+	return StreamingLevelInstance;
+}
+
 void ULevelStreaming::BroadcastLevelLoadedStatus(UWorld* PersistentWorld, FName LevelPackageName, bool bLoaded)
 {
 	for (auto It = PersistentWorld->StreamingLevels.CreateIterator(); It; ++It)
