#UE4 Fixed RecastNavMesh serialization not working in standalone games

diff --git a/Engine/Source/Runtime/Engine/Classes/AI/Navigation/RecastNavMesh.h b/Engine/Source/Runtime/Engine/Classes/AI/Navigation/RecastNavMesh.h
--- a/Engine/Source/Runtime/Engine/Classes/AI/Navigation/RecastNavMesh.h
+++ b/Engine/Source/Runtime/Engine/Classes/AI/Navigation/RecastNavMesh.h
@@ -588,7 +588,6 @@ public:
 	// Begin UObject Interface
 	virtual void PostInitProperties() OVERRIDE;
 	virtual SIZE_T GetResourceSize(EResourceSizeMode::Type Mode) OVERRIDE;	
-	virtual bool NeedsLoadForClient() const OVERRIDE { return false; }
 
 #if WITH_EDITOR
 	virtual void PostEditChangeProperty( struct FPropertyChangedEvent& PropertyChangedEvent) OVERRIDE;
