#ttp 319372 - UE4: RENDERING: The screen will occasionally go black when the scene texture is referenced on a non-seperate translucent materials in the world
#proj Engine
#branch UE4
#summary Added new error to catch a material reading from scene color with separate translucency disabled.
#change set bUsesSceneColor to true if a scene texture expression reads from scene color texture.
#change spit out an error if scene color is used on a non separate transparency.

diff --git a/Engine/Source/Runtime/Engine/Private/Materials/HLSLMaterialTranslator.h b/Engine/Source/Runtime/Engine/Private/Materials/HLSLMaterialTranslator.h
--- a/Engine/Source/Runtime/Engine/Private/Materials/HLSLMaterialTranslator.h
+++ b/Engine/Source/Runtime/Engine/Private/Materials/HLSLMaterialTranslator.h
@@ -368,9 +368,18 @@ public:
 				Errorf(TEXT("Only transparent or postprocess materials can read from scene depth."));
 			}
 
-			if (MaterialCompilationOutput.bUsesSceneColor && Material->GetMaterialDomain() != MD_PostProcess && !IsTranslucentBlendMode(Material->GetBlendMode()))
+			if (MaterialCompilationOutput.bUsesSceneColor && Material->GetMaterialDomain() != MD_PostProcess)
 			{
-				Errorf(TEXT("Only transparent or postprocess materials can read from scene color."));
+				if (!IsTranslucentBlendMode(Material->GetBlendMode()))
+				{
+					Errorf(TEXT("Only transparent or postprocess materials can read from scene color."));
+				}
+				else if (!Material->IsSeparateTranslucencyEnabled())
+				{
+					// rendering without separate translucency could mean the scenecolor surface
+					// may be set simultaneously as both source texture and destination target. (error)
+					Errorf(TEXT("Cannot read from scene color whilst separate translucency is disabled."));
+				}
 			}
 
 			if (Material->IsLightFunction() && Material->GetBlendMode() != BLEND_Opaque)
@@ -394,7 +403,7 @@ public:
 				{
 					if (Material->GetBlendMode() == BLEND_Opaque || Material->GetBlendMode() == BLEND_Masked)
 					{
-						// In opaque pass, none the the textures are available
+						// In opaque pass, none of the textures are available
 						Errorf(TEXT("SceneTexture expressions cannot be used in opaque materials"));
 					}
 					else if (bNeedsSceneTexturePostProcessInputs)
@@ -2168,6 +2177,8 @@ protected:
 
 		ESceneTextureId SceneTextureId = (ESceneTextureId)InSceneTextureId;
 
+		MaterialCompilationOutput.bUsesSceneColor |= SceneTextureId == PPI_SceneColor;
+
 		FString DefaultScreenAligned(TEXT("MaterialFloat2(ScreenAlignedPosition(Parameters.ScreenPosition).xy)"));
 		FString TexCoordCode((UV != INDEX_NONE) ? CoerceParameter(UV, MCT_Float2) : DefaultScreenAligned);
 
