#ttp 330945 - LIVE: UE4: Add FTimerManager::IsTimerPaused() and TimerExists()

#summary Added IsTimerPaused and TimerExists function

#change Added IsTimerPaused
and TimerExists function

diff --git a/Engine/Source/Runtime/Engine/Classes/Kismet/KismetSystemLibrary.h b/Engine/Source/Runtime/Engine/Classes/Kismet/KismetSystemLibrary.h
--- a/Engine/Source/Runtime/Engine/Classes/Kismet/KismetSystemLibrary.h
+++ b/Engine/Source/Runtime/Engine/Classes/Kismet/KismetSystemLibrary.h
@@ -246,13 +246,29 @@ class UKismetSystemLibrary : public UBlueprintFunctionLibrary
 	static void K2_UnPauseTimer(UObject* Object, FString FunctionName);
 
 	/**
-	 * Returns true is a timer is active for the given delegate, false otherwise.
+	 * Returns true is a timer exists and is active for the given delegate, false otherwise.
 	 * @param Object		Object that implements the delegate function. Defaults to self (this blueprint)
 	 * @param FunctionName	Delegate function name. Can be a K2 function or a Custom Event.
 	 */
 	UFUNCTION(BlueprintPure, meta=(FriendlyName = "IsTimerActive", DefaultToSelf = "Object"), Category="Utilities|Time")
 	static bool K2_IsTimerActive(UObject* Object, FString FunctionName);
 
+	/**
+	* Returns true is a timer exists and is paused for the given delegate, false otherwise.
+	* @param Object		Object that implements the delegate function. Defaults to self (this blueprint)
+	* @param FunctionName	Delegate function name. Can be a K2 function or a Custom Event.
+	*/
+	UFUNCTION(BlueprintPure, meta = (FriendlyName = "IsTimerPaused", DefaultToSelf = "Object"), Category = "Utilities|Time")
+	static bool K2_IsTimerPaused(UObject* Object, FString FunctionName);
+
+	/**
+	* Returns true is a timer for the given delegate exists, false otherwise.
+	* @param Object		Object that implements the delegate function. Defaults to self (this blueprint)
+	* @param FunctionName	Delegate function name. Can be a K2 function or a Custom Event.
+	*/
+	UFUNCTION(BlueprintPure, meta = (FriendlyName = "TimerExists", DefaultToSelf = "Object"), Category = "Utilities|Time")
+	static bool K2_TimerExists(UObject* Object, FString FunctionName);
+	
 	/**
 	 * Returns elapsed time for the given delegate (time since current countdown iteration began).
 	 * @param Object		Object that implements the delegate function. Defaults to self (this blueprint)
diff --git a/Engine/Source/Runtime/Engine/Private/KismetSystemLibrary.cpp b/Engine/Source/Runtime/Engine/Private/KismetSystemLibrary.cpp
--- a/Engine/Source/Runtime/Engine/Private/KismetSystemLibrary.cpp
+++ b/Engine/Source/Runtime/Engine/Private/KismetSystemLibrary.cpp
@@ -327,6 +327,44 @@ bool UKismetSystemLibrary::K2_IsTimerActive(UObject* Object, FString FunctionNam
 	}
 }
 
+bool UKismetSystemLibrary::K2_IsTimerPaused(UObject* Object, FString FunctionName)
+{
+	FTimerDynamicDelegate Delegate;
+	Delegate.SetObject(Object);
+	Delegate.SetFunctionName(FName(*FunctionName));
+
+	if (Delegate.IsBound())
+	{
+		UWorld* World = GEngine->GetWorldFromContextObject(Object);
+		return World->GetTimerManager().IsTimerPaused(Delegate);
+	}
+	else
+	{
+		const FString ObjectNameStr = Object != NULL ? Object->GetName() : TEXT("None");
+		UE_LOG(LogBlueprintUserMessages, Warning, TEXT("IsTimerPaused passed a bad function (%s) or object (%s)"), *FunctionName, *ObjectNameStr);
+		return false;
+	}
+}
+
+bool UKismetSystemLibrary::K2_TimerExists(UObject* Object, FString FunctionName)
+{
+	FTimerDynamicDelegate Delegate;
+	Delegate.SetObject(Object);
+	Delegate.SetFunctionName(FName(*FunctionName));
+
+	if (Delegate.IsBound())
+	{
+		UWorld* World = GEngine->GetWorldFromContextObject(Object);
+		return World->GetTimerManager().TimerExists(Delegate);
+	}
+	else
+	{
+		const FString ObjectNameStr = Object != NULL ? Object->GetName() : TEXT("None");
+		UE_LOG(LogBlueprintUserMessages, Warning, TEXT("TimerExists passed a bad function (%s) or object (%s)"), *FunctionName, *ObjectNameStr);
+		return false;
+	}
+}
+
 float UKismetSystemLibrary::K2_GetTimerElapsedTime(UObject* Object, FString FunctionName)
 {
 	FTimerDynamicDelegate Delegate;
diff --git a/Engine/Source/Runtime/Engine/Public/TimerManager.h b/Engine/Source/Runtime/Engine/Public/TimerManager.h
--- a/Engine/Source/Runtime/Engine/Public/TimerManager.h
+++ b/Engine/Source/Runtime/Engine/Public/TimerManager.h
@@ -288,6 +288,7 @@ public:
 	{
 		return GetTimerRate( FTimerDelegate::CreateUObject(inObj, inTimerMethod) );
 	}
+
 	template< class UserClass >	
 	FORCEINLINE float GetTimerRate(UserClass* inObj, typename FTimerDelegate::TUObjectMethodDelegate_Const< UserClass >::FMethodPtr inTimerMethod)
 	{
@@ -299,23 +300,26 @@ public:
 	{
 		return InternalGetTimerRate( FTimerUnifiedDelegate(InDelegate) );
 	}
+
 	/** Version that takes a dynamic delegate (e.g. for UFunctions). */
 	FORCEINLINE float GetTimerRate(FTimerDynamicDelegate const& InDynDelegate)
 	{
 		return InternalGetTimerRate( FTimerUnifiedDelegate(InDynDelegate) );
 	}
+
 	/**
 	 * Returns true if the specified timer exists and is not paused
 	 *
-	 * @param inObj			Object to call the timer function on.
-	 * @param inTimerMethod Method to call when timer fires.
-	 * @return				true if the timer is active, false otherwise.
+	 * @param inObj			Object binding for query.
+	 * @param inTimerMethod Method binding for query.
+	 * @return				true if the timer matching the given criteria exists and is active, false otherwise.
 	 */
 	template< class UserClass >	
 	FORCEINLINE bool IsTimerActive(UserClass* inObj, typename FTimerDelegate::TUObjectMethodDelegate< UserClass >::FMethodPtr inTimerMethod)
 	{
 		return IsTimerActive( FTimerDelegate::CreateUObject(inObj, inTimerMethod) );
 	}
+		
 	template< class UserClass >	
 	FORCEINLINE bool IsTimerActive(UserClass* inObj, typename FTimerDelegate::TUObjectMethodDelegate_Const< UserClass >::FMethodPtr inTimerMethod)
 	{
@@ -328,6 +332,7 @@ public:
 		FTimerData const* const TimerData = FindTimer( FTimerUnifiedDelegate(InDelegate) );
 		return ( (TimerData != NULL) && (TimerData->Status != ETimerStatus::Paused) );
 	}
+
 	/** Version that takes a dynamic delegate (e.g. for UFunctions). */
 	FORCEINLINE bool IsTimerActive(FTimerDynamicDelegate const& InDynDelegate)
 	{
@@ -335,6 +340,71 @@ public:
 		return ( (TimerData != NULL) && (TimerData->Status != ETimerStatus::Paused) );
 	}
 
+	/**
+	* Returns true if the specified timer exists and is paused
+	*
+	* @param inObj			Object binding for query.
+	* @param inTimerMethod	Method binding for query.
+	* @return				true if the timer matching the given criteria exists and is paused, false otherwise.
+	*/
+	template< class UserClass >
+	FORCEINLINE bool IsTimerPaused(UserClass* inObj, typename FTimerDelegate::TUObjectMethodDelegate< UserClass >::FMethodPtr inTimerMethod) const
+	{
+		return IsTimerPaused(FTimerDelegate::CreateUObject(inObj, inTimerMethod));
+	}
+
+	template< class UserClass >
+	FORCEINLINE bool IsTimerPaused(UserClass* inObj, typename FTimerDelegate::TUObjectMethodDelegate_Const< UserClass >::FMethodPtr inTimerMethod)  const
+	{
+		return IsTimerPaused(FTimerDelegate::CreateUObject(inObj, inTimerMethod));
+	}
+
+	/** Version that takes any generic delegate. */
+	FORCEINLINE bool IsTimerPaused(FTimerDelegate const& InDelegate) const
+	{
+		FTimerData const* const TimerData = FindTimer(FTimerUnifiedDelegate(InDelegate));
+		return ((TimerData != NULL) && (TimerData->Status == ETimerStatus::Paused));
+	}
+
+	/** Version that takes a dynamic delegate (e.g. for UFunctions). */
+	FORCEINLINE bool IsTimerPaused(FTimerDynamicDelegate const& InDynDelegate) const
+	{
+		FTimerData const* const TimerData = FindTimer(FTimerUnifiedDelegate(InDynDelegate));
+		return ((TimerData != NULL) && (TimerData->Status == ETimerStatus::Paused));
+	}
+
+	/**
+	* Returns true if the specified timer exists
+	*
+	* @param inObj			Object binding for query.
+	* @param inTimerMethod	Method binding for query.
+	* @return				true if the timer matching the given criteria exists, false otherwise.
+	*/
+	template< class UserClass >
+	FORCEINLINE bool TimerExists(UserClass* inObj, typename FTimerDelegate::TUObjectMethodDelegate< UserClass >::FMethodPtr inTimerMethod) const
+	{
+		return TimerExists(FTimerDelegate::CreateUObject(inObj, inTimerMethod));
+	}
+
+	template< class UserClass >
+	FORCEINLINE bool TimerExists(UserClass* inObj, typename FTimerDelegate::TUObjectMethodDelegate_Const< UserClass >::FMethodPtr inTimerMethod) const
+	{
+		return TimerExists(FTimerDelegate::CreateUObject(inObj, inTimerMethod));
+	}
+
+	/** Version that takes any generic delegate. */
+	FORCEINLINE bool TimerExists(FTimerDelegate const& InDelegate) const
+	{
+		return FindTimer(FTimerUnifiedDelegate(InDelegate)) != NULL;
+		
+	}
+
+	/** Version that takes a dynamic delegate (e.g. for UFunctions). */
+	FORCEINLINE bool TimerExists(FTimerDynamicDelegate const& InDynDelegate) const
+	{
+		return FindTimer(FTimerUnifiedDelegate(InDynDelegate)) != NULL;
+	}
+
 	/**
 	 * Gets the current elapsed time for the specified timer.
 	 *
