#ttp 303393 UObject improvement: Remove header groups

diff --git a/Engine/Source/Runtime/Engine/Classes/Components/BoxReflectionCaptureComponent.h b/Engine/Source/Runtime/Engine/Classes/Components/BoxReflectionCaptureComponent.h
--- a/Engine/Source/Runtime/Engine/Classes/Components/BoxReflectionCaptureComponent.h
+++ b/Engine/Source/Runtime/Engine/Classes/Components/BoxReflectionCaptureComponent.h
@@ -5,7 +5,7 @@
 #include "BoxReflectionCaptureComponent.generated.h"
 
 	// -> will be exported to EngineDecalClasses.h
-UCLASS(hidecategories=(Collision, Object, Physics, SceneComponent), HeaderGroup=Decal, MinimalAPI)
+UCLASS(hidecategories=(Collision, Object, Physics, SceneComponent), MinimalAPI)
 class UBoxReflectionCaptureComponent : public UReflectionCaptureComponent
 {
 	GENERATED_UCLASS_BODY()
diff --git a/Engine/Source/Runtime/Engine/Classes/Components/BrushComponent.h b/Engine/Source/Runtime/Engine/Classes/Components/BrushComponent.h
--- a/Engine/Source/Runtime/Engine/Classes/Components/BrushComponent.h
+++ b/Engine/Source/Runtime/Engine/Classes/Components/BrushComponent.h
@@ -4,7 +4,7 @@
 #pragma once
 #include "BrushComponent.generated.h"
 
-UCLASS(HeaderGroup=Component, editinlinenew, MinimalAPI, showcategories=(Mobility), hidecategories=(Physics, Lighting, LOD, Rendering, TextureStreaming, Transform, Activation, "Components|Activation"))
+UCLASS(editinlinenew, MinimalAPI, showcategories=(Mobility), hidecategories=(Physics, Lighting, LOD, Rendering, TextureStreaming, Transform, Activation, "Components|Activation"))
 class UBrushComponent : public UPrimitiveComponent
 {
 	GENERATED_UCLASS_BODY()
diff --git a/Engine/Source/Runtime/Engine/Classes/Components/CapsuleComponent.h b/Engine/Source/Runtime/Engine/Classes/Components/CapsuleComponent.h
--- a/Engine/Source/Runtime/Engine/Classes/Components/CapsuleComponent.h
+++ b/Engine/Source/Runtime/Engine/Classes/Components/CapsuleComponent.h
@@ -6,7 +6,7 @@
 /** 
  * A capsule generally used for simple collision. Bounds are rendered as lines in the editor.
  */
-UCLASS(HeaderGroup=Component, ClassGroup=Shapes, editinlinenew, hidecategories=(Object,LOD,Lighting,TextureStreaming), meta=(BlueprintSpawnableComponent))
+UCLASS(ClassGroup=Shapes, editinlinenew, hidecategories=(Object,LOD,Lighting,TextureStreaming), meta=(BlueprintSpawnableComponent))
 class ENGINE_API UCapsuleComponent : public UShapeComponent
 {
 	GENERATED_UCLASS_BODY()
diff --git a/Engine/Source/Runtime/Engine/Classes/Components/ChildActorComponent.h b/Engine/Source/Runtime/Engine/Classes/Components/ChildActorComponent.h
--- a/Engine/Source/Runtime/Engine/Classes/Components/ChildActorComponent.h
+++ b/Engine/Source/Runtime/Engine/Classes/Components/ChildActorComponent.h
@@ -5,7 +5,7 @@
 #include "ChildActorComponent.generated.h"
 
 /** A component that spawns an Actor when registered, and destroys it when unregistered.*/
-UCLASS(HeaderGroup=Component, ClassGroup=Utility, hidecategories=(Object,LOD,Physics,Lighting,TextureStreaming,Activation,"Components|Activation",Collision), meta=(BlueprintSpawnableComponent), MinimalAPI)
+UCLASS(ClassGroup=Utility, hidecategories=(Object,LOD,Physics,Lighting,TextureStreaming,Activation,"Components|Activation",Collision), meta=(BlueprintSpawnableComponent), MinimalAPI)
 class UChildActorComponent : public USceneComponent
 {
 	GENERATED_UCLASS_BODY()
diff --git a/Engine/Source/Runtime/Engine/Classes/Components/DecalComponent.h b/Engine/Source/Runtime/Engine/Classes/Components/DecalComponent.h
--- a/Engine/Source/Runtime/Engine/Classes/Components/DecalComponent.h
+++ b/Engine/Source/Runtime/Engine/Classes/Components/DecalComponent.h
@@ -9,7 +9,7 @@
 /** 
  * A material that is rendered onto the surface of a mesh. A kind of 'bumper sticker' for a model.
  */
-UCLASS(hidecategories=(Collision, Object, Physics, SceneComponent, Activation, "Components|Activation", Physics, Mobility), ClassGroup=Rendering, HeaderGroup=Decal, MinimalAPI, meta=(BlueprintSpawnableComponent))
+UCLASS(hidecategories=(Collision, Object, Physics, SceneComponent, Activation, "Components|Activation", Physics, Mobility), ClassGroup=Rendering, MinimalAPI, meta=(BlueprintSpawnableComponent))
 class UDecalComponent : public USceneComponent
 {
 	GENERATED_UCLASS_BODY()
diff --git a/Engine/Source/Runtime/Engine/Classes/Components/DestructibleComponent.h b/Engine/Source/Runtime/Engine/Classes/Components/DestructibleComponent.h
--- a/Engine/Source/Runtime/Engine/Classes/Components/DestructibleComponent.h
+++ b/Engine/Source/Runtime/Engine/Classes/Components/DestructibleComponent.h
@@ -37,7 +37,7 @@ struct FDestructibleChunkInfo
  *
  *	The USkeletalMesh pointer in the base class (SkinnedMeshComponent) MUST be a DestructibleMesh
  */
-UCLASS(HeaderGroup=Component, ClassGroup=Physics, hidecategories=(Object,Mesh,"Components|SkinnedMesh",Mirroring,Activation,"Components|Activation"), config=Engine, editinlinenew, meta=(BlueprintSpawnableComponent))
+UCLASS(ClassGroup=Physics, hidecategories=(Object,Mesh,"Components|SkinnedMesh",Mirroring,Activation,"Components|Activation"), config=Engine, editinlinenew, meta=(BlueprintSpawnableComponent))
 class ENGINE_API UDestructibleComponent : public USkinnedMeshComponent
 {
 	GENERATED_UCLASS_BODY()
diff --git a/Engine/Source/Runtime/Engine/Classes/Components/DirectionalLightComponent.h b/Engine/Source/Runtime/Engine/Classes/Components/DirectionalLightComponent.h
--- a/Engine/Source/Runtime/Engine/Classes/Components/DirectionalLightComponent.h
+++ b/Engine/Source/Runtime/Engine/Classes/Components/DirectionalLightComponent.h
@@ -7,7 +7,7 @@
 /**
  * A light component that has parallel rays. Will provide a uniform lighting across any affected surface (eg. The Sun). This will affect all objects in the defined light-mass importance volume.
  */
-UCLASS(HeaderGroup=Light, ClassGroup=Lights, hidecategories=(Object, LightProfiles), dependson=UEngineTypes, editinlinenew, meta=(BlueprintSpawnableComponent), MinimalAPI)
+UCLASS(ClassGroup=Lights, hidecategories=(Object, LightProfiles), dependson=UEngineTypes, editinlinenew, meta=(BlueprintSpawnableComponent), MinimalAPI)
 class UDirectionalLightComponent : public ULightComponent
 {
 	GENERATED_UCLASS_BODY()
diff --git a/Engine/Source/Runtime/Engine/Classes/Components/DrawFrustumComponent.h b/Engine/Source/Runtime/Engine/Classes/Components/DrawFrustumComponent.h
--- a/Engine/Source/Runtime/Engine/Classes/Components/DrawFrustumComponent.h
+++ b/Engine/Source/Runtime/Engine/Classes/Components/DrawFrustumComponent.h
@@ -8,7 +8,7 @@
  *	Utility component for drawing a view frustum. Origin is at the component location, frustum points down position X axis.
  */ 
 
-UCLASS(HeaderGroup=Component, collapsecategories, hidecategories=Object, editinlinenew, MinimalAPI)
+UCLASS(collapsecategories, hidecategories=Object, editinlinenew, MinimalAPI)
 class UDrawFrustumComponent : public UPrimitiveComponent
 {
 	GENERATED_UCLASS_BODY()
diff --git a/Engine/Source/Runtime/Engine/Classes/Components/DrawSphereComponent.h b/Engine/Source/Runtime/Engine/Classes/Components/DrawSphereComponent.h
--- a/Engine/Source/Runtime/Engine/Classes/Components/DrawSphereComponent.h
+++ b/Engine/Source/Runtime/Engine/Classes/Components/DrawSphereComponent.h
@@ -7,7 +7,7 @@
 /** 
  * A sphere generally used for simple collision. Bounds are rendered as lines in the editor.
  */
-UCLASS(HeaderGroup=Component, collapsecategories, hidecategories=Object, editinlinenew, MinimalAPI)
+UCLASS(collapsecategories, hidecategories=Object, editinlinenew, MinimalAPI)
 class UDrawSphereComponent : public USphereComponent
 {
 	GENERATED_UCLASS_BODY()
diff --git a/Engine/Source/Runtime/Engine/Classes/Components/ExponentialHeightFogComponent.h b/Engine/Source/Runtime/Engine/Classes/Components/ExponentialHeightFogComponent.h
--- a/Engine/Source/Runtime/Engine/Classes/Components/ExponentialHeightFogComponent.h
+++ b/Engine/Source/Runtime/Engine/Classes/Components/ExponentialHeightFogComponent.h
@@ -7,7 +7,7 @@
 /**
  *	Used to create fogging effects such as clouds but with a density that is related to the height of the fog.
  */
-UCLASS(HeaderGroup=Component, ClassGroup=Rendering, collapsecategories, hidecategories=(Object, Mobility), editinlinenew, meta=(BlueprintSpawnableComponent))
+UCLASS(ClassGroup=Rendering, collapsecategories, hidecategories=(Object, Mobility), editinlinenew, meta=(BlueprintSpawnableComponent))
 class ENGINE_API UExponentialHeightFogComponent : public USceneComponent
 {
 	GENERATED_UCLASS_BODY()
diff --git a/Engine/Source/Runtime/Engine/Classes/Components/InputComponent.h b/Engine/Source/Runtime/Engine/Classes/Components/InputComponent.h
--- a/Engine/Source/Runtime/Engine/Classes/Components/InputComponent.h
+++ b/Engine/Source/Runtime/Engine/Classes/Components/InputComponent.h
@@ -411,7 +411,7 @@ namespace EControllerAnalogStick
 	};
 }
 
-UCLASS(transient, HeaderGroup=Component, config=Input, dependson=UPlayerInput, hidecategories=(Activation, "Components|Activation"))
+UCLASS(transient, config=Input, hidecategories=(Activation, "Components|Activation"))
 class ENGINE_API UInputComponent : public UActorComponent
 {
 	GENERATED_UCLASS_BODY()
diff --git a/Engine/Source/Runtime/Engine/Classes/Components/InstancedStaticMeshComponent.h b/Engine/Source/Runtime/Engine/Classes/Components/InstancedStaticMeshComponent.h
--- a/Engine/Source/Runtime/Engine/Classes/Components/InstancedStaticMeshComponent.h
+++ b/Engine/Source/Runtime/Engine/Classes/Components/InstancedStaticMeshComponent.h
@@ -57,7 +57,7 @@ struct FInstancedStaticMeshMappingInfo
 };
 
 /** A component that efficiently renders multiple instances of the same StaticMesh. */
-UCLASS(HeaderGroup=Component, ClassGroup=Rendering, meta=(BlueprintSpawnableComponent))
+UCLASS(ClassGroup=Rendering, meta=(BlueprintSpawnableComponent))
 class ENGINE_API UInstancedStaticMeshComponent : public UStaticMeshComponent
 {
 	GENERATED_UCLASS_BODY()
diff --git a/Engine/Source/Runtime/Engine/Classes/Components/InteractiveFoliageComponent.h b/Engine/Source/Runtime/Engine/Classes/Components/InteractiveFoliageComponent.h
--- a/Engine/Source/Runtime/Engine/Classes/Components/InteractiveFoliageComponent.h
+++ b/Engine/Source/Runtime/Engine/Classes/Components/InteractiveFoliageComponent.h
@@ -4,7 +4,7 @@
 #pragma once
 #include "InteractiveFoliageComponent.generated.h"
 
-UCLASS(HeaderGroup=Component, hidecategories=Object)
+UCLASS(hidecategories=Object)
 class UInteractiveFoliageComponent : public UStaticMeshComponent
 {
 	GENERATED_UCLASS_BODY()
diff --git a/Engine/Source/Runtime/Engine/Classes/Components/LightComponent.h b/Engine/Source/Runtime/Engine/Classes/Components/LightComponent.h
--- a/Engine/Source/Runtime/Engine/Classes/Components/LightComponent.h
+++ b/Engine/Source/Runtime/Engine/Classes/Components/LightComponent.h
@@ -27,7 +27,7 @@ public:
 	bool bPrecomputedLightingIsValid;
 };
 
-UCLASS(HeaderGroup=Component, abstract, HideCategories=(Trigger,Activation,"Components|Activation",Physics), ShowCategories=(Mobility))
+UCLASS(abstract, HideCategories=(Trigger,Activation,"Components|Activation",Physics), ShowCategories=(Mobility))
 class ENGINE_API ULightComponent : public ULightComponentBase
 {
 	GENERATED_UCLASS_BODY()
diff --git a/Engine/Source/Runtime/Engine/Classes/Components/LightComponentBase.h b/Engine/Source/Runtime/Engine/Classes/Components/LightComponentBase.h
--- a/Engine/Source/Runtime/Engine/Classes/Components/LightComponentBase.h
+++ b/Engine/Source/Runtime/Engine/Classes/Components/LightComponentBase.h
@@ -4,7 +4,7 @@
 #pragma once
 #include "LightComponentBase.generated.h"
 
-UCLASS(HeaderGroup=Component, abstract, HideCategories=(Trigger,Activation,"Components|Activation",Physics), ShowCategories=(Mobility))
+UCLASS(abstract, HideCategories=(Trigger,Activation,"Components|Activation",Physics), ShowCategories=(Mobility))
 class ENGINE_API ULightComponentBase : public USceneComponent
 {
 	GENERATED_UCLASS_BODY()
diff --git a/Engine/Source/Runtime/Engine/Classes/Components/LineBatchComponent.h b/Engine/Source/Runtime/Engine/Classes/Components/LineBatchComponent.h
--- a/Engine/Source/Runtime/Engine/Classes/Components/LineBatchComponent.h
+++ b/Engine/Source/Runtime/Engine/Classes/Components/LineBatchComponent.h
@@ -104,7 +104,7 @@ public:
 	float RemainingLifeTime;
 };
 
-UCLASS(HeaderGroup=Component, MinimalAPI)
+UCLASS(MinimalAPI)
 class ULineBatchComponent : public UPrimitiveComponent
 {
 	GENERATED_UCLASS_BODY()
diff --git a/Engine/Source/Runtime/Engine/Classes/Components/MaterialBillboardComponent.h b/Engine/Source/Runtime/Engine/Classes/Components/MaterialBillboardComponent.h
--- a/Engine/Source/Runtime/Engine/Classes/Components/MaterialBillboardComponent.h
+++ b/Engine/Source/Runtime/Engine/Classes/Components/MaterialBillboardComponent.h
@@ -50,7 +50,7 @@ struct FMaterialSpriteElement
 /** 
  * A 2d material that will be rendered always facing the camera.
  */
-UCLASS(HeaderGroup=Component, ClassGroup=Rendering, collapsecategories, hidecategories=Object, editinlinenew, meta=(BlueprintSpawnableComponent))
+UCLASS(ClassGroup=Rendering, collapsecategories, hidecategories=Object, editinlinenew, meta=(BlueprintSpawnableComponent))
 class UMaterialBillboardComponent : public UPrimitiveComponent
 {
 	GENERATED_UCLASS_BODY()
diff --git a/Engine/Source/Runtime/Engine/Classes/Components/MeshComponent.h b/Engine/Source/Runtime/Engine/Classes/Components/MeshComponent.h
--- a/Engine/Source/Runtime/Engine/Classes/Components/MeshComponent.h
+++ b/Engine/Source/Runtime/Engine/Classes/Components/MeshComponent.h
@@ -4,7 +4,7 @@
 #pragma once
 #include "MeshComponent.generated.h"
 
-UCLASS(HeaderGroup=Component, abstract)
+UCLASS(abstract)
 class ENGINE_API UMeshComponent : public UPrimitiveComponent
 {
 	GENERATED_UCLASS_BODY()
diff --git a/Engine/Source/Runtime/Engine/Classes/Components/ModelComponent.h b/Engine/Source/Runtime/Engine/Classes/Components/ModelComponent.h
--- a/Engine/Source/Runtime/Engine/Classes/Components/ModelComponent.h
+++ b/Engine/Source/Runtime/Engine/Classes/Components/ModelComponent.h
@@ -4,7 +4,12 @@
 #pragma once
 #include "ModelComponent.generated.h"
 
-UCLASS(HeaderGroup=Component, MinimalAPI)
+//
+// Forward declarations
+//
+class FModelElement;
+
+UCLASS(MinimalAPI)
 class UModelComponent : public UPrimitiveComponent, public IInterface_CollisionDataProvider
 {
 	GENERATED_UCLASS_BODY()
diff --git a/Engine/Source/Runtime/Engine/Classes/Components/PawnNoiseEmitterComponent.h b/Engine/Source/Runtime/Engine/Classes/Components/PawnNoiseEmitterComponent.h
--- a/Engine/Source/Runtime/Engine/Classes/Components/PawnNoiseEmitterComponent.h
+++ b/Engine/Source/Runtime/Engine/Classes/Components/PawnNoiseEmitterComponent.h
@@ -9,7 +9,7 @@
  * PawnNoiseEmitterComponent tracks noise event data used by SensingComponents to hear a Pawn.
  * This component is intended to exist on either a Pawn or its Controller. It does nothing on network clients.
  */
-UCLASS(HeaderGroup=Component, ClassGroup=AI, meta=(BlueprintSpawnableComponent))
+UCLASS(ClassGroup=AI, meta=(BlueprintSpawnableComponent))
 class ENGINE_API UPawnNoiseEmitterComponent : public UActorComponent
 {
 	GENERATED_UCLASS_BODY()
diff --git a/Engine/Source/Runtime/Engine/Classes/Components/PawnSensingComponent.h b/Engine/Source/Runtime/Engine/Classes/Components/PawnSensingComponent.h
--- a/Engine/Source/Runtime/Engine/Classes/Components/PawnSensingComponent.h
+++ b/Engine/Source/Runtime/Engine/Classes/Components/PawnSensingComponent.h
@@ -9,7 +9,7 @@
  * SensingComponent encapsulates sensory (ie sight and hearing) settings and functionality for an Actor,
  * allowing the actor to see/hear Pawns in the world. It does nothing on network clients.
  */
-UCLASS(HeaderGroup=Component, ClassGroup=AI, HideCategories=(Activation, "Components|Activation", Collision), meta=(BlueprintSpawnableComponent))
+UCLASS(ClassGroup=AI, HideCategories=(Activation, "Components|Activation", Collision), meta=(BlueprintSpawnableComponent))
 class ENGINE_API UPawnSensingComponent : public UActorComponent
 {
 	GENERATED_UCLASS_BODY()
diff --git a/Engine/Source/Runtime/Engine/Classes/Components/PlaneReflectionCaptureComponent.h b/Engine/Source/Runtime/Engine/Classes/Components/PlaneReflectionCaptureComponent.h
--- a/Engine/Source/Runtime/Engine/Classes/Components/PlaneReflectionCaptureComponent.h
+++ b/Engine/Source/Runtime/Engine/Classes/Components/PlaneReflectionCaptureComponent.h
@@ -5,7 +5,7 @@
 #include "PlaneReflectionCaptureComponent.generated.h"
 
 	// -> will be exported to EngineDecalClasses.h
-UCLASS(hidecategories=(Collision, Object, Physics, SceneComponent), HeaderGroup=Decal, MinimalAPI)
+UCLASS(hidecategories=(Collision, Object, Physics, SceneComponent), MinimalAPI)
 class UPlaneReflectionCaptureComponent : public UReflectionCaptureComponent
 {
 	GENERATED_UCLASS_BODY()
diff --git a/Engine/Source/Runtime/Engine/Classes/Components/PointLightComponent.h b/Engine/Source/Runtime/Engine/Classes/Components/PointLightComponent.h
--- a/Engine/Source/Runtime/Engine/Classes/Components/PointLightComponent.h
+++ b/Engine/Source/Runtime/Engine/Classes/Components/PointLightComponent.h
@@ -7,7 +7,7 @@
 /**
  * A light component which emits light from a single point equally in all directions.
  */
-UCLASS(HeaderGroup=Light, ClassGroup=Lights, hidecategories=(Object, LightShafts), editinlinenew, meta=(BlueprintSpawnableComponent))
+UCLASS(ClassGroup=Lights, hidecategories=(Object, LightShafts), editinlinenew, meta=(BlueprintSpawnableComponent))
 class ENGINE_API UPointLightComponent : public ULightComponent
 {
 	GENERATED_UCLASS_BODY()
diff --git a/Engine/Source/Runtime/Engine/Classes/Components/PoseableMeshComponent.h b/Engine/Source/Runtime/Engine/Classes/Components/PoseableMeshComponent.h
--- a/Engine/Source/Runtime/Engine/Classes/Components/PoseableMeshComponent.h
+++ b/Engine/Source/Runtime/Engine/Classes/Components/PoseableMeshComponent.h
@@ -1,9 +1,7 @@
 // Copyright 1998-2014 Epic Games, Inc. All Rights Reserved.
 
-
-
-
 #pragma once
+#include "Animation/AnimationAsset.h"
 #include "PoseableMeshComponent.generated.h"
 
 UENUM()
@@ -23,7 +21,7 @@ namespace EBoneSpaces
 /**
  *	UPoseableMeshComponent that allows bone transforms to be driven by blueprint.
  */
-UCLASS(HeaderGroup=Component, ClassGroup=Rendering, hidecategories=(Object,Physics), config=Engine, editinlinenew, meta=(BlueprintSpawnableComponent))
+UCLASS(ClassGroup=Rendering, hidecategories=(Object,Physics), config=Engine, editinlinenew, meta=(BlueprintSpawnableComponent))
 class ENGINE_API UPoseableMeshComponent : public USkinnedMeshComponent
 {
 	GENERATED_UCLASS_BODY()
diff --git a/Engine/Source/Runtime/Engine/Classes/Components/PostProcessComponent.h b/Engine/Source/Runtime/Engine/Classes/Components/PostProcessComponent.h
--- a/Engine/Source/Runtime/Engine/Classes/Components/PostProcessComponent.h
+++ b/Engine/Source/Runtime/Engine/Classes/Components/PostProcessComponent.h
@@ -1,6 +1,6 @@
 // Copyright 1998-2013 Epic Games, Inc. All Rights Reserved.
 
-
+#pragma once
 #include "PostProcessComponent.generated.h"
 
 /**
@@ -8,7 +8,7 @@
  *	Will use a parent UShapeComponent to provide volume data if available.
  */
 
-UCLASS(HeaderGroup = Component, ClassGroup = Rendering, collapsecategories, hidecategories = (Object), editinlinenew, meta = (BlueprintSpawnableComponent), MinimalAPI)
+UCLASS(ClassGroup = Rendering, collapsecategories, hidecategories = (Object), editinlinenew, meta = (BlueprintSpawnableComponent), MinimalAPI)
 class UPostProcessComponent : public USceneComponent, public IInterface_PostProcessVolume
 {
 	GENERATED_UCLASS_BODY()
diff --git a/Engine/Source/Runtime/Engine/Classes/Components/PrimitiveComponent.h b/Engine/Source/Runtime/Engine/Classes/Components/PrimitiveComponent.h
--- a/Engine/Source/Runtime/Engine/Classes/Components/PrimitiveComponent.h
+++ b/Engine/Source/Runtime/Engine/Classes/Components/PrimitiveComponent.h
@@ -183,7 +183,7 @@ DECLARE_DYNAMIC_MULTICAST_DELEGATE_TwoParams( FComponentOnInputTouchEndSignature
 DECLARE_DYNAMIC_MULTICAST_DELEGATE_TwoParams( FComponentBeginTouchOverSignature, ETouchIndex::Type, FingerIndex, UPrimitiveComponent*, TouchedComponent );
 DECLARE_DYNAMIC_MULTICAST_DELEGATE_TwoParams( FComponentEndTouchOverSignature, ETouchIndex::Type, FingerIndex, UPrimitiveComponent*, TouchedComponent );
 
-UCLASS(dependson=(UScene, ULightComponent, UEngineTypes), HeaderGroup=Component, abstract, HideCategories=(Mobility))
+UCLASS(dependson=(UScene, ULightComponent, UEngineTypes), abstract, HideCategories=(Mobility))
 class ENGINE_API UPrimitiveComponent : public USceneComponent
 {
 	GENERATED_UCLASS_BODY()
diff --git a/Engine/Source/Runtime/Engine/Classes/Components/ReflectionCaptureComponent.h b/Engine/Source/Runtime/Engine/Classes/Components/ReflectionCaptureComponent.h
--- a/Engine/Source/Runtime/Engine/Classes/Components/ReflectionCaptureComponent.h
+++ b/Engine/Source/Runtime/Engine/Classes/Components/ReflectionCaptureComponent.h
@@ -66,7 +66,7 @@ private:
 };
 
 // -> will be exported to EngineDecalClasses.h
-UCLASS(abstract, hidecategories=(Collision, Object, Physics, SceneComponent, Activation, "Components|Activation", Mobility), HeaderGroup=Decal, MinimalAPI)
+UCLASS(abstract, hidecategories=(Collision, Object, Physics, SceneComponent, Activation, "Components|Activation", Mobility), MinimalAPI)
 class UReflectionCaptureComponent : public USceneComponent
 {
 	GENERATED_UCLASS_BODY()
diff --git a/Engine/Source/Runtime/Engine/Classes/Components/SceneCaptureComponent.h b/Engine/Source/Runtime/Engine/Classes/Components/SceneCaptureComponent.h
--- a/Engine/Source/Runtime/Engine/Classes/Components/SceneCaptureComponent.h
+++ b/Engine/Source/Runtime/Engine/Classes/Components/SceneCaptureComponent.h
@@ -5,7 +5,7 @@
 #include "SceneCaptureComponent.generated.h"
 
 	// -> will be exported to EngineDecalClasses.h
-UCLASS(hidecategories=(abstract, Collision, Object, Physics, SceneComponent, Mobility), HeaderGroup=Decal, MinimalAPI)
+UCLASS(hidecategories=(abstract, Collision, Object, Physics, SceneComponent, Mobility), MinimalAPI)
 class USceneCaptureComponent : public USceneComponent
 {
 	GENERATED_UCLASS_BODY()
diff --git a/Engine/Source/Runtime/Engine/Classes/Components/SceneCaptureComponent2D.h b/Engine/Source/Runtime/Engine/Classes/Components/SceneCaptureComponent2D.h
--- a/Engine/Source/Runtime/Engine/Classes/Components/SceneCaptureComponent2D.h
+++ b/Engine/Source/Runtime/Engine/Classes/Components/SceneCaptureComponent2D.h
@@ -16,7 +16,7 @@ enum ESceneCaptureSource
 /**
  *	Used to capture a 'snapshot' of the scene from a single plane and feed it to a render target.
  */
-UCLASS(hidecategories=(Collision, Object, Physics, SceneComponent), ClassGroup=Rendering, HeaderGroup=Decal, MinimalAPI, editinlinenew, meta=(BlueprintSpawnableComponent))
+UCLASS(hidecategories=(Collision, Object, Physics, SceneComponent), ClassGroup=Rendering, MinimalAPI, editinlinenew, meta=(BlueprintSpawnableComponent))
 class USceneCaptureComponent2D : public USceneCaptureComponent
 {
 	GENERATED_UCLASS_BODY()
diff --git a/Engine/Source/Runtime/Engine/Classes/Components/SceneCaptureComponentCube.h b/Engine/Source/Runtime/Engine/Classes/Components/SceneCaptureComponentCube.h
--- a/Engine/Source/Runtime/Engine/Classes/Components/SceneCaptureComponentCube.h
+++ b/Engine/Source/Runtime/Engine/Classes/Components/SceneCaptureComponentCube.h
@@ -6,7 +6,7 @@
 /**
  *	Used to capture a 'snapshot' of the scene from a 6 planes and feed it to a render target.
  */
-UCLASS(hidecategories = (Collision, Object, Physics, SceneComponent), ClassGroup=Rendering, HeaderGroup = Decal, MinimalAPI, editinlinenew, meta = (BlueprintSpawnableComponent))
+UCLASS(hidecategories = (Collision, Object, Physics, SceneComponent), ClassGroup=Rendering, MinimalAPI, editinlinenew, meta = (BlueprintSpawnableComponent))
 class USceneCaptureComponentCube : public USceneCaptureComponent
 {
 	GENERATED_UCLASS_BODY()
diff --git a/Engine/Source/Runtime/Engine/Classes/Components/SceneComponent.h b/Engine/Source/Runtime/Engine/Classes/Components/SceneComponent.h
--- a/Engine/Source/Runtime/Engine/Classes/Components/SceneComponent.h
+++ b/Engine/Source/Runtime/Engine/Classes/Components/SceneComponent.h
@@ -2,6 +2,7 @@
 
 #pragma once
 
+#include "Components.h"
 #include "SceneComponent.generated.h"
 
 /** Overlap info consisting of the primitive and the body that is overlapping */
@@ -62,7 +63,7 @@ FORCEINLINE void operator|=(EMoveComponentFlags& Dest,EMoveComponentFlags Arg)
 
 
 /** A SceneComponent has a transform and supports attachment, but has no rendering or collision capabilities. Useful as a 'dummy' component in the hierarchy to offset others. */
-UCLASS(HeaderGroup=Component, ClassGroup=Utility, BlueprintType, HideCategories=(Trigger), meta=(BlueprintSpawnableComponent))
+UCLASS(ClassGroup=Utility, BlueprintType, HideCategories=(Trigger), meta=(BlueprintSpawnableComponent))
 class ENGINE_API USceneComponent : public UActorComponent
 {
 	GENERATED_UCLASS_BODY()
diff --git a/Engine/Source/Runtime/Engine/Classes/Components/ShapeComponent.h b/Engine/Source/Runtime/Engine/Classes/Components/ShapeComponent.h
--- a/Engine/Source/Runtime/Engine/Classes/Components/ShapeComponent.h
+++ b/Engine/Source/Runtime/Engine/Classes/Components/ShapeComponent.h
@@ -4,7 +4,7 @@
 #pragma once
 #include "ShapeComponent.generated.h"
 
-UCLASS(abstract, HeaderGroup=Component, hidecategories=(Object,LOD,Lighting,TextureStreaming,Activation,"Components|Activation"), editinlinenew, MinimalAPI, meta=(BlueprintSpawnableComponent), showcategories=(Mobility))
+UCLASS(abstract, hidecategories=(Object,LOD,Lighting,TextureStreaming,Activation,"Components|Activation"), editinlinenew, MinimalAPI, meta=(BlueprintSpawnableComponent), showcategories=(Mobility))
 class UShapeComponent : public UPrimitiveComponent
 {
 	GENERATED_UCLASS_BODY()
diff --git a/Engine/Source/Runtime/Engine/Classes/Components/SkeletalMeshComponent.h b/Engine/Source/Runtime/Engine/Classes/Components/SkeletalMeshComponent.h
--- a/Engine/Source/Runtime/Engine/Classes/Components/SkeletalMeshComponent.h
+++ b/Engine/Source/Runtime/Engine/Classes/Components/SkeletalMeshComponent.h
@@ -218,7 +218,7 @@ struct FSkeletalMeshComponentPreClothTickFunction : public FTickFunction
 /**
  * SkeletalMeshComponent is a mesh that supports skeletal animation and physics.
  */
-UCLASS(HeaderGroup=Component, ClassGroup=(Rendering, Common), hidecategories=Object, config=Engine, editinlinenew, meta=(BlueprintSpawnableComponent))
+UCLASS(ClassGroup=(Rendering, Common), hidecategories=Object, config=Engine, editinlinenew, meta=(BlueprintSpawnableComponent))
 class ENGINE_API USkeletalMeshComponent : public USkinnedMeshComponent
 {
 	GENERATED_UCLASS_BODY()
diff --git a/Engine/Source/Runtime/Engine/Classes/Components/SkinnedMeshComponent.h b/Engine/Source/Runtime/Engine/Classes/Components/SkinnedMeshComponent.h
--- a/Engine/Source/Runtime/Engine/Classes/Components/SkinnedMeshComponent.h
+++ b/Engine/Source/Runtime/Engine/Classes/Components/SkinnedMeshComponent.h
@@ -8,8 +8,14 @@
  */
 
 #pragma once
+#include "SkeletalMeshTypes.h"
 #include "SkinnedMeshComponent.generated.h"
 
+//
+// Forward declarations
+//
+class FSkeletalMeshResource;
+
 //
 // Bone Visibility.
 //
@@ -119,7 +125,7 @@ struct FSkelMeshComponentLODInfo
 	
 };
 
-UCLASS(HeaderGroup=Component, hidecategories=Object, config=Engine, editinlinenew, abstract)
+UCLASS(hidecategories=Object, config=Engine, editinlinenew, abstract)
 class ENGINE_API USkinnedMeshComponent : public UMeshComponent
 {
 	GENERATED_UCLASS_BODY()
diff --git a/Engine/Source/Runtime/Engine/Classes/Components/SkyLightComponent.h b/Engine/Source/Runtime/Engine/Classes/Components/SkyLightComponent.h
--- a/Engine/Source/Runtime/Engine/Classes/Components/SkyLightComponent.h
+++ b/Engine/Source/Runtime/Engine/Classes/Components/SkyLightComponent.h
@@ -81,7 +81,7 @@ enum ESkyLightSourceType
 	SLS_MAX,
 };
 
-UCLASS(HeaderGroup=Component, ClassGroup=Lights, HideCategories=(Trigger,Activation,"Components|Activation",Physics), meta=(BlueprintSpawnableComponent))
+UCLASS(ClassGroup=Lights, HideCategories=(Trigger,Activation,"Components|Activation",Physics), meta=(BlueprintSpawnableComponent))
 class ENGINE_API USkyLightComponent : public ULightComponentBase
 {
 	GENERATED_UCLASS_BODY()
diff --git a/Engine/Source/Runtime/Engine/Classes/Components/SphereComponent.h b/Engine/Source/Runtime/Engine/Classes/Components/SphereComponent.h
--- a/Engine/Source/Runtime/Engine/Classes/Components/SphereComponent.h
+++ b/Engine/Source/Runtime/Engine/Classes/Components/SphereComponent.h
@@ -7,7 +7,7 @@
 /** 
  * A sphere generally used for simple collision. Bounds are rendered as lines in the editor.
  */
-UCLASS(HeaderGroup=Component, ClassGroup=Shapes, editinlinenew, hidecategories=(Object,LOD,Lighting,TextureStreaming), meta=(BlueprintSpawnableComponent))
+UCLASS(ClassGroup=Shapes, editinlinenew, hidecategories=(Object,LOD,Lighting,TextureStreaming), meta=(BlueprintSpawnableComponent))
 class ENGINE_API USphereComponent : public UShapeComponent
 {
 	GENERATED_UCLASS_BODY()
diff --git a/Engine/Source/Runtime/Engine/Classes/Components/SplineComponent.h b/Engine/Source/Runtime/Engine/Classes/Components/SplineComponent.h
--- a/Engine/Source/Runtime/Engine/Classes/Components/SplineComponent.h
+++ b/Engine/Source/Runtime/Engine/Classes/Components/SplineComponent.h
@@ -4,7 +4,7 @@
 #pragma once
 #include "SplineComponent.generated.h"
 
-UCLASS(HeaderGroup=Spline, MinimalAPI)
+UCLASS(MinimalAPI)
 class USplineComponent : public UPrimitiveComponent
 {
 	GENERATED_UCLASS_BODY()
diff --git a/Engine/Source/Runtime/Engine/Classes/Components/SplineMeshComponent.h b/Engine/Source/Runtime/Engine/Classes/Components/SplineMeshComponent.h
--- a/Engine/Source/Runtime/Engine/Classes/Components/SplineMeshComponent.h
+++ b/Engine/Source/Runtime/Engine/Classes/Components/SplineMeshComponent.h
@@ -81,7 +81,7 @@ struct FSplineMeshParams
 
 };
 
-UCLASS(HeaderGroup=Component, MinimalAPI)
+UCLASS(MinimalAPI)
 class USplineMeshComponent : public UStaticMeshComponent, public IInterface_CollisionDataProvider
 {
 	GENERATED_UCLASS_BODY()
diff --git a/Engine/Source/Runtime/Engine/Classes/Components/SpotLightComponent.h b/Engine/Source/Runtime/Engine/Classes/Components/SpotLightComponent.h
--- a/Engine/Source/Runtime/Engine/Classes/Components/SpotLightComponent.h
+++ b/Engine/Source/Runtime/Engine/Classes/Components/SpotLightComponent.h
@@ -7,7 +7,7 @@
 /**
  * A spot light component emits a directional cone shaped light (Eg a Torch).
  */
-UCLASS(HeaderGroup=Light, ClassGroup=Lights, hidecategories=Object, editinlinenew, meta=(BlueprintSpawnableComponent))
+UCLASS(ClassGroup=Lights, hidecategories=Object, editinlinenew, meta=(BlueprintSpawnableComponent))
 class ENGINE_API USpotLightComponent : public UPointLightComponent
 {
 	GENERATED_UCLASS_BODY()
diff --git a/Engine/Source/Runtime/Engine/Classes/Components/StaticMeshComponent.h b/Engine/Source/Runtime/Engine/Classes/Components/StaticMeshComponent.h
--- a/Engine/Source/Runtime/Engine/Classes/Components/StaticMeshComponent.h
+++ b/Engine/Source/Runtime/Engine/Classes/Components/StaticMeshComponent.h
@@ -2,6 +2,8 @@
 
 
 #pragma once
+#include "LightMap.h"
+#include "ShadowMap.h"
 #include "StaticMeshComponent.generated.h"
 
 /** Cached vertex information at the time the mesh was painted. */
@@ -169,7 +171,7 @@ public:
 };
 
 /** A StaticMeshComponent is a mesh that does not animate. */
-UCLASS(HeaderGroup=Component, ClassGroup=(Rendering, Common), hidecategories=(Object,Activation,"Components|Activation"), ShowCategories=(Mobility), dependson=ULightmassPrimitiveSettingsObject, editinlinenew, meta=(BlueprintSpawnableComponent))
+UCLASS(ClassGroup=(Rendering, Common), hidecategories=(Object,Activation,"Components|Activation"), ShowCategories=(Mobility), dependson=ULightmassPrimitiveSettingsObject, editinlinenew, meta=(BlueprintSpawnableComponent))
 class ENGINE_API UStaticMeshComponent : public UMeshComponent
 {
 	GENERATED_UCLASS_BODY()
diff --git a/Engine/Source/Runtime/Engine/Classes/Components/TextRenderComponent.h b/Engine/Source/Runtime/Engine/Classes/Components/TextRenderComponent.h
--- a/Engine/Source/Runtime/Engine/Classes/Components/TextRenderComponent.h
+++ b/Engine/Source/Runtime/Engine/Classes/Components/TextRenderComponent.h
@@ -31,7 +31,7 @@ enum EVerticalTextAligment
 /**
  * Renders text in the world with given font. Contains usual font related attributes such as Scale, Alignment, Color etc.
  */
-UCLASS(HeaderGroup=Component, ClassGroup=Rendering, hidecategories=(Object,LOD,Physics,TextureStreaming,Activation,"Components|Activation",Collision), editinlinenew, meta=(BlueprintSpawnableComponent = ""), MinimalAPI)
+UCLASS(ClassGroup=Rendering, hidecategories=(Object,LOD,Physics,TextureStreaming,Activation,"Components|Activation",Collision), editinlinenew, meta=(BlueprintSpawnableComponent = ""), MinimalAPI)
 class UTextRenderComponent : public UPrimitiveComponent
 {
 	GENERATED_UCLASS_BODY()
diff --git a/Engine/Source/Runtime/Engine/Classes/Components/VectorFieldComponent.h b/Engine/Source/Runtime/Engine/Classes/Components/VectorFieldComponent.h
--- a/Engine/Source/Runtime/Engine/Classes/Components/VectorFieldComponent.h
+++ b/Engine/Source/Runtime/Engine/Classes/Components/VectorFieldComponent.h
@@ -8,7 +8,7 @@
 /** 
  * A Component referencing a vector field.
  */
-UCLASS(HeaderGroup=Component, ClassGroup=Rendering, hidecategories=Object, meta=(BlueprintSpawnableComponent), MinimalAPI)
+UCLASS(ClassGroup=Rendering, hidecategories=Object, meta=(BlueprintSpawnableComponent), MinimalAPI)
 class UVectorFieldComponent : public UPrimitiveComponent
 {
 	GENERATED_UCLASS_BODY()
diff --git a/Engine/Source/Runtime/Engine/Classes/Components/WindDirectionalSourceComponent.h b/Engine/Source/Runtime/Engine/Classes/Components/WindDirectionalSourceComponent.h
--- a/Engine/Source/Runtime/Engine/Classes/Components/WindDirectionalSourceComponent.h
+++ b/Engine/Source/Runtime/Engine/Classes/Components/WindDirectionalSourceComponent.h
@@ -4,7 +4,7 @@
 #pragma once
 #include "WindDirectionalSourceComponent.generated.h"
 
-UCLASS(HeaderGroup=Component, collapsecategories, hidecategories=(Object, Mobility), editinlinenew)
+UCLASS(collapsecategories, hidecategories=(Object, Mobility), editinlinenew)
 class UWindDirectionalSourceComponent : public USceneComponent
 {
 	GENERATED_UCLASS_BODY()
diff --git a/Engine/Source/Runtime/Engine/Classes/Debug/GameplayDebuggingControllerComponent.h b/Engine/Source/Runtime/Engine/Classes/Debug/GameplayDebuggingControllerComponent.h
--- a/Engine/Source/Runtime/Engine/Classes/Debug/GameplayDebuggingControllerComponent.h
+++ b/Engine/Source/Runtime/Engine/Classes/Debug/GameplayDebuggingControllerComponent.h
@@ -7,7 +7,7 @@
 #pragma once
 #include "GameplayDebuggingControllerComponent.generated.h"
 
-UCLASS(config=Engine, HeaderGroup=Component)
+UCLASS(config=Engine)
 class ENGINE_API UGameplayDebuggingControllerComponent : public UActorComponent
 {
 	GENERATED_UCLASS_BODY()
diff --git a/Engine/Source/Runtime/Engine/Classes/Distributions/DistributionFloatParticleParameter.h b/Engine/Source/Runtime/Engine/Classes/Distributions/DistributionFloatParticleParameter.h
--- a/Engine/Source/Runtime/Engine/Classes/Distributions/DistributionFloatParticleParameter.h
+++ b/Engine/Source/Runtime/Engine/Classes/Distributions/DistributionFloatParticleParameter.h
@@ -4,7 +4,7 @@
 #pragma once
 #include "DistributionFloatParticleParameter.generated.h"
 
-UCLASS(HeaderGroup=Particle, collapsecategories, hidecategories=Object, editinlinenew, MinimalAPI)
+UCLASS(collapsecategories, hidecategories=Object, editinlinenew, MinimalAPI)
 class UDistributionFloatParticleParameter : public UDistributionFloatParameterBase
 {
 	GENERATED_UCLASS_BODY()
diff --git a/Engine/Source/Runtime/Engine/Classes/Distributions/DistributionVectorParticleParameter.h b/Engine/Source/Runtime/Engine/Classes/Distributions/DistributionVectorParticleParameter.h
--- a/Engine/Source/Runtime/Engine/Classes/Distributions/DistributionVectorParticleParameter.h
+++ b/Engine/Source/Runtime/Engine/Classes/Distributions/DistributionVectorParticleParameter.h
@@ -4,7 +4,7 @@
 #pragma once
 #include "DistributionVectorParticleParameter.generated.h"
 
-UCLASS(HeaderGroup=Particle, collapsecategories, hidecategories=Object, editinlinenew, MinimalAPI)
+UCLASS(collapsecategories, hidecategories=Object, editinlinenew, MinimalAPI)
 class UDistributionVectorParticleParameter : public UDistributionVectorParameterBase
 {
 	GENERATED_UCLASS_BODY()
diff --git a/Engine/Source/Runtime/Engine/Classes/Engine/BoxReflectionCapture.h b/Engine/Source/Runtime/Engine/Classes/Engine/BoxReflectionCapture.h
--- a/Engine/Source/Runtime/Engine/Classes/Engine/BoxReflectionCapture.h
+++ b/Engine/Source/Runtime/Engine/Classes/Engine/BoxReflectionCapture.h
@@ -8,7 +8,7 @@
 #pragma once
 #include "BoxReflectionCapture.generated.h"
 
-UCLASS(hidecategories=(Collision, Attachment, Actor), HeaderGroup=Decal, MinimalAPI)
+UCLASS(hidecategories=(Collision, Attachment, Actor), MinimalAPI)
 class ABoxReflectionCapture : public AReflectionCapture
 {
 	GENERATED_UCLASS_BODY()
diff --git a/Engine/Source/Runtime/Engine/Classes/Engine/ChildConnection.h b/Engine/Source/Runtime/Engine/Classes/Engine/ChildConnection.h
--- a/Engine/Source/Runtime/Engine/Classes/Engine/ChildConnection.h
+++ b/Engine/Source/Runtime/Engine/Classes/Engine/ChildConnection.h
@@ -5,7 +5,7 @@
 #pragma once
 #include "ChildConnection.generated.h"
 
-UCLASS(HeaderGroup=Network, MinimalAPI,transient,config=Engine)
+UCLASS(MinimalAPI,transient,config=Engine)
 class UChildConnection : public UNetConnection
 {
 	GENERATED_UCLASS_BODY()
diff --git a/Engine/Source/Runtime/Engine/Classes/Engine/CloudStorageBase.h b/Engine/Source/Runtime/Engine/Classes/Engine/CloudStorageBase.h
--- a/Engine/Source/Runtime/Engine/Classes/Engine/CloudStorageBase.h
+++ b/Engine/Source/Runtime/Engine/Classes/Engine/CloudStorageBase.h
@@ -29,7 +29,7 @@ enum ECloudStorageDelegate
 };
 
 	// Full API exported so that platform modules can extend it (e.g. FacebookWindows in WinDrv)
-UCLASS(HeaderGroup=PlatformInterface)
+UCLASS()
 class UCloudStorageBase : public UPlatformInterfaceBase
 {
 	GENERATED_UCLASS_BODY()
diff --git a/Engine/Source/Runtime/Engine/Classes/Engine/Console.h b/Engine/Source/Runtime/Engine/Classes/Engine/Console.h
--- a/Engine/Source/Runtime/Engine/Classes/Engine/Console.h
+++ b/Engine/Source/Runtime/Engine/Classes/Engine/Console.h
@@ -59,7 +59,7 @@ struct FAutoCompleteNode
 	
 };
 
-UCLASS(Within=GameViewportClient, transient, config=Input, HeaderGroup=UserInterface)
+UCLASS(Within=GameViewportClient, transient, config=Input)
 class ENGINE_API UConsole : public UObject, public FOutputDevice
 {
 	GENERATED_UCLASS_BODY()
diff --git a/Engine/Source/Runtime/Engine/Classes/Engine/DecalActor.h b/Engine/Source/Runtime/Engine/Classes/Engine/DecalActor.h
--- a/Engine/Source/Runtime/Engine/Classes/Engine/DecalActor.h
+++ b/Engine/Source/Runtime/Engine/Classes/Engine/DecalActor.h
@@ -8,7 +8,12 @@
 #pragma once
 #include "DecalActor.generated.h"
 
-UCLASS(hidecategories=(Collision, Attachment, Actor, Input, Replication), HeaderGroup=Decal, MinimalAPI)
+//
+// Forward declarations.
+//
+class UBoxComponent;
+
+UCLASS(hidecategories=(Collision, Attachment, Actor, Input, Replication), MinimalAPI)
 class ADecalActor : public AActor
 {
 	GENERATED_UCLASS_BODY()
diff --git a/Engine/Source/Runtime/Engine/Classes/Engine/Engine.h b/Engine/Source/Runtime/Engine/Classes/Engine/Engine.h
--- a/Engine/Source/Runtime/Engine/Classes/Engine/Engine.h
+++ b/Engine/Source/Runtime/Engine/Classes/Engine/Engine.h
@@ -462,7 +462,7 @@ namespace EMatineeCaptureType
 
 class IAnalyticsProvider;
 
-UCLASS(HeaderGroup=GameEngine, abstract, config=Engine, transient)
+UCLASS(abstract, config=Engine, transient)
 class ENGINE_API UEngine : public UObject, public FExec
 {
 	GENERATED_UCLASS_BODY()
diff --git a/Engine/Source/Runtime/Engine/Classes/Engine/EngineBaseTypes.h b/Engine/Source/Runtime/Engine/Classes/Engine/EngineBaseTypes.h
--- a/Engine/Source/Runtime/Engine/Classes/Engine/EngineBaseTypes.h
+++ b/Engine/Source/Runtime/Engine/Classes/Engine/EngineBaseTypes.h
@@ -545,7 +545,7 @@ struct FExposureSettings
 };
 
 
-UCLASS(HeaderGroup=Base, abstract, config=Engine)
+UCLASS(abstract, config=Engine)
 class UEngineBaseTypes : public UObject
 {
 	GENERATED_UCLASS_BODY()
diff --git a/Engine/Source/Runtime/Engine/Classes/Engine/GameEngine.h b/Engine/Source/Runtime/Engine/Classes/Engine/GameEngine.h
--- a/Engine/Source/Runtime/Engine/Classes/Engine/GameEngine.h
+++ b/Engine/Source/Runtime/Engine/Classes/Engine/GameEngine.h
@@ -7,7 +7,7 @@
 #pragma once
 #include "GameEngine.generated.h"
 
-UCLASS(HeaderGroup=GameEngine, config=Engine, transient)
+UCLASS(config=Engine, transient)
 class ENGINE_API UGameEngine : public UEngine
 {
 	GENERATED_UCLASS_BODY()
diff --git a/Engine/Source/Runtime/Engine/Classes/Engine/GeneratedMeshAreaLight.h b/Engine/Source/Runtime/Engine/Classes/Engine/GeneratedMeshAreaLight.h
--- a/Engine/Source/Runtime/Engine/Classes/Engine/GeneratedMeshAreaLight.h
+++ b/Engine/Source/Runtime/Engine/Classes/Engine/GeneratedMeshAreaLight.h
@@ -8,7 +8,7 @@
 
 #include "GeneratedMeshAreaLight.generated.h"
 
-UCLASS(HeaderGroup=Light, notplaceable, MinimalAPI)
+UCLASS(notplaceable, MinimalAPI)
 class AGeneratedMeshAreaLight : public ASpotLight
 {
 	GENERATED_UCLASS_BODY()
diff --git a/Engine/Source/Runtime/Engine/Classes/Engine/InGameAdManager.h b/Engine/Source/Runtime/Engine/Classes/Engine/InGameAdManager.h
--- a/Engine/Source/Runtime/Engine/Classes/Engine/InGameAdManager.h
+++ b/Engine/Source/Runtime/Engine/Classes/Engine/InGameAdManager.h
@@ -31,7 +31,7 @@ DECLARE_DYNAMIC_DELEGATE( FOnUserClickedBanner );
 DECLARE_DYNAMIC_DELEGATE( FOnUserClosedAdvertisement );
 
 
-UCLASS(HeaderGroup=PlatformInterface)
+UCLASS()
 class UInGameAdManager : public UPlatformInterfaceBase
 {
 	GENERATED_UCLASS_BODY()
diff --git a/Engine/Source/Runtime/Engine/Classes/Engine/InteractiveFoliageActor.h b/Engine/Source/Runtime/Engine/Classes/Engine/InteractiveFoliageActor.h
--- a/Engine/Source/Runtime/Engine/Classes/Engine/InteractiveFoliageActor.h
+++ b/Engine/Source/Runtime/Engine/Classes/Engine/InteractiveFoliageActor.h
@@ -6,7 +6,7 @@
 #pragma once
 #include "InteractiveFoliageActor.generated.h"
 
-UCLASS(HeaderGroup=Foliage, MinimalAPI)
+UCLASS(MinimalAPI)
 class AInteractiveFoliageActor : public AStaticMeshActor
 {
 	GENERATED_UCLASS_BODY()
diff --git a/Engine/Source/Runtime/Engine/Classes/Engine/LevelScriptActor.h b/Engine/Source/Runtime/Engine/Classes/Engine/LevelScriptActor.h
--- a/Engine/Source/Runtime/Engine/Classes/Engine/LevelScriptActor.h
+++ b/Engine/Source/Runtime/Engine/Classes/Engine/LevelScriptActor.h
@@ -4,7 +4,7 @@
 
 #include "LevelScriptActor.generated.h"
 
-UCLASS(HeaderGroup=LevelScript, notplaceable, meta=(KismetHideOverrides = "ReceiveAnyDamage,ReceivePointDamage,ReceiveRadialDamage,ReceiveActorBeginOverlap,ReceiveActorEndOverlap,ReceiveHit,ReceiveDestroyed,ReceiveActorBeginCursorOver,ReceiveActorEndCursorOver,ReceiveActorOnClicked,ReceiveActorOnReleased,ReceiveActorOnInputTouchBegin,ReceiveActorOnInputTouchEnd,ReceiveActorOnInputTouchEnter,ReceiveActorOnInputTouchLeave"), HideCategories=(Collision,Rendering,"Utilities|Orientation"))
+UCLASS(notplaceable, meta=(KismetHideOverrides = "ReceiveAnyDamage,ReceivePointDamage,ReceiveRadialDamage,ReceiveActorBeginOverlap,ReceiveActorEndOverlap,ReceiveHit,ReceiveDestroyed,ReceiveActorBeginCursorOver,ReceiveActorEndCursorOver,ReceiveActorOnClicked,ReceiveActorOnReleased,ReceiveActorOnInputTouchBegin,ReceiveActorOnInputTouchEnd,ReceiveActorOnInputTouchEnter,ReceiveActorOnInputTouchLeave"), HideCategories=(Collision,Rendering,"Utilities|Orientation"))
 class ENGINE_API ALevelScriptActor : public AActor
 {
 	GENERATED_UCLASS_BODY()
diff --git a/Engine/Source/Runtime/Engine/Classes/Engine/Light.h b/Engine/Source/Runtime/Engine/Classes/Engine/Light.h
--- a/Engine/Source/Runtime/Engine/Classes/Engine/Light.h
+++ b/Engine/Source/Runtime/Engine/Classes/Engine/Light.h
@@ -4,7 +4,7 @@
 
 #include "Light.generated.h"
 
-UCLASS(Abstract, HeaderGroup=Light, ClassGroup=Lights, hidecategories=(Input,Collision,Replication), ConversionRoot, meta=(ChildCanTick))
+UCLASS(Abstract, ClassGroup=Lights, hidecategories=(Input,Collision,Replication), ConversionRoot, meta=(ChildCanTick))
 class ENGINE_API ALight : public AActor
 {
 	GENERATED_UCLASS_BODY()
diff --git a/Engine/Source/Runtime/Engine/Classes/Engine/LocalPlayer.h b/Engine/Source/Runtime/Engine/Classes/Engine/LocalPlayer.h
--- a/Engine/Source/Runtime/Engine/Classes/Engine/LocalPlayer.h
+++ b/Engine/Source/Runtime/Engine/Classes/Engine/LocalPlayer.h
@@ -46,7 +46,7 @@ private:
  *	There may be several spawned in the case of splitscreen/coop
  *	There may be 0 spawned on servers
  */
-UCLASS(HeaderGroup=Network, Within=Engine, config=Engine, transient)
+UCLASS(Within=Engine, config=Engine, transient)
 class ENGINE_API ULocalPlayer : public UPlayer
 {
 	GENERATED_UCLASS_BODY()
diff --git a/Engine/Source/Runtime/Engine/Classes/Engine/NetConnection.h b/Engine/Source/Runtime/Engine/Classes/Engine/NetConnection.h
--- a/Engine/Source/Runtime/Engine/Classes/Engine/NetConnection.h
+++ b/Engine/Source/Runtime/Engine/Classes/Engine/NetConnection.h
@@ -91,7 +91,7 @@ struct DelayedPacket
 #define PING_ACK_DELAY 0.5
 
 
-UCLASS(HeaderGroup=Network, customConstructor, Abstract, MinimalAPI, transient, config=Engine, dependsOn=UGameEngine)
+UCLASS(customConstructor, Abstract, MinimalAPI, transient, config=Engine, dependsOn=UGameEngine)
 class UNetConnection : public UPlayer
 {
 	GENERATED_UCLASS_BODY()
diff --git a/Engine/Source/Runtime/Engine/Classes/Engine/NetDriver.h b/Engine/Source/Runtime/Engine/Classes/Engine/NetDriver.h
--- a/Engine/Source/Runtime/Engine/Classes/Engine/NetDriver.h
+++ b/Engine/Source/Runtime/Engine/Classes/Engine/NetDriver.h
@@ -87,7 +87,7 @@ struct FActorDestructionInfo
 };
 
 
-UCLASS(HeaderGroup=Network, dependson=UEngineTypes, Abstract, customConstructor, transient, MinimalAPI, config=Engine)
+UCLASS(dependson=UEngineTypes, Abstract, customConstructor, transient, MinimalAPI, config=Engine)
 class UNetDriver : public UObject, public FExec
 {
 	GENERATED_UCLASS_BODY()
diff --git a/Engine/Source/Runtime/Engine/Classes/Engine/PlaneReflectionCapture.h b/Engine/Source/Runtime/Engine/Classes/Engine/PlaneReflectionCapture.h
--- a/Engine/Source/Runtime/Engine/Classes/Engine/PlaneReflectionCapture.h
+++ b/Engine/Source/Runtime/Engine/Classes/Engine/PlaneReflectionCapture.h
@@ -7,7 +7,7 @@
 #pragma once
 #include "PlaneReflectionCapture.generated.h"
 
-UCLASS(abstract, hidecategories=(Collision, Attachment, Actor), HeaderGroup=Decal, MinimalAPI)
+UCLASS(abstract, hidecategories=(Collision, Attachment, Actor), MinimalAPI)
 class APlaneReflectionCapture : public AReflectionCapture
 {
 	GENERATED_UCLASS_BODY()
diff --git a/Engine/Source/Runtime/Engine/Classes/Engine/PlatformInterfaceBase.h b/Engine/Source/Runtime/Engine/Classes/Engine/PlatformInterfaceBase.h
--- a/Engine/Source/Runtime/Engine/Classes/Engine/PlatformInterfaceBase.h
+++ b/Engine/Source/Runtime/Engine/Classes/Engine/PlatformInterfaceBase.h
@@ -107,7 +107,7 @@ struct FDelegateArray
 
 };
 
-UCLASS(HeaderGroup=PlatformInterface, transient,MinimalAPI)
+UCLASS(transient,MinimalAPI)
 class UPlatformInterfaceBase : public UObject
 {
 	GENERATED_UCLASS_BODY()
diff --git a/Engine/Source/Runtime/Engine/Classes/Engine/PlatformInterfaceWebResponse.h b/Engine/Source/Runtime/Engine/Classes/Engine/PlatformInterfaceWebResponse.h
--- a/Engine/Source/Runtime/Engine/Classes/Engine/PlatformInterfaceWebResponse.h
+++ b/Engine/Source/Runtime/Engine/Classes/Engine/PlatformInterfaceWebResponse.h
@@ -10,7 +10,7 @@
 #pragma once
 #include "PlatformInterfaceWebResponse.generated.h"
 
-UCLASS(HeaderGroup=PlatformInterface, transient,MinimalAPI)
+UCLASS(transient,MinimalAPI)
 class UPlatformInterfaceWebResponse : public UObject
 {
 	GENERATED_UCLASS_BODY()
diff --git a/Engine/Source/Runtime/Engine/Classes/Engine/Player.h b/Engine/Source/Runtime/Engine/Classes/Engine/Player.h
--- a/Engine/Source/Runtime/Engine/Classes/Engine/Player.h
+++ b/Engine/Source/Runtime/Engine/Classes/Engine/Player.h
@@ -7,7 +7,7 @@
 #pragma once
 #include "Player.generated.h"
 
-UCLASS(HeaderGroup=Network, transient, config=Engine)
+UCLASS(transient, config=Engine)
 class UPlayer : public UObject, public FExec
 {
 	GENERATED_UCLASS_BODY()
diff --git a/Engine/Source/Runtime/Engine/Classes/Engine/ReflectionCapture.h b/Engine/Source/Runtime/Engine/Classes/Engine/ReflectionCapture.h
--- a/Engine/Source/Runtime/Engine/Classes/Engine/ReflectionCapture.h
+++ b/Engine/Source/Runtime/Engine/Classes/Engine/ReflectionCapture.h
@@ -3,7 +3,7 @@
 #pragma once
 #include "ReflectionCapture.generated.h"
 
-UCLASS(abstract, hidecategories=(Collision, Attachment, Actor), HeaderGroup=Decal, MinimalAPI)
+UCLASS(abstract, hidecategories=(Collision, Attachment, Actor), MinimalAPI)
 class AReflectionCapture : public AActor
 {
 	GENERATED_UCLASS_BODY()
diff --git a/Engine/Source/Runtime/Engine/Classes/Engine/Scene.h b/Engine/Source/Runtime/Engine/Classes/Engine/Scene.h
--- a/Engine/Source/Runtime/Engine/Classes/Engine/Scene.h
+++ b/Engine/Source/Runtime/Engine/Classes/Engine/Scene.h
@@ -895,7 +895,7 @@ struct FPostProcessSettings
 	}
 };
 
-UCLASS(HeaderGroup=Scene)
+UCLASS()
 class UScene : public UObject
 {
 	GENERATED_UCLASS_BODY()
diff --git a/Engine/Source/Runtime/Engine/Classes/Engine/SceneCapture2D.h b/Engine/Source/Runtime/Engine/Classes/Engine/SceneCapture2D.h
--- a/Engine/Source/Runtime/Engine/Classes/Engine/SceneCapture2D.h
+++ b/Engine/Source/Runtime/Engine/Classes/Engine/SceneCapture2D.h
@@ -8,7 +8,7 @@
 #pragma once
 #include "SceneCapture2D.generated.h"
 
-UCLASS(hidecategories=(Collision, Material, Attachment, Actor), HeaderGroup=Decal, MinimalAPI)
+UCLASS(hidecategories=(Collision, Material, Attachment, Actor), MinimalAPI)
 class ASceneCapture2D : public ASceneCapture
 {
 	GENERATED_UCLASS_BODY()
diff --git a/Engine/Source/Runtime/Engine/Classes/Engine/SceneCaptureCube.h b/Engine/Source/Runtime/Engine/Classes/Engine/SceneCaptureCube.h
--- a/Engine/Source/Runtime/Engine/Classes/Engine/SceneCaptureCube.h
+++ b/Engine/Source/Runtime/Engine/Classes/Engine/SceneCaptureCube.h
@@ -8,7 +8,7 @@
 #pragma once
 #include "SceneCaptureCube.generated.h"
 
-UCLASS(hidecategories = (Collision, Material, Attachment, Actor), HeaderGroup = Decal, MinimalAPI)
+UCLASS(hidecategories = (Collision, Material, Attachment, Actor), MinimalAPI)
 class ASceneCaptureCube : public ASceneCapture
 {
 	GENERATED_UCLASS_BODY()
diff --git a/Engine/Source/Runtime/Engine/Classes/Engine/SkeletalMesh.h b/Engine/Source/Runtime/Engine/Classes/Engine/SkeletalMesh.h
--- a/Engine/Source/Runtime/Engine/Classes/Engine/SkeletalMesh.h
+++ b/Engine/Source/Runtime/Engine/Classes/Engine/SkeletalMesh.h
@@ -361,7 +361,7 @@ struct FSkeletalMaterial
 
 class FSkeletalMeshResource;
 
-UCLASS(HeaderGroup=SkeletalMesh, hidecategories=Object, MinimalAPI, BlueprintType)
+UCLASS(hidecategories=Object, MinimalAPI, BlueprintType)
 class USkeletalMesh : public UObject
 {
 	GENERATED_UCLASS_BODY()
diff --git a/Engine/Source/Runtime/Engine/Classes/Engine/SkeletalMeshSocket.h b/Engine/Source/Runtime/Engine/Classes/Engine/SkeletalMeshSocket.h
--- a/Engine/Source/Runtime/Engine/Classes/Engine/SkeletalMeshSocket.h
+++ b/Engine/Source/Runtime/Engine/Classes/Engine/SkeletalMeshSocket.h
@@ -4,7 +4,7 @@
 #pragma once
 #include "SkeletalMeshSocket.generated.h"
 
-UCLASS(HeaderGroup=SkeletalMesh, hidecategories=Object, hidecategories=Actor, MinimalAPI)
+UCLASS(hidecategories=Object, hidecategories=Actor, MinimalAPI)
 class USkeletalMeshSocket : public UObject
 {
 	GENERATED_UCLASS_BODY()
diff --git a/Engine/Source/Runtime/Engine/Classes/Engine/SkyLight.h b/Engine/Source/Runtime/Engine/Classes/Engine/SkyLight.h
--- a/Engine/Source/Runtime/Engine/Classes/Engine/SkyLight.h
+++ b/Engine/Source/Runtime/Engine/Classes/Engine/SkyLight.h
@@ -4,7 +4,7 @@
 
 #include "SkyLight.generated.h"
 
-UCLASS(HeaderGroup=Light, ClassGroup=Lights, hidecategories=(Input,Collision,Replication,Info), ConversionRoot)
+UCLASS(ClassGroup=Lights, hidecategories=(Input,Collision,Replication,Info), ConversionRoot)
 class ENGINE_API ASkyLight : public AInfo
 {
 	GENERATED_UCLASS_BODY()
diff --git a/Engine/Source/Runtime/Engine/Classes/Engine/SphereReflectionCapture.h b/Engine/Source/Runtime/Engine/Classes/Engine/SphereReflectionCapture.h
--- a/Engine/Source/Runtime/Engine/Classes/Engine/SphereReflectionCapture.h
+++ b/Engine/Source/Runtime/Engine/Classes/Engine/SphereReflectionCapture.h
@@ -8,7 +8,11 @@
 #pragma once
 #include "SphereReflectionCapture.generated.h"
 
-UCLASS(hidecategories=(Collision, Attachment, Actor), HeaderGroup=Decal, MinimalAPI)
+//
+// Forward declarations.
+//
+class UDrawSphereComponent;
+UCLASS(hidecategories=(Collision, Attachment, Actor), MinimalAPI)
 class ASphereReflectionCapture : public AReflectionCapture
 {
 	GENERATED_UCLASS_BODY()
diff --git a/Engine/Source/Runtime/Engine/Classes/Engine/StaticMesh.h b/Engine/Source/Runtime/Engine/Classes/Engine/StaticMesh.h
--- a/Engine/Source/Runtime/Engine/Classes/Engine/StaticMesh.h
+++ b/Engine/Source/Runtime/Engine/Classes/Engine/StaticMesh.h
@@ -303,7 +303,7 @@ struct FAssetEditorOrbitCameraPosition
 	FRotator CamOrbitRotation;
 };
 
-UCLASS(HeaderGroup=StaticMesh, collapsecategories, hidecategories=Object, customconstructor, MinimalAPI, BlueprintType, config=Engine)
+UCLASS(collapsecategories, hidecategories=Object, customconstructor, MinimalAPI, BlueprintType, config=Engine)
 class UStaticMesh : public UObject, public IInterface_CollisionDataProvider, public IInterface_AssetUserData
 {
 	GENERATED_UCLASS_BODY()
diff --git a/Engine/Source/Runtime/Engine/Classes/Engine/StaticMeshSocket.h b/Engine/Source/Runtime/Engine/Classes/Engine/StaticMeshSocket.h
--- a/Engine/Source/Runtime/Engine/Classes/Engine/StaticMeshSocket.h
+++ b/Engine/Source/Runtime/Engine/Classes/Engine/StaticMeshSocket.h
@@ -4,7 +4,7 @@
 #pragma once
 #include "StaticMeshSocket.generated.h"
 
-UCLASS(HeaderGroup=StaticMesh, hidecategories=Object, hidecategories=Actor, MinimalAPI)
+UCLASS(hidecategories=Object, hidecategories=Actor, MinimalAPI)
 class UStaticMeshSocket : public UObject
 {
 	GENERATED_UCLASS_BODY()
diff --git a/Engine/Source/Runtime/Engine/Classes/Engine/TextRenderActor.h b/Engine/Source/Runtime/Engine/Classes/Engine/TextRenderActor.h
--- a/Engine/Source/Runtime/Engine/Classes/Engine/TextRenderActor.h
+++ b/Engine/Source/Runtime/Engine/Classes/Engine/TextRenderActor.h
@@ -8,7 +8,7 @@
 
 #include "TextRenderActor.generated.h"
 
-UCLASS(MinimalAPI, hidecategories=(Collision, Attachment, Actor), HeaderGroup=Decal)
+UCLASS(MinimalAPI, hidecategories=(Collision, Attachment, Actor))
 class ATextRenderActor : public AActor
 {
 	GENERATED_UCLASS_BODY()
diff --git a/Engine/Source/Runtime/Engine/Classes/Engine/TwitterIntegrationBase.h b/Engine/Source/Runtime/Engine/Classes/Engine/TwitterIntegrationBase.h
--- a/Engine/Source/Runtime/Engine/Classes/Engine/TwitterIntegrationBase.h
+++ b/Engine/Source/Runtime/Engine/Classes/Engine/TwitterIntegrationBase.h
@@ -28,7 +28,7 @@ enum ETwitterIntegrationDelegate
 	TID_MAX,
 };
 
-UCLASS(HeaderGroup=PlatformInterface, config=Engine)
+UCLASS(config=Engine)
 class UTwitterIntegrationBase : public UPlatformInterfaceBase
 {
 	GENERATED_UCLASS_BODY()
diff --git a/Engine/Source/Runtime/Engine/Classes/Foliage/InstancedFoliageActor.h b/Engine/Source/Runtime/Engine/Classes/Foliage/InstancedFoliageActor.h
--- a/Engine/Source/Runtime/Engine/Classes/Foliage/InstancedFoliageActor.h
+++ b/Engine/Source/Runtime/Engine/Classes/Foliage/InstancedFoliageActor.h
@@ -2,9 +2,10 @@
 
 
 #pragma once
+#include "InstancedFoliage.h"
 #include "InstancedFoliageActor.generated.h"
 
-UCLASS(notplaceable, HeaderGroup = Foliage, hidecategories = Object, MinimalAPI, NotBlueprintable)
+UCLASS(notplaceable, hidecategories = Object, MinimalAPI, NotBlueprintable)
 class AInstancedFoliageActor : public AActor
 {
 	GENERATED_UCLASS_BODY()
diff --git a/Engine/Source/Runtime/Engine/Classes/Foliage/InstancedFoliageSettings.h b/Engine/Source/Runtime/Engine/Classes/Foliage/InstancedFoliageSettings.h
--- a/Engine/Source/Runtime/Engine/Classes/Foliage/InstancedFoliageSettings.h
+++ b/Engine/Source/Runtime/Engine/Classes/Foliage/InstancedFoliageSettings.h
@@ -15,7 +15,7 @@ enum FoliageVertexColorMask
 };
 
 
-UCLASS(HeaderGroup=Foliage, hidecategories=Object, editinlinenew, MinimalAPI)
+UCLASS(hidecategories=Object, editinlinenew, MinimalAPI)
 class UInstancedFoliageSettings : public UObject
 {
 	GENERATED_UCLASS_BODY()
diff --git a/Engine/Source/Runtime/Engine/Classes/GameFramework/Character.h b/Engine/Source/Runtime/Engine/Classes/GameFramework/Character.h
--- a/Engine/Source/Runtime/Engine/Classes/GameFramework/Character.h
+++ b/Engine/Source/Runtime/Engine/Classes/GameFramework/Character.h
@@ -4,6 +4,10 @@
 
 #include "Character.generated.h"
 
+//
+// Forward declarations
+//
+struct FAnimMontageInstance;
 /** Replicated data when playing a root motion montage. */
 USTRUCT()
 struct FRepRootMotionMontage
@@ -130,7 +134,7 @@ struct FRepRelativeMovement
 // to a CharacterMovementComponent that handles movement of the collision capsule, and they
 // also have implementations of basic networking and input models.
 //=============================================================================
-UCLASS(abstract, HeaderGroup=Pawn, config=Game, dependson=(AController, UCharacterMovementComponent), BlueprintType)
+UCLASS(abstract, config=Game, dependson=(AController, UCharacterMovementComponent), BlueprintType)
 class ENGINE_API ACharacter : public APawn
 {
 	GENERATED_UCLASS_BODY()
diff --git a/Engine/Source/Runtime/Engine/Classes/GameFramework/CharacterMovementComponent.h b/Engine/Source/Runtime/Engine/Classes/GameFramework/CharacterMovementComponent.h
--- a/Engine/Source/Runtime/Engine/Classes/GameFramework/CharacterMovementComponent.h
+++ b/Engine/Source/Runtime/Engine/Classes/GameFramework/CharacterMovementComponent.h
@@ -137,7 +137,7 @@ struct FRootMotionMovementParams
 };
 
 
-UCLASS(HeaderGroup=Component, dependson=(UNetworkPredictionInterface))
+UCLASS(dependson=(UNetworkPredictionInterface))
 class ENGINE_API UCharacterMovementComponent : public UPawnMovementComponent, public INetworkPredictionInterface
 {
 	GENERATED_UCLASS_BODY()
diff --git a/Engine/Source/Runtime/Engine/Classes/GameFramework/DefaultPawn.h b/Engine/Source/Runtime/Engine/Classes/GameFramework/DefaultPawn.h
--- a/Engine/Source/Runtime/Engine/Classes/GameFramework/DefaultPawn.h
+++ b/Engine/Source/Runtime/Engine/Classes/GameFramework/DefaultPawn.h
@@ -6,9 +6,10 @@
 //=============================================================================
 
 #pragma once
+#include "FloatingPawnMovement.h"
 #include "DefaultPawn.generated.h"
 
-UCLASS(HeaderGroup=Pawn, config=Game, Blueprintable, BlueprintType)
+UCLASS(config=Game, Blueprintable, BlueprintType)
 class ENGINE_API ADefaultPawn : public APawn
 {
 	GENERATED_UCLASS_BODY()
diff --git a/Engine/Source/Runtime/Engine/Classes/GameFramework/ForceFeedbackEffect.h b/Engine/Source/Runtime/Engine/Classes/GameFramework/ForceFeedbackEffect.h
--- a/Engine/Source/Runtime/Engine/Classes/GameFramework/ForceFeedbackEffect.h
+++ b/Engine/Source/Runtime/Engine/Classes/GameFramework/ForceFeedbackEffect.h
@@ -70,7 +70,7 @@ struct FActiveForceFeedbackEffect
 /**
  * A predefined force-feedback effect to be played on a controller
  */
-UCLASS(BlueprintType, headergroup=UserInterface, MinimalAPI)
+UCLASS(BlueprintType, MinimalAPI)
 class UForceFeedbackEffect : public UObject
 {
 	GENERATED_UCLASS_BODY()
diff --git a/Engine/Source/Runtime/Engine/Classes/GameFramework/GameState.h b/Engine/Source/Runtime/Engine/Classes/GameFramework/GameState.h
--- a/Engine/Source/Runtime/Engine/Classes/GameFramework/GameState.h
+++ b/Engine/Source/Runtime/Engine/Classes/GameFramework/GameState.h
@@ -8,7 +8,7 @@
 // GameState is replicated and is valid on servers and clients.
 //=============================================================================
 
-UCLASS(config=Game, HeaderGroup=ReplicationInfo, notplaceable, BlueprintType, Blueprintable)
+UCLASS(config=Game, notplaceable, BlueprintType, Blueprintable)
 class ENGINE_API AGameState : public AInfo
 {
 	GENERATED_UCLASS_BODY()
diff --git a/Engine/Source/Runtime/Engine/Classes/GameFramework/InputSettings.h b/Engine/Source/Runtime/Engine/Classes/GameFramework/InputSettings.h
--- a/Engine/Source/Runtime/Engine/Classes/GameFramework/InputSettings.h
+++ b/Engine/Source/Runtime/Engine/Classes/GameFramework/InputSettings.h
@@ -1,10 +1,10 @@
 // Copyright 1998-2014 Epic Games, Inc. All Rights Reserved.
 
 #pragma once
-
+#include "PlayerInput.h"
 #include "InputSettings.generated.h"
 
-UCLASS(config=Input, HeaderGroup=UserInterface)
+UCLASS(config=Input)
 class ENGINE_API UInputSettings : public UObject
 {
 	GENERATED_UCLASS_BODY()
diff --git a/Engine/Source/Runtime/Engine/Classes/GameFramework/MovementComponent.h b/Engine/Source/Runtime/Engine/Classes/GameFramework/MovementComponent.h
--- a/Engine/Source/Runtime/Engine/Classes/GameFramework/MovementComponent.h
+++ b/Engine/Source/Runtime/Engine/Classes/GameFramework/MovementComponent.h
@@ -5,10 +5,10 @@
  */
 
 #pragma once
+#include "Components/SceneComponent.h"
 #include "MovementComponent.generated.h"
 
-
-UCLASS(HeaderGroup=Component, ClassGroup=Movement, abstract, BlueprintType)
+UCLASS(ClassGroup=Movement, abstract, BlueprintType)
 class ENGINE_API UMovementComponent : public UActorComponent
 {
 	GENERATED_UCLASS_BODY()
diff --git a/Engine/Source/Runtime/Engine/Classes/GameFramework/NavMovementComponent.h b/Engine/Source/Runtime/Engine/Classes/GameFramework/NavMovementComponent.h
--- a/Engine/Source/Runtime/Engine/Classes/GameFramework/NavMovementComponent.h
+++ b/Engine/Source/Runtime/Engine/Classes/GameFramework/NavMovementComponent.h
@@ -8,7 +8,7 @@
 #include "AI/Navigation/NavAgentInterface.h"
 #include "NavMovementComponent.generated.h"
 
-UCLASS(HeaderGroup=Component, abstract, dependson=(UNavigationSystem, UPrimitiveComponent, INavAgentInterface))
+UCLASS(abstract, dependson=(UNavigationSystem, UPrimitiveComponent, INavAgentInterface))
 class ENGINE_API UNavMovementComponent : public UMovementComponent
 {
 	GENERATED_UCLASS_BODY()
diff --git a/Engine/Source/Runtime/Engine/Classes/GameFramework/Pawn.h b/Engine/Source/Runtime/Engine/Classes/GameFramework/Pawn.h
--- a/Engine/Source/Runtime/Engine/Classes/GameFramework/Pawn.h
+++ b/Engine/Source/Runtime/Engine/Classes/GameFramework/Pawn.h
@@ -1,7 +1,7 @@
 // Copyright 1998-2014 Epic Games, Inc. All Rights Reserved.
 
 #pragma once
-
+#include "PawnMovementComponent.h"
 #include "Pawn.generated.h"
 
 ENGINE_API DECLARE_LOG_CATEGORY_EXTERN(LogDamage, Warning, All);
@@ -11,7 +11,7 @@ ENGINE_API DECLARE_LOG_CATEGORY_EXTERN(LogDamage, Warning, All);
  *	Pawns are the physical representations of players and creatures in a level.
  *	They are responsible for all physical interaction between the player or AI and the world.
  */
-UCLASS(abstract, HeaderGroup=Pawn, config=Game, dependson=AController, BlueprintType, Blueprintable, hidecategories=(Navigation, "AI|Navigation"))
+UCLASS(abstract, config=Game, dependson=AController, BlueprintType, Blueprintable, hidecategories=(Navigation, "AI|Navigation"))
 class ENGINE_API APawn : public AActor, public INavAgentInterface
 {
 	GENERATED_UCLASS_BODY()
diff --git a/Engine/Source/Runtime/Engine/Classes/GameFramework/PawnMovementComponent.h b/Engine/Source/Runtime/Engine/Classes/GameFramework/PawnMovementComponent.h
--- a/Engine/Source/Runtime/Engine/Classes/GameFramework/PawnMovementComponent.h
+++ b/Engine/Source/Runtime/Engine/Classes/GameFramework/PawnMovementComponent.h
@@ -7,7 +7,7 @@
 #pragma once
 #include "PawnMovementComponent.generated.h"
 
-UCLASS(HeaderGroup=Component, abstract)
+UCLASS(abstract)
 class ENGINE_API UPawnMovementComponent : public UNavMovementComponent
 {
 	GENERATED_UCLASS_BODY()
diff --git a/Engine/Source/Runtime/Engine/Classes/GameFramework/PlayerInput.h b/Engine/Source/Runtime/Engine/Classes/GameFramework/PlayerInput.h
--- a/Engine/Source/Runtime/Engine/Classes/GameFramework/PlayerInput.h
+++ b/Engine/Source/Runtime/Engine/Classes/GameFramework/PlayerInput.h
@@ -7,13 +7,17 @@
 //=============================================================================
 
 #pragma once
-
-#include "KeyState.h"
+#include "Components/InputComponent.h"
 #include "GestureRecognizer.h"
+#include "KeyState.h"
 #include "PlayerInput.generated.h"
 
 ENGINE_API DECLARE_LOG_CATEGORY_EXTERN(LogInput, Log, All);
 
+//
+// Forward declarations.
+//
+struct FDelegateDispatchDetails;
 
 USTRUCT()
 struct FKeyBind
@@ -223,7 +227,7 @@ struct FAxisKeyDetails
 	}
 };
 
-UCLASS(Within=PlayerController, config=Input, dependsOn=UEngineTypes, transient, HeaderGroup=UserInterface)
+UCLASS(Within=PlayerController, config=Input, dependsOn=UEngineTypes, transient)
 class ENGINE_API UPlayerInput : public UObject
 {
 	GENERATED_UCLASS_BODY()
diff --git a/Engine/Source/Runtime/Engine/Classes/GameFramework/PlayerState.h b/Engine/Source/Runtime/Engine/Classes/GameFramework/PlayerState.h
--- a/Engine/Source/Runtime/Engine/Classes/GameFramework/PlayerState.h
+++ b/Engine/Source/Runtime/Engine/Classes/GameFramework/PlayerState.h
@@ -30,7 +30,7 @@ struct PingAvgData
 // such as playername, score, etc.
 //=============================================================================
 
-UCLASS(HeaderGroup=ReplicationInfo, BlueprintType, Blueprintable, notplaceable)
+UCLASS(BlueprintType, Blueprintable, notplaceable)
 class ENGINE_API APlayerState : public AInfo
 {
 	GENERATED_UCLASS_BODY()
diff --git a/Engine/Source/Runtime/Engine/Classes/GameFramework/ProjectileMovementComponent.h b/Engine/Source/Runtime/Engine/Classes/GameFramework/ProjectileMovementComponent.h
--- a/Engine/Source/Runtime/Engine/Classes/GameFramework/ProjectileMovementComponent.h
+++ b/Engine/Source/Runtime/Engine/Classes/GameFramework/ProjectileMovementComponent.h
@@ -8,7 +8,7 @@
  * If the updated component is simulating physics, only the initial launch parameters (when initial velocity is non-zero) will affect the projectile,
  * and the physics sim will take over from there.
  */
-UCLASS(HeaderGroup=Component, ClassGroup=Movement, meta=(BlueprintSpawnableComponent), ShowCategories=(Velocity))
+UCLASS(ClassGroup=Movement, meta=(BlueprintSpawnableComponent), ShowCategories=(Velocity))
 class ENGINE_API UProjectileMovementComponent : public UMovementComponent
 {
 	GENERATED_UCLASS_BODY()
diff --git a/Engine/Source/Runtime/Engine/Classes/GameFramework/RotatingMovementComponent.h b/Engine/Source/Runtime/Engine/Classes/GameFramework/RotatingMovementComponent.h
--- a/Engine/Source/Runtime/Engine/Classes/GameFramework/RotatingMovementComponent.h
+++ b/Engine/Source/Runtime/Engine/Classes/GameFramework/RotatingMovementComponent.h
@@ -6,7 +6,7 @@
 /**
  * Movement component updates position of associated PrimitiveComponent during its tick.
  */
-UCLASS(HeaderGroup=Component, ClassGroup=Movement, meta=(BlueprintSpawnableComponent), HideCategories=(PlanarMovement, "Components|Movement|Planar", Velocity))
+UCLASS(ClassGroup=Movement, meta=(BlueprintSpawnableComponent), HideCategories=(PlanarMovement, "Components|Movement|Planar", Velocity))
 class ENGINE_API URotatingMovementComponent : public UMovementComponent
 {
 	GENERATED_UCLASS_BODY()
diff --git a/Engine/Source/Runtime/Engine/Classes/GameFramework/SpectatorPawn.h b/Engine/Source/Runtime/Engine/Classes/GameFramework/SpectatorPawn.h
--- a/Engine/Source/Runtime/Engine/Classes/GameFramework/SpectatorPawn.h
+++ b/Engine/Source/Runtime/Engine/Classes/GameFramework/SpectatorPawn.h
@@ -10,7 +10,7 @@
 #include "SpectatorPawn.generated.h"
 
 
-UCLASS(HeaderGroup=Pawn, config=Game, Blueprintable, BlueprintType)
+UCLASS(config=Game, Blueprintable, BlueprintType)
 class ENGINE_API ASpectatorPawn : public ADefaultPawn
 {
 	GENERATED_UCLASS_BODY()
diff --git a/Engine/Source/Runtime/Engine/Classes/GameFramework/SpectatorPawnMovement.h b/Engine/Source/Runtime/Engine/Classes/GameFramework/SpectatorPawnMovement.h
--- a/Engine/Source/Runtime/Engine/Classes/GameFramework/SpectatorPawnMovement.h
+++ b/Engine/Source/Runtime/Engine/Classes/GameFramework/SpectatorPawnMovement.h
@@ -7,7 +7,7 @@
 #pragma once
 #include "SpectatorPawnMovement.generated.h"
 
-UCLASS(HeaderGroup=Component, DependsOn=UFloatingPawnMovement)
+UCLASS(DependsOn=UFloatingPawnMovement)
 class ENGINE_API USpectatorPawnMovement : public UFloatingPawnMovement
 {
 	GENERATED_UCLASS_BODY()
diff --git a/Engine/Source/Runtime/Engine/Classes/GameFramework/SpringArmComponent.h b/Engine/Source/Runtime/Engine/Classes/GameFramework/SpringArmComponent.h
--- a/Engine/Source/Runtime/Engine/Classes/GameFramework/SpringArmComponent.h
+++ b/Engine/Source/Runtime/Engine/Classes/GameFramework/SpringArmComponent.h
@@ -11,7 +11,7 @@
  * Example: Use as a 'camera boom' to keep the follow camera for a player from colliding into the world.
  */
 
-UCLASS(HeaderGroup=Component, ClassGroup=Camera, meta=(BlueprintSpawnableComponent), hideCategories=(Mobility))
+UCLASS(ClassGroup=Camera, meta=(BlueprintSpawnableComponent), hideCategories=(Mobility))
 class ENGINE_API USpringArmComponent : public USceneComponent
 {
 	GENERATED_UCLASS_BODY()
diff --git a/Engine/Source/Runtime/Engine/Classes/GameFramework/WheeledVehicle.h b/Engine/Source/Runtime/Engine/Classes/GameFramework/WheeledVehicle.h
--- a/Engine/Source/Runtime/Engine/Classes/GameFramework/WheeledVehicle.h
+++ b/Engine/Source/Runtime/Engine/Classes/GameFramework/WheeledVehicle.h
@@ -3,7 +3,7 @@
 #pragma once
 #include "WheeledVehicle.generated.h"
 
-UCLASS(abstract, HeaderGroup=Pawn, config=Game, dependson=(AController), BlueprintType)
+UCLASS(abstract, config=Game, dependson=(AController), BlueprintType)
 class ENGINE_API AWheeledVehicle : public APawn
 {
 	GENERATED_UCLASS_BODY()
diff --git a/Engine/Source/Runtime/Engine/Classes/GameFramework/WorldSettings.h b/Engine/Source/Runtime/Engine/Classes/GameFramework/WorldSettings.h
--- a/Engine/Source/Runtime/Engine/Classes/GameFramework/WorldSettings.h
+++ b/Engine/Source/Runtime/Engine/Classes/GameFramework/WorldSettings.h
@@ -173,7 +173,7 @@ struct ENGINE_API FNetViewer
 /**
  * Actor containing all script accessible world properties.
  */
-UCLASS(HeaderGroup=GameEngine, config=game, showcategories=(Rendering), hidecategories=(Actor, Advanced, Display, Events, Object, Attachment, Info, Input, Blueprint, Layers), notplaceable, dependson=(APostProcessVolume, UMusicTrackDataStructures))
+UCLASS(config=game, showcategories=(Rendering), hidecategories=(Actor, Advanced, Display, Events, Object, Attachment, Info, Input, Blueprint, Layers), notplaceable, dependson=(APostProcessVolume, UMusicTrackDataStructures))
 class ENGINE_API AWorldSettings : public AInfo
 {
 	GENERATED_UCLASS_BODY()
diff --git a/Engine/Source/Runtime/Engine/Classes/Kismet/GameplayStatics.h b/Engine/Source/Runtime/Engine/Classes/Kismet/GameplayStatics.h
--- a/Engine/Source/Runtime/Engine/Classes/Kismet/GameplayStatics.h
+++ b/Engine/Source/Runtime/Engine/Classes/Kismet/GameplayStatics.h
@@ -5,6 +5,10 @@
 #include "../Sound/DialogueTypes.h"
 #include "GameplayStatics.generated.h"
 
+//
+// Forward declarations.
+//
+class USaveGame;
 struct FDialogueContext;
 
 UENUM()
@@ -18,7 +22,7 @@ namespace ESuggestProjVelocityTraceOption
 	};
 }
 
-UCLASS(HeaderGroup=KismetLibrary)
+UCLASS()
 class ENGINE_API UGameplayStatics : public UBlueprintFunctionLibrary
 {
 	GENERATED_UCLASS_BODY()
diff --git a/Engine/Source/Runtime/Engine/Classes/Kismet/KismetAIHelperLibrary.h b/Engine/Source/Runtime/Engine/Classes/Kismet/KismetAIHelperLibrary.h
--- a/Engine/Source/Runtime/Engine/Classes/Kismet/KismetAIHelperLibrary.h
+++ b/Engine/Source/Runtime/Engine/Classes/Kismet/KismetAIHelperLibrary.h
@@ -8,7 +8,7 @@
 #pragma once
 #include "KismetAIHelperLibrary.generated.h"
 
-UCLASS(HeaderGroup=KismetLibrary)
+UCLASS()
 class ENGINE_API UKismetAIHelperLibrary : public UBlueprintFunctionLibrary
 {
 	GENERATED_UCLASS_BODY()
diff --git a/Engine/Source/Runtime/Engine/Classes/Kismet/KismetArrayLibrary.h b/Engine/Source/Runtime/Engine/Classes/Kismet/KismetArrayLibrary.h
--- a/Engine/Source/Runtime/Engine/Classes/Kismet/KismetArrayLibrary.h
+++ b/Engine/Source/Runtime/Engine/Classes/Kismet/KismetArrayLibrary.h
@@ -3,7 +3,7 @@
 #pragma once
 #include "KismetArrayLibrary.generated.h"
 
-UCLASS(HeaderGroup=KismetLibrary,MinimalAPI)
+UCLASS(MinimalAPI)
 class UKismetArrayLibrary : public UBlueprintFunctionLibrary
 {
 	GENERATED_UCLASS_BODY()
diff --git a/Engine/Source/Runtime/Engine/Classes/Kismet/KismetMaterialLibrary.h b/Engine/Source/Runtime/Engine/Classes/Kismet/KismetMaterialLibrary.h
--- a/Engine/Source/Runtime/Engine/Classes/Kismet/KismetMaterialLibrary.h
+++ b/Engine/Source/Runtime/Engine/Classes/Kismet/KismetMaterialLibrary.h
@@ -4,7 +4,7 @@
 
 #include "KismetMaterialLibrary.generated.h"
 
-UCLASS(HeaderGroup=KismetLibrary,MinimalAPI)
+UCLASS(MinimalAPI)
 class UKismetMaterialLibrary : public UBlueprintFunctionLibrary
 {
 	GENERATED_UCLASS_BODY()
diff --git a/Engine/Source/Runtime/Engine/Classes/Landscape/ControlPointMeshComponent.h b/Engine/Source/Runtime/Engine/Classes/Landscape/ControlPointMeshComponent.h
--- a/Engine/Source/Runtime/Engine/Classes/Landscape/ControlPointMeshComponent.h
+++ b/Engine/Source/Runtime/Engine/Classes/Landscape/ControlPointMeshComponent.h
@@ -4,7 +4,7 @@
 
 #include "ControlPointMeshComponent.generated.h"
 
-UCLASS(HeaderGroup=Terrain, MinimalAPI)
+UCLASS(MinimalAPI)
 class UControlPointMeshComponent : public UStaticMeshComponent
 {
 	GENERATED_UCLASS_BODY()
diff --git a/Engine/Source/Runtime/Engine/Classes/Landscape/Landscape.h b/Engine/Source/Runtime/Engine/Classes/Landscape/Landscape.h
--- a/Engine/Source/Runtime/Engine/Classes/Landscape/Landscape.h
+++ b/Engine/Source/Runtime/Engine/Classes/Landscape/Landscape.h
@@ -17,7 +17,7 @@ enum ELandscapeSetupErrors
 	LSE_MAX,
 };
 
-UCLASS(dependson=ULightComponent, HeaderGroup=Terrain, Placeable, hidecategories=LandscapeProxy, showcategories=(Display, Movement, Collision, Lighting, LOD, Input), MinimalAPI)
+UCLASS(dependson=ULightComponent, Placeable, hidecategories=LandscapeProxy, showcategories=(Display, Movement, Collision, Lighting, LOD, Input), MinimalAPI)
 class ALandscape : public ALandscapeProxy
 {
 	GENERATED_UCLASS_BODY()
diff --git a/Engine/Source/Runtime/Engine/Classes/Landscape/LandscapeComponent.h b/Engine/Source/Runtime/Engine/Classes/Landscape/LandscapeComponent.h
--- a/Engine/Source/Runtime/Engine/Classes/Landscape/LandscapeComponent.h
+++ b/Engine/Source/Runtime/Engine/Classes/Landscape/LandscapeComponent.h
@@ -2,7 +2,8 @@
 
 #pragma once
 #include "LandscapeLayerInfoObject.h"
-
+#include "LightMap.h"
+#include "ShadowMap.h"
 #include "LandscapeComponent.generated.h"
 
 class FLandscapeComponentDerivedData
@@ -80,7 +81,7 @@ struct FWeightmapLayerAllocationInfo
 	}
 };
 
-UCLASS(HeaderGroup=Terrain, hidecategories=(Display, Attachment, Physics, Debug, Collision, Movement, Rendering, PrimitiveComponent, Object, Transform), MinimalAPI)
+UCLASS(hidecategories=(Display, Attachment, Physics, Debug, Collision, Movement, Rendering, PrimitiveComponent, Object, Transform), MinimalAPI)
 class ULandscapeComponent : public UPrimitiveComponent
 {
 	GENERATED_UCLASS_BODY()
diff --git a/Engine/Source/Runtime/Engine/Classes/Landscape/LandscapeGizmoActiveActor.h b/Engine/Source/Runtime/Engine/Classes/Landscape/LandscapeGizmoActiveActor.h
--- a/Engine/Source/Runtime/Engine/Classes/Landscape/LandscapeGizmoActiveActor.h
+++ b/Engine/Source/Runtime/Engine/Classes/Landscape/LandscapeGizmoActiveActor.h
@@ -39,7 +39,7 @@ struct FGizmoSelectData
 	
 };
 
-UCLASS(notplaceable, HeaderGroup=Terrain, MinimalAPI)
+UCLASS(notplaceable, MinimalAPI)
 class ALandscapeGizmoActiveActor : public ALandscapeGizmoActor
 {
 	GENERATED_UCLASS_BODY()
diff --git a/Engine/Source/Runtime/Engine/Classes/Landscape/LandscapeGizmoActor.h b/Engine/Source/Runtime/Engine/Classes/Landscape/LandscapeGizmoActor.h
--- a/Engine/Source/Runtime/Engine/Classes/Landscape/LandscapeGizmoActor.h
+++ b/Engine/Source/Runtime/Engine/Classes/Landscape/LandscapeGizmoActor.h
@@ -4,7 +4,7 @@
 #pragma once
 #include "LandscapeGizmoActor.generated.h"
 
-UCLASS(notplaceable, HeaderGroup=Terrain, MinimalAPI, NotBlueprintable)
+UCLASS(notplaceable, MinimalAPI, NotBlueprintable)
 class ALandscapeGizmoActor : public AActor
 {
 	GENERATED_UCLASS_BODY()
diff --git a/Engine/Source/Runtime/Engine/Classes/Landscape/LandscapeGizmoRenderComponent.h b/Engine/Source/Runtime/Engine/Classes/Landscape/LandscapeGizmoRenderComponent.h
--- a/Engine/Source/Runtime/Engine/Classes/Landscape/LandscapeGizmoRenderComponent.h
+++ b/Engine/Source/Runtime/Engine/Classes/Landscape/LandscapeGizmoRenderComponent.h
@@ -4,7 +4,7 @@
 #pragma once
 #include "LandscapeGizmoRenderComponent.generated.h"
 
-UCLASS(HeaderGroup=Terrain, hidecategories=Object)
+UCLASS(hidecategories=Object)
 class ULandscapeGizmoRenderComponent : public UPrimitiveComponent
 {
 	GENERATED_UCLASS_BODY()
diff --git a/Engine/Source/Runtime/Engine/Classes/Landscape/LandscapeHeightfieldCollisionComponent.h b/Engine/Source/Runtime/Engine/Classes/Landscape/LandscapeHeightfieldCollisionComponent.h
--- a/Engine/Source/Runtime/Engine/Classes/Landscape/LandscapeHeightfieldCollisionComponent.h
+++ b/Engine/Source/Runtime/Engine/Classes/Landscape/LandscapeHeightfieldCollisionComponent.h
@@ -4,7 +4,7 @@
 #pragma once
 #include "LandscapeHeightfieldCollisionComponent.generated.h"
 
-UCLASS(HeaderGroup=Terrain, MinimalAPI)
+UCLASS(MinimalAPI)
 class ULandscapeHeightfieldCollisionComponent : public UPrimitiveComponent
 {
 	GENERATED_UCLASS_BODY()
diff --git a/Engine/Source/Runtime/Engine/Classes/Landscape/LandscapeInfo.h b/Engine/Source/Runtime/Engine/Classes/Landscape/LandscapeInfo.h
--- a/Engine/Source/Runtime/Engine/Classes/Landscape/LandscapeInfo.h
+++ b/Engine/Source/Runtime/Engine/Classes/Landscape/LandscapeInfo.h
@@ -96,7 +96,7 @@ struct FLandscapeInfoLayerSettings
 #endif
 };
 
-UCLASS(HeaderGroup=Terrain)
+UCLASS()
 class ULandscapeInfo : public UObject
 {
 	GENERATED_UCLASS_BODY()
diff --git a/Engine/Source/Runtime/Engine/Classes/Landscape/LandscapeMaterialInstanceConstant.h b/Engine/Source/Runtime/Engine/Classes/Landscape/LandscapeMaterialInstanceConstant.h
--- a/Engine/Source/Runtime/Engine/Classes/Landscape/LandscapeMaterialInstanceConstant.h
+++ b/Engine/Source/Runtime/Engine/Classes/Landscape/LandscapeMaterialInstanceConstant.h
@@ -3,7 +3,7 @@
 #pragma once
 #include "LandscapeMaterialInstanceConstant.generated.h"
 
-UCLASS(HeaderGroup=Material, MinimalAPI)
+UCLASS(MinimalAPI)
 class ULandscapeMaterialInstanceConstant : public UMaterialInstanceConstant
 {
 	GENERATED_UCLASS_BODY()
diff --git a/Engine/Source/Runtime/Engine/Classes/Landscape/LandscapeMeshCollisionComponent.h b/Engine/Source/Runtime/Engine/Classes/Landscape/LandscapeMeshCollisionComponent.h
--- a/Engine/Source/Runtime/Engine/Classes/Landscape/LandscapeMeshCollisionComponent.h
+++ b/Engine/Source/Runtime/Engine/Classes/Landscape/LandscapeMeshCollisionComponent.h
@@ -4,7 +4,7 @@
 #pragma once
 #include "LandscapeMeshCollisionComponent.generated.h"
 
-UCLASS(HeaderGroup=Terrain)
+UCLASS()
 class ULandscapeMeshCollisionComponent : public ULandscapeHeightfieldCollisionComponent
 {
 	GENERATED_UCLASS_BODY()
diff --git a/Engine/Source/Runtime/Engine/Classes/Landscape/LandscapeProxy.h b/Engine/Source/Runtime/Engine/Classes/Landscape/LandscapeProxy.h
--- a/Engine/Source/Runtime/Engine/Classes/Landscape/LandscapeProxy.h
+++ b/Engine/Source/Runtime/Engine/Classes/Landscape/LandscapeProxy.h
@@ -6,6 +6,11 @@
 
 #include "LandscapeProxy.generated.h"
 
+//
+// Forward declarations.
+//
+class ULandscapeMaterialInstanceConstant;
+
 /** Structure storing channel usage for weightmap textures */
 USTRUCT()
 struct FLandscapeWeightmapUsage
@@ -164,7 +169,7 @@ namespace ELandscapeLayerPaintingRestriction
 	};
 }
 
-UCLASS(HeaderGroup=Terrain, dependson=UEngineTypes, NotPlaceable, hidecategories=(Display, Attachment, Physics, Debug, Lighting, LOD), showcategories=(Rendering, "Utilities|Orientation"), MinimalAPI)
+UCLASS(dependson=UEngineTypes, NotPlaceable, hidecategories=(Display, Attachment, Physics, Debug, Lighting, LOD), showcategories=(Rendering, "Utilities|Orientation"), MinimalAPI)
 class ALandscapeProxy : public AActor, public INavRelevantActorInterface
 {
 	GENERATED_UCLASS_BODY()
diff --git a/Engine/Source/Runtime/Engine/Classes/Landscape/LandscapeSplineControlPoint.h b/Engine/Source/Runtime/Engine/Classes/Landscape/LandscapeSplineControlPoint.h
--- a/Engine/Source/Runtime/Engine/Classes/Landscape/LandscapeSplineControlPoint.h
+++ b/Engine/Source/Runtime/Engine/Classes/Landscape/LandscapeSplineControlPoint.h
@@ -36,7 +36,7 @@ struct FLandscapeSplineConnection
 	ENGINE_API struct FLandscapeSplineSegmentConnection& GetFarConnection() const;
 };
 
-UCLASS(HeaderGroup=Terrain,Within=LandscapeSplinesComponent,autoExpandCategories=LandscapeSplineControlPoint,MinimalAPI)
+UCLASS(Within=LandscapeSplinesComponent,autoExpandCategories=LandscapeSplineControlPoint,MinimalAPI)
 class ULandscapeSplineControlPoint : public UObject
 {
 	GENERATED_UCLASS_BODY()
diff --git a/Engine/Source/Runtime/Engine/Classes/Landscape/LandscapeSplineSegment.h b/Engine/Source/Runtime/Engine/Classes/Landscape/LandscapeSplineSegment.h
--- a/Engine/Source/Runtime/Engine/Classes/Landscape/LandscapeSplineSegment.h
+++ b/Engine/Source/Runtime/Engine/Classes/Landscape/LandscapeSplineSegment.h
@@ -140,7 +140,7 @@ struct FLandscapeSplineMeshEntry
 };
 
 
-UCLASS(HeaderGroup=Terrain,Within=LandscapeSplinesComponent,autoExpandCategories=(LandscapeSplineSegment,LandscapeSplineMeshes),MinimalAPI)
+UCLASS(Within=LandscapeSplinesComponent,autoExpandCategories=(LandscapeSplineSegment,LandscapeSplineMeshes),MinimalAPI)
 class ULandscapeSplineSegment : public UObject
 {
 	GENERATED_UCLASS_BODY()
diff --git a/Engine/Source/Runtime/Engine/Classes/Landscape/LandscapeSplinesComponent.h b/Engine/Source/Runtime/Engine/Classes/Landscape/LandscapeSplinesComponent.h
--- a/Engine/Source/Runtime/Engine/Classes/Landscape/LandscapeSplinesComponent.h
+++ b/Engine/Source/Runtime/Engine/Classes/Landscape/LandscapeSplinesComponent.h
@@ -4,7 +4,7 @@
 #pragma once
 #include "LandscapeSplinesComponent.generated.h"
 
-UCLASS(HeaderGroup=Terrain,MinimalAPI)
+UCLASS(MinimalAPI)
 class ULandscapeSplinesComponent : public UPrimitiveComponent
 {
 	GENERATED_UCLASS_BODY()
diff --git a/Engine/Source/Runtime/Engine/Classes/Materials/Material.h b/Engine/Source/Runtime/Engine/Classes/Materials/Material.h
--- a/Engine/Source/Runtime/Engine/Classes/Materials/Material.h
+++ b/Engine/Source/Runtime/Engine/Classes/Materials/Material.h
@@ -254,7 +254,7 @@ struct FMaterialParameterCollectionInfo
  * when light from the scene hits the surface, the lighting model of the material is used to calculate how
  * that light interacts with the surface. 
  */
-UCLASS(HeaderGroup=Material, hidecategories=Object, MinimalAPI, BlueprintType)
+UCLASS(hidecategories=Object, MinimalAPI, BlueprintType)
 class UMaterial : public UMaterialInterface
 {
 	GENERATED_UCLASS_BODY()
diff --git a/Engine/Source/Runtime/Engine/Classes/Materials/MaterialExpression.h b/Engine/Source/Runtime/Engine/Classes/Materials/MaterialExpression.h
--- a/Engine/Source/Runtime/Engine/Classes/Materials/MaterialExpression.h
+++ b/Engine/Source/Runtime/Engine/Classes/Materials/MaterialExpression.h
@@ -77,7 +77,7 @@ struct FExpressionOutput
 };
 #endif
 
-UCLASS(HeaderGroup=Material, abstract, hidecategories=Object, MinimalAPI)
+UCLASS(abstract, hidecategories=Object, MinimalAPI)
 class UMaterialExpression : public UObject
 {
 	GENERATED_UCLASS_BODY()
diff --git a/Engine/Source/Runtime/Engine/Classes/Materials/MaterialExpressionAbs.h b/Engine/Source/Runtime/Engine/Classes/Materials/MaterialExpressionAbs.h
--- a/Engine/Source/Runtime/Engine/Classes/Materials/MaterialExpressionAbs.h
+++ b/Engine/Source/Runtime/Engine/Classes/Materials/MaterialExpressionAbs.h
@@ -8,7 +8,7 @@
 #pragma once
 #include "MaterialExpressionAbs.generated.h"
 
-UCLASS(HeaderGroup=Material, collapsecategories, hidecategories=Object)
+UCLASS(collapsecategories, hidecategories=Object)
 class UMaterialExpressionAbs : public UMaterialExpression
 {
 	GENERATED_UCLASS_BODY()
diff --git a/Engine/Source/Runtime/Engine/Classes/Materials/MaterialExpressionActorPositionWS.h b/Engine/Source/Runtime/Engine/Classes/Materials/MaterialExpressionActorPositionWS.h
--- a/Engine/Source/Runtime/Engine/Classes/Materials/MaterialExpressionActorPositionWS.h
+++ b/Engine/Source/Runtime/Engine/Classes/Materials/MaterialExpressionActorPositionWS.h
@@ -4,7 +4,7 @@
 #pragma once
 #include "MaterialExpressionActorPositionWS.generated.h"
 
-UCLASS(HeaderGroup=Material, collapsecategories, hidecategories=Object)
+UCLASS(collapsecategories, hidecategories=Object)
 class UMaterialExpressionActorPositionWS : public UMaterialExpression
 {
 	GENERATED_UCLASS_BODY()
diff --git a/Engine/Source/Runtime/Engine/Classes/Materials/MaterialExpressionAdd.h b/Engine/Source/Runtime/Engine/Classes/Materials/MaterialExpressionAdd.h
--- a/Engine/Source/Runtime/Engine/Classes/Materials/MaterialExpressionAdd.h
+++ b/Engine/Source/Runtime/Engine/Classes/Materials/MaterialExpressionAdd.h
@@ -4,7 +4,7 @@
 #pragma once
 #include "MaterialExpressionAdd.generated.h"
 
-UCLASS(HeaderGroup=Material, MinimalAPI)
+UCLASS(MinimalAPI)
 class UMaterialExpressionAdd : public UMaterialExpression
 {
 	GENERATED_UCLASS_BODY()
diff --git a/Engine/Source/Runtime/Engine/Classes/Materials/MaterialExpressionBlackBody.h b/Engine/Source/Runtime/Engine/Classes/Materials/MaterialExpressionBlackBody.h
--- a/Engine/Source/Runtime/Engine/Classes/Materials/MaterialExpressionBlackBody.h
+++ b/Engine/Source/Runtime/Engine/Classes/Materials/MaterialExpressionBlackBody.h
@@ -4,7 +4,7 @@
 #pragma once
 #include "MaterialExpressionBlackBody.generated.h"
 
-UCLASS(HeaderGroup=Material)
+UCLASS()
 class UMaterialExpressionBlackBody : public UMaterialExpression
 {
 	GENERATED_UCLASS_BODY()
diff --git a/Engine/Source/Runtime/Engine/Classes/Materials/MaterialExpressionBreakMaterialAttributes.h b/Engine/Source/Runtime/Engine/Classes/Materials/MaterialExpressionBreakMaterialAttributes.h
--- a/Engine/Source/Runtime/Engine/Classes/Materials/MaterialExpressionBreakMaterialAttributes.h
+++ b/Engine/Source/Runtime/Engine/Classes/Materials/MaterialExpressionBreakMaterialAttributes.h
@@ -4,7 +4,7 @@
 #pragma once
 #include "MaterialExpressionBreakMaterialAttributes.generated.h"
 
-UCLASS(HeaderGroup=Material, collapsecategories, hidecategories=Object, MinimalAPI)
+UCLASS(collapsecategories, hidecategories=Object, MinimalAPI)
 class UMaterialExpressionBreakMaterialAttributes : public UMaterialExpression
 {
 	GENERATED_UCLASS_BODY()
diff --git a/Engine/Source/Runtime/Engine/Classes/Materials/MaterialExpressionBumpOffset.h b/Engine/Source/Runtime/Engine/Classes/Materials/MaterialExpressionBumpOffset.h
--- a/Engine/Source/Runtime/Engine/Classes/Materials/MaterialExpressionBumpOffset.h
+++ b/Engine/Source/Runtime/Engine/Classes/Materials/MaterialExpressionBumpOffset.h
@@ -4,7 +4,7 @@
 #pragma once
 #include "MaterialExpressionBumpOffset.generated.h"
 
-UCLASS(HeaderGroup=Material, collapsecategories, hidecategories=Object, MinimalAPI)
+UCLASS(collapsecategories, hidecategories=Object, MinimalAPI)
 class UMaterialExpressionBumpOffset : public UMaterialExpression
 {
 	GENERATED_UCLASS_BODY()
diff --git a/Engine/Source/Runtime/Engine/Classes/Materials/MaterialExpressionCollectionParameter.h b/Engine/Source/Runtime/Engine/Classes/Materials/MaterialExpressionCollectionParameter.h
--- a/Engine/Source/Runtime/Engine/Classes/Materials/MaterialExpressionCollectionParameter.h
+++ b/Engine/Source/Runtime/Engine/Classes/Materials/MaterialExpressionCollectionParameter.h
@@ -7,7 +7,7 @@
 #pragma once
 #include "MaterialExpressionCollectionParameter.generated.h"
 
-UCLASS(HeaderGroup=Material, hidecategories=object, MinimalAPI)
+UCLASS(hidecategories=object, MinimalAPI)
 class UMaterialExpressionCollectionParameter : public UMaterialExpression
 {
 	GENERATED_UCLASS_BODY()
diff --git a/Engine/Source/Runtime/Engine/Classes/Materials/MaterialExpressionConstantBiasScale.h b/Engine/Source/Runtime/Engine/Classes/Materials/MaterialExpressionConstantBiasScale.h
--- a/Engine/Source/Runtime/Engine/Classes/Materials/MaterialExpressionConstantBiasScale.h
+++ b/Engine/Source/Runtime/Engine/Classes/Materials/MaterialExpressionConstantBiasScale.h
@@ -4,7 +4,7 @@
 #pragma once
 #include "MaterialExpressionConstantBiasScale.generated.h"
 
-UCLASS(HeaderGroup=Material)
+UCLASS()
 class UMaterialExpressionConstantBiasScale : public UMaterialExpression
 {
 	GENERATED_UCLASS_BODY()
diff --git a/Engine/Source/Runtime/Engine/Classes/Materials/MaterialExpressionFunctionInput.h b/Engine/Source/Runtime/Engine/Classes/Materials/MaterialExpressionFunctionInput.h
--- a/Engine/Source/Runtime/Engine/Classes/Materials/MaterialExpressionFunctionInput.h
+++ b/Engine/Source/Runtime/Engine/Classes/Materials/MaterialExpressionFunctionInput.h
@@ -19,7 +19,7 @@ enum EFunctionInputType
 	FunctionInput_MAX,
 };
 
-UCLASS(HeaderGroup=Material, hidecategories=object, MinimalAPI)
+UCLASS(hidecategories=object, MinimalAPI)
 class UMaterialExpressionFunctionInput : public UMaterialExpression
 {
 	GENERATED_UCLASS_BODY()
diff --git a/Engine/Source/Runtime/Engine/Classes/Materials/MaterialExpressionTextureBase.h b/Engine/Source/Runtime/Engine/Classes/Materials/MaterialExpressionTextureBase.h
--- a/Engine/Source/Runtime/Engine/Classes/Materials/MaterialExpressionTextureBase.h
+++ b/Engine/Source/Runtime/Engine/Classes/Materials/MaterialExpressionTextureBase.h
@@ -8,7 +8,7 @@
 #pragma once
 #include "MaterialExpressionTextureBase.generated.h"
 
-UCLASS(HeaderGroup=Material, abstract, hidecategories=Object, MinimalAPI)
+UCLASS(abstract, hidecategories=Object, MinimalAPI)
 class UMaterialExpressionTextureBase : public UMaterialExpression 
 {
 	GENERATED_UCLASS_BODY()
diff --git a/Engine/Source/Runtime/Engine/Classes/Materials/MaterialExpressionTextureSampleParameter.h b/Engine/Source/Runtime/Engine/Classes/Materials/MaterialExpressionTextureSampleParameter.h
--- a/Engine/Source/Runtime/Engine/Classes/Materials/MaterialExpressionTextureSampleParameter.h
+++ b/Engine/Source/Runtime/Engine/Classes/Materials/MaterialExpressionTextureSampleParameter.h
@@ -4,7 +4,7 @@
 #pragma once
 #include "MaterialExpressionTextureSampleParameter.generated.h"
 
-UCLASS(HeaderGroup=Material, collapsecategories, abstract, hidecategories=Object, MinimalAPI)
+UCLASS(collapsecategories, abstract, hidecategories=Object, MinimalAPI)
 class UMaterialExpressionTextureSampleParameter : public UMaterialExpressionTextureSample
 {
 	GENERATED_UCLASS_BODY()
diff --git a/Engine/Source/Runtime/Engine/Classes/Materials/MaterialInstance.h b/Engine/Source/Runtime/Engine/Classes/Materials/MaterialInstance.h
--- a/Engine/Source/Runtime/Engine/Classes/Materials/MaterialInstance.h
+++ b/Engine/Source/Runtime/Engine/Classes/Materials/MaterialInstance.h
@@ -4,6 +4,11 @@
 #pragma once
 #include "MaterialInstance.generated.h"
 
+//
+// Forward declarations.
+//
+struct FMaterialInstanceBasePropertyOverrides;
+
 /** Editable font parameter. */
 USTRUCT()
 struct FFontParameterValue
@@ -121,7 +126,7 @@ struct FVectorParameterValue
 	static ValueType GetValue(const FVectorParameterValue& Parameter) { return Parameter.ParameterValue; }
 };
 
-UCLASS(abstract, HeaderGroup=Material, BlueprintType,MinimalAPI)
+UCLASS(abstract, BlueprintType,MinimalAPI)
 class UMaterialInstance : public UMaterialInterface
 {
 	GENERATED_UCLASS_BODY()
diff --git a/Engine/Source/Runtime/Engine/Classes/Materials/MaterialInstanceConstant.h b/Engine/Source/Runtime/Engine/Classes/Materials/MaterialInstanceConstant.h
--- a/Engine/Source/Runtime/Engine/Classes/Materials/MaterialInstanceConstant.h
+++ b/Engine/Source/Runtime/Engine/Classes/Materials/MaterialInstanceConstant.h
@@ -8,7 +8,7 @@
  * General modification of the material cannot be supported without recompilation, so the instances are limited to changing the values of
  * predefined material parameters. The parameters are statically defined in the compiled material by a unique name, type and default value.
  */
-UCLASS(HeaderGroup=Material, hidecategories=Object, collapsecategories, BlueprintType,MinimalAPI)
+UCLASS(hidecategories=Object, collapsecategories, BlueprintType,MinimalAPI)
 class UMaterialInstanceConstant : public UMaterialInstance
 {
 	GENERATED_UCLASS_BODY()
diff --git a/Engine/Source/Runtime/Engine/Classes/Materials/MaterialInstanceDynamic.h b/Engine/Source/Runtime/Engine/Classes/Materials/MaterialInstanceDynamic.h
--- a/Engine/Source/Runtime/Engine/Classes/Materials/MaterialInstanceDynamic.h
+++ b/Engine/Source/Runtime/Engine/Classes/Materials/MaterialInstanceDynamic.h
@@ -4,7 +4,7 @@
 #pragma once
 #include "MaterialInstanceDynamic.generated.h"
 
-UCLASS(HeaderGroup=Material, hidecategories=Object, collapsecategories, BlueprintType, MinimalAPI)
+UCLASS(hidecategories=Object, collapsecategories, BlueprintType, MinimalAPI)
 class UMaterialInstanceDynamic : public UMaterialInstance
 {
 	GENERATED_UCLASS_BODY()
diff --git a/Engine/Source/Runtime/Engine/Classes/Materials/MaterialInterface.h b/Engine/Source/Runtime/Engine/Classes/Materials/MaterialInterface.h
--- a/Engine/Source/Runtime/Engine/Classes/Materials/MaterialInterface.h
+++ b/Engine/Source/Runtime/Engine/Classes/Materials/MaterialInterface.h
@@ -90,7 +90,7 @@ struct FLightmassMaterialInterfaceSettings
 	
 };
 
-UCLASS(abstract, HeaderGroup=Material, BlueprintType,MinimalAPI)
+UCLASS(abstract, BlueprintType,MinimalAPI)
 class UMaterialInterface : public UObject, public IBlendableInterface
 {
 	GENERATED_UCLASS_BODY()
diff --git a/Engine/Source/Runtime/Engine/Classes/Matinee/InterpFilter.h b/Engine/Source/Runtime/Engine/Classes/Matinee/InterpFilter.h
--- a/Engine/Source/Runtime/Engine/Classes/Matinee/InterpFilter.h
+++ b/Engine/Source/Runtime/Engine/Classes/Matinee/InterpFilter.h
@@ -9,7 +9,7 @@
 #pragma once
 #include "InterpFilter.generated.h"
 
-UCLASS(HeaderGroup=Interpolation)
+UCLASS()
 class UInterpFilter : public UObject
 {
 	GENERATED_UCLASS_BODY()
diff --git a/Engine/Source/Runtime/Engine/Classes/Matinee/InterpFilter_Classes.h b/Engine/Source/Runtime/Engine/Classes/Matinee/InterpFilter_Classes.h
--- a/Engine/Source/Runtime/Engine/Classes/Matinee/InterpFilter_Classes.h
+++ b/Engine/Source/Runtime/Engine/Classes/Matinee/InterpFilter_Classes.h
@@ -9,7 +9,7 @@
 #pragma once
 #include "InterpFilter_Classes.generated.h"
 
-UCLASS(HeaderGroup=Interpolation)
+UCLASS()
 class UInterpFilter_Classes : public UInterpFilter
 {
 	GENERATED_UCLASS_BODY()
diff --git a/Engine/Source/Runtime/Engine/Classes/Matinee/InterpFilter_Custom.h b/Engine/Source/Runtime/Engine/Classes/Matinee/InterpFilter_Custom.h
--- a/Engine/Source/Runtime/Engine/Classes/Matinee/InterpFilter_Custom.h
+++ b/Engine/Source/Runtime/Engine/Classes/Matinee/InterpFilter_Custom.h
@@ -9,7 +9,7 @@
 #pragma once
 #include "InterpFilter_Custom.generated.h"
 
-UCLASS(HeaderGroup=Interpolation, MinimalAPI)
+UCLASS(MinimalAPI)
 class UInterpFilter_Custom : public UInterpFilter
 {
 	GENERATED_UCLASS_BODY()
diff --git a/Engine/Source/Runtime/Engine/Classes/Matinee/InterpGroup.h b/Engine/Source/Runtime/Engine/Classes/Matinee/InterpGroup.h
--- a/Engine/Source/Runtime/Engine/Classes/Matinee/InterpGroup.h
+++ b/Engine/Source/Runtime/Engine/Classes/Matinee/InterpGroup.h
@@ -1,7 +1,7 @@
 // Copyright 1998-2014 Epic Games, Inc. All Rights Reserved.
 
 #pragma once
-
+#include "Interpolation.h"
 #include "InterpGroup.generated.h"
 
 /**
@@ -74,7 +74,7 @@ struct FInterpEdSelKey
 	
 };
 
-UCLASS(HeaderGroup=Interpolation, collapsecategories, hidecategories=Object, MinimalAPI)
+UCLASS(collapsecategories, hidecategories=Object, MinimalAPI)
 class UInterpGroup : public UObject, public FInterpEdInputInterface
 {
 	GENERATED_UCLASS_BODY()
diff --git a/Engine/Source/Runtime/Engine/Classes/Matinee/InterpGroupCamera.h b/Engine/Source/Runtime/Engine/Classes/Matinee/InterpGroupCamera.h
--- a/Engine/Source/Runtime/Engine/Classes/Matinee/InterpGroupCamera.h
+++ b/Engine/Source/Runtime/Engine/Classes/Matinee/InterpGroupCamera.h
@@ -46,7 +46,7 @@ struct FCameraPreviewInfo
 	
 };
 
-UCLASS(HeaderGroup=Interpolation, collapsecategories, hidecategories=Object, MinimalAPI)
+UCLASS(collapsecategories, hidecategories=Object, MinimalAPI)
 class UInterpGroupCamera : public UInterpGroup
 {
 	GENERATED_UCLASS_BODY()
diff --git a/Engine/Source/Runtime/Engine/Classes/Matinee/InterpGroupDirector.h b/Engine/Source/Runtime/Engine/Classes/Matinee/InterpGroupDirector.h
--- a/Engine/Source/Runtime/Engine/Classes/Matinee/InterpGroupDirector.h
+++ b/Engine/Source/Runtime/Engine/Classes/Matinee/InterpGroupDirector.h
@@ -8,7 +8,7 @@
 #pragma once
 #include "InterpGroupDirector.generated.h"
 
-UCLASS(HeaderGroup=Interpolation, collapsecategories, hidecategories=Object, MinimalAPI)
+UCLASS(collapsecategories, hidecategories=Object, MinimalAPI)
 class UInterpGroupDirector : public UInterpGroup
 {
 	GENERATED_UCLASS_BODY()
diff --git a/Engine/Source/Runtime/Engine/Classes/Matinee/InterpGroupInst.h b/Engine/Source/Runtime/Engine/Classes/Matinee/InterpGroupInst.h
--- a/Engine/Source/Runtime/Engine/Classes/Matinee/InterpGroupInst.h
+++ b/Engine/Source/Runtime/Engine/Classes/Matinee/InterpGroupInst.h
@@ -3,7 +3,7 @@
 #pragma once
 #include "InterpGroupInst.generated.h"
 
-UCLASS(HeaderGroup=Interpolation, MinimalAPI)
+UCLASS(MinimalAPI)
 class UInterpGroupInst : public UObject
 {
 	GENERATED_UCLASS_BODY()
diff --git a/Engine/Source/Runtime/Engine/Classes/Matinee/InterpTrack.h b/Engine/Source/Runtime/Engine/Classes/Matinee/InterpTrack.h
--- a/Engine/Source/Runtime/Engine/Classes/Matinee/InterpTrack.h
+++ b/Engine/Source/Runtime/Engine/Classes/Matinee/InterpTrack.h
@@ -10,6 +10,12 @@
 
 #include "InterpTrack.generated.h"
 
+//
+// Forward declarations
+//
+class FCanvas;
+
+
 /** Required condition for this track to be enabled */
 UENUM()
 enum ETrackActiveCondition
@@ -81,7 +87,7 @@ struct FSubTrackGroup
 
 };
 
-UCLASS(HeaderGroup=Interpolation, collapsecategories, hidecategories=Object, abstract, MinimalAPI)
+UCLASS(collapsecategories, hidecategories=Object, abstract, MinimalAPI)
 class UInterpTrack : public UObject, public FCurveEdInterface, public FInterpEdInputInterface
 {
 	GENERATED_UCLASS_BODY()
diff --git a/Engine/Source/Runtime/Engine/Classes/Matinee/InterpTrackAnimControl.h b/Engine/Source/Runtime/Engine/Classes/Matinee/InterpTrackAnimControl.h
--- a/Engine/Source/Runtime/Engine/Classes/Matinee/InterpTrackAnimControl.h
+++ b/Engine/Source/Runtime/Engine/Classes/Matinee/InterpTrackAnimControl.h
@@ -51,7 +51,7 @@ struct FAnimControlTrackKey
 
 };
 
-UCLASS(HeaderGroup=Interpolation, MinimalAPI, meta=( DisplayName = "Animation Track" ))
+UCLASS(MinimalAPI, meta=( DisplayName = "Animation Track" ))
 class UInterpTrackAnimControl : public UInterpTrackFloatBase
 {
 	GENERATED_UCLASS_BODY()
