#ttp 303393 UObject improvement: Remove header groups

diff --git a/Engine/Source/Runtime/Engine/Classes/Matinee/InterpTrackAudioMaster.h b/Engine/Source/Runtime/Engine/Classes/Matinee/InterpTrackAudioMaster.h
--- a/Engine/Source/Runtime/Engine/Classes/Matinee/InterpTrackAudioMaster.h
+++ b/Engine/Source/Runtime/Engine/Classes/Matinee/InterpTrackAudioMaster.h
@@ -4,7 +4,7 @@
 #pragma once
 #include "InterpTrackAudioMaster.generated.h"
 
-UCLASS(HeaderGroup=Interpolation, MinimalAPI, meta=( DisplayName = "Audio Master Track" ) )
+UCLASS(MinimalAPI, meta=( DisplayName = "Audio Master Track" ) )
 class UInterpTrackAudioMaster : public UInterpTrackVectorBase
 {
 	GENERATED_UCLASS_BODY()
diff --git a/Engine/Source/Runtime/Engine/Classes/Matinee/InterpTrackBoolProp.h b/Engine/Source/Runtime/Engine/Classes/Matinee/InterpTrackBoolProp.h
--- a/Engine/Source/Runtime/Engine/Classes/Matinee/InterpTrackBoolProp.h
+++ b/Engine/Source/Runtime/Engine/Classes/Matinee/InterpTrackBoolProp.h
@@ -25,7 +25,7 @@ struct FBoolTrackKey
 
 };
 
-UCLASS(HeaderGroup=Interpolation, MinimalAPI, meta=( DisplayName = "Bool Property Track" ) )
+UCLASS(MinimalAPI, meta=( DisplayName = "Bool Property Track" ) )
 class UInterpTrackBoolProp : public UInterpTrack
 {
 	GENERATED_UCLASS_BODY()
diff --git a/Engine/Source/Runtime/Engine/Classes/Matinee/InterpTrackColorProp.h b/Engine/Source/Runtime/Engine/Classes/Matinee/InterpTrackColorProp.h
--- a/Engine/Source/Runtime/Engine/Classes/Matinee/InterpTrackColorProp.h
+++ b/Engine/Source/Runtime/Engine/Classes/Matinee/InterpTrackColorProp.h
@@ -4,7 +4,7 @@
 #pragma once
 #include "InterpTrackColorProp.generated.h"
 
-UCLASS(HeaderGroup=Interpolation, MinimalAPI, meta=( DisplayName = "Color Property Track" ) )
+UCLASS(MinimalAPI, meta=( DisplayName = "Color Property Track" ) )
 class UInterpTrackColorProp : public UInterpTrackVectorBase
 {
 	GENERATED_UCLASS_BODY()
diff --git a/Engine/Source/Runtime/Engine/Classes/Matinee/InterpTrackColorScale.h b/Engine/Source/Runtime/Engine/Classes/Matinee/InterpTrackColorScale.h
--- a/Engine/Source/Runtime/Engine/Classes/Matinee/InterpTrackColorScale.h
+++ b/Engine/Source/Runtime/Engine/Classes/Matinee/InterpTrackColorScale.h
@@ -8,7 +8,7 @@
 #pragma once
 #include "InterpTrackColorScale.generated.h"
 
-UCLASS(HeaderGroup=Interpolation, meta=( DisplayName = "Color Scale Track" ) )
+UCLASS(meta=( DisplayName = "Color Scale Track" ) )
 class UInterpTrackColorScale : public UInterpTrackVectorBase
 {
 	GENERATED_UCLASS_BODY()
diff --git a/Engine/Source/Runtime/Engine/Classes/Matinee/InterpTrackDirector.h b/Engine/Source/Runtime/Engine/Classes/Matinee/InterpTrackDirector.h
--- a/Engine/Source/Runtime/Engine/Classes/Matinee/InterpTrackDirector.h
+++ b/Engine/Source/Runtime/Engine/Classes/Matinee/InterpTrackDirector.h
@@ -41,7 +41,7 @@ struct FDirectorTrackCut
 
 };
 
-UCLASS(HeaderGroup=Interpolation, MinimalAPI, meta=( DisplayName = "Director Track" ) )
+UCLASS(MinimalAPI, meta=( DisplayName = "Director Track" ) )
 class UInterpTrackDirector : public UInterpTrack
 {
 	GENERATED_UCLASS_BODY()
diff --git a/Engine/Source/Runtime/Engine/Classes/Matinee/InterpTrackEvent.h b/Engine/Source/Runtime/Engine/Classes/Matinee/InterpTrackEvent.h
--- a/Engine/Source/Runtime/Engine/Classes/Matinee/InterpTrackEvent.h
+++ b/Engine/Source/Runtime/Engine/Classes/Matinee/InterpTrackEvent.h
@@ -30,7 +30,7 @@ struct FEventTrackKey
 
 };
 
-UCLASS(HeaderGroup=Interpolation, MinimalAPI, meta=( DisplayName = "Event Track" ) )
+UCLASS(MinimalAPI, meta=( DisplayName = "Event Track" ) )
 class UInterpTrackEvent : public UInterpTrack
 {
 	GENERATED_UCLASS_BODY()
diff --git a/Engine/Source/Runtime/Engine/Classes/Matinee/InterpTrackFade.h b/Engine/Source/Runtime/Engine/Classes/Matinee/InterpTrackFade.h
--- a/Engine/Source/Runtime/Engine/Classes/Matinee/InterpTrackFade.h
+++ b/Engine/Source/Runtime/Engine/Classes/Matinee/InterpTrackFade.h
@@ -3,7 +3,7 @@
 #pragma once
 #include "InterpTrackFade.generated.h"
 
-UCLASS(HeaderGroup=Interpolation, meta=( DisplayName = "Fade Track" ) )
+UCLASS(meta=( DisplayName = "Fade Track" ) )
 class UInterpTrackFade : public UInterpTrackFloatBase
 {
 	GENERATED_UCLASS_BODY()
diff --git a/Engine/Source/Runtime/Engine/Classes/Matinee/InterpTrackFloatBase.h b/Engine/Source/Runtime/Engine/Classes/Matinee/InterpTrackFloatBase.h
--- a/Engine/Source/Runtime/Engine/Classes/Matinee/InterpTrackFloatBase.h
+++ b/Engine/Source/Runtime/Engine/Classes/Matinee/InterpTrackFloatBase.h
@@ -3,7 +3,7 @@
 #pragma once
 #include "InterpTrackFloatBase.generated.h"
 
-UCLASS(HeaderGroup=Interpolation, abstract, MinimalAPI)
+UCLASS(abstract, MinimalAPI)
 class UInterpTrackFloatBase : public UInterpTrack
 {
 	GENERATED_UCLASS_BODY()
diff --git a/Engine/Source/Runtime/Engine/Classes/Matinee/InterpTrackFloatMaterialParam.h b/Engine/Source/Runtime/Engine/Classes/Matinee/InterpTrackFloatMaterialParam.h
--- a/Engine/Source/Runtime/Engine/Classes/Matinee/InterpTrackFloatMaterialParam.h
+++ b/Engine/Source/Runtime/Engine/Classes/Matinee/InterpTrackFloatMaterialParam.h
@@ -3,7 +3,7 @@
 #pragma once
 #include "InterpTrackFloatMaterialParam.generated.h"
 
-UCLASS(HeaderGroup=Interpolation, meta=( DisplayName = "Float Material Parameter Track" ) )
+UCLASS(meta=( DisplayName = "Float Material Parameter Track" ) )
 class UInterpTrackFloatMaterialParam : public UInterpTrackFloatBase
 {
 	GENERATED_UCLASS_BODY()
diff --git a/Engine/Source/Runtime/Engine/Classes/Matinee/InterpTrackFloatParticleParam.h b/Engine/Source/Runtime/Engine/Classes/Matinee/InterpTrackFloatParticleParam.h
--- a/Engine/Source/Runtime/Engine/Classes/Matinee/InterpTrackFloatParticleParam.h
+++ b/Engine/Source/Runtime/Engine/Classes/Matinee/InterpTrackFloatParticleParam.h
@@ -3,7 +3,7 @@
 #pragma once
 #include "InterpTrackFloatParticleParam.generated.h"
 
-UCLASS(HeaderGroup=Interpolation, meta=( DisplayName = "Float Particle Param Track" ) )
+UCLASS(meta=( DisplayName = "Float Particle Param Track" ) )
 class UInterpTrackFloatParticleParam : public UInterpTrackFloatBase
 {
 	GENERATED_UCLASS_BODY()
diff --git a/Engine/Source/Runtime/Engine/Classes/Matinee/InterpTrackFloatProp.h b/Engine/Source/Runtime/Engine/Classes/Matinee/InterpTrackFloatProp.h
--- a/Engine/Source/Runtime/Engine/Classes/Matinee/InterpTrackFloatProp.h
+++ b/Engine/Source/Runtime/Engine/Classes/Matinee/InterpTrackFloatProp.h
@@ -3,7 +3,7 @@
 #pragma once
 #include "InterpTrackFloatProp.generated.h"
 
-UCLASS(HeaderGroup=Interpolation, MinimalAPI, meta=( DisplayName = "Float Property Track" ) )
+UCLASS(MinimalAPI, meta=( DisplayName = "Float Property Track" ) )
 class UInterpTrackFloatProp : public UInterpTrackFloatBase
 {
 	GENERATED_UCLASS_BODY()
diff --git a/Engine/Source/Runtime/Engine/Classes/Matinee/InterpTrackLinearColorBase.h b/Engine/Source/Runtime/Engine/Classes/Matinee/InterpTrackLinearColorBase.h
--- a/Engine/Source/Runtime/Engine/Classes/Matinee/InterpTrackLinearColorBase.h
+++ b/Engine/Source/Runtime/Engine/Classes/Matinee/InterpTrackLinearColorBase.h
@@ -4,7 +4,7 @@
 #pragma once
 #include "InterpTrackLinearColorBase.generated.h"
 
-UCLASS(HeaderGroup=Interpolation, abstract, MinimalAPI)
+UCLASS(abstract, MinimalAPI)
 class UInterpTrackLinearColorBase : public UInterpTrack
 {
 	GENERATED_UCLASS_BODY()
diff --git a/Engine/Source/Runtime/Engine/Classes/Matinee/InterpTrackLinearColorProp.h b/Engine/Source/Runtime/Engine/Classes/Matinee/InterpTrackLinearColorProp.h
--- a/Engine/Source/Runtime/Engine/Classes/Matinee/InterpTrackLinearColorProp.h
+++ b/Engine/Source/Runtime/Engine/Classes/Matinee/InterpTrackLinearColorProp.h
@@ -4,7 +4,7 @@
 #pragma once
 #include "InterpTrackLinearColorProp.generated.h"
 
-UCLASS(HeaderGroup=Interpolation, MinimalAPI, meta=( DisplayName = "LinearColor Property Track" ) )
+UCLASS(MinimalAPI, meta=( DisplayName = "LinearColor Property Track" ) )
 class UInterpTrackLinearColorProp : public UInterpTrackLinearColorBase
 {
 	GENERATED_UCLASS_BODY()
diff --git a/Engine/Source/Runtime/Engine/Classes/Matinee/InterpTrackMove.h b/Engine/Source/Runtime/Engine/Classes/Matinee/InterpTrackMove.h
--- a/Engine/Source/Runtime/Engine/Classes/Matinee/InterpTrackMove.h
+++ b/Engine/Source/Runtime/Engine/Classes/Matinee/InterpTrackMove.h
@@ -86,7 +86,7 @@ struct FInterpLookupTrack
 };
 
 /** Track containing data for moving an actor around over time. */
-UCLASS(HeaderGroup=Interpolation, MinimalAPI, meta=( DisplayName = "Movement Track" ) )
+UCLASS(MinimalAPI, meta=( DisplayName = "Movement Track" ) )
 class UInterpTrackMove : public UInterpTrack
 {
 	GENERATED_UCLASS_BODY()
diff --git a/Engine/Source/Runtime/Engine/Classes/Matinee/InterpTrackMoveAxis.h b/Engine/Source/Runtime/Engine/Classes/Matinee/InterpTrackMoveAxis.h
--- a/Engine/Source/Runtime/Engine/Classes/Matinee/InterpTrackMoveAxis.h
+++ b/Engine/Source/Runtime/Engine/Classes/Matinee/InterpTrackMoveAxis.h
@@ -21,7 +21,7 @@ enum EInterpMoveAxis
 	AXIS_RotationZ,
 };
 
-UCLASS(dependson=UInterpTrackMove, HeaderGroup=Interpolation, MinimalAPI, meta=( DisplayName = "Move Axis Track" ) )
+UCLASS(dependson=UInterpTrackMove, MinimalAPI, meta=( DisplayName = "Move Axis Track" ) )
 class UInterpTrackMoveAxis : public UInterpTrackFloatBase
 {
 	GENERATED_UCLASS_BODY()
diff --git a/Engine/Source/Runtime/Engine/Classes/Matinee/InterpTrackParticleReplay.h b/Engine/Source/Runtime/Engine/Classes/Matinee/InterpTrackParticleReplay.h
--- a/Engine/Source/Runtime/Engine/Classes/Matinee/InterpTrackParticleReplay.h
+++ b/Engine/Source/Runtime/Engine/Classes/Matinee/InterpTrackParticleReplay.h
@@ -87,7 +87,7 @@ struct FParticleReplayTrackKey
 
 };
 
-UCLASS(HeaderGroup=Interpolation, MinimalAPI, meta=( DisplayName = "Particle Replay Track" ) )
+UCLASS(MinimalAPI, meta=( DisplayName = "Particle Replay Track" ) )
 class UInterpTrackParticleReplay : public UInterpTrack
 {
 	GENERATED_UCLASS_BODY()
diff --git a/Engine/Source/Runtime/Engine/Classes/Matinee/InterpTrackSlomo.h b/Engine/Source/Runtime/Engine/Classes/Matinee/InterpTrackSlomo.h
--- a/Engine/Source/Runtime/Engine/Classes/Matinee/InterpTrackSlomo.h
+++ b/Engine/Source/Runtime/Engine/Classes/Matinee/InterpTrackSlomo.h
@@ -4,7 +4,7 @@
 #pragma once
 #include "InterpTrackSlomo.generated.h"
 
-UCLASS(HeaderGroup=Interpolation, meta=( DisplayName = "Slomo Track" ) )
+UCLASS(meta=( DisplayName = "Slomo Track" ) )
 class UInterpTrackSlomo : public UInterpTrackFloatBase
 {
 	GENERATED_UCLASS_BODY()
diff --git a/Engine/Source/Runtime/Engine/Classes/Matinee/InterpTrackSound.h b/Engine/Source/Runtime/Engine/Classes/Matinee/InterpTrackSound.h
--- a/Engine/Source/Runtime/Engine/Classes/Matinee/InterpTrackSound.h
+++ b/Engine/Source/Runtime/Engine/Classes/Matinee/InterpTrackSound.h
@@ -38,7 +38,7 @@ struct FSoundTrackKey
 	
 };
 
-UCLASS(HeaderGroup=Interpolation, MinimalAPI, meta=( DisplayName = "Sound Track" ) )
+UCLASS(MinimalAPI, meta=( DisplayName = "Sound Track" ) )
 class UInterpTrackSound : public UInterpTrackVectorBase
 {
 	GENERATED_UCLASS_BODY()
diff --git a/Engine/Source/Runtime/Engine/Classes/Matinee/InterpTrackToggle.h b/Engine/Source/Runtime/Engine/Classes/Matinee/InterpTrackToggle.h
--- a/Engine/Source/Runtime/Engine/Classes/Matinee/InterpTrackToggle.h
+++ b/Engine/Source/Runtime/Engine/Classes/Matinee/InterpTrackToggle.h
@@ -42,7 +42,7 @@ struct FToggleTrackKey
 
 };
 
-UCLASS(HeaderGroup=Interpolation, MinimalAPI, meta=( DisplayName = "Toggle Track" ) )
+UCLASS(MinimalAPI, meta=( DisplayName = "Toggle Track" ) )
 class UInterpTrackToggle : public UInterpTrack
 {
 	GENERATED_UCLASS_BODY()
diff --git a/Engine/Source/Runtime/Engine/Classes/Matinee/InterpTrackVectorMaterialParam.h b/Engine/Source/Runtime/Engine/Classes/Matinee/InterpTrackVectorMaterialParam.h
--- a/Engine/Source/Runtime/Engine/Classes/Matinee/InterpTrackVectorMaterialParam.h
+++ b/Engine/Source/Runtime/Engine/Classes/Matinee/InterpTrackVectorMaterialParam.h
@@ -3,7 +3,7 @@
 #pragma once
 #include "InterpTrackVectorMaterialParam.generated.h"
 
-UCLASS(HeaderGroup=Interpolation, MinimalAPI, meta=( DisplayName = "Vector Material Parameter Track" ))
+UCLASS(MinimalAPI, meta=( DisplayName = "Vector Material Parameter Track" ))
 class UInterpTrackVectorMaterialParam : public UInterpTrackVectorBase
 {
 	GENERATED_UCLASS_BODY()
diff --git a/Engine/Source/Runtime/Engine/Classes/Matinee/InterpTrackVectorProp.h b/Engine/Source/Runtime/Engine/Classes/Matinee/InterpTrackVectorProp.h
--- a/Engine/Source/Runtime/Engine/Classes/Matinee/InterpTrackVectorProp.h
+++ b/Engine/Source/Runtime/Engine/Classes/Matinee/InterpTrackVectorProp.h
@@ -4,7 +4,7 @@
 #pragma once
 #include "InterpTrackVectorProp.generated.h"
 
-UCLASS(HeaderGroup=Interpolation, MinimalAPI, meta=( DisplayName = "Vector Property Track" ) )
+UCLASS(MinimalAPI, meta=( DisplayName = "Vector Property Track" ) )
 class UInterpTrackVectorProp : public UInterpTrackVectorBase
 {
 	GENERATED_UCLASS_BODY()
diff --git a/Engine/Source/Runtime/Engine/Classes/Matinee/InterpTrackVisibility.h b/Engine/Source/Runtime/Engine/Classes/Matinee/InterpTrackVisibility.h
--- a/Engine/Source/Runtime/Engine/Classes/Matinee/InterpTrackVisibility.h
+++ b/Engine/Source/Runtime/Engine/Classes/Matinee/InterpTrackVisibility.h
@@ -61,7 +61,7 @@ struct FVisibilityTrackKey
 
 };
 
-UCLASS(HeaderGroup=Interpolation, MinimalAPI, meta=( DisplayName = "Visibility Track" ) )
+UCLASS(MinimalAPI, meta=( DisplayName = "Visibility Track" ) )
 class UInterpTrackVisibility : public UInterpTrack
 {
 	GENERATED_UCLASS_BODY()
diff --git a/Engine/Source/Runtime/Engine/Classes/Particles/Acceleration/ParticleModuleAcceleration.h b/Engine/Source/Runtime/Engine/Classes/Particles/Acceleration/ParticleModuleAcceleration.h
--- a/Engine/Source/Runtime/Engine/Classes/Particles/Acceleration/ParticleModuleAcceleration.h
+++ b/Engine/Source/Runtime/Engine/Classes/Particles/Acceleration/ParticleModuleAcceleration.h
@@ -4,7 +4,7 @@
 #pragma once
 #include "ParticleModuleAcceleration.generated.h"
 
-UCLASS(HeaderGroup=Particle, editinlinenew, hidecategories=Object, meta=(DisplayName = "Acceleration"), MinimalAPI)
+UCLASS(editinlinenew, hidecategories=Object, meta=(DisplayName = "Acceleration"), MinimalAPI)
 class UParticleModuleAcceleration : public UParticleModuleAccelerationBase
 {
 	GENERATED_UCLASS_BODY()
diff --git a/Engine/Source/Runtime/Engine/Classes/Particles/Acceleration/ParticleModuleAccelerationBase.h b/Engine/Source/Runtime/Engine/Classes/Particles/Acceleration/ParticleModuleAccelerationBase.h
--- a/Engine/Source/Runtime/Engine/Classes/Particles/Acceleration/ParticleModuleAccelerationBase.h
+++ b/Engine/Source/Runtime/Engine/Classes/Particles/Acceleration/ParticleModuleAccelerationBase.h
@@ -4,7 +4,7 @@
 #pragma once
 #include "ParticleModuleAccelerationBase.generated.h"
 
-UCLASS(HeaderGroup=Particle, editinlinenew, hidecategories=Object, meta=(DisplayName = "Acceleration"), abstract)
+UCLASS(editinlinenew, hidecategories=Object, meta=(DisplayName = "Acceleration"), abstract)
 class UParticleModuleAccelerationBase : public UParticleModule
 {
 	GENERATED_UCLASS_BODY()
diff --git a/Engine/Source/Runtime/Engine/Classes/Particles/Acceleration/ParticleModuleAccelerationConstant.h b/Engine/Source/Runtime/Engine/Classes/Particles/Acceleration/ParticleModuleAccelerationConstant.h
--- a/Engine/Source/Runtime/Engine/Classes/Particles/Acceleration/ParticleModuleAccelerationConstant.h
+++ b/Engine/Source/Runtime/Engine/Classes/Particles/Acceleration/ParticleModuleAccelerationConstant.h
@@ -8,7 +8,7 @@
 
 #include "ParticleModuleAccelerationConstant.generated.h"
 
-UCLASS(HeaderGroup=Particle, editinlinenew, hidecategories=(Object, Acceleration), meta=(DisplayName = "Const Acceleration"))
+UCLASS(editinlinenew, hidecategories=(Object, Acceleration), meta=(DisplayName = "Const Acceleration"))
 class UParticleModuleAccelerationConstant : public UParticleModuleAccelerationBase
 {
 	GENERATED_UCLASS_BODY()
diff --git a/Engine/Source/Runtime/Engine/Classes/Particles/Acceleration/ParticleModuleAccelerationDrag.h b/Engine/Source/Runtime/Engine/Classes/Particles/Acceleration/ParticleModuleAccelerationDrag.h
--- a/Engine/Source/Runtime/Engine/Classes/Particles/Acceleration/ParticleModuleAccelerationDrag.h
+++ b/Engine/Source/Runtime/Engine/Classes/Particles/Acceleration/ParticleModuleAccelerationDrag.h
@@ -8,7 +8,7 @@
 
 #include "ParticleModuleAccelerationDrag.generated.h"
 
-UCLASS(HeaderGroup=Particle, editinlinenew, hidecategories=(Object, Acceleration), MinimalAPI, meta=(DisplayName = "Drag"))
+UCLASS(editinlinenew, hidecategories=(Object, Acceleration), MinimalAPI, meta=(DisplayName = "Drag"))
 class UParticleModuleAccelerationDrag : public UParticleModuleAccelerationBase
 {
 	GENERATED_UCLASS_BODY()
diff --git a/Engine/Source/Runtime/Engine/Classes/Particles/Acceleration/ParticleModuleAccelerationDragScaleOverLife.h b/Engine/Source/Runtime/Engine/Classes/Particles/Acceleration/ParticleModuleAccelerationDragScaleOverLife.h
--- a/Engine/Source/Runtime/Engine/Classes/Particles/Acceleration/ParticleModuleAccelerationDragScaleOverLife.h
+++ b/Engine/Source/Runtime/Engine/Classes/Particles/Acceleration/ParticleModuleAccelerationDragScaleOverLife.h
@@ -8,7 +8,7 @@
 
 #include "ParticleModuleAccelerationDragScaleOverLife.generated.h"
 
-UCLASS(HeaderGroup=Particle, editinlinenew, hidecategories=(UObject, Acceleration), MinimalAPI, meta=(DisplayName = "Drag Scale/Life"))
+UCLASS(editinlinenew, hidecategories=(UObject, Acceleration), MinimalAPI, meta=(DisplayName = "Drag Scale/Life"))
 class UParticleModuleAccelerationDragScaleOverLife : public UParticleModuleAccelerationBase
 {
 	GENERATED_UCLASS_BODY()
diff --git a/Engine/Source/Runtime/Engine/Classes/Particles/Acceleration/ParticleModuleAccelerationOverLifetime.h b/Engine/Source/Runtime/Engine/Classes/Particles/Acceleration/ParticleModuleAccelerationOverLifetime.h
--- a/Engine/Source/Runtime/Engine/Classes/Particles/Acceleration/ParticleModuleAccelerationOverLifetime.h
+++ b/Engine/Source/Runtime/Engine/Classes/Particles/Acceleration/ParticleModuleAccelerationOverLifetime.h
@@ -4,7 +4,7 @@
 #pragma once
 #include "ParticleModuleAccelerationOverLifetime.generated.h"
 
-UCLASS(HeaderGroup=Particle, editinlinenew, hidecategories=Object, meta=(DisplayName = "Acceleration/Life"), MinimalAPI)
+UCLASS(editinlinenew, hidecategories=Object, meta=(DisplayName = "Acceleration/Life"), MinimalAPI)
 class UParticleModuleAccelerationOverLifetime : public UParticleModuleAccelerationBase
 {
 	GENERATED_UCLASS_BODY()
diff --git a/Engine/Source/Runtime/Engine/Classes/Particles/Attractor/ParticleModuleAttractorBase.h b/Engine/Source/Runtime/Engine/Classes/Particles/Attractor/ParticleModuleAttractorBase.h
--- a/Engine/Source/Runtime/Engine/Classes/Particles/Attractor/ParticleModuleAttractorBase.h
+++ b/Engine/Source/Runtime/Engine/Classes/Particles/Attractor/ParticleModuleAttractorBase.h
@@ -4,7 +4,7 @@
 #pragma once
 #include "ParticleModuleAttractorBase.generated.h"
 
-UCLASS(HeaderGroup=Particle, editinlinenew, hidecategories=Object, abstract, meta=(DisplayName = "Attraction"))
+UCLASS(editinlinenew, hidecategories=Object, abstract, meta=(DisplayName = "Attraction"))
 class UParticleModuleAttractorBase : public UParticleModule
 {
 	GENERATED_UCLASS_BODY()
diff --git a/Engine/Source/Runtime/Engine/Classes/Particles/Attractor/ParticleModuleAttractorLine.h b/Engine/Source/Runtime/Engine/Classes/Particles/Attractor/ParticleModuleAttractorLine.h
--- a/Engine/Source/Runtime/Engine/Classes/Particles/Attractor/ParticleModuleAttractorLine.h
+++ b/Engine/Source/Runtime/Engine/Classes/Particles/Attractor/ParticleModuleAttractorLine.h
@@ -4,7 +4,7 @@
 #pragma once
 #include "ParticleModuleAttractorLine.generated.h"
 
-UCLASS(HeaderGroup=Particle, editinlinenew, hidecategories=Object, MinimalAPI, meta=(DisplayName = "Line Attractor"))
+UCLASS(editinlinenew, hidecategories=Object, MinimalAPI, meta=(DisplayName = "Line Attractor"))
 class UParticleModuleAttractorLine : public UParticleModuleAttractorBase
 {
 	GENERATED_UCLASS_BODY()
diff --git a/Engine/Source/Runtime/Engine/Classes/Particles/Attractor/ParticleModuleAttractorParticle.h b/Engine/Source/Runtime/Engine/Classes/Particles/Attractor/ParticleModuleAttractorParticle.h
--- a/Engine/Source/Runtime/Engine/Classes/Particles/Attractor/ParticleModuleAttractorParticle.h
+++ b/Engine/Source/Runtime/Engine/Classes/Particles/Attractor/ParticleModuleAttractorParticle.h
@@ -12,7 +12,7 @@ enum EAttractorParticleSelectionMethod
 	EAPSM_MAX,
 };
 
-UCLASS(HeaderGroup=Particle, editinlinenew, hidecategories=Object, MinimalAPI, meta=(DisplayName = "Particle Attractor"))
+UCLASS(editinlinenew, hidecategories=Object, MinimalAPI, meta=(DisplayName = "Particle Attractor"))
 class UParticleModuleAttractorParticle : public UParticleModuleAttractorBase
 {
 	GENERATED_UCLASS_BODY()
diff --git a/Engine/Source/Runtime/Engine/Classes/Particles/Attractor/ParticleModuleAttractorPoint.h b/Engine/Source/Runtime/Engine/Classes/Particles/Attractor/ParticleModuleAttractorPoint.h
--- a/Engine/Source/Runtime/Engine/Classes/Particles/Attractor/ParticleModuleAttractorPoint.h
+++ b/Engine/Source/Runtime/Engine/Classes/Particles/Attractor/ParticleModuleAttractorPoint.h
@@ -4,7 +4,7 @@
 #pragma once
 #include "ParticleModuleAttractorPoint.generated.h"
 
-UCLASS(HeaderGroup=Particle, editinlinenew, hidecategories=Object, MinimalAPI, meta=(DisplayName = "Point Attractor"))
+UCLASS(editinlinenew, hidecategories=Object, MinimalAPI, meta=(DisplayName = "Point Attractor"))
 class UParticleModuleAttractorPoint : public UParticleModuleAttractorBase
 {
 	GENERATED_UCLASS_BODY()
diff --git a/Engine/Source/Runtime/Engine/Classes/Particles/Attractor/ParticleModuleAttractorPointGravity.h b/Engine/Source/Runtime/Engine/Classes/Particles/Attractor/ParticleModuleAttractorPointGravity.h
--- a/Engine/Source/Runtime/Engine/Classes/Particles/Attractor/ParticleModuleAttractorPointGravity.h
+++ b/Engine/Source/Runtime/Engine/Classes/Particles/Attractor/ParticleModuleAttractorPointGravity.h
@@ -9,7 +9,7 @@
 
 #include "ParticleModuleAttractorPointGravity.generated.h"
 
-UCLASS(HeaderGroup=Particle, editinlinenew, hidecategories=Object, meta=(DisplayName = "Point Gravity"))
+UCLASS(editinlinenew, hidecategories=Object, meta=(DisplayName = "Point Gravity"))
 class UParticleModuleAttractorPointGravity : public UParticleModuleAttractorBase
 {
 	GENERATED_UCLASS_BODY()
diff --git a/Engine/Source/Runtime/Engine/Classes/Particles/Beam/ParticleModuleBeamBase.h b/Engine/Source/Runtime/Engine/Classes/Particles/Beam/ParticleModuleBeamBase.h
--- a/Engine/Source/Runtime/Engine/Classes/Particles/Beam/ParticleModuleBeamBase.h
+++ b/Engine/Source/Runtime/Engine/Classes/Particles/Beam/ParticleModuleBeamBase.h
@@ -51,7 +51,7 @@ enum Beam2SourceTargetTangentMethod
 	PEB2STTM_MAX,
 };
 
-UCLASS(HeaderGroup=Particle, editinlinenew, hidecategories=Object, abstract, meta=(DisplayName = "Beam"))
+UCLASS(editinlinenew, hidecategories=Object, abstract, meta=(DisplayName = "Beam"))
 class UParticleModuleBeamBase : public UParticleModule
 {
 	GENERATED_UCLASS_BODY()
diff --git a/Engine/Source/Runtime/Engine/Classes/Particles/Beam/ParticleModuleBeamModifier.h b/Engine/Source/Runtime/Engine/Classes/Particles/Beam/ParticleModuleBeamModifier.h
--- a/Engine/Source/Runtime/Engine/Classes/Particles/Beam/ParticleModuleBeamModifier.h
+++ b/Engine/Source/Runtime/Engine/Classes/Particles/Beam/ParticleModuleBeamModifier.h
@@ -48,7 +48,7 @@ struct FBeamModifierOptions
 
 };
 
-UCLASS(HeaderGroup=Particle, editinlinenew, hidecategories=Object, meta=(DisplayName = "Beam Modifier"))
+UCLASS(editinlinenew, hidecategories=Object, meta=(DisplayName = "Beam Modifier"))
 class UParticleModuleBeamModifier : public UParticleModuleBeamBase
 {
 	GENERATED_UCLASS_BODY()
diff --git a/Engine/Source/Runtime/Engine/Classes/Particles/Beam/ParticleModuleBeamNoise.h b/Engine/Source/Runtime/Engine/Classes/Particles/Beam/ParticleModuleBeamNoise.h
--- a/Engine/Source/Runtime/Engine/Classes/Particles/Beam/ParticleModuleBeamNoise.h
+++ b/Engine/Source/Runtime/Engine/Classes/Particles/Beam/ParticleModuleBeamNoise.h
@@ -10,7 +10,7 @@
 #pragma once
 #include "ParticleModuleBeamNoise.generated.h"
 
-UCLASS(HeaderGroup=Particle, editinlinenew, hidecategories=Object, MinimalAPI, meta=(DisplayName = "Noise"))
+UCLASS(editinlinenew, hidecategories=Object, MinimalAPI, meta=(DisplayName = "Noise"))
 class UParticleModuleBeamNoise : public UParticleModuleBeamBase
 {
 	GENERATED_UCLASS_BODY()
diff --git a/Engine/Source/Runtime/Engine/Classes/Particles/Beam/ParticleModuleBeamSource.h b/Engine/Source/Runtime/Engine/Classes/Particles/Beam/ParticleModuleBeamSource.h
--- a/Engine/Source/Runtime/Engine/Classes/Particles/Beam/ParticleModuleBeamSource.h
+++ b/Engine/Source/Runtime/Engine/Classes/Particles/Beam/ParticleModuleBeamSource.h
@@ -10,7 +10,7 @@
 #pragma once
 #include "ParticleModuleBeamSource.generated.h"
 
-UCLASS(HeaderGroup=Particle, editinlinenew, hidecategories=Object, MinimalAPI, meta=(DisplayName = "Source"))
+UCLASS(editinlinenew, hidecategories=Object, MinimalAPI, meta=(DisplayName = "Source"))
 class UParticleModuleBeamSource : public UParticleModuleBeamBase
 {
 	GENERATED_UCLASS_BODY()
diff --git a/Engine/Source/Runtime/Engine/Classes/Particles/Beam/ParticleModuleBeamTarget.h b/Engine/Source/Runtime/Engine/Classes/Particles/Beam/ParticleModuleBeamTarget.h
--- a/Engine/Source/Runtime/Engine/Classes/Particles/Beam/ParticleModuleBeamTarget.h
+++ b/Engine/Source/Runtime/Engine/Classes/Particles/Beam/ParticleModuleBeamTarget.h
@@ -10,7 +10,7 @@
 #pragma once
 #include "ParticleModuleBeamTarget.generated.h"
 
-UCLASS(HeaderGroup=Particle, editinlinenew, hidecategories=Object, MinimalAPI, meta=(DisplayName = "Target"))
+UCLASS(editinlinenew, hidecategories=Object, MinimalAPI, meta=(DisplayName = "Target"))
 class UParticleModuleBeamTarget : public UParticleModuleBeamBase
 {
 	GENERATED_UCLASS_BODY()
diff --git a/Engine/Source/Runtime/Engine/Classes/Particles/Camera/ParticleModuleCameraBase.h b/Engine/Source/Runtime/Engine/Classes/Particles/Camera/ParticleModuleCameraBase.h
--- a/Engine/Source/Runtime/Engine/Classes/Particles/Camera/ParticleModuleCameraBase.h
+++ b/Engine/Source/Runtime/Engine/Classes/Particles/Camera/ParticleModuleCameraBase.h
@@ -4,7 +4,7 @@
 #pragma once
 #include "ParticleModuleCameraBase.generated.h"
 
-UCLASS(HeaderGroup=Particle, editinlinenew, hidecategories=Object, abstract, meta=(DisplayName = "Camera"))
+UCLASS(editinlinenew, hidecategories=Object, abstract, meta=(DisplayName = "Camera"))
 class UParticleModuleCameraBase : public UParticleModule
 {
 	GENERATED_UCLASS_BODY()
diff --git a/Engine/Source/Runtime/Engine/Classes/Particles/Camera/ParticleModuleCameraOffset.h b/Engine/Source/Runtime/Engine/Classes/Particles/Camera/ParticleModuleCameraOffset.h
--- a/Engine/Source/Runtime/Engine/Classes/Particles/Camera/ParticleModuleCameraOffset.h
+++ b/Engine/Source/Runtime/Engine/Classes/Particles/Camera/ParticleModuleCameraOffset.h
@@ -16,7 +16,7 @@ enum EParticleCameraOffsetUpdateMethod
 	EPCOUM_MAX,
 };
 
-UCLASS(HeaderGroup=Particle, editinlinenew, hidecategories=Object, meta=(DisplayName = "Camera Offset"))
+UCLASS(editinlinenew, hidecategories=Object, meta=(DisplayName = "Camera Offset"))
 class UParticleModuleCameraOffset : public UParticleModuleCameraBase
 {
 	GENERATED_UCLASS_BODY()
diff --git a/Engine/Source/Runtime/Engine/Classes/Particles/Collision/ParticleModuleCollision.h b/Engine/Source/Runtime/Engine/Classes/Particles/Collision/ParticleModuleCollision.h
--- a/Engine/Source/Runtime/Engine/Classes/Particles/Collision/ParticleModuleCollision.h
+++ b/Engine/Source/Runtime/Engine/Classes/Particles/Collision/ParticleModuleCollision.h
@@ -4,7 +4,7 @@
 #pragma once
 #include "ParticleModuleCollision.generated.h"
 
-UCLASS(HeaderGroup=Particle, editinlinenew, hidecategories=Object, MinimalAPI, meta=(DisplayName = "Collision"))
+UCLASS(editinlinenew, hidecategories=Object, MinimalAPI, meta=(DisplayName = "Collision"))
 class UParticleModuleCollision : public UParticleModuleCollisionBase
 {
 	GENERATED_UCLASS_BODY()
diff --git a/Engine/Source/Runtime/Engine/Classes/Particles/Collision/ParticleModuleCollisionBase.h b/Engine/Source/Runtime/Engine/Classes/Particles/Collision/ParticleModuleCollisionBase.h
--- a/Engine/Source/Runtime/Engine/Classes/Particles/Collision/ParticleModuleCollisionBase.h
+++ b/Engine/Source/Runtime/Engine/Classes/Particles/Collision/ParticleModuleCollisionBase.h
@@ -25,7 +25,7 @@ enum EParticleCollisionComplete
 	EPCC_MAX,
 };
 
-UCLASS(HeaderGroup=Particle, editinlinenew, hidecategories=Object, abstract, meta=(DisplayName = "Collision"))
+UCLASS(editinlinenew, hidecategories=Object, abstract, meta=(DisplayName = "Collision"))
 class UParticleModuleCollisionBase : public UParticleModule
 {
 	GENERATED_UCLASS_BODY()
diff --git a/Engine/Source/Runtime/Engine/Classes/Particles/Collision/ParticleModuleCollisionGPU.h b/Engine/Source/Runtime/Engine/Classes/Particles/Collision/ParticleModuleCollisionGPU.h
--- a/Engine/Source/Runtime/Engine/Classes/Particles/Collision/ParticleModuleCollisionGPU.h
+++ b/Engine/Source/Runtime/Engine/Classes/Particles/Collision/ParticleModuleCollisionGPU.h
@@ -20,7 +20,7 @@ namespace EParticleCollisionResponse
 	};
 }
 
-UCLASS(HeaderGroup=Particle, editinlinenew, hidecategories=Object, MinimalAPI, meta=(DisplayName = "Collision (Scene Depth)"))
+UCLASS(editinlinenew, hidecategories=Object, MinimalAPI, meta=(DisplayName = "Collision (Scene Depth)"))
 class UParticleModuleCollisionGPU : public UParticleModuleCollisionBase
 {
 	GENERATED_UCLASS_BODY()
diff --git a/Engine/Source/Runtime/Engine/Classes/Particles/Color/ParticleModuleColor.h b/Engine/Source/Runtime/Engine/Classes/Particles/Color/ParticleModuleColor.h
--- a/Engine/Source/Runtime/Engine/Classes/Particles/Color/ParticleModuleColor.h
+++ b/Engine/Source/Runtime/Engine/Classes/Particles/Color/ParticleModuleColor.h
@@ -4,7 +4,7 @@
 #pragma once
 #include "ParticleModuleColor.generated.h"
 
-UCLASS(HeaderGroup=Particle, editinlinenew, hidecategories=Object, MinimalAPI, meta=(DisplayName = "Initial Color"))
+UCLASS(editinlinenew, hidecategories=Object, MinimalAPI, meta=(DisplayName = "Initial Color"))
 class UParticleModuleColor : public UParticleModuleColorBase
 {
 	GENERATED_UCLASS_BODY()
diff --git a/Engine/Source/Runtime/Engine/Classes/Particles/Color/ParticleModuleColorBase.h b/Engine/Source/Runtime/Engine/Classes/Particles/Color/ParticleModuleColorBase.h
--- a/Engine/Source/Runtime/Engine/Classes/Particles/Color/ParticleModuleColorBase.h
+++ b/Engine/Source/Runtime/Engine/Classes/Particles/Color/ParticleModuleColorBase.h
@@ -4,7 +4,7 @@
 #pragma once
 #include "ParticleModuleColorBase.generated.h"
 
-UCLASS(HeaderGroup=Particle, editinlinenew, hidecategories=Object, abstract, meta=(DisplayName = "Color"))
+UCLASS(editinlinenew, hidecategories=Object, abstract, meta=(DisplayName = "Color"))
 class UParticleModuleColorBase : public UParticleModule
 {
 	GENERATED_UCLASS_BODY()
diff --git a/Engine/Source/Runtime/Engine/Classes/Particles/Color/ParticleModuleColorOverLife.h b/Engine/Source/Runtime/Engine/Classes/Particles/Color/ParticleModuleColorOverLife.h
--- a/Engine/Source/Runtime/Engine/Classes/Particles/Color/ParticleModuleColorOverLife.h
+++ b/Engine/Source/Runtime/Engine/Classes/Particles/Color/ParticleModuleColorOverLife.h
@@ -4,7 +4,7 @@
 #pragma once
 #include "ParticleModuleColorOverLife.generated.h"
 
-UCLASS(HeaderGroup=Particle, editinlinenew, hidecategories=Object, MinimalAPI, meta=(DisplayName = "Color Over Life"))
+UCLASS(editinlinenew, hidecategories=Object, MinimalAPI, meta=(DisplayName = "Color Over Life"))
 class UParticleModuleColorOverLife : public UParticleModuleColorBase
 {
 	GENERATED_UCLASS_BODY()
diff --git a/Engine/Source/Runtime/Engine/Classes/Particles/Color/ParticleModuleColorScaleOverLife.h b/Engine/Source/Runtime/Engine/Classes/Particles/Color/ParticleModuleColorScaleOverLife.h
--- a/Engine/Source/Runtime/Engine/Classes/Particles/Color/ParticleModuleColorScaleOverLife.h
+++ b/Engine/Source/Runtime/Engine/Classes/Particles/Color/ParticleModuleColorScaleOverLife.h
@@ -10,7 +10,7 @@
 #pragma once
 #include "ParticleModuleColorScaleOverLife.generated.h"
 
-UCLASS(HeaderGroup=Particle, editinlinenew, hidecategories=Object, MinimalAPI, meta=(DisplayName = "Scale Color / Life"))
+UCLASS(editinlinenew, hidecategories=Object, MinimalAPI, meta=(DisplayName = "Scale Color / Life"))
 class UParticleModuleColorScaleOverLife : public UParticleModuleColorBase
 {
 	GENERATED_UCLASS_BODY()
diff --git a/Engine/Source/Runtime/Engine/Classes/Particles/Color/ParticleModuleColor_Seeded.h b/Engine/Source/Runtime/Engine/Classes/Particles/Color/ParticleModuleColor_Seeded.h
--- a/Engine/Source/Runtime/Engine/Classes/Particles/Color/ParticleModuleColor_Seeded.h
+++ b/Engine/Source/Runtime/Engine/Classes/Particles/Color/ParticleModuleColor_Seeded.h
@@ -4,7 +4,7 @@
 #pragma once
 #include "ParticleModuleColor_Seeded.generated.h"
 
-UCLASS(HeaderGroup=Particle, editinlinenew, hidecategories=Object, meta=(DisplayName = "Init Color (Seed)"))
+UCLASS(editinlinenew, hidecategories=Object, meta=(DisplayName = "Init Color (Seed)"))
 class UParticleModuleColor_Seeded : public UParticleModuleColor
 {
 	GENERATED_UCLASS_BODY()
diff --git a/Engine/Source/Runtime/Engine/Classes/Particles/EmitterCameraLensEffectBase.h b/Engine/Source/Runtime/Engine/Classes/Particles/EmitterCameraLensEffectBase.h
--- a/Engine/Source/Runtime/Engine/Classes/Particles/EmitterCameraLensEffectBase.h
+++ b/Engine/Source/Runtime/Engine/Classes/Particles/EmitterCameraLensEffectBase.h
@@ -9,7 +9,7 @@
 #pragma once
 #include "EmitterCameraLensEffectBase.generated.h"
 
-UCLASS(abstract, Blueprintable, dependson=AEmitter, HeaderGroup=Particle, MinimalAPI)
+UCLASS(abstract, Blueprintable, dependson=AEmitter, MinimalAPI)
 class AEmitterCameraLensEffectBase : public AEmitter
 {
 	GENERATED_UCLASS_BODY()
diff --git a/Engine/Source/Runtime/Engine/Classes/Particles/Event/ParticleModuleEventBase.h b/Engine/Source/Runtime/Engine/Classes/Particles/Event/ParticleModuleEventBase.h
--- a/Engine/Source/Runtime/Engine/Classes/Particles/Event/ParticleModuleEventBase.h
+++ b/Engine/Source/Runtime/Engine/Classes/Particles/Event/ParticleModuleEventBase.h
@@ -4,7 +4,7 @@
 #pragma once
 #include "ParticleModuleEventBase.generated.h"
 
-UCLASS(HeaderGroup=Particle, editinlinenew, hidecategories=Object, abstract, meta=(DisplayName = "Event"))
+UCLASS(editinlinenew, hidecategories=Object, abstract, meta=(DisplayName = "Event"))
 class UParticleModuleEventBase : public UParticleModule
 {
 	GENERATED_UCLASS_BODY()
diff --git a/Engine/Source/Runtime/Engine/Classes/Particles/Event/ParticleModuleEventGenerator.h b/Engine/Source/Runtime/Engine/Classes/Particles/Event/ParticleModuleEventGenerator.h
--- a/Engine/Source/Runtime/Engine/Classes/Particles/Event/ParticleModuleEventGenerator.h
+++ b/Engine/Source/Runtime/Engine/Classes/Particles/Event/ParticleModuleEventGenerator.h
@@ -53,7 +53,7 @@ struct FParticleEvent_GenerateInfo
 	}
 };
 
-UCLASS(HeaderGroup=Particle, editinlinenew, hidecategories=Object, MinimalAPI, meta=(DisplayName = "Event Generator"))
+UCLASS(editinlinenew, hidecategories=Object, MinimalAPI, meta=(DisplayName = "Event Generator"))
 class UParticleModuleEventGenerator : public UParticleModuleEventBase
 {
 	GENERATED_UCLASS_BODY()
diff --git a/Engine/Source/Runtime/Engine/Classes/Particles/Event/ParticleModuleEventReceiverBase.h b/Engine/Source/Runtime/Engine/Classes/Particles/Event/ParticleModuleEventReceiverBase.h
--- a/Engine/Source/Runtime/Engine/Classes/Particles/Event/ParticleModuleEventReceiverBase.h
+++ b/Engine/Source/Runtime/Engine/Classes/Particles/Event/ParticleModuleEventReceiverBase.h
@@ -4,7 +4,7 @@
 #pragma once
 #include "ParticleModuleEventReceiverBase.generated.h"
 
-UCLASS(HeaderGroup=Particle, editinlinenew, hidecategories=Object, abstract)
+UCLASS(editinlinenew, hidecategories=Object, abstract)
 class UParticleModuleEventReceiverBase : public UParticleModuleEventBase
 {
 	GENERATED_UCLASS_BODY()
diff --git a/Engine/Source/Runtime/Engine/Classes/Particles/Event/ParticleModuleEventReceiverKillParticles.h b/Engine/Source/Runtime/Engine/Classes/Particles/Event/ParticleModuleEventReceiverKillParticles.h
--- a/Engine/Source/Runtime/Engine/Classes/Particles/Event/ParticleModuleEventReceiverKillParticles.h
+++ b/Engine/Source/Runtime/Engine/Classes/Particles/Event/ParticleModuleEventReceiverKillParticles.h
@@ -4,7 +4,7 @@
 #pragma once
 #include "ParticleModuleEventReceiverKillParticles.generated.h"
 
-UCLASS(HeaderGroup=Particle, editinlinenew, hidecategories=Object, meta=(DisplayName = "EventReceiver Kill All"))
+UCLASS(editinlinenew, hidecategories=Object, meta=(DisplayName = "EventReceiver Kill All"))
 class UParticleModuleEventReceiverKillParticles : public UParticleModuleEventReceiverBase
 {
 	GENERATED_UCLASS_BODY()
diff --git a/Engine/Source/Runtime/Engine/Classes/Particles/Event/ParticleModuleEventReceiverSpawn.h b/Engine/Source/Runtime/Engine/Classes/Particles/Event/ParticleModuleEventReceiverSpawn.h
--- a/Engine/Source/Runtime/Engine/Classes/Particles/Event/ParticleModuleEventReceiverSpawn.h
+++ b/Engine/Source/Runtime/Engine/Classes/Particles/Event/ParticleModuleEventReceiverSpawn.h
@@ -4,7 +4,7 @@
 #pragma once
 #include "ParticleModuleEventReceiverSpawn.generated.h"
 
-UCLASS(HeaderGroup=Particle, editinlinenew, hidecategories=Object, meta=(DisplayName = "EventReceiver Spawn"))
+UCLASS(editinlinenew, hidecategories=Object, meta=(DisplayName = "EventReceiver Spawn"))
 class UParticleModuleEventReceiverSpawn : public UParticleModuleEventReceiverBase
 {
 	GENERATED_UCLASS_BODY()
diff --git a/Engine/Source/Runtime/Engine/Classes/Particles/Event/ParticleModuleEventSendToGame.h b/Engine/Source/Runtime/Engine/Classes/Particles/Event/ParticleModuleEventSendToGame.h
--- a/Engine/Source/Runtime/Engine/Classes/Particles/Event/ParticleModuleEventSendToGame.h
+++ b/Engine/Source/Runtime/Engine/Classes/Particles/Event/ParticleModuleEventSendToGame.h
@@ -4,7 +4,7 @@
 #pragma once
 #include "ParticleModuleEventSendToGame.generated.h"
 
-UCLASS(HeaderGroup=Particle, abstract, editinlinenew, hidecategories=Object)
+UCLASS(abstract, editinlinenew, hidecategories=Object)
 class UParticleModuleEventSendToGame : public UObject
 {
 	GENERATED_UCLASS_BODY()
diff --git a/Engine/Source/Runtime/Engine/Classes/Particles/Kill/ParticleModuleKillBase.h b/Engine/Source/Runtime/Engine/Classes/Particles/Kill/ParticleModuleKillBase.h
--- a/Engine/Source/Runtime/Engine/Classes/Particles/Kill/ParticleModuleKillBase.h
+++ b/Engine/Source/Runtime/Engine/Classes/Particles/Kill/ParticleModuleKillBase.h
@@ -4,7 +4,7 @@
 #pragma once
 #include "ParticleModuleKillBase.generated.h"
 
-UCLASS(HeaderGroup=Particle, editinlinenew, hidecategories=Object, abstract, meta=(DisplayName = "Kill"))
+UCLASS(editinlinenew, hidecategories=Object, abstract, meta=(DisplayName = "Kill"))
 class UParticleModuleKillBase : public UParticleModule
 {
 	GENERATED_UCLASS_BODY()
diff --git a/Engine/Source/Runtime/Engine/Classes/Particles/Kill/ParticleModuleKillBox.h b/Engine/Source/Runtime/Engine/Classes/Particles/Kill/ParticleModuleKillBox.h
--- a/Engine/Source/Runtime/Engine/Classes/Particles/Kill/ParticleModuleKillBox.h
+++ b/Engine/Source/Runtime/Engine/Classes/Particles/Kill/ParticleModuleKillBox.h
@@ -4,7 +4,7 @@
 #pragma once
 #include "ParticleModuleKillBox.generated.h"
 
-UCLASS(HeaderGroup=Particle, editinlinenew, hidecategories=Object, meta=(DisplayName = "Kill Box"))
+UCLASS(editinlinenew, hidecategories=Object, meta=(DisplayName = "Kill Box"))
 class UParticleModuleKillBox : public UParticleModuleKillBase
 {
 	GENERATED_UCLASS_BODY()
diff --git a/Engine/Source/Runtime/Engine/Classes/Particles/Kill/ParticleModuleKillHeight.h b/Engine/Source/Runtime/Engine/Classes/Particles/Kill/ParticleModuleKillHeight.h
--- a/Engine/Source/Runtime/Engine/Classes/Particles/Kill/ParticleModuleKillHeight.h
+++ b/Engine/Source/Runtime/Engine/Classes/Particles/Kill/ParticleModuleKillHeight.h
@@ -4,7 +4,7 @@
 #pragma once
 #include "ParticleModuleKillHeight.generated.h"
 
-UCLASS(HeaderGroup=Particle, editinlinenew, hidecategories=Object, meta=(DisplayName = "Kill Height"))
+UCLASS(editinlinenew, hidecategories=Object, meta=(DisplayName = "Kill Height"))
 class UParticleModuleKillHeight : public UParticleModuleKillBase
 {
 	GENERATED_UCLASS_BODY()
diff --git a/Engine/Source/Runtime/Engine/Classes/Particles/Lifetime/ParticleModuleLifetime.h b/Engine/Source/Runtime/Engine/Classes/Particles/Lifetime/ParticleModuleLifetime.h
--- a/Engine/Source/Runtime/Engine/Classes/Particles/Lifetime/ParticleModuleLifetime.h
+++ b/Engine/Source/Runtime/Engine/Classes/Particles/Lifetime/ParticleModuleLifetime.h
@@ -4,7 +4,7 @@
 #pragma once
 #include "ParticleModuleLifetime.generated.h"
 
-UCLASS(HeaderGroup=Particle, editinlinenew, hidecategories=Object, MinimalAPI, meta=(DisplayName = "Lifetime"))
+UCLASS(editinlinenew, hidecategories=Object, MinimalAPI, meta=(DisplayName = "Lifetime"))
 class UParticleModuleLifetime : public UParticleModuleLifetimeBase
 {
 	GENERATED_UCLASS_BODY()
diff --git a/Engine/Source/Runtime/Engine/Classes/Particles/Lifetime/ParticleModuleLifetimeBase.h b/Engine/Source/Runtime/Engine/Classes/Particles/Lifetime/ParticleModuleLifetimeBase.h
--- a/Engine/Source/Runtime/Engine/Classes/Particles/Lifetime/ParticleModuleLifetimeBase.h
+++ b/Engine/Source/Runtime/Engine/Classes/Particles/Lifetime/ParticleModuleLifetimeBase.h
@@ -4,7 +4,7 @@
 #pragma once
 #include "ParticleModuleLifetimeBase.generated.h"
 
-UCLASS(HeaderGroup=Particle, editinlinenew, hidecategories=Object, abstract, meta=(DisplayName = "Lifetime"))
+UCLASS(editinlinenew, hidecategories=Object, abstract, meta=(DisplayName = "Lifetime"))
 class UParticleModuleLifetimeBase : public UParticleModule
 {
 	GENERATED_UCLASS_BODY()
diff --git a/Engine/Source/Runtime/Engine/Classes/Particles/Lifetime/ParticleModuleLifetime_Seeded.h b/Engine/Source/Runtime/Engine/Classes/Particles/Lifetime/ParticleModuleLifetime_Seeded.h
--- a/Engine/Source/Runtime/Engine/Classes/Particles/Lifetime/ParticleModuleLifetime_Seeded.h
+++ b/Engine/Source/Runtime/Engine/Classes/Particles/Lifetime/ParticleModuleLifetime_Seeded.h
@@ -4,7 +4,7 @@
 #pragma once
 #include "ParticleModuleLifetime_Seeded.generated.h"
 
-UCLASS(HeaderGroup=Particle, editinlinenew, hidecategories=Object, meta=(DisplayName = "Lifetime (Seed)"))
+UCLASS(editinlinenew, hidecategories=Object, meta=(DisplayName = "Lifetime (Seed)"))
 class UParticleModuleLifetime_Seeded : public UParticleModuleLifetime
 {
 	GENERATED_UCLASS_BODY()
diff --git a/Engine/Source/Runtime/Engine/Classes/Particles/Light/ParticleModuleLight.h b/Engine/Source/Runtime/Engine/Classes/Particles/Light/ParticleModuleLight.h
--- a/Engine/Source/Runtime/Engine/Classes/Particles/Light/ParticleModuleLight.h
+++ b/Engine/Source/Runtime/Engine/Classes/Particles/Light/ParticleModuleLight.h
@@ -4,7 +4,7 @@
 #pragma once
 #include "ParticleModuleLight.generated.h"
 
-UCLASS(HeaderGroup=Particle, editinlinenew, hidecategories=Object, MinimalAPI, meta=(DisplayName = "Light"))
+UCLASS(editinlinenew, hidecategories=Object, MinimalAPI, meta=(DisplayName = "Light"))
 class UParticleModuleLight : public UParticleModuleLightBase
 {
 	GENERATED_UCLASS_BODY()
diff --git a/Engine/Source/Runtime/Engine/Classes/Particles/Light/ParticleModuleLightBase.h b/Engine/Source/Runtime/Engine/Classes/Particles/Light/ParticleModuleLightBase.h
--- a/Engine/Source/Runtime/Engine/Classes/Particles/Light/ParticleModuleLightBase.h
+++ b/Engine/Source/Runtime/Engine/Classes/Particles/Light/ParticleModuleLightBase.h
@@ -4,7 +4,7 @@
 #pragma once
 #include "ParticleModuleLightBase.generated.h"
 
-UCLASS(HeaderGroup=Particle, editinlinenew, hidecategories=Object, abstract, meta=(DisplayName = "Light"))
+UCLASS(editinlinenew, hidecategories=Object, abstract, meta=(DisplayName = "Light"))
 class UParticleModuleLightBase : public UParticleModule
 {
 	GENERATED_UCLASS_BODY()
diff --git a/Engine/Source/Runtime/Engine/Classes/Particles/Light/ParticleModuleLight_Seeded.h b/Engine/Source/Runtime/Engine/Classes/Particles/Light/ParticleModuleLight_Seeded.h
--- a/Engine/Source/Runtime/Engine/Classes/Particles/Light/ParticleModuleLight_Seeded.h
+++ b/Engine/Source/Runtime/Engine/Classes/Particles/Light/ParticleModuleLight_Seeded.h
@@ -4,7 +4,7 @@
 #pragma once
 #include "ParticleModuleLight_Seeded.generated.h"
 
-UCLASS(HeaderGroup=Particle, editinlinenew, hidecategories=Object, MinimalAPI, meta=(DisplayName = "Light (Seed)"))
+UCLASS(editinlinenew, hidecategories=Object, MinimalAPI, meta=(DisplayName = "Light (Seed)"))
 class UParticleModuleLight_Seeded : public UParticleModuleLight
 {
 	GENERATED_UCLASS_BODY()
diff --git a/Engine/Source/Runtime/Engine/Classes/Particles/Location/ParticleModuleLocation.h b/Engine/Source/Runtime/Engine/Classes/Particles/Location/ParticleModuleLocation.h
--- a/Engine/Source/Runtime/Engine/Classes/Particles/Location/ParticleModuleLocation.h
+++ b/Engine/Source/Runtime/Engine/Classes/Particles/Location/ParticleModuleLocation.h
@@ -4,7 +4,7 @@
 #pragma once
 #include "ParticleModuleLocation.generated.h"
 
-UCLASS(HeaderGroup=Particle, editinlinenew, hidecategories=Object, MinimalAPI, meta=(DisplayName = "Initial Location"))
+UCLASS(editinlinenew, hidecategories=Object, MinimalAPI, meta=(DisplayName = "Initial Location"))
 class UParticleModuleLocation : public UParticleModuleLocationBase
 {
 	GENERATED_UCLASS_BODY()
diff --git a/Engine/Source/Runtime/Engine/Classes/Particles/Location/ParticleModuleLocationBase.h b/Engine/Source/Runtime/Engine/Classes/Particles/Location/ParticleModuleLocationBase.h
--- a/Engine/Source/Runtime/Engine/Classes/Particles/Location/ParticleModuleLocationBase.h
+++ b/Engine/Source/Runtime/Engine/Classes/Particles/Location/ParticleModuleLocationBase.h
@@ -4,7 +4,7 @@
 #pragma once
 #include "ParticleModuleLocationBase.generated.h"
 
-UCLASS(HeaderGroup=Particle, editinlinenew, hidecategories=Object, abstract, meta=(DisplayName = "Location"))
+UCLASS(editinlinenew, hidecategories=Object, abstract, meta=(DisplayName = "Location"))
 class UParticleModuleLocationBase : public UParticleModule
 {
 	GENERATED_UCLASS_BODY()
diff --git a/Engine/Source/Runtime/Engine/Classes/Particles/Location/ParticleModuleLocationBoneSocket.h b/Engine/Source/Runtime/Engine/Classes/Particles/Location/ParticleModuleLocationBoneSocket.h
--- a/Engine/Source/Runtime/Engine/Classes/Particles/Location/ParticleModuleLocationBoneSocket.h
+++ b/Engine/Source/Runtime/Engine/Classes/Particles/Location/ParticleModuleLocationBoneSocket.h
@@ -41,7 +41,7 @@ struct FLocationBoneSocketInfo
 
 };
 
-UCLASS(HeaderGroup=Particle, editinlinenew, hidecategories=Object, meta=(DisplayName = "Bone/Socket Location"))
+UCLASS(editinlinenew, hidecategories=Object, meta=(DisplayName = "Bone/Socket Location"))
 class UParticleModuleLocationBoneSocket : public UParticleModuleLocationBase
 {
 	GENERATED_UCLASS_BODY()
diff --git a/Engine/Source/Runtime/Engine/Classes/Particles/Location/ParticleModuleLocationDirect.h b/Engine/Source/Runtime/Engine/Classes/Particles/Location/ParticleModuleLocationDirect.h
--- a/Engine/Source/Runtime/Engine/Classes/Particles/Location/ParticleModuleLocationDirect.h
+++ b/Engine/Source/Runtime/Engine/Classes/Particles/Location/ParticleModuleLocationDirect.h
@@ -10,7 +10,7 @@
 #pragma once
 #include "ParticleModuleLocationDirect.generated.h"
 
-UCLASS(HeaderGroup=Particle, editinlinenew, hidecategories=Object, MinimalAPI, meta=(DisplayName = "Direct Location"))
+UCLASS(editinlinenew, hidecategories=Object, MinimalAPI, meta=(DisplayName = "Direct Location"))
 class UParticleModuleLocationDirect : public UParticleModuleLocationBase
 {
 	GENERATED_UCLASS_BODY()
diff --git a/Engine/Source/Runtime/Engine/Classes/Particles/Location/ParticleModuleLocationEmitter.h b/Engine/Source/Runtime/Engine/Classes/Particles/Location/ParticleModuleLocationEmitter.h
--- a/Engine/Source/Runtime/Engine/Classes/Particles/Location/ParticleModuleLocationEmitter.h
+++ b/Engine/Source/Runtime/Engine/Classes/Particles/Location/ParticleModuleLocationEmitter.h
@@ -19,7 +19,7 @@ enum ELocationEmitterSelectionMethod
 	ELESM_MAX,
 };
 
-UCLASS(HeaderGroup=Particle, editinlinenew, hidecategories=Object, MinimalAPI, meta=(DisplayName = "Emitter Initial Location"))
+UCLASS(editinlinenew, hidecategories=Object, MinimalAPI, meta=(DisplayName = "Emitter Initial Location"))
 class UParticleModuleLocationEmitter : public UParticleModuleLocationBase
 {
 	GENERATED_UCLASS_BODY()
diff --git a/Engine/Source/Runtime/Engine/Classes/Particles/Location/ParticleModuleLocationEmitterDirect.h b/Engine/Source/Runtime/Engine/Classes/Particles/Location/ParticleModuleLocationEmitterDirect.h
--- a/Engine/Source/Runtime/Engine/Classes/Particles/Location/ParticleModuleLocationEmitterDirect.h
+++ b/Engine/Source/Runtime/Engine/Classes/Particles/Location/ParticleModuleLocationEmitterDirect.h
@@ -11,7 +11,7 @@
 #pragma once
 #include "ParticleModuleLocationEmitterDirect.generated.h"
 
-UCLASS(HeaderGroup=Particle, editinlinenew, hidecategories=Object, meta=(DisplayName = "Emitter Direct Location"))
+UCLASS(editinlinenew, hidecategories=Object, meta=(DisplayName = "Emitter Direct Location"))
 class UParticleModuleLocationEmitterDirect : public UParticleModuleLocationBase
 {
 	GENERATED_UCLASS_BODY()
diff --git a/Engine/Source/Runtime/Engine/Classes/Particles/Location/ParticleModuleLocationPrimitiveBase.h b/Engine/Source/Runtime/Engine/Classes/Particles/Location/ParticleModuleLocationPrimitiveBase.h
--- a/Engine/Source/Runtime/Engine/Classes/Particles/Location/ParticleModuleLocationPrimitiveBase.h
+++ b/Engine/Source/Runtime/Engine/Classes/Particles/Location/ParticleModuleLocationPrimitiveBase.h
@@ -8,7 +8,7 @@
 #pragma once
 #include "ParticleModuleLocationPrimitiveBase.generated.h"
 
-UCLASS(HeaderGroup=Particle, editinlinenew, hidecategories=Object)
+UCLASS(editinlinenew, hidecategories=Object)
 class UParticleModuleLocationPrimitiveBase : public UParticleModuleLocationBase
 {
 	GENERATED_UCLASS_BODY()
diff --git a/Engine/Source/Runtime/Engine/Classes/Particles/Location/ParticleModuleLocationPrimitiveCylinder.h b/Engine/Source/Runtime/Engine/Classes/Particles/Location/ParticleModuleLocationPrimitiveCylinder.h
--- a/Engine/Source/Runtime/Engine/Classes/Particles/Location/ParticleModuleLocationPrimitiveCylinder.h
+++ b/Engine/Source/Runtime/Engine/Classes/Particles/Location/ParticleModuleLocationPrimitiveCylinder.h
@@ -17,7 +17,7 @@ enum CylinderHeightAxis
 	PMLPC_HEIGHTAXIS_MAX,
 };
 
-UCLASS(HeaderGroup=Particle, editinlinenew, hidecategories=Object, MinimalAPI)
+UCLASS(editinlinenew, hidecategories=Object, MinimalAPI)
 class UParticleModuleLocationPrimitiveCylinder : public UParticleModuleLocationPrimitiveBase
 {
 	GENERATED_UCLASS_BODY()
diff --git a/Engine/Source/Runtime/Engine/Classes/Particles/Location/ParticleModuleLocationPrimitiveCylinder_Seeded.h b/Engine/Source/Runtime/Engine/Classes/Particles/Location/ParticleModuleLocationPrimitiveCylinder_Seeded.h
--- a/Engine/Source/Runtime/Engine/Classes/Particles/Location/ParticleModuleLocationPrimitiveCylinder_Seeded.h
+++ b/Engine/Source/Runtime/Engine/Classes/Particles/Location/ParticleModuleLocationPrimitiveCylinder_Seeded.h
@@ -4,7 +4,7 @@
 #pragma once
 #include "ParticleModuleLocationPrimitiveCylinder_Seeded.generated.h"
 
-UCLASS(HeaderGroup=Particle, editinlinenew, hidecategories=Object, meta=(DisplayName = "Cylinder (Seed)"))
+UCLASS(editinlinenew, hidecategories=Object, meta=(DisplayName = "Cylinder (Seed)"))
 class UParticleModuleLocationPrimitiveCylinder_Seeded : public UParticleModuleLocationPrimitiveCylinder
 {
 	GENERATED_UCLASS_BODY()
diff --git a/Engine/Source/Runtime/Engine/Classes/Particles/Location/ParticleModuleLocationPrimitiveSphere.h b/Engine/Source/Runtime/Engine/Classes/Particles/Location/ParticleModuleLocationPrimitiveSphere.h
--- a/Engine/Source/Runtime/Engine/Classes/Particles/Location/ParticleModuleLocationPrimitiveSphere.h
+++ b/Engine/Source/Runtime/Engine/Classes/Particles/Location/ParticleModuleLocationPrimitiveSphere.h
@@ -8,7 +8,7 @@
 #pragma once
 #include "ParticleModuleLocationPrimitiveSphere.generated.h"
 
-UCLASS(HeaderGroup=Particle, editinlinenew, hidecategories=Object, MinimalAPI, meta=(DisplayName = "Sphere"))
+UCLASS(editinlinenew, hidecategories=Object, MinimalAPI, meta=(DisplayName = "Sphere"))
 class UParticleModuleLocationPrimitiveSphere : public UParticleModuleLocationPrimitiveBase
 {
 	GENERATED_UCLASS_BODY()
diff --git a/Engine/Source/Runtime/Engine/Classes/Particles/Location/ParticleModuleLocationPrimitiveSphere_Seeded.h b/Engine/Source/Runtime/Engine/Classes/Particles/Location/ParticleModuleLocationPrimitiveSphere_Seeded.h
--- a/Engine/Source/Runtime/Engine/Classes/Particles/Location/ParticleModuleLocationPrimitiveSphere_Seeded.h
+++ b/Engine/Source/Runtime/Engine/Classes/Particles/Location/ParticleModuleLocationPrimitiveSphere_Seeded.h
@@ -4,7 +4,7 @@
 #pragma once
 #include "ParticleModuleLocationPrimitiveSphere_Seeded.generated.h"
 
-UCLASS(HeaderGroup=Particle, editinlinenew, hidecategories=Object, meta=(DisplayName = "Sphere (Seed)"))
+UCLASS(editinlinenew, hidecategories=Object, meta=(DisplayName = "Sphere (Seed)"))
 class UParticleModuleLocationPrimitiveSphere_Seeded : public UParticleModuleLocationPrimitiveSphere
 {
 	GENERATED_UCLASS_BODY()
diff --git a/Engine/Source/Runtime/Engine/Classes/Particles/Location/ParticleModuleLocationPrimitiveTriangle.h b/Engine/Source/Runtime/Engine/Classes/Particles/Location/ParticleModuleLocationPrimitiveTriangle.h
--- a/Engine/Source/Runtime/Engine/Classes/Particles/Location/ParticleModuleLocationPrimitiveTriangle.h
+++ b/Engine/Source/Runtime/Engine/Classes/Particles/Location/ParticleModuleLocationPrimitiveTriangle.h
@@ -8,7 +8,7 @@
 #pragma once
 #include "ParticleModuleLocationPrimitiveTriangle.generated.h"
 
-UCLASS(HeaderGroup=Particle, editinlinenew, hidecategories=Object, MinimalAPI, meta=(DisplayName = "Triangle"))
+UCLASS(editinlinenew, hidecategories=Object, MinimalAPI, meta=(DisplayName = "Triangle"))
 class UParticleModuleLocationPrimitiveTriangle : public UParticleModuleLocationBase
 {
 	GENERATED_UCLASS_BODY()
diff --git a/Engine/Source/Runtime/Engine/Classes/Particles/Location/ParticleModuleLocationSkelVertSurface.h b/Engine/Source/Runtime/Engine/Classes/Particles/Location/ParticleModuleLocationSkelVertSurface.h
--- a/Engine/Source/Runtime/Engine/Classes/Particles/Location/ParticleModuleLocationSkelVertSurface.h
+++ b/Engine/Source/Runtime/Engine/Classes/Particles/Location/ParticleModuleLocationSkelVertSurface.h
@@ -12,7 +12,7 @@ enum ELocationSkelVertSurfaceSource
 	VERTSURFACESOURCE_MAX,
 };
 
-UCLASS(HeaderGroup=Particle, editinlinenew, hidecategories=Object, meta=(DisplayName = "Skel Vert/Surf Location"))
+UCLASS(editinlinenew, hidecategories=Object, meta=(DisplayName = "Skel Vert/Surf Location"))
 class UParticleModuleLocationSkelVertSurface : public UParticleModuleLocationBase
 {
 	GENERATED_UCLASS_BODY()
diff --git a/Engine/Source/Runtime/Engine/Classes/Particles/Location/ParticleModuleLocationWorldOffset.h b/Engine/Source/Runtime/Engine/Classes/Particles/Location/ParticleModuleLocationWorldOffset.h
--- a/Engine/Source/Runtime/Engine/Classes/Particles/Location/ParticleModuleLocationWorldOffset.h
+++ b/Engine/Source/Runtime/Engine/Classes/Particles/Location/ParticleModuleLocationWorldOffset.h
@@ -4,7 +4,7 @@
 #pragma once
 #include "ParticleModuleLocationWorldOffset.generated.h"
 
-UCLASS(HeaderGroup=Particle, editinlinenew, meta=(DisplayName = "World Offset"))
+UCLASS(editinlinenew, meta=(DisplayName = "World Offset"))
 class UParticleModuleLocationWorldOffset : public UParticleModuleLocation
 {
 	GENERATED_UCLASS_BODY()
diff --git a/Engine/Source/Runtime/Engine/Classes/Particles/Location/ParticleModuleLocationWorldOffset_Seeded.h b/Engine/Source/Runtime/Engine/Classes/Particles/Location/ParticleModuleLocationWorldOffset_Seeded.h
--- a/Engine/Source/Runtime/Engine/Classes/Particles/Location/ParticleModuleLocationWorldOffset_Seeded.h
+++ b/Engine/Source/Runtime/Engine/Classes/Particles/Location/ParticleModuleLocationWorldOffset_Seeded.h
@@ -4,7 +4,7 @@
 #pragma once
 #include "ParticleModuleLocationWorldOffset_Seeded.generated.h"
 
-UCLASS(HeaderGroup=Particle, editinlinenew, meta=(DisplayName = "World Offset (Seed)"))
+UCLASS(editinlinenew, meta=(DisplayName = "World Offset (Seed)"))
 class UParticleModuleLocationWorldOffset_Seeded : public UParticleModuleLocationWorldOffset
 {
 	GENERATED_UCLASS_BODY()
diff --git a/Engine/Source/Runtime/Engine/Classes/Particles/Location/ParticleModuleLocation_Seeded.h b/Engine/Source/Runtime/Engine/Classes/Particles/Location/ParticleModuleLocation_Seeded.h
--- a/Engine/Source/Runtime/Engine/Classes/Particles/Location/ParticleModuleLocation_Seeded.h
+++ b/Engine/Source/Runtime/Engine/Classes/Particles/Location/ParticleModuleLocation_Seeded.h
@@ -4,7 +4,7 @@
 #pragma once
 #include "ParticleModuleLocation_Seeded.generated.h"
 
-UCLASS(HeaderGroup=Particle, editinlinenew, hidecategories=Object, meta=(DisplayName = "Initial Location (Seed)"))
+UCLASS(editinlinenew, hidecategories=Object, meta=(DisplayName = "Initial Location (Seed)"))
 class UParticleModuleLocation_Seeded : public UParticleModuleLocation
 {
 	GENERATED_UCLASS_BODY()
diff --git a/Engine/Source/Runtime/Engine/Classes/Particles/Location/ParticleModuleSourceMovement.h b/Engine/Source/Runtime/Engine/Classes/Particles/Location/ParticleModuleSourceMovement.h
--- a/Engine/Source/Runtime/Engine/Classes/Particles/Location/ParticleModuleSourceMovement.h
+++ b/Engine/Source/Runtime/Engine/Classes/Particles/Location/ParticleModuleSourceMovement.h
@@ -4,7 +4,7 @@
 #pragma once
 #include "ParticleModuleSourceMovement.generated.h"
 
-UCLASS(HeaderGroup=Particle, editinlinenew, hidecategories=Object, meta=(DisplayName = "Source Movement"))
+UCLASS(editinlinenew, hidecategories=Object, meta=(DisplayName = "Source Movement"))
 class UParticleModuleSourceMovement : public UParticleModuleLocationBase
 {
 	GENERATED_UCLASS_BODY()
diff --git a/Engine/Source/Runtime/Engine/Classes/Particles/Material/ParticleModuleMaterialBase.h b/Engine/Source/Runtime/Engine/Classes/Particles/Material/ParticleModuleMaterialBase.h
--- a/Engine/Source/Runtime/Engine/Classes/Particles/Material/ParticleModuleMaterialBase.h
+++ b/Engine/Source/Runtime/Engine/Classes/Particles/Material/ParticleModuleMaterialBase.h
@@ -4,7 +4,7 @@
 #pragma once
 #include "ParticleModuleMaterialBase.generated.h"
 
-UCLASS(HeaderGroup=Particle, editinlinenew, hidecategories=Object, abstract, meta=(DisplayName = "Material"))
+UCLASS(editinlinenew, hidecategories=Object, abstract, meta=(DisplayName = "Material"))
 class UParticleModuleMaterialBase : public UParticleModule
 {
 	GENERATED_UCLASS_BODY()
diff --git a/Engine/Source/Runtime/Engine/Classes/Particles/Material/ParticleModuleMeshMaterial.h b/Engine/Source/Runtime/Engine/Classes/Particles/Material/ParticleModuleMeshMaterial.h
--- a/Engine/Source/Runtime/Engine/Classes/Particles/Material/ParticleModuleMeshMaterial.h
+++ b/Engine/Source/Runtime/Engine/Classes/Particles/Material/ParticleModuleMeshMaterial.h
@@ -4,7 +4,7 @@
 #pragma once
 #include "ParticleModuleMeshMaterial.generated.h"
 
-UCLASS(HeaderGroup=Particle, editinlinenew, hidecategories=Object, MinimalAPI, meta=(DisplayName = "Mesh Material"))
+UCLASS(editinlinenew, hidecategories=Object, MinimalAPI, meta=(DisplayName = "Mesh Material"))
 class UParticleModuleMeshMaterial : public UParticleModuleMaterialBase
 {
 	GENERATED_UCLASS_BODY()
diff --git a/Engine/Source/Runtime/Engine/Classes/Particles/Modules/Location/ParticleModulePivotOffset.h b/Engine/Source/Runtime/Engine/Classes/Particles/Modules/Location/ParticleModulePivotOffset.h
--- a/Engine/Source/Runtime/Engine/Classes/Particles/Modules/Location/ParticleModulePivotOffset.h
+++ b/Engine/Source/Runtime/Engine/Classes/Particles/Modules/Location/ParticleModulePivotOffset.h
@@ -3,7 +3,7 @@
 #pragma once
 #include "ParticleModulePivotOffset.generated.h"
 
-UCLASS(HeaderGroup=Particle, editinlinenew, hidecategories=Object, MinimalAPI, DisplayName="Pivot Offset")
+UCLASS(editinlinenew, hidecategories=Object, MinimalAPI, DisplayName="Pivot Offset")
 class UParticleModulePivotOffset : public UParticleModuleLocationBase
 {
 	GENERATED_UCLASS_BODY()
diff --git a/Engine/Source/Runtime/Engine/Classes/Particles/Orbit/ParticleModuleOrbit.h b/Engine/Source/Runtime/Engine/Classes/Particles/Orbit/ParticleModuleOrbit.h
--- a/Engine/Source/Runtime/Engine/Classes/Particles/Orbit/ParticleModuleOrbit.h
+++ b/Engine/Source/Runtime/Engine/Classes/Particles/Orbit/ParticleModuleOrbit.h
@@ -53,7 +53,7 @@ struct FOrbitOptions
 	
 };
 
-UCLASS(HeaderGroup=Particle, editinlinenew, hidecategories=(Object, Orbit), MinimalAPI, meta=(DisplayName = "Orbit"))
+UCLASS(editinlinenew, hidecategories=(Object, Orbit), MinimalAPI, meta=(DisplayName = "Orbit"))
 class UParticleModuleOrbit : public UParticleModuleOrbitBase
 {
 	GENERATED_UCLASS_BODY()
diff --git a/Engine/Source/Runtime/Engine/Classes/Particles/Orbit/ParticleModuleOrbitBase.h b/Engine/Source/Runtime/Engine/Classes/Particles/Orbit/ParticleModuleOrbitBase.h
--- a/Engine/Source/Runtime/Engine/Classes/Particles/Orbit/ParticleModuleOrbitBase.h
+++ b/Engine/Source/Runtime/Engine/Classes/Particles/Orbit/ParticleModuleOrbitBase.h
@@ -4,7 +4,7 @@
 #pragma once
 #include "ParticleModuleOrbitBase.generated.h"
 
-UCLASS(HeaderGroup=Particle, editinlinenew, hidecategories=Object, abstract, meta=(DisplayName = "Orbit"))
+UCLASS(editinlinenew, hidecategories=Object, abstract, meta=(DisplayName = "Orbit"))
 class UParticleModuleOrbitBase : public UParticleModule
 {
 	GENERATED_UCLASS_BODY()
diff --git a/Engine/Source/Runtime/Engine/Classes/Particles/Orientation/ParticleModuleOrientationAxisLock.h b/Engine/Source/Runtime/Engine/Classes/Particles/Orientation/ParticleModuleOrientationAxisLock.h
--- a/Engine/Source/Runtime/Engine/Classes/Particles/Orientation/ParticleModuleOrientationAxisLock.h
+++ b/Engine/Source/Runtime/Engine/Classes/Particles/Orientation/ParticleModuleOrientationAxisLock.h
@@ -31,7 +31,7 @@ enum EParticleAxisLock
 	EPAL_MAX,
 };
 
-UCLASS(HeaderGroup=Particle, editinlinenew, hidecategories=Object, MinimalAPI, meta=(DisplayName = "Lock Axis"))
+UCLASS(editinlinenew, hidecategories=Object, MinimalAPI, meta=(DisplayName = "Lock Axis"))
 class UParticleModuleOrientationAxisLock : public UParticleModuleOrientationBase
 {
 	GENERATED_UCLASS_BODY()
diff --git a/Engine/Source/Runtime/Engine/Classes/Particles/Orientation/ParticleModuleOrientationBase.h b/Engine/Source/Runtime/Engine/Classes/Particles/Orientation/ParticleModuleOrientationBase.h
--- a/Engine/Source/Runtime/Engine/Classes/Particles/Orientation/ParticleModuleOrientationBase.h
+++ b/Engine/Source/Runtime/Engine/Classes/Particles/Orientation/ParticleModuleOrientationBase.h
@@ -4,7 +4,7 @@
 #pragma once
 #include "ParticleModuleOrientationBase.generated.h"
 
-UCLASS(HeaderGroup=Particle, editinlinenew, hidecategories=Object, abstract, meta=(DisplayName = "Orientation"))
+UCLASS(editinlinenew, hidecategories=Object, abstract, meta=(DisplayName = "Orientation"))
 class UParticleModuleOrientationBase : public UParticleModule
 {
 	GENERATED_UCLASS_BODY()
diff --git a/Engine/Source/Runtime/Engine/Classes/Particles/Parameter/ParticleModuleParameterBase.h b/Engine/Source/Runtime/Engine/Classes/Particles/Parameter/ParticleModuleParameterBase.h
--- a/Engine/Source/Runtime/Engine/Classes/Particles/Parameter/ParticleModuleParameterBase.h
+++ b/Engine/Source/Runtime/Engine/Classes/Particles/Parameter/ParticleModuleParameterBase.h
@@ -4,7 +4,7 @@
 #pragma once
 #include "ParticleModuleParameterBase.generated.h"
 
-UCLASS(HeaderGroup=Particle, editinlinenew, hidecategories=Object, abstract, meta=(DisplayName = "Parameter"))
+UCLASS(editinlinenew, hidecategories=Object, abstract, meta=(DisplayName = "Parameter"))
 class UParticleModuleParameterBase : public UParticleModule
 {
 	GENERATED_UCLASS_BODY()
diff --git a/Engine/Source/Runtime/Engine/Classes/Particles/Parameter/ParticleModuleParameterDynamic.h b/Engine/Source/Runtime/Engine/Classes/Particles/Parameter/ParticleModuleParameterDynamic.h
--- a/Engine/Source/Runtime/Engine/Classes/Particles/Parameter/ParticleModuleParameterDynamic.h
+++ b/Engine/Source/Runtime/Engine/Classes/Particles/Parameter/ParticleModuleParameterDynamic.h
@@ -78,7 +78,7 @@ struct FEmitterDynamicParameter
 
 };
 
-UCLASS(HeaderGroup=Particle, editinlinenew, hidecategories=Object, MinimalAPI, meta=(DisplayName = "Dynamic"))
+UCLASS(editinlinenew, hidecategories=Object, MinimalAPI, meta=(DisplayName = "Dynamic"))
 class UParticleModuleParameterDynamic : public UParticleModuleParameterBase
 {
 	GENERATED_UCLASS_BODY()
diff --git a/Engine/Source/Runtime/Engine/Classes/Particles/Parameter/ParticleModuleParameterDynamic_Seeded.h b/Engine/Source/Runtime/Engine/Classes/Particles/Parameter/ParticleModuleParameterDynamic_Seeded.h
--- a/Engine/Source/Runtime/Engine/Classes/Particles/Parameter/ParticleModuleParameterDynamic_Seeded.h
+++ b/Engine/Source/Runtime/Engine/Classes/Particles/Parameter/ParticleModuleParameterDynamic_Seeded.h
@@ -4,7 +4,7 @@
 #pragma once
 #include "ParticleModuleParameterDynamic_Seeded.generated.h"
 
-UCLASS(HeaderGroup=Particle, editinlinenew, hidecategories=Object, meta=(DisplayName = "Dynamic (Seed)"))
+UCLASS(editinlinenew, hidecategories=Object, meta=(DisplayName = "Dynamic (Seed)"))
 class UParticleModuleParameterDynamic_Seeded : public UParticleModuleParameterDynamic
 {
 	GENERATED_UCLASS_BODY()
diff --git a/Engine/Source/Runtime/Engine/Classes/Particles/ParticleEmitter.h b/Engine/Source/Runtime/Engine/Classes/Particles/ParticleEmitter.h
--- a/Engine/Source/Runtime/Engine/Classes/Particles/ParticleEmitter.h
+++ b/Engine/Source/Runtime/Engine/Classes/Particles/ParticleEmitter.h
@@ -6,6 +6,7 @@
 //=============================================================================
 
 #pragma once
+#include "ParticleEmitterInstances.h"
 #include "ParticleEmitter.generated.h"
 
 //=============================================================================
@@ -74,7 +75,7 @@ struct FParticleBurst
 	
 };
 
-UCLASS(HeaderGroup=Particle, dependson=UParticleLODLevel, hidecategories=Object, editinlinenew, abstract, MinimalAPI)
+UCLASS(dependson=UParticleLODLevel, hidecategories=Object, editinlinenew, abstract, MinimalAPI)
 class UParticleEmitter : public UObject
 {
 	GENERATED_UCLASS_BODY()
diff --git a/Engine/Source/Runtime/Engine/Classes/Particles/ParticleEventManager.h b/Engine/Source/Runtime/Engine/Classes/Particles/ParticleEventManager.h
--- a/Engine/Source/Runtime/Engine/Classes/Particles/ParticleEventManager.h
+++ b/Engine/Source/Runtime/Engine/Classes/Particles/ParticleEventManager.h
@@ -4,7 +4,7 @@
 #pragma once
 #include "ParticleEventManager.generated.h"
 
-UCLASS(HeaderGroup=Particle, dependson=UParticleSystemComponent, config=Game, NotBlueprintable, notplaceable)
+UCLASS(dependson=UParticleSystemComponent, config=Game, NotBlueprintable, notplaceable)
 class AParticleEventManager : public AActor
 {
 	GENERATED_UCLASS_BODY()
diff --git a/Engine/Source/Runtime/Engine/Classes/Particles/ParticleModule.h b/Engine/Source/Runtime/Engine/Classes/Particles/ParticleModule.h
--- a/Engine/Source/Runtime/Engine/Classes/Particles/ParticleModule.h
+++ b/Engine/Source/Runtime/Engine/Classes/Particles/ParticleModule.h
@@ -119,7 +119,7 @@ struct FParticleRandomSeedInfo
 	
 };
 
-UCLASS(HeaderGroup=Particle, editinlinenew, hidecategories=Object, abstract, MinimalAPI)
+UCLASS(editinlinenew, hidecategories=Object, abstract, MinimalAPI)
 class UParticleModule : public UObject
 {
 	GENERATED_UCLASS_BODY()
diff --git a/Engine/Source/Runtime/Engine/Classes/Particles/ParticleModuleRequired.h b/Engine/Source/Runtime/Engine/Classes/Particles/ParticleModuleRequired.h
--- a/Engine/Source/Runtime/Engine/Classes/Particles/ParticleModuleRequired.h
+++ b/Engine/Source/Runtime/Engine/Classes/Particles/ParticleModuleRequired.h
@@ -27,7 +27,7 @@ enum EEmitterNormalsMode
 	ENM_MAX,
 };
 
-UCLASS(HeaderGroup=Particle, editinlinenew, hidecategories=(Object, Cascade), meta=(DisplayName = "Required"), MinimalAPI)
+UCLASS(editinlinenew, hidecategories=(Object, Cascade), meta=(DisplayName = "Required"), MinimalAPI)
 class UParticleModuleRequired : public UParticleModule
 {
 	GENERATED_UCLASS_BODY()
diff --git a/Engine/Source/Runtime/Engine/Classes/Particles/ParticleSystem.h b/Engine/Source/Runtime/Engine/Classes/Particles/ParticleSystem.h
--- a/Engine/Source/Runtime/Engine/Classes/Particles/ParticleSystem.h
+++ b/Engine/Source/Runtime/Engine/Classes/Particles/ParticleSystem.h
@@ -79,7 +79,7 @@ struct FLODSoloTrack
  * in a system, the designer can create elaborate particle effects that are held in a single system. Once created using
  * Cascade, a ParticleSystem can then be inserted into a level or created in script.
  */
-UCLASS(HeaderGroup=Particle, hidecategories=Object, MinimalAPI, BlueprintType)
+UCLASS(hidecategories=Object, MinimalAPI, BlueprintType)
 class UParticleSystem : public UObject
 {
 	GENERATED_UCLASS_BODY()
diff --git a/Engine/Source/Runtime/Engine/Classes/Particles/ParticleSystemComponent.h b/Engine/Source/Runtime/Engine/Classes/Particles/ParticleSystemComponent.h
--- a/Engine/Source/Runtime/Engine/Classes/Particles/ParticleSystemComponent.h
+++ b/Engine/Source/Runtime/Engine/Classes/Particles/ParticleSystemComponent.h
@@ -4,6 +4,17 @@
 #pragma once
 #include "ParticleSystemComponent.generated.h"
 
+//
+// Forward declarations.
+//
+class UAnimNotifyState;
+class FParticleDynamicData;
+class FParticleSystemSceneProxy;
+class UAnimNotifyState;
+struct FDynamicEmitterDataBase;
+struct FDynamicEmitterReplayDataBase;
+struct FParticleAnimTrailEmitterInstance;
+
 /** Enum for specifying type of a name instance parameter. */
 UENUM()
 enum EParticleSysParamType
@@ -232,7 +243,7 @@ struct FParticleEventKismetData : public FParticleEventData
 /** 
  * A particle emmitter.
  */
-UCLASS(HeaderGroup=Particle, ClassGroup=Rendering, hidecategories=Object, hidecategories=Physics, hidecategories=Collision, showcategories=Trigger, editinlinenew, dependson=(AEmitter, UParticleSystem), meta=(BlueprintSpawnableComponent))
+UCLASS(ClassGroup=Rendering, hidecategories=Object, hidecategories=Physics, hidecategories=Collision, showcategories=Trigger, editinlinenew, dependson=(AEmitter, UParticleSystem), meta=(BlueprintSpawnableComponent))
 class ENGINE_API UParticleSystemComponent : public UPrimitiveComponent
 {
 	GENERATED_UCLASS_BODY()
diff --git a/Engine/Source/Runtime/Engine/Classes/Particles/ParticleSystemReplay.h b/Engine/Source/Runtime/Engine/Classes/Particles/ParticleSystemReplay.h
--- a/Engine/Source/Runtime/Engine/Classes/Particles/ParticleSystemReplay.h
+++ b/Engine/Source/Runtime/Engine/Classes/Particles/ParticleSystemReplay.h
@@ -58,7 +58,7 @@ struct FParticleSystemReplayFrame
 	
 };
 
-UCLASS(HeaderGroup=Particle, hidecategories=Object, AutoExpandCategories=ParticleSystemReplay)
+UCLASS(hidecategories=Object, AutoExpandCategories=ParticleSystemReplay)
 class UParticleSystemReplay : public UObject
 {
 	GENERATED_UCLASS_BODY()
diff --git a/Engine/Source/Runtime/Engine/Classes/Particles/Rotation/ParticleModuleMeshRotation.h b/Engine/Source/Runtime/Engine/Classes/Particles/Rotation/ParticleModuleMeshRotation.h
--- a/Engine/Source/Runtime/Engine/Classes/Particles/Rotation/ParticleModuleMeshRotation.h
+++ b/Engine/Source/Runtime/Engine/Classes/Particles/Rotation/ParticleModuleMeshRotation.h
@@ -4,7 +4,7 @@
 #pragma once
 #include "ParticleModuleMeshRotation.generated.h"
 
-UCLASS(HeaderGroup=Particle, editinlinenew, hidecategories=Object, MinimalAPI, meta=(DisplayName = "Init Mesh Rotation"))
+UCLASS(editinlinenew, hidecategories=Object, MinimalAPI, meta=(DisplayName = "Init Mesh Rotation"))
 class UParticleModuleMeshRotation : public UParticleModuleRotationBase
 {
 	GENERATED_UCLASS_BODY()
diff --git a/Engine/Source/Runtime/Engine/Classes/Particles/Rotation/ParticleModuleMeshRotation_Seeded.h b/Engine/Source/Runtime/Engine/Classes/Particles/Rotation/ParticleModuleMeshRotation_Seeded.h
--- a/Engine/Source/Runtime/Engine/Classes/Particles/Rotation/ParticleModuleMeshRotation_Seeded.h
+++ b/Engine/Source/Runtime/Engine/Classes/Particles/Rotation/ParticleModuleMeshRotation_Seeded.h
@@ -4,7 +4,7 @@
 #pragma once
 #include "ParticleModuleMeshRotation_Seeded.generated.h"
 
-UCLASS(HeaderGroup=Particle, editinlinenew, hidecategories=Object, meta=(DisplayName = "Init Mesh Rot (Seed)"))
+UCLASS(editinlinenew, hidecategories=Object, meta=(DisplayName = "Init Mesh Rot (Seed)"))
 class UParticleModuleMeshRotation_Seeded : public UParticleModuleMeshRotation
 {
 	GENERATED_UCLASS_BODY()
diff --git a/Engine/Source/Runtime/Engine/Classes/Particles/Rotation/ParticleModuleRotation.h b/Engine/Source/Runtime/Engine/Classes/Particles/Rotation/ParticleModuleRotation.h
--- a/Engine/Source/Runtime/Engine/Classes/Particles/Rotation/ParticleModuleRotation.h
+++ b/Engine/Source/Runtime/Engine/Classes/Particles/Rotation/ParticleModuleRotation.h
@@ -4,7 +4,7 @@
 #pragma once
 #include "ParticleModuleRotation.generated.h"
 
-UCLASS(HeaderGroup=Particle, editinlinenew, hidecategories=Object, MinimalAPI, meta=(DisplayName = "Initial Rotation"))
+UCLASS(editinlinenew, hidecategories=Object, MinimalAPI, meta=(DisplayName = "Initial Rotation"))
 class UParticleModuleRotation : public UParticleModuleRotationBase
 {
 	GENERATED_UCLASS_BODY()
diff --git a/Engine/Source/Runtime/Engine/Classes/Particles/Rotation/ParticleModuleRotationBase.h b/Engine/Source/Runtime/Engine/Classes/Particles/Rotation/ParticleModuleRotationBase.h
--- a/Engine/Source/Runtime/Engine/Classes/Particles/Rotation/ParticleModuleRotationBase.h
+++ b/Engine/Source/Runtime/Engine/Classes/Particles/Rotation/ParticleModuleRotationBase.h
@@ -4,7 +4,7 @@
 #pragma once
 #include "ParticleModuleRotationBase.generated.h"
 
-UCLASS(HeaderGroup=Particle, editinlinenew, hidecategories=Object, abstract, meta=(DisplayName = "Rotation"))
+UCLASS(editinlinenew, hidecategories=Object, abstract, meta=(DisplayName = "Rotation"))
 class UParticleModuleRotationBase : public UParticleModule
 {
 	GENERATED_UCLASS_BODY()
diff --git a/Engine/Source/Runtime/Engine/Classes/Particles/Rotation/ParticleModuleRotationOverLifetime.h b/Engine/Source/Runtime/Engine/Classes/Particles/Rotation/ParticleModuleRotationOverLifetime.h
--- a/Engine/Source/Runtime/Engine/Classes/Particles/Rotation/ParticleModuleRotationOverLifetime.h
+++ b/Engine/Source/Runtime/Engine/Classes/Particles/Rotation/ParticleModuleRotationOverLifetime.h
@@ -4,7 +4,7 @@
 #pragma once
 #include "ParticleModuleRotationOverLifetime.generated.h"
 
-UCLASS(HeaderGroup=Particle, editinlinenew, hidecategories=Object, MinimalAPI, meta=(DisplayName = "Rotation/Life"))
+UCLASS(editinlinenew, hidecategories=Object, MinimalAPI, meta=(DisplayName = "Rotation/Life"))
 class UParticleModuleRotationOverLifetime : public UParticleModuleRotationBase
 {
 	GENERATED_UCLASS_BODY()
diff --git a/Engine/Source/Runtime/Engine/Classes/Particles/Rotation/ParticleModuleRotation_Seeded.h b/Engine/Source/Runtime/Engine/Classes/Particles/Rotation/ParticleModuleRotation_Seeded.h
--- a/Engine/Source/Runtime/Engine/Classes/Particles/Rotation/ParticleModuleRotation_Seeded.h
+++ b/Engine/Source/Runtime/Engine/Classes/Particles/Rotation/ParticleModuleRotation_Seeded.h
@@ -4,7 +4,7 @@
 #pragma once
 #include "ParticleModuleRotation_Seeded.generated.h"
 
-UCLASS(HeaderGroup=Particle, editinlinenew, hidecategories=Object, meta=(DisplayName = "Init Rotation (Seed)"))
+UCLASS(editinlinenew, hidecategories=Object, meta=(DisplayName = "Init Rotation (Seed)"))
 class UParticleModuleRotation_Seeded : public UParticleModuleRotation
 {
 	GENERATED_UCLASS_BODY()
diff --git a/Engine/Source/Runtime/Engine/Classes/Particles/RotationRate/ParticleModuleMeshRotationRate.h b/Engine/Source/Runtime/Engine/Classes/Particles/RotationRate/ParticleModuleMeshRotationRate.h
--- a/Engine/Source/Runtime/Engine/Classes/Particles/RotationRate/ParticleModuleMeshRotationRate.h
+++ b/Engine/Source/Runtime/Engine/Classes/Particles/RotationRate/ParticleModuleMeshRotationRate.h
@@ -4,7 +4,7 @@
 #pragma once
 #include "ParticleModuleMeshRotationRate.generated.h"
 
-UCLASS(HeaderGroup=Particle, editinlinenew, hidecategories=Object, MinimalAPI, meta=(DisplayName = "Init Mesh Rotation Rate"))
+UCLASS(editinlinenew, hidecategories=Object, MinimalAPI, meta=(DisplayName = "Init Mesh Rotation Rate"))
 class UParticleModuleMeshRotationRate : public UParticleModuleRotationRateBase
 {
 	GENERATED_UCLASS_BODY()
diff --git a/Engine/Source/Runtime/Engine/Classes/Particles/RotationRate/ParticleModuleMeshRotationRateMultiplyLife.h b/Engine/Source/Runtime/Engine/Classes/Particles/RotationRate/ParticleModuleMeshRotationRateMultiplyLife.h
--- a/Engine/Source/Runtime/Engine/Classes/Particles/RotationRate/ParticleModuleMeshRotationRateMultiplyLife.h
+++ b/Engine/Source/Runtime/Engine/Classes/Particles/RotationRate/ParticleModuleMeshRotationRateMultiplyLife.h
@@ -4,7 +4,7 @@
 #pragma once
 #include "ParticleModuleMeshRotationRateMultiplyLife.generated.h"
 
-UCLASS(HeaderGroup=Particle, editinlinenew, hidecategories=Object, meta=(DisplayName = "Mesh Rotation Rate * Life"))
+UCLASS(editinlinenew, hidecategories=Object, meta=(DisplayName = "Mesh Rotation Rate * Life"))
 class UParticleModuleMeshRotationRateMultiplyLife : public UParticleModuleRotationRateBase
 {
 	GENERATED_UCLASS_BODY()
diff --git a/Engine/Source/Runtime/Engine/Classes/Particles/RotationRate/ParticleModuleMeshRotationRateOverLife.h b/Engine/Source/Runtime/Engine/Classes/Particles/RotationRate/ParticleModuleMeshRotationRateOverLife.h
--- a/Engine/Source/Runtime/Engine/Classes/Particles/RotationRate/ParticleModuleMeshRotationRateOverLife.h
+++ b/Engine/Source/Runtime/Engine/Classes/Particles/RotationRate/ParticleModuleMeshRotationRateOverLife.h
@@ -4,7 +4,7 @@
 #pragma once
 #include "ParticleModuleMeshRotationRateOverLife.generated.h"
 
-UCLASS(HeaderGroup=Particle, editinlinenew, hidecategories=Object, meta=(DisplayName = "Mesh Rotation Rate over Life"))
+UCLASS(editinlinenew, hidecategories=Object, meta=(DisplayName = "Mesh Rotation Rate over Life"))
 class UParticleModuleMeshRotationRateOverLife : public UParticleModuleRotationRateBase
 {
 	GENERATED_UCLASS_BODY()
diff --git a/Engine/Source/Runtime/Engine/Classes/Particles/RotationRate/ParticleModuleMeshRotationRate_Seeded.h b/Engine/Source/Runtime/Engine/Classes/Particles/RotationRate/ParticleModuleMeshRotationRate_Seeded.h
--- a/Engine/Source/Runtime/Engine/Classes/Particles/RotationRate/ParticleModuleMeshRotationRate_Seeded.h
+++ b/Engine/Source/Runtime/Engine/Classes/Particles/RotationRate/ParticleModuleMeshRotationRate_Seeded.h
@@ -4,7 +4,7 @@
 #pragma once
 #include "ParticleModuleMeshRotationRate_Seeded.generated.h"
 
-UCLASS(HeaderGroup=Particle, editinlinenew, hidecategories=Object, meta=(DisplayName = "Init Mesh Rotation Rate (Seed)"))
+UCLASS(editinlinenew, hidecategories=Object, meta=(DisplayName = "Init Mesh Rotation Rate (Seed)"))
 class UParticleModuleMeshRotationRate_Seeded : public UParticleModuleMeshRotationRate
 {
 	GENERATED_UCLASS_BODY()
diff --git a/Engine/Source/Runtime/Engine/Classes/Particles/RotationRate/ParticleModuleRotationRate.h b/Engine/Source/Runtime/Engine/Classes/Particles/RotationRate/ParticleModuleRotationRate.h
--- a/Engine/Source/Runtime/Engine/Classes/Particles/RotationRate/ParticleModuleRotationRate.h
+++ b/Engine/Source/Runtime/Engine/Classes/Particles/RotationRate/ParticleModuleRotationRate.h
@@ -4,7 +4,7 @@
 #pragma once
 #include "ParticleModuleRotationRate.generated.h"
 
-UCLASS(HeaderGroup=Particle, editinlinenew, hidecategories=Object, MinimalAPI, meta=(DisplayName = "Initial Rotation Rate"))
+UCLASS(editinlinenew, hidecategories=Object, MinimalAPI, meta=(DisplayName = "Initial Rotation Rate"))
 class UParticleModuleRotationRate : public UParticleModuleRotationRateBase
 {
 	GENERATED_UCLASS_BODY()
diff --git a/Engine/Source/Runtime/Engine/Classes/Particles/RotationRate/ParticleModuleRotationRateBase.h b/Engine/Source/Runtime/Engine/Classes/Particles/RotationRate/ParticleModuleRotationRateBase.h
--- a/Engine/Source/Runtime/Engine/Classes/Particles/RotationRate/ParticleModuleRotationRateBase.h
+++ b/Engine/Source/Runtime/Engine/Classes/Particles/RotationRate/ParticleModuleRotationRateBase.h
@@ -4,7 +4,7 @@
 #pragma once
 #include "ParticleModuleRotationRateBase.generated.h"
 
-UCLASS(HeaderGroup=Particle, editinlinenew, hidecategories=Object, abstract, meta=(DisplayName = "Rotation Rate"))
+UCLASS(editinlinenew, hidecategories=Object, abstract, meta=(DisplayName = "Rotation Rate"))
 class UParticleModuleRotationRateBase : public UParticleModule
 {
 	GENERATED_UCLASS_BODY()
diff --git a/Engine/Source/Runtime/Engine/Classes/Particles/RotationRate/ParticleModuleRotationRateMultiplyLife.h b/Engine/Source/Runtime/Engine/Classes/Particles/RotationRate/ParticleModuleRotationRateMultiplyLife.h
--- a/Engine/Source/Runtime/Engine/Classes/Particles/RotationRate/ParticleModuleRotationRateMultiplyLife.h
+++ b/Engine/Source/Runtime/Engine/Classes/Particles/RotationRate/ParticleModuleRotationRateMultiplyLife.h
@@ -4,7 +4,7 @@
 #pragma once
 #include "ParticleModuleRotationRateMultiplyLife.generated.h"
 
-UCLASS(HeaderGroup=Particle, editinlinenew, hidecategories=Object, MinimalAPI, meta=(DisplayName = "Rotation Rate * Life"))
+UCLASS(editinlinenew, hidecategories=Object, MinimalAPI, meta=(DisplayName = "Rotation Rate * Life"))
 class UParticleModuleRotationRateMultiplyLife : public UParticleModuleRotationRateBase
 {
 	GENERATED_UCLASS_BODY()
diff --git a/Engine/Source/Runtime/Engine/Classes/Particles/RotationRate/ParticleModuleRotationRate_Seeded.h b/Engine/Source/Runtime/Engine/Classes/Particles/RotationRate/ParticleModuleRotationRate_Seeded.h
--- a/Engine/Source/Runtime/Engine/Classes/Particles/RotationRate/ParticleModuleRotationRate_Seeded.h
+++ b/Engine/Source/Runtime/Engine/Classes/Particles/RotationRate/ParticleModuleRotationRate_Seeded.h
@@ -4,7 +4,7 @@
 #pragma once
 #include "ParticleModuleRotationRate_Seeded.generated.h"
 
-UCLASS(HeaderGroup=Particle, editinlinenew, hidecategories=Object, meta=(DisplayName = "Init Rotation Rate (Seed)"))
+UCLASS(editinlinenew, hidecategories=Object, meta=(DisplayName = "Init Rotation Rate (Seed)"))
 class UParticleModuleRotationRate_Seeded : public UParticleModuleRotationRate
 {
 	GENERATED_UCLASS_BODY()
diff --git a/Engine/Source/Runtime/Engine/Classes/Particles/Size/ParticleModuleSize.h b/Engine/Source/Runtime/Engine/Classes/Particles/Size/ParticleModuleSize.h
--- a/Engine/Source/Runtime/Engine/Classes/Particles/Size/ParticleModuleSize.h
+++ b/Engine/Source/Runtime/Engine/Classes/Particles/Size/ParticleModuleSize.h
@@ -4,7 +4,7 @@
 #pragma once
 #include "ParticleModuleSize.generated.h"
 
-UCLASS(HeaderGroup=Particle, editinlinenew, hidecategories=Object, MinimalAPI, meta=(DisplayName = "Initial Size"))
+UCLASS(editinlinenew, hidecategories=Object, MinimalAPI, meta=(DisplayName = "Initial Size"))
 class UParticleModuleSize : public UParticleModuleSizeBase
 {
 	GENERATED_UCLASS_BODY()
diff --git a/Engine/Source/Runtime/Engine/Classes/Particles/Size/ParticleModuleSizeBase.h b/Engine/Source/Runtime/Engine/Classes/Particles/Size/ParticleModuleSizeBase.h
--- a/Engine/Source/Runtime/Engine/Classes/Particles/Size/ParticleModuleSizeBase.h
+++ b/Engine/Source/Runtime/Engine/Classes/Particles/Size/ParticleModuleSizeBase.h
@@ -4,7 +4,7 @@
 #pragma once
 #include "ParticleModuleSizeBase.generated.h"
 
-UCLASS(HeaderGroup=Particle, editinlinenew, hidecategories=Object, abstract, meta=(DisplayName = "Size"))
+UCLASS(editinlinenew, hidecategories=Object, abstract, meta=(DisplayName = "Size"))
 class UParticleModuleSizeBase : public UParticleModule
 {
 	GENERATED_UCLASS_BODY()
diff --git a/Engine/Source/Runtime/Engine/Classes/Particles/Size/ParticleModuleSizeMultiplyLife.h b/Engine/Source/Runtime/Engine/Classes/Particles/Size/ParticleModuleSizeMultiplyLife.h
--- a/Engine/Source/Runtime/Engine/Classes/Particles/Size/ParticleModuleSizeMultiplyLife.h
+++ b/Engine/Source/Runtime/Engine/Classes/Particles/Size/ParticleModuleSizeMultiplyLife.h
@@ -4,7 +4,7 @@
 #pragma once
 #include "ParticleModuleSizeMultiplyLife.generated.h"
 
-UCLASS(HeaderGroup=Particle, editinlinenew, hidecategories=Object, MinimalAPI, meta=(DisplayName = "Size By Life"))
+UCLASS(editinlinenew, hidecategories=Object, MinimalAPI, meta=(DisplayName = "Size By Life"))
 class UParticleModuleSizeMultiplyLife : public UParticleModuleSizeBase
 {
 	GENERATED_UCLASS_BODY()
diff --git a/Engine/Source/Runtime/Engine/Classes/Particles/Size/ParticleModuleSizeScale.h b/Engine/Source/Runtime/Engine/Classes/Particles/Size/ParticleModuleSizeScale.h
--- a/Engine/Source/Runtime/Engine/Classes/Particles/Size/ParticleModuleSizeScale.h
+++ b/Engine/Source/Runtime/Engine/Classes/Particles/Size/ParticleModuleSizeScale.h
@@ -4,7 +4,7 @@
 #pragma once
 #include "ParticleModuleSizeScale.generated.h"
 
-UCLASS(HeaderGroup=Particle, editinlinenew, hidecategories=Object, MinimalAPI, meta=(DisplayName = "Size Scale"))
+UCLASS(editinlinenew, hidecategories=Object, MinimalAPI, meta=(DisplayName = "Size Scale"))
 class UParticleModuleSizeScale : public UParticleModuleSizeBase
 {
 	GENERATED_UCLASS_BODY()
diff --git a/Engine/Source/Runtime/Engine/Classes/Particles/Size/ParticleModuleSizeScaleBySpeed.h b/Engine/Source/Runtime/Engine/Classes/Particles/Size/ParticleModuleSizeScaleBySpeed.h
--- a/Engine/Source/Runtime/Engine/Classes/Particles/Size/ParticleModuleSizeScaleBySpeed.h
+++ b/Engine/Source/Runtime/Engine/Classes/Particles/Size/ParticleModuleSizeScaleBySpeed.h
@@ -8,7 +8,7 @@
 
 #include "ParticleModuleSizeScaleBySpeed.generated.h"
 
-UCLASS(HeaderGroup=Particle, editinlinenew, hidecategories=Object, meta=(DisplayName = "Size By Speed"))
+UCLASS(editinlinenew, hidecategories=Object, meta=(DisplayName = "Size By Speed"))
 class UParticleModuleSizeScaleBySpeed : public UParticleModuleSizeBase
 {
 	GENERATED_UCLASS_BODY()
diff --git a/Engine/Source/Runtime/Engine/Classes/Particles/Size/ParticleModuleSize_Seeded.h b/Engine/Source/Runtime/Engine/Classes/Particles/Size/ParticleModuleSize_Seeded.h
--- a/Engine/Source/Runtime/Engine/Classes/Particles/Size/ParticleModuleSize_Seeded.h
+++ b/Engine/Source/Runtime/Engine/Classes/Particles/Size/ParticleModuleSize_Seeded.h
@@ -4,7 +4,7 @@
 #pragma once
 #include "ParticleModuleSize_Seeded.generated.h"
 
-UCLASS(HeaderGroup=Particle, editinlinenew, hidecategories=Object, meta=(DisplayName = "Initial Size (Seed)"))
+UCLASS(editinlinenew, hidecategories=Object, meta=(DisplayName = "Initial Size (Seed)"))
 class UParticleModuleSize_Seeded : public UParticleModuleSize
 {
 	GENERATED_UCLASS_BODY()
diff --git a/Engine/Source/Runtime/Engine/Classes/Particles/Spawn/ParticleModuleSpawn.h b/Engine/Source/Runtime/Engine/Classes/Particles/Spawn/ParticleModuleSpawn.h
--- a/Engine/Source/Runtime/Engine/Classes/Particles/Spawn/ParticleModuleSpawn.h
+++ b/Engine/Source/Runtime/Engine/Classes/Particles/Spawn/ParticleModuleSpawn.h
@@ -4,7 +4,7 @@
 #pragma once
 #include "ParticleModuleSpawn.generated.h"
 
-UCLASS(HeaderGroup=Particle, editinlinenew, hidecategories=Object, hidecategories=ParticleModuleSpawnBase, MinimalAPI, meta=(DisplayName = "Spawn"))
+UCLASS(editinlinenew, hidecategories=Object, hidecategories=ParticleModuleSpawnBase, MinimalAPI, meta=(DisplayName = "Spawn"))
 class UParticleModuleSpawn : public UParticleModuleSpawnBase
 {
 	GENERATED_UCLASS_BODY()
diff --git a/Engine/Source/Runtime/Engine/Classes/Particles/Spawn/ParticleModuleSpawnBase.h b/Engine/Source/Runtime/Engine/Classes/Particles/Spawn/ParticleModuleSpawnBase.h
--- a/Engine/Source/Runtime/Engine/Classes/Particles/Spawn/ParticleModuleSpawnBase.h
+++ b/Engine/Source/Runtime/Engine/Classes/Particles/Spawn/ParticleModuleSpawnBase.h
@@ -4,7 +4,7 @@
 #pragma once
 #include "ParticleModuleSpawnBase.generated.h"
 
-UCLASS(HeaderGroup=Particle, editinlinenew, hidecategories=Object, abstract, meta=(DisplayName = "Spawn"))
+UCLASS(editinlinenew, hidecategories=Object, abstract, meta=(DisplayName = "Spawn"))
 class UParticleModuleSpawnBase : public UParticleModule
 {
 	GENERATED_UCLASS_BODY()
diff --git a/Engine/Source/Runtime/Engine/Classes/Particles/Spawn/ParticleModuleSpawnPerUnit.h b/Engine/Source/Runtime/Engine/Classes/Particles/Spawn/ParticleModuleSpawnPerUnit.h
--- a/Engine/Source/Runtime/Engine/Classes/Particles/Spawn/ParticleModuleSpawnPerUnit.h
+++ b/Engine/Source/Runtime/Engine/Classes/Particles/Spawn/ParticleModuleSpawnPerUnit.h
@@ -4,7 +4,7 @@
 #pragma once
 #include "ParticleModuleSpawnPerUnit.generated.h"
 
-UCLASS(HeaderGroup=Particle, editinlinenew, hidecategories=Object, meta=(DisplayName = "Spawn PerUnit"))
+UCLASS(editinlinenew, hidecategories=Object, meta=(DisplayName = "Spawn PerUnit"))
 class UParticleModuleSpawnPerUnit : public UParticleModuleSpawnBase
 {
 	GENERATED_UCLASS_BODY()
diff --git a/Engine/Source/Runtime/Engine/Classes/Particles/SubUV/ParticleModuleSubUV.h b/Engine/Source/Runtime/Engine/Classes/Particles/SubUV/ParticleModuleSubUV.h
--- a/Engine/Source/Runtime/Engine/Classes/Particles/SubUV/ParticleModuleSubUV.h
+++ b/Engine/Source/Runtime/Engine/Classes/Particles/SubUV/ParticleModuleSubUV.h
@@ -4,7 +4,7 @@
 #pragma once
 #include "ParticleModuleSubUV.generated.h"
 
-UCLASS(HeaderGroup=Particle, editinlinenew, hidecategories=Object, MinimalAPI, meta=(DisplayName = "SubImage Index"))
+UCLASS(editinlinenew, hidecategories=Object, MinimalAPI, meta=(DisplayName = "SubImage Index"))
 class UParticleModuleSubUV : public UParticleModuleSubUVBase
 {
 	GENERATED_UCLASS_BODY()
diff --git a/Engine/Source/Runtime/Engine/Classes/Particles/SubUV/ParticleModuleSubUVBase.h b/Engine/Source/Runtime/Engine/Classes/Particles/SubUV/ParticleModuleSubUVBase.h
--- a/Engine/Source/Runtime/Engine/Classes/Particles/SubUV/ParticleModuleSubUVBase.h
+++ b/Engine/Source/Runtime/Engine/Classes/Particles/SubUV/ParticleModuleSubUVBase.h
@@ -4,7 +4,7 @@
 #pragma once
 #include "ParticleModuleSubUVBase.generated.h"
 
-UCLASS(HeaderGroup=Particle, editinlinenew, hidecategories=Object, abstract, meta=(DisplayName = "SubUV"))
+UCLASS(editinlinenew, hidecategories=Object, abstract, meta=(DisplayName = "SubUV"))
 class UParticleModuleSubUVBase : public UParticleModule
 {
 	GENERATED_UCLASS_BODY()
diff --git a/Engine/Source/Runtime/Engine/Classes/Particles/SubUV/ParticleModuleSubUVMovie.h b/Engine/Source/Runtime/Engine/Classes/Particles/SubUV/ParticleModuleSubUVMovie.h
--- a/Engine/Source/Runtime/Engine/Classes/Particles/SubUV/ParticleModuleSubUVMovie.h
+++ b/Engine/Source/Runtime/Engine/Classes/Particles/SubUV/ParticleModuleSubUVMovie.h
@@ -4,7 +4,7 @@
 #pragma once
 #include "ParticleModuleSubUVMovie.generated.h"
 
-UCLASS(HeaderGroup=Particle, editinlinenew, hidecategories=Object, hidecategories=SubUV, meta=(DisplayName = "SubUV Movie"))
+UCLASS(editinlinenew, hidecategories=Object, hidecategories=SubUV, meta=(DisplayName = "SubUV Movie"))
 class UParticleModuleSubUVMovie : public UParticleModuleSubUV
 {
 	GENERATED_UCLASS_BODY()
diff --git a/Engine/Source/Runtime/Engine/Classes/Particles/Trail/ParticleModuleTrailBase.h b/Engine/Source/Runtime/Engine/Classes/Particles/Trail/ParticleModuleTrailBase.h
--- a/Engine/Source/Runtime/Engine/Classes/Particles/Trail/ParticleModuleTrailBase.h
+++ b/Engine/Source/Runtime/Engine/Classes/Particles/Trail/ParticleModuleTrailBase.h
@@ -9,7 +9,7 @@
 #pragma once
 #include "ParticleModuleTrailBase.generated.h"
 
-UCLASS(HeaderGroup=Particle, editinlinenew, hidecategories=Object, abstract, meta=(DisplayName = "Trail"))
+UCLASS(editinlinenew, hidecategories=Object, abstract, meta=(DisplayName = "Trail"))
 class UParticleModuleTrailBase : public UParticleModule
 {
 	GENERATED_UCLASS_BODY()
diff --git a/Engine/Source/Runtime/Engine/Classes/Particles/Trail/ParticleModuleTrailSource.h b/Engine/Source/Runtime/Engine/Classes/Particles/Trail/ParticleModuleTrailSource.h
--- a/Engine/Source/Runtime/Engine/Classes/Particles/Trail/ParticleModuleTrailSource.h
+++ b/Engine/Source/Runtime/Engine/Classes/Particles/Trail/ParticleModuleTrailSource.h
@@ -28,7 +28,7 @@ enum ETrail2SourceMethod
 	PET2SRCM_MAX,
 };
 
-UCLASS(HeaderGroup=Particle, editinlinenew, hidecategories=Object, MinimalAPI, meta=(DisplayName = "Source"))
+UCLASS(editinlinenew, hidecategories=Object, MinimalAPI, meta=(DisplayName = "Source"))
 class UParticleModuleTrailSource : public UParticleModuleTrailBase
 {
 	GENERATED_UCLASS_BODY()
diff --git a/Engine/Source/Runtime/Engine/Classes/Particles/TypeData/ParticleModuleTypeDataAnimTrail.h b/Engine/Source/Runtime/Engine/Classes/Particles/TypeData/ParticleModuleTypeDataAnimTrail.h
--- a/Engine/Source/Runtime/Engine/Classes/Particles/TypeData/ParticleModuleTypeDataAnimTrail.h
+++ b/Engine/Source/Runtime/Engine/Classes/Particles/TypeData/ParticleModuleTypeDataAnimTrail.h
@@ -9,7 +9,7 @@
 #pragma once
 #include "ParticleModuleTypeDataAnimTrail.generated.h"
 
-UCLASS(HeaderGroup=Particle, editinlinenew, hidecategories=Object, meta=(DisplayName = "AnimTrail Data"))
+UCLASS(editinlinenew, hidecategories=Object, meta=(DisplayName = "AnimTrail Data"))
 class UParticleModuleTypeDataAnimTrail : public UParticleModuleTypeDataBase
 {
 	GENERATED_UCLASS_BODY()
diff --git a/Engine/Source/Runtime/Engine/Classes/Particles/TypeData/ParticleModuleTypeDataBase.h b/Engine/Source/Runtime/Engine/Classes/Particles/TypeData/ParticleModuleTypeDataBase.h
--- a/Engine/Source/Runtime/Engine/Classes/Particles/TypeData/ParticleModuleTypeDataBase.h
+++ b/Engine/Source/Runtime/Engine/Classes/Particles/TypeData/ParticleModuleTypeDataBase.h
@@ -4,7 +4,7 @@
 #pragma once
 #include "ParticleModuleTypeDataBase.generated.h"
 
-UCLASS(HeaderGroup=Particle, editinlinenew, hidecategories=Object, abstract, MinimalAPI)
+UCLASS(editinlinenew, hidecategories=Object, abstract, MinimalAPI)
 class UParticleModuleTypeDataBase : public UParticleModule
 {
 	GENERATED_UCLASS_BODY()
diff --git a/Engine/Source/Runtime/Engine/Classes/Particles/TypeData/ParticleModuleTypeDataBeam2.h b/Engine/Source/Runtime/Engine/Classes/Particles/TypeData/ParticleModuleTypeDataBeam2.h
--- a/Engine/Source/Runtime/Engine/Classes/Particles/TypeData/ParticleModuleTypeDataBeam2.h
+++ b/Engine/Source/Runtime/Engine/Classes/Particles/TypeData/ParticleModuleTypeDataBeam2.h
@@ -49,7 +49,7 @@ struct FBeamTargetData
 
 };
 
-UCLASS(HeaderGroup=Particle, editinlinenew, dontcollapsecategories, hidecategories=Object, MinimalAPI, meta=(DisplayName = "Beam Data"))
+UCLASS(editinlinenew, dontcollapsecategories, hidecategories=Object, MinimalAPI, meta=(DisplayName = "Beam Data"))
 class UParticleModuleTypeDataBeam2 : public UParticleModuleTypeDataBase
 {
 	GENERATED_UCLASS_BODY()
diff --git a/Engine/Source/Runtime/Engine/Classes/Particles/TypeData/ParticleModuleTypeDataGpu.h b/Engine/Source/Runtime/Engine/Classes/Particles/TypeData/ParticleModuleTypeDataGpu.h
--- a/Engine/Source/Runtime/Engine/Classes/Particles/TypeData/ParticleModuleTypeDataGpu.h
+++ b/Engine/Source/Runtime/Engine/Classes/Particles/TypeData/ParticleModuleTypeDataGpu.h
@@ -377,7 +377,7 @@ struct FGPUSpriteResourceData
 
 };
 
-UCLASS(HeaderGroup=Particle, editinlinenew, hidecategories=Object, MinimalAPI, meta=(DisplayName = "GPU Sprites"))
+UCLASS(editinlinenew, hidecategories=Object, MinimalAPI, meta=(DisplayName = "GPU Sprites"))
 class UParticleModuleTypeDataGpu : public UParticleModuleTypeDataBase
 {
 	GENERATED_UCLASS_BODY()
diff --git a/Engine/Source/Runtime/Engine/Classes/Particles/TypeData/ParticleModuleTypeDataMesh.h b/Engine/Source/Runtime/Engine/Classes/Particles/TypeData/ParticleModuleTypeDataMesh.h
--- a/Engine/Source/Runtime/Engine/Classes/Particles/TypeData/ParticleModuleTypeDataMesh.h
+++ b/Engine/Source/Runtime/Engine/Classes/Particles/TypeData/ParticleModuleTypeDataMesh.h
@@ -46,7 +46,7 @@ enum EMeshCameraFacingOptions
 	EMeshCameraFacingOptions_MAX,
 };
 
-UCLASS(HeaderGroup=Particle, editinlinenew, hidecategories=Object, MinimalAPI, meta=(DisplayName = "Mesh Data"))
+UCLASS(editinlinenew, hidecategories=Object, MinimalAPI, meta=(DisplayName = "Mesh Data"))
 class UParticleModuleTypeDataMesh : public UParticleModuleTypeDataBase
 {
 	GENERATED_UCLASS_BODY()
diff --git a/Engine/Source/Runtime/Engine/Classes/Particles/TypeData/ParticleModuleTypeDataRibbon.h b/Engine/Source/Runtime/Engine/Classes/Particles/TypeData/ParticleModuleTypeDataRibbon.h
--- a/Engine/Source/Runtime/Engine/Classes/Particles/TypeData/ParticleModuleTypeDataRibbon.h
+++ b/Engine/Source/Runtime/Engine/Classes/Particles/TypeData/ParticleModuleTypeDataRibbon.h
@@ -18,7 +18,7 @@ enum ETrailsRenderAxisOption
 	Trails_MAX,
 };
 
-UCLASS(HeaderGroup=Particle, editinlinenew, hidecategories=Object, meta=(DisplayName = "Ribbon Data"))
+UCLASS(editinlinenew, hidecategories=Object, meta=(DisplayName = "Ribbon Data"))
 class UParticleModuleTypeDataRibbon : public UParticleModuleTypeDataBase
 {
 	GENERATED_UCLASS_BODY()
diff --git a/Engine/Source/Runtime/Engine/Classes/Particles/VectorField/ParticleModuleVectorFieldBase.h b/Engine/Source/Runtime/Engine/Classes/Particles/VectorField/ParticleModuleVectorFieldBase.h
--- a/Engine/Source/Runtime/Engine/Classes/Particles/VectorField/ParticleModuleVectorFieldBase.h
+++ b/Engine/Source/Runtime/Engine/Classes/Particles/VectorField/ParticleModuleVectorFieldBase.h
@@ -9,7 +9,7 @@
 
 #include "ParticleModuleVectorFieldBase.generated.h"
 
-UCLASS(HeaderGroup=Particle, editinlinenew, hidecategories=Object, abstract, meta=(DisplayName = "Vector Field"))
+UCLASS(editinlinenew, hidecategories=Object, abstract, meta=(DisplayName = "Vector Field"))
 class UParticleModuleVectorFieldBase : public UParticleModule
 {
 	GENERATED_UCLASS_BODY()
diff --git a/Engine/Source/Runtime/Engine/Classes/Particles/VectorField/ParticleModuleVectorFieldGlobal.h b/Engine/Source/Runtime/Engine/Classes/Particles/VectorField/ParticleModuleVectorFieldGlobal.h
--- a/Engine/Source/Runtime/Engine/Classes/Particles/VectorField/ParticleModuleVectorFieldGlobal.h
+++ b/Engine/Source/Runtime/Engine/Classes/Particles/VectorField/ParticleModuleVectorFieldGlobal.h
@@ -8,7 +8,7 @@
 
 #include "ParticleModuleVectorFieldGlobal.generated.h"
 
-UCLASS(HeaderGroup=Particle, editinlinenew, hidecategories=Object, meta=(DisplayName = "Global Vector Fields"))
+UCLASS(editinlinenew, hidecategories=Object, meta=(DisplayName = "Global Vector Fields"))
 class UParticleModuleVectorFieldGlobal : public UParticleModuleVectorFieldBase
 {
 	GENERATED_UCLASS_BODY()
diff --git a/Engine/Source/Runtime/Engine/Classes/Particles/VectorField/ParticleModuleVectorFieldLocal.h b/Engine/Source/Runtime/Engine/Classes/Particles/VectorField/ParticleModuleVectorFieldLocal.h
--- a/Engine/Source/Runtime/Engine/Classes/Particles/VectorField/ParticleModuleVectorFieldLocal.h
+++ b/Engine/Source/Runtime/Engine/Classes/Particles/VectorField/ParticleModuleVectorFieldLocal.h
@@ -8,7 +8,7 @@
 
 #include "ParticleModuleVectorFieldLocal.generated.h"
 
-UCLASS(HeaderGroup=Particle, editinlinenew, hidecategories=Object, MinimalAPI, meta=(DisplayName = "Local Vector Field"))
+UCLASS(editinlinenew, hidecategories=Object, MinimalAPI, meta=(DisplayName = "Local Vector Field"))
 class UParticleModuleVectorFieldLocal : public UParticleModuleVectorFieldBase
 {
 	GENERATED_UCLASS_BODY()
