#ttp 303393 UObject improvement: Remove header groups

diff --git a/Engine/Source/Runtime/Engine/Classes/Particles/VectorField/ParticleModuleVectorFieldRotation.h b/Engine/Source/Runtime/Engine/Classes/Particles/VectorField/ParticleModuleVectorFieldRotation.h
--- a/Engine/Source/Runtime/Engine/Classes/Particles/VectorField/ParticleModuleVectorFieldRotation.h
+++ b/Engine/Source/Runtime/Engine/Classes/Particles/VectorField/ParticleModuleVectorFieldRotation.h
@@ -9,7 +9,7 @@
 
 #include "ParticleModuleVectorFieldRotation.generated.h"
 
-UCLASS(HeaderGroup=Particle, editinlinenew, hidecategories=Object, meta=(DisplayName = "VF Init Rotation"))
+UCLASS(editinlinenew, hidecategories=Object, meta=(DisplayName = "VF Init Rotation"))
 class UParticleModuleVectorFieldRotation : public UParticleModuleVectorFieldBase
 {
 	GENERATED_UCLASS_BODY()
diff --git a/Engine/Source/Runtime/Engine/Classes/Particles/VectorField/ParticleModuleVectorFieldRotationRate.h b/Engine/Source/Runtime/Engine/Classes/Particles/VectorField/ParticleModuleVectorFieldRotationRate.h
--- a/Engine/Source/Runtime/Engine/Classes/Particles/VectorField/ParticleModuleVectorFieldRotationRate.h
+++ b/Engine/Source/Runtime/Engine/Classes/Particles/VectorField/ParticleModuleVectorFieldRotationRate.h
@@ -8,7 +8,7 @@
 
 #include "ParticleModuleVectorFieldRotationRate.generated.h"
 
-UCLASS(HeaderGroup=Particle, editinlinenew, hidecategories=Object, meta=(DisplayName = "VF Rotation Rate"))
+UCLASS(editinlinenew, hidecategories=Object, meta=(DisplayName = "VF Rotation Rate"))
 class UParticleModuleVectorFieldRotationRate : public UParticleModuleVectorFieldBase
 {
 	GENERATED_UCLASS_BODY()
diff --git a/Engine/Source/Runtime/Engine/Classes/Particles/VectorField/ParticleModuleVectorFieldScale.h b/Engine/Source/Runtime/Engine/Classes/Particles/VectorField/ParticleModuleVectorFieldScale.h
--- a/Engine/Source/Runtime/Engine/Classes/Particles/VectorField/ParticleModuleVectorFieldScale.h
+++ b/Engine/Source/Runtime/Engine/Classes/Particles/VectorField/ParticleModuleVectorFieldScale.h
@@ -8,7 +8,7 @@
 
 #include "ParticleModuleVectorFieldScale.generated.h"
 
-UCLASS(HeaderGroup=Particle, editinlinenew, hidecategories=Object, MinimalAPI, meta=(DisplayName = "Vector Field Scale"))
+UCLASS(editinlinenew, hidecategories=Object, MinimalAPI, meta=(DisplayName = "Vector Field Scale"))
 class UParticleModuleVectorFieldScale : public UParticleModuleVectorFieldBase
 {
 	GENERATED_UCLASS_BODY()
diff --git a/Engine/Source/Runtime/Engine/Classes/Particles/VectorField/ParticleModuleVectorFieldScaleOverLife.h b/Engine/Source/Runtime/Engine/Classes/Particles/VectorField/ParticleModuleVectorFieldScaleOverLife.h
--- a/Engine/Source/Runtime/Engine/Classes/Particles/VectorField/ParticleModuleVectorFieldScaleOverLife.h
+++ b/Engine/Source/Runtime/Engine/Classes/Particles/VectorField/ParticleModuleVectorFieldScaleOverLife.h
@@ -8,7 +8,7 @@
 
 #include "ParticleModuleVectorFieldScaleOverLife.generated.h"
 
-UCLASS(HeaderGroup=Particle, editinlinenew, hidecategories=Object, MinimalAPI, meta=(DisplayName = "VF Scale/Life"))
+UCLASS(editinlinenew, hidecategories=Object, MinimalAPI, meta=(DisplayName = "VF Scale/Life"))
 class UParticleModuleVectorFieldScaleOverLife : public UParticleModuleVectorFieldBase
 {
 	GENERATED_UCLASS_BODY()
diff --git a/Engine/Source/Runtime/Engine/Classes/Particles/Velocity/ParticleModuleVelocity.h b/Engine/Source/Runtime/Engine/Classes/Particles/Velocity/ParticleModuleVelocity.h
--- a/Engine/Source/Runtime/Engine/Classes/Particles/Velocity/ParticleModuleVelocity.h
+++ b/Engine/Source/Runtime/Engine/Classes/Particles/Velocity/ParticleModuleVelocity.h
@@ -4,7 +4,7 @@
 #pragma once
 #include "ParticleModuleVelocity.generated.h"
 
-UCLASS(HeaderGroup=Particle, editinlinenew, hidecategories=Object, MinimalAPI, meta=(DisplayName = "Initial Velocity"))
+UCLASS(editinlinenew, hidecategories=Object, MinimalAPI, meta=(DisplayName = "Initial Velocity"))
 class UParticleModuleVelocity : public UParticleModuleVelocityBase
 {
 	GENERATED_UCLASS_BODY()
diff --git a/Engine/Source/Runtime/Engine/Classes/Particles/Velocity/ParticleModuleVelocityBase.h b/Engine/Source/Runtime/Engine/Classes/Particles/Velocity/ParticleModuleVelocityBase.h
--- a/Engine/Source/Runtime/Engine/Classes/Particles/Velocity/ParticleModuleVelocityBase.h
+++ b/Engine/Source/Runtime/Engine/Classes/Particles/Velocity/ParticleModuleVelocityBase.h
@@ -4,7 +4,7 @@
 #pragma once
 #include "ParticleModuleVelocityBase.generated.h"
 
-UCLASS(HeaderGroup=Particle, editinlinenew, hidecategories=Object, abstract, meta=(DisplayName = "Velocity"))
+UCLASS(editinlinenew, hidecategories=Object, abstract, meta=(DisplayName = "Velocity"))
 class UParticleModuleVelocityBase : public UParticleModule
 {
 	GENERATED_UCLASS_BODY()
diff --git a/Engine/Source/Runtime/Engine/Classes/Particles/Velocity/ParticleModuleVelocityCone.h b/Engine/Source/Runtime/Engine/Classes/Particles/Velocity/ParticleModuleVelocityCone.h
--- a/Engine/Source/Runtime/Engine/Classes/Particles/Velocity/ParticleModuleVelocityCone.h
+++ b/Engine/Source/Runtime/Engine/Classes/Particles/Velocity/ParticleModuleVelocityCone.h
@@ -4,7 +4,7 @@
 #pragma once
 #include "ParticleModuleVelocityCone.generated.h"
 
-UCLASS(HeaderGroup=Particle, editinlinenew, hidecategories=Object, meta=(DisplayName = "Velocity Cone"))
+UCLASS(editinlinenew, hidecategories=Object, meta=(DisplayName = "Velocity Cone"))
 class UParticleModuleVelocityCone : public UParticleModuleVelocityBase
 {
 	GENERATED_UCLASS_BODY()
diff --git a/Engine/Source/Runtime/Engine/Classes/Particles/Velocity/ParticleModuleVelocityInheritParent.h b/Engine/Source/Runtime/Engine/Classes/Particles/Velocity/ParticleModuleVelocityInheritParent.h
--- a/Engine/Source/Runtime/Engine/Classes/Particles/Velocity/ParticleModuleVelocityInheritParent.h
+++ b/Engine/Source/Runtime/Engine/Classes/Particles/Velocity/ParticleModuleVelocityInheritParent.h
@@ -4,7 +4,7 @@
 #pragma once
 #include "ParticleModuleVelocityInheritParent.generated.h"
 
-UCLASS(HeaderGroup=Particle, editinlinenew, hidecategories=Object, MinimalAPI, meta=(DisplayName = "Inherit Parent Velocity"))
+UCLASS(editinlinenew, hidecategories=Object, MinimalAPI, meta=(DisplayName = "Inherit Parent Velocity"))
 class UParticleModuleVelocityInheritParent : public UParticleModuleVelocityBase
 {
 	GENERATED_UCLASS_BODY()
diff --git a/Engine/Source/Runtime/Engine/Classes/Particles/Velocity/ParticleModuleVelocityOverLifetime.h b/Engine/Source/Runtime/Engine/Classes/Particles/Velocity/ParticleModuleVelocityOverLifetime.h
--- a/Engine/Source/Runtime/Engine/Classes/Particles/Velocity/ParticleModuleVelocityOverLifetime.h
+++ b/Engine/Source/Runtime/Engine/Classes/Particles/Velocity/ParticleModuleVelocityOverLifetime.h
@@ -4,7 +4,7 @@
 #pragma once
 #include "ParticleModuleVelocityOverLifetime.generated.h"
 
-UCLASS(HeaderGroup=Particle, editinlinenew, hidecategories=Object, MinimalAPI, meta=(DisplayName = "Velocity/Life"))
+UCLASS(editinlinenew, hidecategories=Object, MinimalAPI, meta=(DisplayName = "Velocity/Life"))
 class UParticleModuleVelocityOverLifetime : public UParticleModuleVelocityBase
 {
 	GENERATED_UCLASS_BODY()
diff --git a/Engine/Source/Runtime/Engine/Classes/Particles/Velocity/ParticleModuleVelocity_Seeded.h b/Engine/Source/Runtime/Engine/Classes/Particles/Velocity/ParticleModuleVelocity_Seeded.h
--- a/Engine/Source/Runtime/Engine/Classes/Particles/Velocity/ParticleModuleVelocity_Seeded.h
+++ b/Engine/Source/Runtime/Engine/Classes/Particles/Velocity/ParticleModuleVelocity_Seeded.h
@@ -4,7 +4,7 @@
 #pragma once
 #include "ParticleModuleVelocity_Seeded.generated.h"
 
-UCLASS(HeaderGroup=Particle, editinlinenew, hidecategories=Object, meta=(DisplayName = "Initial Velocity (Seed)"))
+UCLASS(editinlinenew, hidecategories=Object, meta=(DisplayName = "Initial Velocity (Seed)"))
 class UParticleModuleVelocity_Seeded : public UParticleModuleVelocity
 {
 	GENERATED_UCLASS_BODY()
diff --git a/Engine/Source/Runtime/Engine/Classes/PhysicsEngine/PhysicsConstraintComponent.h b/Engine/Source/Runtime/Engine/Classes/PhysicsEngine/PhysicsConstraintComponent.h
--- a/Engine/Source/Runtime/Engine/Classes/PhysicsEngine/PhysicsConstraintComponent.h
+++ b/Engine/Source/Runtime/Engine/Classes/PhysicsEngine/PhysicsConstraintComponent.h
@@ -9,7 +9,7 @@ DECLARE_DYNAMIC_MULTICAST_DELEGATE_OneParam(FConstraintBrokenSignature, int32, C
 /**
  *	This is effectively a joint that allows you to connect 2 rigid bodies together. You can create different types of joints using the various parameters of this component.
  */
-UCLASS(HeaderGroup=Component, ClassGroup=Physics, dependson(ConstraintInstance), MinimalAPI, meta=(BlueprintSpawnableComponent), HideCategories=(Activation,"Components|Activation",Physics, Mobility))
+UCLASS(ClassGroup=Physics, dependson(ConstraintInstance), MinimalAPI, meta=(BlueprintSpawnableComponent), HideCategories=(Activation,"Components|Activation",Physics, Mobility))
 class UPhysicsConstraintComponent : public USceneComponent
 {
 	GENERATED_UCLASS_BODY()
diff --git a/Engine/Source/Runtime/Engine/Classes/PhysicsEngine/PhysicsHandleComponent.h b/Engine/Source/Runtime/Engine/Classes/PhysicsEngine/PhysicsHandleComponent.h
--- a/Engine/Source/Runtime/Engine/Classes/PhysicsEngine/PhysicsHandleComponent.h
+++ b/Engine/Source/Runtime/Engine/Classes/PhysicsEngine/PhysicsHandleComponent.h
@@ -14,7 +14,7 @@ namespace physx
 /**
  *	Utility object for moving physics objects around.
  */
-UCLASS(collapsecategories, ClassGroup=Physics, hidecategories=Object, HeaderGroup=Component, MinimalAPI, meta=(BlueprintSpawnableComponent))
+UCLASS(collapsecategories, ClassGroup=Physics, hidecategories=Object, MinimalAPI, meta=(BlueprintSpawnableComponent))
 class UPhysicsHandleComponent : public UActorComponent
 {
 	GENERATED_UCLASS_BODY()
diff --git a/Engine/Source/Runtime/Engine/Classes/PhysicsEngine/PhysicsThrusterComponent.h b/Engine/Source/Runtime/Engine/Classes/PhysicsEngine/PhysicsThrusterComponent.h
--- a/Engine/Source/Runtime/Engine/Classes/PhysicsEngine/PhysicsThrusterComponent.h
+++ b/Engine/Source/Runtime/Engine/Classes/PhysicsEngine/PhysicsThrusterComponent.h
@@ -8,7 +8,7 @@
  *	Used with objects that have physics to apply a force down the negative-X direction
  *	ie. point X in the direction you want the thrust in.
  */
-UCLASS(hidecategories=(Object, Mobility, LOD), ClassGroup=Physics, showcategories=Trigger, HeaderGroup=Component, meta=(BlueprintSpawnableComponent))
+UCLASS(hidecategories=(Object, Mobility, LOD), ClassGroup=Physics, showcategories=Trigger, meta=(BlueprintSpawnableComponent))
 class UPhysicsThrusterComponent : public USceneComponent
 {
 	GENERATED_UCLASS_BODY()
diff --git a/Engine/Source/Runtime/Engine/Classes/PhysicsEngine/RadialForceComponent.h b/Engine/Source/Runtime/Engine/Classes/PhysicsEngine/RadialForceComponent.h
--- a/Engine/Source/Runtime/Engine/Classes/PhysicsEngine/RadialForceComponent.h
+++ b/Engine/Source/Runtime/Engine/Classes/PhysicsEngine/RadialForceComponent.h
@@ -7,7 +7,7 @@
 /**
  *	Used to emit a radial force or impulse that can affect physics objects and or destructible objects.
  */
-UCLASS(hidecategories=(Object, Mobility), ClassGroup=Physics, showcategories=Trigger, HeaderGroup=Component, meta=(BlueprintSpawnableComponent), MinimalAPI)
+UCLASS(hidecategories=(Object, Mobility), ClassGroup=Physics, showcategories=Trigger, meta=(BlueprintSpawnableComponent), MinimalAPI)
 class URadialForceComponent : public USceneComponent
 {
 	GENERATED_UCLASS_BODY()
diff --git a/Engine/Source/Runtime/Engine/Classes/Vehicles/WheeledVehicleMovementComponent.h b/Engine/Source/Runtime/Engine/Classes/Vehicles/WheeledVehicleMovementComponent.h
--- a/Engine/Source/Runtime/Engine/Classes/Vehicles/WheeledVehicleMovementComponent.h
+++ b/Engine/Source/Runtime/Engine/Classes/Vehicles/WheeledVehicleMovementComponent.h
@@ -209,7 +209,7 @@ void ConvertCurve( const TArray<FFloatPair>& Pairs, struct FRichCurve& OutCurve
 /**
  * Component to handle the vehicle simulation for an actor.
  */
-UCLASS(Abstract, HeaderGroup=Component, dependson=UCurveBase, hidecategories=(PlanarMovement, "Components|Movement|Planar", Activation, "Components|Activation"))
+UCLASS(Abstract, dependson=UCurveBase, hidecategories=(PlanarMovement, "Components|Movement|Planar", Activation, "Components|Activation"))
 class ENGINE_API UWheeledVehicleMovementComponent : public UPawnMovementComponent
 {
 	GENERATED_UCLASS_BODY()
diff --git a/Engine/Source/Runtime/Engine/Private/AI/Navigation/NavRelevantComponent.cpp b/Engine/Source/Runtime/Engine/Private/AI/Navigation/NavRelevantComponent.cpp
--- a/Engine/Source/Runtime/Engine/Private/AI/Navigation/NavRelevantComponent.cpp
+++ b/Engine/Source/Runtime/Engine/Private/AI/Navigation/NavRelevantComponent.cpp
@@ -1,7 +1,6 @@
 // Copyright 1998-2014 Epic Games, Inc. All Rights Reserved.
 
 #include "EnginePrivate.h"
-#include "EngineNavigationClasses.h"
 
 UNavRelevantComponent::UNavRelevantComponent(const class FPostConstructInitializeProperties& PCIP) : Super(PCIP)
 {
diff --git a/Engine/Source/Runtime/Engine/Private/AI/Navigation/NavigationOctree.cpp b/Engine/Source/Runtime/Engine/Private/AI/Navigation/NavigationOctree.cpp
--- a/Engine/Source/Runtime/Engine/Private/AI/Navigation/NavigationOctree.cpp
+++ b/Engine/Source/Runtime/Engine/Private/AI/Navigation/NavigationOctree.cpp
@@ -8,8 +8,6 @@
 #include "RecastNavMeshGenerator.h"
 #endif // NAVOCTREE_CONTAINS_COLLISION_DATA
 
-#include "EngineNavigationClasses.h"
-
 //----------------------------------------------------------------------//
 // FNavigationOctree
 //----------------------------------------------------------------------//
diff --git a/Engine/Source/Runtime/Engine/Private/AI/Navigation/SmartNavLinkComponent.cpp b/Engine/Source/Runtime/Engine/Private/AI/Navigation/SmartNavLinkComponent.cpp
--- a/Engine/Source/Runtime/Engine/Private/AI/Navigation/SmartNavLinkComponent.cpp
+++ b/Engine/Source/Runtime/Engine/Private/AI/Navigation/SmartNavLinkComponent.cpp
@@ -1,7 +1,6 @@
 // Copyright 1998-2014 Epic Games, Inc. All Rights Reserved.
 
 #include "EnginePrivate.h"
-#include "EngineNavigationClasses.h"
 #include "VisualLog.h"
 
 USmartNavLinkComponent::USmartNavLinkComponent(const class FPostConstructInitializeProperties& PCIP) : Super(PCIP)
diff --git a/Engine/Source/Runtime/Engine/Private/Actor.cpp b/Engine/Source/Runtime/Engine/Private/Actor.cpp
--- a/Engine/Source/Runtime/Engine/Private/Actor.cpp
+++ b/Engine/Source/Runtime/Engine/Private/Actor.cpp
@@ -4,8 +4,6 @@
 #include "EnginePrivate.h"
 #include "ParticleDefinitions.h"
 #include "SoundDefinitions.h"
-#include "EngineNavigationClasses.h"
-#include "EngineKismetLibraryClasses.h"
 #include "LatentActions.h"
 #include "MessageLog.h"
 #include "Net/UnrealNetwork.h"
diff --git a/Engine/Source/Runtime/Engine/Private/AllEngineHeaders.h b/Engine/Source/Runtime/Engine/Private/AllEngineHeaders.h
--- a/Engine/Source/Runtime/Engine/Private/AllEngineHeaders.h
+++ b/Engine/Source/Runtime/Engine/Private/AllEngineHeaders.h
@@ -4,11 +4,4 @@
 #include "ParticleDefinitions.h"
 #include "SoundDefinitions.h"
 #include "BlueprintUtilities.h"
-#include "EngineSplineClasses.h"
-#include "EnginePlatformInterfaceClasses.h"
-#include "EngineFoliageClasses.h"
-#include "EngineLevelScriptClasses.h"
-#include "EngineDecalClasses.h"
 #include "NavigationPathBuilder.h"
-#include "EngineUserInterfaceClasses.h"
-#include "EngineKismetLibraryClasses.h"
diff --git a/Engine/Source/Runtime/Engine/Private/Animation/AnimNode_BlendSpaceEvaluator.cpp b/Engine/Source/Runtime/Engine/Private/Animation/AnimNode_BlendSpaceEvaluator.cpp
--- a/Engine/Source/Runtime/Engine/Private/Animation/AnimNode_BlendSpaceEvaluator.cpp
+++ b/Engine/Source/Runtime/Engine/Private/Animation/AnimNode_BlendSpaceEvaluator.cpp
@@ -1,8 +1,6 @@
 // Copyright 1998-2014 Epic Games, Inc. All Rights Reserved.
 
 #include "EnginePrivate.h"
-#include "EngineAnimClasses.h"
-#include "EngineAnimationNodeClasses.h"
 
 /////////////////////////////////////////////////////
 // FAnimNode_BlendSpaceEvaluator
diff --git a/Engine/Source/Runtime/Engine/Private/Animation/AnimNotifyState.cpp b/Engine/Source/Runtime/Engine/Private/Animation/AnimNotifyState.cpp
--- a/Engine/Source/Runtime/Engine/Private/Animation/AnimNotifyState.cpp
+++ b/Engine/Source/Runtime/Engine/Private/Animation/AnimNotifyState.cpp
@@ -1,7 +1,6 @@
 // Copyright 1998-2014 Epic Games, Inc. All Rights Reserved.
 
 #include "EnginePrivate.h"
-#include "EngineAnimClasses.h"
 
 /////////////////////////////////////////////////////
 // UAnimNotifyState
diff --git a/Engine/Source/Runtime/Engine/Private/Animation/AnimNotifyState_Trail.cpp b/Engine/Source/Runtime/Engine/Private/Animation/AnimNotifyState_Trail.cpp
--- a/Engine/Source/Runtime/Engine/Private/Animation/AnimNotifyState_Trail.cpp
+++ b/Engine/Source/Runtime/Engine/Private/Animation/AnimNotifyState_Trail.cpp
@@ -1,10 +1,8 @@
 // Copyright 1998-2014 Epic Games, Inc. All Rights Reserved.
 
 #include "EnginePrivate.h"
-#include "EngineAnimClasses.h"
 #include "MessageLog.h"
 #include "ParticleDefinitions.h"
-#include "EngineKismetLibraryClasses.h"
 
 #define LOCTEXT_NAMESPACE "AnimNotifyState_Trail"
 
diff --git a/Engine/Source/Runtime/Engine/Private/Animation/AnimSequence.cpp b/Engine/Source/Runtime/Engine/Private/Animation/AnimSequence.cpp
--- a/Engine/Source/Runtime/Engine/Private/Animation/AnimSequence.cpp
+++ b/Engine/Source/Runtime/Engine/Private/Animation/AnimSequence.cpp
@@ -5,7 +5,6 @@
 =============================================================================*/ 
 
 #include "EnginePrivate.h"
-#include "EngineAnimClasses.h"
 #include "AnimationCompression.h"
 #include "AnimEncoding.h"
 #include "AnimationUtils.h"
diff --git a/Engine/Source/Runtime/Engine/Private/CheatManager.cpp b/Engine/Source/Runtime/Engine/Private/CheatManager.cpp
--- a/Engine/Source/Runtime/Engine/Private/CheatManager.cpp
+++ b/Engine/Source/Runtime/Engine/Private/CheatManager.cpp
@@ -1,7 +1,75 @@
 // Copyright 1998-2014 Epic Games, Inc. All Rights Reserved.
 
 #include "EnginePrivate.h"
-#include "EngineComponentClasses.h"
+
+#include "Components/ApplicationLifecycleComponent.h"
+#include "AI/BehaviorTree/BlackboardComponent.h"
+#include "AI/BrainComponent.h"
+#include "AI/BehaviorTree/BehaviorTreeComponent.h"
+#include "Debug/GameplayDebuggingControllerComponent.h"
+#include "Components/InputComponent.h"
+#include "GameFramework/MovementComponent.h"
+#include "Components/SceneComponent.h"
+#include "Components/LightComponentBase.h"
+#include "Components/LightComponent.h"
+#include "Components/PrimitiveComponent.h"
+#include "GameFramework/NavMovementComponent.h"
+#include "GameFramework/PawnMovementComponent.h"
+#include "GameFramework/CharacterMovementComponent.h"
+#include "GameFramework/FloatingPawnMovement.h"
+#include "GameFramework/SpectatorPawnMovement.h"
+#include "Vehicles/WheeledVehicleMovementComponent.h"
+#include "Vehicles/WheeledVehicleMovementComponent4W.h"
+#include "GameFramework/ProjectileMovementComponent.h"
+#include "GameFramework/RotatingMovementComponent.h"
+#include "AI/Navigation/NavigationComponent.h"
+#include "AI/Navigation/PathFollowingComponent.h"
+#include "Components/PawnNoiseEmitterComponent.h"
+#include "Components/PawnSensingComponent.h"
+#include "PhysicsEngine/PhysicsHandleComponent.h"
+#include "Atmosphere/AtmosphericFogComponent.h"
+#include "Components/AudioComponent.h"
+#include "Camera/CameraComponent.h"
+#include "Components/ChildActorComponent.h"
+#include "Components/ExponentialHeightFogComponent.h"
+#include "Components/SkyLightComponent.h"
+#include "AI/Navigation/NavigationGraphNodeComponent.h"
+#include "PhysicsEngine/PhysicsConstraintComponent.h"
+#include "PhysicsEngine/PhysicsThrusterComponent.h"
+#include "Components/PostProcessComponent.h"
+#include "Components/ArrowComponent.h"
+#include "Components/BillboardComponent.h"
+#include "Components/BrushComponent.h"
+#include "Components/DrawFrustumComponent.h"
+#include "AI/EnvironmentQuery/EQSRenderingComponent.h"
+#include "Debug/GameplayDebuggingComponent.h"
+#include "Components/LineBatchComponent.h"
+#include "Components/MaterialBillboardComponent.h"
+#include "Components/MeshComponent.h"
+#include "Components/SkinnedMeshComponent.h"
+#include "Components/DestructibleComponent.h"
+#include "Components/PoseableMeshComponent.h"
+#include "Components/SkeletalMeshComponent.h"
+#include "Components/StaticMeshComponent.h"
+#include "Components/InstancedStaticMeshComponent.h"
+#include "Components/InteractiveFoliageComponent.h"
+#include "Components/SplineMeshComponent.h"
+#include "Components/ModelComponent.h"
+#include "AI/Navigation/NavLinkRenderingComponent.h"
+#include "AI/Navigation/NavMeshRenderingComponent.h"
+#include "AI/Navigation/NavTestRenderingComponent.h"
+#include "Components/NiagaraComponent.h"
+#include "Components/ShapeComponent.h"
+#include "Components/BoxComponent.h"
+#include "Components/CapsuleComponent.h"
+#include "Components/SphereComponent.h"
+#include "Components/DrawSphereComponent.h"
+#include "Components/TextRenderComponent.h"
+#include "Components/VectorFieldComponent.h"
+#include "PhysicsEngine/RadialForceComponent.h"
+#include "GameFramework/SpringArmComponent.h"
+#include "Components/WindDirectionalSourceComponent.h"
+#include "Components/TimelineComponent.h"
 #include "Slate.h"
 #include "AI/NavDataGenerator.h"
 #include "OnlineSubsystemUtils.h"
diff --git a/Engine/Source/Runtime/Engine/Private/ComponentInstanceDataCache.cpp b/Engine/Source/Runtime/Engine/Private/ComponentInstanceDataCache.cpp
--- a/Engine/Source/Runtime/Engine/Private/ComponentInstanceDataCache.cpp
+++ b/Engine/Source/Runtime/Engine/Private/ComponentInstanceDataCache.cpp
@@ -1,7 +1,78 @@
 // Copyright 1998-2014 Epic Games, Inc. All Rights Reserved.
 
 #include "EnginePrivate.h"
-#include "EngineComponentClasses.h"
+
+#include "Components/ApplicationLifecycleComponent.h"
+#include "AI/BehaviorTree/BlackboardComponent.h"
+#include "AI/BrainComponent.h"
+#include "AI/BehaviorTree/BehaviorTreeComponent.h"
+#include "Debug/GameplayDebuggingControllerComponent.h"
+#include "Components/InputComponent.h"
+#include "GameFramework/MovementComponent.h"
+#include "Components/SceneComponent.h"
+#include "Components/LightComponentBase.h"
+#include "Components/LightComponent.h"
+#include "Components/PrimitiveComponent.h"
+#include "GameFramework/NavMovementComponent.h"
+#include "GameFramework/PawnMovementComponent.h"
+#include "GameFramework/CharacterMovementComponent.h"
+#include "GameFramework/FloatingPawnMovement.h"
+#include "GameFramework/SpectatorPawnMovement.h"
+#include "Vehicles/WheeledVehicleMovementComponent.h"
+#include "Vehicles/WheeledVehicleMovementComponent4W.h"
+#include "GameFramework/ProjectileMovementComponent.h"
+#include "GameFramework/RotatingMovementComponent.h"
+#include "AI/Navigation/NavigationComponent.h"
+#include "AI/Navigation/PathFollowingComponent.h"
+#include "Components/PawnNoiseEmitterComponent.h"
+#include "Components/PawnSensingComponent.h"
+#include "PhysicsEngine/PhysicsHandleComponent.h"
+#include "Atmosphere/AtmosphericFogComponent.h"
+#include "Components/AudioComponent.h"
+#include "Camera/CameraComponent.h"
+#include "Components/ChildActorComponent.h"
+#include "Components/ExponentialHeightFogComponent.h"
+#include "Components/SkyLightComponent.h"
+#include "AI/Navigation/NavigationGraphNodeComponent.h"
+#include "PhysicsEngine/PhysicsConstraintComponent.h"
+#include "PhysicsEngine/PhysicsThrusterComponent.h"
+#include "Components/PostProcessComponent.h"
+#include "Components/ArrowComponent.h"
+#include "Components/BillboardComponent.h"
+#include "Components/BrushComponent.h"
+#include "Components/DrawFrustumComponent.h"
+#include "AI/EnvironmentQuery/EQSRenderingComponent.h"
+#include "Debug/GameplayDebuggingComponent.h"
+#include "Components/LineBatchComponent.h"
+#include "Components/MaterialBillboardComponent.h"
+#include "Components/MeshComponent.h"
+#include "Components/SkinnedMeshComponent.h"
+#include "Components/DestructibleComponent.h"
+#include "Components/PoseableMeshComponent.h"
+#include "Components/SkeletalMeshComponent.h"
+#include "Components/StaticMeshComponent.h"
+#include "Components/InstancedStaticMeshComponent.h"
+#include "Components/InteractiveFoliageComponent.h"
+#include "Components/SplineMeshComponent.h"
+#include "Components/ModelComponent.h"
+#include "AI/Navigation/NavLinkRenderingComponent.h"
+#include "AI/Navigation/NavMeshRenderingComponent.h"
+#include "AI/Navigation/NavTestRenderingComponent.h"
+#include "Components/NiagaraComponent.h"
+#include "Components/ShapeComponent.h"
+#include "Components/BoxComponent.h"
+#include "Components/CapsuleComponent.h"
+#include "Components/SphereComponent.h"
+#include "Components/DrawSphereComponent.h"
+#include "Components/TextRenderComponent.h"
+#include "Components/VectorFieldComponent.h"
+#include "PhysicsEngine/RadialForceComponent.h"
+#include "GameFramework/SpringArmComponent.h"
+#include "Components/WindDirectionalSourceComponent.h"
+#include "Components/TimelineComponent.h"
+#include "Slate.h"
+#include "AI/NavDataGenerator.h"
+#include "OnlineSubsystemUtils.h"
 
 void FComponentInstanceDataCache::AddInstanceData(TSharedPtr<FComponentInstanceDataBase> NewData)
 {
diff --git a/Engine/Source/Runtime/Engine/Private/Controller.cpp b/Engine/Source/Runtime/Engine/Private/Controller.cpp
--- a/Engine/Source/Runtime/Engine/Private/Controller.cpp
+++ b/Engine/Source/Runtime/Engine/Private/Controller.cpp
@@ -8,8 +8,6 @@
 #include "Net/UnrealNetwork.h"
 #include "ConfigCacheIni.h"
 #include "NavigationPathBuilder.h"
-#include "EngineUserInterfaceClasses.h"
-#include "EngineInterpolationClasses.h"
 
 DEFINE_LOG_CATEGORY(LogPath);
 
diff --git a/Engine/Source/Runtime/Engine/Private/FoliageComponent.cpp b/Engine/Source/Runtime/Engine/Private/FoliageComponent.cpp
--- a/Engine/Source/Runtime/Engine/Private/FoliageComponent.cpp
+++ b/Engine/Source/Runtime/Engine/Private/FoliageComponent.cpp
@@ -5,9 +5,6 @@
 =============================================================================*/
 
 #include "EnginePrivate.h"
-#include "EngineFoliageClasses.h"
-
-
 
 /** Scene proxy class for UInteractiveFoliageComponent. */
 class FInteractiveFoliageSceneProxy : public FStaticMeshSceneProxy
diff --git a/Engine/Source/Runtime/Engine/Private/GameEngine.cpp b/Engine/Source/Runtime/Engine/Private/GameEngine.cpp
--- a/Engine/Source/Runtime/Engine/Private/GameEngine.cpp
+++ b/Engine/Source/Runtime/Engine/Private/GameEngine.cpp
@@ -5,18 +5,14 @@
 =============================================================================*/
 
 #include "EnginePrivate.h"
-#include "EngineUserInterfaceClasses.h"
 #include "ParticleDefinitions.h"
 #include "SoundDefinitions.h"
-#include "EnginePlatformInterfaceClasses.h"
-#include "EngineDecalClasses.h"
 #include "Net/UnrealNetwork.h"
 #include "AllocatorFixedSizeFreeList.h"
 #include "Database.h"
 #include "MallocProfiler.h"
 #include "Net/NetworkProfiler.h"
 #include "ConfigCacheIni.h"
-#include "EngineLevelScriptClasses.h"
 
 #include "AVIWriter.h"
 #include "StreamingPauseRendering.h"
diff --git a/Engine/Source/Runtime/Engine/Private/GameViewportClient.cpp b/Engine/Source/Runtime/Engine/Private/GameViewportClient.cpp
--- a/Engine/Source/Runtime/Engine/Private/GameViewportClient.cpp
+++ b/Engine/Source/Runtime/Engine/Private/GameViewportClient.cpp
@@ -1,8 +1,6 @@
 // Copyright 1998-2014 Epic Games, Inc. All Rights Reserved.
 
 #include "EnginePrivate.h"
-#include "EngineUserInterfaceClasses.h"
-#include "EngineKismetLibraryClasses.h"
 #include "ParticleDefinitions.h"
 #include "SoundDefinitions.h"
 #include "FXSystem.h"
diff --git a/Engine/Source/Runtime/Engine/Private/GameplayStatics.cpp b/Engine/Source/Runtime/Engine/Private/GameplayStatics.cpp
--- a/Engine/Source/Runtime/Engine/Private/GameplayStatics.cpp
+++ b/Engine/Source/Runtime/Engine/Private/GameplayStatics.cpp
@@ -1,8 +1,6 @@
 // Copyright 1998-2014 Epic Games, Inc. All Rights Reserved.
 
 #include "EnginePrivate.h"
-#include "EngineKismetLibraryClasses.h"
-#include "EngineDecalClasses.h"
 #include "ParticleDefinitions.h"
 #include "SoundDefinitions.h"
 #include "PlatformFeatures.h"
diff --git a/Engine/Source/Runtime/Engine/Private/InGameAdvertising.cpp b/Engine/Source/Runtime/Engine/Private/InGameAdvertising.cpp
--- a/Engine/Source/Runtime/Engine/Private/InGameAdvertising.cpp
+++ b/Engine/Source/Runtime/Engine/Private/InGameAdvertising.cpp
@@ -5,8 +5,6 @@
 =============================================================================*/ 
 
 #include "EnginePrivate.h"
-#include "EngineGameEngineClasses.h"
-#include "EnginePlatformInterfaceClasses.h"
 
 UInGameAdManager::UInGameAdManager(const class FPostConstructInitializeProperties& PCIP)
 	: Super(PCIP)
diff --git a/Engine/Source/Runtime/Engine/Private/InputComponent.cpp b/Engine/Source/Runtime/Engine/Private/InputComponent.cpp
--- a/Engine/Source/Runtime/Engine/Private/InputComponent.cpp
+++ b/Engine/Source/Runtime/Engine/Private/InputComponent.cpp
@@ -1,8 +1,6 @@
 // Copyright 1998-2014 Epic Games, Inc. All Rights Reserved.
 
 #include "EnginePrivate.h"
-#include "EngineUserInterfaceClasses.h"
-#include "EngineLevelScriptClasses.h"		// For ALevelScriptActor
 
 FInputChord::RelationshipType FInputChord::GetRelationship(const FInputChord& OtherChord) const
 {
diff --git a/Engine/Source/Runtime/Engine/Private/InstancedStaticMesh.cpp b/Engine/Source/Runtime/Engine/Private/InstancedStaticMesh.cpp
--- a/Engine/Source/Runtime/Engine/Private/InstancedStaticMesh.cpp
+++ b/Engine/Source/Runtime/Engine/Private/InstancedStaticMesh.cpp
@@ -6,7 +6,79 @@
 
 #include "EnginePrivate.h"
 #include "ShaderParameters.h"
-#include "EngineComponentClasses.h"
+
+#include "Components/ApplicationLifecycleComponent.h"
+#include "AI/BehaviorTree/BlackboardComponent.h"
+#include "AI/BrainComponent.h"
+#include "AI/BehaviorTree/BehaviorTreeComponent.h"
+#include "Debug/GameplayDebuggingControllerComponent.h"
+#include "Components/InputComponent.h"
+#include "GameFramework/MovementComponent.h"
+#include "Components/SceneComponent.h"
+#include "Components/LightComponentBase.h"
+#include "Components/LightComponent.h"
+#include "Components/PrimitiveComponent.h"
+#include "GameFramework/NavMovementComponent.h"
+#include "GameFramework/PawnMovementComponent.h"
+#include "GameFramework/CharacterMovementComponent.h"
+#include "GameFramework/FloatingPawnMovement.h"
+#include "GameFramework/SpectatorPawnMovement.h"
+#include "Vehicles/WheeledVehicleMovementComponent.h"
+#include "Vehicles/WheeledVehicleMovementComponent4W.h"
+#include "GameFramework/ProjectileMovementComponent.h"
+#include "GameFramework/RotatingMovementComponent.h"
+#include "AI/Navigation/NavigationComponent.h"
+#include "AI/Navigation/PathFollowingComponent.h"
+#include "Components/PawnNoiseEmitterComponent.h"
+#include "Components/PawnSensingComponent.h"
+#include "PhysicsEngine/PhysicsHandleComponent.h"
+#include "Atmosphere/AtmosphericFogComponent.h"
+#include "Components/AudioComponent.h"
+#include "Camera/CameraComponent.h"
+#include "Components/ChildActorComponent.h"
+#include "Components/ExponentialHeightFogComponent.h"
+#include "Components/SkyLightComponent.h"
+#include "AI/Navigation/NavigationGraphNodeComponent.h"
+#include "PhysicsEngine/PhysicsConstraintComponent.h"
+#include "PhysicsEngine/PhysicsThrusterComponent.h"
+#include "Components/PostProcessComponent.h"
+#include "Components/ArrowComponent.h"
+#include "Components/BillboardComponent.h"
+#include "Components/BrushComponent.h"
+#include "Components/DrawFrustumComponent.h"
+#include "AI/EnvironmentQuery/EQSRenderingComponent.h"
+#include "Debug/GameplayDebuggingComponent.h"
+#include "Components/LineBatchComponent.h"
+#include "Components/MaterialBillboardComponent.h"
+#include "Components/MeshComponent.h"
+#include "Components/SkinnedMeshComponent.h"
+#include "Components/DestructibleComponent.h"
+#include "Components/PoseableMeshComponent.h"
+#include "Components/SkeletalMeshComponent.h"
+#include "Components/StaticMeshComponent.h"
+#include "Components/InstancedStaticMeshComponent.h"
+#include "Components/InteractiveFoliageComponent.h"
+#include "Components/SplineMeshComponent.h"
+#include "Components/ModelComponent.h"
+#include "AI/Navigation/NavLinkRenderingComponent.h"
+#include "AI/Navigation/NavMeshRenderingComponent.h"
+#include "AI/Navigation/NavTestRenderingComponent.h"
+#include "Components/NiagaraComponent.h"
+#include "Components/ShapeComponent.h"
+#include "Components/BoxComponent.h"
+#include "Components/CapsuleComponent.h"
+#include "Components/SphereComponent.h"
+#include "Components/DrawSphereComponent.h"
+#include "Components/TextRenderComponent.h"
+#include "Components/VectorFieldComponent.h"
+#include "PhysicsEngine/RadialForceComponent.h"
+#include "GameFramework/SpringArmComponent.h"
+#include "Components/WindDirectionalSourceComponent.h"
+#include "Components/TimelineComponent.h"
+#include "Slate.h"
+#include "AI/NavDataGenerator.h"
+#include "OnlineSubsystemUtils.h"
+
 #include "StaticMeshLight.h"
 
 #if WITH_PHYSX
diff --git a/Engine/Source/Runtime/Engine/Private/Interpolation.cpp b/Engine/Source/Runtime/Engine/Private/Interpolation.cpp
--- a/Engine/Source/Runtime/Engine/Private/Interpolation.cpp
+++ b/Engine/Source/Runtime/Engine/Private/Interpolation.cpp
@@ -6,16 +6,11 @@
 
 #include "EnginePrivate.h"
 #include "EngineClasses.h"
-#include "EngineInterpolationClasses.h"
-#include "EngineKismetLibraryClasses.h"
 #include "ParticleDefinitions.h"
-#include "EngineMaterialClasses.h"
-#include "EngineLevelScriptClasses.h"
 #include "SoundDefinitions.h"
 #include "InterpolationHitProxy.h"
 #include "AVIWriter.h"
 #include "AnimationUtils.h"
-#include "EngineDecalClasses.h"
 #include "MatineeUtils.h"
 
 #if WITH_EDITOR
diff --git a/Engine/Source/Runtime/Engine/Private/KismetSystemLibrary.cpp b/Engine/Source/Runtime/Engine/Private/KismetSystemLibrary.cpp
--- a/Engine/Source/Runtime/Engine/Private/KismetSystemLibrary.cpp
+++ b/Engine/Source/Runtime/Engine/Private/KismetSystemLibrary.cpp
@@ -1,8 +1,6 @@
 // Copyright 1998-2014 Epic Games, Inc. All Rights Reserved.
 
 #include "EnginePrivate.h"
-#include "EngineKismetLibraryClasses.h"
-#include "EngineUserInterfaceClasses.h"
 #include "LatentActions.h"
 #include "DelayAction.h"
 #include "InterpolateComponentToAction.h"
diff --git a/Engine/Source/Runtime/Engine/Private/Landscape/LandscapeEdit.cpp b/Engine/Source/Runtime/Engine/Private/Landscape/LandscapeEdit.cpp
--- a/Engine/Source/Runtime/Engine/Private/Landscape/LandscapeEdit.cpp
+++ b/Engine/Source/Runtime/Engine/Private/Landscape/LandscapeEdit.cpp
@@ -5,7 +5,6 @@ LandscapeEdit.cpp: Landscape editing
 =============================================================================*/
 
 #include "EnginePrivate.h"
-#include "EngineFoliageClasses.h"
 #include "Landscape/LandscapeDataAccess.h"
 #include "Landscape/LandscapeEdit.h"
 #include "Landscape/LandscapeRender.h"
@@ -2677,7 +2676,6 @@ void ULandscapeInfo::ExportLayer(ULandscapeLayerInfoObject* LayerInfo, const FSt
 //////////////////////////////////////////////////////////////////////////
 // Apply splines
 //////////////////////////////////////////////////////////////////////////
-#include "EngineSplineClasses.h"
 #include "Raster.h"
 
 class FLandscapeSplineHeightsRasterPolicy
diff --git a/Engine/Source/Runtime/Engine/Private/Level.cpp b/Engine/Source/Runtime/Engine/Private/Level.cpp
--- a/Engine/Source/Runtime/Engine/Private/Level.cpp
+++ b/Engine/Source/Runtime/Engine/Private/Level.cpp
@@ -7,7 +7,6 @@ Level.cpp: Level-related functions
 #include "EnginePrivate.h"
 #include "Net/UnrealNetwork.h"
 #include "SoundDefinitions.h"
-#include "EngineMaterialClasses.h"
 #include "PrecomputedLightVolume.h"
 #include "TickTaskManagerInterface.h"
 #include "BlueprintUtilities.h"
@@ -15,7 +14,6 @@ Level.cpp: Level-related functions
 #include "Editor/UnrealEd/Public/Kismet2/KismetEditorUtilities.h"
 #include "Editor/UnrealEd/Public/Kismet2/BlueprintEditorUtils.h"
 #endif
-#include "EngineLevelScriptClasses.h"
 #include "Runtime/Engine/Classes/MovieScene/RuntimeMovieScenePlayerInterface.h"
 #include "LevelUtils.h"
 #include "TargetPlatform.h"
diff --git a/Engine/Source/Runtime/Engine/Private/LevelActor.cpp b/Engine/Source/Runtime/Engine/Private/LevelActor.cpp
--- a/Engine/Source/Runtime/Engine/Private/LevelActor.cpp
+++ b/Engine/Source/Runtime/Engine/Private/LevelActor.cpp
@@ -3,8 +3,6 @@
 
 #include "EnginePrivate.h"
 #include "Net/UnrealNetwork.h"
-#include "EngineLevelScriptClasses.h"
-#include "EngineKismetLibraryClasses.h"
 #include "Collision.h"
 
 #if WITH_PHYSX
diff --git a/Engine/Source/Runtime/Engine/Private/LevelScriptActor.cpp b/Engine/Source/Runtime/Engine/Private/LevelScriptActor.cpp
--- a/Engine/Source/Runtime/Engine/Private/LevelScriptActor.cpp
+++ b/Engine/Source/Runtime/Engine/Private/LevelScriptActor.cpp
@@ -2,10 +2,7 @@
 
 
 #include "EnginePrivate.h"
-#include "EngineLevelScriptClasses.h"
 #include "SoundDefinitions.h"
-#include "EngineUserInterfaceClasses.h"
-#include "EngineKismetLibraryClasses.h"
 #include "BlueprintUtilities.h"
 
 
diff --git a/Engine/Source/Runtime/Engine/Private/LevelTick.cpp b/Engine/Source/Runtime/Engine/Private/LevelTick.cpp
--- a/Engine/Source/Runtime/Engine/Private/LevelTick.cpp
+++ b/Engine/Source/Runtime/Engine/Private/LevelTick.cpp
@@ -8,8 +8,6 @@
 #include "Net/UnrealNetwork.h"
 #include "Collision.h"
 
-#include "EngineUserInterfaceClasses.h"
-#include "EngineDecalClasses.h"
 #include "ParticleDefinitions.h"
 //#include "SoundDefinitions.h"
 #include "FXSystem.h"
diff --git a/Engine/Source/Runtime/Engine/Private/LightComponent.cpp b/Engine/Source/Runtime/Engine/Private/LightComponent.cpp
--- a/Engine/Source/Runtime/Engine/Private/LightComponent.cpp
+++ b/Engine/Source/Runtime/Engine/Private/LightComponent.cpp
@@ -5,7 +5,6 @@
 =============================================================================*/
 
 #include "EnginePrivate.h"
-#include "EngineLightClasses.h"
 #if WITH_EDITOR
 #include "ObjectEditorUtils.h"
 #endif
diff --git a/Engine/Source/Runtime/Engine/Private/LocalPlayer.cpp b/Engine/Source/Runtime/Engine/Private/LocalPlayer.cpp
--- a/Engine/Source/Runtime/Engine/Private/LocalPlayer.cpp
+++ b/Engine/Source/Runtime/Engine/Private/LocalPlayer.cpp
@@ -1,8 +1,6 @@
 // Copyright 1998-2014 Epic Games, Inc. All Rights Reserved.
 
 #include "EnginePrivate.h"
-#include "EngineUserInterfaceClasses.h"
-#include "EngineInterpolationClasses.h"
 
 #include "SubtitleManager.h"
 #include "Net/UnrealNetwork.h"
diff --git a/Engine/Source/Runtime/Engine/Private/Pawn.cpp b/Engine/Source/Runtime/Engine/Private/Pawn.cpp
--- a/Engine/Source/Runtime/Engine/Private/Pawn.cpp
+++ b/Engine/Source/Runtime/Engine/Private/Pawn.cpp
@@ -10,10 +10,7 @@
 #include "Net/UnrealNetwork.h"
 #include "ConfigCacheIni.h"
 #include "NavigationPathBuilder.h"
-#include "EngineInterpolationClasses.h"
-#include "EngineKismetLibraryClasses.h"
 #include "ParticleDefinitions.h"
-#include "EngineComponentClasses.h"
 #include "DisplayDebugHelpers.h"
 
 DEFINE_LOG_CATEGORY(LogDamage);
diff --git a/Engine/Source/Runtime/Engine/Private/PlayerController.cpp b/Engine/Source/Runtime/Engine/Private/PlayerController.cpp
--- a/Engine/Source/Runtime/Engine/Private/PlayerController.cpp
+++ b/Engine/Source/Runtime/Engine/Private/PlayerController.cpp
@@ -3,11 +3,7 @@
 #include "EnginePrivate.h"
 #include "Net/UnrealNetwork.h"
 #include "ConfigCacheIni.h"
-#include "EngineKismetLibraryClasses.h"
 #include "SoundDefinitions.h"
-#include "EngineUserInterfaceClasses.h"
-#include "EngineInterpolationClasses.h"
-#include "EngineLevelScriptClasses.h"
 #include "OnlineSubsystemUtils.h"
 #include "IHeadMountedDisplay.h"
 #include "IForceFeedbackSystem.h"
diff --git a/Engine/Source/Runtime/Engine/Private/ScriptPlatformInterface.cpp b/Engine/Source/Runtime/Engine/Private/ScriptPlatformInterface.cpp
--- a/Engine/Source/Runtime/Engine/Private/ScriptPlatformInterface.cpp
+++ b/Engine/Source/Runtime/Engine/Private/ScriptPlatformInterface.cpp
@@ -5,8 +5,6 @@
 =============================================================================*/ 
 
 #include "EnginePrivate.h"
-#include "EngineGameEngineClasses.h"
-#include "EnginePlatformInterfaceClasses.h"
 
 DEFINE_LOG_CATEGORY_STATIC(LogScriptPlatformInterface, Log, All);
 
diff --git a/Engine/Source/Runtime/Engine/Private/SkyLightComponent.cpp b/Engine/Source/Runtime/Engine/Private/SkyLightComponent.cpp
--- a/Engine/Source/Runtime/Engine/Private/SkyLightComponent.cpp
+++ b/Engine/Source/Runtime/Engine/Private/SkyLightComponent.cpp
@@ -5,7 +5,6 @@
 =============================================================================*/
 
 #include "EnginePrivate.h"
-#include "EngineLightClasses.h"
 #if WITH_EDITOR
 #include "ObjectEditorUtils.h"
 #endif
diff --git a/Engine/Source/Runtime/Engine/Private/UserInterface/Canvas.cpp b/Engine/Source/Runtime/Engine/Private/UserInterface/Canvas.cpp
--- a/Engine/Source/Runtime/Engine/Private/UserInterface/Canvas.cpp
+++ b/Engine/Source/Runtime/Engine/Private/UserInterface/Canvas.cpp
@@ -5,9 +5,7 @@
 =============================================================================*/
 
 #include "EnginePrivate.h"
-#include "EngineUserInterfaceClasses.h"
 #include "TileRendering.h"
-#include "EngineMaterialClasses.h"
 #include "RHIStaticStates.h"
 #include "WordWrapper.h"
 
diff --git a/Engine/Source/Runtime/Engine/Private/UserInterface/CanvasItem.cpp b/Engine/Source/Runtime/Engine/Private/UserInterface/CanvasItem.cpp
--- a/Engine/Source/Runtime/Engine/Private/UserInterface/CanvasItem.cpp
+++ b/Engine/Source/Runtime/Engine/Private/UserInterface/CanvasItem.cpp
@@ -6,9 +6,7 @@
 
 #include "EnginePrivate.h"
 #include "CanvasItem.h"
-#include "EngineUserInterfaceClasses.h"
 #include "TileRendering.h"
-#include "EngineMaterialClasses.h"
 #include "RHIStaticStates.h"
 #include "Engine.h"
 #include "Internationalization.h"
diff --git a/Engine/Source/Runtime/Engine/Private/UserInterface/Console.cpp b/Engine/Source/Runtime/Engine/Private/UserInterface/Console.cpp
--- a/Engine/Source/Runtime/Engine/Private/UserInterface/Console.cpp
+++ b/Engine/Source/Runtime/Engine/Private/UserInterface/Console.cpp
@@ -5,8 +5,6 @@
 =============================================================================*/
 
 #include "EnginePrivate.h"
-#include "EngineUserInterfaceClasses.h"
-#include "EngineLevelScriptClasses.h"
 #include "Slate.h"
 #include "DefaultValueHelper.h"
 
diff --git a/Engine/Source/Runtime/Engine/Private/World.cpp b/Engine/Source/Runtime/Engine/Private/World.cpp
--- a/Engine/Source/Runtime/Engine/Private/World.cpp
+++ b/Engine/Source/Runtime/Engine/Private/World.cpp
@@ -6,8 +6,6 @@
 
 #include "EnginePrivate.h"
 #include "Net/UnrealNetwork.h"
-#include "EngineLevelScriptClasses.h"
-#include "EngineUserInterfaceClasses.h"
 #include "ParticleDefinitions.h"
 #include "Database.h"
 #include "Net/NetworkProfiler.h"
diff --git a/Engine/Source/Runtime/Engine/Public/Audio.h b/Engine/Source/Runtime/Engine/Public/Audio.h
--- a/Engine/Source/Runtime/Engine/Public/Audio.h
+++ b/Engine/Source/Runtime/Engine/Public/Audio.h
@@ -5,6 +5,7 @@
 =============================================================================*/
 
 #pragma once
+#include "Sound/SoundClass.h"
 
 ENGINE_API DECLARE_LOG_CATEGORY_EXTERN(LogAudio, Warning, All);
 
@@ -97,6 +98,8 @@ enum EAudioSpeakers
 
 // Forward declarations.
 class UAudioComponent;
+class USoundNode;
+struct FWaveInstance;
 struct FReverbSettings;
 struct FSampleLoop;
 
diff --git a/Engine/Source/Runtime/Engine/Public/Engine.h b/Engine/Source/Runtime/Engine/Public/Engine.h
--- a/Engine/Source/Runtime/Engine/Public/Engine.h
+++ b/Engine/Source/Runtime/Engine/Public/Engine.h
@@ -602,7 +602,6 @@ struct FURL;
 #include "RenderUtils.h"					// Render utility classes.
 #include "Distributions.h"					// Distributions
 #include "PhysxUserData.h"
-#include "EngineSceneClasses.h"
 #include "EngineClasses.h"					// All actor classes.
 #include "VisualLog.h"
 #include "MaterialShared.h"					// Shared material definitions.
@@ -615,10 +614,7 @@ struct FURL;
 #include "LightMap.h"						// Light-maps.
 #include "ShadowMap.h"
 #include "Model.h"							// Model class.
-#include "EngineComponentClasses.h"
-#include "EngineMaterialClasses.h"
-#include "EngineStaticMeshClasses.h"
-#include "EngineNavigationClasses.h"
+
 #include "AI/NavDataGenerator.h"
 #include "AI/NavLinkRenderingProxy.h"
 #include "AI/NavigationModifier.h"
@@ -626,16 +622,8 @@ struct FURL;
 #include "StaticMeshResources.h"			// Static meshes.
 #include "AnimTree.h"						// Animation.
 #include "SkeletalMeshTypes.h"				// Skeletal animated mesh.
-#include "EngineSkeletalMeshClasses.h"
-#include "EngineLightClasses.h"
 #include "Animation/SkeletalMeshActor.h"
 #include "Interpolation.h"					// Matinee.
-#include "EngineInterpolationClasses.h"
-#include "EnginePawnClasses.h"
-#include "EngineReplicationInfoClasses.h"
-#include "EngineGameEngineClasses.h"		// Main Unreal engine declarations
-#include "EngineTerrainClasses.h"
-#include "EngineNetworkClasses.h"
 #include "ContentStreaming.h"				// Content streaming class definitions.
 #include "LightingBuildOptions.h"			// Definition of lighting build option struct.
 #include "PixelFormat.h"
@@ -650,8 +638,6 @@ struct FURL;
 #include "TimerManager.h"					// Game play timers
 #include "EngineService.h"
 #include "AI/NavigationSystemHelpers.h"
-#include "EngineAnimationNodeClasses.h"
-#include "EngineNavigationClasses.h"
 #include "HardwareInfo.h"
 
 /** Implements the engine module. */
diff --git a/Engine/Source/Runtime/Engine/Public/InstancedFoliage.h b/Engine/Source/Runtime/Engine/Public/InstancedFoliage.h
--- a/Engine/Source/Runtime/Engine/Public/InstancedFoliage.h
+++ b/Engine/Source/Runtime/Engine/Public/InstancedFoliage.h
@@ -6,6 +6,12 @@
 
 #pragma once
 
+#include "Components/SceneComponent.h"
+//
+// Forward declarations.
+//
+class UInstancedStaticMeshComponent;
+
 /**
  * Flags stored with each instance
  */
diff --git a/Engine/Source/Runtime/Engine/Public/LightMap.h b/Engine/Source/Runtime/Engine/Public/LightMap.h
--- a/Engine/Source/Runtime/Engine/Public/LightMap.h
+++ b/Engine/Source/Runtime/Engine/Public/LightMap.h
@@ -7,6 +7,7 @@
 #pragma once
 
 #include "ChunkedArray.h"
+#include "Engine/LightMapTexture2D.h"
 
 /** Whether to use bilinear filtering on lightmaps */
 extern ENGINE_API bool GUseBilinearLightmaps;
diff --git a/Engine/Source/Runtime/Engine/Public/ParticleDefinitions.h b/Engine/Source/Runtime/Engine/Public/ParticleDefinitions.h
--- a/Engine/Source/Runtime/Engine/Public/ParticleDefinitions.h
+++ b/Engine/Source/Runtime/Engine/Public/ParticleDefinitions.h
@@ -9,4 +9,3 @@
 #include "Engine.h"
 #include "ParticleHelper.h"
 #include "ParticleEmitterInstances.h"
-#include "EngineParticleClasses.h"
diff --git a/Engine/Source/Runtime/Engine/Public/ParticleHelper.h b/Engine/Source/Runtime/Engine/Public/ParticleHelper.h
--- a/Engine/Source/Runtime/Engine/Public/ParticleHelper.h
+++ b/Engine/Source/Runtime/Engine/Public/ParticleHelper.h
@@ -134,6 +134,8 @@ struct FParticleSpriteEmitterInstance;
 struct FParticleMeshEmitterInstance;
 struct FParticleBeam2EmitterInstance;
 
+struct FStaticMeshLODResources;
+
 // Special module indices...
 #define INDEX_TYPEDATAMODULE	(INDEX_NONE - 1)
 #define INDEX_REQUIREDMODULE	(INDEX_NONE - 2)
diff --git a/Engine/Source/Runtime/Engine/Public/SceneManagement.h b/Engine/Source/Runtime/Engine/Public/SceneManagement.h
--- a/Engine/Source/Runtime/Engine/Public/SceneManagement.h
+++ b/Engine/Source/Runtime/Engine/Public/SceneManagement.h
@@ -9,7 +9,6 @@
 // Includes the draw mesh macros
 #include "SceneUtils.h"
 #include "UniformBuffer.h"
-#include "EngineSceneClasses.h"
 #include "BufferVisualizationData.h"
 #include "ConvexVolume.h"
 #include "SystemSettings.h"
diff --git a/Engine/Source/Runtime/Online/OnlineSubsystemUtils/Private/OnlineSessionClient.cpp b/Engine/Source/Runtime/Online/OnlineSubsystemUtils/Private/OnlineSessionClient.cpp
--- a/Engine/Source/Runtime/Online/OnlineSubsystemUtils/Private/OnlineSessionClient.cpp
+++ b/Engine/Source/Runtime/Online/OnlineSubsystemUtils/Private/OnlineSessionClient.cpp
@@ -6,7 +6,6 @@
 =============================================================================*/
 
 #include "OnlineSubsystemUtilsPrivatePCH.h"
-#include "EngineNetworkClasses.h"
 
 #define INVALID_CONTROLLERID 255
 
diff --git a/Engine/Source/Runtime/SkillSystem/Classes/AttributeComponent.h b/Engine/Source/Runtime/SkillSystem/Classes/AttributeComponent.h
--- a/Engine/Source/Runtime/SkillSystem/Classes/AttributeComponent.h
+++ b/Engine/Source/Runtime/SkillSystem/Classes/AttributeComponent.h
@@ -21,7 +21,7 @@ struct SKILLSYSTEM_API FAttributeDefaults
 /** 
  * 
  */
-UCLASS(HeaderGroup=Component, ClassGroup=SkillSystem, hidecategories=(Object,LOD,Lighting,Transform,Sockets,TextureStreaming), editinlinenew, meta=(BlueprintSpawnableComponent))
+UCLASS(ClassGroup=SkillSystem, hidecategories=(Object,LOD,Lighting,Transform,Sockets,TextureStreaming), editinlinenew, meta=(BlueprintSpawnableComponent))
 class SKILLSYSTEM_API UAttributeComponent : public UActorComponent
 {
 	GENERATED_UCLASS_BODY()
diff --git a/Engine/Source/Runtime/SkillSystem/Private/SkillSystemModulePrivatePCH.h b/Engine/Source/Runtime/SkillSystem/Private/SkillSystemModulePrivatePCH.h
--- a/Engine/Source/Runtime/SkillSystem/Private/SkillSystemModulePrivatePCH.h
+++ b/Engine/Source/Runtime/SkillSystem/Private/SkillSystemModulePrivatePCH.h
@@ -14,7 +14,6 @@
 #include "Engine.h"
 #include "Engine/DataTable.h"
 #include "SkillSystemClasses.h"
-#include "SkillSystemComponentClasses.h"
 
 #include "SkillSystemModule.h"
 #include "SkillSystemLog.h"
