Fixed crash which can occur when undoing/redoing placing a material function in the material editor
#ttp 332887 - EDITOR: MATERIAL EDITOR: CRASH: Changing the name of a material function that is placed inside a material can lead to a crash with undo/redo
#branch UE4
#proj Runtime.Engine
#add Corrected UEdGraph::CreateNode so that the graph node is constructed with ConstructObject<>, passing RF_Transactional, instead of constructing it with NewObject<> and setting the RF_Transactional flag afterwards.  This ensures that a version of the object with RF_PendingKill set is recorded in the undo buffer (as graph nodes are created inside a scoped transaction), meaning that a GC sweep will not attempt to clean out the reference to the MaterialExpression when the transaction is undone (as the MaterialExpression is also created during the transaction, and will be marked as RF_PendingKill in the undo buffer).

diff --git a/Engine/Source/Runtime/Engine/Private/EdGraph/EdGraph.cpp b/Engine/Source/Runtime/Engine/Private/EdGraph/EdGraph.cpp
--- a/Engine/Source/Runtime/Engine/Private/EdGraph/EdGraph.cpp
+++ b/Engine/Source/Runtime/Engine/Private/EdGraph/EdGraph.cpp
@@ -251,8 +251,7 @@ void UEdGraph::RemoveOnGraphChangedHandler( const FOnGraphChanged::FDelegate& In
 
 UEdGraphNode* UEdGraph::CreateNode( TSubclassOf<UEdGraphNode> NewNodeClass, bool bSelectNewNode/* = true*/ )
 {
-	UEdGraphNode* NewNode = NewObject<UEdGraphNode>(this, NewNodeClass);
-	NewNode->SetFlags(RF_Transactional);
+	UEdGraphNode* NewNode = ConstructObject<UEdGraphNode>(NewNodeClass, this, NAME_None, RF_Transactional);
 	AddNode(NewNode, false, bSelectNewNode );
 	return NewNode;
 }
