Fix issues with using mouse for touch and Simulate mode
TTP #332927

diff --git a/Engine/Source/Editor/LevelEditor/Private/SLevelViewport.cpp b/Engine/Source/Editor/LevelEditor/Private/SLevelViewport.cpp
--- a/Engine/Source/Editor/LevelEditor/Private/SLevelViewport.cpp
+++ b/Engine/Source/Editor/LevelEditor/Private/SLevelViewport.cpp
@@ -2873,7 +2873,7 @@ bool SLevelViewport::IsViewportConfigurationSet(FName ConfigurationName) const
 	return false;
 }
 
-void SLevelViewport::StartPlayInEditorSession( UGameViewportClient* PlayClient )
+void SLevelViewport::StartPlayInEditorSession( UGameViewportClient* PlayClient, const bool bInSimulateInEditor )
 {
 	check( !HasPlayInEditorViewport() );
 
@@ -2898,6 +2898,7 @@ void SLevelViewport::StartPlayInEditorSession( UGameViewportClient* PlayClient )
 
 	// Whether to start with the game taking mouse control or leaving it shown in the editor
 	ActiveViewport->SetPlayInEditorGetsMouseControl(GetDefault<ULevelEditorPlaySettings>()->GameGetsMouseControl);
+	ActiveViewport->SetPlayInEditorIsSimulate(bInSimulateInEditor);
 	
 	ActiveViewport->OnPlayWorldViewportSwapped( *InactiveViewport );
 
diff --git a/Engine/Source/Editor/LevelEditor/Public/ILevelViewport.h b/Engine/Source/Editor/LevelEditor/Public/ILevelViewport.h
--- a/Engine/Source/Editor/LevelEditor/Public/ILevelViewport.h
+++ b/Engine/Source/Editor/LevelEditor/Public/ILevelViewport.h
@@ -14,7 +14,7 @@ public:
 	 *
 	 * @param PlayClient	The client to swap with the editor client
 	 */
-	virtual void StartPlayInEditorSession( class UGameViewportClient* PlayClient ) = 0;
+	virtual void StartPlayInEditorSession( class UGameViewportClient* PlayClient, const bool bInSimulateInEditor ) = 0;
 
 	/**
 	 * Ends a currently active play in editor session in this viewport
diff --git a/Engine/Source/Editor/LevelEditor/Public/SLevelViewport.h b/Engine/Source/Editor/LevelEditor/Public/SLevelViewport.h
--- a/Engine/Source/Editor/LevelEditor/Public/SLevelViewport.h
+++ b/Engine/Source/Editor/LevelEditor/Public/SLevelViewport.h
@@ -75,7 +75,7 @@ public:
 	void SaveConfig(const FString& ConfigName);
 
 	/** ILevelViewport Interface */
-	virtual void StartPlayInEditorSession( UGameViewportClient* PlayClient ) OVERRIDE;
+	virtual void StartPlayInEditorSession( UGameViewportClient* PlayClient, const bool bInSimulateInEditor ) OVERRIDE;
 	virtual void EndPlayInEditorSession() OVERRIDE;
 	virtual void SwapViewportsForSimulateInEditor() OVERRIDE;
 	virtual void SwapViewportsForPlayInEditor() OVERRIDE;
diff --git a/Engine/Source/Editor/UnrealEd/Private/PlayLevel.cpp b/Engine/Source/Editor/UnrealEd/Private/PlayLevel.cpp
--- a/Engine/Source/Editor/UnrealEd/Private/PlayLevel.cpp
+++ b/Engine/Source/Editor/UnrealEd/Private/PlayLevel.cpp
@@ -2266,7 +2266,7 @@ UWorld* UEditorEngine::CreatePlayInEditorWorld(FWorldContext &PieWorldContext, b
 			{
 				TSharedPtr<ILevelViewport> LevelViewportRef = SlatePlayInEditorSession.DestinationSlateViewport.Pin();
 
-				LevelViewportRef->StartPlayInEditorSession( ViewportClient );
+				LevelViewportRef->StartPlayInEditorSession( ViewportClient, bInSimulateInEditor );
 			}
 			else
 			{		
@@ -2554,6 +2554,9 @@ void UEditorEngine::ToggleBetweenPIEandSIE()
 			// The undo system may have a reference to a SIE object that is about to be destroyed, so clear the transactions
 			ResetTransaction( NSLOCTEXT("UnrealEd", "ToggleBetweenPIEandSIE", "Toggle Between PIE and SIE") );
 
+			// The Game's viewport needs to know about the change away from simluate before the PC is (potentially) created
+			GameViewport->GetGameViewport()->SetPlayInEditorIsSimulate(false);
+
 			// The editor viewport client wont be visible so temporarily disable it being realtime
 			EditorViewportClient.SetRealtime( false, true );
 
@@ -2566,7 +2569,7 @@ void UEditorEngine::ToggleBetweenPIEandSIE()
 				{
 					APlayerController* PC = World->GetFirstPlayerController();
 					PC->PlayerState->bOnlySpectator = false;
-					World->GetAuthGameMode()->RestartPlayer(World->GetFirstPlayerController());
+					World->GetAuthGameMode()->RestartPlayer(PC);
 				}
 
 				OnSwitchWorldsForPIE(false);
@@ -2576,7 +2579,8 @@ void UEditorEngine::ToggleBetweenPIEandSIE()
 			LevelViewport->SwapViewportsForPlayInEditor();
 
 			// No longer simulating
-			EditorViewportClient.SetIsSimulateInEditorViewport( false );
+			GameViewport->SetIsSimulateInEditorViewport(false);
+			EditorViewportClient.SetIsSimulateInEditorViewport(false);
 			bIsSimulatingInEditor = false;
 		}
 		else
@@ -2584,7 +2588,9 @@ void UEditorEngine::ToggleBetweenPIEandSIE()
 			// Swap to simulate from PIE
 			LevelViewport->SwapViewportsForSimulateInEditor();
 	
-			EditorViewportClient.SetIsSimulateInEditorViewport( true );
+			GameViewport->SetIsSimulateInEditorViewport(true);
+			GameViewport->GetGameViewport()->SetPlayInEditorIsSimulate(true);
+			EditorViewportClient.SetIsSimulateInEditorViewport(true);
 			bIsSimulatingInEditor = true;
 
 			// Make sure the viewport is in real-time mode
diff --git a/Engine/Source/Editor/UnrealEd/Public/LevelEditorViewport.h b/Engine/Source/Editor/UnrealEd/Public/LevelEditorViewport.h
--- a/Engine/Source/Editor/UnrealEd/Public/LevelEditorViewport.h
+++ b/Engine/Source/Editor/UnrealEd/Public/LevelEditorViewport.h
@@ -297,7 +297,7 @@ public:
 	/** Updates the rotate widget with the passed in delta rotation. */
 	void ApplyDeltaToRotateWidget( const FRotator& InRot );
 
-	void SetIsSimulateInEditorViewport( bool bInIsSimulateInEditorViewport );
+	virtual void SetIsSimulateInEditorViewport( bool bInIsSimulateInEditorViewport ) OVERRIDE;
 
 	/**
 	 * Draws a screen space bounding box around the specified actor
diff --git a/Engine/Source/Runtime/Engine/Classes/Engine/GameViewportClient.h b/Engine/Source/Runtime/Engine/Classes/Engine/GameViewportClient.h
--- a/Engine/Source/Runtime/Engine/Classes/Engine/GameViewportClient.h
+++ b/Engine/Source/Runtime/Engine/Classes/Engine/GameViewportClient.h
@@ -246,6 +246,8 @@ public:
 
 	virtual void SetReferenceToWorldContext(struct FWorldContext& WorldContext);
 
+	class FSceneViewport* GetGameViewport();
+
 public:
 	// Begin UObject Interface
 	virtual void PostInitProperties() OVERRIDE;
@@ -271,6 +273,7 @@ public:
 	virtual bool IsOrtho() const OVERRIDE;
 	virtual void MouseEnter(FViewport* Viewport, int32 x, int32 y) OVERRIDE;
 	virtual void MouseLeave(FViewport* Viewport) OVERRIDE;
+	virtual void SetIsSimulateInEditorViewport(bool bInIsSimulateInEditorViewport) OVERRIDE;
 	// End of FViewportClient interface.
 
 	// Begin FExec interface.
diff --git a/Engine/Source/Runtime/Engine/Private/GameViewportClient.cpp b/Engine/Source/Runtime/Engine/Private/GameViewportClient.cpp
--- a/Engine/Source/Runtime/Engine/Private/GameViewportClient.cpp
+++ b/Engine/Source/Runtime/Engine/Private/GameViewportClient.cpp
@@ -119,6 +119,11 @@ void UGameViewportClient::DetachViewportClient()
 	RemoveFromRoot();
 }
 
+FSceneViewport* UGameViewportClient::GetGameViewport()
+{
+	return static_cast<FSceneViewport*>(Viewport);
+}
+
 
 void UGameViewportClient::Tick( float DeltaTime )
 {
@@ -279,19 +284,42 @@ bool UGameViewportClient::InputMotion(FViewport* InViewport, int32 ControllerId,
 	return bResult;
 }
 
-void UGameViewportClient::MouseEnter(FViewport* Viewport, int32 x, int32 y)
+void UGameViewportClient::SetIsSimulateInEditorViewport(bool bInIsSimulateInEditorViewport)
 {
-	Super::MouseEnter(Viewport, x, y);
-
 	if (GetDefault<UInputSettings>()->bUseMouseForTouch)
+	{
+		FSlateApplication::Get().SetGameIsFakingTouchEvents(!bInIsSimulateInEditorViewport);
+	}
+
+	for (ULocalPlayer* LocalPlayer : GetOuterUEngine()->GetGamePlayers(this))
+	{
+		if (LocalPlayer->PlayerController)
+		{
+			if (bInIsSimulateInEditorViewport)
+			{
+				LocalPlayer->PlayerController->CleanupGameViewport();
+			}
+			else
+			{
+				LocalPlayer->PlayerController->CreateTouchInterface();
+			}
+		}
+	}
+}
+
+void UGameViewportClient::MouseEnter(FViewport* InViewport, int32 x, int32 y)
+{
+	Super::MouseEnter(InViewport, x, y);
+
+	if (GetDefault<UInputSettings>()->bUseMouseForTouch && !GetGameViewport()->GetPlayInEditorIsSimulate())
 	{
 		FSlateApplication::Get().SetGameIsFakingTouchEvents(true);
 	}
 }
 
-void UGameViewportClient::MouseLeave(FViewport* Viewport)
+void UGameViewportClient::MouseLeave(FViewport* InViewport)
 {
-	Super::MouseLeave(Viewport);
+	Super::MouseLeave(InViewport);
 
 	if (GetDefault<UInputSettings>()->bUseMouseForTouch)
 	{
@@ -331,7 +359,7 @@ bool UGameViewportClient::RequiresUncapturedAxisInput() const
 		}
 		else if (GetWorld() && GetWorld()->GetFirstPlayerController())
 		{
-			bRequired = GetWorld()->GetFirstPlayerController() ->ShouldShowMouseCursor();
+			bRequired = GetWorld()->GetFirstPlayerController()->ShouldShowMouseCursor();
 		}
 	}
 
@@ -1061,7 +1089,7 @@ void UGameViewportClient::ReceivedFocus(FViewport* InViewport)
 {
 	InViewport->CaptureJoystickInput(true);
 
-	if (GetDefault<UInputSettings>()->bUseMouseForTouch)
+	if (GetDefault<UInputSettings>()->bUseMouseForTouch && !GetGameViewport()->GetPlayInEditorIsSimulate())
 	{
 		FSlateApplication::Get().SetGameIsFakingTouchEvents(true);
 	}
diff --git a/Engine/Source/Runtime/Engine/Private/Slate/SceneViewport.cpp b/Engine/Source/Runtime/Engine/Private/Slate/SceneViewport.cpp
--- a/Engine/Source/Runtime/Engine/Private/Slate/SceneViewport.cpp
+++ b/Engine/Source/Runtime/Engine/Private/Slate/SceneViewport.cpp
@@ -31,6 +31,7 @@ FSceneViewport::FSceneViewport( FViewportClient* InViewportClient, TSharedPtr<SV
 	, bUseSeparateRenderTarget( InViewportWidget.IsValid() ? !InViewportWidget->ShouldRenderDirectly() : true )
 	, bIsResizing( false )
 	, bPlayInEditorGetsMouseControl( true )
+	, bPlayInEditorIsSimulate( false )
 {
 	bIsSlateViewport = true;
 }
diff --git a/Engine/Source/Runtime/Engine/Public/Slate/SceneViewport.h b/Engine/Source/Runtime/Engine/Public/Slate/SceneViewport.h
--- a/Engine/Source/Runtime/Engine/Public/Slate/SceneViewport.h
+++ b/Engine/Source/Runtime/Engine/Public/Slate/SceneViewport.h
@@ -156,11 +156,20 @@ public:
 	* Sets whether a PIE viewport takes mouse control on startup.
 	* @param bGetsMouseControl Takes control if true, or not if false. 
 	*/
-	void SetPlayInEditorGetsMouseControl( bool bGetsMouseControl )
+	void SetPlayInEditorGetsMouseControl(const bool bGetsMouseControl)
 	{
 		bPlayInEditorGetsMouseControl = bGetsMouseControl;
 	}
 
+	void SetPlayInEditorIsSimulate(const bool bIsSimulate)
+	{
+		bPlayInEditorIsSimulate = bIsSimulate;
+	}
+	bool GetPlayInEditorIsSimulate() const
+	{
+		return bPlayInEditorIsSimulate;
+	}
+
 	/** Updates the viewport RHI with a new size and fullscreen flag */
 	virtual void UpdateViewportRHI(bool bDestroyed,uint32 NewSizeX,uint32 NewSizeY,EWindowMode::Type NewWindowMode) OVERRIDE;
 
@@ -289,6 +298,8 @@ private:
 	FOnSceneViewportResize OnSceneViewportResizeDel;
 	/** Whether a PIE viewport should take mouse control on startup */
 	bool bPlayInEditorGetsMouseControl;
+	/** Whether the PIE viewport is currently in simulate in editor mode */
+	bool bPlayInEditorIsSimulate;
 };
 
 
diff --git a/Engine/Source/Runtime/Engine/Public/UnrealClient.h b/Engine/Source/Runtime/Engine/Public/UnrealClient.h
--- a/Engine/Source/Runtime/Engine/Public/UnrealClient.h
+++ b/Engine/Source/Runtime/Engine/Public/UnrealClient.h
@@ -619,6 +619,7 @@ public:
 	 */
 	virtual bool InputMotion(FViewport* Viewport, int32 ControllerId, const FVector& Tilt, const FVector& RotationRate, const FVector& Gravity, const FVector& Acceleration) { return false; }
 
+	virtual void SetIsSimulateInEditorViewport(bool bInIsSimulateInEditorViewport) { };
 
 	virtual bool WantsPollingMouseMovement(void) const { return true; }
 
