Merging UE4-Pretest @ 2042161 to UE4

diff --git a/Engine/Source/Runtime/Engine/Classes/EdGraph/EdGraphNode.h b/Engine/Source/Runtime/Engine/Classes/EdGraph/EdGraphNode.h
--- a/Engine/Source/Runtime/Engine/Classes/EdGraph/EdGraphNode.h
+++ b/Engine/Source/Runtime/Engine/Classes/EdGraph/EdGraphNode.h
@@ -135,7 +135,7 @@ public:
 	TWeakPtr<SGraphNode> NodeWidget;
 
 	/** Create a new pin on this node using the supplied info, and return the new pin */
-	UEdGraphPin* CreatePin(EEdGraphPinDirection Dir, const FString& PinCategory, const FString& PinSubCategory, UObject* PinSubCategoryObject, bool bIsArray, bool bIsReference, const FString& PinName, bool bIsConst = false);
+	UEdGraphPin* CreatePin(EEdGraphPinDirection Dir, const FString& PinCategory, const FString& PinSubCategory, UObject* PinSubCategoryObject, bool bIsArray, bool bIsReference, const FString& PinName, bool bIsConst = false, int32 Index = INDEX_NONE);
 
 	// Allocates a pin from the pool
 	static UEdGraphPin* AllocatePinFromPool(UEdGraphNode* OuterNode);
@@ -149,6 +149,9 @@ public:
 	/** Find a pin on this node with the supplied name and assert if it is not present */
 	UEdGraphPin* FindPinChecked(const FString& PinName) const;
 	
+	/** Find a pin on this node with the supplied name and remove it */
+	void DiscardPin(UEdGraphPin* Pin);
+
 	/** Whether or not this node should be given the chance to override pin names.  If this returns true, then GetPinNameOverride() will be called for each pin, each frame */
 	virtual bool ShouldOverridePinNames() const { return false; }
 
@@ -166,6 +169,9 @@ public:
 	 */
 	virtual void GetPinHoverText(const UEdGraphPin& Pin, FString& HoverTextOut) const;
 
+	/** Gets the index for a pin */
+	int32 GetPinIndex(UEdGraphPin* Pin) const;
+
 	/** Break all links on this node */
 	void BreakAllNodeLinks();
 
diff --git a/Engine/Source/Runtime/Engine/Classes/EdGraph/EdGraphPin.h b/Engine/Source/Runtime/Engine/Classes/EdGraph/EdGraphPin.h
--- a/Engine/Source/Runtime/Engine/Classes/EdGraph/EdGraphPin.h
+++ b/Engine/Source/Runtime/Engine/Classes/EdGraph/EdGraphPin.h
@@ -13,6 +13,45 @@ enum EEdGraphPinDirection
 	EGPD_MAX,
 };
 
+USTRUCT()
+struct FSimpleMemberReference
+{
+	GENERATED_USTRUCT_BODY()
+
+	/** Class that this member is defined in. */
+	UPROPERTY()
+	TSubclassOf<class UObject> MemberParentClass;
+
+	/** Name of the member */
+	UPROPERTY()
+	FName MemberName;
+
+	/** The Guid of the member */
+	UPROPERTY()
+	FGuid MemberGuid;
+
+	void Reset()
+	{
+		operator=(FSimpleMemberReference());
+	}
+
+	bool operator==(const FSimpleMemberReference& Other) const
+	{
+		return (MemberParentClass == Other.MemberParentClass)
+			&& (MemberName == Other.MemberName)
+			&& (MemberGuid == Other.MemberGuid);
+	}
+};
+
+FORCEINLINE FArchive& operator<<(FArchive& Ar, FSimpleMemberReference& Data)
+{
+	Ar << Data.MemberParentClass;
+	Ar << Data.MemberName;
+	Ar << Data.MemberGuid;
+	return Ar;
+}
+
+
 /** Struct used to define the type of information carried on this pin */
 USTRUCT()
 struct FEdGraphPinType
@@ -31,6 +70,10 @@ struct FEdGraphPinType
 	UPROPERTY()
 	TWeakObjectPtr<UObject> PinSubCategoryObject;
 
+	/** Sub-category member reference */
+	UPROPERTY()
+	FSimpleMemberReference PinSubCategoryMemberReference;
+
 	/** Whether or not this pin represents an array of values */
 	UPROPERTY()
 	bool bIsArray;
@@ -71,7 +114,8 @@ public:
 			&& (PinSubCategoryObject == Other.PinSubCategoryObject) 
 			&& (bIsArray == Other.bIsArray) 
 			&& (bIsReference == Other.bIsReference)
-			&& (bIsWeakPointer == Other.bIsWeakPointer);
+			&& (bIsWeakPointer == Other.bIsWeakPointer)
+			&& (PinSubCategoryMemberReference == Other.PinSubCategoryMemberReference);
 	}
 	bool operator != ( const FEdGraphPinType& Other ) const
 	{
@@ -88,6 +132,7 @@ public:
 		PinCategory.Empty();
 		PinSubCategory.Empty();
 		PinSubCategoryObject = NULL;
+		PinSubCategoryMemberReference.Reset();
 		bIsArray = false;
 		bIsReference = false;
 		bIsWeakPointer = false;
@@ -181,6 +226,10 @@ class UEdGraphPin : public UObject
 	/** If true, the pin may be hidden by user */
 	UPROPERTY()
 	uint32 bAdvancedView:1;
+
+	/** Pin name could be changed, so whenever possible it's good to have a persistent GUID identifying Pin to reconstruct Node seamlessly */
+	UPROPERTY()
+	FGuid PersistentGuid;
 #endif // WITH_EDITORONLY_DATA
 
 public:
diff --git a/Engine/Source/Runtime/Engine/Classes/Engine/Blueprint.h b/Engine/Source/Runtime/Engine/Classes/Engine/Blueprint.h
--- a/Engine/Source/Runtime/Engine/Classes/Engine/Blueprint.h
+++ b/Engine/Source/Runtime/Engine/Classes/Engine/Blueprint.h
@@ -4,7 +4,6 @@
 
 #include "EdGraph/EdGraphPin.h"
 #include "BlueprintCore.h"
-#include "BlueprintGeneratedStruct.h"
 #include "Blueprint.generated.h"
 
 /** States a blueprint can be in */
@@ -56,7 +55,6 @@ namespace EKismetCompileType
 		Full,
 		StubAfterFailure, 
 		BytecodeOnly,
-		StructuresOnly,
 	};
 };
 
@@ -80,7 +78,7 @@ public:
 	/** Whether or not this compile type should operate on the generated class of the blueprint, as opposed to just the skeleton */
 	bool IsGeneratedClassCompileType() const
 	{
-		return (CompileType != EKismetCompileType::SkeletonOnly) && (CompileType != EKismetCompileType::StructuresOnly);
+		return (CompileType != EKismetCompileType::SkeletonOnly);
 	}
 
 	FKismetCompilerOptions()
@@ -157,10 +155,6 @@ struct FBPVariableDescription
 	UPROPERTY(EditAnywhere, Category=BPVariableDescription)
 	FString DefaultValue;
 
-	/** Used by blueprint generated struct to mark invalid fields */
-	UPROPERTY(Transient)
-	bool bInvalidMember;
-
 	FBPVariableDescription()
 		: PropertyFlags(CPF_Edit)
 	{
@@ -245,27 +239,6 @@ struct FEditedDocumentInfo
 	}
 };
 
-USTRUCT()
-struct FBPStructureDescription
-{
-	GENERATED_USTRUCT_BODY()
-
-	UPROPERTY(EditAnywhere, Category=BPStructureDescription)
-	FName Name;
-
-	UPROPERTY(EditAnywhere, Category=BPStructureDescription)
-	TArray<FBPVariableDescription> Fields;
-
-	UPROPERTY()
-	UBlueprintGeneratedStruct* CompiledStruct;
-
-	FBPStructureDescription() 
-		: CompiledStruct(NULL)
-	{
-
-	}
-};
-
 ///////////////
 
 /**
@@ -377,9 +350,6 @@ class ENGINE_API UBlueprint : public UBlueprintCore
 	/** Array of info about the interfaces we implement in this blueprint */
 	UPROPERTY(AssetRegistrySearchable)
 	TArray<struct FBPInterfaceDescription> ImplementedInterfaces;
-
-	UPROPERTY(EditAnywhere, Category=Blueprint)
-	TArray<struct FBPStructureDescription> UserDefinedStructures;
 	
 #endif // WITH_EDITORONLY_DATA
 
@@ -585,6 +555,9 @@ public:
 
 		return false;
 	}
+
+	static FName GetFunctionNameFromClassByGuid(const UClass* InClass, const FGuid FunctionGuid);
+	static bool GetFunctionGuidFromClassByFieldName(const UClass* InClass, const FName FunctionName, FGuid& FunctionGuid);
 #endif
 
 	/** Find a function given its name and optionally an object property name within this Blueprint */
@@ -624,45 +597,12 @@ public:
 template<>
 inline FName UBlueprint::GetFieldNameFromClassByGuid<UFunction>(const UClass* InClass, const FGuid FunctionGuid)
 {
-	TArray<UBlueprint*> Blueprints;
-	UBlueprint::GetBlueprintHierarchyFromClass(InClass, Blueprints);
-
-	for (int32 BPIndex = 0; BPIndex < Blueprints.Num(); ++BPIndex)
-	{
-		UBlueprint* Blueprint = Blueprints[BPIndex];
-		for (int32 FunctionIndex = 0; FunctionIndex < Blueprint->FunctionGraphs.Num(); ++FunctionIndex)
-		{
-			UEdGraph* FunctionGraph = Blueprint->FunctionGraphs[FunctionIndex];
-			if (FunctionGraph && FunctionGraph->GraphGuid == FunctionGuid)
-			{	
-				return FunctionGraph->GetFName();
-			}
-		}
-	}
-
-	return NAME_None;
+	return GetFunctionNameFromClassByGuid(InClass, FunctionGuid);
 }
 
 template<>
 inline bool UBlueprint::GetGuidFromClassByFieldName<UFunction>(const UClass* InClass, const FName FunctionName, FGuid& FunctionGuid)
 {
-	TArray<UBlueprint*> Blueprints;
-	UBlueprint::GetBlueprintHierarchyFromClass(InClass, Blueprints);
-
-	for (int32 BPIndex = 0; BPIndex < Blueprints.Num(); ++BPIndex)
-	{
-		UBlueprint* Blueprint = Blueprints[BPIndex];
-		for (int32 FunctionIndex = 0; FunctionIndex < Blueprint->FunctionGraphs.Num(); ++FunctionIndex)
-		{
-			UEdGraph* FunctionGraph = Blueprint->FunctionGraphs[FunctionIndex];
-			if (FunctionGraph && FunctionGraph->GetFName() == FunctionName)
-			{	
-				FunctionGuid = FunctionGraph->GraphGuid;
-				return true;
-			}
-		}
-	}
-
-	return false;
+	return GetFunctionGuidFromClassByFieldName(InClass, FunctionName, FunctionGuid);
 }
 #endif // #if WITH_EDITORONLY_DATA
diff --git a/Engine/Source/Runtime/Engine/Classes/Engine/EngineTypes.h b/Engine/Source/Runtime/Engine/Classes/Engine/EngineTypes.h
--- a/Engine/Source/Runtime/Engine/Classes/Engine/EngineTypes.h
+++ b/Engine/Source/Runtime/Engine/Classes/Engine/EngineTypes.h
@@ -1258,7 +1258,7 @@ struct FPrimitiveMaterialRef
 };
 
 /** Structure containing information about one hit of the trace */
-USTRUCT(BlueprintType)
+USTRUCT(BlueprintType, meta=(HasNativeMakeBreak="true"))
 struct ENGINE_API FHitResult
 {
 	GENERATED_USTRUCT_BODY()
diff --git a/Engine/Source/Runtime/Engine/Classes/Engine/UserDefinedStruct.h b/Engine/Source/Runtime/Engine/Classes/Engine/UserDefinedStruct.h
new file mode 100644
--- /dev/null
+++ b/Engine/Source/Runtime/Engine/Classes/Engine/UserDefinedStruct.h
@@ -0,0 +1,53 @@
+// Copyright 1998-2014 Epic Games, Inc. All Rights Reserved.
+
+#pragma once
+
+#include "UserDefinedStruct.generated.h"
+
+UENUM()
+enum EUserDefinedStructureStatus
+{
+	// Struct is in an unknown state
+	UDSS_UpToDate,
+	// Struct has been modified but not recompiled
+	UDSS_Dirty,
+	// Struct tried but failed to be compiled
+	UDSS_Error,
+	// Struct is a duplicate, the original one was changed
+	UDSS_Duplicate,
+
+	UDSS_MAX,
+};
+
+UCLASS()
+class ENGINE_API UUserDefinedStruct : public UScriptStruct
+{
+	GENERATED_UCLASS_BODY()
+
+public:
+#if WITH_EDITORONLY_DATA
+	UPROPERTY()
+	TEnumAsByte<enum EUserDefinedStructureStatus> Status;
+
+	// the original struct, when current struct isn't a temporary duplicate, the field should be null
+	UPROPERTY(Transient)
+	TWeakObjectPtr<UUserDefinedStruct> PrimaryStruct;
+
+	UPROPERTY()
+	UObject* EditorData;
+#endif // WITH_EDITORONLY_DATA
+
+#if WITH_EDITOR
+	// UObject interface.
+	virtual void Serialize(FArchive& Ar) OVERRIDE;
+	virtual void PostDuplicate(bool bDuplicateForPIE) OVERRIDE;
+	// End of UObject interface.
+
+#endif	// WITH_EDITOR
+
+	virtual void SerializeTaggedProperties( FArchive& Ar, uint8* Data, UStruct* DefaultsStruct, uint8* Defaults ) const OVERRIDE;
+
+	// UScriptStruct interface.
+	virtual void RecursivelyPreload() OVERRIDE;
+	// End of  UScriptStruct interface.
+};
diff --git a/Engine/Source/Runtime/Engine/Classes/Kismet/GameplayStatics.h b/Engine/Source/Runtime/Engine/Classes/Kismet/GameplayStatics.h
--- a/Engine/Source/Runtime/Engine/Classes/Kismet/GameplayStatics.h
+++ b/Engine/Source/Runtime/Engine/Classes/Kismet/GameplayStatics.h
@@ -371,7 +371,7 @@ class ENGINE_API UGameplayStatics : public UBlueprintFunctionLibrary
 	 * @param HitComponent
 	 * @param HitBoneName
 	 */
-	UFUNCTION(BlueprintPure, Category="Collision")
+	UFUNCTION(BlueprintPure, Category = "Collision", meta=(NativeBreakFunc))
 	static void BreakHitResult(const struct FHitResult& Hit, FVector& Location, FVector& Normal, FVector& ImpactPoint, FVector& ImpactNormal, class UPhysicalMaterial*& PhysMat, class AActor*& HitActor, class UPrimitiveComponent*& HitComponent, FName& HitBoneName);
 
 	/** Returns the EPhysicalSurface type of the given Hit. 
diff --git a/Engine/Source/Runtime/Engine/Classes/Kismet/KismetSystemLibrary.h b/Engine/Source/Runtime/Engine/Classes/Kismet/KismetSystemLibrary.h
--- a/Engine/Source/Runtime/Engine/Classes/Kismet/KismetSystemLibrary.h
+++ b/Engine/Source/Runtime/Engine/Classes/Kismet/KismetSystemLibrary.h
@@ -31,6 +31,15 @@ namespace EMoveComponentAction
 	};
 }
 
+USTRUCT()
+struct FGenericStruct
+{
+	GENERATED_USTRUCT_BODY()
+
+	UPROPERTY()
+	int32 Data;
+};
+
 UCLASS(HeaderGroup=KismetLibrary,MinimalAPI)
 class UKismetSystemLibrary : public UBlueprintFunctionLibrary
 {
@@ -336,6 +345,32 @@ class UKismetSystemLibrary : public UBlueprintFunctionLibrary
 	UFUNCTION(BlueprintCallable, meta=(BlueprintInternalUseOnly = "true"))
 	static void SetTransformPropertyByName(UObject* Object, FName PropertyName, const FTransform& Value);
 
+	/** Set a custom structure property by name */
+	UFUNCTION(BlueprintCallable, CustomThunk, meta = (BlueprintInternalUseOnly = "true", CustomStructureParam = "Value"))
+	static void SetStructurePropertyByName(UObject* Object, FName PropertyName, const FGenericStruct& Value);
+
+	/** Based on UKismetArrayLibrary::execSetArrayPropertyByName */
+	DECLARE_FUNCTION(execSetStructurePropertyByName)
+	{
+		P_GET_OBJECT(UObject, OwnerObject);
+		P_GET_PROPERTY(UNameProperty, StructPropertyName);
+
+		Stack.StepCompiledIn<UStructProperty>(NULL);
+		void* SrcStructAddr = Stack.MostRecentPropertyAddress;
+
+		P_FINISH;
+
+		if (OwnerObject != NULL)
+		{
+			UStructProperty* StructProp = FindField<UStructProperty>(OwnerObject->GetClass(), StructPropertyName);
+			if (StructProp != NULL)
+			{
+				void* Dest = StructProp->ContainerPtrToValuePtr<void>(OwnerObject);
+				StructProp->CopyValuesInternal(Dest, SrcStructAddr, 1);
+			}
+		}
+	}
+
 	// --- Collision functions ------------------------------
 
 	// @NOTE!!!!! please remove _NEW when we can remove all _DEPRECATED functions and create redirect in BaseEngine.INI
diff --git a/Engine/Source/Runtime/Engine/Classes/Matinee/InterpData.h b/Engine/Source/Runtime/Engine/Classes/Matinee/InterpData.h
--- a/Engine/Source/Runtime/Engine/Classes/Matinee/InterpData.h
+++ b/Engine/Source/Runtime/Engine/Classes/Matinee/InterpData.h
@@ -112,8 +112,11 @@ class UInterpData : public UObject
 	/** Find a DirectorGroup in the data. Should only ever be 0 or 1 of these! */
 	ENGINE_API class UInterpGroupDirector* FindDirectorGroup();
 
+	/** Checks to see if the event name exists */
+	ENGINE_API bool IsEventName(const FName& InEventName) const;
+
 	/** Get the list of all unique event names */
-	ENGINE_API void GetAllEventNames(TArray<FName>& OutEventNames);
+	ENGINE_API void GetAllEventNames(TArray<FName>& OutEventNames) const;
 
 	/** Update the AllEventNames array */
 	ENGINE_API void UpdateEventNames();
diff --git a/Engine/Source/Runtime/Engine/Classes/Particles/ParticleEmitter.h b/Engine/Source/Runtime/Engine/Classes/Particles/ParticleEmitter.h
--- a/Engine/Source/Runtime/Engine/Classes/Particles/ParticleEmitter.h
+++ b/Engine/Source/Runtime/Engine/Classes/Particles/ParticleEmitter.h
@@ -6,6 +6,8 @@
 //=============================================================================
 
 #pragma once
+#include "ParticleHelper.h"
+#include "ParticleEmitterInstances.h"
 #include "ParticleEmitter.generated.h"
 
 //=============================================================================
diff --git a/Engine/Source/Runtime/Engine/Classes/Particles/ParticleSystemComponent.h b/Engine/Source/Runtime/Engine/Classes/Particles/ParticleSystemComponent.h
--- a/Engine/Source/Runtime/Engine/Classes/Particles/ParticleSystemComponent.h
+++ b/Engine/Source/Runtime/Engine/Classes/Particles/ParticleSystemComponent.h
@@ -2,6 +2,7 @@
 
 
 #pragma once
+#include "ParticleHelper.h"
 #include "ParticleSystemComponent.generated.h"
 
 /** Enum for specifying type of a name instance parameter. */
@@ -615,7 +616,7 @@ public:
 	* @param OutTrailEmitters	The array to fill with pointers to
 	* @param bIncludeUnassociated If true, trails that have not yet been associated with an AnimNotifyState are included in OutTrailEmitters.
 	*/
-	virtual void GetTrailEmitters(UAnimNotifyState* InAnimNotifyState, TArray< FParticleAnimTrailEmitterInstance* >& OutTrailEmitters, bool bIncludeUnassociated = false);
+	virtual void GetTrailEmitters(UAnimNotifyState* InAnimNotifyState, TArray< struct FParticleAnimTrailEmitterInstance* >& OutTrailEmitters, bool bIncludeUnassociated = false);
 
 public:
 	TArray<struct FParticleEmitterInstance*> EmitterInstances;
diff --git a/Engine/Source/Runtime/Engine/Private/Blueprint.cpp b/Engine/Source/Runtime/Engine/Private/Blueprint.cpp
--- a/Engine/Source/Runtime/Engine/Private/Blueprint.cpp
+++ b/Engine/Source/Runtime/Engine/Private/Blueprint.cpp
@@ -192,20 +192,6 @@ void UBlueprint::Serialize(FArchive& Ar)
 		}
 	}
 
-	// If we don't have a skeleton class via compile-on-load, generate one now
-	if( Ar.IsLoading() && (SkeletonGeneratedClass == NULL) )
-	{
-		bool bWasRegen = bIsRegeneratingOnLoad;
-		bIsRegeneratingOnLoad = true;
-
-		FBlueprintEditorUtils::PreloadMembers(this);
-		FBlueprintEditorUtils::PreloadConstructionScript(this);
-		FBlueprintEditorUtils::RefreshInputDelegatePins(this);
-		FKismetEditorUtilities::GenerateBlueprintSkeleton(this);
-
-		bIsRegeneratingOnLoad = bWasRegen;
-	}
-
 	if (Ar.UE4Ver() < VER_UE4_BP_ACTOR_VARIABLE_DEFAULT_PREVENTING)
 	{
 		// Actor variables can't have default values (because Blueprint templates are library elements that can 
@@ -259,18 +245,6 @@ bool UBlueprint::Rename( const TCHAR* InName, UObject* NewOuter, ERenameFlags Fl
 		}
 	}
 
-	for (auto StructIter = UserDefinedStructures.CreateIterator(); StructIter; ++StructIter)
-	{
-		if (UBlueprintGeneratedStruct* CompiledStruct = (*StructIter).CompiledStruct)
-		{
-			bool bMovedOK = CompiledStruct->Rename(*CompiledStruct->GetName(), NewOuter, Flags);
-			if(!bMovedOK)
-			{
-				return false;
-			}
-		}
-	}
-
 	bool bSuccess = Super::Rename( InName, NewOuter, Flags );
 
 	// Finally, do a compile 
@@ -812,17 +786,6 @@ void UBlueprint::TagSubobjects(EObjectFlags NewFlags)
 		SkeletonGeneratedClass->SetFlags(NewFlags);
 		SkeletonGeneratedClass->TagSubobjects(NewFlags);
 	}
-#if WITH_EDITORONLY_DATA
-	for (auto StructDescIter = UserDefinedStructures.CreateIterator(); StructDescIter; ++StructDescIter)
-	{
-		UBlueprintGeneratedStruct* Struct = (*StructDescIter).CompiledStruct;
-		if (Struct && !Struct->HasAnyFlags(GARBAGE_COLLECTION_KEEPFLAGS | RF_RootSet))
-		{
-			Struct->SetFlags(NewFlags);
-			Struct->TagSubobjects(NewFlags);
-		}
-	}
-#endif // WITH_EDITORONLY_DATA
 }
 
 void UBlueprint::GetAllGraphs(TArray<UEdGraph*>& Graphs) const
@@ -970,3 +933,17 @@ void UBlueprint::PostLoadSubobjects(FObjectInstancingGraph* OuterInstanceGraph)
 }
 
 #endif
+
+#if WITH_EDITORONLY_DATA
+
+FName UBlueprint::GetFunctionNameFromClassByGuid(const UClass* InClass, const FGuid FunctionGuid)
+{
+	return FBlueprintEditorUtils::GetFunctionNameFromClassByGuid(InClass, FunctionGuid);
+}
+
+bool UBlueprint::GetFunctionGuidFromClassByFieldName(const UClass* InClass, const FName FunctionName, FGuid& FunctionGuid)
+{
+	return FBlueprintEditorUtils::GetFunctionGuidFromClassByFieldName(InClass, FunctionName, FunctionGuid);
+}
+
+#endif //WITH_EDITORONLY_DATA
\ No newline at end of file
diff --git a/Engine/Source/Runtime/Engine/Private/EdGraph/EdGraph.cpp b/Engine/Source/Runtime/Engine/Private/EdGraph/EdGraph.cpp
--- a/Engine/Source/Runtime/Engine/Private/EdGraph/EdGraph.cpp
+++ b/Engine/Source/Runtime/Engine/Private/EdGraph/EdGraph.cpp
@@ -401,7 +401,7 @@ UEdGraphNode::UEdGraphNode(const class FPostConstructInitializeProperties& PCIP)
 }
 
 #if WITH_EDITOR
-UEdGraphPin* UEdGraphNode::CreatePin(EEdGraphPinDirection Dir, const FString& PinCategory, const FString& PinSubCategory, UObject* PinSubCategoryObject, bool bIsArray, bool bIsReference, const FString& PinName, bool bIsConst /*= false*/)
+UEdGraphPin* UEdGraphNode::CreatePin(EEdGraphPinDirection Dir, const FString& PinCategory, const FString& PinSubCategory, UObject* PinSubCategoryObject, bool bIsArray, bool bIsReference, const FString& PinName, bool bIsConst /*= false*/, int32 Index /*= INDEX_NONE*/)
 {
 #if 0
 	UEdGraphPin* NewPin = AllocatePinFromPool(this);
@@ -418,7 +418,15 @@ UEdGraphPin* UEdGraphNode::CreatePin(EEdGraphPinDirection Dir, const FString& Pi
 	NewPin->PinType.bIsConst = bIsConst;
 	NewPin->SetFlags(RF_Transactional);
 
-	Pins.Add(NewPin);
+	Modify(false);
+	if ( Pins.IsValidIndex( Index ) )
+	{
+		Pins.Insert(NewPin, Index);
+	}
+	else
+	{
+		Pins.Add(NewPin);
+	}
 	return NewPin;
 }
 
@@ -442,6 +450,14 @@ UEdGraphPin* UEdGraphNode::FindPinChecked(const FString& PinName) const
 	return Result;
 }
 
+void UEdGraphNode::DiscardPin(UEdGraphPin* Pin)
+{
+	check( Pin );
+
+	Modify();
+	Pins.Remove( Pin );
+}
+
 void UEdGraphNode::BreakAllNodeLinks()
 {
 	TSet<UEdGraphNode*> NodeList;
@@ -656,6 +672,11 @@ FString UEdGraphNode::GetPinDisplayName(const UEdGraphPin* Pin) const
 	return GetSchema()->GetPinDisplayName(Pin);
 }
 
+int32 UEdGraphNode::GetPinIndex(UEdGraphPin* Pin) const
+{
+	return Pins.Find(Pin);
+}
+
 #endif	//#if WITH_EDITOR
 
 
@@ -801,8 +822,12 @@ UEdGraphPin::UEdGraphPin(const class FPostConstructInitializeProperties& PCIP)
 
 void UEdGraphPin::MakeLinkTo(UEdGraphPin* ToPin)
 {
+	Modify();
+
 	if (ToPin != NULL)
 	{
+		ToPin->Modify();
+
 		// Make sure we don't already link to it
 		if (!LinkedTo.Contains(ToPin))
 		{
@@ -861,10 +886,15 @@ void UEdGraphPin::CopyPersistentDataFromOldPin(const UEdGraphPin& SourcePin)
 	// The PinType, Direction, and bNotConnectable are properties generated from the schema
 
 	// Only move the default value if it was modified; inherit the new default value otherwise
-	if (SourcePin.DefaultValue != SourcePin.AutogeneratedDefaultValue || SourcePin.DefaultObject != NULL || !SourcePin.DefaultTextValue.IsEmpty())
+	if (SourcePin.DefaultValue != SourcePin.AutogeneratedDefaultValue || SourcePin.DefaultObject != NULL)
 	{
 		DefaultObject = SourcePin.DefaultObject;
 		DefaultValue = SourcePin.DefaultValue;
+	}
+
+	const bool bTextValueWasModified = (SourcePin.DefaultTextValue.ToString() != SourcePin.AutogeneratedDefaultValue);
+	if (bTextValueWasModified)
+	{
 		DefaultTextValue = SourcePin.DefaultTextValue;
 	}
 
@@ -1478,6 +1508,27 @@ bool FEdGraphPinType::Serialize(FArchive& Ar)
 	Ar << bIsReference;
 	Ar << bIsWeakPointer;
 	
+	if (Ar.UE4Ver() >= VER_UE4_MEMBERREFERENCE_IN_PINTYPE)
+	{
+		Ar << PinSubCategoryMemberReference;
+	}
+	else if (Ar.IsLoading() && Ar.IsPersistent())
+	{
+		if ((PinCategory == TEXT("delegate")) || (PinCategory == TEXT("mcdelegate")))
+		{
+			if (const UFunction* Signature = Cast<const UFunction>(PinSubCategoryObject.Get()))
+			{
+				PinSubCategoryMemberReference.MemberName = Signature->GetFName();
+				PinSubCategoryMemberReference.MemberParentClass = Signature->GetOwnerClass();
+				PinSubCategoryObject = NULL;
+			}
+			else
+			{
+				ensure(true);
+			}
+		}
+	}
+
 	return true;
 }
 
diff --git a/Engine/Source/Runtime/Engine/Private/Engine.GeneratedWrapper.cpp b/Engine/Source/Runtime/Engine/Private/Engine.GeneratedWrapper.cpp
--- a/Engine/Source/Runtime/Engine/Private/Engine.GeneratedWrapper.cpp
+++ b/Engine/Source/Runtime/Engine/Private/Engine.GeneratedWrapper.cpp
@@ -4,4 +4,4 @@
 #include "AllEngineHeaders.h"
 
 
-#include "Engine.generated.inl"
+
diff --git a/Engine/Source/Runtime/Engine/Private/Interpolation.cpp b/Engine/Source/Runtime/Engine/Private/Interpolation.cpp
--- a/Engine/Source/Runtime/Engine/Private/Interpolation.cpp
+++ b/Engine/Source/Runtime/Engine/Private/Interpolation.cpp
@@ -1851,7 +1851,12 @@ UInterpGroupDirector* UInterpData::FindDirectorGroup()
 	return DirGroup;
 }
 
-void UInterpData::GetAllEventNames(TArray<FName>& OutEventNames)
+bool UInterpData::IsEventName(const FName& InEventName) const
+{
+	return AllEventNames.Contains( InEventName );
+}
+
+void UInterpData::GetAllEventNames(TArray<FName>& OutEventNames) const
 {
 	OutEventNames = AllEventNames;
 }
@@ -6298,6 +6303,7 @@ void UInterpTrackEvent::RemoveKeyframe(int32 KeyIndex)
 	if ( Group )
 	{
 		UInterpData* IData = CastChecked<UInterpData>( Group->GetOuter() );
+		IData->Modify();
 		IData->UpdateEventNames();
 	}
 }
diff --git a/Engine/Source/Runtime/Engine/Private/KismetSystemLibrary.cpp b/Engine/Source/Runtime/Engine/Private/KismetSystemLibrary.cpp
--- a/Engine/Source/Runtime/Engine/Private/KismetSystemLibrary.cpp
+++ b/Engine/Source/Runtime/Engine/Private/KismetSystemLibrary.cpp
@@ -2303,3 +2303,8 @@ void UKismetSystemLibrary::ShowPlatformSpecificAchievementsScreen()
 	AndroidThunkCpp_ShowAchievements();
 #endif
 }
+void UKismetSystemLibrary::SetStructurePropertyByName(UObject* Object, FName PropertyName, const FGenericStruct& Value)
+{
+	// We should never hit these!  They're stubs to avoid NoExport on the class.
+	check(0);
+}
\ No newline at end of file
diff --git a/Engine/Source/Runtime/Engine/Private/MatineeDelegates.cpp b/Engine/Source/Runtime/Engine/Private/MatineeDelegates.cpp
new file mode 100644
--- /dev/null
+++ b/Engine/Source/Runtime/Engine/Private/MatineeDelegates.cpp
@@ -0,0 +1,11 @@
+// Copyright 1998-2013 Epic Games, Inc. All Rights Reserved.
+
+#include "EnginePrivate.h"
+#include "MatineeDelegates.h"
+
+FMatineeDelegates& FMatineeDelegates::Get()
+{
+	// return the singleton object
+	static FMatineeDelegates Singleton;
+	return Singleton;
+}
diff --git a/Engine/Source/Runtime/Engine/Private/UserDefinedStruct.cpp b/Engine/Source/Runtime/Engine/Private/UserDefinedStruct.cpp
new file mode 100644
--- /dev/null
+++ b/Engine/Source/Runtime/Engine/Private/UserDefinedStruct.cpp
@@ -0,0 +1,104 @@
+// Copyright 1998-2014 Epic Games, Inc. All Rights Reserved.
+
+#include "EnginePrivate.h"
+
+#if WITH_EDITOR
+#include "StructureEditorUtils.h"
+#endif //WITH_EDITOR
+
+UUserDefinedStruct::UUserDefinedStruct(const class FPostConstructInitializeProperties& PCIP)
+	: Super(PCIP)
+{
+}
+
+#if WITH_EDITOR
+
+void UUserDefinedStruct::Serialize(FArchive& Ar)
+{
+	Super::Serialize( Ar );
+
+	if (Ar.IsLoading() && (EUserDefinedStructureStatus::UDSS_UpToDate == Status))
+	{
+		FString ErrorMsg;
+		const FStructureEditorUtils::EStructureError Result = FStructureEditorUtils::IsStructureValid(this, NULL, &ErrorMsg);
+		if (FStructureEditorUtils::EStructureError::Ok != Result)
+		{
+			Status = EUserDefinedStructureStatus::UDSS_Error;
+			UE_LOG(LogClass, Log, TEXT("UUserDefinedStruct.Serialize '%s' validation: %s"), *GetName(), *ErrorMsg);
+		}
+	}
+}
+
+void UUserDefinedStruct::PostDuplicate(bool bDuplicateForPIE)
+{
+	Super::PostDuplicate(bDuplicateForPIE);
+	if (!bDuplicateForPIE && (GetOuter() != GetTransientPackage()))
+	{
+		SetMetaData(TEXT("BlueprintType"), TEXT("true"));
+		FStructureEditorUtils::OnStructureChanged(this);
+	}
+}
+
+#endif	// WITH_EDITOR
+
+void UUserDefinedStruct::SerializeTaggedProperties(FArchive& Ar, uint8* Data, UStruct* DefaultsStruct, uint8* Defaults) const
+{
+#if WITH_EDITOR
+	/*	The following code is responsible for UUserDefinedStruct's default values serialization.	*/
+
+	auto UDDefaultsStruct = Cast<UUserDefinedStruct>(DefaultsStruct);
+
+	/*	When saving delta, we want the difference between current data and true structure's default values. 
+		When Defaults is NULL then zeroed data will be used for comparison.*/
+	const bool bUseNewDefaults = !Defaults
+		&& UDDefaultsStruct
+		&& Ar.DoDelta()
+		&& Ar.IsSaving()
+		&& (0 == (Ar.GetPortFlags() & PPF_Duplicate))
+		&& !Ar.IsCooking();
+
+	/*	Object serialized from delta will have missing properties filled with zeroed data, 
+		we want structure's default data instead */
+	const bool bLoadDefaultFirst = UDDefaultsStruct 
+		&& Ar.IsLoading() 
+		&& Ar.IsPersistent();
+
+	const bool bPrepareDefaultStruct = bUseNewDefaults || bLoadDefaultFirst;
+	FStructureEditorUtils::FStructOnScope StructDefaultMem(bPrepareDefaultStruct ? UDDefaultsStruct : NULL);
+	if (bPrepareDefaultStruct)
+	{
+		FStructureEditorUtils::Fill_MakeStructureDefaultValue(UDDefaultsStruct, StructDefaultMem.GetStructMemory());
+	}
+
+	if (bUseNewDefaults)
+	{
+		Defaults = StructDefaultMem.GetStructMemory();
+	}
+	if (bLoadDefaultFirst)
+	{
+		UDDefaultsStruct->CopyScriptStruct(Data, StructDefaultMem.GetStructMemory());
+	}
+#endif // WITH_EDITOR
+	Super::SerializeTaggedProperties(Ar, Data, DefaultsStruct, Defaults);
+}
+
+void UUserDefinedStruct::RecursivelyPreload()
+{
+	ULinkerLoad* Linker = GetLinker();
+	if( Linker && (NULL == PropertyLink) )
+	{
+		TArray<UObject*> AllChildMembers;
+		GetObjectsWithOuter(this, AllChildMembers);
+		for (int32 Index = 0; Index < AllChildMembers.Num(); ++Index)
+		{
+			UObject* Member = AllChildMembers[Index];
+			Linker->Preload(Member);
+		}
+
+		Linker->Preload(this);
+		if (NULL == PropertyLink)
+		{
+			StaticLink(true);
+		}
+	}
+}
\ No newline at end of file
diff --git a/Engine/Source/Runtime/Engine/Public/MatineeDelegates.h b/Engine/Source/Runtime/Engine/Public/MatineeDelegates.h
new file mode 100644
--- /dev/null
+++ b/Engine/Source/Runtime/Engine/Public/MatineeDelegates.h
@@ -0,0 +1,48 @@
+// Copyright 1998-2013 Epic Games, Inc. All Rights Reserved.
+
+#ifndef __MatineeDelegates_h__
+#define __MatineeDelegates_h__
+
+#pragma once
+
+/** 
+ * FMatineeDelegates
+ * Delegates used by matinee.
+ **/
+class ENGINE_API FMatineeDelegates
+{
+public:
+	/** Return a single FMatineeDelegates singleton */
+	static FMatineeDelegates& Get();
+
+	/**
+	 * Fired whenever a track event keyframe is added
+	 *
+	 * @param	AMatineeActor	The actor the keyframe belongs to
+	 * @param	FName			The name of the keyframe that has been added
+	 * @param	int32			The index of the keyframe
+	 */
+	DECLARE_MULTICAST_DELEGATE_ThreeParams(FOnEventKeyframeAdded, const class AMatineeActor*, const FName&, int32);
+	FOnEventKeyframeAdded OnEventKeyframeAdded;
+
+	/**
+	 * Fired whenever a track event keyframe is renamed
+	 *
+	 * @param	AMatineeActor	The actor the keyframe belongs to
+	 * @param	FName			The old name of the keyframe
+	 * @param	FName			The new name of the keyframe
+	 */
+	DECLARE_MULTICAST_DELEGATE_ThreeParams(FOnEventKeyframeRenamed, const class AMatineeActor*, const FName&, const FName&);
+	FOnEventKeyframeRenamed OnEventKeyframeRenamed;
+
+	/**
+	 * Fired whenever a track event keyframe is removed
+	 *
+	 * @param	AMatineeActor	The actor the keyframe belongs to
+	 * @param	FName			The names of the keyframes that have been removed
+	 */
+	DECLARE_MULTICAST_DELEGATE_TwoParams(FOnEventKeyframeRemoved, const class AMatineeActor*, const TArray<FName>&);
+	FOnEventKeyframeRemoved OnEventKeyframeRemoved;
+};
+
+#endif// __MatineeDelegates_h__
diff --git a/Engine/Source/Runtime/EngineMessages/Private/EngineMessagesModule.cpp b/Engine/Source/Runtime/EngineMessages/Private/EngineMessagesModule.cpp
--- a/Engine/Source/Runtime/EngineMessages/Private/EngineMessagesModule.cpp
+++ b/Engine/Source/Runtime/EngineMessages/Private/EngineMessagesModule.cpp
@@ -5,7 +5,6 @@
 =============================================================================*/
 
 #include "EngineMessagesPrivatePCH.h"
-#include "EngineMessages.generated.inl"
 
 
 /**
diff --git a/Engine/Source/Runtime/EngineSettings/Private/EngineSettingsModule.cpp b/Engine/Source/Runtime/EngineSettings/Private/EngineSettingsModule.cpp
--- a/Engine/Source/Runtime/EngineSettings/Private/EngineSettingsModule.cpp
+++ b/Engine/Source/Runtime/EngineSettings/Private/EngineSettingsModule.cpp
@@ -5,7 +5,6 @@
 =============================================================================*/
 
 #include "EngineSettingsPrivatePCH.h"
-#include "EngineSettings.generated.inl"
 
 
 /**
diff --git a/Engine/Source/Runtime/GameplayTags/Private/GameplayTagsModule.cpp b/Engine/Source/Runtime/GameplayTags/Private/GameplayTagsModule.cpp
--- a/Engine/Source/Runtime/GameplayTags/Private/GameplayTagsModule.cpp
+++ b/Engine/Source/Runtime/GameplayTags/Private/GameplayTagsModule.cpp
@@ -1,7 +1,7 @@
 // Copyright 1998-2014 Epic Games, Inc. All Rights Reserved.
 
 #include "GameplayTagsModulePrivatePCH.h"
-#include "GameplayTags.generated.inl"
+
 
 
 class FGameplayTagsModule : public IGameplayTagsModule
diff --git a/Engine/Source/Runtime/IOS/IOSRuntimeSettings/Private/IOSRuntimeSettingsModule.cpp b/Engine/Source/Runtime/IOS/IOSRuntimeSettings/Private/IOSRuntimeSettingsModule.cpp
--- a/Engine/Source/Runtime/IOS/IOSRuntimeSettings/Private/IOSRuntimeSettingsModule.cpp
+++ b/Engine/Source/Runtime/IOS/IOSRuntimeSettings/Private/IOSRuntimeSettingsModule.cpp
@@ -1,6 +1,6 @@
 // Copyright 1998-2014 Epic Games, Inc. All Rights Reserved.
 
 #include "IOSRuntimeSettingsPrivatePCH.h"
-#include "IOSRuntimeSettings.generated.inl"
+
 
 IMPLEMENT_MODULE(FDefaultModuleImpl, IOSRuntimeSettings);
diff --git a/Engine/Source/Runtime/IOS/LaunchDaemonMessages/Private/LaunchDaemonMessagesModule.cpp b/Engine/Source/Runtime/IOS/LaunchDaemonMessages/Private/LaunchDaemonMessagesModule.cpp
--- a/Engine/Source/Runtime/IOS/LaunchDaemonMessages/Private/LaunchDaemonMessagesModule.cpp
+++ b/Engine/Source/Runtime/IOS/LaunchDaemonMessages/Private/LaunchDaemonMessagesModule.cpp
@@ -5,7 +5,7 @@
 =============================================================================*/
 
 #include "LaunchDaemonMessagesPrivatePCH.h"
-#include "LaunchDaemonMessages.generated.inl"
+
 
 
 /**
diff --git a/Engine/Source/Runtime/InputCore/Private/InputCore.cpp b/Engine/Source/Runtime/InputCore/Private/InputCore.cpp
--- a/Engine/Source/Runtime/InputCore/Private/InputCore.cpp
+++ b/Engine/Source/Runtime/InputCore/Private/InputCore.cpp
@@ -4,7 +4,5 @@
 #include "InputCoreModule.h"
 #include "ModuleManager.h"
 
-#include "InputCore.generated.inl"
-
 
 IMPLEMENT_MODULE( FInputCoreModule, InputCore );
\ No newline at end of file
diff --git a/Engine/Source/Runtime/MoviePlayer/Private/MoviePlayerSettings.cpp b/Engine/Source/Runtime/MoviePlayer/Private/MoviePlayerSettings.cpp
--- a/Engine/Source/Runtime/MoviePlayer/Private/MoviePlayerSettings.cpp
+++ b/Engine/Source/Runtime/MoviePlayer/Private/MoviePlayerSettings.cpp
@@ -2,7 +2,7 @@
 
 #include "MoviePlayer.h"
 #include "MoviePlayerSettings.h"
-#include "MoviePlayer.generated.inl"
+
 
 #define LOCTEXT_NAMESPACE "MoviePlayerSettings"
 
diff --git a/Engine/Source/Runtime/MovieSceneCore/Private/MovieSceneCoreModule.cpp b/Engine/Source/Runtime/MovieSceneCore/Private/MovieSceneCoreModule.cpp
--- a/Engine/Source/Runtime/MovieSceneCore/Private/MovieSceneCoreModule.cpp
+++ b/Engine/Source/Runtime/MovieSceneCore/Private/MovieSceneCoreModule.cpp
@@ -6,7 +6,7 @@
 #include "MovieScene.h"
 #include "MovieSceneBindings.h"
 #include "RuntimeMovieScenePlayer.h"
-#include "MovieSceneCore.generated.inl"
+
 
 
 /**
diff --git a/Engine/Source/Runtime/MovieSceneCoreTypes/Private/MovieSceneCoreTypesModule.cpp b/Engine/Source/Runtime/MovieSceneCoreTypes/Private/MovieSceneCoreTypesModule.cpp
--- a/Engine/Source/Runtime/MovieSceneCoreTypes/Private/MovieSceneCoreTypesModule.cpp
+++ b/Engine/Source/Runtime/MovieSceneCoreTypes/Private/MovieSceneCoreTypesModule.cpp
@@ -4,7 +4,7 @@
 #include "ModuleManager.h"
 
 
-#include "MovieSceneCoreTypes.generated.inl"
+
 
 /**
  * FMovieSceneCoreTypesModule module implementation (private)
diff --git a/Engine/Source/Runtime/Online/OnlineSubsystem/Private/OnlineSubsystem.cpp b/Engine/Source/Runtime/Online/OnlineSubsystem/Private/OnlineSubsystem.cpp
--- a/Engine/Source/Runtime/Online/OnlineSubsystem/Private/OnlineSubsystem.cpp
+++ b/Engine/Source/Runtime/Online/OnlineSubsystem/Private/OnlineSubsystem.cpp
@@ -5,8 +5,6 @@
 #include "OnlineIdentityInterface.h"
 #include "NboSerializer.h"
 
-#include "OnlineSubsystem.generated.inl"
-
 DEFINE_LOG_CATEGORY(LogOnline);
 DEFINE_LOG_CATEGORY(LogOnlineGame);
 
diff --git a/Engine/Source/Runtime/Online/OnlineSubsystemSteam/OnlineSubsystemSteam.Build.cs b/Engine/Source/Runtime/Online/OnlineSubsystemSteam/OnlineSubsystemSteam.Build.cs
--- a/Engine/Source/Runtime/Online/OnlineSubsystemSteam/OnlineSubsystemSteam.Build.cs
+++ b/Engine/Source/Runtime/Online/OnlineSubsystemSteam/OnlineSubsystemSteam.Build.cs
@@ -9,6 +9,12 @@ public class OnlineSubsystemSteam : ModuleRules
 	{
 		Definitions.Add("ONLINESUBSYSTEMSTEAM_PACKAGE=1");
 
+		PublicDependencyModuleNames.AddRange(
+			new string[] { 
+				"OnlineSubsystemUtils",
+			}
+		);
+
 		PrivateDependencyModuleNames.AddRange(
 			new string[] {
 				"Core", 
@@ -17,9 +23,8 @@ public class OnlineSubsystemSteam : ModuleRules
 				"Sockets", 
 				"Voice",
 				"OnlineSubsystem", 
-				"OnlineSubsystemUtils", 
 			}
-			);
+		);
 
 		AddThirdPartyPrivateStaticDependencies(Target, "Steamworks");
 
diff --git a/Engine/Source/Runtime/Online/OnlineSubsystemSteam/Private/OnlineSubsystemSteam.cpp b/Engine/Source/Runtime/Online/OnlineSubsystemSteam/Private/OnlineSubsystemSteam.cpp
--- a/Engine/Source/Runtime/Online/OnlineSubsystemSteam/Private/OnlineSubsystemSteam.cpp
+++ b/Engine/Source/Runtime/Online/OnlineSubsystemSteam/Private/OnlineSubsystemSteam.cpp
@@ -7,7 +7,7 @@
 #include "IpAddressSteam.h"
 #include "OnlineSubsystemUtilsClasses.h"
 #include "OnlineSubsystemSteamClasses.h"
-#include "OnlineSubsystemSteam.generated.inl"
+
 
 #include "SocketSubsystemSteam.h"
 
diff --git a/Engine/Source/Runtime/Online/OnlineSubsystemSteam/Private/OnlineSubsystemSteamPrivatePCH.h b/Engine/Source/Runtime/Online/OnlineSubsystemSteam/Private/OnlineSubsystemSteamPrivatePCH.h
--- a/Engine/Source/Runtime/Online/OnlineSubsystemSteam/Private/OnlineSubsystemSteamPrivatePCH.h
+++ b/Engine/Source/Runtime/Online/OnlineSubsystemSteam/Private/OnlineSubsystemSteamPrivatePCH.h
@@ -5,6 +5,8 @@
 #include "Engine.h"
 #include "OnlineSubsystemSteamModule.h"
 #include "OnlineSubsystemModule.h"
+#include "OnlineSubsystemUtils.h"
+#include "OnlineSubsystemUtilsClasses.h"
 #include "OnlineSubsystem.h"
 #include "OnlineSubsystemUtils.h"
 #include "SocketSubsystem.h"
diff --git a/Engine/Source/Runtime/Online/OnlineSubsystemUtils/Private/OnlineSubsystemUtils.cpp b/Engine/Source/Runtime/Online/OnlineSubsystemUtils/Private/OnlineSubsystemUtils.cpp
--- a/Engine/Source/Runtime/Online/OnlineSubsystemUtils/Private/OnlineSubsystemUtils.cpp
+++ b/Engine/Source/Runtime/Online/OnlineSubsystemUtils/Private/OnlineSubsystemUtils.cpp
@@ -5,7 +5,7 @@
 #include "ModuleManager.h"
 
 #include "IPAddress.h"
-#include "OnlineSubsystemUtils.generated.inl"
+
 #include "NboSerializer.h"
 
 #include "Voice.h"
diff --git a/Engine/Source/Runtime/SkillSystem/Private/SkillSystemModule.cpp b/Engine/Source/Runtime/SkillSystem/Private/SkillSystemModule.cpp
--- a/Engine/Source/Runtime/SkillSystem/Private/SkillSystemModule.cpp
+++ b/Engine/Source/Runtime/SkillSystem/Private/SkillSystemModule.cpp
@@ -1,7 +1,6 @@
 // Copyright 1998-2013 Epic Games, Inc. All Rights Reserved.
 
 #include "SkillSystemModulePrivatePCH.h"
-#include "SKillSystem.generated.inl"
 #include "SkillSystemGlobals.h"
 
 
diff --git a/Engine/Source/Runtime/Slate/Private/Slate.cpp b/Engine/Source/Runtime/Slate/Private/Slate.cpp
--- a/Engine/Source/Runtime/Slate/Private/Slate.cpp
+++ b/Engine/Source/Runtime/Slate/Private/Slate.cpp
@@ -3,7 +3,7 @@
 #include "Slate.h"
 #include "ModuleManager.h"
 
-#include "Slate.generated.inl"
+
 
 IMPLEMENT_MODULE( FDefaultModuleImpl, Slate );
 
diff --git a/Engine/Source/Runtime/UMG/Private/UMGModule.cpp b/Engine/Source/Runtime/UMG/Private/UMGModule.cpp
--- a/Engine/Source/Runtime/UMG/Private/UMGModule.cpp
+++ b/Engine/Source/Runtime/UMG/Private/UMGModule.cpp
@@ -5,8 +5,6 @@
 #include "IUMGModule.h"
 #include "ModuleManager.h"
 
-#include "UMG.generated.inl"
-
 class FUMGModule : public IUMGModule
 {
 public:
