#UE4:
* improved game console auto completion to be more like editor

diff --git a/Engine/Source/Runtime/Engine/Classes/Engine/Console.h b/Engine/Source/Runtime/Engine/Classes/Engine/Console.h
--- a/Engine/Source/Runtime/Engine/Classes/Engine/Console.h
+++ b/Engine/Source/Runtime/Engine/Classes/Engine/Console.h
@@ -136,6 +136,7 @@ class ENGINE_API UConsole : public UObject, public FOutputDevice
 	/** Currently selected auto complete index */
 	int32 AutoCompleteIndex;
 
+	/** -1: auto complete cursor is not visible, >=0 */
 	int32 AutoCompleteCursor;
 
 	/** Do we need to rebuild auto complete? */
diff --git a/Engine/Source/Runtime/Engine/Private/UserInterface/Console.cpp b/Engine/Source/Runtime/Engine/Private/UserInterface/Console.cpp
--- a/Engine/Source/Runtime/Engine/Private/UserInterface/Console.cpp
+++ b/Engine/Source/Runtime/Engine/Private/UserInterface/Console.cpp
@@ -44,6 +44,7 @@ public:
 
 UConsole::UConsole(const class FPostConstructInitializeProperties& PCIP)
 	: Super(PCIP)
+	, AutoCompleteCursor(-1)
 {
 	// Structure to hold one-time initialization
 	struct FConstructorStatics
@@ -252,7 +253,7 @@ void UConsole::UpdateCompleteIndices()
 	}
 	bNavigatingHistory = false;
 	AutoCompleteIndex = 0;
-	AutoCompleteCursor = 0;
+	AutoCompleteCursor = -1;
 	AutoCompleteIndices.Empty();
 	FAutoCompleteNode *Node = &AutoCompleteTree;
 	FString LowerTypedStr = TypedStr.ToLower();
@@ -592,7 +593,7 @@ bool UConsole::InputKey_InputLine( int32 ControllerId, FKey Key, EInputEvent Eve
 		}
 		else if (Key == EKeys::Tab && Event == IE_Pressed)
 		{
-			if (AutoCompleteIndices.Num() > 0 && !bAutoCompleteLocked)
+			if (AutoCompleteIndices.Num() > 0 && !bAutoCompleteLocked && AutoCompleteCursor >= 0)
 			{
 				const FAutoCompleteCommand& Cmd = AutoCompleteList[AutoCompleteIndices[AutoCompleteIndex + AutoCompleteCursor]];
 
@@ -778,6 +779,20 @@ bool UConsole::InputKey_InputLine( int32 ControllerId, FKey Key, EInputEvent Eve
 			SetCursorPos(TypedStr.Len());
 			return true;
 		}
+		else
+		{
+			// auto complete is open and the user adds a key
+			if (AutoCompleteIndices.Num() > 0 && !bAutoCompleteLocked && AutoCompleteCursor >= 0)
+			{
+				const FAutoCompleteCommand& Cmd = AutoCompleteList[AutoCompleteIndices[AutoCompleteIndex + AutoCompleteCursor]];
+
+				TypedStr = Cmd.Command;
+				SetCursorPos(TypedStr.Len());
+				bAutoCompleteLocked = true;
+
+				// no "return true" so that the following key will be added the auto completed string
+			}
+		}
 	}
 
 	return false;
