Fix speedtree for ps4

diff --git a/Engine/Source/Runtime/Engine/Public/SpeedTreeWind.h b/Engine/Source/Runtime/Engine/Public/SpeedTreeWind.h
--- a/Engine/Source/Runtime/Engine/Public/SpeedTreeWind.h
+++ b/Engine/Source/Runtime/Engine/Public/SpeedTreeWind.h
@@ -300,7 +300,7 @@ protected:
 					float		m_afBranchWindAnchor[3];
 					float		m_fMaxBranchLevel1Length;
 
-	MS_ALIGN(16)	float		m_afShaderTable[NUM_SHADER_VALUES];
+	MS_ALIGN(16)	float		m_afShaderTable[NUM_SHADER_VALUES] GCC_ALIGN(16);
 };
 
 
