Convert vehicle units to use Nm for torque and RPM for wheel and engine angular velocity

diff --git a/Engine/Config/BaseEngine.ini b/Engine/Config/BaseEngine.ini
--- a/Engine/Config/BaseEngine.ini
+++ b/Engine/Config/BaseEngine.ini
@@ -1087,6 +1087,10 @@ bUseStreamingPause=false
 +K2FieldRedirects=(OldFieldName="PointLightComponent.Radius",NewFieldName="PointLightComponent.AttenuationRadius")
 +K2FieldRedirects=(OldFieldName="LightComponent.InverseSquaredFalloff",NewFieldName="PointLightComponent.bUseInverseSquaredFalloff")
 
+; Vehicle rename properties
++TaggedPropertyRedirects=(ClassName="VehicleEngineData",OldPropertyName="MaxOmega",NewPropertyName="MaxRPM")
++TaggedPropertyRedirects=(ClassName="VehicleMovementComponentBoosted4w",OldPropertyName="BoostedMaxOmega",NewPropertyName="BoostedMaxRPM")
+
 [CoreUObject.Metadata]
 ; Note: UnrealHeaderTool should be rerun after making changes to MetadataRedirects to catch any keys specified in class headers
 +MetadataRedirects=(OldKey="K2Protected", NewKey="BlueprintProtected")
diff --git a/Engine/Source/Runtime/Core/Public/UObject/ObjectVersion.h b/Engine/Source/Runtime/Core/Public/UObject/ObjectVersion.h
--- a/Engine/Source/Runtime/Core/Public/UObject/ObjectVersion.h
+++ b/Engine/Source/Runtime/Core/Public/UObject/ObjectVersion.h
@@ -522,6 +522,8 @@ enum EUnrealEngineObjectUE4Version
 	VER_UE4_FLIP_MATERIAL_COORDS,
 	// PinSubCategoryMemberReference added to FEdGraphPinType
 	VER_UE4_MEMBERREFERENCE_IN_PINTYPE,
+	// Vehicles use Nm for Torque instead of cm and RPM instead of rad/s
+	VER_UE4_VEHICLES_UNIT_CHANGE,
 	// -----<new versions can be added before this line>-------------------------------------------------
 	// - this needs to be the last line (see note below)
 	VER_UE4_AUTOMATIC_VERSION_PLUS_ONE,
diff --git a/Engine/Source/Runtime/Engine/Classes/Vehicles/WheeledVehicleMovementComponent.h b/Engine/Source/Runtime/Engine/Classes/Vehicles/WheeledVehicleMovementComponent.h
--- a/Engine/Source/Runtime/Engine/Classes/Vehicles/WheeledVehicleMovementComponent.h
+++ b/Engine/Source/Runtime/Engine/Classes/Vehicles/WheeledVehicleMovementComponent.h
@@ -100,12 +100,12 @@ struct FWheelSetup
 	UPROPERTY(EditAnywhere, Category=WheelsSetup)
 	float SteerAngle;
 
-	// max brake torque for this wheel
+	// max brake torque for this wheel (Nm)
 	UPROPERTY(EditAnywhere, Category=WheelsSetup)
 	float MaxBrakeTorque;
 
-	/** Max handbreak brake torque for this wheel. A handbrake should have a stronger brake torque
-		than the break. This will be ignored for wheels that are not affected by the handbrake. */
+	/** Max handbrake brake torque for this wheel (Nm). A handbrake should have a stronger brake torque
+		than the brake. This will be ignored for wheels that are not affected by the handbrake. */
 	UPROPERTY(EditAnywhere, Category=WheelsSetup)
 	float MaxHandBrakeTorque;
 
@@ -122,8 +122,8 @@ struct FWheelSetup
 		, BoneName(NAME_None)
 		, AdditionalOffset(0.0f)
 		, SteerAngle(70.0f)
-		, MaxBrakeTorque(1500 * 100 * 100)
-		, MaxHandBrakeTorque(0.0f)
+		, MaxBrakeTorque(1500.f)
+		, MaxHandBrakeTorque(3000.f)
 		, DampingRate(0.25f * 100.f * 100.f)
 		, Mass(20.f)
 	{
@@ -357,6 +357,8 @@ class ENGINE_API UWheeledVehicleMovementComponent : public UPawnMovementComponen
 	UFUNCTION(BlueprintCallable, Category="Game|Components|WheeledVehicleMovement")
 	bool GetUseAutoGears() const;
 
+	virtual void Serialize(FArchive & Ar) OVERRIDE;
+
 protected:
 
 	// replicated state of vehicle 
@@ -496,4 +498,41 @@ protected:
 	class USkinnedMeshComponent* GetMesh();
 
 #endif // WITH_PHYSX
+	
+
 };
+
+//some helper functions for converting units
+
+//rev per minute to rad/s
+inline float RPMToOmega(float RPM)
+{
+	return RPM * PI / 30.f;
+}
+
+//rad/s to rev per minute
+inline float OmegaToRPM(float Omega)
+{
+	return Omega * 30.f / PI;
+}
+
+//km/h to cm/s
+inline float KmHToCmS(float KmH)
+{
+	return KmH * 100000.f / 3600.f;
+}
+
+inline float CmSToKmH(float CmS)
+{
+	return CmS * 3600.f / 100000.f;
+}
+
+inline float M2ToCm2(float M2)
+{
+	return M2 * 100.f * 100.f;
+}
+
+inline float Cm2ToM2(float Cm2)
+{
+	return Cm2 / (100.f * 100.f);
+}
\ No newline at end of file
diff --git a/Engine/Source/Runtime/Engine/Classes/Vehicles/WheeledVehicleMovementComponent4W.h b/Engine/Source/Runtime/Engine/Classes/Vehicles/WheeledVehicleMovementComponent4W.h
--- a/Engine/Source/Runtime/Engine/Classes/Vehicles/WheeledVehicleMovementComponent4W.h
+++ b/Engine/Source/Runtime/Engine/Classes/Vehicles/WheeledVehicleMovementComponent4W.h
@@ -63,13 +63,13 @@ struct FVehicleEngineData
 	UPROPERTY(EditAnywhere, Category=Setup)
 	float MOI;
 
-	/** Maximum torque available to apply to the engine (kg m^2 s^-2). */ 
+	/** Maximum torque available to apply to the engine (Nm). */ 
 	UPROPERTY(EditAnywhere, Category=Setup)
 	float PeakTorque;
 
-	/** Maximum rotation speed of the engine (radians per second: rad s^-1) */
+	/** Maximum revolutions per minute of the engine */
 	UPROPERTY(EditAnywhere, Category=Setup)
-	float MaxOmega;
+	float MaxRPM;
 
 	/** Damping rate of engine when full throttle is applied */
 	UPROPERTY(EditAnywhere, Category=Setup, AdvancedDisplay)
@@ -171,7 +171,7 @@ class ENGINE_API UWheeledVehicleMovementComponent4W : public UWheeledVehicleMove
 	UPROPERTY(EditAnywhere, Category=VehicleSetup, AdvancedDisplay)
 	float AckermannAccuracy;
 
-	/** Max speed value in steering curve */
+	/** Max speed value in steering curve (km/h)*/
 	UPROPERTY(EditAnywhere, Category=VehicleSetup)
 	float MaxSteeringSpeed;
 
@@ -179,6 +179,8 @@ class ENGINE_API UWheeledVehicleMovementComponent4W : public UWheeledVehicleMove
 	UPROPERTY(EditAnywhere, Category=VehicleSetup, AdvancedDisplay)
 	TArray<FFloatPair> SteeringCurve;
 
+	virtual void Serialize(FArchive & Ar) OVERRIDE;
+
 protected:
 
 #if WITH_PHYSX
diff --git a/Engine/Source/Runtime/Engine/Private/Vehicles/WheeledVehicleMovementComponent.cpp b/Engine/Source/Runtime/Engine/Private/Vehicles/WheeledVehicleMovementComponent.cpp
--- a/Engine/Source/Runtime/Engine/Private/Vehicles/WheeledVehicleMovementComponent.cpp
+++ b/Engine/Source/Runtime/Engine/Private/Vehicles/WheeledVehicleMovementComponent.cpp
@@ -283,8 +283,8 @@ void UWheeledVehicleMovementComponent::SetupWheels( PxVehicleWheelsSimData* PWhe
 		PWheelData.mRadius = Wheel->ShapeRadius;
 		PWheelData.mWidth = Wheel->ShapeWidth;	
 		PWheelData.mMaxSteer = FMath::DegreesToRadians(WheelSetups[WheelIdx].SteerAngle);
-		PWheelData.mMaxBrakeTorque = WheelSetups[WheelIdx].MaxBrakeTorque;
-		PWheelData.mMaxHandBrakeTorque = Wheel->bAffectedByHandbrake ? WheelSetups[WheelIdx].MaxHandBrakeTorque : 0.0f;
+		PWheelData.mMaxBrakeTorque = M2ToCm2(WheelSetups[WheelIdx].MaxBrakeTorque);
+		PWheelData.mMaxHandBrakeTorque = Wheel->bAffectedByHandbrake ? M2ToCm2(WheelSetups[WheelIdx].MaxHandBrakeTorque) : 0.0f;
 
 //		PWheelData.mMass = 1.0f;
 //		PWheelData.mMOI = Wheel->Inertia * Mass; // Multiply by vehicle mass so that as vehicle mass shrinks, the torque and inertia can shrink as well
@@ -295,7 +295,7 @@ void UWheeledVehicleMovementComponent::SetupWheels( PxVehicleWheelsSimData* PWhe
 		// init tire data
 		PxVehicleTireData PTireData; 
 		PTireData.mType = Wheel->TireType ? Wheel->TireType->GetTireTypeID() : GEngine->DefaultTireType->GetTireTypeID();
-		PTireData.mCamberStiffnessPerUnitGravity = 0.0f;
+		//PTireData.mCamberStiffnessPerUnitGravity = 0.0f;
 		PTireData.mLatStiffX = Wheel->LatStiffMaxLoad;
 		PTireData.mLatStiffY = Wheel->LatStiffValue;
 		PTireData.mLongitudinalStiffnessPerUnitGravity = Wheel->LongStiffValue;
@@ -706,10 +706,10 @@ void UWheeledVehicleMovementComponent::UpdateState( float DeltaTime )
 			}
 		}
 		
-		SteeringInput = SteeringInputRate.InterpInputValue( DeltaTime, SteeringInput, CalcSteeringInput() );
+		SteeringInput = SteeringInputRate.InterpInputValue(DeltaTime, SteeringInput, CalcSteeringInput());
 		ThrottleInput = ThrottleInputRate.InterpInputValue( DeltaTime, ThrottleInput, CalcThrottleInput() );
-		BrakeInput = BrakeInputRate.InterpInputValue( DeltaTime, BrakeInput, CalcBrakeInput() );
-		HandbrakeInput = HandbrakeInputRate.InterpInputValue( DeltaTime, HandbrakeInput, CalcHandbrakeInput() );
+		BrakeInput = BrakeInputRate.InterpInputValue(DeltaTime, BrakeInput, CalcBrakeInput());
+		HandbrakeInput = HandbrakeInputRate.InterpInputValue(DeltaTime, HandbrakeInput, CalcHandbrakeInput());
 
 		// and send to server
 		ServerUpdateState(SteeringInput, ThrottleInput, BrakeInput, HandbrakeInput, GetCurrentGear());
@@ -1064,6 +1064,22 @@ bool UWheeledVehicleMovementComponent::CheckSlipThreshold(float AbsLongSlipThres
 	return false;
 }
 
+void UWheeledVehicleMovementComponent::Serialize(FArchive & Ar)
+{
+	Super::Serialize(Ar);
+	if (Ar.IsLoading() && Ar.UE4Ver() < VER_UE4_VEHICLES_UNIT_CHANGE)
+	{
+		for (int32 WheelIdx = 0; WheelIdx < WheelSetups.Num(); ++WheelIdx)
+		{
+			//this backwards compatable code breaks if they were using some strange value that is exactly the new default in the new units
+			FWheelSetup & WheelSetup = WheelSetups[WheelIdx];
+			WheelSetup.MaxHandBrakeTorque = WheelSetup.MaxHandBrakeTorque != 1500.f ? Cm2ToM2(WheelSetup.MaxHandBrakeTorque) : 1500.f;	//need to convert from cm^2 to m^2
+			WheelSetup.MaxBrakeTorque = WheelSetup.MaxHandBrakeTorque != 3000.f ? Cm2ToM2(WheelSetup.MaxBrakeTorque) : 3000.f;			//need to convert from cm^2 to m^2
+		}
+
+	}
+}
+
 float UWheeledVehicleMovementComponent::GetMaxSpringForce() const
 {
 	if (PVehicle == NULL)
diff --git a/Engine/Source/Runtime/Engine/Private/Vehicles/WheeledVehicleMovementComponent4W.cpp b/Engine/Source/Runtime/Engine/Private/Vehicles/WheeledVehicleMovementComponent4W.cpp
--- a/Engine/Source/Runtime/Engine/Private/Vehicles/WheeledVehicleMovementComponent4W.cpp
+++ b/Engine/Source/Runtime/Engine/Private/Vehicles/WheeledVehicleMovementComponent4W.cpp
@@ -24,8 +24,8 @@ UWheeledVehicleMovementComponent4W::UWheeledVehicleMovementComponent4W(const cla
 
 	PxVehicleEngineData DefEngineData;
 	EngineSetup.MOI = DefEngineData.mMOI * LengthScale * LengthScale;
-	EngineSetup.PeakTorque = 1000 * LengthScale * LengthScale;
-	EngineSetup.MaxOmega = 1000;
+	EngineSetup.PeakTorque = 1000.f;
+	EngineSetup.MaxRPM = 7000.f;
 	EngineSetup.DampingRateFullThrottle = 0.5f * LengthScale * LengthScale;
 	EngineSetup.DampingRateZeroThrottleClutchEngaged = 2.0f * LengthScale * LengthScale;
 	EngineSetup.DampingRateZeroThrottleClutchDisengaged = 0.5f * LengthScale * LengthScale;
@@ -57,7 +57,7 @@ UWheeledVehicleMovementComponent4W::UWheeledVehicleMovementComponent4W(const cla
 	AutoBoxSetup.GearAutoBoxLatency = DefAutoBoxSetup.getLatency();
 	bUseGearAutoBox = true;
 
-	MaxSteeringSpeed = 6000.0f; // editable in vehicle blueprint
+	MaxSteeringSpeed = 100.f; // editable in vehicle blueprint
 	SteeringCurve.AddZeroed(4);
 	SteeringCurve[0].InVal = 0.0f;
 	SteeringCurve[0].OutVal = 0.75f;
@@ -115,8 +115,8 @@ static void GetVehicleDifferential4WSetup(const FVehicleDifferential4WData& Setu
 static void GetVehicleEngineSetup(const FVehicleEngineData& Setup, PxVehicleEngineData& PxSetup)
 {
 	PxSetup.mMOI = Setup.MOI; // convert from m to cm scale
-	PxSetup.mPeakTorque = Setup.PeakTorque;
-	PxSetup.mMaxOmega = Setup.MaxOmega;
+	PxSetup.mPeakTorque = M2ToCm2(Setup.PeakTorque);	//convert Nm to (kg cm^2/s^2)
+	PxSetup.mMaxOmega = RPMToOmega(Setup.MaxRPM);
 	PxSetup.mDampingRateFullThrottle = Setup.DampingRateFullThrottle;
 	PxSetup.mDampingRateZeroThrottleClutchEngaged = Setup.DampingRateZeroThrottleClutchEngaged;
 	PxSetup.mDampingRateZeroThrottleClutchDisengaged = Setup.DampingRateZeroThrottleClutchDisengaged;
@@ -228,7 +228,7 @@ void UWheeledVehicleMovementComponent4W::SetupVehicle()
 	const int MaxSteeringSamples = FMath::Min(8, SteeringCurve.Num());
 	for (int i = 0; i < 8; i++)
 	{
-		SteeringMap[i*2 + 0] = (i < MaxSteeringSamples) ? (SteeringCurve[i].InVal * MaxSteeringSpeed) : PX_MAX_F32;
+		SteeringMap[i*2 + 0] = (i < MaxSteeringSamples) ? (SteeringCurve[i].InVal * KmHToCmS(MaxSteeringSpeed)) : PX_MAX_F32;
 		SteeringMap[i*2 + 1] = (i < MaxSteeringSamples) ? FMath::Clamp<float>(SteeringCurve[i].OutVal, 0.0f, 1.0f) : PX_MAX_F32;
 	}
 }
@@ -317,3 +317,16 @@ void UWheeledVehicleMovementComponent4W::UpdateAutoBoxSetup(const FVehicleAutoBo
 	}
 #endif
 }
+
+void UWheeledVehicleMovementComponent4W::Serialize(FArchive & Ar)
+{
+	Super::Serialize(Ar);
+	if (Ar.IsLoading() && Ar.UE4Ver() < VER_UE4_VEHICLES_UNIT_CHANGE)
+	{
+		//we need to convert from old units to new. This backwards compatable code fails in the rare case that they were using very strange values that are the new defaults in the correct units.
+		EngineSetup.PeakTorque = EngineSetup.PeakTorque != 1000.f ? Cm2ToM2(EngineSetup.PeakTorque) : 1000.f;	//need to convert kg cm^2/s^2 into Nm
+		EngineSetup.MaxRPM = EngineSetup.MaxRPM != 7000.f ? OmegaToRPM(EngineSetup.MaxRPM) : 7000.f;	//need to convert from rad/s to RPM
+		MaxSteeringSpeed = MaxSteeringSpeed != 100.f ? CmSToKmH(MaxSteeringSpeed) : 100.f;	//we now store as km/h instead of cm/s
+	}
+}
+
