Fix inconsistency between RaycastSingle and RaycastTest

diff --git a/Engine/Source/Runtime/Engine/Private/Collision/PhysXCollision.cpp b/Engine/Source/Runtime/Engine/Private/Collision/PhysXCollision.cpp
--- a/Engine/Source/Runtime/Engine/Private/Collision/PhysXCollision.cpp
+++ b/Engine/Source/Runtime/Engine/Private/Collision/PhysXCollision.cpp
@@ -323,6 +323,7 @@ bool RaycastTest(const UWorld * World, const FVector Start, const FVector End, E
 		// Create filter data used to filter collisions
 		PxFilterData PFilter = CreateQueryFilterData(TraceChannel, Params.bTraceComplex, ResponseParams.CollisionResponse, ObjectParams, false);
 		PxSceneQueryFilterData PQueryFilterData(PFilter, PxSceneQueryFilterFlag::eSTATIC | PxSceneQueryFilterFlag::eDYNAMIC | PxSceneQueryFilterFlag::ePREFILTER);
+		PxSceneQueryFlags POutputFlags = PxHitFlags();
 		FPxQueryFilterCallback PQueryCallback(Params.IgnoreActors);
 
 		FPhysScene* PhysScene = World->GetPhysicsScene();
@@ -332,15 +333,14 @@ bool RaycastTest(const UWorld * World, const FVector Start, const FVector End, E
 		SCOPED_SCENE_READ_LOCK(SyncScene);
 
 		const PxVec3 PDir = U2PVector(Delta/DeltaMag);
-
-		PxSceneQueryHit PQueryHit;
-		bHaveBlockingHit = SyncScene->raycastAny(U2PVector(Start), PDir, DeltaMag, PQueryHit, PQueryFilterData, &PQueryCallback);
+		PxRaycastHit PHit;
+		bHaveBlockingHit = SyncScene->raycastSingle(U2PVector(Start), PDir, DeltaMag, POutputFlags, PHit, PQueryFilterData, &PQueryCallback);
 
 		// Test async scene if we have no blocking hit, and async tests are requested
 		if( !bHaveBlockingHit && Params.bTraceAsyncScene && PhysScene->HasAsyncScene())
 		{
 			PxScene* AsyncScene = PhysScene->GetPhysXScene(PST_Async);
-			SCOPED_SCENE_READ_LOCK(AsyncScene);			bHaveBlockingHit = AsyncScene->raycastAny(U2PVector(Start), PDir, DeltaMag, PQueryHit, PQueryFilterData, &PQueryCallback);
+			SCOPED_SCENE_READ_LOCK(AsyncScene);			bHaveBlockingHit = AsyncScene->raycastSingle(U2PVector(Start), PDir, DeltaMag, POutputFlags, PHit, PQueryFilterData, &PQueryCallback);
 		}
 	}
 
