Fix non-unity(hopefully)

diff --git a/Engine/Source/Runtime/Engine/Private/GameViewportClient.cpp b/Engine/Source/Runtime/Engine/Private/GameViewportClient.cpp
--- a/Engine/Source/Runtime/Engine/Private/GameViewportClient.cpp
+++ b/Engine/Source/Runtime/Engine/Private/GameViewportClient.cpp
@@ -15,6 +15,7 @@
 #include "SceneViewExtension.h"
 #include "IHeadMountedDisplay.h"
 #include "SVirtualJoystick.h"
+#include "SceneViewport.h"
 
 #include "HighResScreenshot.h"
 
