Fix vehicle UI to show correct RPM and speed

diff --git a/Engine/Source/Runtime/Engine/Classes/Vehicles/WheeledVehicleMovementComponent.h b/Engine/Source/Runtime/Engine/Classes/Vehicles/WheeledVehicleMovementComponent.h
--- a/Engine/Source/Runtime/Engine/Classes/Vehicles/WheeledVehicleMovementComponent.h
+++ b/Engine/Source/Runtime/Engine/Classes/Vehicles/WheeledVehicleMovementComponent.h
@@ -235,6 +235,25 @@ class ENGINE_API UWheeledVehicleMovementComponent : public UPawnMovementComponen
 	UPROPERTY(EditAnywhere, Category = VehicleSetup)
 	float ChassisHeight;
 
+	// Drag area in cm^2
+	UPROPERTY(transient)
+	float DragArea;
+
+	// Estimated mad speed for engine
+	UPROPERTY(transient)
+	float EstimatedMaxEngineSpeed;
+
+	// Max RPM for engine
+	UPROPERTY(transient)
+	float MaxEngineRPM;
+
+	// Debug drag magnitude last applied
+	UPROPERTY(transient)
+	float DebugDragMagnitude;
+
+	/** When vehicle is created we want to compute some helper data like drag area, etc.... Derived classes should use this to properly compute things like engine RPM */
+	virtual void ComputeConstants();
+
 	/** Override center of mass offset, makes tweaking easier [uu] */
 	UPROPERTY(EditAnywhere, Category=VehicleSetup, AdvancedDisplay)
 	FVector COMOffset;
@@ -432,10 +451,6 @@ protected:
 	UPROPERTY(EditAnywhere, Category=VehicleInput, AdvancedDisplay)
 	FVehicleInputRate SteeringInputRate;
 
-	// Debug drag magnitude last applied
-	UPROPERTY(transient)
-	float DebugDragMagnitude;
-
 	/** Compute steering input */
 	float CalcSteeringInput();
 
diff --git a/Engine/Source/Runtime/Engine/Classes/Vehicles/WheeledVehicleMovementComponent4W.h b/Engine/Source/Runtime/Engine/Classes/Vehicles/WheeledVehicleMovementComponent4W.h
--- a/Engine/Source/Runtime/Engine/Classes/Vehicles/WheeledVehicleMovementComponent4W.h
+++ b/Engine/Source/Runtime/Engine/Classes/Vehicles/WheeledVehicleMovementComponent4W.h
@@ -180,6 +180,7 @@ class ENGINE_API UWheeledVehicleMovementComponent4W : public UWheeledVehicleMove
 	TArray<FFloatPair> SteeringCurve;
 
 	virtual void Serialize(FArchive & Ar) OVERRIDE;
+	virtual void ComputeConstants() OVERRIDE;
 
 protected:
 
diff --git a/Engine/Source/Runtime/Engine/Private/Vehicles/WheeledVehicleMovementComponent.cpp b/Engine/Source/Runtime/Engine/Private/Vehicles/WheeledVehicleMovementComponent.cpp
--- a/Engine/Source/Runtime/Engine/Private/Vehicles/WheeledVehicleMovementComponent.cpp
+++ b/Engine/Source/Runtime/Engine/Private/Vehicles/WheeledVehicleMovementComponent.cpp
@@ -143,6 +143,8 @@ bool UWheeledVehicleMovementComponent::CanCreateVehicle() const
 
 void UWheeledVehicleMovementComponent::CreateVehicle()
 {
+	ComputeConstants();
+
 	if ( PVehicle == NULL )
 	{
 		if ( CanCreateVehicle() )
@@ -1103,6 +1105,13 @@ float UWheeledVehicleMovementComponent::GetMaxSpringForce() const
 
 }
 
+void UWheeledVehicleMovementComponent::ComputeConstants()
+{
+	DragArea = ChassisWidth * ChassisHeight;
+	MaxEngineRPM = 5000.f;
+}
+
+
 void UWheeledVehicleMovementComponent::DrawDebug( UCanvas* Canvas, float& YL, float& YPos )
 {
 	if ( PVehicle == NULL )
@@ -1201,6 +1210,8 @@ void UWheeledVehicleMovementComponent::DrawDebug( UCanvas* Canvas, float& YL, fl
 			YPos += YL;
 		}
 	}
+
+	DrawDebugLines();
 }
 
 void UWheeledVehicleMovementComponent::DrawDebugLines()
@@ -1301,4 +1312,4 @@ void UWheeledVehicleMovementComponent::GetLifetimeReplicatedProps( TArray< FLife
 	Super::GetLifetimeReplicatedProps( OutLifetimeProps );
 
 	DOREPLIFETIME( UWheeledVehicleMovementComponent, ReplicatedState );
-}
+}
\ No newline at end of file
diff --git a/Engine/Source/Runtime/Engine/Private/Vehicles/WheeledVehicleMovementComponent4W.cpp b/Engine/Source/Runtime/Engine/Private/Vehicles/WheeledVehicleMovementComponent4W.cpp
--- a/Engine/Source/Runtime/Engine/Private/Vehicles/WheeledVehicleMovementComponent4W.cpp
+++ b/Engine/Source/Runtime/Engine/Private/Vehicles/WheeledVehicleMovementComponent4W.cpp
@@ -330,3 +330,9 @@ void UWheeledVehicleMovementComponent4W::Serialize(FArchive & Ar)
 	}
 }
 
+void UWheeledVehicleMovementComponent4W::ComputeConstants()
+{
+	Super::ComputeConstants();
+	MaxEngineRPM = EngineSetup.MaxRPM;
+}
+
