Fixed particle trails only rendering in one eye when in stereo.

diff --git a/Engine/Source/Runtime/Engine/Private/Particles/ParticleSystemRender.cpp b/Engine/Source/Runtime/Engine/Private/Particles/ParticleSystemRender.cpp
--- a/Engine/Source/Runtime/Engine/Private/Particles/ParticleSystemRender.cpp
+++ b/Engine/Source/Runtime/Engine/Private/Particles/ParticleSystemRender.cpp
@@ -5558,6 +5558,12 @@ void FDynamicTrailsEmitterData::PreRenderView(FParticleSystemSceneProxy* Proxy,
 
 		bool bOnlyOneView = ShouldUsePrerenderView() || ((GEngine && GEngine->GameViewport && (GEngine->GameViewport->GetCurrentSplitscreenConfiguration() == ESplitScreenType::None)) ? true : false);
 
+		// Render both views when in stereo
+		if (GEngine->HMDDevice.IsValid() && GEngine->IsStereoscopic3D())
+		{
+			bOnlyOneView = false;
+		}
+
 		BuildViewFillDataAndSubmit(Proxy,ViewFamily,VisibilityMap,bOnlyOneView,SourcePointer->VertexCount,VertexStride, DynamicParameterVertexStride);
 
 		// Set the frame tracker
