Lower/improve network log verbosity

* It is really hard to see what is happening when you are running tests that perform many connect/disconnects.
* I've tried to narrow it down to just having Open/Close commands, and remove spam that makes this hard to see that doesn't add value or redundant

diff --git a/Engine/Source/Runtime/Engine/Private/NetConnection.cpp b/Engine/Source/Runtime/Engine/Private/NetConnection.cpp
--- a/Engine/Source/Runtime/Engine/Private/NetConnection.cpp
+++ b/Engine/Source/Runtime/Engine/Private/NetConnection.cpp
@@ -191,10 +191,12 @@ void UNetConnection::Close()
 	if (Driver != NULL && State != USOCK_Closed)
 	{
 		NETWORK_PROFILER(GNetworkProfiler.TrackEvent(TEXT("CLOSE"), *(GetName() + TEXT(" ") + LowLevelGetRemoteAddress())));
-		UE_LOG(LogNet, Log, TEXT("UNetConnection::Close: Name: %s, Driver: %s, Owner: %s, RemoteAddr: %s, Time: %s"), 
+		UE_LOG(LogNet, Log, TEXT("UNetConnection::Close: Name: %s, Driver: %s, PC: %s, Owner: %s, Channels: %i, RemoteAddr: %s, Time: %s"), 
 			*GetName(), 
 			*Driver->GetDescription(), 
-			OwningActor ? *OwningActor->GetName() : TEXT("No Owner"),
+			PlayerController ? *PlayerController->GetName() : TEXT("NULL"),
+			OwningActor ? *OwningActor->GetName() : TEXT("NULL"),
+			OpenChannels.Num(),
 			*LowLevelGetRemoteAddress(true),
 			*FDateTime::UtcNow().ToString(TEXT("%Y.%m.%d-%H.%M.%S")));
 
@@ -207,7 +209,7 @@ void UNetConnection::Close()
 	}
 	else
 	{
-		UE_LOG(LogNet, Log, TEXT("UNetConnection::Close: Already closed. Name: %s"), *GetName() );
+		UE_LOG(LogNet, Verbose, TEXT("UNetConnection::Close: Already closed. Name: %s"), *GetName() );
 	}
 
 	LogCallLastTime		= 0;
@@ -224,12 +226,15 @@ void UNetConnection::CleanUp()
 	}
 	Children.Empty();
 
-	UE_LOG(LogNet, Log, TEXT("UNetConnection::Cleanup: Closing connection. Name: %s, RemoteAddr: %s [%s] [%s] [%s] from CleanUp()"),
-		*GetName(),
-		*LowLevelGetRemoteAddress(true),
-		Driver ? *Driver->NetDriverName.ToString() : TEXT("NULL"),
-		PlayerController ? *PlayerController->GetName() : TEXT("NoPC"),
-		OwningActor ? *OwningActor->GetName() : TEXT("No Owner"));
+	if ( State != USOCK_Closed )
+	{
+		UE_LOG(LogNet, Log, TEXT("UNetConnection::Cleanup: Closing open connection. Name: %s, RemoteAddr: %s Driver: %s, PC: %s, Owner: %s"),
+			*GetName(),
+			*LowLevelGetRemoteAddress(true),
+			Driver ? *Driver->NetDriverName.ToString() : TEXT("NULL"),
+			PlayerController ? *PlayerController->GetName() : TEXT("NoPC"),
+			OwningActor ? *OwningActor->GetName() : TEXT("No Owner"));
+	}
 
 	Close();
 
