Added more useful data to network logs + Fixed comment

diff --git a/Engine/Source/Runtime/Engine/Private/NetConnection.cpp b/Engine/Source/Runtime/Engine/Private/NetConnection.cpp
--- a/Engine/Source/Runtime/Engine/Private/NetConnection.cpp
+++ b/Engine/Source/Runtime/Engine/Private/NetConnection.cpp
@@ -891,7 +891,7 @@ void UNetConnection::ReceivedPacket( FBitReader& Reader )
 				// Can't handle other channels until control channel exists.
 				if ( Channels[0] == NULL )
 				{
-					UE_LOG( LogNetTraffic, Error, TEXT( "UNetConnection::ReceivedPacket: Received bunch before connected" ) );
+					UE_LOG( LogNetTraffic, Error, TEXT( "UNetConnection::ReceivedPacket: Received bunch before connected. ChIndex: %i, ChType: %i" ), Bunch.ChIndex, Bunch.ChType );
 					Close();
 					return;
 				}
@@ -944,9 +944,8 @@ void UNetConnection::ReceivedPacket( FBitReader& Reader )
 				if (!ValidUnreliableOpen)
 				{
 					UE_LOG( LogNetTraffic, Warning, TEXT( "      Received unreliable bunch before open (Channel %d Current Sequence %i)" ), Bunch.ChIndex, InReliable[Bunch.ChIndex] );
-					// Just return, stop processing packets, and DON'T ACK
-					// This is because if we ACK, the sender could get the wrong impression, 
-					// and think they don't need to resend any important contents of this bunch (properties, guid's, etc)
+					// Since we won't be processing this packet, don't ack it
+					// We don't want the sender to think this bunch was processed when it really wasn't
 					bSkipAck = true;
 					continue;
 				}
