Add ExternalUIInterfaceGooglePlay and implement ShowAchievementsUI
Since ShowAchievementsUI takes a LocalUserNum parameter, also added a player controller parameter to the ShowPlatformSpecificAchievementsScreen blueprint function.
ShowPlatformSpecificAchievementsScreen now uses ExternalUIInterfaceGooglePlay instead of the hacked JNI call directly for Android.

diff --git a/Engine/Source/Runtime/Engine/Classes/Kismet/KismetSystemLibrary.h b/Engine/Source/Runtime/Engine/Classes/Kismet/KismetSystemLibrary.h
--- a/Engine/Source/Runtime/Engine/Classes/Kismet/KismetSystemLibrary.h
+++ b/Engine/Source/Runtime/Engine/Classes/Kismet/KismetSystemLibrary.h
@@ -1055,7 +1055,9 @@ class UKismetSystemLibrary : public UBlueprintFunctionLibrary
 
 	/**
 	 * Displays the built-in achievements GUI (iOS and Android only; this function may be renamed or moved in a future release)
+	 *
+	 * @param SpecificPlayer Specific player's achievements to show. May not be supported on all platforms. If null, defaults to the player with ControllerId 0
 	 */
 	UFUNCTION(BlueprintCallable, Category = "Utilities|Platform")
-	static void ShowPlatformSpecificAchievementsScreen();
+	static void ShowPlatformSpecificAchievementsScreen(class APlayerController* SpecificPlayer);
 };
diff --git a/Engine/Source/Runtime/Engine/Private/KismetSystemLibrary.cpp b/Engine/Source/Runtime/Engine/Private/KismetSystemLibrary.cpp
--- a/Engine/Source/Runtime/Engine/Private/KismetSystemLibrary.cpp
+++ b/Engine/Source/Runtime/Engine/Private/KismetSystemLibrary.cpp
@@ -7,6 +7,7 @@
 #include "DelayAction.h"
 #include "InterpolateComponentToAction.h"
 #include "Advertising.h"
+#include "Online.h"
 
 //////////////////////////////////////////////////////////////////////////
 // UKismetSystemLibrary
@@ -2293,14 +2294,27 @@ void UKismetSystemLibrary::ShowPlatformSpecificLeaderboardScreen(const FString&
 #endif
 }
 
-void UKismetSystemLibrary::ShowPlatformSpecificAchievementsScreen()
+void UKismetSystemLibrary::ShowPlatformSpecificAchievementsScreen(class APlayerController* SpecificPlayer)
 {
 #if PLATFORM_IOS
 	extern CORE_API void IOSShowAchievementsUI();
 	IOSShowAchievementsUI();
-#elif PLATFORM_ANDROID
-	extern void AndroidThunkCpp_ShowAchievements();
-	AndroidThunkCpp_ShowAchievements();
+#else
+	IOnlineExternalUIPtr ExternalUI = Online::GetExternalUIInterface();
+	if(ExternalUI.IsValid())
+	{
+		// Get the controller id from the player
+		int LocalUserNum = 0;
+		if(SpecificPlayer)
+		{
+			ULocalPlayer* LocalPlayer = Cast<ULocalPlayer>(SpecificPlayer->Player);
+			if(LocalPlayer)
+			{
+				LocalUserNum = LocalPlayer->ControllerId;
+			}
+		}
+		ExternalUI->ShowAchievementsUI(LocalUserNum);
+	}
 #endif
 }
 void UKismetSystemLibrary::SetStructurePropertyByName(UObject* Object, FName PropertyName, const FGenericStruct& Value)
diff --git a/Engine/Source/Runtime/Online/Android/OnlineSubsystemGooglePlay/Private/OnlineExternalUIInterfaceGooglePlay.cpp b/Engine/Source/Runtime/Online/Android/OnlineSubsystemGooglePlay/Private/OnlineExternalUIInterfaceGooglePlay.cpp
new file mode 100644
--- /dev/null
+++ b/Engine/Source/Runtime/Online/Android/OnlineSubsystemGooglePlay/Private/OnlineExternalUIInterfaceGooglePlay.cpp
@@ -0,0 +1,41 @@
+// Copyright 1998-2014 Epic Games, Inc. All Rights Reserved.
+
+#include "OnlineSubsystemGooglePlayPrivatePCH.h"
+
+FOnlineExternalUIGooglePlay::FOnlineExternalUIGooglePlay() 
+{
+
+}
+
+bool FOnlineExternalUIGooglePlay::ShowLoginUI(const int ControllerIndex, bool bShowOnlineOnly, const FOnLoginUIClosedDelegate& Delegate)
+{
+	return false;
+}
+
+bool FOnlineExternalUIGooglePlay::ShowFriendsUI(int32 LocalUserNum)
+{
+	return false;
+}
+
+bool FOnlineExternalUIGooglePlay::ShowInviteUI(int32 LocalUserNum) 
+{
+	return false;
+}
+
+bool FOnlineExternalUIGooglePlay::ShowAchievementsUI(int32 LocalUserNum) 
+{
+	// Will always show the achievements UI for the current local signed-in user
+	extern void AndroidThunkCpp_ShowAchievements();
+	AndroidThunkCpp_ShowAchievements();
+	return true;
+}
+
+bool FOnlineExternalUIGooglePlay::ShowWebURL(const FString& WebURL) 
+{
+	return false;
+}
+
+bool FOnlineExternalUIGooglePlay::ShowProfileUI( const FUniqueNetId& Requestor, const FUniqueNetId& Requestee, const FOnProfileUIClosedDelegate& Delegate )
+{
+	return false;
+}
diff --git a/Engine/Source/Runtime/Online/Android/OnlineSubsystemGooglePlay/Private/OnlineExternalUIInterfaceGooglePlay.h b/Engine/Source/Runtime/Online/Android/OnlineSubsystemGooglePlay/Private/OnlineExternalUIInterfaceGooglePlay.h
new file mode 100644
--- /dev/null
+++ b/Engine/Source/Runtime/Online/Android/OnlineSubsystemGooglePlay/Private/OnlineExternalUIInterfaceGooglePlay.h
@@ -0,0 +1,29 @@
+// Copyright 1998-2014 Epic Games, Inc. All Rights Reserved.
+
+#pragma once
+#include "OnlineDelegateMacros.h"
+
+#include "OnlineExternalUIInterface.h"
+#include "OnlineSubsystemGooglePlayPackage.h"
+
+/** 
+ * Interface definition for the online services external UIs
+ * Any online service that provides extra UI overlays will implement the relevant functions
+ */
+class FOnlineExternalUIGooglePlay :
+	public IOnlineExternalUI
+{
+
+public:
+	FOnlineExternalUIGooglePlay();
+
+	virtual bool ShowLoginUI(const int ControllerIndex, bool bShowOnlineOnly, const FOnLoginUIClosedDelegate& Delegate) OVERRIDE;
+	virtual bool ShowFriendsUI(int32 LocalUserNum) OVERRIDE;
+	virtual bool ShowInviteUI(int32 LocalUserNum) OVERRIDE;
+	virtual bool ShowAchievementsUI(int32 LocalUserNum) OVERRIDE;
+	virtual bool ShowWebURL(const FString& WebURL) OVERRIDE;
+	virtual bool ShowProfileUI(const FUniqueNetId& Requestor, const FUniqueNetId& Requestee, const FOnProfileUIClosedDelegate& Delegate) OVERRIDE;
+};
+
+typedef TSharedPtr<FOnlineExternalUIGooglePlay, ESPMode::ThreadSafe> FOnlineExternalUIGooglePlayPtr;
+
diff --git a/Engine/Source/Runtime/Online/Android/OnlineSubsystemGooglePlay/Private/OnlineSubsystemGooglePlay.cpp b/Engine/Source/Runtime/Online/Android/OnlineSubsystemGooglePlay/Private/OnlineSubsystemGooglePlay.cpp
--- a/Engine/Source/Runtime/Online/Android/OnlineSubsystemGooglePlay/Private/OnlineSubsystemGooglePlay.cpp
+++ b/Engine/Source/Runtime/Online/Android/OnlineSubsystemGooglePlay/Private/OnlineSubsystemGooglePlay.cpp
@@ -62,7 +62,7 @@ IOnlineVoicePtr FOnlineSubsystemGooglePlay::GetVoiceInterface() const
 
 IOnlineExternalUIPtr FOnlineSubsystemGooglePlay::GetExternalUIInterface() const
 {
-	return nullptr;
+	return ExternalUIInterface;
 }
 
 
@@ -95,6 +95,7 @@ bool FOnlineSubsystemGooglePlay::Init()
 	
 	IdentityInterface = MakeShareable(new FOnlineIdentityGooglePlay(this));
 	AchievementsInterface = MakeShareable(new FOnlineAchievementsGooglePlay(this));
+	ExternalUIInterface = MakeShareable(new FOnlineExternalUIGooglePlay());
 
 	extern void AndroidThunkCpp_GooglePlayConnect();
 	AndroidThunkCpp_GooglePlayConnect();
@@ -120,6 +121,7 @@ bool FOnlineSubsystemGooglePlay::Shutdown()
 	}
 
 	// Destruct the interfaces
+	DESTRUCT_INTERFACE(ExternalUIInterface);
 	DESTRUCT_INTERFACE(AchievementsInterface);
 	DESTRUCT_INTERFACE(IdentityInterface);
 #undef DESTRUCT_INTERFACE
diff --git a/Engine/Source/Runtime/Online/Android/OnlineSubsystemGooglePlay/Public/OnlineSubsystemGooglePlay.h b/Engine/Source/Runtime/Online/Android/OnlineSubsystemGooglePlay/Public/OnlineSubsystemGooglePlay.h
--- a/Engine/Source/Runtime/Online/Android/OnlineSubsystemGooglePlay/Public/OnlineSubsystemGooglePlay.h
+++ b/Engine/Source/Runtime/Online/Android/OnlineSubsystemGooglePlay/Public/OnlineSubsystemGooglePlay.h
@@ -5,6 +5,7 @@
 #include "OnlineSubsystem.h"
 #include "OnlineIdentityInterfaceGooglePlay.h"
 #include "OnlineAchievementsInterfaceGooglePlay.h"
+#include "OnlineExternalUIInterfaceGooglePlay.h"
 
 /**
  * OnlineSubsystemGooglePlay - Implementation of the online subsystem for Google Play services
@@ -69,6 +70,9 @@ private:
 
 	/** Interface to the online achievements */
 	FOnlineAchievementsGooglePlayPtr AchievementsInterface;
+
+	/** Interface to the external UI services */
+	FOnlineExternalUIGooglePlayPtr ExternalUIInterface;
 };
 
 typedef TSharedPtr<FOnlineSubsystemGooglePlay, ESPMode::ThreadSafe> FOnlineSubsystemGooglePlayPtr;
