Back out changelist 2044347

diff --git a/Engine/Source/Runtime/Engine/Classes/Animation/SkeletalMeshActor.h b/Engine/Source/Runtime/Engine/Classes/Animation/SkeletalMeshActor.h
--- a/Engine/Source/Runtime/Engine/Classes/Animation/SkeletalMeshActor.h
+++ b/Engine/Source/Runtime/Engine/Classes/Animation/SkeletalMeshActor.h
@@ -76,10 +76,7 @@ public:
 	// End IMatineeAnimInterface Interface
 
 private:
-	// Saved flag used to restore animation mode
 	EAnimationMode::Type	SavedAnimationMode;
-	// Saved flag used to restore mesh update mode
-	EMeshComponentUpdateFlag::Type	SavedMeshUpdateFlag;
 	// utility function to see if it can play animation or not
 	bool CanPlayAnimation();
 };
diff --git a/Engine/Source/Runtime/Engine/Private/SkeletalMesh.cpp b/Engine/Source/Runtime/Engine/Private/SkeletalMesh.cpp
--- a/Engine/Source/Runtime/Engine/Private/SkeletalMesh.cpp
+++ b/Engine/Source/Runtime/Engine/Private/SkeletalMesh.cpp
@@ -3472,8 +3472,6 @@ void ASkeletalMeshActor::PreviewBeginAnimControl(UInterpGroup* InInterpGroup)
 	{
 		SavedAnimationMode = SkeletalMeshComponent->GetAnimationMode();
 		SkeletalMeshComponent->SetAnimationMode(EAnimationMode::Type::AnimationSingleNode);
-		SavedMeshUpdateFlag = SkeletalMeshComponent->MeshComponentUpdateFlag;
-		SkeletalMeshComponent->MeshComponentUpdateFlag = EMeshComponentUpdateFlag::AlwaysTickPoseAndRefreshBones;
 	}
 }
 
@@ -3485,7 +3483,6 @@ void ASkeletalMeshActor::PreviewFinishAnimControl(UInterpGroup* InInterpGroup)
 		// if in editor, reset the Animations, makes easier for artist to see them visually and align them
 		// in game, we keep the last pose that matinee kept. If you'd like it to have animation, you'll need to have AnimTree or AnimGraph to handle correctly
 		SkeletalMeshComponent->SetAnimationMode(SavedAnimationMode);
-		SkeletalMeshComponent->MeshComponentUpdateFlag = SavedMeshUpdateFlag;
 
 		// Update space bases to reset it back to ref pose
 		SkeletalMeshComponent->RefreshBoneTransforms();
@@ -3644,8 +3641,6 @@ void ASkeletalMeshActor::BeginAnimControl(UInterpGroup* InInterpGroup)
 	{
 		SavedAnimationMode = SkeletalMeshComponent->GetAnimationMode();
 		SkeletalMeshComponent->SetAnimationMode(EAnimationMode::Type::AnimationSingleNode);
-		SavedMeshUpdateFlag = SkeletalMeshComponent->MeshComponentUpdateFlag;
-		SkeletalMeshComponent->MeshComponentUpdateFlag = EMeshComponentUpdateFlag::AlwaysTickPoseAndRefreshBones;
 	}
 }
 
@@ -3696,7 +3691,6 @@ void ASkeletalMeshActor::FinishAnimControl(UInterpGroup* InInterpGroup)
 		// if in editor, reset the Animations, makes easier for artist to see them visually and align them
 		// in game, we keep the last pose that matinee kept. If you'd like it to have animation, you'll need to have AnimTree or AnimGraph to handle correctly
 		SkeletalMeshComponent->SetAnimationMode(SavedAnimationMode);
-		SkeletalMeshComponent->MeshComponentUpdateFlag = SavedMeshUpdateFlag;
 	}
 }
 
