Fix up code for removing NaN

diff --git a/Engine/Source/Runtime/Core/Public/UObject/ObjectVersion.h b/Engine/Source/Runtime/Core/Public/UObject/ObjectVersion.h
--- a/Engine/Source/Runtime/Core/Public/UObject/ObjectVersion.h
+++ b/Engine/Source/Runtime/Core/Public/UObject/ObjectVersion.h
@@ -524,6 +524,9 @@ enum EUnrealEngineObjectUE4Version
 	VER_UE4_MEMBERREFERENCE_IN_PINTYPE,
 	// Vehicles use Nm for Torque instead of cm and RPM instead of rad/s
 	VER_UE4_VEHICLES_UNIT_CHANGE,
+	// removes NANs from all animations when loaded
+	// now importing should detect NaNs, so we should not have NaNs in source data
+	VER_UE4_ANIMATION_REMOVE_NANS, 
 	// -----<new versions can be added before this line>-------------------------------------------------
 	// - this needs to be the last line (see note below)
 	VER_UE4_AUTOMATIC_VERSION_PLUS_ONE,
diff --git a/Engine/Source/Runtime/Engine/Classes/Animation/AnimSequence.h b/Engine/Source/Runtime/Engine/Classes/Animation/AnimSequence.h
--- a/Engine/Source/Runtime/Engine/Classes/Animation/AnimSequence.h
+++ b/Engine/Source/Runtime/Engine/Classes/Animation/AnimSequence.h
@@ -706,8 +706,16 @@ private:
 	 * Remap Tracks to New Skeleton
 	 */
 	void RemapTracksToNewSkeleton( USkeleton * NewSkeleton );
+	/**
+	 * Remap NaN tracks from the RawAnimation data and recompress
+	 */	
+	void RemoveNaNTracks();
 #endif
 
+	/*
+	 * Utility function that helps to remove track, you can't just remove RawAnimationData
+	 */
+	void RemoveTrack(int32 TrackIndex);
 	friend class UAnimationAsset;
 };
 
diff --git a/Engine/Source/Runtime/Engine/Private/Animation/AnimSequence.cpp b/Engine/Source/Runtime/Engine/Private/Animation/AnimSequence.cpp
--- a/Engine/Source/Runtime/Engine/Private/Animation/AnimSequence.cpp
+++ b/Engine/Source/Runtime/Engine/Private/Animation/AnimSequence.cpp
@@ -266,6 +266,15 @@ int32 UAnimSequence::GetNumberOfTracks() const
 
 void UAnimSequence::PostLoad()
 {
+#if WITH_EDITOR
+	// I have to do this first thing in here
+	// so that remove all NaNs before even being read
+	if(GetLinkerUE4Version() < VER_UE4_ANIMATION_REMOVE_NANS)
+	{
+		RemoveNaNTracks();
+	}
+#endif // WITH_EDITOR
+
 	Super::PostLoad();
 
 	// If RAW animation data exists, and needs to be recompressed, do so.
@@ -1219,9 +1228,7 @@ bool UAnimSequence::CompressRawAnimData(float MaxPosDiff, float MaxAngleDiff)
 				if (OneKeyTransform.Equals(LocalRefPoses[BoneIndex]))
 				{
 					// remove this key
-					RawAnimationData.RemoveAt(I);
-					AnimationTrackNames.RemoveAt(I);
-					TrackToSkeletonMapTable.RemoveAt(I);
+					RemoveTrack(I);
 					--I;
 					bRemovedKeys = true;
 				}
@@ -1560,9 +1567,7 @@ void UAnimSequence::RemapTracksToNewSkeleton( USkeleton * NewSkeleton )
 			else
 			{
 				// if not found, delete the track data
-				RawAnimationData.RemoveAt(TableId);
-				AnimationTrackNames.RemoveAt(TableId);
-				TrackToSkeletonMapTable.RemoveAt(TableId);
+				RemoveTrack(TableId);
 				NewTrackToSkeletonMapTable.RemoveAt(TableId);
 				--TableId;
 			}
@@ -1614,6 +1619,60 @@ void UAnimSequence::PostProcessSequence()
 #endif
 }
 
+#if WITH_EDITOR
+void UAnimSequence::RemoveNaNTracks()
+{
+	for( int32 TrackIndex=0; TrackIndex<RawAnimationData.Num(); ++TrackIndex )
+	{
+		const FRawAnimSequenceTrack & RawTrack = RawAnimationData[TrackIndex];
+
+		bool bContainsNaN = false;
+		for ( auto Key : RawTrack.PosKeys )
+		{
+			bContainsNaN |= Key.ContainsNaN();
+		}
+
+		if (!bContainsNaN)
+		{
+			for(auto Key : RawTrack.RotKeys)
+			{
+				bContainsNaN |= Key.ContainsNaN();
+			}
+		}
+
+		if (!bContainsNaN)
+		{
+			for(auto Key : RawTrack.ScaleKeys)
+			{
+				bContainsNaN |= Key.ContainsNaN();
+			}
+		}
+
+		if (bContainsNaN)
+		{
+			UE_LOG(LogAnimation, Warning, TEXT("Animation raw data contains NaNs - Removing the following track [%s Track (%s)]"), (GetOuter() ? *GetOuter()->GetFullName() : *GetFullName()), *AnimationTrackNames[TrackIndex].ToString());
+			// remove this track
+			RemoveTrack(TrackIndex);
+			--TrackIndex;
+		}
+	}
+
+	FAnimationUtils::CompressAnimSequence(this, false, false);
+}
+#endif // WITH_EDITOR
+
+void UAnimSequence::RemoveTrack(int32 TrackIndex)
+{
+	if (RawAnimationData.IsValidIndex(TrackIndex))
+	{
+		RawAnimationData.RemoveAt(TrackIndex);
+		AnimationTrackNames.RemoveAt(TrackIndex);
+		TrackToSkeletonMapTable.RemoveAt(TrackIndex);
+
+		check (RawAnimationData.Num() == AnimationTrackNames.Num() && AnimationTrackNames.Num() == TrackToSkeletonMapTable.Num() );
+	}
+}
+
 bool UAnimSequence::GetAllAnimationSequencesReferred(TArray<UAnimSequence*>& AnimationSequences)
 {
 	AnimationSequences.Add(this);
