Enabled caching for platform data in postload.
Fix for ddc commandlet.

diff --git a/Engine/Source/Runtime/Engine/Private/Materials/Material.cpp b/Engine/Source/Runtime/Engine/Private/Materials/Material.cpp
--- a/Engine/Source/Runtime/Engine/Private/Materials/Material.cpp
+++ b/Engine/Source/Runtime/Engine/Private/Materials/Material.cpp
@@ -1598,38 +1598,7 @@ void UMaterial::Serialize(FArchive& Ar)
 
 	if (Ar.UE4Ver() >= VER_UE4_PURGED_FMATERIAL_COMPILE_OUTPUTS)
 	{
-
 		SerializeInlineShaderMaps( CachedMaterialResourcesForCooking, Ar, MaterialResources );
-
-
-
-		/*if ( Ar.IsLoading() || (Ar.IsCooking() == false) )
-		{
-			// shouldn't ever load into this saved material resources array
-			// if we are not cooking then SerializeInlineShaderMaps will ignore the SavedMaterialResources array
-			TArray<FMaterialResource*> SavedMaterialResources;
-			SerializeInlineShaderMaps( SavedMaterialResources, Ar, MaterialResources );
-			check( SavedMaterialResources.Num() == 0);
-		}
-		else if ( Ar.IsSaving() )
-		{
-			check( Ar.IsCooking() );
-			const TArray<FMaterialResource*> *CachedMaterialResourcesForPlatform = CachedMaterialResourcesForCooking.Find(Ar.CookingTarget());
-			// check( CachedMaterialResourcesForPlatform != NULL || (!HasAnyMarks( OBJECTMARK_TagExp ))  );
-			if ( CachedMaterialResourcesForPlatform != NULL )
-			{
-				SerializeInlineShaderMaps(*CachedMaterialResourcesForPlatform, Ar, MaterialResources);
-			}
-			else
-			{
-				// if we are saving to a linker and we are cooking then we should have definatly have had a cached material for the target platform (see UPackage::SavePackage() -> UObject::BeginCacheResourceForPlatform
-				//  if this check is triggering most likely because some object didn't have BeginCacheResourceForPlatform called from SavePackage OR
-				//  the archive isn't actually saving and GetLinker is returning a valid value for some other reason
-				check( Ar.GetLinker() == NULL );
-			}
-			
-		}*/
-		
 	}
 	else
 	{
@@ -1983,6 +1952,19 @@ void UMaterial::PostLoad()
 	{
 		SCOPE_SECONDS_COUNTER(MaterialLoadTime);
 
+		// enable caching in postload for derived data cache commandlet and cook by the book
+		ITargetPlatformManagerModule* TPM = GetTargetPlatformManager();
+		if (TPM && (TPM->RestrictFormatsToRuntimeOnly() == false))
+		{
+			ITargetPlatformManagerModule* TPM = GetTargetPlatformManager();
+			TArray<ITargetPlatform*> Platforms = TPM->GetActiveTargetPlatforms();
+			// Cache for all the shader formats that the cooking target requires
+			for (int32 FormatIndex = 0; FormatIndex < Platforms.Num(); FormatIndex++)
+			{
+				BeginCacheForCookedPlatformData(Platforms[FormatIndex]);
+			}
+		}
+
 		//Don't compile shaders in post load for dev overhead materials.
 		if (FApp::CanEverRender() && !bIsMaterialEditorStatsMaterial)
 		{
@@ -2022,8 +2004,6 @@ void UMaterial::BeginCacheForCookedPlatformData( const ITargetPlatform *TargetPl
 
 	if ( CachedMaterialResourcesForPlatform == NULL )
 	{
-		check( CachedMaterialResourcesForPlatform == NULL );
-
 		CachedMaterialResourcesForCooking.Add( TargetPlatform );
 		CachedMaterialResourcesForPlatform = CachedMaterialResourcesForCooking.Find( TargetPlatform );
 
diff --git a/Engine/Source/Runtime/Engine/Private/Materials/MaterialInstance.cpp b/Engine/Source/Runtime/Engine/Private/Materials/MaterialInstance.cpp
--- a/Engine/Source/Runtime/Engine/Private/Materials/MaterialInstance.cpp
+++ b/Engine/Source/Runtime/Engine/Private/Materials/MaterialInstance.cpp
@@ -1169,8 +1169,6 @@ void UMaterialInstance::InitStaticPermutation()
 	// Allocate material resources if needed even if we are cooking, so that StaticPermutationMaterialResources will always be valid
 	UpdatePermutationAllocations();
 
-	const TArray<FName>& DesiredShaderFormats = GetTargetShaderFormats();
-
 	if ( FApp::CanEverRender() ) 
 	{
 		// Cache shaders for the current platform to be used for rendering
@@ -1656,6 +1654,19 @@ void UMaterialInstance::PostLoad()
 
 		// Make sure static parameters are up to date and shaders are cached for the current platform
 		InitStaticPermutation();
+
+		// enable caching in postload for derived data cache commandlet and cook by the book
+		ITargetPlatformManagerModule* TPM = GetTargetPlatformManager();
+		if (TPM && (TPM->RestrictFormatsToRuntimeOnly() == false)) 
+		{
+			ITargetPlatformManagerModule* TPM = GetTargetPlatformManager();
+			TArray<ITargetPlatform*> Platforms = TPM->GetActiveTargetPlatforms();
+			// Cache for all the shader formats that the cooking target requires
+			for (int32 FormatIndex = 0; FormatIndex < Platforms.Num(); FormatIndex++)
+			{
+				BeginCacheForCookedPlatformData(Platforms[FormatIndex]);
+			}
+		}
 	}
 	INC_FLOAT_STAT_BY(STAT_ShaderCompiling_MaterialLoading,(float)MaterialLoadTime);
 
diff --git a/Engine/Source/Runtime/Engine/Private/TextureDerivedData.cpp b/Engine/Source/Runtime/Engine/Private/TextureDerivedData.cpp
--- a/Engine/Source/Runtime/Engine/Private/TextureDerivedData.cpp
+++ b/Engine/Source/Runtime/Engine/Private/TextureDerivedData.cpp
@@ -1376,6 +1376,19 @@ void UTexture::UpdateCachedLODBias()
 void UTexture::BeginCachePlatformData()
 {
 	CachePlatformData(true);
+
+	// enable caching in postload for derived data cache commandlet and cook by the book
+	ITargetPlatformManagerModule* TPM = GetTargetPlatformManager();
+	if (TPM && (TPM->RestrictFormatsToRuntimeOnly() == false))
+	{
+		ITargetPlatformManagerModule* TPM = GetTargetPlatformManager();
+		TArray<ITargetPlatform*> Platforms = TPM->GetActiveTargetPlatforms();
+		// Cache for all the shader formats that the cooking target requires
+		for (int32 FormatIndex = 0; FormatIndex < Platforms.Num(); FormatIndex++)
+		{
+			BeginCacheForCookedPlatformData(Platforms[FormatIndex]);
+		}
+	}
 }
 
 void UTexture::BeginCacheForCookedPlatformData( const ITargetPlatform *TargetPlatform )
