Had a request from a licensee to make PointLight, DirectionalLight, and SpotLight able to be extended by users by adding "ENGINE_API" to those classes.

diff --git a/Engine/Source/Runtime/Engine/Classes/Engine/DirectionalLight.h b/Engine/Source/Runtime/Engine/Classes/Engine/DirectionalLight.h
--- a/Engine/Source/Runtime/Engine/Classes/Engine/DirectionalLight.h
+++ b/Engine/Source/Runtime/Engine/Classes/Engine/DirectionalLight.h
@@ -4,8 +4,8 @@
 
 #include "DirectionalLight.generated.h"
 
-UCLASS(HeaderGroup=Light, ClassGroup=(Lights, DirectionalLights), MinimalAPI, meta=(ChildCanTick))
-class ADirectionalLight : public ALight
+UCLASS(HeaderGroup=Light, ClassGroup=(Lights, DirectionalLights), meta=(ChildCanTick))
+class ENGINE_API ADirectionalLight : public ALight
 {
 	GENERATED_UCLASS_BODY()
 
diff --git a/Engine/Source/Runtime/Engine/Classes/Engine/PointLight.h b/Engine/Source/Runtime/Engine/Classes/Engine/PointLight.h
--- a/Engine/Source/Runtime/Engine/Classes/Engine/PointLight.h
+++ b/Engine/Source/Runtime/Engine/Classes/Engine/PointLight.h
@@ -4,8 +4,8 @@
 
 #include "PointLight.generated.h"
 
-UCLASS(HeaderGroup=Light, ClassGroup=(Lights, PointLights),MinimalAPI, meta=(ChildCanTick))
-class APointLight : public ALight
+UCLASS(HeaderGroup=Light, ClassGroup=(Lights, PointLights), meta=(ChildCanTick))
+class ENGINE_API APointLight : public ALight
 {
 	GENERATED_UCLASS_BODY()
 
@@ -14,9 +14,9 @@ class APointLight : public ALight
 
 	// BEGIN DEPRECATED (use component functions now in level script)
 	UFUNCTION(BlueprintCallable, Category="Rendering|Lighting", meta=(DeprecatedFunction))
-	ENGINE_API void SetRadius(float NewRadius);
+	void SetRadius(float NewRadius);
 	UFUNCTION(BlueprintCallable, Category="Rendering|Lighting", meta=(DeprecatedFunction))
-	ENGINE_API void SetLightFalloffExponent(float NewLightFalloffExponent);
+	void SetLightFalloffExponent(float NewLightFalloffExponent);
 	// END DEPRECATED
 
 #if WITH_EDITOR
diff --git a/Engine/Source/Runtime/Engine/Classes/Engine/SpotLight.h b/Engine/Source/Runtime/Engine/Classes/Engine/SpotLight.h
--- a/Engine/Source/Runtime/Engine/Classes/Engine/SpotLight.h
+++ b/Engine/Source/Runtime/Engine/Classes/Engine/SpotLight.h
@@ -5,8 +5,8 @@
 
 #include "SpotLight.generated.h"
 
-UCLASS(HeaderGroup=Light, ClassGroup=(Lights, SpotLights), MinimalAPI, meta=(ChildCanTick))
-class ASpotLight : public ALight
+UCLASS(HeaderGroup=Light, ClassGroup=(Lights, SpotLights), meta=(ChildCanTick))
+class ENGINE_API ASpotLight : public ALight
 {
 	GENERATED_UCLASS_BODY()
 
