Added additional cloth visualization menus

diff --git a/Engine/Source/Editor/Persona/Private/AnimViewportShowCommands.cpp b/Engine/Source/Editor/Persona/Private/AnimViewportShowCommands.cpp
--- a/Engine/Source/Editor/Persona/Private/AnimViewportShowCommands.cpp
+++ b/Engine/Source/Editor/Persona/Private/AnimViewportShowCommands.cpp
@@ -32,7 +32,7 @@ void FAnimViewportShowCommands::RegisterCommands()
 	UI_COMMAND( ShowLocalAxesAll, "All", "Shows all local hierarchy axes", EUserInterfaceActionType::RadioButton, FInputGesture() );
 
 #if WITH_APEX_CLOTHING
-	UI_COMMAND( DisableClothSimulation, "Disable Clothing", "Disable Cloth Simulation", EUserInterfaceActionType::ToggleButton, FInputGesture() );
+	UI_COMMAND( DisableClothSimulation, "Disable Cloth Simulation", "Disable Cloth Simulation and Show non-simulated mesh", EUserInterfaceActionType::ToggleButton, FInputGesture() );
 	UI_COMMAND( ApplyClothWind, "Apply Wind", "Apply Wind to Clothing", EUserInterfaceActionType::ToggleButton, FInputGesture() );
 
 	UI_COMMAND( ShowClothSimulationNormals, "Simulation Normals", "Display Cloth Simulation Normals", EUserInterfaceActionType::ToggleButton, FInputGesture() );
@@ -40,6 +40,15 @@ void FAnimViewportShowCommands::RegisterCommands()
 	UI_COMMAND( ShowClothCollisionVolumes, "Collision Volumes", "Display Cloth Collision Volumes", EUserInterfaceActionType::ToggleButton, FInputGesture() );
 	UI_COMMAND( EnableCollisionWithAttachedClothChildren, "Collide with Cloth Children", "Enables collision detection between collision primitives in the base mesh and clothing on any attachments in the preview scene.", EUserInterfaceActionType::ToggleButton, FInputGesture() );	
 
-	UI_COMMAND( ShowOnlyClothSections, "Only Cloth Sections", "Display Only Cloth Mapped Sections", EUserInterfaceActionType::ToggleButton, FInputGesture() );
+	UI_COMMAND( ShowClothPhysicalMeshWire, "Physical Mesh Wire", "Display Simulation Mesh's wire frame. Shows all physical mesh loaded from an asset file when cloth simulation is disabled and shows simulated results when simulating. The part having zero max distance value will be drawn in pink.", EUserInterfaceActionType::ToggleButton, FInputGesture() );
+	UI_COMMAND( ShowClothMaxDistances, "Max Distances", "Display Max Distances only when simulation is disabled. If turned on this option, disables cloth simulation automatically", EUserInterfaceActionType::ToggleButton, FInputGesture() );
+	UI_COMMAND( ShowClothBackstop, "Back stops", "Display Back stops. If it has a proper value, draws a line in red. Otherwise, draws a vertex in white.", EUserInterfaceActionType::ToggleButton, FInputGesture() );
+	UI_COMMAND( ShowClothFixedVertices, "Fixed Vertices", "Display Fixed Vertices", EUserInterfaceActionType::ToggleButton, FInputGesture() );	
+
+	// below 3 menus are radio button styles
+	UI_COMMAND(ShowAllSections, "Show All Sections", "Display All sections including Cloth Mapped Sections", EUserInterfaceActionType::RadioButton, FInputGesture());
+	UI_COMMAND(ShowOnlyClothSections, "Show Only Cloth Sections", "Display Only Cloth Mapped Sections", EUserInterfaceActionType::RadioButton, FInputGesture());
+	UI_COMMAND(HideOnlyClothSections, "Hide Only Cloth Sections", "Display All Except Cloth Mapped Sections", EUserInterfaceActionType::RadioButton, FInputGesture());
+
 #endif// #if WITH_APEX_CLOTHING
 }
diff --git a/Engine/Source/Editor/Persona/Private/AnimViewportShowCommands.h b/Engine/Source/Editor/Persona/Private/AnimViewportShowCommands.h
--- a/Engine/Source/Editor/Persona/Private/AnimViewportShowCommands.h
+++ b/Engine/Source/Editor/Persona/Private/AnimViewportShowCommands.h
@@ -99,8 +99,22 @@ public:
 	*/
 	TSharedPtr< FUICommandInfo > EnableCollisionWithAttachedClothChildren;
 
+	/** Show all sections which means the original state */
+	TSharedPtr< FUICommandInfo > ShowAllSections;
 	/** Show only clothing mapped sections */
 	TSharedPtr< FUICommandInfo > ShowOnlyClothSections;
+	/** Show all except clothing mapped sections */
+	TSharedPtr< FUICommandInfo > HideOnlyClothSections;
+
+	/** Show max distances of simulation vertices */
+	TSharedPtr< FUICommandInfo > ShowClothMaxDistances;
+
+	TSharedPtr< FUICommandInfo > ShowClothBackstop;
+
+	/** Show simulation mesh's wire frame */
+	TSharedPtr< FUICommandInfo > ShowClothPhysicalMeshWire;
+
+	TSharedPtr< FUICommandInfo > ShowClothFixedVertices;	
 #endif// #if WITH_APEX_CLOTHING
 
 public:
diff --git a/Engine/Source/Editor/Persona/Private/AnimationEditorViewportClient.cpp b/Engine/Source/Editor/Persona/Private/AnimationEditorViewportClient.cpp
--- a/Engine/Source/Editor/Persona/Private/AnimationEditorViewportClient.cpp
+++ b/Engine/Source/Editor/Persona/Private/AnimationEditorViewportClient.cpp
@@ -171,7 +171,7 @@ FAnimationViewportClient::FAnimationViewportClient( FPreviewScene& InPreviewScen
 
 	//wind actor's initial position, rotation and strength
 	PrevWindLocation = FVector(100, 100, 100);
-	PrevWindRotation = FRotator(0,0,0);
+	PrevWindRotation = FRotator(0,0,0); // roll, yaw, pitch
 	PrevWindStrength = 0.2f;
 
 	bDrawUVs = false;
@@ -408,6 +408,29 @@ void FAnimationViewportClient::Draw(const FSceneView* View, FPrimitiveDrawInterf
 			PreviewSkelMeshComp->DrawClothingCollisionVolumes(PDI);
 		}
 
+		// Display collision volumes of current selected cloth
+		if ( PreviewSkelMeshComp->bDisplayClothPhysicalMeshWire )
+		{
+			PreviewSkelMeshComp->DrawClothingPhysicalMeshWire(PDI);
+		}
+
+		// Display collision volumes of current selected cloth
+		if ( PreviewSkelMeshComp->bDisplayClothMaxDistances )
+		{
+			PreviewSkelMeshComp->DrawClothingMaxDistances(PDI);
+		}
+
+		// Display collision volumes of current selected cloth
+		if ( PreviewSkelMeshComp->bDisplayClothBackstops )
+		{
+			PreviewSkelMeshComp->DrawClothingBackstops(PDI);
+		}
+
+		if( PreviewSkelMeshComp->bDisplayClothFixedVertices )
+		{
+			PreviewSkelMeshComp->DrawClothingFixedVertices(PDI);
+		}
+		
 		// Display socket hit points
 		if ( PreviewSkelMeshComp->bDrawSockets )
 		{
@@ -823,7 +846,7 @@ void FAnimationViewportClient::DisplayInfo(FCanvas* Canvas, FSceneView* View, bo
 			FMath::Round(PreviewSkelMeshComp->Bounds.BoxExtent.X * 2.0f),
 			FMath::Round(PreviewSkelMeshComp->Bounds.BoxExtent.Y * 2.0f),
 			FMath::Round(PreviewSkelMeshComp->Bounds.BoxExtent.Z * 2.0f));
-		Canvas->DrawShadowedString(CurXOffset, CurYOffset, *InfoString, GEngine->GetSmallFont(), TextColor);
+		Canvas->DrawShadowedString( CurXOffset, CurYOffset, *InfoString, GEngine->GetSmallFont(), TextColor );
 
 		uint32 NumNotiesWithErrors = PreviewSkelMeshComp->AnimNotifyErrors.Num();
 		for (uint32 i = 0; i < NumNotiesWithErrors; ++i)
@@ -965,16 +988,16 @@ bool FAnimationViewportClient::InputWidgetDelta( FViewport* Viewport, EAxisList:
 			{
 				FVector RotAxis;
 				float RotAngle;
-				Rot.Quaternion().ToAxisAndAngle(RotAxis, RotAngle);
+				Rot.Quaternion().ToAxisAndAngle( RotAxis, RotAngle );
 
-				FVector4 BoneSpaceAxis = BaseTM.TransformVector(RotAxis);
+				FVector4 BoneSpaceAxis = BaseTM.TransformVector( RotAxis );
 
 				//Calculate the new delta rotation
-				FQuat DeltaQuat(BoneSpaceAxis, RotAngle);
+				FQuat DeltaQuat( BoneSpaceAxis, RotAngle );
 
-				FRotator NewRotation = (CurrentSkelControlTM * FTransform(DeltaQuat)).Rotator();
+				FRotator NewRotation = ( CurrentSkelControlTM * FTransform( DeltaQuat ) ).Rotator();
 
-				if (SelectedSocket)
+				if ( SelectedSocket )
 				{
 					SelectedSocket->RelativeRotation = NewRotation;
 				}
@@ -998,8 +1021,8 @@ bool FAnimationViewportClient::InputWidgetDelta( FViewport* Viewport, EAxisList:
 				}
 			}
 
-			//@TODO: ANIMATION: Add scaling support here
-		}
+				//@TODO: ANIMATION: Add scaling support here
+			}
 		else if( WindActor.IsValid() )
 		{
 			if (WidgetMode == FWidget::WM_Rotate)
@@ -1039,34 +1062,34 @@ void FAnimationViewportClient::TrackingStarted( const struct FInputEventState& I
 		if(InInputState.IsLeftMouseButtonPressed() && (Widget->GetCurrentAxis() & EAxisList::XYZ) != 0)
 		{
 			if ( PreviewSkelMeshComp->SocketsOfInterest.Num() == 1 )
+		{
+			const bool bAltDown = InInputState.IsAltButtonPressed();
+
+			if ( bAltDown && PersonaPtr.IsValid() )
 			{
-				const bool bAltDown = InInputState.IsAltButtonPressed();
+				// Rather than moving/rotating the selected socket, copy it and move the copy instead
+				PersonaPtr.Pin()->DuplicateAndSelectSocket( PreviewSkelMeshComp->SocketsOfInterest[0] );
+			}
+
+			// Socket movement is transactional - we want undo/redo and saving of it
+			USkeletalMeshSocket* Socket = PreviewSkelMeshComp->SocketsOfInterest[0].Socket;
 
-				if ( bAltDown && PersonaPtr.IsValid() )
+			if ( Socket && bInTransaction == false )
+			{
+				if (WidgetMode == FWidget::WM_Rotate )
 				{
-					// Rather than moving/rotating the selected socket, copy it and move the copy instead
-					PersonaPtr.Pin()->DuplicateAndSelectSocket( PreviewSkelMeshComp->SocketsOfInterest[0] );
+					GEditor->BeginTransaction( LOCTEXT("AnimationEditorViewport_RotateSocket", "Rotate Socket" ) );
 				}
-
-				// Socket movement is transactional - we want undo/redo and saving of it
-				USkeletalMeshSocket* Socket = PreviewSkelMeshComp->SocketsOfInterest[0].Socket;
-
-				if ( Socket && bInTransaction == false )
+				else
 				{
-					if (WidgetMode == FWidget::WM_Rotate )
-					{
-						GEditor->BeginTransaction( LOCTEXT("AnimationEditorViewport_RotateSocket", "Rotate Socket" ) );
-					}
-					else
-					{
-						GEditor->BeginTransaction( LOCTEXT("AnimationEditorViewport_TranslateSocket", "Translate Socket" ) );
-					}
-
-					Socket->SetFlags( RF_Transactional );	// Undo doesn't work without this!
-					Socket->Modify();
-					bInTransaction = true;
+					GEditor->BeginTransaction( LOCTEXT("AnimationEditorViewport_TranslateSocket", "Translate Socket" ) );
 				}
+
+				Socket->SetFlags( RF_Transactional );	// Undo doesn't work without this!
+				Socket->Modify();
+				bInTransaction = true;
 			}
+		}
 			else if( BoneIndex >= 0 )
 			{
 				if ( bInTransaction == false )
@@ -1118,14 +1141,14 @@ FWidget::EWidgetMode FAnimationViewportClient::GetWidgetMode() const
 	FWidget::EWidgetMode Mode = FWidget::WM_None;
 	if (!PreviewSkelMeshComp->IsAnimBlueprintInstanced())
 	{
-		if ((PreviewSkelMeshComp != NULL) && ((PreviewSkelMeshComp->BonesOfInterest.Num() == 1) || (PreviewSkelMeshComp->SocketsOfInterest.Num() == 1)))
-		{
+	if ((PreviewSkelMeshComp != NULL) && ((PreviewSkelMeshComp->BonesOfInterest.Num() == 1) || (PreviewSkelMeshComp->SocketsOfInterest.Num() == 1)))
+	{
 			Mode = WidgetMode;
-		}
-		else if (SelectedWindActor.IsValid())
-		{
+	}
+	else if (SelectedWindActor.IsValid())
+	{
 			Mode = WidgetMode;
-		}
+	}
 	}
 
 	return Mode;
diff --git a/Engine/Source/Editor/Persona/Private/Persona.cpp b/Engine/Source/Editor/Persona/Private/Persona.cpp
--- a/Engine/Source/Editor/Persona/Private/Persona.cpp
+++ b/Engine/Source/Editor/Persona/Private/Persona.cpp
@@ -454,6 +454,9 @@ FPersona::~FPersona()
 	if(PreviewComponent)
 	{
 		PreviewComponent->RemoveFromRoot();
+#if WITH_APEX_CLOTHING
+		PreviewComponent->RestoreClothSectionsVisibility();
+#endif //#if WITH_APEX_CLOTHING
 	}
 	
 	// NOTE: Any tabs that we still have hanging out when destroyed will be cleaned up by FBaseToolkit's destructor
diff --git a/Engine/Source/Editor/Persona/Private/SAnimViewportToolBar.cpp b/Engine/Source/Editor/Persona/Private/SAnimViewportToolBar.cpp
--- a/Engine/Source/Editor/Persona/Private/SAnimViewportToolBar.cpp
+++ b/Engine/Source/Editor/Persona/Private/SAnimViewportToolBar.cpp
@@ -157,7 +157,7 @@ public:
 				.SliderValue(AnimViewportPtr.Pin().ToSharedRef(), &SAnimationEditorViewportTabBody::GetWindStrengthSliderValue)
 				.OnSliderValueChanged(AnimViewportPtr.Pin().ToSharedRef(), &SAnimationEditorViewportTabBody::SetWindStrength)
 				.SliderTooltip( LOCTEXT("WindStrength_ToolTip", "Change wind strength") )
-				.OnPlusClicked( this, &SClothWindSettings::OnIncreasetWindStrength )
+				.OnPlusClicked( this, &SClothWindSettings::OnIncreaseWindStrength )
 				.PlusTooltip( LOCTEXT("IncreasetWindStrength_ToolTip", "Increase Wind Strength") )
 				.ExtraWidget( ExtraWidget )
 			];
@@ -173,7 +173,7 @@ protected:
 	}
 
 	/** Callback function for increasing size */
-	FReply OnIncreasetWindStrength( )
+	FReply OnIncreaseWindStrength()
 	{
 		const float DeltaValue = 0.1f;
 		AnimViewportPtr.Pin()->SetWindStrength( AnimViewportPtr.Pin()->GetWindStrengthSliderValue() + DeltaValue );
@@ -422,90 +422,80 @@ TSharedRef<SWidget> SAnimViewportToolBar::GenerateShowMenu() const
 	const bool bInShouldCloseWindowAfterMenuSelection = true;
 	FMenuBuilder ShowMenuBuilder(bInShouldCloseWindowAfterMenuSelection, Viewport.Pin()->GetCommandList());
 	{
-		ShowMenuBuilder.BeginSection("AnimViewportFOV", LOCTEXT("Viewport_FOVLabel", "Field Of View"));
-		{
-			const float FOVMin = 5.f;
-			const float FOVMax = 170.f;
-
-			TSharedPtr<SWidget> FOVWidget = SNew(SSpinBox<float>)
-				.Font(FEditorStyle::GetFontStyle(TEXT("MenuItem.Font")))
-				.MinValue(FOVMin)
-				.MaxValue(FOVMax)
-				.Value(this, &SAnimViewportToolBar::OnGetFOVValue)
-				.OnValueChanged(this, &SAnimViewportToolBar::OnFOVValueChanged)
-				.OnValueCommitted(this, &SAnimViewportToolBar::OnFOVValueCommitted);
-
-			ShowMenuBuilder.AddWidget(FOVWidget.ToSharedRef(), FText());
-		}
-		ShowMenuBuilder.EndSection();
+			ShowMenuBuilder.BeginSection("AnimViewportFOV", LOCTEXT("Viewport_FOVLabel", "Field Of View"));
+			{
+				const float FOVMin = 5.f;
+				const float FOVMax = 170.f;
+
+				TSharedPtr<SWidget> FOVWidget = SNew(SSpinBox<float>)
+					.Font(FEditorStyle::GetFontStyle(TEXT("MenuItem.Font")))
+					.MinValue(FOVMin)
+					.MaxValue(FOVMax)
+					.Value(this, &SAnimViewportToolBar::OnGetFOVValue)
+					.OnValueChanged(this, &SAnimViewportToolBar::OnFOVValueChanged)
+					.OnValueCommitted(this, &SAnimViewportToolBar::OnFOVValueCommitted);
+
+				ShowMenuBuilder.AddWidget(FOVWidget.ToSharedRef(), FText());
+			}
+			ShowMenuBuilder.EndSection();
 
-		ShowMenuBuilder.BeginSection("AnimViewportAudio", LOCTEXT("Viewport_AudioLabel", "Audio"));
-		{
-			ShowMenuBuilder.AddMenuEntry(Actions.MuteAudio);
-		}
-		ShowMenuBuilder.EndSection();
+			ShowMenuBuilder.BeginSection("AnimViewportAudio", LOCTEXT("Viewport_AudioLabel", "Audio"));
+			{
+				ShowMenuBuilder.AddMenuEntry(Actions.MuteAudio);
+			}
+			ShowMenuBuilder.EndSection();
 
-		ShowMenuBuilder.BeginSection("AnimViewportMesh", LOCTEXT("ShowMenu_Actions_Mesh", "Mesh"));
-		{
-			ShowMenuBuilder.AddMenuEntry( Actions.ShowReferencePose );
-			ShowMenuBuilder.AddMenuEntry( Actions.ShowBound );
-			ShowMenuBuilder.AddMenuEntry( Actions.ShowPreviewMesh );
-		}
-		ShowMenuBuilder.EndSection();
+			ShowMenuBuilder.BeginSection("AnimViewportMesh", LOCTEXT("ShowMenu_Actions_Mesh", "Mesh"));
+			{
+				ShowMenuBuilder.AddMenuEntry( Actions.ShowReferencePose );
+				ShowMenuBuilder.AddMenuEntry( Actions.ShowBound );
+				ShowMenuBuilder.AddMenuEntry( Actions.ShowPreviewMesh );
+			}
+			ShowMenuBuilder.EndSection();
 
-		ShowMenuBuilder.BeginSection("AnimViewportAnimation", LOCTEXT("ShowMenu_Actions_Asset", "Asset"));
-		{
-			ShowMenuBuilder.AddMenuEntry( Actions.ShowRawAnimation );
-			ShowMenuBuilder.AddMenuEntry( Actions.ShowNonRetargetedAnimation );
-			ShowMenuBuilder.AddMenuEntry( Actions.ShowAdditiveBaseBones );
-		}
-		ShowMenuBuilder.EndSection();
+			ShowMenuBuilder.BeginSection("AnimViewportAnimation", LOCTEXT("ShowMenu_Actions_Asset", "Asset"));
+			{
+				ShowMenuBuilder.AddMenuEntry( Actions.ShowRawAnimation );
+				ShowMenuBuilder.AddMenuEntry( Actions.ShowNonRetargetedAnimation );
+				ShowMenuBuilder.AddMenuEntry( Actions.ShowAdditiveBaseBones );
+			}
+			ShowMenuBuilder.EndSection();
 
-		ShowMenuBuilder.BeginSection("AnimViewportPreviewBones", LOCTEXT("ShowMenu_Actions_Bones", "Hierarchy") );
-		{
-			ShowMenuBuilder.AddMenuEntry( Actions.ShowSockets );
-			ShowMenuBuilder.AddMenuEntry( Actions.ShowBones );
-			ShowMenuBuilder.AddMenuEntry( Actions.ShowBoneNames );
-			ShowMenuBuilder.AddMenuEntry( Actions.ShowBoneWeight );
-		}
-		ShowMenuBuilder.EndSection();
+			ShowMenuBuilder.BeginSection("AnimViewportPreviewBones", LOCTEXT("ShowMenu_Actions_Bones", "Hierarchy") );
+			{
+				ShowMenuBuilder.AddMenuEntry( Actions.ShowSockets );
+				ShowMenuBuilder.AddMenuEntry( Actions.ShowBones );
+				ShowMenuBuilder.AddMenuEntry( Actions.ShowBoneNames );
+				ShowMenuBuilder.AddMenuEntry( Actions.ShowBoneWeight );
+			}
+			ShowMenuBuilder.EndSection();
 
-		ShowMenuBuilder.BeginSection("AnimviewportInfo", LOCTEXT("ShowInfo_Actions_Info", "Info") );
-		{
-			ShowMenuBuilder.AddMenuEntry( Actions.ShowDisplayInfo );
-		}
-		ShowMenuBuilder.EndSection();
+			ShowMenuBuilder.BeginSection("AnimviewportInfo", LOCTEXT("ShowInfo_Actions_Info", "Info") );
+			{
+				ShowMenuBuilder.AddMenuEntry( Actions.ShowDisplayInfo );
+			}
+			ShowMenuBuilder.EndSection();
 
 #if WITH_APEX_CLOTHING
-		UDebugSkelMeshComponent* PreviewComp = Viewport.Pin()->GetPersona().Pin()->PreviewComponent;
+			UDebugSkelMeshComponent* PreviewComp = Viewport.Pin()->GetPersona().Pin()->PreviewComponent;
 
-		if( PreviewComp && PreviewComp->HasValidClothingActors() )
-		{
-			ShowMenuBuilder.BeginSection("AnimViewportClothingOptions", LOCTEXT("ShowMenu_Actions_Clothing", "Clothing") );
+			if( PreviewComp && PreviewComp->HasValidClothingActors() )
 			{
-				ShowMenuBuilder.AddMenuEntry( Actions.DisableClothSimulation );
-				ShowMenuBuilder.AddMenuEntry( Actions.ApplyClothWind );
-				TSharedPtr<SWidget> WindWidget = SNew( SClothWindSettings ).AnimEditorViewport( Viewport );
-				ShowMenuBuilder.AddWidget( WindWidget.ToSharedRef(),  FText() );
-				ShowMenuBuilder.AddMenuEntry( Actions.ShowClothSimulationNormals );
-				ShowMenuBuilder.AddMenuEntry( Actions.ShowClothGraphicalTangents );
-				ShowMenuBuilder.AddMenuEntry( Actions.ShowClothCollisionVolumes );
-				TSharedPtr<SWidget> GravityWidget = SNew( SGravitySettings ).AnimEditorViewport( Viewport );
-				ShowMenuBuilder.AddWidget( GravityWidget.ToSharedRef(),  FText() );
-				ShowMenuBuilder.AddMenuEntry( Actions.EnableCollisionWithAttachedClothChildren );
-				ShowMenuBuilder.AddMenuEntry( Actions.ShowOnlyClothSections );
+				ShowMenuBuilder.AddMenuSeparator();
+				ShowMenuBuilder.AddSubMenu(
+					LOCTEXT("AnimViewportClothingSubMenu", "Clothing"),
+					LOCTEXT("AnimViewportClothingSubMenuToolTip", "Options relating to clothing"),
+					FNewMenuDelegate::CreateRaw(this, &SAnimViewportToolBar::FillShowClothingMenu));
 			}
-			ShowMenuBuilder.EndSection();
-		}
-		else // if skeletal mesh has clothing assets without mapping yet or assets have only collision volumes without clothing sections, then need to show collision volumes which assets include 
-		if( PreviewComp && PreviewComp->SkeletalMesh && PreviewComp->SkeletalMesh->ClothingAssets.Num() > 0)
-		{				
-			ShowMenuBuilder.BeginSection("AnimViewportClothingOptions", LOCTEXT("ShowMenu_Actions_Clothing", "Clothing") );
-			{
-				ShowMenuBuilder.AddMenuEntry( Actions.ShowClothCollisionVolumes );
+			else // if skeletal mesh has clothing assets without mapping yet or assets have only collision volumes without clothing sections, then need to show collision volumes which assets include 
+			if( PreviewComp && PreviewComp->SkeletalMesh && PreviewComp->SkeletalMesh->ClothingAssets.Num() > 0)
+			{				
+				ShowMenuBuilder.BeginSection("AnimViewportClothingOptions", LOCTEXT("ShowMenu_Actions_Clothing", "Clothing") );
+				{
+					ShowMenuBuilder.AddMenuEntry( Actions.ShowClothCollisionVolumes );
+				}
+				ShowMenuBuilder.EndSection();
 			}
-			ShowMenuBuilder.EndSection();
-		}
 #endif // #if WITH_APEX_CLOTHING
 
 		ShowMenuBuilder.AddMenuSeparator();
@@ -533,7 +523,7 @@ void SAnimViewportToolBar::FillShowSceneMenu(FMenuBuilder& MenuBuilder) const
 	{
 		MenuBuilder.AddMenuEntry(Actions.ToggleFloor);
 		MenuBuilder.AddMenuEntry(Actions.ToggleSky);
-	}
+		}
 	MenuBuilder.EndSection();
 
 	MenuBuilder.BeginSection("AnimViewportFloorOffset", LOCTEXT("Viewport_FloorOffsetLabel", "Floor Height Offset"));
@@ -593,6 +583,47 @@ void SAnimViewportToolBar::FillShowAdvancedMenu(FMenuBuilder& MenuBuilder) const
 	MenuBuilder.EndSection();
 }
 
+void SAnimViewportToolBar::FillShowClothingMenu(FMenuBuilder& MenuBuilder) const
+{
+	const FAnimViewportShowCommands& Actions = FAnimViewportShowCommands::Get();
+
+	MenuBuilder.BeginSection("ClothPreview", LOCTEXT("ClothPreview_Label", "Preview"));
+	{
+		MenuBuilder.AddMenuEntry(Actions.DisableClothSimulation);
+		MenuBuilder.AddMenuEntry(Actions.ApplyClothWind);
+		TSharedPtr<SWidget> WindWidget = SNew(SClothWindSettings).AnimEditorViewport(Viewport);
+		MenuBuilder.AddWidget(WindWidget.ToSharedRef(), FText());
+		TSharedPtr<SWidget> GravityWidget = SNew(SGravitySettings).AnimEditorViewport(Viewport);
+		MenuBuilder.AddWidget(GravityWidget.ToSharedRef(), FText());
+		MenuBuilder.AddMenuEntry(Actions.EnableCollisionWithAttachedClothChildren);
+	}
+	MenuBuilder.EndSection();
+
+	MenuBuilder.BeginSection("ClothNormalVisualization", LOCTEXT("ClothNormalVisualization_Label", "Normal Visualization"));
+	{
+		MenuBuilder.AddMenuEntry(Actions.ShowClothSimulationNormals);
+		MenuBuilder.AddMenuEntry(Actions.ShowClothGraphicalTangents);
+	}
+	MenuBuilder.EndSection();
+
+	MenuBuilder.BeginSection("ClothConstraintsVisualization", LOCTEXT("ClothConstraintsVisualization_Label", "Constraints Visualization"));
+	{
+		MenuBuilder.AddMenuEntry(Actions.ShowClothCollisionVolumes);
+		MenuBuilder.AddMenuEntry(Actions.ShowClothPhysicalMeshWire);
+		MenuBuilder.AddMenuEntry(Actions.ShowClothMaxDistances);
+		MenuBuilder.AddMenuEntry(Actions.ShowClothBackstop);
+		MenuBuilder.AddMenuEntry(Actions.ShowClothFixedVertices);
+	}
+	MenuBuilder.EndSection();
+
+	MenuBuilder.BeginSection("ClothAdditionalVisualization", LOCTEXT("ClothAdditionalVisualization_Label", "Sections Display Mode"));
+	{
+		MenuBuilder.AddMenuEntry(Actions.ShowAllSections);
+		MenuBuilder.AddMenuEntry(Actions.ShowOnlyClothSections);
+		MenuBuilder.AddMenuEntry(Actions.HideOnlyClothSections);
+	}
+}
+
 FText SAnimViewportToolBar::GetLODMenuLabel() const
 {
 	FText Label = LOCTEXT("LODMenu_AutoLabel", "LOD Auto");
diff --git a/Engine/Source/Editor/Persona/Private/SAnimViewportToolBar.h b/Engine/Source/Editor/Persona/Private/SAnimViewportToolBar.h
--- a/Engine/Source/Editor/Persona/Private/SAnimViewportToolBar.h
+++ b/Engine/Source/Editor/Persona/Private/SAnimViewportToolBar.h
@@ -46,7 +46,10 @@ private:
 	 * Generates the Show -> Advanced sub menu content
 	 */
 	void FillShowAdvancedMenu(FMenuBuilder& MenuBuilder) const;
-
+	/**
+	* Generates the Show -> Clothing sub menu content
+	*/
+	void FillShowClothingMenu(FMenuBuilder& MenuBuilder) const;
 	/**
 	 * Generates the toolbar LOD menu content 
 	 *
diff --git a/Engine/Source/Editor/Persona/Private/SAnimationEditorViewport.cpp b/Engine/Source/Editor/Persona/Private/SAnimationEditorViewport.cpp
--- a/Engine/Source/Editor/Persona/Private/SAnimationEditorViewport.cpp
+++ b/Engine/Source/Editor/Persona/Private/SAnimationEditorViewport.cpp
@@ -31,7 +31,7 @@ namespace EAnimationPlaybackSpeeds
 SAnimationEditorViewport::~SAnimationEditorViewport()
 {
 	if(PersonaPtr.IsValid())
-	{
+	{		
 		PersonaPtr.Pin()->UnregisterOnPostUndo(this);
 	}
 }
@@ -559,10 +559,49 @@ void SAnimationEditorViewportTabBody::BindCommands()
 		FIsActionChecked::CreateSP(this, &SAnimationEditorViewportTabBody::IsEnablingCollisionWithAttachedClothChildren));
 
 	CommandList.MapAction( 
+		ViewportShowMenuCommands.ShowClothPhysicalMeshWire,
+		FExecuteAction::CreateSP( this, &SAnimationEditorViewportTabBody::OnShowClothPhysicalMeshWire ),
+		FCanExecuteAction(),
+		FIsActionChecked::CreateSP( this, &SAnimationEditorViewportTabBody::IsShowingClothPhysicalMeshWire ) );
+
+	CommandList.MapAction( 
+		ViewportShowMenuCommands.ShowClothMaxDistances,
+		FExecuteAction::CreateSP( this, &SAnimationEditorViewportTabBody::OnShowClothMaxDistances ),
+		FCanExecuteAction(),
+		FIsActionChecked::CreateSP( this, &SAnimationEditorViewportTabBody::IsShowingClothMaxDistances ) );
+
+	CommandList.MapAction( 
+		ViewportShowMenuCommands.ShowClothBackstop,
+		FExecuteAction::CreateSP( this, &SAnimationEditorViewportTabBody::OnShowClothBackstops ),
+		FCanExecuteAction(),
+		FIsActionChecked::CreateSP( this, &SAnimationEditorViewportTabBody::IsShowingClothBackstops ) );
+
+	CommandList.MapAction( 
+		ViewportShowMenuCommands.ShowClothFixedVertices,
+		FExecuteAction::CreateSP( this, &SAnimationEditorViewportTabBody::OnShowClothFixedVertices ),
+		FCanExecuteAction(),
+		FIsActionChecked::CreateSP( this, &SAnimationEditorViewportTabBody::IsShowingClothFixedVertices ) );
+
+	SectionsDisplayMode = ESectionDisplayMode::ShowAll;
+
+	CommandList.MapAction(
+		ViewportShowMenuCommands.ShowAllSections,
+		FExecuteAction::CreateSP(this, &SAnimationEditorViewportTabBody::OnSetSectionsDisplayMode, (int32)ESectionDisplayMode::ShowAll),
+		FCanExecuteAction(),
+		FIsActionChecked::CreateSP(this, &SAnimationEditorViewportTabBody::IsSectionsDisplayMode, (int32)ESectionDisplayMode::ShowAll));
+
+	CommandList.MapAction(
 		ViewportShowMenuCommands.ShowOnlyClothSections,
-		FExecuteAction::CreateSP(this, &SAnimationEditorViewportTabBody::OnShowOnlyClothSections),
+		FExecuteAction::CreateSP(this, &SAnimationEditorViewportTabBody::OnSetSectionsDisplayMode, (int32)ESectionDisplayMode::ShowOnlyClothSections),
 		FCanExecuteAction(),
-		FIsActionChecked::CreateSP(this, &SAnimationEditorViewportTabBody::IsShowingOnlyClothSections));
+		FIsActionChecked::CreateSP(this, &SAnimationEditorViewportTabBody::IsSectionsDisplayMode, (int32)ESectionDisplayMode::ShowOnlyClothSections));
+
+	CommandList.MapAction(
+		ViewportShowMenuCommands.HideOnlyClothSections,
+		FExecuteAction::CreateSP(this, &SAnimationEditorViewportTabBody::OnSetSectionsDisplayMode, (int32)ESectionDisplayMode::HideOnlyClothSections),
+		FCanExecuteAction(),
+		FIsActionChecked::CreateSP(this, &SAnimationEditorViewportTabBody::IsSectionsDisplayMode, (int32)ESectionDisplayMode::HideOnlyClothSections));
+
 #endif// #if WITH_APEX_CLOTHING		
 
 	//Bind LOD preview menu commands
@@ -1349,6 +1388,19 @@ void SAnimationEditorViewportTabBody::OnDisableClothSimulation()
 	if( PreviewComponent )
 	{
 		PreviewComponent->bDisableClothSimulation = !PreviewComponent->bDisableClothSimulation;
+
+		// if the user turns on cloth simulation option while displaying max distances, then turns off max distance option
+		if(!PreviewComponent->bDisableClothSimulation && PreviewComponent->bDisplayClothMaxDistances)
+		{
+			PreviewComponent->bDisplayClothMaxDistances = false;
+		}
+
+		// if the user turns on cloth simulation option while displaying back stops, then turns off back stops option
+		if (!PreviewComponent->bDisableClothSimulation && PreviewComponent->bDisplayClothBackstops)
+		{
+			PreviewComponent->bDisplayClothBackstops = false;
+		}
+
 		RefreshViewport();
 	}
 }
@@ -1465,8 +1517,6 @@ bool SAnimationEditorViewportTabBody::IsShowingClothCollisionVolumes() const
 void SAnimationEditorViewportTabBody::SetGravityScale(float SliderPos)
 {
 	GetAnimationViewportClient()->SetGravityScale(SliderPos);
-
-
 	RefreshViewport();
 }
 
@@ -1503,76 +1553,158 @@ bool SAnimationEditorViewportTabBody::IsEnablingCollisionWithAttachedClothChildr
 	return false;
 }
 
-void SAnimationEditorViewportTabBody::OnShowOnlyClothSections()
+void SAnimationEditorViewportTabBody::OnShowClothPhysicalMeshWire()
 {
 	UDebugSkelMeshComponent* PreviewComponent = PersonaPtr.Pin()->PreviewComponent;
 
 	if( PreviewComponent )
 	{
-		PreviewComponent->bDisplayOnlyClothSections = !PreviewComponent->bDisplayOnlyClothSections;
+		PreviewComponent->bDisplayClothPhysicalMeshWire = !PreviewComponent->bDisplayClothPhysicalMeshWire;
+		RefreshViewport();
+	}
+}
 
-		PreviewComponent->PreEditChange(NULL);
+bool SAnimationEditorViewportTabBody::IsShowingClothPhysicalMeshWire() const
+{
+	UDebugSkelMeshComponent* PreviewComponent = PersonaPtr.Pin()->PreviewComponent;
 
-		if(PreviewComponent->bDisplayOnlyClothSections)
-		{
-			// disable all except clothing sections
-			FSkeletalMeshResource* SkelMeshResource = PreviewComponent->GetSkeletalMeshResource();
-			check(SkelMeshResource);
-			const int32 LODIndex = FMath::Clamp(PreviewComponent->PredictedLODLevel, 0, SkelMeshResource->LODModels.Num()-1);
-			FStaticLODModel& LODModel = SkelMeshResource->LODModels[ LODIndex ];
+	if( PreviewComponent )
+	{
+		return PreviewComponent->bDisplayClothPhysicalMeshWire;
+	}
 
-			for(int32 SecIdx=0; SecIdx < LODModel.Sections.Num(); SecIdx++)
-			{
-				FSkelMeshSection& Section = LODModel.Sections[SecIdx];
+	return false;
+}
 
-				if(!LODModel.Chunks[Section.ChunkIndex].HasApexClothData())
-				{
-					Section.bDisabled = true;
-				}
-			}
+void SAnimationEditorViewportTabBody::OnShowClothMaxDistances()
+{
+	UDebugSkelMeshComponent* PreviewComponent = PersonaPtr.Pin()->PreviewComponent;
+
+	if( PreviewComponent )
+	{
+		PreviewComponent->bDisplayClothMaxDistances = !PreviewComponent->bDisplayClothMaxDistances;
+
+		// disables cloth simulation and stop animation because showing max distances is only useful when the cloth is not moving
+		if(PreviewComponent->bDisplayClothMaxDistances)
+		{
+			PreviewComponent->bPrevDisableClothSimulation = PreviewComponent->bDisableClothSimulation;
+			PreviewComponent->bDisableClothSimulation = true;
+			PreviewComponent->InitAnim(false);
 		}
 		else
 		{
-			// restore to previous states
-			FSkeletalMeshResource* SkelMeshResource = PreviewComponent->GetSkeletalMeshResource();
-			check(SkelMeshResource);
-			const int32 LODIndex = FMath::Clamp(PreviewComponent->PredictedLODLevel, 0, SkelMeshResource->LODModels.Num()-1);
-			FStaticLODModel& LODModel = SkelMeshResource->LODModels[ LODIndex ];
+			// restore previous state
+			PreviewComponent->bDisableClothSimulation = PreviewComponent->bPrevDisableClothSimulation;
+		}
 
-			for(int32 SecIdx=0; SecIdx < LODModel.Sections.Num(); SecIdx++)
-			{
-				LODModel.Sections[SecIdx].bDisabled = false;
-			}
+		RefreshViewport();
+	}
+}
 
-			for(int32 SecIdx=0; SecIdx < LODModel.Sections.Num(); SecIdx++)
-			{
-				FSkelMeshSection& Section = LODModel.Sections[SecIdx];
+bool SAnimationEditorViewportTabBody::IsShowingClothMaxDistances() const
+{
+	UDebugSkelMeshComponent* PreviewComponent = PersonaPtr.Pin()->PreviewComponent;
 
-				if(LODModel.Chunks[Section.ChunkIndex].HasApexClothData())
-				{
-					LODModel.Sections[Section.CorrespondClothSectionIndex].bDisabled = true;
-				}
-			}
+	if( PreviewComponent )
+	{
+		return PreviewComponent->bDisplayClothMaxDistances;
+	}
+
+	return false;
+}
+
+void SAnimationEditorViewportTabBody::OnShowClothBackstops()
+{
+	UDebugSkelMeshComponent* PreviewComponent = PersonaPtr.Pin()->PreviewComponent;
+
+	if( PreviewComponent )
+	{
+		PreviewComponent->bDisplayClothBackstops = !PreviewComponent->bDisplayClothBackstops;
+		// disables cloth simulation and stop animation because showing back stops is only useful when the cloth is not moving
+		if (PreviewComponent->bDisplayClothBackstops)
+		{
+			PreviewComponent->bPrevDisableClothSimulation = PreviewComponent->bDisableClothSimulation;
+			PreviewComponent->bDisableClothSimulation = true;
+			PreviewComponent->InitAnim(false);
 		}
+		else
+		{
+			// restore previous state
+			PreviewComponent->bDisableClothSimulation = PreviewComponent->bPrevDisableClothSimulation;
+		}
+		RefreshViewport();
+	}
+}
 
-		PreviewComponent->PostEditChange();
+bool SAnimationEditorViewportTabBody::IsShowingClothBackstops() const
+{
+	UDebugSkelMeshComponent* PreviewComponent = PersonaPtr.Pin()->PreviewComponent;
 
+	if( PreviewComponent )
+	{
+		return PreviewComponent->bDisplayClothBackstops;
+	}
+
+	return false;
+}
+
+void SAnimationEditorViewportTabBody::OnShowClothFixedVertices()
+{
+	UDebugSkelMeshComponent* PreviewComponent = PersonaPtr.Pin()->PreviewComponent;
+
+	if( PreviewComponent )
+	{
+		PreviewComponent->bDisplayClothFixedVertices = !PreviewComponent->bDisplayClothFixedVertices;
 		RefreshViewport();
 	}
 }
 
-bool SAnimationEditorViewportTabBody::IsShowingOnlyClothSections() const
+bool SAnimationEditorViewportTabBody::IsShowingClothFixedVertices() const
 {
 	UDebugSkelMeshComponent* PreviewComponent = PersonaPtr.Pin()->PreviewComponent;
 
 	if( PreviewComponent )
 	{
-		return PreviewComponent->bDisplayOnlyClothSections;
+		return PreviewComponent->bDisplayClothFixedVertices;
 	}
 
 	return false;
 }
-#endif // #if WITH_APEX_CLOTHING
+
+void SAnimationEditorViewportTabBody::OnSetSectionsDisplayMode(int32 DisplayMode)
+{
+	UDebugSkelMeshComponent* PreviewComponent = PersonaPtr.Pin()->PreviewComponent;
+
+	SectionsDisplayMode = DisplayMode;
+
+	if (!PreviewComponent)
+	{
+		return;
+	}
+
+	switch (DisplayMode)
+	{
+	case ESectionDisplayMode::ShowAll:
+		// restore to the original states
+		PreviewComponent->RestoreClothSectionsVisibility();
+		break;
+	case ESectionDisplayMode::ShowOnlyClothSections:
+		// disable all except clothing sections
+		PreviewComponent->ShowOnlyClothSections(true, PreviewComponent->PredictedLODLevel);
+		break;
+	case ESectionDisplayMode::HideOnlyClothSections:
+		// disables only clothing sections
+		PreviewComponent->ShowOnlyClothSections(false, PreviewComponent->PredictedLODLevel);
+		break;
+	}
+
+	RefreshViewport();
+}
+
+bool SAnimationEditorViewportTabBody::IsSectionsDisplayMode(int32 DisplayMode) const
+{
+	return (SectionsDisplayMode == DisplayMode);
+}
 
 EVisibility SAnimationEditorViewportTabBody::GetViewportCornerTextVisibility() const
 {
@@ -1629,4 +1761,5 @@ FReply SAnimationEditorViewportTabBody::ClickedOnViewportCornerText()
 	return FReply::Handled();
 }
 
+#endif // #if WITH_APEX_CLOTHING
 #undef LOCTEXT_NAMESPACE
diff --git a/Engine/Source/Editor/Persona/Private/SAnimationEditorViewport.h b/Engine/Source/Editor/Persona/Private/SAnimationEditorViewport.h
--- a/Engine/Source/Editor/Persona/Private/SAnimationEditorViewport.h
+++ b/Engine/Source/Editor/Persona/Private/SAnimationEditorViewport.h
@@ -328,6 +328,14 @@ private:
 	bool IsAudioMuted();
 
 #if WITH_APEX_CLOTHING
+
+	enum ESectionDisplayMode
+	{
+		ShowAll,
+		ShowOnlyClothSections,
+		HideOnlyClothSections,
+		NumSectionDisplayMode
+	};
 	/** 
 	 * clothing show options 
 	*/
@@ -353,9 +361,29 @@ private:
 	void OnEnableCollisionWithAttachedClothChildren();
 	bool IsEnablingCollisionWithAttachedClothChildren() const;
 
-	/** Show cloth simulation normals */
-	void OnShowOnlyClothSections();
-	bool IsShowingOnlyClothSections() const;
+	/** Show cloth physical mesh wire */
+	void OnShowClothPhysicalMeshWire();
+	bool IsShowingClothPhysicalMeshWire() const;
+
+	/** Show cloth max distances */
+	void OnShowClothMaxDistances();
+	bool IsShowingClothMaxDistances() const;
+
+	/** Show cloth back stops */
+	void OnShowClothBackstops();
+	bool IsShowingClothBackstops() const;
+
+	/** Show only fixed vertices */
+	void OnShowClothFixedVertices();
+	bool IsShowingClothFixedVertices() const;
+
+	/** Show all sections which means the original state */
+	void OnSetSectionsDisplayMode(int32 DisplayMode);
+	bool IsSectionsDisplayMode(int32 DisplayMode) const;
+
+	/** Draw All/ Draw only clothing sections/ Hide only clothing sections */
+	uint32 SectionsDisplayMode;
+
 #endif // #if WITH_APEX_CLOTHING
 
 private:
diff --git a/Engine/Source/Editor/UnrealEd/Classes/Animation/DebugSkelMeshComponent.h b/Engine/Source/Editor/UnrealEd/Classes/Animation/DebugSkelMeshComponent.h
--- a/Engine/Source/Editor/UnrealEd/Classes/Animation/DebugSkelMeshComponent.h
+++ b/Engine/Source/Editor/UnrealEd/Classes/Animation/DebugSkelMeshComponent.h
@@ -210,6 +210,13 @@ class UDebugSkelMeshComponent : public USkeletalMeshComponent
 	virtual void ReportAnimNotifyError(const FText& Error, UObject* InSourceNotify);
 	virtual void ClearAnimNotifyErrors(UObject* InSourceNotify);
 #endif
+
+#if WITH_APEX_CLOTHING
+	/** toggle visibility between cloth sections and non-cloth sections */
+	UNREALED_API void ShowOnlyClothSections(bool bShow, int32 LODIndex);
+	/** Restore all section visibilities to original states for all LODs */
+	UNREALED_API void RestoreClothSectionsVisibility();
+#endif //#if WITH_APEX_CLOTHING
 };
 
 
diff --git a/Engine/Source/Editor/UnrealEd/Private/Animation/DebugSkelMeshComponent.cpp b/Engine/Source/Editor/UnrealEd/Private/Animation/DebugSkelMeshComponent.cpp
--- a/Engine/Source/Editor/UnrealEd/Private/Animation/DebugSkelMeshComponent.cpp
+++ b/Engine/Source/Editor/UnrealEd/Private/Animation/DebugSkelMeshComponent.cpp
@@ -250,6 +250,14 @@ void UDebugSkelMeshComponent::EnablePreview(bool bEnable, UAnimationAsset* Previ
 
 				AnimScriptInstance = PreviewInstance;
 
+#if WITH_APEX_CLOTHING
+			    // turn off these options when playing animations because max distances / back stops don't have meaning while moving
+			    bDisplayClothMaxDistances = false;
+				bDisplayClothBackstops = false;
+			    // restore previous state
+			    bDisableClothSimulation = bPrevDisableClothSimulation;
+#endif // #if WITH_APEX_CLOTHING
+    
 				if(PreviewAsset)
 				{
 					PreviewInstance->SetVertexAnimation(NULL);
@@ -422,4 +430,82 @@ void UDebugSkelMeshComponent::ClearAnimNotifyErrors(UObject* InSourceNotify)
 		}
 	}
 }
-#endif
\ No newline at end of file
+#endif
+
+#if WITH_APEX_CLOTHING
+
+void UDebugSkelMeshComponent::ShowOnlyClothSections(bool bShow, int32 LODIndex)
+{
+	PreEditChange(NULL);
+
+	FSkeletalMeshResource* SkelMeshResource = GetSkeletalMeshResource();
+	check(SkelMeshResource);
+	check(LODIndex >= 0 && LODIndex < SkelMeshResource->LODModels.Num());
+	FStaticLODModel& LODModel = SkelMeshResource->LODModels[LODIndex];
+
+	for (int32 SecIdx = 0; SecIdx < LODModel.Sections.Num(); SecIdx++)
+	{
+		FSkelMeshSection& Section = LODModel.Sections[SecIdx];
+
+		// toggle visibility between cloth sections and non-cloth sections
+		if (bShow)
+		{
+			// enables only cloth sections
+			if (LODModel.Chunks[Section.ChunkIndex].HasApexClothData())
+			{
+				Section.bDisabled = false;
+			}
+			else
+			{
+				Section.bDisabled = true;
+			}
+		}
+		else
+		{   // disables cloth sections and also corresponding original sections
+			if (LODModel.Chunks[Section.ChunkIndex].HasApexClothData())
+			{
+				Section.bDisabled = true;
+				LODModel.Sections[Section.CorrespondClothSectionIndex].bDisabled = true;
+			}
+			else
+			{
+				Section.bDisabled = false;
+			}
+		}
+	}
+
+	PostEditChange();
+}
+
+void UDebugSkelMeshComponent::RestoreClothSectionsVisibility()
+{
+	PreEditChange(NULL);
+
+	FSkeletalMeshResource* SkelMeshResource = GetSkeletalMeshResource();
+	check(SkelMeshResource);
+
+	for (int32 LODIndex = 0; LODIndex < SkelMeshResource->LODModels.Num(); LODIndex++)
+	{
+		FStaticLODModel& LODModel = SkelMeshResource->LODModels[LODIndex];
+
+		// enables all sections first
+		for (int32 SecIdx = 0; SecIdx < LODModel.Sections.Num(); SecIdx++)
+		{
+			LODModel.Sections[SecIdx].bDisabled = false;
+		}
+
+		// disables corresponding original section to enable the cloth section instead
+		for (int32 SecIdx = 0; SecIdx < LODModel.Sections.Num(); SecIdx++)
+		{
+			FSkelMeshSection& Section = LODModel.Sections[SecIdx];
+
+			if (LODModel.Chunks[Section.ChunkIndex].HasApexClothData())
+			{
+				LODModel.Sections[Section.CorrespondClothSectionIndex].bDisabled = true;
+			}
+		}
+	}
+
+	PostEditChange();
+}
+#endif // #if WITH_APEX_CLOTHING
\ No newline at end of file
diff --git a/Engine/Source/Runtime/Engine/Classes/Components/SkeletalMeshComponent.h b/Engine/Source/Runtime/Engine/Classes/Components/SkeletalMeshComponent.h
--- a/Engine/Source/Runtime/Engine/Classes/Components/SkeletalMeshComponent.h
+++ b/Engine/Source/Runtime/Engine/Classes/Components/SkeletalMeshComponent.h
@@ -342,9 +342,20 @@ public:
 	 * Supports up to 16 capsules / 32 planes per convex and ignored collisions by a max number will be drawn in Dark Gray.
 	 */
 	uint32 bDisplayClothingCollisionVolumes:1;
+	
+	/** Draw clothing physical mesh wire frame */
+	uint32 	bDisplayClothPhysicalMeshWire:1;	
+
+	/** Draw max distances of clothing simulation vertices */
+	uint32 	bDisplayClothMaxDistances:1;
+
+	/** Draw back stops of clothing simulation vertices */
+	uint32 	bDisplayClothBackstops:1;
+
+	/** To save previous state */
+	uint32 bPrevDisableClothSimulation:1;
 
-	/** Draw only clothing sections, disables non-clothing sections */
-	uint32 bDisplayOnlyClothSections:1;
+	uint32 bDisplayClothFixedVertices:1;
 	/**
 	 * Vertex Animation
 	 */
@@ -607,11 +618,28 @@ public:
 
 	bool IsValidClothingActor(int32 ActorIndex) const;
 	/** Draws APEX Clothing simulated normals on cloth meshes **/
-	void DrawClothingNormals(FPrimitiveDrawInterface* PDI) const;
+	void DrawClothingNormals(FPrimitiveDrawInterface* PDI);
 	/** Draws APEX Clothing Graphical Tangents on cloth meshes **/
-	void DrawClothingTangents(FPrimitiveDrawInterface* PDI) const;
+	void DrawClothingTangents(FPrimitiveDrawInterface* PDI);
 	/** Draws internal collision volumes which the character has, colliding with cloth **/
 	void DrawClothingCollisionVolumes(FPrimitiveDrawInterface* PDI);
+	/** Draws max distances of clothing simulation vertices 
+	  * clothing simulation will be disabled and animation will be reset when drawing this option 
+	  * because max distances do have meaning only in initial pose
+	 **/
+	void DrawClothingMaxDistances(FPrimitiveDrawInterface* PDI);
+	/** Draws Clothing back stops **/
+	void DrawClothingBackstops(FPrimitiveDrawInterface* PDI);
+	/** Draws Clothing Physical mesh wire **/
+	void DrawClothingPhysicalMeshWire(FPrimitiveDrawInterface* PDI);
+
+	void DrawClothingFixedVertices(FPrimitiveDrawInterface* PDI);
+
+	/** Loads clothing extra infos dynamically just for Previewing in Editor 
+	 *  such as MaxDistances, Physical mesh wire
+	 **/
+	void LoadClothingVisualizationInfo(int32 AssetIndex);
+	void LoadAllClothingVisualizationInfos();
 
 	/** freezing clothing actor now */
 	void FreezeClothSection(bool bFreeze);
diff --git a/Engine/Source/Runtime/Engine/Classes/Engine/SkeletalMesh.h b/Engine/Source/Runtime/Engine/Classes/Engine/SkeletalMesh.h
--- a/Engine/Source/Runtime/Engine/Classes/Engine/SkeletalMesh.h
+++ b/Engine/Source/Runtime/Engine/Classes/Engine/SkeletalMesh.h
@@ -300,6 +300,52 @@ private:
 };
 #endif // #if WITH_APEX_CLOTHING
 
+
+/** 
+ * constrain Coefficients - max distance, collisionSphere radius, collision sphere distance 
+ */
+struct FClothConstrainCoeff
+{
+	float ClothMaxDistance;
+	float ClothBackstopRadius;
+	float ClothBackstopDistance;
+};
+
+/** 
+ * bone weights & bone indices for visualization of cloth physical mesh 
+ */
+struct FClothBoneWeightsInfo
+{
+	// support up 4 bone influences but used MAX_TOTAL_INFLUENCES for the future
+	uint16 Indices[MAX_TOTAL_INFLUENCES];
+	float Weights[MAX_TOTAL_INFLUENCES];
+};
+/** 
+ * save temporary data only for debugging on Persona editor 
+ */
+struct FClothVisualizationInfo
+{
+	TArray<FVector> ClothPhysicalMeshVertices;
+	TArray<FVector> ClothPhysicalMeshNormals;
+	TArray<uint32> ClothPhysicalMeshIndices;
+	TArray<FClothConstrainCoeff> ClothConstrainCoeffs;
+	// bone weights & bone indices
+	TArray<FClothBoneWeightsInfo> ClothPhysicalMeshBoneWeightsInfo;
+	uint8	NumMaxBoneInfluences;
+
+	// Max value of max distances
+	float MaximumMaxDistance;
+};
+
+/** 
+ * data structure for loading bone planes of collision volumes data
+ */
+struct FClothBonePlane
+{
+	int32 BoneIndex;
+	FPlane PlaneData;
+};
+
 USTRUCT()
 struct FClothingAssetData
 {
@@ -318,9 +364,16 @@ struct FClothingAssetData
 	/* Collision volume data for showing to the users whether collision shape is correct or not */
 	TArray<FApexClothCollisionVolumeData> ClothCollisionVolumes;
 	TArray<uint32> ClothCollisionConvexPlaneIndices;
-	TArray<FPlane> ClothCollisionVolumePlanes;
+	TArray<FClothBonePlane> ClothCollisionVolumePlanes;
 	TArray<FApexClothBoneSphereData> ClothBoneSpheres;
 	TArray<uint16> BoneSphereConnections;
+
+	/**
+	 * saved temporarily just for debugging / visualization 
+	 * Num of this array means LOD number of clothing physical meshes 
+	 */
+	TArray<FClothVisualizationInfo> ClothVisualizationInfos;
+
 #endif// #if WITH_APEX_CLOTHING
 
 	// serialization
diff --git a/Engine/Source/Runtime/Engine/Private/SkeletalMeshComponentPhysics.cpp b/Engine/Source/Runtime/Engine/Private/SkeletalMeshComponentPhysics.cpp
--- a/Engine/Source/Runtime/Engine/Private/SkeletalMeshComponentPhysics.cpp
+++ b/Engine/Source/Runtime/Engine/Private/SkeletalMeshComponentPhysics.cpp
@@ -13,6 +13,18 @@
 
 #include "Collision/CollisionDebugDrawing.h"
 
+#if WITH_APEX
+	#include "NxParamUtils.h"
+	#include "NxApex.h"
+
+#if WITH_APEX_CLOTHING
+	#include "NxClothingAsset.h"
+	#include "NxClothingActor.h"
+	#include "NxClothingCollision.h"
+#endif// #if WITH_APEX_CLOTHING
+
+#endif//#if WITH_APEX
+
 #define LOCTEXT_NAMESPACE "SkeletalMeshComponentPhysics"
 
 #if WITH_APEX_CLOTHING
@@ -121,6 +133,9 @@ void USkeletalMesh::LoadClothCollisionVolumes(int32 AssetIndex, NxClothingAsset*
 
 	for(int32 PlaneIdx=0; PlaneIdx<NumPlanes; PlaneIdx++)
 	{
+		FClothBonePlane BonePlane;
+		FCStringAnsi::Sprintf(ParameterName, "bonePlanes[%d].boneIndex", PlaneIdx);
+		verify(NxParameterized::getParamI32(*AssetParams, ParameterName, BonePlane.BoneIndex));
 		FCStringAnsi::Sprintf(ParameterName, "bonePlanes[%d].n", PlaneIdx);
 		verify(NxParameterized::getParamVec3(*AssetParams, ParameterName, PlaneNormal));
 		FCStringAnsi::Sprintf(ParameterName, "bonePlanes[%d].d", PlaneIdx);
@@ -132,8 +147,8 @@ void USkeletalMesh::LoadClothCollisionVolumes(int32 AssetIndex, NxClothingAsset*
 		}
 
 		PlaneData[3] = PlaneDist;
-
-		Asset.ClothCollisionVolumePlanes.Add(P2UPlane(PlaneData));
+		BonePlane.PlaneData = P2UPlane(PlaneData);
+		Asset.ClothCollisionVolumePlanes.Add(BonePlane);
 	}
 
 	// load bone spheres
@@ -1789,11 +1804,11 @@ bool USkeletalMeshComponent::HasValidClothingActors()
 	return false;
 }
 
-// if any changes in clothing assets, will create new actors 
+//if any changes in clothing assets, will create new actors 
 void USkeletalMeshComponent::ValidateClothingActors()
 {
-	// newly spawned component's tick group could be "specified tick group + 1" for the first few frames
-	// but it comes back to original tick group soon
+	//newly spawned component's tick group could be "specified tick group + 1" for the first few frames
+	//but it comes back to original tick group soon
 	if(PrimaryComponentTick.GetActualTickGroup() != PrimaryComponentTick.TickGroup)
 	{
 		return;
@@ -1806,7 +1821,7 @@ void USkeletalMeshComponent::ValidateClothingActors()
 
 	int32 NumAssets = SkeletalMesh->ClothingAssets.Num();
 
-	// if clothing assets are added or removed, re-create after removing all
+	//if clothing assets are added or removed, re-create after removing all
 	if(ClothingActors.Num() != NumAssets)
 	{
 		RemoveAllClothingActors();
@@ -1833,10 +1848,10 @@ void USkeletalMeshComponent::ValidateClothingActors()
 			}
 		}
 		else 
-		{	// don't have cloth sections but clothing actor is alive
+		{	//don't have cloth sections but clothing actor is alive
 			if(IsValidClothingActor(AssetIdx))
 			{
-				// clear because mapped sections are removed
+				//clear because mapped sections are removed
 				ClothingActors[AssetIdx].Clear(true);
 			}
 		}
@@ -1851,7 +1866,7 @@ void USkeletalMeshComponent::ValidateClothingActors()
 /** 
  * APEX clothing actor is created from APEX clothing asset for cloth simulation 
  * If this is invalid, re-create actor , but if valid ,just skip to create
- */
+*/
 bool USkeletalMeshComponent::CreateClothingActor(int32 AssetIndex, TSharedPtr<FClothingAssetWrapper> ClothingAssetWrapper)
 {	
 	int32 NumActors = ClothingActors.Num();
@@ -1891,7 +1906,7 @@ bool USkeletalMeshComponent::CreateClothingActor(int32 AssetIndex, TSharedPtr<FC
 
 	FVector ScaleVector = ComponentToWorld.GetScale3D();
 
-	// support only uniform scale
+	//support only uniform scale
 	verify(NxParameterized::setParamF32(*ActorDesc, "actorScale", ScaleVector.X));
 
 	bool bUseInternalBoneOrder = true;
@@ -1967,6 +1982,12 @@ bool USkeletalMeshComponent::CreateClothingActor(int32 AssetIndex, TSharedPtr<FC
 	ClothingActor->forcePhysicalLod(1); // 1 will be changed to "GetActivePhysicalLod()" later
 	ClothingActor->setFrozen(false);
 
+	// process clothing collisions once for the case that this component doesn't move
+	if(bCollideWithEnvironment)
+	{
+		ProcessClothCollisionWithEnvironment();
+	}
+
 	return true;
 }
 
@@ -2165,7 +2186,7 @@ void USkeletalMeshComponent::DrawDebugClothCollisions()
 {
 	FColor Colors[6] = { FColor::Red, FColor::Green, FColor::Blue, FColor::Cyan, FColor::Yellow, FColor::Magenta };
 
-	for (auto It = ClothOverlappedComponentsMap.CreateConstIterator(); It; ++It)
+	for( auto It = ClothOverlappedComponentsMap.CreateConstIterator(); It; ++It )
 	{
 		TWeakObjectPtr<UPrimitiveComponent> PrimComp = It.Key();
 
@@ -2174,36 +2195,36 @@ void USkeletalMeshComponent::DrawDebugClothCollisions()
 		if (Channel == ECollisionChannel::ECC_WorldStatic)
 		{
 			if (!GetClothCollisionDataFromStaticMesh(PrimComp.Get(), CollisionPrims))
-			{
-				return;
-			}
+		{
+			return;
+		}
 
-			for (int32 PrimIndex = 0; PrimIndex < CollisionPrims.Num(); PrimIndex++)
+		for(int32 PrimIndex=0; PrimIndex < CollisionPrims.Num(); PrimIndex++)
+		{
+			switch(CollisionPrims[PrimIndex].PrimType)
 			{
-				switch (CollisionPrims[PrimIndex].PrimType)
-				{
-				case FClothCollisionPrimitive::SPHERE:
-					DrawDebugSphere(GetWorld(), CollisionPrims[PrimIndex].Origin, CollisionPrims[PrimIndex].Radius, 10, FColor::Red);
-					break;
+			case FClothCollisionPrimitive::SPHERE:
+				DrawDebugSphere(GetWorld(), CollisionPrims[PrimIndex].Origin, CollisionPrims[PrimIndex].Radius, 10, FColor::Red);
+				break;
 
-				case FClothCollisionPrimitive::CAPSULE:
+			case FClothCollisionPrimitive::CAPSULE:
 				{
 					FVector DiffVec = CollisionPrims[PrimIndex].SpherePos2 - CollisionPrims[PrimIndex].SpherePos1;
 					float HalfHeight = DiffVec.Size() * 0.5f;
 					FQuat Rotation = PrimComp->ComponentToWorld.GetRotation();
-					DrawDebugCapsule(GetWorld(), CollisionPrims[PrimIndex].Origin, HalfHeight, CollisionPrims[PrimIndex].Radius, Rotation, FColor::Red);
-				}
-
-					break;
-				case FClothCollisionPrimitive::CONVEX:
-					DrawDebugConvexFromPlanes(CollisionPrims[PrimIndex], Colors[PrimIndex % 6]);
-					break;
+					DrawDebugCapsule(GetWorld(), CollisionPrims[PrimIndex].Origin, HalfHeight, CollisionPrims[PrimIndex].Radius, Rotation, FColor::Red );
 				}
 
-				//draw a bounding box for double checking
-				DrawDebugBox(GetWorld(), PrimComp->Bounds.Origin, PrimComp->Bounds.BoxExtent, FColor::Red);
+				break;
+			case FClothCollisionPrimitive::CONVEX:	
+				DrawDebugConvexFromPlanes(CollisionPrims[PrimIndex], Colors[PrimIndex%6]);
+				break;
 			}
+
+			//draw a bounding box for double checking
+			DrawDebugBox(GetWorld(), PrimComp->Bounds.Origin, PrimComp->Bounds.BoxExtent, FColor::Red );
 		}
+	}
 		else if (Channel == ECollisionChannel::ECC_PhysicsBody) // supports an interaction between the character and other clothing such as a curtain
 		{
 			bool bCreatedCollisions = false;
@@ -2391,7 +2412,7 @@ bool USkeletalMeshComponent::GetClothCollisionDataFromStaticMesh(UPrimitiveCompo
 			}
 			break;
 		}
-	}		
+	}
 	return true;
 }
 
@@ -2439,12 +2460,12 @@ void USkeletalMeshComponent::FindClothCollisions(TArray<FApexClothCollisionVolum
 				NewCollision.LocalPose = LocalToWorld;
 				OutCollisions.Add(NewCollision);
 			}
+			}
 		}
 	}
-}
 
 void USkeletalMeshComponent::CreateInternalClothCollisions(TArray<FApexClothCollisionVolumeData>& InCollisions, TArray<physx::apex::NxClothingCollision*>& OutCollisions)
-{
+	{
 	int32 NumCollisions = InCollisions.Num();
 
 	const int32 MaxNumCapsules = 16;
@@ -2452,42 +2473,42 @@ void USkeletalMeshComponent::CreateInternalClothCollisions(TArray<FApexClothColl
 	NumCollisions = FMath::Min(NumCollisions, MaxNumCapsules);
 	int32 NumActors = ClothingActors.Num();
 
-	for (int32 ActorIdx = 0; ActorIdx < NumActors; ActorIdx++)
-	{
-		if (!IsValidClothingActor(ActorIdx))
+		for(int32 ActorIdx=0; ActorIdx < NumActors; ActorIdx++)
 		{
-			continue;
-		}
+		if (!IsValidClothingActor(ActorIdx))
+			{
+				continue;
+			}
 
 		NxClothingActor* Actor = ClothingActors[ActorIdx].ApexClothingActor;
 		int32 NumCurrentCapsules = SkeletalMesh->ClothingAssets[ActorIdx].ClothCollisionVolumes.Num(); // # of capsules 
 
-		for (int32 ColIdx = 0; ColIdx < NumCollisions; ColIdx++)
-		{
-			// support only capsules now 
-			if (InCollisions[ColIdx].IsCapsule() && NumCurrentCapsules < MaxNumCapsules)
+			for(int32 ColIdx=0; ColIdx < NumCollisions; ColIdx++)
 			{
+				// support only capsules now 
+			if (InCollisions[ColIdx].IsCapsule() && NumCurrentCapsules < MaxNumCapsules)
+				{
 				FVector Origin = InCollisions[ColIdx].LocalPose.GetOrigin();
-				// apex uses y-axis as the up-axis of capsule
+					// apex uses y-axis as the up-axis of capsule
 				FVector UpAxis = InCollisions[ColIdx].LocalPose.GetScaledAxis(EAxis::Y);
-
+					
 				float Radius = InCollisions[ColIdx].CapsuleRadius*UpAxis.Size();
 
 				float HalfHeight = InCollisions[ColIdx].CapsuleHeight*0.5f;
-				const FVector TopEnd = Origin + (HalfHeight * UpAxis);
-				const FVector BottomEnd = Origin - (HalfHeight * UpAxis);
+					const FVector TopEnd = Origin + (HalfHeight * UpAxis);
+					const FVector BottomEnd = Origin - (HalfHeight * UpAxis);
 
-				NxClothingSphere* Sphere1 = Actor->createCollisionSphere(U2PVector(TopEnd), Radius);
-				NxClothingSphere* Sphere2 = Actor->createCollisionSphere(U2PVector(BottomEnd), Radius);
+					NxClothingSphere* Sphere1 = Actor->createCollisionSphere(U2PVector(TopEnd), Radius);
+					NxClothingSphere* Sphere2 = Actor->createCollisionSphere(U2PVector(BottomEnd), Radius);
 
-				NxClothingCapsule* Capsule = Actor->createCollisionCapsule(*Sphere1, *Sphere2);
+					NxClothingCapsule* Capsule = Actor->createCollisionCapsule(*Sphere1, *Sphere2);
 
 				OutCollisions.Add(Capsule);
 				NumCurrentCapsules++;
+				}
 			}
 		}
 	}
-}
 
 void USkeletalMeshComponent::CopyClothCollisionsToChildren()
 {
@@ -2578,7 +2599,7 @@ void USkeletalMeshComponent::ReleaseClothingCollision(NxClothingCollision* Colli
 
 			check(Capsule);
 
-			Capsule->releaseWithSpheres(); 
+			Capsule->releaseWithSpheres();
 		}
 		break;
 
@@ -2610,77 +2631,77 @@ FApexClothCollisionInfo* USkeletalMeshComponent::CreateNewClothingCollsions(UPri
 	if (Channel == ECollisionChannel::ECC_WorldStatic)
 	{
 		if (!GetClothCollisionDataFromStaticMesh(PrimitiveComponent, CollisionPrims))
-		{
-			return NULL;
-		}
+	{
+		return NULL;
+	}
 
 		NewInfo.OverlapCompType = FApexClothCollisionInfo::OCT_STATIC;
 
-		int32 NumActors = ClothingActors.Num();
+	int32 NumActors = ClothingActors.Num();
 
-		for (int32 ActorIdx = 0; ActorIdx < NumActors; ActorIdx++)
-		{
-			NxClothingActor* Actor = ClothingActors[ActorIdx].ApexClothingActor;
+	for(int32 ActorIdx=0; ActorIdx < NumActors; ActorIdx++)
+	{
+		NxClothingActor* Actor = ClothingActors[ActorIdx].ApexClothingActor;
 
-			if (Actor)
+		if(Actor)
+		{	
+			for(int32 PrimIndex=0; PrimIndex < CollisionPrims.Num(); PrimIndex++)
 			{
-				for (int32 PrimIndex = 0; PrimIndex < CollisionPrims.Num(); PrimIndex++)
-				{
-					NxClothingCollision* ClothCol = NULL;
+				NxClothingCollision* ClothCol = NULL;
 
-					switch (CollisionPrims[PrimIndex].PrimType)
+				switch(CollisionPrims[PrimIndex].PrimType)
+				{
+				case FClothCollisionPrimitive::SPHERE:
+					ClothCol = Actor->createCollisionSphere(U2PVector(CollisionPrims[PrimIndex].Origin), CollisionPrims[PrimIndex].Radius);
+					if(ClothCol)
 					{
-					case FClothCollisionPrimitive::SPHERE:
-						ClothCol = Actor->createCollisionSphere(U2PVector(CollisionPrims[PrimIndex].Origin), CollisionPrims[PrimIndex].Radius);
-						if (ClothCol)
-						{
-							NewInfo.ClothingCollisions.Add(ClothCol);
-						}
-						break;
-					case FClothCollisionPrimitive::CAPSULE:
+						NewInfo.ClothingCollisions.Add(ClothCol);
+					}
+					break;
+				case FClothCollisionPrimitive::CAPSULE:
 					{
 						float Radius = CollisionPrims[PrimIndex].Radius;
 						NxClothingSphere* Sphere1 = Actor->createCollisionSphere(U2PVector(CollisionPrims[PrimIndex].SpherePos1), Radius);
 						NxClothingSphere* Sphere2 = Actor->createCollisionSphere(U2PVector(CollisionPrims[PrimIndex].SpherePos2), Radius);
 
 						ClothCol = Actor->createCollisionCapsule(*Sphere1, *Sphere2);
-						if (ClothCol)
+						if(ClothCol)
 						{
 							NewInfo.ClothingCollisions.Add(ClothCol);
 						}
 					}
-						break;
-					case FClothCollisionPrimitive::CONVEX:
-
-						int32 NumPlanes = CollisionPrims[PrimIndex].ConvexPlanes.Num();
+					break;
+				case FClothCollisionPrimitive::CONVEX:
 
-						TArray<NxClothingPlane*> ClothingPlanes;
+					int32 NumPlanes = CollisionPrims[PrimIndex].ConvexPlanes.Num();
 
-						//can not exceed 32 planes
-						NumPlanes = FMath::Min(NumPlanes, 32);
+					TArray<NxClothingPlane*> ClothingPlanes;
 
-						ClothingPlanes.AddUninitialized(NumPlanes);
+					//can not exceed 32 planes
+					NumPlanes = FMath::Min(NumPlanes, 32);
 
-						for (int32 PlaneIdx = 0; PlaneIdx < NumPlanes; PlaneIdx++)
-						{
-							PxPlane PPlane = U2PPlane(CollisionPrims[PrimIndex].ConvexPlanes[PlaneIdx]);
+					ClothingPlanes.AddUninitialized(NumPlanes);
 
-							ClothingPlanes[PlaneIdx] = Actor->createCollisionPlane(PPlane);
-						}
+					for(int32 PlaneIdx=0; PlaneIdx < NumPlanes; PlaneIdx++)
+					{
+						PxPlane PPlane = U2PPlane(CollisionPrims[PrimIndex].ConvexPlanes[PlaneIdx]);
 
-						ClothCol = Actor->createCollisionConvex(ClothingPlanes.GetTypedData(), ClothingPlanes.Num());
+						ClothingPlanes[PlaneIdx] = Actor->createCollisionPlane(PPlane);
+					}
 
-						if (ClothCol)
-						{
-							NewInfo.ClothingCollisions.Add(ClothCol);
-						}
+					ClothCol = Actor->createCollisionConvex(ClothingPlanes.GetTypedData(), ClothingPlanes.Num());
 
-						break;
+					if(ClothCol)
+					{
+						NewInfo.ClothingCollisions.Add(ClothCol);
 					}
+
+					break;
 				}
 			}
 		}
 	}
+	}
 	else if (Channel == ECollisionChannel::ECC_PhysicsBody) // creates new clothing collisions for clothing objects
 	{
 		bool bCreatedCollisions = false;
@@ -2957,10 +2978,10 @@ void USkeletalMeshComponent::PreClothTick(float DeltaTime)
 
 	// if skeletal mesh has clothing assets, call TickClothing
 	if (SkeletalMesh && SkeletalMesh->ClothingAssets.Num() > 0)
-	{
+		{
 		TickClothing(DeltaTime + SkippedTickDeltaTime);
+		}
 	}
-}
 
 void USkeletalMeshComponent::UpdateClothState(float DeltaTime)
 {
@@ -3079,7 +3100,7 @@ void USkeletalMeshComponent::TickClothing(float DeltaTime)
 	}
 	else
 	{
-		ValidateClothingActors();
+	ValidateClothingActors();
 		UpdateClothState(DeltaTime);
 	}
 #if 0 //if turn on this flag, you can check which primitive objects are activated for collision detection
@@ -3200,7 +3221,7 @@ bool USkeletalMeshComponent::IsValidClothingActor(int32 ActorIndex) const
 
 }
 
-void USkeletalMeshComponent::DrawClothingNormals(FPrimitiveDrawInterface* PDI) const
+void USkeletalMeshComponent::DrawClothingNormals(FPrimitiveDrawInterface* PDI)
 {
 #if WITH_APEX_CLOTHING
 	if (!SkeletalMesh) 
@@ -3223,7 +3244,8 @@ void USkeletalMeshComponent::DrawClothingNormals(FPrimitiveDrawInterface* PDI) c
 		{
 			uint32 NumSimulVertices = ClothingActor->getNumSimulationVertices();
 
-			if(NumSimulVertices > 0)
+			// simulation is enabled and there exist simulated vertices
+			if(!bDisableClothSimulation && NumSimulVertices > 0)
 			{
 				const physx::PxVec3* Vertices = ClothingActor->getSimulationPositions();
 				const physx::PxVec3* Normals = ClothingActor->getSimulationNormals();
@@ -3235,6 +3257,39 @@ void USkeletalMeshComponent::DrawClothingNormals(FPrimitiveDrawInterface* PDI) c
 				{
 					Start = P2UVector(Vertices[i]);
 					End = P2UVector(Normals[i]); 
+					End *= 5.0f;
+
+					End = Start + End;
+
+					PDI->DrawLine(Start, End, LineColor, SDPG_World);
+				}
+			}
+			else // otherwise, draws initial values loaded from the asset file
+			{
+				if(SkeletalMesh->ClothingAssets[ActorIdx].ClothVisualizationInfos.Num() == 0)
+				{
+					LoadClothingVisualizationInfo(ActorIdx);
+				}
+
+				int32 PhysicalMeshLOD = PredictedLODLevel;
+
+				// if this asset doesn't have current physical mesh LOD, then skip to visualize
+				if(!SkeletalMesh->ClothingAssets[ActorIdx].ClothVisualizationInfos.IsValidIndex(PhysicalMeshLOD))
+				{
+					continue;
+				}
+
+				FClothVisualizationInfo& VisualInfo = SkeletalMesh->ClothingAssets[ActorIdx].ClothVisualizationInfos[PhysicalMeshLOD];				
+
+				uint32 NumVertices = VisualInfo.ClothPhysicalMeshVertices.Num();
+
+				FLinearColor LineColor = FColor::Red;
+				FVector Start, End;
+
+				for(uint32 i=0; i<NumVertices; i++)
+				{
+					Start = VisualInfo.ClothPhysicalMeshVertices[i];
+					End = VisualInfo.ClothPhysicalMeshNormals[i]; 
 
 					End *= 5.0f;
 
@@ -3248,7 +3303,7 @@ void USkeletalMeshComponent::DrawClothingNormals(FPrimitiveDrawInterface* PDI) c
 #endif // #if WITH_APEX_CLOTHING
 }
 
-void USkeletalMeshComponent::DrawClothingTangents(FPrimitiveDrawInterface* PDI) const
+void USkeletalMeshComponent::DrawClothingTangents(FPrimitiveDrawInterface* PDI)
 {
 #if WITH_APEX_CLOTHING
 
@@ -3272,11 +3327,9 @@ void USkeletalMeshComponent::DrawClothingTangents(FPrimitiveDrawInterface* PDI)
 		{
 			uint32 NumSimulVertices = ClothingActor->getNumSimulationVertices();
 
-			if(NumSimulVertices == 0)
+			// simulation is enabled and there exist simulated vertices
+			if(!bDisableClothSimulation && NumSimulVertices > 0)
 			{
-				return;
-			}
-
 				uint16 ChunkIndex = 0;
 
 				FStaticLODModel& LODModel = MeshObject->GetSkeletalMeshResource().LODModels[PredictedLODLevel];
@@ -3355,19 +3408,127 @@ void USkeletalMeshComponent::DrawClothingTangents(FPrimitiveDrawInterface* PDI)
 						End *= 5.0f;
 						End = Start + End;
 
-						PDI->DrawLine(Start, End, FColor::Red, SDPG_World);
+						PDI->DrawLine(Start, End, FColor::Green, SDPG_World);
 
 						End = Tangent; 
 						End *= 5.0f;
 						End = Start + End;
 
-						PDI->DrawLine(Start, End, FColor::Blue, SDPG_World);
+						PDI->DrawLine(Start, End, FColor::Red, SDPG_World);
 
 						End = BiNormal; 
 						End *= 5.0f;
 						End = Start + End;
 
+						PDI->DrawLine(Start, End, FColor::Blue, SDPG_World);
+					}
+				}
+			}
+			else
+			{
+				if(SkeletalMesh->ClothingAssets[ActorIdx].ClothVisualizationInfos.Num() == 0)
+				{
+					LoadClothingVisualizationInfo(ActorIdx);
+				}
+
+				int32 PhysicalMeshLOD = PredictedLODLevel;
+
+				// if this asset doesn't have current physical mesh LOD, then skip to visualize
+				if(!SkeletalMesh->ClothingAssets[ActorIdx].ClothVisualizationInfos.IsValidIndex(PhysicalMeshLOD))
+				{
+					continue;
+				}
+
+				FClothVisualizationInfo& VisualInfo = SkeletalMesh->ClothingAssets[ActorIdx].ClothVisualizationInfos[PhysicalMeshLOD];				
+
+				NumSimulVertices = VisualInfo.ClothPhysicalMeshVertices.Num();
+
+				uint16 ChunkIndex = 0;
+
+				FStaticLODModel& LODModel = MeshObject->GetSkeletalMeshResource().LODModels[PredictedLODLevel];
+
+				TArray<uint32> SectionIndices;
+				SkeletalMesh->GetClothSectionIndices(PredictedLODLevel, ActorIdx, SectionIndices);
+
+				int32 NumSections = SectionIndices.Num();
+
+				for(int32 SecIdx=0; SecIdx <NumSections; SecIdx++)
+				{
+					uint16 SectionIndex = SectionIndices[SecIdx];
+
+					ChunkIndex = LODModel.Sections[SectionIndex].ChunkIndex;
+
+					FSkelMeshChunk& Chunk = LODModel.Chunks[ChunkIndex];
+
+					const TArray<FVector>& SimulVertices = VisualInfo.ClothPhysicalMeshVertices;
+					const TArray<FVector>& SimulNormals = VisualInfo.ClothPhysicalMeshNormals;
+
+					uint32 NumMppingData = Chunk.ApexClothMappingData.Num();
+
+					FVector Start, End;		
+
+					for(uint32 MappingIndex=0; MappingIndex < NumMppingData; MappingIndex++)
+					{
+						FVector4 BaryCoordPos = Chunk.ApexClothMappingData[MappingIndex].PositionBaryCoordsAndDist;
+						FVector4 BaryCoordNormal = Chunk.ApexClothMappingData[MappingIndex].NormalBaryCoordsAndDist;
+						FVector4 BaryCoordTangent = Chunk.ApexClothMappingData[MappingIndex].TangentBaryCoordsAndDist;
+						uint16*  SimulIndices = Chunk.ApexClothMappingData[MappingIndex].SimulMeshVertIndices;
+
+						bool bFixed = (SimulIndices[3] == 0xFFFF);
+
+						check(SimulIndices[0] < NumSimulVertices && SimulIndices[1] < NumSimulVertices && SimulIndices[2] < NumSimulVertices);
+
+						FVector a = SimulVertices[SimulIndices[0]];
+						FVector b = SimulVertices[SimulIndices[1]];
+						FVector c = SimulVertices[SimulIndices[2]];
+
+						FVector na = SimulNormals[SimulIndices[0]];
+						FVector nb = SimulNormals[SimulIndices[1]];
+						FVector nc = SimulNormals[SimulIndices[2]];
+
+						FVector Position = 
+							BaryCoordPos.X*(a + BaryCoordPos.W*na)
+							+ BaryCoordPos.Y*(b + BaryCoordPos.W*nb)
+							+ BaryCoordPos.Z*(c + BaryCoordPos.W*nc);
+
+						FVector Normal = 
+							BaryCoordNormal.X*(a + BaryCoordNormal.W*na)
+							+ BaryCoordNormal.Y*(b + BaryCoordNormal.W*nb) 
+							+ BaryCoordNormal.Z*(c + BaryCoordNormal.W*nc);
+
+						FVector Tangent = 
+							BaryCoordTangent.X*(a + BaryCoordTangent.W*na) 
+							+ BaryCoordTangent.Y*(b + BaryCoordTangent.W*nb)  
+							+ BaryCoordTangent.Z*(c + BaryCoordTangent.W*nc);
+
+						Normal -= Position;
+						Normal.Normalize();
+
+						Tangent -= Position;
+						Tangent.Normalize();
+
+						FVector BiNormal = FVector::CrossProduct(Normal, Tangent);
+						BiNormal.Normalize();
+
+						Start = Position;
+						End = Normal; 
+						End *= 5.0f;
+						End = Start + End;
+
 						PDI->DrawLine(Start, End, FColor::Green, SDPG_World);
+
+						End = Tangent; 
+						End *= 5.0f;
+						End = Start + End;
+
+						PDI->DrawLine(Start, End, FColor::Red, SDPG_World);
+
+						End = BiNormal; 
+						End *= 5.0f;
+						End = Start + End;
+
+						PDI->DrawLine(Start, End, FColor::Blue, SDPG_World);
+					}
 					}
 				}
 			}
@@ -3431,7 +3592,7 @@ void USkeletalMeshComponent::DrawClothingCollisionVolumes(FPrimitiveDrawInterfac
 				DrawWireCapsule(PDI, LocalToWorld.GetOrigin(), LocalToWorld.GetUnitAxis(EAxis::X), LocalToWorld.GetUnitAxis(EAxis::Z), LocalToWorld.GetUnitAxis(EAxis::Y), CapsuleColor, Collisions[i].CapsuleRadius, CapsuleHalfHeight, CapsuleSides, SDPG_World);
 				SphereCount += 2; // 1 capsule takes 2 spheres
 			}
-		}
+			}
 
 		// draw bone spheres
 		TArray<FApexClothBoneSphereData>& Spheres = SkeletalMesh->ClothingAssets[AssetIdx].ClothBoneSpheres;
@@ -3443,7 +3604,7 @@ void USkeletalMeshComponent::DrawClothingCollisionVolumes(FPrimitiveDrawInterfac
 
 		for(int32 i=0; i<NumSpheres; i++)
 		{
-			if(Spheres[i].BoneIndex < 0) 
+			if(Spheres[i].BoneIndex < 0)
 			{
 				continue;
 			}
@@ -3490,6 +3651,603 @@ void USkeletalMeshComponent::DrawClothingCollisionVolumes(FPrimitiveDrawInterfac
 #endif // #if WITH_APEX_CLOTHING
 }
 
+void USkeletalMeshComponent::DrawClothingFixedVertices(FPrimitiveDrawInterface* PDI)
+{
+#if WITH_APEX_CLOTHING
+
+	if (!SkeletalMesh || !MeshObject)
+	{
+		return;
+	}
+
+	int32 NumActors = ClothingActors.Num();
+
+	TArray<FMatrix> RefToLocals;
+	// get local matrices for skinning
+	UpdateRefToLocalMatrices(RefToLocals, this, GetSkeletalMeshResource(), 0, NULL);
+	const float Inv_255 = 1.f/255.f;
+
+	for(int32 ActorIdx=0; ActorIdx<NumActors; ActorIdx++)
+	{
+		if(!IsValidClothingActor(ActorIdx))
+		{
+			continue;
+		}
+
+		NxClothingActor* ClothingActor = ClothingActors[ActorIdx].ApexClothingActor;
+
+		if (ClothingActor)
+		{
+			uint32 NumSimulVertices = ClothingActor->getNumSimulationVertices();
+
+			if(SkeletalMesh->ClothingAssets[ActorIdx].ClothVisualizationInfos.Num() == 0)
+			{
+				LoadClothingVisualizationInfo(ActorIdx);
+			}
+
+			int32 PhysicalMeshLOD = PredictedLODLevel;
+
+			// if this asset doesn't have current physical mesh LOD, then skip to visualize
+			if(!SkeletalMesh->ClothingAssets[ActorIdx].ClothVisualizationInfos.IsValidIndex(PhysicalMeshLOD))
+			{
+				continue;
+			}
+
+			FClothVisualizationInfo& VisualInfo = SkeletalMesh->ClothingAssets[ActorIdx].ClothVisualizationInfos[PhysicalMeshLOD];				
+
+			NumSimulVertices = VisualInfo.ClothPhysicalMeshVertices.Num();
+
+			uint16 ChunkIndex = 0;
+
+			FStaticLODModel& LODModel = MeshObject->GetSkeletalMeshResource().LODModels[PredictedLODLevel];
+
+			TArray<uint32> SectionIndices;
+			SkeletalMesh->GetClothSectionIndices(PredictedLODLevel, ActorIdx, SectionIndices);
+
+			int32 NumSections = SectionIndices.Num();
+
+			for(int32 SecIdx=0; SecIdx <NumSections; SecIdx++)
+			{
+				uint16 SectionIndex = SectionIndices[SecIdx];
+
+				ChunkIndex = LODModel.Sections[SectionIndex].ChunkIndex;
+
+				FSkelMeshChunk& Chunk = LODModel.Chunks[ChunkIndex];
+
+				const TArray<FVector>& SimulVertices = VisualInfo.ClothPhysicalMeshVertices;
+				const TArray<FVector>& SimulNormals = VisualInfo.ClothPhysicalMeshNormals;
+
+				uint32 NumMppingData = Chunk.ApexClothMappingData.Num();
+
+				FVector Start, End;		
+
+				for(uint32 MappingIndex=0; MappingIndex < NumMppingData; MappingIndex++)
+				{
+					FVector4 BaryCoordPos = Chunk.ApexClothMappingData[MappingIndex].PositionBaryCoordsAndDist;
+					FVector4 BaryCoordNormal = Chunk.ApexClothMappingData[MappingIndex].NormalBaryCoordsAndDist;
+					FVector4 BaryCoordTangent = Chunk.ApexClothMappingData[MappingIndex].TangentBaryCoordsAndDist;
+					uint16*  SimulIndices = Chunk.ApexClothMappingData[MappingIndex].SimulMeshVertIndices;
+
+					bool bFixed = (SimulIndices[3] == 0xFFFF);
+
+					if(!bFixed)
+					{
+						continue;
+					}
+
+					check(SimulIndices[0] < NumSimulVertices && SimulIndices[1] < NumSimulVertices && SimulIndices[2] < NumSimulVertices);
+					
+					FVector a = SimulVertices[SimulIndices[0]];
+					FVector b = SimulVertices[SimulIndices[1]];
+					FVector c = SimulVertices[SimulIndices[2]];
+
+					FVector na = SimulNormals[SimulIndices[0]];
+					FVector nb = SimulNormals[SimulIndices[1]];
+					FVector nc = SimulNormals[SimulIndices[2]];
+
+					FVector Position = 
+						BaryCoordPos.X*(a + BaryCoordPos.W*na)
+						+ BaryCoordPos.Y*(b + BaryCoordPos.W*nb)
+						+ BaryCoordPos.Z*(c + BaryCoordPos.W*nc);
+
+					FSoftSkinVertex& SoftVert = Chunk.SoftVertices[MappingIndex];
+
+					// skinning, APEX clothing supports up to 4 bone influences for now
+					const uint8* BoneIndices = SoftVert.InfluenceBones;
+					const uint8* BoneWeights = SoftVert.InfluenceWeights;
+				
+					FMatrix SkinningMat(ForceInitToZero);
+
+					for(int32 BoneWeightIdx=0; BoneWeightIdx < Chunk.MaxBoneInfluences; BoneWeightIdx++)
+					{
+						float Weight = BoneWeights[BoneWeightIdx]*Inv_255;
+						SkinningMat += RefToLocals[Chunk.BoneMap[BoneIndices[BoneWeightIdx]]]*Weight;
+					}
+
+					FVector SkinnedPosition = SkinningMat.TransformPosition(Position);
+					// draws a skinned fixed vertex
+					PDI->DrawPoint(SkinnedPosition, FColor::Yellow, 2, SDPG_World);
+				}
+			}
+		}
+	}
+
+#endif // #if WITH_APEX_CLOTHING
+}
+
+void USkeletalMeshComponent::LoadClothingVisualizationInfo(int32 AssetIndex)
+{
+#if WITH_APEX_CLOTHING
+
+	if (!SkeletalMesh || !SkeletalMesh->ClothingAssets.IsValidIndex(AssetIndex)) 
+	{
+		return;
+	}
+
+
+	FClothingAssetData& AssetData = SkeletalMesh->ClothingAssets[AssetIndex];
+	NxClothingAsset* ApexClothingAsset = AssetData.ApexClothingAsset->GetAsset();
+	const NxParameterized::Interface* AssetParams = ApexClothingAsset->getAssetNxParameterized();
+
+	int32 NumPhysicalLODs;
+	NxParameterized::getParamArraySize(*AssetParams, "physicalMeshes", NumPhysicalLODs);
+
+	check(NumPhysicalLODs == ApexClothingAsset->getNumGraphicalLodLevels());
+
+	// prepares to load data for each LOD
+	AssetData.ClothVisualizationInfos.Empty(NumPhysicalLODs);
+	AssetData.ClothVisualizationInfos.AddZeroed(NumPhysicalLODs);
+
+	for(int32 LODIndex=0; LODIndex<NumPhysicalLODs; LODIndex++)
+	{
+		FClothVisualizationInfo& VisualInfo = AssetData.ClothVisualizationInfos[LODIndex];
+
+		char ParameterName[MAX_SPRINTF];
+		FCStringAnsi::Sprintf(ParameterName, "physicalMeshes[%d]", LODIndex);
+
+		NxParameterized::Interface* PhysicalMeshParams;
+		uint32 NumVertices = 0;
+		uint32 NumIndices = 0;
+		// physical mesh vertices & normals
+		if (NxParameterized::getParamRef(*AssetParams, ParameterName, PhysicalMeshParams))
+		{
+			if(PhysicalMeshParams != NULL)
+			{
+				verify(NxParameterized::getParamU32(*PhysicalMeshParams, "physicalMesh.numVertices", NumVertices));
+
+				// physical mesh vertices
+				physx::PxI32 VertexCount = 0;
+				if (NxParameterized::getParamArraySize(*PhysicalMeshParams, "physicalMesh.vertices", VertexCount))
+				{
+					check(VertexCount == NumVertices);
+					VisualInfo.ClothPhysicalMeshVertices.Empty(NumVertices);
+
+					for (uint32 VertexIndex = 0; VertexIndex < NumVertices; ++VertexIndex)
+					{
+						FCStringAnsi::Sprintf( ParameterName, "physicalMesh.vertices[%d]", VertexIndex );
+						NxParameterized::Handle MeshVertexHandle(*PhysicalMeshParams);
+						if (NxParameterized::findParam(*PhysicalMeshParams, ParameterName, MeshVertexHandle) != NULL)
+						{
+							physx::PxVec3  Vertex;
+							MeshVertexHandle.getParamVec3(Vertex);
+							VisualInfo.ClothPhysicalMeshVertices.Add(P2UVector(Vertex));
+						}
+					}
+				}
+
+				// bone weights & bone indices
+
+				physx::PxI32 BoneWeightsCount = 0;
+				if (NxParameterized::getParamArraySize(*PhysicalMeshParams, "physicalMesh.boneWeights", BoneWeightsCount))
+				{
+					VisualInfo.ClothPhysicalMeshBoneWeightsInfo.AddZeroed(VertexCount);
+
+					int32 MaxBoneWeights = BoneWeightsCount / VertexCount;
+					VisualInfo.NumMaxBoneInfluences = MaxBoneWeights;
+
+					physx::PxI32 BoneIndicesCount = 0;
+					verify(NxParameterized::getParamArraySize(*PhysicalMeshParams, "physicalMesh.boneIndices", BoneIndicesCount));
+					check(BoneIndicesCount == BoneWeightsCount);
+
+					NxParameterized::Handle BoneWeightHandle(*PhysicalMeshParams);
+					for (int32 VertexIndex = 0; VertexIndex < VertexCount; ++VertexIndex)
+					{
+						for(uint8 WeightIndex=0; WeightIndex < MaxBoneWeights; ++WeightIndex)
+						{
+							uint32 CurBoneWeightIndex = VertexIndex*MaxBoneWeights + WeightIndex;
+
+							NxParameterized::Handle BoneIndexHandle(*PhysicalMeshParams);
+							FCStringAnsi::Sprintf( ParameterName, "physicalMesh.boneIndices[%d]", CurBoneWeightIndex );
+							verify(NxParameterized::findParam(*PhysicalMeshParams, ParameterName, BoneIndexHandle));
+							uint16 BoneIndex;
+							BoneIndexHandle.getParamU16(BoneIndex);
+							VisualInfo.ClothPhysicalMeshBoneWeightsInfo[VertexIndex].Indices[WeightIndex] = BoneIndex;
+
+							NxParameterized::Handle BoneWeightHandle(*PhysicalMeshParams);
+							FCStringAnsi::Sprintf( ParameterName, "physicalMesh.boneWeights[%d]", CurBoneWeightIndex );
+							verify(NxParameterized::findParam(*PhysicalMeshParams, ParameterName, BoneWeightHandle));
+							float BoneWeight;
+							BoneWeightHandle.getParamF32(BoneWeight);
+							VisualInfo.ClothPhysicalMeshBoneWeightsInfo[VertexIndex].Weights[WeightIndex] = BoneWeight;
+						}
+					}
+				}
+
+				// physical mesh normals
+				physx::PxI32 NormalCount = 0;
+				if (NxParameterized::getParamArraySize(*PhysicalMeshParams, "physicalMesh.normals", NormalCount))
+				{
+					check(NormalCount == NumVertices);
+					VisualInfo.ClothPhysicalMeshNormals.Empty(NormalCount);
+
+					for (int32 NormalIndex = 0; NormalIndex < NormalCount; ++NormalIndex)
+					{
+						FCStringAnsi::Sprintf( ParameterName, "physicalMesh.normals[%d]", NormalIndex );
+						NxParameterized::Handle MeshNormalHandle(*PhysicalMeshParams);
+						if (NxParameterized::findParam(*PhysicalMeshParams, ParameterName, MeshNormalHandle) != NULL)
+						{
+							physx::PxVec3  PxNormal;
+							MeshNormalHandle.getParamVec3(PxNormal);
+							VisualInfo.ClothPhysicalMeshNormals.Add(P2UVector(PxNormal));
+						}
+					}
+				}
+				// physical mesh indices
+				verify(NxParameterized::getParamU32(*PhysicalMeshParams, "physicalMesh.numIndices", NumIndices));
+				physx::PxI32 IndexCount = 0;
+				if (NxParameterized::getParamArraySize(*PhysicalMeshParams, "physicalMesh.indices", IndexCount))
+				{
+					check(IndexCount == NumIndices);
+					VisualInfo.ClothPhysicalMeshIndices.Empty(NumIndices);
+
+					for (uint32 IndexIdx = 0; IndexIdx < NumIndices; ++IndexIdx)
+					{
+						FCStringAnsi::Sprintf( ParameterName, "physicalMesh.indices[%d]", IndexIdx );
+						NxParameterized::Handle MeshIndexHandle(*PhysicalMeshParams);
+						if (NxParameterized::findParam(*PhysicalMeshParams, ParameterName, MeshIndexHandle) != NULL)
+						{
+							uint32 IndexParam;
+							MeshIndexHandle.getParamU32(IndexParam);
+							VisualInfo.ClothPhysicalMeshIndices.Add(IndexParam);
+						}
+					}
+				}
+
+				// constraint coefficient parameters (max distances & backstop data)
+				verify(NxParameterized::getParamF32(*PhysicalMeshParams, "physicalMesh.maximumMaxDistance", VisualInfo.MaximumMaxDistance));
+
+				physx::PxI32  ConstraintCoeffCount = 0;
+				if (NxParameterized::getParamArraySize(*PhysicalMeshParams, "physicalMesh.constrainCoefficients", ConstraintCoeffCount))
+				{
+					check(ConstraintCoeffCount == NumVertices);
+					VisualInfo.ClothConstrainCoeffs.Empty(ConstraintCoeffCount);
+					VisualInfo.ClothConstrainCoeffs.AddZeroed(ConstraintCoeffCount);
+
+					for (uint32 ConstCoeffIdx = 0; ConstCoeffIdx < NumIndices; ++ConstCoeffIdx)
+					{
+						// max distances
+						FCStringAnsi::Sprintf( ParameterName, "physicalMesh.constrainCoefficients[%d].maxDistance", ConstCoeffIdx );
+						NxParameterized::Handle MeshConstCoeffHandle(*PhysicalMeshParams);
+						if (NxParameterized::findParam(*PhysicalMeshParams, ParameterName, MeshConstCoeffHandle) != NULL)
+						{
+							float MaxDistance;
+							MeshConstCoeffHandle.getParamF32(MaxDistance);
+							VisualInfo.ClothConstrainCoeffs[ConstCoeffIdx].ClothMaxDistance = MaxDistance;
+						}
+
+						// backstop data
+						FCStringAnsi::Sprintf( ParameterName, "physicalMesh.constrainCoefficients[%d].collisionSphereRadius", ConstCoeffIdx );
+						if (NxParameterized::findParam(*PhysicalMeshParams, ParameterName, MeshConstCoeffHandle) != NULL)
+						{
+							float BackstopCollisionSphereRadius;
+							MeshConstCoeffHandle.getParamF32(BackstopCollisionSphereRadius);
+							VisualInfo.ClothConstrainCoeffs[ConstCoeffIdx].ClothBackstopRadius = BackstopCollisionSphereRadius;
+						}
+
+						FCStringAnsi::Sprintf( ParameterName, "physicalMesh.constrainCoefficients[%d].collisionSphereDistance", ConstCoeffIdx );
+						if (NxParameterized::findParam(*PhysicalMeshParams, ParameterName, MeshConstCoeffHandle) != NULL)
+						{
+							float BackstopCollisionSphereDistance;
+							MeshConstCoeffHandle.getParamF32(BackstopCollisionSphereDistance);
+							VisualInfo.ClothConstrainCoeffs[ConstCoeffIdx].ClothBackstopDistance = BackstopCollisionSphereDistance;
+						}
+					}
+				}
+			}
+		}
+	}
+
+#endif // #if WITH_APEX_CLOTHING
+}
+
+void USkeletalMeshComponent::LoadAllClothingVisualizationInfos()
+{
+#if WITH_APEX_CLOTHING
+
+	if (!SkeletalMesh || SkeletalMesh->ClothingAssets.Num() == 0) 
+	{
+		return;
+	}
+
+	int32 NumAssets = SkeletalMesh->ClothingAssets.Num();
+
+	for(int32 AssetIdx=0; AssetIdx<NumAssets; AssetIdx++)
+	{
+		LoadClothingVisualizationInfo(AssetIdx);
+	}
+#endif // #if WITH_APEX_CLOTHING
+}
+
+void USkeletalMeshComponent::DrawClothingMaxDistances(FPrimitiveDrawInterface* PDI)
+{
+#if WITH_APEX_CLOTHING
+	if (!SkeletalMesh
+		|| SkeletalMesh->ClothingAssets.Num() == 0) 
+	{
+		return;
+	}
+
+	int32 PhysicalMeshLOD = PredictedLODLevel;
+
+	int32 NumAssets = SkeletalMesh->ClothingAssets.Num();
+
+	for(int32 AssetIdx=0; AssetIdx < NumAssets; AssetIdx++)
+	{
+		if(SkeletalMesh->ClothingAssets[AssetIdx].ClothVisualizationInfos.Num() == 0)
+		{
+			LoadClothingVisualizationInfo(AssetIdx);
+		}
+
+		// if this asset doesn't have current physical mesh LOD, then skip to visualize
+		if(!SkeletalMesh->ClothingAssets[AssetIdx].ClothVisualizationInfos.IsValidIndex(PhysicalMeshLOD))
+		{
+			continue;
+		}
+
+		FClothVisualizationInfo& VisualInfo = SkeletalMesh->ClothingAssets[AssetIdx].ClothVisualizationInfos[PhysicalMeshLOD];
+
+		int32 NumMaxDistances = VisualInfo.ClothConstrainCoeffs.Num();
+		for(int32 MaxDistIdx=0; MaxDistIdx < NumMaxDistances; MaxDistIdx++)
+		{
+			float MaxDistance = VisualInfo.ClothConstrainCoeffs[MaxDistIdx].ClothMaxDistance;
+			if(MaxDistance > 0.0f)
+			{
+				FVector LineStart = VisualInfo.ClothPhysicalMeshVertices[MaxDistIdx];
+				FVector LineEnd = LineStart + (VisualInfo.ClothPhysicalMeshNormals[MaxDistIdx] * MaxDistance);
+
+				// calculates gray scale for this line color
+				uint8 GrayLevel = (uint8)((MaxDistance / VisualInfo.MaximumMaxDistance)*255);
+				PDI->DrawLine(LineStart, LineEnd, FColor(GrayLevel, GrayLevel, GrayLevel), SDPG_World);
+			}
+			else // fixed vertices
+			{
+				FVector FixedPoint = VisualInfo.ClothPhysicalMeshVertices[MaxDistIdx];
+				PDI->DrawPoint(FixedPoint, FColor::Blue, 2, SDPG_World);
+			}
+		}
+	}	
+
+#endif // #if WITH_APEX_CLOTHING
+}
+
+void USkeletalMeshComponent::DrawClothingBackstops(FPrimitiveDrawInterface* PDI)
+{
+#if WITH_APEX_CLOTHING
+	if (!SkeletalMesh
+		|| SkeletalMesh->ClothingAssets.Num() == 0) 
+	{
+		return;
+	}
+
+	int32 PhysicalMeshLOD = PredictedLODLevel;
+
+	int32 NumAssets = SkeletalMesh->ClothingAssets.Num();
+
+	for(int32 AssetIdx=0; AssetIdx < NumAssets; AssetIdx++)
+	{
+		if(SkeletalMesh->ClothingAssets[AssetIdx].ClothVisualizationInfos.Num() == 0)
+		{
+			LoadClothingVisualizationInfo(AssetIdx);
+		}
+
+		// if this asset doesn't have current physical mesh LOD, then skip to visualize
+		if(!SkeletalMesh->ClothingAssets[AssetIdx].ClothVisualizationInfos.IsValidIndex(PhysicalMeshLOD))
+		{
+			continue;
+		}
+
+		FClothVisualizationInfo& VisualInfo = SkeletalMesh->ClothingAssets[AssetIdx].ClothVisualizationInfos[PhysicalMeshLOD];
+
+		int32 NumBackstopSpheres = VisualInfo.ClothConstrainCoeffs.Num();
+		for(int32 BackstopIdx=0; BackstopIdx < NumBackstopSpheres; BackstopIdx++)
+		{
+			float Distance = VisualInfo.ClothConstrainCoeffs[BackstopIdx].ClothBackstopDistance;
+			float MaxDistance = VisualInfo.ClothConstrainCoeffs[BackstopIdx].ClothMaxDistance;
+
+			FColor FixedColor = FColor::White;
+			if(Distance > MaxDistance) // Backstop is disabled if the value is bigger than its Max Distance value
+			{
+				Distance = 0.0f;
+				FixedColor = FColor::Black;
+			}
+
+			if(Distance > 0.0f)
+			{
+				FVector LineStart = VisualInfo.ClothPhysicalMeshVertices[BackstopIdx];
+				FVector LineEnd = LineStart + (VisualInfo.ClothPhysicalMeshNormals[BackstopIdx] * Distance);
+				PDI->DrawLine(LineStart, LineEnd, FColor::Red, SDPG_World);
+			}
+			else
+			if(Distance < 0.0f)
+			{
+				FVector LineStart = VisualInfo.ClothPhysicalMeshVertices[BackstopIdx];
+				FVector LineEnd = LineStart + (VisualInfo.ClothPhysicalMeshNormals[BackstopIdx] * Distance);
+				PDI->DrawLine(LineStart, LineEnd, FColor::Blue, SDPG_World);
+			}
+			else // White means collision distance is set to 0
+			{
+				FVector FixedPoint = VisualInfo.ClothPhysicalMeshVertices[BackstopIdx];
+				PDI->DrawPoint(FixedPoint, FixedColor, 2, SDPG_World);
+			}
+		}
+	}
+
+#endif // #if WITH_APEX_CLOTHING
+}
+
+void USkeletalMeshComponent::DrawClothingPhysicalMeshWire(FPrimitiveDrawInterface* PDI)
+{
+#if WITH_APEX_CLOTHING
+	if (!SkeletalMesh
+		|| SkeletalMesh->ClothingAssets.Num() == 0) 
+	{
+		return;
+	}
+
+	int32 PhysicalMeshLOD = PredictedLODLevel;
+
+	int32 NumAssets = SkeletalMesh->ClothingAssets.Num();
+
+	TArray<FMatrix> RefToLocals;
+	// get local matrices for skinning
+	UpdateRefToLocalMatrices(RefToLocals, this, GetSkeletalMeshResource(), 0, NULL);
+
+	for(int32 AssetIdx=0; AssetIdx < NumAssets; AssetIdx++)
+	{
+		if(SkeletalMesh->ClothingAssets[AssetIdx].ClothVisualizationInfos.Num() == 0)
+		{
+			LoadClothingVisualizationInfo(AssetIdx);
+		}
+
+		// if this asset doesn't have current physical mesh LOD, then skip to visualize
+		if(!SkeletalMesh->ClothingAssets[AssetIdx].ClothVisualizationInfos.IsValidIndex(PhysicalMeshLOD))
+		{
+			continue;
+		}
+
+		FClothVisualizationInfo& VisualInfo = SkeletalMesh->ClothingAssets[AssetIdx].ClothVisualizationInfos[PhysicalMeshLOD];
+
+		uint32 NumPhysicalMeshVerts = VisualInfo.ClothPhysicalMeshVertices.Num();
+		const physx::PxVec3* SimulVertices = NULL;
+
+		bool bUseSimulatedResult = false;
+		// skinning for fixed vertices
+		TArray<FVector> SimulatedPhysicalMeshVertices;
+		// if cloth simulation is enabled and there exist a clothing actor and simulation vertices, 
+		// then draws wire frame using simulated vertices
+		if(!bDisableClothSimulation && ClothingActors.IsValidIndex(AssetIdx))
+		{
+			NxClothingActor* ClothingActor = ClothingActors[AssetIdx].ApexClothingActor;
+
+			if(ClothingActor)
+			{
+				uint32 NumSimulVertices = ClothingActor->getNumSimulationVertices();
+
+				if(NumSimulVertices > 0)
+				{
+					bUseSimulatedResult = true;
+					SimulatedPhysicalMeshVertices.AddUninitialized(NumPhysicalMeshVerts);
+
+					const physx::PxVec3* SimulVertices = ClothingActor->getSimulationPositions();
+					for(uint32 SimulVertIdx=0; SimulVertIdx < NumSimulVertices; SimulVertIdx++)
+					{
+						SimulatedPhysicalMeshVertices[SimulVertIdx] = P2UVector(SimulVertices[SimulVertIdx]);
+					}
+
+					// skinning for fixed vertices
+					for(uint32 FixedVertIdx=NumSimulVertices; FixedVertIdx < NumPhysicalMeshVerts; FixedVertIdx++)
+					{
+						FMatrix SkinningMat(ForceInitToZero);
+
+						for(uint32 BoneWeightIdx=0; BoneWeightIdx < VisualInfo.NumMaxBoneInfluences; BoneWeightIdx++)
+						{
+							uint16 ApexBoneIndex = VisualInfo.ClothPhysicalMeshBoneWeightsInfo[FixedVertIdx].Indices[BoneWeightIdx];
+
+							FName BoneName = SkeletalMesh->ClothingAssets[AssetIdx].ApexClothingAsset->GetConvertedBoneName(ApexBoneIndex);
+
+							int32 BoneIndex = GetBoneIndex(BoneName);
+
+							if(BoneIndex < 0)
+							{
+								continue;
+							}
+
+							float Weight = VisualInfo.ClothPhysicalMeshBoneWeightsInfo[FixedVertIdx].Weights[BoneWeightIdx];
+							SkinningMat += RefToLocals[BoneIndex]*Weight;
+						}
+
+						SimulatedPhysicalMeshVertices[FixedVertIdx] = SkinningMat.TransformPosition(VisualInfo.ClothPhysicalMeshVertices[FixedVertIdx]);
+					}
+				}
+			}
+		}
+
+		TArray<FVector>* PhysicalMeshVertices = &VisualInfo.ClothPhysicalMeshVertices;
+		if(bUseSimulatedResult)
+		{
+			PhysicalMeshVertices = &SimulatedPhysicalMeshVertices;
+		}
+
+		uint32 NumIndices = VisualInfo.ClothPhysicalMeshIndices.Num();
+		check(NumIndices % 3 == 0); // check triangles count
+
+		for(uint32 IndexIdx=0; IndexIdx < NumIndices; IndexIdx+=3)
+		{
+			// draw a triangle
+			uint32 Index0 = VisualInfo.ClothPhysicalMeshIndices[IndexIdx];
+			uint32 Index1 = VisualInfo.ClothPhysicalMeshIndices[IndexIdx + 1];
+			uint32 Index2 = VisualInfo.ClothPhysicalMeshIndices[IndexIdx + 2];
+
+			// If index is greater than Num of vertices, then skip 
+			if(Index0 >= NumPhysicalMeshVerts 
+			|| Index1 >= NumPhysicalMeshVerts 
+			|| Index2 >= NumPhysicalMeshVerts)
+			{
+				continue;
+			}
+
+			FVector V[3];
+			V[0] = (*PhysicalMeshVertices)[Index0];
+			V[1] = (*PhysicalMeshVertices)[Index1];
+			V[2] = (*PhysicalMeshVertices)[Index2];
+
+			float MaxDists[3];
+			MaxDists[0] = VisualInfo.ClothConstrainCoeffs[Index0].ClothMaxDistance;
+			MaxDists[1] = VisualInfo.ClothConstrainCoeffs[Index1].ClothMaxDistance;
+			MaxDists[2] = VisualInfo.ClothConstrainCoeffs[Index2].ClothMaxDistance;
+
+			for(int32 i=0; i<3; i++)
+			{
+				uint32 Index0 = i;
+				uint32 Index1 = i+1;
+				if(Index1 >= 3)
+				{
+					Index1 = 0;
+				}
+
+				// calculates gray scaled color
+				uint8 GrayLevel0 = (uint8)((MaxDists[Index0] / VisualInfo.MaximumMaxDistance)*255.0f);
+				uint8 GrayLevel1 = (uint8)((MaxDists[Index1] / VisualInfo.MaximumMaxDistance)*255.0f);
+
+				uint8 GrayMidColor = (uint8)(((uint32)GrayLevel0 + (uint32)GrayLevel1)*0.5);
+				FColor LineColor(GrayMidColor, GrayMidColor, GrayMidColor);
+				if(GrayMidColor == 0)
+				{
+					LineColor = FColor::Magenta;
+				}
+				else
+				{
+					LineColor = FColor::White;
+				}
+
+				PDI->DrawLine(V[Index0], V[Index1], LineColor, SDPG_World);
+			}
+		}
+	}
+
+#endif // #if WITH_APEX_CLOTHING
+}
+
 float USkeletalMeshComponent::GetMass() const
 {
 	if (bUseSingleBodyPhysics)
