#ttp 303393 UObject improvement: Remove header groups
#proj core
#branch UE4
#summary Removed header groups, removed some redundant includes and replaced with forward declarations. Fixed compilation errors.

diff --git a/Engine/Source/Developer/MovieSceneTools/Private/K2Node_PlayMovieScene.cpp b/Engine/Source/Developer/MovieSceneTools/Private/K2Node_PlayMovieScene.cpp
--- a/Engine/Source/Developer/MovieSceneTools/Private/K2Node_PlayMovieScene.cpp
+++ b/Engine/Source/Developer/MovieSceneTools/Private/K2Node_PlayMovieScene.cpp
@@ -4,7 +4,6 @@
 #include "ScopedTransaction.h"
 #include "K2Node_PlayMovieScene.h"
 #include "BlueprintUtilities.h"
-#include "EngineLevelScriptClasses.h"		// For ALevelScriptActor
 #include "MovieScene.h"
 #include "Kismet2/BlueprintEditorUtils.h"
 #include "MovieSceneBindings.h"
diff --git a/Engine/Source/Editor/AnimGraph/Classes/AnimGraphNode_Base.h b/Engine/Source/Editor/AnimGraph/Classes/AnimGraphNode_Base.h
--- a/Engine/Source/Editor/AnimGraph/Classes/AnimGraphNode_Base.h
+++ b/Engine/Source/Editor/AnimGraph/Classes/AnimGraphNode_Base.h
@@ -1,10 +1,13 @@
 // Copyright 1998-2014 Epic Games, Inc. All Rights Reserved.
 
 #pragma once
-
+#include "K2Node.h"
 #include "AnimGraphNode_Base.generated.h"
 
 //
+// Forward declarations.
+//
+struct FEdGraphSchemaAction_K2NewNode;
 
 struct FPoseLinkMappingRecord
 {
diff --git a/Engine/Source/Editor/AnimGraph/Classes/AnimGraphNode_BlendListByInt.h b/Engine/Source/Editor/AnimGraph/Classes/AnimGraphNode_BlendListByInt.h
--- a/Engine/Source/Editor/AnimGraph/Classes/AnimGraphNode_BlendListByInt.h
+++ b/Engine/Source/Editor/AnimGraph/Classes/AnimGraphNode_BlendListByInt.h
@@ -1,6 +1,7 @@
 // Copyright 1998-2014 Epic Games, Inc. All Rights Reserved.
 
 #pragma once
+#include "AnimGraphNode_BlendListBase.h"
 #include "AnimGraphNode_BlendListByInt.generated.h"
 
 UCLASS(MinimalAPI)
diff --git a/Engine/Source/Editor/AnimGraph/Classes/AnimGraphNode_BlendSpaceEvaluator.h b/Engine/Source/Editor/AnimGraph/Classes/AnimGraphNode_BlendSpaceEvaluator.h
--- a/Engine/Source/Editor/AnimGraph/Classes/AnimGraphNode_BlendSpaceEvaluator.h
+++ b/Engine/Source/Editor/AnimGraph/Classes/AnimGraphNode_BlendSpaceEvaluator.h
@@ -1,7 +1,7 @@
 // Copyright 1998-2014 Epic Games, Inc. All Rights Reserved.
 
 #pragma once
-
+#include "AnimGraphNode_BlendSpaceBase.h"
 #include "Animation/AnimNode_BlendSpaceEvaluator.h"
 #include "AnimGraphNode_BlendSpaceEvaluator.generated.h"
 
diff --git a/Engine/Source/Editor/AnimGraph/Classes/AnimGraphNode_CopyBone.h b/Engine/Source/Editor/AnimGraph/Classes/AnimGraphNode_CopyBone.h
--- a/Engine/Source/Editor/AnimGraph/Classes/AnimGraphNode_CopyBone.h
+++ b/Engine/Source/Editor/AnimGraph/Classes/AnimGraphNode_CopyBone.h
@@ -1,6 +1,7 @@
 // Copyright 1998-2014 Epic Games, Inc. All Rights Reserved.
 
 #pragma once
+#include "AnimGraphNode_SkeletalControlBase.h"
 #include "AnimGraphNode_CopyBone.generated.h"
 
 UCLASS(MinimalAPI)
diff --git a/Engine/Source/Editor/AnimGraph/Classes/AnimGraphNode_Fabrik.h b/Engine/Source/Editor/AnimGraph/Classes/AnimGraphNode_Fabrik.h
--- a/Engine/Source/Editor/AnimGraph/Classes/AnimGraphNode_Fabrik.h
+++ b/Engine/Source/Editor/AnimGraph/Classes/AnimGraphNode_Fabrik.h
@@ -1,7 +1,7 @@
 // Copyright 1998-2013 Epic Games, Inc. All Rights Reserved.
 
 #pragma once
-
+#include "AnimGraphNode_SkeletalControlBase.h"
 #include "AnimGraphNode_Fabrik.generated.h"
 
 // Editor node for FABRIK IK skeletal controller
diff --git a/Engine/Source/Editor/AnimGraph/Classes/AnimGraphNode_IdentityPose.h b/Engine/Source/Editor/AnimGraph/Classes/AnimGraphNode_IdentityPose.h
--- a/Engine/Source/Editor/AnimGraph/Classes/AnimGraphNode_IdentityPose.h
+++ b/Engine/Source/Editor/AnimGraph/Classes/AnimGraphNode_IdentityPose.h
@@ -1,6 +1,7 @@
 // Copyright 1998-2014 Epic Games, Inc. All Rights Reserved.
 
 #pragma once
+#include "AnimGraphNode_RefPoseBase.h"
 #include "AnimGraphNode_IdentityPose.generated.h"
 
 UCLASS(MinimalAPI)
diff --git a/Engine/Source/Editor/AnimGraph/Classes/AnimGraphNode_LocalRefPose.h b/Engine/Source/Editor/AnimGraph/Classes/AnimGraphNode_LocalRefPose.h
--- a/Engine/Source/Editor/AnimGraph/Classes/AnimGraphNode_LocalRefPose.h
+++ b/Engine/Source/Editor/AnimGraph/Classes/AnimGraphNode_LocalRefPose.h
@@ -1,6 +1,7 @@
 // Copyright 1998-2014 Epic Games, Inc. All Rights Reserved.
 
 #pragma once
+#include "AnimGraphNode_RefPoseBase.h"
 #include "AnimGraphNode_LocalRefPose.generated.h"
 
 UCLASS(MinimalAPI)
diff --git a/Engine/Source/Editor/AnimGraph/Classes/AnimGraphNode_RotationMultiplier.h b/Engine/Source/Editor/AnimGraph/Classes/AnimGraphNode_RotationMultiplier.h
--- a/Engine/Source/Editor/AnimGraph/Classes/AnimGraphNode_RotationMultiplier.h
+++ b/Engine/Source/Editor/AnimGraph/Classes/AnimGraphNode_RotationMultiplier.h
@@ -1,6 +1,7 @@
 // Copyright 1998-2014 Epic Games, Inc. All Rights Reserved.
 
 #pragma once
+#include "AnimGraphNode_SkeletalControlBase.h"
 #include "AnimGraphNode_RotationMultiplier.generated.h"
 
 UCLASS(MinimalAPI)
diff --git a/Engine/Source/Editor/AnimGraph/Classes/AnimGraphNode_StateMachine.h b/Engine/Source/Editor/AnimGraph/Classes/AnimGraphNode_StateMachine.h
--- a/Engine/Source/Editor/AnimGraph/Classes/AnimGraphNode_StateMachine.h
+++ b/Engine/Source/Editor/AnimGraph/Classes/AnimGraphNode_StateMachine.h
@@ -1,6 +1,7 @@
 // Copyright 1998-2014 Epic Games, Inc. All Rights Reserved.
 
 #pragma once
+#include "AnimGraphNode_StateMachineBase.h"
 #include "AnimGraphNode_StateMachine.generated.h"
 
 UCLASS(MinimalAPI)
diff --git a/Engine/Source/Editor/AnimGraph/Classes/AnimPreviewInstance.h b/Engine/Source/Editor/AnimGraph/Classes/AnimPreviewInstance.h
--- a/Engine/Source/Editor/AnimGraph/Classes/AnimPreviewInstance.h
+++ b/Engine/Source/Editor/AnimGraph/Classes/AnimPreviewInstance.h
@@ -4,6 +4,7 @@
 
 #include "BlueprintGraphDefinitions.h"
 #include "AnimGraphNode_ModifyBone.h"
+#include "Animation/AnimSingleNodeInstance.h"
 #include "AnimPreviewInstance.generated.h"
 
 /** Enum to know how montage is being played */
diff --git a/Engine/Source/Editor/AnimGraph/Classes/AnimStateTransitionNode.h b/Engine/Source/Editor/AnimGraph/Classes/AnimStateTransitionNode.h
--- a/Engine/Source/Editor/AnimGraph/Classes/AnimStateTransitionNode.h
+++ b/Engine/Source/Editor/AnimGraph/Classes/AnimStateTransitionNode.h
@@ -1,9 +1,9 @@
 // Copyright 1998-2014 Epic Games, Inc. All Rights Reserved.
 
 #pragma once
+#include "AnimStateNodeBase.h"
 #include "AnimStateTransitionNode.generated.h"
 
-
 UCLASS(MinimalAPI, DependsOn=UAnimInstance, config=Editor)
 class UAnimStateTransitionNode : public UAnimStateNodeBase
 {
diff --git a/Engine/Source/Editor/AnimGraph/Classes/AnimationCustomTransitionSchema.h b/Engine/Source/Editor/AnimGraph/Classes/AnimationCustomTransitionSchema.h
--- a/Engine/Source/Editor/AnimGraph/Classes/AnimationCustomTransitionSchema.h
+++ b/Engine/Source/Editor/AnimGraph/Classes/AnimationCustomTransitionSchema.h
@@ -1,6 +1,7 @@
 // Copyright 1998-2014 Epic Games, Inc. All Rights Reserved.
 
 #pragma once
+#include "AnimationGraphSchema.h"
 #include "AnimationCustomTransitionSchema.generated.h"
 
 UCLASS(MinimalAPI)
diff --git a/Engine/Source/Editor/AnimGraph/Classes/AnimationStateMachineGraph.h b/Engine/Source/Editor/AnimGraph/Classes/AnimationStateMachineGraph.h
--- a/Engine/Source/Editor/AnimGraph/Classes/AnimationStateMachineGraph.h
+++ b/Engine/Source/Editor/AnimGraph/Classes/AnimationStateMachineGraph.h
@@ -1,7 +1,7 @@
 // Copyright 1998-2014 Epic Games, Inc. All Rights Reserved.
 
 #pragma once
-
+#include "EdGraph/EdGraph.h"
 #include "AnimationStateMachineGraph.generated.h"
 
 UCLASS(MinimalAPI)
diff --git a/Engine/Source/Editor/AnimGraph/Private/AnimGraphNode_ApplyAdditive.cpp b/Engine/Source/Editor/AnimGraph/Private/AnimGraphNode_ApplyAdditive.cpp
--- a/Engine/Source/Editor/AnimGraph/Private/AnimGraphNode_ApplyAdditive.cpp
+++ b/Engine/Source/Editor/AnimGraph/Private/AnimGraphNode_ApplyAdditive.cpp
@@ -4,6 +4,7 @@
 
 #include "GraphEditorActions.h"
 #include "ScopedTransaction.h"
+#include "AnimGraphNode_ApplyAdditive.h"
 
 /////////////////////////////////////////////////////
 // UAnimGraphNode_ApplyAdditive
diff --git a/Engine/Source/Editor/AnimGraph/Private/AnimGraphNode_BlendListByBool.cpp b/Engine/Source/Editor/AnimGraph/Private/AnimGraphNode_BlendListByBool.cpp
--- a/Engine/Source/Editor/AnimGraph/Private/AnimGraphNode_BlendListByBool.cpp
+++ b/Engine/Source/Editor/AnimGraph/Private/AnimGraphNode_BlendListByBool.cpp
@@ -4,6 +4,7 @@
 
 #include "GraphEditorActions.h"
 #include "ScopedTransaction.h"
+#include "AnimGraphNode_BlendListByBool.h"
 
 /////////////////////////////////////////////////////
 // UAnimGraphNode_BlendListByBool
diff --git a/Engine/Source/Editor/AnimGraph/Private/AnimGraphNode_BlendListByEnum.cpp b/Engine/Source/Editor/AnimGraph/Private/AnimGraphNode_BlendListByEnum.cpp
--- a/Engine/Source/Editor/AnimGraph/Private/AnimGraphNode_BlendListByEnum.cpp
+++ b/Engine/Source/Editor/AnimGraph/Private/AnimGraphNode_BlendListByEnum.cpp
@@ -6,6 +6,7 @@
 #include "ScopedTransaction.h"
 #include "CompilerResultsLog.h"
 #include "K2ActionMenuBuilder.h" // for FK2ActionMenuBuilder::AddNewNodeAction()
+#include "AnimGraphNode_BlendListByEnum.h"
 
 #define LOCTEXT_NAMESPACE "BlendListByEnum"
 
diff --git a/Engine/Source/Editor/AnimGraph/Private/AnimGraphNode_BlendListByInt.cpp b/Engine/Source/Editor/AnimGraph/Private/AnimGraphNode_BlendListByInt.cpp
--- a/Engine/Source/Editor/AnimGraph/Private/AnimGraphNode_BlendListByInt.cpp
+++ b/Engine/Source/Editor/AnimGraph/Private/AnimGraphNode_BlendListByInt.cpp
@@ -5,6 +5,7 @@
 
 #include "GraphEditorActions.h"
 #include "ScopedTransaction.h"
+#include "AnimGraphNode_BlendListByInt.h"
 
 /////////////////////////////////////////////////////
 // UAnimGraphNode_BlendListByInt
diff --git a/Engine/Source/Editor/AnimGraph/Private/AnimGraphNode_CopyBone.cpp b/Engine/Source/Editor/AnimGraph/Private/AnimGraphNode_CopyBone.cpp
--- a/Engine/Source/Editor/AnimGraph/Private/AnimGraphNode_CopyBone.cpp
+++ b/Engine/Source/Editor/AnimGraph/Private/AnimGraphNode_CopyBone.cpp
@@ -1,6 +1,7 @@
 // Copyright 1998-2014 Epic Games, Inc. All Rights Reserved.
 
 #include "AnimGraphPrivatePCH.h"
+#include "AnimGraphNode_CopyBone.h"
 
 /////////////////////////////////////////////////////
 // UAnimGraphNode_CopyBoneSkeletalControl
diff --git a/Engine/Source/Editor/AnimGraph/Private/AnimGraphNode_Fabrik.cpp b/Engine/Source/Editor/AnimGraph/Private/AnimGraphNode_Fabrik.cpp
--- a/Engine/Source/Editor/AnimGraph/Private/AnimGraphNode_Fabrik.cpp
+++ b/Engine/Source/Editor/AnimGraph/Private/AnimGraphNode_Fabrik.cpp
@@ -1,6 +1,7 @@
 // Copyright 1998-2013 Epic Games, Inc. All Rights Reserved.
 
 #include "AnimGraphPrivatePCH.h"
+#include "AnimGraphNode_Fabrik.h"
 
 /////////////////////////////////////////////////////
 // UAnimGraphNode_Fabrik 
diff --git a/Engine/Source/Editor/AnimGraph/Private/AnimGraphNode_IdentityPose.cpp b/Engine/Source/Editor/AnimGraph/Private/AnimGraphNode_IdentityPose.cpp
--- a/Engine/Source/Editor/AnimGraph/Private/AnimGraphNode_IdentityPose.cpp
+++ b/Engine/Source/Editor/AnimGraph/Private/AnimGraphNode_IdentityPose.cpp
@@ -1,6 +1,7 @@
 // Copyright 1998-2014 Epic Games, Inc. All Rights Reserved.
 
 #include "AnimGraphPrivatePCH.h"
+#include "AnimGraphNode_IdentityPose.h"
 
 /////////////////////////////////////////////////////
 // UAnimGraphNode_IdentityPose
diff --git a/Engine/Source/Editor/AnimGraph/Private/AnimGraphNode_LayeredBoneBlend.cpp b/Engine/Source/Editor/AnimGraph/Private/AnimGraphNode_LayeredBoneBlend.cpp
--- a/Engine/Source/Editor/AnimGraph/Private/AnimGraphNode_LayeredBoneBlend.cpp
+++ b/Engine/Source/Editor/AnimGraph/Private/AnimGraphNode_LayeredBoneBlend.cpp
@@ -4,6 +4,7 @@
 
 #include "GraphEditorActions.h"
 #include "ScopedTransaction.h"
+#include "AnimGraphNode_LayeredBoneBlend.h"
 
 /////////////////////////////////////////////////////
 // UAnimGraphNode_LayeredBoneBlend
diff --git a/Engine/Source/Editor/AnimGraph/Private/AnimGraphNode_LocalRefPose.cpp b/Engine/Source/Editor/AnimGraph/Private/AnimGraphNode_LocalRefPose.cpp
--- a/Engine/Source/Editor/AnimGraph/Private/AnimGraphNode_LocalRefPose.cpp
+++ b/Engine/Source/Editor/AnimGraph/Private/AnimGraphNode_LocalRefPose.cpp
@@ -1,6 +1,7 @@
 // Copyright 1998-2014 Epic Games, Inc. All Rights Reserved.
 
 #include "AnimGraphPrivatePCH.h"
+#include "AnimGraphNode_LocalRefPose.h"
 
 /////////////////////////////////////////////////////
 // UAnimGraphNode_LocalRefPose
diff --git a/Engine/Source/Editor/AnimGraph/Private/AnimGraphNode_MeshRefPose.cpp b/Engine/Source/Editor/AnimGraph/Private/AnimGraphNode_MeshRefPose.cpp
--- a/Engine/Source/Editor/AnimGraph/Private/AnimGraphNode_MeshRefPose.cpp
+++ b/Engine/Source/Editor/AnimGraph/Private/AnimGraphNode_MeshRefPose.cpp
@@ -1,6 +1,7 @@
 // Copyright 1998-2014 Epic Games, Inc. All Rights Reserved.
 
 #include "AnimGraphPrivatePCH.h"
+#include "AnimGraphNode_MeshRefPose.h"
 
 /////////////////////////////////////////////////////
 // UAnimGraphNode_MeshRefPose
diff --git a/Engine/Source/Editor/AnimGraph/Private/AnimGraphNode_ModifyBone.cpp b/Engine/Source/Editor/AnimGraph/Private/AnimGraphNode_ModifyBone.cpp
--- a/Engine/Source/Editor/AnimGraph/Private/AnimGraphNode_ModifyBone.cpp
+++ b/Engine/Source/Editor/AnimGraph/Private/AnimGraphNode_ModifyBone.cpp
@@ -1,6 +1,7 @@
 // Copyright 1998-2014 Epic Games, Inc. All Rights Reserved.
 
 #include "AnimGraphPrivatePCH.h"
+#include "AnimGraphNode_ModifyBone.h"
 
 /////////////////////////////////////////////////////
 // UAnimGraphNode_ModifyBone
diff --git a/Engine/Source/Editor/AnimGraph/Private/AnimGraphNode_RotateRootBone.cpp b/Engine/Source/Editor/AnimGraph/Private/AnimGraphNode_RotateRootBone.cpp
--- a/Engine/Source/Editor/AnimGraph/Private/AnimGraphNode_RotateRootBone.cpp
+++ b/Engine/Source/Editor/AnimGraph/Private/AnimGraphNode_RotateRootBone.cpp
@@ -4,6 +4,7 @@
 
 #include "GraphEditorActions.h"
 #include "ScopedTransaction.h"
+#include "AnimGraphNode_RotateRootBone.h"
 
 /////////////////////////////////////////////////////
 // UAnimGraphNode_RotateRootBone
diff --git a/Engine/Source/Editor/AnimGraph/Private/AnimGraphNode_RotationMultiplier.cpp b/Engine/Source/Editor/AnimGraph/Private/AnimGraphNode_RotationMultiplier.cpp
--- a/Engine/Source/Editor/AnimGraph/Private/AnimGraphNode_RotationMultiplier.cpp
+++ b/Engine/Source/Editor/AnimGraph/Private/AnimGraphNode_RotationMultiplier.cpp
@@ -1,7 +1,7 @@
 // Copyright 1998-2014 Epic Games, Inc. All Rights Reserved.
 
 #include "AnimGraphPrivatePCH.h"
-
+#include "AnimGraphNode_RotationMultiplier.h"
 /////////////////////////////////////////////////////
 // UAnimGraphNode_RotationMultiplier
 
diff --git a/Engine/Source/Editor/AnimGraph/Private/AnimGraphNode_SaveCachedPose.cpp b/Engine/Source/Editor/AnimGraph/Private/AnimGraphNode_SaveCachedPose.cpp
--- a/Engine/Source/Editor/AnimGraph/Private/AnimGraphNode_SaveCachedPose.cpp
+++ b/Engine/Source/Editor/AnimGraph/Private/AnimGraphNode_SaveCachedPose.cpp
@@ -6,6 +6,7 @@
 #include "ScopedTransaction.h"
 #include "Kismet2NameValidators.h"
 #include "K2ActionMenuBuilder.h" // for FK2ActionMenuBuilder::AddNewNodeAction()
+#include "AnimGraphNode_SaveCachedPose.h"
 
 /////////////////////////////////////////////////////
 // FCachedPoseNameValidator
@@ -19,9 +20,8 @@ public:
 		TArray<UAnimGraphNode_SaveCachedPose*> Nodes;
 		FBlueprintEditorUtils::GetAllNodesOfClass<UAnimGraphNode_SaveCachedPose>(InBlueprint, Nodes);
 
-		for (auto NodeIt = Nodes.CreateIterator(); NodeIt; ++NodeIt)
+		for (auto Node : Nodes)
 		{
-			auto Node = *NodeIt;
 			Names.Add(Node->CacheName);
 		}
 	}
diff --git a/Engine/Source/Editor/AnimGraph/Private/AnimGraphNode_SequenceEvaluator.cpp b/Engine/Source/Editor/AnimGraph/Private/AnimGraphNode_SequenceEvaluator.cpp
--- a/Engine/Source/Editor/AnimGraph/Private/AnimGraphNode_SequenceEvaluator.cpp
+++ b/Engine/Source/Editor/AnimGraph/Private/AnimGraphNode_SequenceEvaluator.cpp
@@ -4,6 +4,7 @@
 
 #include "CompilerResultsLog.h"
 #include "GraphEditorActions.h"
+#include "AnimGraphNode_SequenceEvaluator.h"
 
 /////////////////////////////////////////////////////
 // UAnimGraphNode_SequenceEvaluator
diff --git a/Engine/Source/Editor/AnimGraph/Private/AnimGraphNode_Slot.cpp b/Engine/Source/Editor/AnimGraph/Private/AnimGraphNode_Slot.cpp
--- a/Engine/Source/Editor/AnimGraph/Private/AnimGraphNode_Slot.cpp
+++ b/Engine/Source/Editor/AnimGraph/Private/AnimGraphNode_Slot.cpp
@@ -4,6 +4,7 @@
 
 #include "GraphEditorActions.h"
 #include "ScopedTransaction.h"
+#include "AnimGraphNode_Slot.h"
 
 /////////////////////////////////////////////////////
 // UAnimGraphNode_Slot
diff --git a/Engine/Source/Editor/AnimGraph/Private/AnimGraphNode_SpringBone.cpp b/Engine/Source/Editor/AnimGraph/Private/AnimGraphNode_SpringBone.cpp
--- a/Engine/Source/Editor/AnimGraph/Private/AnimGraphNode_SpringBone.cpp
+++ b/Engine/Source/Editor/AnimGraph/Private/AnimGraphNode_SpringBone.cpp
@@ -1,6 +1,7 @@
 // Copyright 1998-2014 Epic Games, Inc. All Rights Reserved.
 
 #include "AnimGraphPrivatePCH.h"
+#include "AnimGraphNode_SpringBone.h"
 
 /////////////////////////////////////////////////////
 // UAnimGraphNode_SpringBone
diff --git a/Engine/Source/Editor/AnimGraph/Private/AnimGraphNode_StateMachine.cpp b/Engine/Source/Editor/AnimGraph/Private/AnimGraphNode_StateMachine.cpp
--- a/Engine/Source/Editor/AnimGraph/Private/AnimGraphNode_StateMachine.cpp
+++ b/Engine/Source/Editor/AnimGraph/Private/AnimGraphNode_StateMachine.cpp
@@ -1,6 +1,7 @@
 // Copyright 1998-2014 Epic Games, Inc. All Rights Reserved.
 
 #include "AnimGraphPrivatePCH.h"
+#include "AnimGraphNode_StateMachine.h"
 
 /////////////////////////////////////////////////////
 // UAnimGraphNode_StateMachine
diff --git a/Engine/Source/Editor/AnimGraph/Private/AnimGraphNode_StateMachineBase.cpp b/Engine/Source/Editor/AnimGraph/Private/AnimGraphNode_StateMachineBase.cpp
--- a/Engine/Source/Editor/AnimGraph/Private/AnimGraphNode_StateMachineBase.cpp
+++ b/Engine/Source/Editor/AnimGraph/Private/AnimGraphNode_StateMachineBase.cpp
@@ -6,6 +6,7 @@
 #include "ScopedTransaction.h"
 #include "EdGraphUtilities.h"
 #include "Kismet2NameValidators.h"
+#include "AnimGraphNode_StateMachineBase.h"
 
 /////////////////////////////////////////////////////
 // FAnimStateMachineNodeNameValidator
@@ -21,9 +22,8 @@ public:
 		UAnimationGraph* StateMachine = CastChecked<UAnimationGraph>(InStateMachineNode->GetOuter());
 		StateMachine->GetNodesOfClassEx<UAnimGraphNode_StateMachine, UAnimGraphNode_StateMachineBase>(Nodes);
 
-		for (auto NodeIt = Nodes.CreateIterator(); NodeIt; ++NodeIt)
+		for (auto Node : Nodes)
 		{
-			auto Node = *NodeIt;
 			if (Node != InStateMachineNode)
 			{
 				Names.Add(Node->GetStateMachineName());
diff --git a/Engine/Source/Editor/AnimGraph/Private/AnimGraphNode_StateResult.cpp b/Engine/Source/Editor/AnimGraph/Private/AnimGraphNode_StateResult.cpp
--- a/Engine/Source/Editor/AnimGraph/Private/AnimGraphNode_StateResult.cpp
+++ b/Engine/Source/Editor/AnimGraph/Private/AnimGraphNode_StateResult.cpp
@@ -1,6 +1,7 @@
 // Copyright 1998-2014 Epic Games, Inc. All Rights Reserved.
 
 #include "AnimGraphPrivatePCH.h"
+#include "AnimGraphNode_StateResult.h"
 
 /////////////////////////////////////////////////////
 // UAnimGraphNode_StateResult
diff --git a/Engine/Source/Editor/AnimGraph/Private/AnimGraphNode_TwoBoneIK.cpp b/Engine/Source/Editor/AnimGraph/Private/AnimGraphNode_TwoBoneIK.cpp
--- a/Engine/Source/Editor/AnimGraph/Private/AnimGraphNode_TwoBoneIK.cpp
+++ b/Engine/Source/Editor/AnimGraph/Private/AnimGraphNode_TwoBoneIK.cpp
@@ -1,6 +1,7 @@
 // Copyright 1998-2014 Epic Games, Inc. All Rights Reserved.
 
 #include "AnimGraphPrivatePCH.h"
+#include "AnimGraphNode_TwoBoneIK.h"
 
 /////////////////////////////////////////////////////
 // UAnimGraphNode_TwoBoneIK
diff --git a/Engine/Source/Editor/AnimGraph/Private/AnimGraphNode_TwoWayBlend.cpp b/Engine/Source/Editor/AnimGraph/Private/AnimGraphNode_TwoWayBlend.cpp
--- a/Engine/Source/Editor/AnimGraph/Private/AnimGraphNode_TwoWayBlend.cpp
+++ b/Engine/Source/Editor/AnimGraph/Private/AnimGraphNode_TwoWayBlend.cpp
@@ -1,6 +1,7 @@
 // Copyright 1998-2014 Epic Games, Inc. All Rights Reserved.
 
 #include "AnimGraphPrivatePCH.h"
+#include "AnimGraphNode_TwoWayBlend.h"
 
 /////////////////////////////////////////////////////
 // UAnimGraphNode_TwoWayBlend
diff --git a/Engine/Source/Editor/AnimGraph/Private/AnimGraphNode_UseCachedPose.cpp b/Engine/Source/Editor/AnimGraph/Private/AnimGraphNode_UseCachedPose.cpp
--- a/Engine/Source/Editor/AnimGraph/Private/AnimGraphNode_UseCachedPose.cpp
+++ b/Engine/Source/Editor/AnimGraph/Private/AnimGraphNode_UseCachedPose.cpp
@@ -5,6 +5,7 @@
 #include "GraphEditorActions.h"
 #include "ScopedTransaction.h"
 #include "K2ActionMenuBuilder.h" // for FK2ActionMenuBuilder::AddNewNodeAction()
+#include "AnimGraphNode_UseCachedPose.h"
 
 /////////////////////////////////////////////////////
 // UAnimGraphNode_UseCachedPose
diff --git a/Engine/Source/Editor/AnimGraph/Private/AnimStateNode.cpp b/Engine/Source/Editor/AnimGraph/Private/AnimStateNode.cpp
--- a/Engine/Source/Editor/AnimGraph/Private/AnimStateNode.cpp
+++ b/Engine/Source/Editor/AnimGraph/Private/AnimStateNode.cpp
@@ -8,6 +8,7 @@
 
 #include "EdGraphUtilities.h"
 #include "Kismet2/Kismet2NameValidators.h"
+#include "AnimGraphNode_StateResult.h"
 
 /////////////////////////////////////////////////////
 // UAnimStateNode
diff --git a/Engine/Source/Editor/AnimGraph/Private/AnimationConduitGraphSchema.cpp b/Engine/Source/Editor/AnimGraph/Private/AnimationConduitGraphSchema.cpp
--- a/Engine/Source/Editor/AnimGraph/Private/AnimationConduitGraphSchema.cpp
+++ b/Engine/Source/Editor/AnimGraph/Private/AnimationConduitGraphSchema.cpp
@@ -1,6 +1,10 @@
 // Copyright 1998-2014 Epic Games, Inc. All Rights Reserved.
 
 #include "AnimGraphPrivatePCH.h"
+#include "AnimationConduitGraphSchema.h"
+#include "AnimGraphNode_TransitionResult.h"
+#include "AnimationTransitionGraph.h"
+#include "AnimStateConduitNode.h"
 
 /////////////////////////////////////////////////////
 // UAnimationConduitGraphSchema
diff --git a/Engine/Source/Editor/AnimGraph/Private/AnimationCustomTransitionGraph.cpp b/Engine/Source/Editor/AnimGraph/Private/AnimationCustomTransitionGraph.cpp
--- a/Engine/Source/Editor/AnimGraph/Private/AnimationCustomTransitionGraph.cpp
+++ b/Engine/Source/Editor/AnimGraph/Private/AnimationCustomTransitionGraph.cpp
@@ -1,7 +1,7 @@
 // Copyright 1998-2014 Epic Games, Inc. All Rights Reserved.
 
 #include "AnimGraphPrivatePCH.h"
-
+#include "AnimationCustomTransitionGraph.h"
 
 #define LOCTEXT_NAMESPACE "AnimationCustomTransitionGraph"
 
diff --git a/Engine/Source/Editor/AnimGraph/Private/AnimationCustomTransitionSchema.cpp b/Engine/Source/Editor/AnimGraph/Private/AnimationCustomTransitionSchema.cpp
--- a/Engine/Source/Editor/AnimGraph/Private/AnimationCustomTransitionSchema.cpp
+++ b/Engine/Source/Editor/AnimGraph/Private/AnimationCustomTransitionSchema.cpp
@@ -1,7 +1,10 @@
 // Copyright 1998-2014 Epic Games, Inc. All Rights Reserved.
 
 #include "AnimGraphPrivatePCH.h"
-
+#include "AnimationCustomTransitionSchema.h"
+#include "AnimGraphNode_CustomTransitionResult.h"
+#include "AnimGraphNode_TransitionPoseEvaluator.h"
+#include "AnimStateTransitionNode.h"
 
 /////////////////////////////////////////////////////
 // UAnimationCustomTransitionSchema
diff --git a/Engine/Source/Editor/AnimGraph/Private/AnimationGraphSchema.cpp b/Engine/Source/Editor/AnimGraph/Private/AnimationGraphSchema.cpp
--- a/Engine/Source/Editor/AnimGraph/Private/AnimationGraphSchema.cpp
+++ b/Engine/Source/Editor/AnimGraph/Private/AnimationGraphSchema.cpp
@@ -10,6 +10,14 @@
 #include "ScopedTransaction.h"
 #include "AssetData.h"
 #include "K2ActionMenuBuilder.h" // for FK2ActionMenuBuilder::AddNewNodeAction()
+#include "AnimationGraphSchema.h"
+#include "K2Node_TransitionRuleGetter.h"
+#include "AnimStateNode.h"
+#include "AnimGraphNode_BlendSpacePlayer.h"
+#include "AnimGraphNode_ComponentToLocalSpace.h"
+#include "AnimGraphNode_LocalToComponentSpace.h"
+#include "AnimGraphNode_RotationOffsetBlendSpace.h"
+#include "AnimGraphNode_SequencePlayer.h"
 
 #define LOCTEXT_NAMESPACE "AnimationGraphSchema"
 
diff --git a/Engine/Source/Editor/AnimGraph/Private/AnimationStateGraph.cpp b/Engine/Source/Editor/AnimGraph/Private/AnimationStateGraph.cpp
--- a/Engine/Source/Editor/AnimGraph/Private/AnimationStateGraph.cpp
+++ b/Engine/Source/Editor/AnimGraph/Private/AnimationStateGraph.cpp
@@ -1,7 +1,7 @@
 // Copyright 1998-2014 Epic Games, Inc. All Rights Reserved.
 
 #include "AnimGraphPrivatePCH.h"
-
+#include "AnimationStateGraph.h"
 
 #define LOCTEXT_NAMESPACE "AnimationStateGraph"
 
diff --git a/Engine/Source/Editor/AnimGraph/Private/AnimationStateGraphSchema.cpp b/Engine/Source/Editor/AnimGraph/Private/AnimationStateGraphSchema.cpp
--- a/Engine/Source/Editor/AnimGraph/Private/AnimationStateGraphSchema.cpp
+++ b/Engine/Source/Editor/AnimGraph/Private/AnimationStateGraphSchema.cpp
@@ -5,6 +5,7 @@
 #include "GraphEditorActions.h"
 #include "ScopedTransaction.h"
 #include "AssetData.h"
+#include "AnimationStateGraphSchema.h"
 
 #define LOCTEXT_NAMESPACE "AnimationStateGraphSchema"
 
diff --git a/Engine/Source/Editor/AnimGraph/Private/AnimationStateMachineGraph.cpp b/Engine/Source/Editor/AnimGraph/Private/AnimationStateMachineGraph.cpp
--- a/Engine/Source/Editor/AnimGraph/Private/AnimationStateMachineGraph.cpp
+++ b/Engine/Source/Editor/AnimGraph/Private/AnimationStateMachineGraph.cpp
@@ -1,7 +1,7 @@
 // Copyright 1998-2014 Epic Games, Inc. All Rights Reserved.
 
 #include "AnimGraphPrivatePCH.h"
-
+#include "AnimationStatemachineGraph.h"
 
 #define LOCTEXT_NAMESPACE "AnimationGraph"
 
diff --git a/Engine/Source/Editor/AnimGraph/Private/AnimationStateMachineSchema.cpp b/Engine/Source/Editor/AnimGraph/Private/AnimationStateMachineSchema.cpp
--- a/Engine/Source/Editor/AnimGraph/Private/AnimationStateMachineSchema.cpp
+++ b/Engine/Source/Editor/AnimGraph/Private/AnimationStateMachineSchema.cpp
@@ -6,6 +6,9 @@
 
 #include "AnimGraphPrivatePCH.h"
 
+#include "AnimationStateMachineSchema.h"
+#include "AnimStateEntryNode.h"
+
 #include "ScopedTransaction.h"
 #include "GraphEditorActions.h"
 
diff --git a/Engine/Source/Editor/AnimGraph/Private/AnimationTransitionSchema.cpp b/Engine/Source/Editor/AnimGraph/Private/AnimationTransitionSchema.cpp
--- a/Engine/Source/Editor/AnimGraph/Private/AnimationTransitionSchema.cpp
+++ b/Engine/Source/Editor/AnimGraph/Private/AnimationTransitionSchema.cpp
@@ -5,6 +5,7 @@
 =============================================================================*/
 
 #include "AnimGraphPrivatePCH.h"
+#include "AnimationTransitionSchema.h"
 
 #include "BlueprintUtilities.h"
 #include "AnimGraphDefinitions.h"
diff --git a/Engine/Source/Editor/AnimGraph/Public/AnimGraphDefinitions.h b/Engine/Source/Editor/AnimGraph/Public/AnimGraphDefinitions.h
--- a/Engine/Source/Editor/AnimGraph/Public/AnimGraphDefinitions.h
+++ b/Engine/Source/Editor/AnimGraph/Public/AnimGraphDefinitions.h
@@ -1,5 +1,3 @@
 // Copyright 1998-2014 Epic Games, Inc. All Rights Reserved.
 
 #include "BlueprintGraphDefinitions.h"
-#include "EngineAnimationNodeClasses.h"
-#include "AnimGraphClasses.h"
diff --git a/Engine/Source/Editor/BlueprintGraph/Classes/K2Node.h b/Engine/Source/Editor/BlueprintGraph/Classes/K2Node.h
--- a/Engine/Source/Editor/BlueprintGraph/Classes/K2Node.h
+++ b/Engine/Source/Editor/BlueprintGraph/Classes/K2Node.h
@@ -1,8 +1,7 @@
 // Copyright 1998-2014 Epic Games, Inc. All Rights Reserved.
 
 #pragma once
-
-
+#include "EdGraph/EdGraphNode.h"
 #include "K2Node.generated.h"
 
 /** Helper struct to allow us to redirect properties and functions through renames and additionally between classes if necessary */
diff --git a/Engine/Source/Editor/Cascade/Classes/CascadeParticleSystemComponent.h b/Engine/Source/Editor/Cascade/Classes/CascadeParticleSystemComponent.h
--- a/Engine/Source/Editor/Cascade/Classes/CascadeParticleSystemComponent.h
+++ b/Engine/Source/Editor/Cascade/Classes/CascadeParticleSystemComponent.h
@@ -8,7 +8,7 @@
 /** 
  * Used to provide an extended particle system component to allow collision to function in the preview window.
  */
-UCLASS(HeaderGroup=Component, hidecategories=Object, hidecategories=Physics, hidecategories=Collision, editinlinenew)
+UCLASS(hidecategories=Object, hidecategories=Physics, hidecategories=Collision, editinlinenew)
 class UCascadeParticleSystemComponent : public UParticleSystemComponent
 {
 	GENERATED_UCLASS_BODY()
diff --git a/Engine/Source/Editor/DetailCustomizations/Private/AnimMontageSegmentDetails.cpp b/Engine/Source/Editor/DetailCustomizations/Private/AnimMontageSegmentDetails.cpp
--- a/Engine/Source/Editor/DetailCustomizations/Private/AnimMontageSegmentDetails.cpp
+++ b/Engine/Source/Editor/DetailCustomizations/Private/AnimMontageSegmentDetails.cpp
@@ -2,10 +2,10 @@
 
 #include "DetailCustomizationsPrivatePCH.h"
 #include "AnimMontageSegmentDetails.h"
+#include "AnimPreviewInstance.h"
 #include "Runtime/Engine/Public/Slate/SceneViewport.h"
 #include "Editor/KismetWidgets/Public/SScrubControlPanel.h"
 #include "Runtime/Engine/Public/FXSystem.h"
-
 #define LOCTEXT_NAMESPACE "AnimMontageSegmentDetails"
 
 /////////////////////////////////////////////////////////////////////////
diff --git a/Engine/Source/Editor/DetailCustomizations/Private/AnimSequenceDetails.cpp b/Engine/Source/Editor/DetailCustomizations/Private/AnimSequenceDetails.cpp
--- a/Engine/Source/Editor/DetailCustomizations/Private/AnimSequenceDetails.cpp
+++ b/Engine/Source/Editor/DetailCustomizations/Private/AnimSequenceDetails.cpp
@@ -3,6 +3,7 @@
 #include "DetailCustomizationsPrivatePCH.h"
 #include "AnimSequenceDetails.h"
 #include "AnimMontageSegmentDetails.h"
+#include "AnimPreviewInstance.h"
 #include "Runtime/Engine/Public/Slate/SceneViewport.h"
 #include "Editor/KismetWidgets/Public/SScrubControlPanel.h"
 #include "Runtime/Engine/Public/FXSystem.h"
diff --git a/Engine/Source/Editor/DetailCustomizations/Private/AnimTransitionNodeDetails.cpp b/Engine/Source/Editor/DetailCustomizations/Private/AnimTransitionNodeDetails.cpp
--- a/Engine/Source/Editor/DetailCustomizations/Private/AnimTransitionNodeDetails.cpp
+++ b/Engine/Source/Editor/DetailCustomizations/Private/AnimTransitionNodeDetails.cpp
@@ -2,8 +2,12 @@
 
 #include "DetailCustomizationsPrivatePCH.h"
 
+#include "AnimationTransitionGraph.h"
 #include "AnimGraphDefinitions.h"
-
+#include "AnimGraphNode_TransitionResult.h"
+#include "AnimStateConduitNode.h"
+#include "AnimStateNodeBase.h"
+#include "AnimStateTransitionNode.h"
 #include "AnimTransitionNodeDetails.h"
 #include "KismetEditorUtilities.h"
 #include "SKismetLinearExpression.h"
diff --git a/Engine/Source/Editor/DetailCustomizations/Private/TransitionPoseEvaluatorNodeDetails.cpp b/Engine/Source/Editor/DetailCustomizations/Private/TransitionPoseEvaluatorNodeDetails.cpp
--- a/Engine/Source/Editor/DetailCustomizations/Private/TransitionPoseEvaluatorNodeDetails.cpp
+++ b/Engine/Source/Editor/DetailCustomizations/Private/TransitionPoseEvaluatorNodeDetails.cpp
@@ -3,6 +3,7 @@
 #include "DetailCustomizationsPrivatePCH.h"
 
 #include "AnimGraphDefinitions.h"
+#include "AnimGraphNode_TransitionPoseEvaluator.h"
 
 #include "AnimTransitionNodeDetails.h"
 #include "KismetEditorUtilities.h"
diff --git a/Engine/Source/Editor/DetailCustomizations/Private/TransitionPoseEvaluatorNodeDetails.h b/Engine/Source/Editor/DetailCustomizations/Private/TransitionPoseEvaluatorNodeDetails.h
--- a/Engine/Source/Editor/DetailCustomizations/Private/TransitionPoseEvaluatorNodeDetails.h
+++ b/Engine/Source/Editor/DetailCustomizations/Private/TransitionPoseEvaluatorNodeDetails.h
@@ -5,6 +5,11 @@
 #include "PreviewScene.h"
 #include "AnimGraphDefinitions.h"
 
+//
+// Forward declarations.
+//
+class UAnimGraphNode_TransitionPoseEvaluator;
+
 class FTransitionPoseEvaluatorNodeDetails : public IDetailCustomization
 {
 public:
diff --git a/Engine/Source/Editor/GraphEditor/Private/AnimationNodes/SGraphNodeStateMachineInstance.cpp b/Engine/Source/Editor/GraphEditor/Private/AnimationNodes/SGraphNodeStateMachineInstance.cpp
--- a/Engine/Source/Editor/GraphEditor/Private/AnimationNodes/SGraphNodeStateMachineInstance.cpp
+++ b/Engine/Source/Editor/GraphEditor/Private/AnimationNodes/SGraphNodeStateMachineInstance.cpp
@@ -1,10 +1,12 @@
 // Copyright 1998-2014 Epic Games, Inc. All Rights Reserved.
 
 #include "GraphEditorCommon.h"
-#include "../KismetNodes/SGraphNodeK2Base.h"
-#include "../KismetNodes/SGraphNodeK2Composite.h"
+#include "KismetNodes/SGraphNodeK2Base.h"
+#include "KismetNodes/SGraphNodeK2Composite.h"
 #include "SGraphNodeStateMachineInstance.h"
 #include "SGraphPreviewer.h"
+#include "AnimationStateMachineGraph.h"
+#include "AnimGraphNode_StateMachineBase.h"
 
 /////////////////////////////////////////////////////
 // SGraphNodeStateMachineInstance
diff --git a/Engine/Source/Editor/GraphEditor/Private/AnimationStateNodes/SGraphNodeAnimState.cpp b/Engine/Source/Editor/GraphEditor/Private/AnimationStateNodes/SGraphNodeAnimState.cpp
--- a/Engine/Source/Editor/GraphEditor/Private/AnimationStateNodes/SGraphNodeAnimState.cpp
+++ b/Engine/Source/Editor/GraphEditor/Private/AnimationStateNodes/SGraphNodeAnimState.cpp
@@ -6,6 +6,8 @@
 #include "Editor/UnrealEd/Public/Kismet2/BlueprintEditorUtils.h"
 #include "SGraphNode.h"
 #include "IDocumentation.h"
+#include "AnimStateConduitNode.h"
+#include "AnimationStateMachineGraph.h"
 
 /////////////////////////////////////////////////////
 // SStateMachineOutputPin
diff --git a/Engine/Source/Editor/GraphEditor/Private/AnimationStateNodes/SGraphNodeAnimState.h b/Engine/Source/Editor/GraphEditor/Private/AnimationStateNodes/SGraphNodeAnimState.h
--- a/Engine/Source/Editor/GraphEditor/Private/AnimationStateNodes/SGraphNodeAnimState.h
+++ b/Engine/Source/Editor/GraphEditor/Private/AnimationStateNodes/SGraphNodeAnimState.h
@@ -2,6 +2,13 @@
 
 #pragma once
 
+//
+// Forward declarations.
+//
+class UAnimStateNodeBase;
+class UAnimStateConduitNode;
+class UAnimStateEntryNode;
+
 class SGraphNodeAnimState : public SGraphNode
 {
 public:
diff --git a/Engine/Source/Editor/GraphEditor/Private/AnimationStateNodes/SGraphNodeAnimTransition.cpp b/Engine/Source/Editor/GraphEditor/Private/AnimationStateNodes/SGraphNodeAnimTransition.cpp
--- a/Engine/Source/Editor/GraphEditor/Private/AnimationStateNodes/SGraphNodeAnimTransition.cpp
+++ b/Engine/Source/Editor/GraphEditor/Private/AnimationStateNodes/SGraphNodeAnimTransition.cpp
@@ -12,6 +12,9 @@
 #include "ConnectionDrawingPolicy.h"
 #include "IDocumentation.h"
 
+#include "AnimationTransitionGraph.h"
+#include "AnimGraphNode_TransitionResult.h"
+
 #define LOCTEXT_NAMESPACE "TransitionNodes"
 
 /////////////////////////////////////////////////////
diff --git a/Engine/Source/Editor/GraphEditor/Private/ConnectionDrawingPolicy.cpp b/Engine/Source/Editor/GraphEditor/Private/ConnectionDrawingPolicy.cpp
--- a/Engine/Source/Editor/GraphEditor/Private/ConnectionDrawingPolicy.cpp
+++ b/Engine/Source/Editor/GraphEditor/Private/ConnectionDrawingPolicy.cpp
@@ -8,6 +8,10 @@
 #include "KismetNodes/KismetNodeInfoContext.h"
 #include "SoundDefinitions.h"
 #include "AnimationStateNodes/SGraphNodeAnimTransition.h"
+#include "AnimStateTransitionNode.h"
+#include "AnimStateEntryNode.h"
+#include "AnimationGraphSchema.h"
+#include "AnimGraphNode_Base.h"
 
 DEFINE_LOG_CATEGORY_STATIC(LogConnectionDrawingPolicy, Log, All);
 
diff --git a/Engine/Source/Editor/GraphEditor/Private/NodeFactory.cpp b/Engine/Source/Editor/GraphEditor/Private/NodeFactory.cpp
--- a/Engine/Source/Editor/GraphEditor/Private/NodeFactory.cpp
+++ b/Engine/Source/Editor/GraphEditor/Private/NodeFactory.cpp
@@ -26,6 +26,12 @@
 #include "KismetNodes/SGraphNodeFormatText.h"
 #include "KismetNodes/SGraphNodeK2ArrayFunction.h"
 
+#include "AnimGraphNode_Root.h"
+#include "AnimGraphNode_SequencePlayer.h"
+#include "AnimStateNode.h"
+
+#include "AnimationStateMachineSchema.h"
+
 #include "AnimationStateNodes/SGraphNodeAnimState.h"
 #include "AnimationStateNodes/SGraphNodeAnimTransition.h"
 #include "AnimationStateNodes/SGraphNodeAnimStateEntry.h"
diff --git a/Engine/Source/Editor/GraphEditor/Private/SGraphPanel.cpp b/Engine/Source/Editor/GraphEditor/Private/SGraphPanel.cpp
--- a/Engine/Source/Editor/GraphEditor/Private/SGraphPanel.cpp
+++ b/Engine/Source/Editor/GraphEditor/Private/SGraphPanel.cpp
@@ -15,6 +15,8 @@
 #include "KismetNodes/KismetNodeInfoContext.h"
 #include "GraphDiffControl.h"
 
+#include "AnimationGraphSchema.h"
+
 DEFINE_LOG_CATEGORY_STATIC(LogGraphPanel, Log, All);
 
 /**
diff --git a/Engine/Source/Editor/Kismet/Private/BlueprintEditor.cpp b/Engine/Source/Editor/Kismet/Private/BlueprintEditor.cpp
--- a/Engine/Source/Editor/Kismet/Private/BlueprintEditor.cpp
+++ b/Engine/Source/Editor/Kismet/Private/BlueprintEditor.cpp
@@ -4,7 +4,6 @@
 #include "ScopedTransaction.h"
 #include "Editor/UnrealEd/Public/Kismet2/KismetDebugUtilities.h"
 #include "Editor/UnrealEd/Public/Kismet2/BlueprintEditorUtils.h"
-#include "EngineLevelScriptClasses.h"
 #include "GraphEditor.h"
 #include "BlueprintUtilities.h"
 #include "Editor/KismetCompiler/Public/KismetCompilerModule.h"
@@ -60,6 +59,11 @@
 #include "BlueprintEditorTabFactories.h"
 #include "SPinTypeSelector.h"
 #include "AnimGraphDefinitions.h"
+#include "AnimStateConduitNode.h"
+#include "AnimationGraphSchema.h"
+#include "AnimationStateGraph.h"
+#include "AnimationStateMachineSchema.h"
+#include "AnimationTransitionGraph.h"
 #include "BlueprintEditorModes.h"
 
 #include "EngineAnalytics.h"
diff --git a/Engine/Source/Editor/Kismet/Private/SBlueprintPalette.cpp b/Engine/Source/Editor/Kismet/Private/SBlueprintPalette.cpp
--- a/Engine/Source/Editor/Kismet/Private/SBlueprintPalette.cpp
+++ b/Engine/Source/Editor/Kismet/Private/SBlueprintPalette.cpp
@@ -22,6 +22,11 @@
 #include "BlueprintPaletteFavorites.h"
 #include "STutorialWrapper.h"
 #include "ClassIconFinder.h"
+#include "AnimationStateMachineGraph.h"
+#include "AnimationStateMachineSchema.h"
+#include "AnimationStateGraph.h"
+#include "AnimStateConduitNode.h"
+#include "AnimationTransitionGraph.h"
 
 #define LOCTEXT_NAMESPACE "BlueprintPalette"
 
diff --git a/Engine/Source/Editor/Persona/Private/AnimationEditorViewportClient.cpp b/Engine/Source/Editor/Persona/Private/AnimationEditorViewportClient.cpp
--- a/Engine/Source/Editor/Persona/Private/AnimationEditorViewportClient.cpp
+++ b/Engine/Source/Editor/Persona/Private/AnimationEditorViewportClient.cpp
@@ -8,7 +8,9 @@
 #include "SAnimViewportToolBar.h"
 #include "AnimViewportShowCommands.h"
 #include "AnimGraphDefinitions.h"
+#include "AnimPreviewInstance.h"
 #include "AnimationEditorViewportClient.h"
+#include "AnimGraphNode_SkeletalControlBase.h"
 #include "Runtime/Engine/Classes/Components/ReflectionCaptureComponent.h"
 #include "Runtime/Engine/Classes/Components/SphereReflectionCaptureComponent.h"
 #include "UnrealWidget.h"
diff --git a/Engine/Source/Editor/Persona/Private/Persona.cpp b/Engine/Source/Editor/Persona/Private/Persona.cpp
--- a/Engine/Source/Editor/Persona/Private/Persona.cpp
+++ b/Engine/Source/Editor/Persona/Private/Persona.cpp
@@ -40,7 +40,12 @@
 #include "Editor/Persona/Private/AnimationEditorViewportClient.h"
 
 #include "ComponentAssetBroker.h"
-
+#include "AnimGraphNode_BlendListByInt.h"
+#include "AnimGraphNode_BlendSpaceEvaluator.h"
+#include "AnimGraphNode_BlendSpacePlayer.h"
+#include "AnimGraphNode_LayeredBoneBlend.h"
+#include "AnimGraphNode_SequencePlayer.h"
+#include "AnimGraphNode_SequenceEvaluator.h"
 #define LOCTEXT_NAMESPACE "FPersona"
 
 /////////////////////////////////////////////////////
diff --git a/Engine/Source/Editor/Persona/Private/SAnimationScrubPanel.cpp b/Engine/Source/Editor/Persona/Private/SAnimationScrubPanel.cpp
--- a/Engine/Source/Editor/Persona/Private/SAnimationScrubPanel.cpp
+++ b/Engine/Source/Editor/Persona/Private/SAnimationScrubPanel.cpp
@@ -1,6 +1,8 @@
 // Copyright 1998-2014 Epic Games, Inc. All Rights Reserved.
 #include "PersonaPrivatePCH.h"
 #include "AnimGraphDefinitions.h"
+#include "AnimPreviewInstance.h"
+#include "Animation/AnimSingleNodeInstance.h"
 #include "Persona.h"
 #include "SAnimationScrubPanel.h"
 #include "Editor/KismetWidgets/Public/SScrubControlPanel.h"
diff --git a/Engine/Source/Editor/UnrealEd/Private/AssetSelection.cpp b/Engine/Source/Editor/UnrealEd/Private/AssetSelection.cpp
--- a/Engine/Source/Editor/UnrealEd/Private/AssetSelection.cpp
+++ b/Engine/Source/Editor/UnrealEd/Private/AssetSelection.cpp
@@ -2,8 +2,6 @@
 
 #include "UnrealEd.h"
 #include "AssetSelection.h"
-#include "EngineLightClasses.h"
-#include "EngineDecalClasses.h"
 #include "ScopedTransaction.h"
 
 #include "LevelUtils.h"
diff --git a/Engine/Source/Editor/UnrealEd/Private/Editor.cpp b/Engine/Source/Editor/UnrealEd/Private/Editor.cpp
--- a/Engine/Source/Editor/UnrealEd/Private/Editor.cpp
+++ b/Engine/Source/Editor/UnrealEd/Private/Editor.cpp
@@ -7,8 +7,6 @@
 
 // needed for the RemotePropagator
 #include "SoundDefinitions.h"
-#include "EngineFoliageClasses.h"
-#include "EngineDecalClasses.h"
 #include "Database.h"
 #include "SurfaceIterators.h"
 #include "ScopedTransaction.h"
diff --git a/Engine/Source/Editor/UnrealEd/Private/EditorActor.cpp b/Engine/Source/Editor/UnrealEd/Private/EditorActor.cpp
--- a/Engine/Source/Editor/UnrealEd/Private/EditorActor.cpp
+++ b/Engine/Source/Editor/UnrealEd/Private/EditorActor.cpp
@@ -4,7 +4,6 @@
 #include "UnrealEd.h"
 #include "ScopedTransaction.h"
 #include "Factories.h"
-#include "EngineFoliageClasses.h"
 #include "LevelUtils.h"
 #include "BusyCursor.h"
 #include "BSPOps.h"
@@ -14,7 +13,6 @@
 #include "Editor/GeometryMode/Public/GeometryEdMode.h"
 #include "Editor/GeometryMode/Public/EditorGeometry.h"
 #include "ActorEditorUtils.h"
-#include "EngineLevelScriptClasses.h"
 
 #define LOCTEXT_NAMESPACE "UnrealEd.EditorActor"
 
diff --git a/Engine/Source/Editor/UnrealEd/Private/EditorAnimUtils.cpp b/Engine/Source/Editor/UnrealEd/Private/EditorAnimUtils.cpp
--- a/Engine/Source/Editor/UnrealEd/Private/EditorAnimUtils.cpp
+++ b/Engine/Source/Editor/UnrealEd/Private/EditorAnimUtils.cpp
@@ -1,6 +1,7 @@
 // Copyright 1998-2014 Epic Games, Inc. All Rights Reserved.
 
 #include "UnrealEd.h"
+#include "AnimGraphNode_Base.h"
 #include "AssetData.h"
 #include "EditorAnimUtils.h"
 #include "BlueprintEditorUtils.h"
@@ -10,6 +11,7 @@
 #include "NotificationManager.h"
 #include "Editor/Persona/Public/PersonaModule.h"
 
+
 #define LOCTEXT_NAMESPACE "EditorAnimUtils"
 
 namespace EditorAnimUtils
diff --git a/Engine/Source/Editor/UnrealEd/Private/EditorExporters.cpp b/Engine/Source/Editor/UnrealEd/Private/EditorExporters.cpp
--- a/Engine/Source/Editor/UnrealEd/Private/EditorExporters.cpp
+++ b/Engine/Source/Editor/UnrealEd/Private/EditorExporters.cpp
@@ -6,8 +6,6 @@
 
 #include "UnrealEd.h"
 #include "SoundDefinitions.h"
-#include "EngineMaterialClasses.h"
-#include "EngineFoliageClasses.h"
 #include "Landscape/LandscapeDataAccess.h"
 #include "Kismet2/DebuggerCommands.h"
 #include "MaterialCompiler.h"
diff --git a/Engine/Source/Editor/UnrealEd/Private/Factories/ActorFactory.cpp b/Engine/Source/Editor/UnrealEd/Private/Factories/ActorFactory.cpp
--- a/Engine/Source/Editor/UnrealEd/Private/Factories/ActorFactory.cpp
+++ b/Engine/Source/Editor/UnrealEd/Private/Factories/ActorFactory.cpp
@@ -6,12 +6,8 @@ ActorFactory.cpp:
 
 #include "UnrealEd.h"
 #include "ParticleDefinitions.h"
-#include "EngineMaterialClasses.h"
-#include "EngineDecalClasses.h"
-#include "EngineFoliageClasses.h"
 #include "SoundDefinitions.h"
 #include "BlueprintUtilities.h"
-#include "EngineDecalClasses.h"
 #include "Kismet2/BlueprintEditorUtils.h"
 #include "Kismet2/KismetEditorUtilities.h"
 #include "AssetData.h"
diff --git a/Engine/Source/Editor/UnrealEd/Private/Kismet2/BlueprintEditorUtils.cpp b/Engine/Source/Editor/UnrealEd/Private/Kismet2/BlueprintEditorUtils.cpp
--- a/Engine/Source/Editor/UnrealEd/Private/Kismet2/BlueprintEditorUtils.cpp
+++ b/Engine/Source/Editor/UnrealEd/Private/Kismet2/BlueprintEditorUtils.cpp
@@ -3,11 +3,18 @@
 #include "UnrealEd.h"
 
 #include "LatentActions.h"
-#include "EngineLevelScriptClasses.h"
 #include "Kismet2/CompilerResultsLog.h"
 
 #include "Editor/KismetCompiler/Public/KismetCompilerModule.h"
 #include "AnimGraphDefinitions.h"
+#include "AnimGraphNode_StateMachineBase.h"
+#include "AnimStateNode.h"
+#include "AnimStateNodeBase.h"
+#include "AnimStateTransitionNode.h"
+#include "AnimationTransitionSchema.h"
+#include "AnimationGraph.h"
+#include "AnimationGraphSchema.h"
+#include "AnimationStateMachineGraph.h"
 #include "Editor/UnrealEd/Public/Kismet2/KismetEditorUtilities.h"
 #include "Editor/UnrealEd/Public/Kismet2/KismetDebugUtilities.h"
 #include "Editor/UnrealEd/Public/Kismet2/BlueprintEditorUtils.h"
diff --git a/Engine/Source/Editor/UnrealEd/Private/Kismet2/Kismet2.cpp b/Engine/Source/Editor/UnrealEd/Private/Kismet2/Kismet2.cpp
--- a/Engine/Source/Editor/UnrealEd/Private/Kismet2/Kismet2.cpp
+++ b/Engine/Source/Editor/UnrealEd/Private/Kismet2/Kismet2.cpp
@@ -4,6 +4,8 @@
 #include "UnrealEd.h"
 #include "BlueprintUtilities.h"
 #include "AnimGraphDefinitions.h"
+#include "AnimationGraph.h"
+#include "AnimationGraphSchema.h"
 #include "Kismet2/KismetEditorUtilities.h"
 #include "Kismet2/KismetReinstanceUtilities.h"
 #include "Kismet2/BlueprintEditorUtils.h"
diff --git a/Engine/Source/Editor/UnrealEd/Private/Kismet2/Kismet2NameValidators.cpp b/Engine/Source/Editor/UnrealEd/Private/Kismet2/Kismet2NameValidators.cpp
--- a/Engine/Source/Editor/UnrealEd/Private/Kismet2/Kismet2NameValidators.cpp
+++ b/Engine/Source/Editor/UnrealEd/Private/Kismet2/Kismet2NameValidators.cpp
@@ -4,6 +4,8 @@
 #include "Slate.h"
 #include "Kismet2NameValidators.h"
 #include "AnimGraphDefinitions.h"
+#include "AnimStateTransitionNode.h"
+#include "AnimationStateMachineGraph.h"
 #include "KismetEditorUtilities.h"
 #include "KismetDebugUtilities.h"
 #include "BlueprintEditorUtils.h"
diff --git a/Engine/Source/Editor/UnrealEd/Private/LevelEditorViewport.cpp b/Engine/Source/Editor/UnrealEd/Private/LevelEditorViewport.cpp
--- a/Engine/Source/Editor/UnrealEd/Private/LevelEditorViewport.cpp
+++ b/Engine/Source/Editor/UnrealEd/Private/LevelEditorViewport.cpp
@@ -2,9 +2,7 @@
 
 
 #include "UnrealEd.h"
-#include "EngineMaterialClasses.h"
 #include "SoundDefinitions.h"
-#include "EngineSplineClasses.h"
 #include "CameraController.h"
 #include "MouseDeltaTracker.h"
 #include "ScopedTransaction.h"
diff --git a/Engine/Source/Editor/UnrealEd/Private/StaticLightingSystem/StaticLightingExport.cpp b/Engine/Source/Editor/UnrealEd/Private/StaticLightingSystem/StaticLightingExport.cpp
--- a/Engine/Source/Editor/UnrealEd/Private/StaticLightingSystem/StaticLightingExport.cpp
+++ b/Engine/Source/Editor/UnrealEd/Private/StaticLightingSystem/StaticLightingExport.cpp
@@ -8,10 +8,9 @@
 
 #include "StaticMeshLight.h"
 #include "ModelLight.h"
-#include "EngineFoliageClasses.h"
 #include "Runtime/Engine/Public/Landscape/LandscapeLight.h"
 
-#include "../Lightmass/Lightmass.h"
+#include "Lightmass/Lightmass.h"
 
 // Doxygen cannot parse these definitions correctly since the declaration is in a header from a different module
 #if !UE_BUILD_DOCS
diff --git a/Engine/Source/Editor/UnrealEd/Public/UnrealEd.h b/Engine/Source/Editor/UnrealEd/Public/UnrealEd.h
--- a/Engine/Source/Editor/UnrealEd/Public/UnrealEd.h
+++ b/Engine/Source/Editor/UnrealEd/Public/UnrealEd.h
@@ -11,8 +11,6 @@
 
 #include "Engine.h"
 
-#include "EngineUserInterfaceClasses.h"
-
 #include "Slate.h"
 #include "EditorStyle.h"
 
diff --git a/Engine/Source/Runtime/Engine/Classes/AI/BehaviorTree/BlackboardComponent.h b/Engine/Source/Runtime/Engine/Classes/AI/BehaviorTree/BlackboardComponent.h
--- a/Engine/Source/Runtime/Engine/Classes/AI/BehaviorTree/BlackboardComponent.h
+++ b/Engine/Source/Runtime/Engine/Classes/AI/BehaviorTree/BlackboardComponent.h
@@ -11,7 +11,10 @@
  */
 
 #pragma once
+#include "Components/ActorComponent.h"
 #include "BehaviorTreeTypes.h"
+#include "Blackboard/BlackboardKeyType.h"
+#include "BlackboardData.h"
 #include "BlackboardComponent.generated.h"
 
 namespace EBlackboardDescription
@@ -25,7 +28,7 @@ namespace EBlackboardDescription
 	};
 }
 
-UCLASS(HeaderGroup=Component)
+UCLASS()
 class ENGINE_API UBlackboardComponent : public UActorComponent
 {
 	GENERATED_UCLASS_BODY()
diff --git a/Engine/Source/Runtime/Engine/Classes/AI/BrainComponent.h b/Engine/Source/Runtime/Engine/Classes/AI/BrainComponent.h
--- a/Engine/Source/Runtime/Engine/Classes/AI/BrainComponent.h
+++ b/Engine/Source/Runtime/Engine/Classes/AI/BrainComponent.h
@@ -83,7 +83,7 @@ private:
 	TWeakObjectPtr<UBrainComponent> Owner;
 };
 
-UCLASS(HeaderGroup=Component, abstract, BlueprintType)
+UCLASS(abstract, BlueprintType)
 class ENGINE_API UBrainComponent : public UActorComponent, public IAIResourceInterface
 {
 	GENERATED_UCLASS_BODY()
diff --git a/Engine/Source/Runtime/Engine/Classes/AI/EnvironmentQuery/EQSRenderingComponent.h b/Engine/Source/Runtime/Engine/Classes/AI/EnvironmentQuery/EQSRenderingComponent.h
--- a/Engine/Source/Runtime/Engine/Classes/AI/EnvironmentQuery/EQSRenderingComponent.h
+++ b/Engine/Source/Runtime/Engine/Classes/AI/EnvironmentQuery/EQSRenderingComponent.h
@@ -2,6 +2,8 @@
 
 #pragma once
 #include "DebugRenderSceneProxy.h"
+#include "Debug/DebugDrawService.h"
+#include "EnvQueryTypes.h"
 #include "EQSRenderingComponent.generated.h"
 
 class FEQSSceneProxy : public FDebugRenderSceneProxy
@@ -28,7 +30,7 @@ private:
 	bool SafeIsActorSelected() const;
 };
 
-UCLASS(HeaderGroup=Component, hidecategories=Object)
+UCLASS(hidecategories=Object)
 class ENGINE_API UEQSRenderingComponent : public UPrimitiveComponent
 {
 	GENERATED_UCLASS_BODY()
diff --git a/Engine/Source/Runtime/Engine/Classes/AI/EnvironmentQuery/EQSTestingPawn.h b/Engine/Source/Runtime/Engine/Classes/AI/EnvironmentQuery/EQSTestingPawn.h
--- a/Engine/Source/Runtime/Engine/Classes/AI/EnvironmentQuery/EQSTestingPawn.h
+++ b/Engine/Source/Runtime/Engine/Classes/AI/EnvironmentQuery/EQSTestingPawn.h
@@ -7,7 +7,7 @@
 /** this class is abstract even though it's perfectly functional on its own.
  *	The reason is to stop it from showing as valid player pawn type when configuring 
  *	project's game mode. */
-UCLASS(abstract, HeaderGroup=Pawn, hidecategories=(Advanced, Attachment, Collision, Animation),DependsOn=(UEnvQueryTypes))
+UCLASS(abstract, hidecategories=(Advanced, Attachment, Collision, Animation),DependsOn=(UEnvQueryTypes))
 class ENGINE_API AEQSTestingPawn : public ACharacter, public IEQSQueryResultSourceInterface
 {
 	GENERATED_UCLASS_BODY()
@@ -75,4 +75,4 @@ protected:
 	void UpdateDrawing();
 
 	static void OnEditorSelectionChanged(UObject* NewSelection);
-};
\ No newline at end of file
+};
diff --git a/Engine/Source/Runtime/Engine/Classes/AI/Navigation/NavigationComponent.h b/Engine/Source/Runtime/Engine/Classes/AI/Navigation/NavigationComponent.h
--- a/Engine/Source/Runtime/Engine/Classes/AI/Navigation/NavigationComponent.h
+++ b/Engine/Source/Runtime/Engine/Classes/AI/Navigation/NavigationComponent.h
@@ -3,7 +3,7 @@
 #pragma once
 #include "NavigationComponent.generated.h"
 
-UCLASS(HeaderGroup=Component, dependson=(UNavAgentInterface, UNavPathObserverInterface, UNavigationSystem), config=Engine)
+UCLASS(dependson=(UNavAgentInterface, UNavPathObserverInterface, UNavigationSystem), config=Engine)
 class ENGINE_API UNavigationComponent : public UActorComponent, public INavigationPathGenerator
 {
 	GENERATED_UCLASS_BODY()
diff --git a/Engine/Source/Runtime/Engine/Classes/AI/Navigation/NavigationGraphNodeComponent.h b/Engine/Source/Runtime/Engine/Classes/AI/Navigation/NavigationGraphNodeComponent.h
--- a/Engine/Source/Runtime/Engine/Classes/AI/Navigation/NavigationGraphNodeComponent.h
+++ b/Engine/Source/Runtime/Engine/Classes/AI/Navigation/NavigationGraphNodeComponent.h
@@ -3,7 +3,7 @@
 #pragma once
 #include "NavigationGraphNodeComponent.generated.h"
 
-UCLASS(HeaderGroup=Component, config=Engine, dependson=ANavigationGraph, MinimalAPI, HideCategories=(Mobility))
+UCLASS(config=Engine, dependson=ANavigationGraph, MinimalAPI, HideCategories=(Mobility))
 class UNavigationGraphNodeComponent : public USceneComponent
 {
 	GENERATED_UCLASS_BODY()
diff --git a/Engine/Source/Runtime/Engine/Classes/AI/Navigation/NavigationTestingActor.h b/Engine/Source/Runtime/Engine/Classes/AI/Navigation/NavigationTestingActor.h
--- a/Engine/Source/Runtime/Engine/Classes/AI/Navigation/NavigationTestingActor.h
+++ b/Engine/Source/Runtime/Engine/Classes/AI/Navigation/NavigationTestingActor.h
@@ -26,7 +26,7 @@ namespace ENavCostDisplay
 	};
 }
 
-UCLASS(hidecategories=(Object, Actor, Input, Rendering), HeaderGroup=Navigation, Blueprintable)
+UCLASS(hidecategories=(Object, Actor, Input, Rendering), Blueprintable)
 class ENGINE_API ANavigationTestingActor : public AActor, public INavAgentInterface, public INavPathObserverInterface
 {
 	GENERATED_UCLASS_BODY()
diff --git a/Engine/Source/Runtime/Engine/Classes/AI/Navigation/PathFollowingComponent.h b/Engine/Source/Runtime/Engine/Classes/AI/Navigation/PathFollowingComponent.h
--- a/Engine/Source/Runtime/Engine/Classes/AI/Navigation/PathFollowingComponent.h
+++ b/Engine/Source/Runtime/Engine/Classes/AI/Navigation/PathFollowingComponent.h
@@ -58,7 +58,7 @@ namespace EPathFollowingRequestResult
 	};
 }
 
-UCLASS(HeaderGroup=Component, config=Engine)
+UCLASS(config=Engine)
 class ENGINE_API UPathFollowingComponent : public UActorComponent, public IAIResourceInterface
 {
 	GENERATED_UCLASS_BODY()
diff --git a/Engine/Source/Runtime/Engine/Classes/Animation/AnimBlueprintGeneratedClass.h b/Engine/Source/Runtime/Engine/Classes/Animation/AnimBlueprintGeneratedClass.h
--- a/Engine/Source/Runtime/Engine/Classes/Animation/AnimBlueprintGeneratedClass.h
+++ b/Engine/Source/Runtime/Engine/Classes/Animation/AnimBlueprintGeneratedClass.h
@@ -10,6 +10,7 @@ struct FBakedAnimationStateMachine;
 class  UAnimInstance;
 class  UEdGraph;
 class  UEdGraphNode;
+class  UAnimGraphNode_StateResult;
 
 // This structure represents debugging information for a single state machine
 USTRUCT()
diff --git a/Engine/Source/Runtime/Engine/Classes/Animation/AnimNode_BlendSpaceEvaluator.h b/Engine/Source/Runtime/Engine/Classes/Animation/AnimNode_BlendSpaceEvaluator.h
--- a/Engine/Source/Runtime/Engine/Classes/Animation/AnimNode_BlendSpaceEvaluator.h
+++ b/Engine/Source/Runtime/Engine/Classes/Animation/AnimNode_BlendSpaceEvaluator.h
@@ -8,7 +8,7 @@
 // Evaluates a point in a blendspace, using a specific time input rather than advancing time internally.
 // Typically the playback position of the animation for this node will represent something other than time, like jump height.
 // This node will not trigger any notifies present in the associated sequence.
-USTRUCT(HeaderGroup = AnimationNode)
+USTRUCT()
 struct ENGINE_API FAnimNode_BlendSpaceEvaluator : public FAnimNode_BlendSpacePlayer
 {
 	GENERATED_USTRUCT_BODY()
diff --git a/Engine/Source/Runtime/Engine/Classes/Animation/AnimSingleNodeInstance.h b/Engine/Source/Runtime/Engine/Classes/Animation/AnimSingleNodeInstance.h
--- a/Engine/Source/Runtime/Engine/Classes/Animation/AnimSingleNodeInstance.h
+++ b/Engine/Source/Runtime/Engine/Classes/Animation/AnimSingleNodeInstance.h
@@ -7,6 +7,7 @@
  */
 
 #pragma once
+#include "AnimInstance.h"
 #include "AnimSingleNodeInstance.generated.h"
 
 DECLARE_DYNAMIC_DELEGATE(FPostEvaluateAnimEvent);
diff --git a/Engine/Source/Runtime/Engine/Classes/Animation/SkeletalMeshActor.h b/Engine/Source/Runtime/Engine/Classes/Animation/SkeletalMeshActor.h
--- a/Engine/Source/Runtime/Engine/Classes/Animation/SkeletalMeshActor.h
+++ b/Engine/Source/Runtime/Engine/Classes/Animation/SkeletalMeshActor.h
@@ -1,9 +1,9 @@
 // Copyright 1998-2014 Epic Games, Inc. All Rights Reserved.
 #pragma once
-
+#include "Components/SkeletalMeshComponent.h"
 #include "SkeletalMeshActor.generated.h"
 
-UCLASS(HeaderGroup=Anim, ClassGroup=ISkeletalMeshes, Blueprintable, ConversionRoot, meta=(ChildCanTick))
+UCLASS(ClassGroup=ISkeletalMeshes, Blueprintable, ConversionRoot, meta=(ChildCanTick))
 class ENGINE_API ASkeletalMeshActor : public AActor, public IMatineeAnimInterface
 {
 	GENERATED_UCLASS_BODY()
diff --git a/Engine/Source/Runtime/Engine/Classes/Atmosphere/AtmosphericFogComponent.h b/Engine/Source/Runtime/Engine/Classes/Atmosphere/AtmosphericFogComponent.h
--- a/Engine/Source/Runtime/Engine/Classes/Atmosphere/AtmosphericFogComponent.h
+++ b/Engine/Source/Runtime/Engine/Classes/Atmosphere/AtmosphericFogComponent.h
@@ -109,7 +109,7 @@ public:
 /**
  *	Used to create fogging effects such as clouds.
  */
-UCLASS(HeaderGroup=Component, ClassGroup=Rendering, collapsecategories, hidecategories=(Object, Mobility, Activation, "Components|Activation"), editinlinenew, meta=(BlueprintSpawnableComponent), MinimalAPI)
+UCLASS(ClassGroup=Rendering, collapsecategories, hidecategories=(Object, Mobility, Activation, "Components|Activation"), editinlinenew, meta=(BlueprintSpawnableComponent), MinimalAPI)
 class UAtmosphericFogComponent : public USceneComponent
 {
 	GENERATED_UCLASS_BODY()
diff --git a/Engine/Source/Runtime/Engine/Classes/Camera/CameraComponent.h b/Engine/Source/Runtime/Engine/Classes/Camera/CameraComponent.h
--- a/Engine/Source/Runtime/Engine/Classes/Camera/CameraComponent.h
+++ b/Engine/Source/Runtime/Engine/Classes/Camera/CameraComponent.h
@@ -7,7 +7,7 @@
   * Represents a camera viewpoint and settings, such as projection type, field of view, and post-process overrides.
   * The default behavior for an actor used as the camera view target is to look for an attached camera component and use its location, rotation, and settings.
   */
-UCLASS(HeaderGroup=Component, HideCategories=(Mobility, Rendering, LOD), ClassGroup=Camera, meta=(BlueprintSpawnableComponent), MinimalAPI)
+UCLASS(HideCategories=(Mobility, Rendering, LOD), ClassGroup=Camera, meta=(BlueprintSpawnableComponent), MinimalAPI)
 class UCameraComponent : public USceneComponent
 {
 	GENERATED_UCLASS_BODY()
diff --git a/Engine/Source/Runtime/Engine/Classes/Components/ActorComponent.h b/Engine/Source/Runtime/Engine/Classes/Components/ActorComponent.h
--- a/Engine/Source/Runtime/Engine/Classes/Components/ActorComponent.h
+++ b/Engine/Source/Runtime/Engine/Classes/Components/ActorComponent.h
@@ -3,9 +3,10 @@
 #pragma once
 
 #include "ComponentInstanceDataCache.h"
+#include "Engine/EngineTypes.h"
 #include "ActorComponent.generated.h"
 
-UCLASS(DefaultToInstanced, HeaderGroup=Component, abstract, hidecategories=(ComponentReplication))
+UCLASS(DefaultToInstanced, abstract, hidecategories=(ComponentReplication))
 class ENGINE_API UActorComponent : public UObject
 {
 	GENERATED_UCLASS_BODY()
diff --git a/Engine/Source/Runtime/Engine/Classes/Components/ApplicationLifecycleComponent.h b/Engine/Source/Runtime/Engine/Classes/Components/ApplicationLifecycleComponent.h
--- a/Engine/Source/Runtime/Engine/Classes/Components/ApplicationLifecycleComponent.h
+++ b/Engine/Source/Runtime/Engine/Classes/Components/ApplicationLifecycleComponent.h
@@ -5,7 +5,7 @@
 
 #include "ApplicationLifecycleComponent.generated.h"
 
-UCLASS(HeaderGroup=Component, ClassGroup=Utility, HideCategories=(Activation, "Components|Activation", Collision), meta=(BlueprintSpawnableComponent))
+UCLASS(ClassGroup=Utility, HideCategories=(Activation, "Components|Activation", Collision), meta=(BlueprintSpawnableComponent))
 class ENGINE_API UApplicationLifecycleComponent : public UActorComponent
 {
 	GENERATED_UCLASS_BODY()
diff --git a/Engine/Source/Runtime/Engine/Classes/Components/ArrowComponent.h b/Engine/Source/Runtime/Engine/Classes/Components/ArrowComponent.h
--- a/Engine/Source/Runtime/Engine/Classes/Components/ArrowComponent.h
+++ b/Engine/Source/Runtime/Engine/Classes/Components/ArrowComponent.h
@@ -7,7 +7,7 @@
 /** 
  * A simple arrow rendered using lines.Useful for indicating which way an object is facing.
  */
-UCLASS(HeaderGroup=Component, ClassGroup=Shapes, hidecategories=(Object,LOD,Physics,Lighting,TextureStreaming,Activation,"Components|Activation",Collision), editinlinenew, meta=(BlueprintSpawnableComponent), MinimalAPI)
+UCLASS(ClassGroup=Shapes, hidecategories=(Object,LOD,Physics,Lighting,TextureStreaming,Activation,"Components|Activation",Collision), editinlinenew, meta=(BlueprintSpawnableComponent), MinimalAPI)
 class UArrowComponent : public UPrimitiveComponent
 {
 	GENERATED_UCLASS_BODY()
diff --git a/Engine/Source/Runtime/Engine/Classes/Components/AudioComponent.h b/Engine/Source/Runtime/Engine/Classes/Components/AudioComponent.h
--- a/Engine/Source/Runtime/Engine/Classes/Components/AudioComponent.h
+++ b/Engine/Source/Runtime/Engine/Classes/Components/AudioComponent.h
@@ -2,6 +2,7 @@
 
 
 #pragma once
+#include "Audio.h"
 #include "AudioComponent.generated.h"
 
 /** called when we finish playing audio, either because it played to completion or because a Stop() call turned it off early */
@@ -89,7 +90,7 @@ struct FSoundParseParameters
 /**
  * Used to provide objects with audio.
  */
-UCLASS(HeaderGroup=Component, ClassGroup=(Audio, Common), hidecategories=(Object, ActorComponent, Physics, Rendering, Mobility, LOD), ShowCategories=Trigger, dependson=(AReverbVolume, UEngineTypes, USoundAttenuation), meta=(BlueprintSpawnableComponent))
+UCLASS(ClassGroup=(Audio, Common), hidecategories=(Object, ActorComponent, Physics, Rendering, Mobility, LOD), ShowCategories=Trigger, dependson=(AReverbVolume, UEngineTypes, USoundAttenuation), meta=(BlueprintSpawnableComponent))
 class ENGINE_API UAudioComponent : public USceneComponent
 {
 	GENERATED_UCLASS_BODY()
diff --git a/Engine/Source/Runtime/Engine/Classes/Components/BillboardComponent.h b/Engine/Source/Runtime/Engine/Classes/Components/BillboardComponent.h
--- a/Engine/Source/Runtime/Engine/Classes/Components/BillboardComponent.h
+++ b/Engine/Source/Runtime/Engine/Classes/Components/BillboardComponent.h
@@ -7,7 +7,7 @@
 /** 
  * A 2d texture that will be rendered always facing the camera.
  */
-UCLASS(HeaderGroup=Component, ClassGroup=Rendering, collapsecategories, hidecategories=(Object,Activation,"Components|Activation",Physics,Collision,Lighting,LOD,Mesh), editinlinenew, meta=(BlueprintSpawnableComponent),MinimalAPI)
+UCLASS(ClassGroup=Rendering, collapsecategories, hidecategories=(Object,Activation,"Components|Activation",Physics,Collision,Lighting,LOD,Mesh), editinlinenew, meta=(BlueprintSpawnableComponent),MinimalAPI)
 class UBillboardComponent : public UPrimitiveComponent
 {
 	GENERATED_UCLASS_BODY()
diff --git a/Engine/Source/Runtime/Engine/Classes/Components/BoxComponent.h b/Engine/Source/Runtime/Engine/Classes/Components/BoxComponent.h
--- a/Engine/Source/Runtime/Engine/Classes/Components/BoxComponent.h
+++ b/Engine/Source/Runtime/Engine/Classes/Components/BoxComponent.h
@@ -7,7 +7,7 @@
 /** 
  * A box generally used for simple collision. Bounds are rendered as lines in the editor.
  */
-UCLASS(HeaderGroup=Component, ClassGroup=Shapes, hidecategories=(Object,LOD,Lighting,TextureStreaming), editinlinenew, meta=(BlueprintSpawnableComponent))
+UCLASS(ClassGroup=Shapes, hidecategories=(Object,LOD,Lighting,TextureStreaming), editinlinenew, meta=(BlueprintSpawnableComponent))
 class ENGINE_API UBoxComponent : public UShapeComponent
 {
 	GENERATED_UCLASS_BODY()
diff --git a/Engine/Source/Runtime/Engine/Classes/Components/BrushComponent.h b/Engine/Source/Runtime/Engine/Classes/Components/BrushComponent.h
--- a/Engine/Source/Runtime/Engine/Classes/Components/BrushComponent.h
+++ b/Engine/Source/Runtime/Engine/Classes/Components/BrushComponent.h
@@ -4,7 +4,7 @@
 #pragma once
 #include "BrushComponent.generated.h"
 
-UCLASS(HeaderGroup=Component, editinlinenew, MinimalAPI, showcategories=(Mobility), hidecategories=(Physics, Lighting, LOD, Rendering, TextureStreaming, Transform, Activation, "Components|Activation"))
+UCLASS(editinlinenew, MinimalAPI, showcategories=(Mobility), hidecategories=(Physics, Lighting, LOD, Rendering, TextureStreaming, Transform, Activation, "Components|Activation"))
 class UBrushComponent : public UPrimitiveComponent
 {
 	GENERATED_UCLASS_BODY()
diff --git a/Engine/Source/Runtime/Engine/Classes/Components/ChildActorComponent.h b/Engine/Source/Runtime/Engine/Classes/Components/ChildActorComponent.h
--- a/Engine/Source/Runtime/Engine/Classes/Components/ChildActorComponent.h
+++ b/Engine/Source/Runtime/Engine/Classes/Components/ChildActorComponent.h
@@ -5,7 +5,7 @@
 #include "ChildActorComponent.generated.h"
 
 /** A component that spawns an Actor when registered, and destroys it when unregistered.*/
-UCLASS(HeaderGroup=Component, ClassGroup=Utility, hidecategories=(Object,LOD,Physics,Lighting,TextureStreaming,Activation,"Components|Activation",Collision), meta=(BlueprintSpawnableComponent), MinimalAPI)
+UCLASS(ClassGroup=Utility, hidecategories=(Object,LOD,Physics,Lighting,TextureStreaming,Activation,"Components|Activation",Collision), meta=(BlueprintSpawnableComponent), MinimalAPI)
 class UChildActorComponent : public USceneComponent
 {
 	GENERATED_UCLASS_BODY()
diff --git a/Engine/Source/Runtime/Engine/Classes/Components/DecalComponent.h b/Engine/Source/Runtime/Engine/Classes/Components/DecalComponent.h
--- a/Engine/Source/Runtime/Engine/Classes/Components/DecalComponent.h
+++ b/Engine/Source/Runtime/Engine/Classes/Components/DecalComponent.h
@@ -9,7 +9,7 @@
 /** 
  * A material that is rendered onto the surface of a mesh. A kind of 'bumper sticker' for a model.
  */
-UCLASS(hidecategories=(Collision, Object, Physics, SceneComponent, Activation, "Components|Activation", Physics, Mobility), ClassGroup=Rendering, HeaderGroup=Decal, MinimalAPI, meta=(BlueprintSpawnableComponent))
+UCLASS(hidecategories=(Collision, Object, Physics, SceneComponent, Activation, "Components|Activation", Physics, Mobility), ClassGroup=Rendering, MinimalAPI, meta=(BlueprintSpawnableComponent))
 class UDecalComponent : public USceneComponent
 {
 	GENERATED_UCLASS_BODY()
diff --git a/Engine/Source/Runtime/Engine/Classes/Components/DestructibleComponent.h b/Engine/Source/Runtime/Engine/Classes/Components/DestructibleComponent.h
--- a/Engine/Source/Runtime/Engine/Classes/Components/DestructibleComponent.h
+++ b/Engine/Source/Runtime/Engine/Classes/Components/DestructibleComponent.h
@@ -37,7 +37,7 @@ struct FDestructibleChunkInfo
  *
  *	The USkeletalMesh pointer in the base class (SkinnedMeshComponent) MUST be a DestructibleMesh
  */
-UCLASS(HeaderGroup=Component, ClassGroup=Physics, hidecategories=(Object,Mesh,"Components|SkinnedMesh",Mirroring,Activation,"Components|Activation"), config=Engine, editinlinenew, meta=(BlueprintSpawnableComponent))
+UCLASS(ClassGroup=Physics, hidecategories=(Object,Mesh,"Components|SkinnedMesh",Mirroring,Activation,"Components|Activation"), config=Engine, editinlinenew, meta=(BlueprintSpawnableComponent))
 class ENGINE_API UDestructibleComponent : public USkinnedMeshComponent
 {
 	GENERATED_UCLASS_BODY()
diff --git a/Engine/Source/Runtime/Engine/Classes/Components/DirectionalLightComponent.h b/Engine/Source/Runtime/Engine/Classes/Components/DirectionalLightComponent.h
--- a/Engine/Source/Runtime/Engine/Classes/Components/DirectionalLightComponent.h
+++ b/Engine/Source/Runtime/Engine/Classes/Components/DirectionalLightComponent.h
@@ -7,7 +7,7 @@
 /**
  * A light component that has parallel rays. Will provide a uniform lighting across any affected surface (eg. The Sun). This will affect all objects in the defined light-mass importance volume.
  */
-UCLASS(HeaderGroup=Light, ClassGroup=Lights, hidecategories=(Object, LightProfiles), dependson=UEngineTypes, editinlinenew, meta=(BlueprintSpawnableComponent), MinimalAPI)
+UCLASS(ClassGroup=Lights, hidecategories=(Object, LightProfiles), dependson=UEngineTypes, editinlinenew, meta=(BlueprintSpawnableComponent), MinimalAPI)
 class UDirectionalLightComponent : public ULightComponent
 {
 	GENERATED_UCLASS_BODY()
diff --git a/Engine/Source/Runtime/Engine/Classes/Components/ExponentialHeightFogComponent.h b/Engine/Source/Runtime/Engine/Classes/Components/ExponentialHeightFogComponent.h
--- a/Engine/Source/Runtime/Engine/Classes/Components/ExponentialHeightFogComponent.h
+++ b/Engine/Source/Runtime/Engine/Classes/Components/ExponentialHeightFogComponent.h
@@ -7,7 +7,7 @@
 /**
  *	Used to create fogging effects such as clouds but with a density that is related to the height of the fog.
  */
-UCLASS(HeaderGroup=Component, ClassGroup=Rendering, collapsecategories, hidecategories=(Object, Mobility), editinlinenew, meta=(BlueprintSpawnableComponent))
+UCLASS(ClassGroup=Rendering, collapsecategories, hidecategories=(Object, Mobility), editinlinenew, meta=(BlueprintSpawnableComponent))
 class ENGINE_API UExponentialHeightFogComponent : public USceneComponent
 {
 	GENERATED_UCLASS_BODY()
diff --git a/Engine/Source/Runtime/Engine/Classes/Components/InputComponent.h b/Engine/Source/Runtime/Engine/Classes/Components/InputComponent.h
--- a/Engine/Source/Runtime/Engine/Classes/Components/InputComponent.h
+++ b/Engine/Source/Runtime/Engine/Classes/Components/InputComponent.h
@@ -411,7 +411,7 @@ namespace EControllerAnalogStick
 	};
 }
 
-UCLASS(transient, HeaderGroup=Component, config=Input, dependson=UPlayerInput, hidecategories=(Activation, "Components|Activation"))
+UCLASS(transient, config=Input, hidecategories=(Activation, "Components|Activation"))
 class ENGINE_API UInputComponent : public UActorComponent
 {
 	GENERATED_UCLASS_BODY()
diff --git a/Engine/Source/Runtime/Engine/Classes/Components/InstancedStaticMeshComponent.h b/Engine/Source/Runtime/Engine/Classes/Components/InstancedStaticMeshComponent.h
--- a/Engine/Source/Runtime/Engine/Classes/Components/InstancedStaticMeshComponent.h
+++ b/Engine/Source/Runtime/Engine/Classes/Components/InstancedStaticMeshComponent.h
@@ -57,7 +57,7 @@ struct FInstancedStaticMeshMappingInfo
 };
 
 /** A component that efficiently renders multiple instances of the same StaticMesh. */
-UCLASS(HeaderGroup=Component, ClassGroup=Rendering, meta=(BlueprintSpawnableComponent))
+UCLASS(ClassGroup=Rendering, meta=(BlueprintSpawnableComponent))
 class ENGINE_API UInstancedStaticMeshComponent : public UStaticMeshComponent
 {
 	GENERATED_UCLASS_BODY()
diff --git a/Engine/Source/Runtime/Engine/Classes/Components/MaterialBillboardComponent.h b/Engine/Source/Runtime/Engine/Classes/Components/MaterialBillboardComponent.h
--- a/Engine/Source/Runtime/Engine/Classes/Components/MaterialBillboardComponent.h
+++ b/Engine/Source/Runtime/Engine/Classes/Components/MaterialBillboardComponent.h
@@ -50,7 +50,7 @@ struct FMaterialSpriteElement
 /** 
  * A 2d material that will be rendered always facing the camera.
  */
-UCLASS(HeaderGroup=Component, ClassGroup=Rendering, collapsecategories, hidecategories=Object, editinlinenew, meta=(BlueprintSpawnableComponent))
+UCLASS(ClassGroup=Rendering, collapsecategories, hidecategories=Object, editinlinenew, meta=(BlueprintSpawnableComponent))
 class UMaterialBillboardComponent : public UPrimitiveComponent
 {
 	GENERATED_UCLASS_BODY()
diff --git a/Engine/Source/Runtime/Engine/Classes/Components/ModelComponent.h b/Engine/Source/Runtime/Engine/Classes/Components/ModelComponent.h
--- a/Engine/Source/Runtime/Engine/Classes/Components/ModelComponent.h
+++ b/Engine/Source/Runtime/Engine/Classes/Components/ModelComponent.h
@@ -4,7 +4,12 @@
 #pragma once
 #include "ModelComponent.generated.h"
 
-UCLASS(HeaderGroup=Component, MinimalAPI)
+//
+// Forward declarations
+//
+class FModelElement;
+
+UCLASS(MinimalAPI)
 class UModelComponent : public UPrimitiveComponent, public IInterface_CollisionDataProvider
 {
 	GENERATED_UCLASS_BODY()
diff --git a/Engine/Source/Runtime/Engine/Classes/Components/PawnNoiseEmitterComponent.h b/Engine/Source/Runtime/Engine/Classes/Components/PawnNoiseEmitterComponent.h
--- a/Engine/Source/Runtime/Engine/Classes/Components/PawnNoiseEmitterComponent.h
+++ b/Engine/Source/Runtime/Engine/Classes/Components/PawnNoiseEmitterComponent.h
@@ -9,7 +9,7 @@
  * PawnNoiseEmitterComponent tracks noise event data used by SensingComponents to hear a Pawn.
  * This component is intended to exist on either a Pawn or its Controller. It does nothing on network clients.
  */
-UCLASS(HeaderGroup=Component, ClassGroup=AI, meta=(BlueprintSpawnableComponent))
+UCLASS(ClassGroup=AI, meta=(BlueprintSpawnableComponent))
 class ENGINE_API UPawnNoiseEmitterComponent : public UActorComponent
 {
 	GENERATED_UCLASS_BODY()
diff --git a/Engine/Source/Runtime/Engine/Classes/Components/PawnSensingComponent.h b/Engine/Source/Runtime/Engine/Classes/Components/PawnSensingComponent.h
--- a/Engine/Source/Runtime/Engine/Classes/Components/PawnSensingComponent.h
+++ b/Engine/Source/Runtime/Engine/Classes/Components/PawnSensingComponent.h
@@ -9,7 +9,7 @@
  * SensingComponent encapsulates sensory (ie sight and hearing) settings and functionality for an Actor,
  * allowing the actor to see/hear Pawns in the world. It does nothing on network clients.
  */
-UCLASS(HeaderGroup=Component, ClassGroup=AI, HideCategories=(Activation, "Components|Activation", Collision), meta=(BlueprintSpawnableComponent))
+UCLASS(ClassGroup=AI, HideCategories=(Activation, "Components|Activation", Collision), meta=(BlueprintSpawnableComponent))
 class ENGINE_API UPawnSensingComponent : public UActorComponent
 {
 	GENERATED_UCLASS_BODY()
diff --git a/Engine/Source/Runtime/Engine/Classes/Components/PointLightComponent.h b/Engine/Source/Runtime/Engine/Classes/Components/PointLightComponent.h
--- a/Engine/Source/Runtime/Engine/Classes/Components/PointLightComponent.h
+++ b/Engine/Source/Runtime/Engine/Classes/Components/PointLightComponent.h
@@ -7,7 +7,7 @@
 /**
  * A light component which emits light from a single point equally in all directions.
  */
-UCLASS(HeaderGroup=Light, ClassGroup=Lights, hidecategories=(Object, LightShafts), editinlinenew, meta=(BlueprintSpawnableComponent))
+UCLASS(ClassGroup=Lights, hidecategories=(Object, LightShafts), editinlinenew, meta=(BlueprintSpawnableComponent))
 class ENGINE_API UPointLightComponent : public ULightComponent
 {
 	GENERATED_UCLASS_BODY()
diff --git a/Engine/Source/Runtime/Engine/Classes/Components/PoseableMeshComponent.h b/Engine/Source/Runtime/Engine/Classes/Components/PoseableMeshComponent.h
--- a/Engine/Source/Runtime/Engine/Classes/Components/PoseableMeshComponent.h
+++ b/Engine/Source/Runtime/Engine/Classes/Components/PoseableMeshComponent.h
@@ -1,9 +1,7 @@
 // Copyright 1998-2014 Epic Games, Inc. All Rights Reserved.
 
-
-
-
 #pragma once
+#include "Animation/AnimationAsset.h"
 #include "PoseableMeshComponent.generated.h"
 
 UENUM()
@@ -23,7 +21,7 @@ namespace EBoneSpaces
 /**
  *	UPoseableMeshComponent that allows bone transforms to be driven by blueprint.
  */
-UCLASS(HeaderGroup=Component, ClassGroup=Rendering, hidecategories=(Object,Physics), config=Engine, editinlinenew, meta=(BlueprintSpawnableComponent))
+UCLASS(ClassGroup=Rendering, hidecategories=(Object,Physics), config=Engine, editinlinenew, meta=(BlueprintSpawnableComponent))
 class ENGINE_API UPoseableMeshComponent : public USkinnedMeshComponent
 {
 	GENERATED_UCLASS_BODY()
diff --git a/Engine/Source/Runtime/Engine/Classes/Components/PostProcessComponent.h b/Engine/Source/Runtime/Engine/Classes/Components/PostProcessComponent.h
--- a/Engine/Source/Runtime/Engine/Classes/Components/PostProcessComponent.h
+++ b/Engine/Source/Runtime/Engine/Classes/Components/PostProcessComponent.h
@@ -1,6 +1,6 @@
 // Copyright 1998-2013 Epic Games, Inc. All Rights Reserved.
 
-
+#pragma once
 #include "PostProcessComponent.generated.h"
 
 /**
@@ -8,7 +8,7 @@
  *	Will use a parent UShapeComponent to provide volume data if available.
  */
 
-UCLASS(HeaderGroup = Component, ClassGroup = Rendering, collapsecategories, hidecategories = (Object), editinlinenew, meta = (BlueprintSpawnableComponent), MinimalAPI)
+UCLASS(ClassGroup = Rendering, collapsecategories, hidecategories = (Object), editinlinenew, meta = (BlueprintSpawnableComponent), MinimalAPI)
 class UPostProcessComponent : public USceneComponent, public IInterface_PostProcessVolume
 {
 	GENERATED_UCLASS_BODY()
diff --git a/Engine/Source/Runtime/Engine/Classes/Components/PrimitiveComponent.h b/Engine/Source/Runtime/Engine/Classes/Components/PrimitiveComponent.h
--- a/Engine/Source/Runtime/Engine/Classes/Components/PrimitiveComponent.h
+++ b/Engine/Source/Runtime/Engine/Classes/Components/PrimitiveComponent.h
@@ -183,7 +183,7 @@ DECLARE_DYNAMIC_MULTICAST_DELEGATE_TwoParams( FComponentOnInputTouchEndSignature
 DECLARE_DYNAMIC_MULTICAST_DELEGATE_TwoParams( FComponentBeginTouchOverSignature, ETouchIndex::Type, FingerIndex, UPrimitiveComponent*, TouchedComponent );
 DECLARE_DYNAMIC_MULTICAST_DELEGATE_TwoParams( FComponentEndTouchOverSignature, ETouchIndex::Type, FingerIndex, UPrimitiveComponent*, TouchedComponent );
 
-UCLASS(dependson=(UScene, ULightComponent, UEngineTypes), HeaderGroup=Component, abstract, HideCategories=(Mobility))
+UCLASS(dependson=(UScene, ULightComponent, UEngineTypes), abstract, HideCategories=(Mobility))
 class ENGINE_API UPrimitiveComponent : public USceneComponent
 {
 	GENERATED_UCLASS_BODY()
diff --git a/Engine/Source/Runtime/Engine/Classes/Components/ReflectionCaptureComponent.h b/Engine/Source/Runtime/Engine/Classes/Components/ReflectionCaptureComponent.h
--- a/Engine/Source/Runtime/Engine/Classes/Components/ReflectionCaptureComponent.h
+++ b/Engine/Source/Runtime/Engine/Classes/Components/ReflectionCaptureComponent.h
@@ -66,7 +66,7 @@ private:
 };
 
 // -> will be exported to EngineDecalClasses.h
-UCLASS(abstract, hidecategories=(Collision, Object, Physics, SceneComponent, Activation, "Components|Activation", Mobility), HeaderGroup=Decal, MinimalAPI)
+UCLASS(abstract, hidecategories=(Collision, Object, Physics, SceneComponent, Activation, "Components|Activation", Mobility), MinimalAPI)
 class UReflectionCaptureComponent : public USceneComponent
 {
 	GENERATED_UCLASS_BODY()
diff --git a/Engine/Source/Runtime/Engine/Classes/Components/SceneCaptureComponent.h b/Engine/Source/Runtime/Engine/Classes/Components/SceneCaptureComponent.h
--- a/Engine/Source/Runtime/Engine/Classes/Components/SceneCaptureComponent.h
+++ b/Engine/Source/Runtime/Engine/Classes/Components/SceneCaptureComponent.h
@@ -5,7 +5,7 @@
 #include "SceneCaptureComponent.generated.h"
 
 	// -> will be exported to EngineDecalClasses.h
-UCLASS(hidecategories=(abstract, Collision, Object, Physics, SceneComponent, Mobility), HeaderGroup=Decal, MinimalAPI)
+UCLASS(hidecategories=(abstract, Collision, Object, Physics, SceneComponent, Mobility), MinimalAPI)
 class USceneCaptureComponent : public USceneComponent
 {
 	GENERATED_UCLASS_BODY()
diff --git a/Engine/Source/Runtime/Engine/Classes/Components/SceneCaptureComponent2D.h b/Engine/Source/Runtime/Engine/Classes/Components/SceneCaptureComponent2D.h
--- a/Engine/Source/Runtime/Engine/Classes/Components/SceneCaptureComponent2D.h
+++ b/Engine/Source/Runtime/Engine/Classes/Components/SceneCaptureComponent2D.h
@@ -16,7 +16,7 @@ enum ESceneCaptureSource
 /**
  *	Used to capture a 'snapshot' of the scene from a single plane and feed it to a render target.
  */
-UCLASS(hidecategories=(Collision, Object, Physics, SceneComponent), ClassGroup=Rendering, HeaderGroup=Decal, MinimalAPI, editinlinenew, meta=(BlueprintSpawnableComponent))
+UCLASS(hidecategories=(Collision, Object, Physics, SceneComponent), ClassGroup=Rendering, MinimalAPI, editinlinenew, meta=(BlueprintSpawnableComponent))
 class USceneCaptureComponent2D : public USceneCaptureComponent
 {
 	GENERATED_UCLASS_BODY()
diff --git a/Engine/Source/Runtime/Engine/Classes/Components/SceneCaptureComponentCube.h b/Engine/Source/Runtime/Engine/Classes/Components/SceneCaptureComponentCube.h
--- a/Engine/Source/Runtime/Engine/Classes/Components/SceneCaptureComponentCube.h
+++ b/Engine/Source/Runtime/Engine/Classes/Components/SceneCaptureComponentCube.h
@@ -6,7 +6,7 @@
 /**
  *	Used to capture a 'snapshot' of the scene from a 6 planes and feed it to a render target.
  */
-UCLASS(hidecategories = (Collision, Object, Physics, SceneComponent), ClassGroup=Rendering, HeaderGroup = Decal, MinimalAPI, editinlinenew, meta = (BlueprintSpawnableComponent))
+UCLASS(hidecategories = (Collision, Object, Physics, SceneComponent), ClassGroup=Rendering, MinimalAPI, editinlinenew, meta = (BlueprintSpawnableComponent))
 class USceneCaptureComponentCube : public USceneCaptureComponent
 {
 	GENERATED_UCLASS_BODY()
diff --git a/Engine/Source/Runtime/Engine/Classes/Components/SceneComponent.h b/Engine/Source/Runtime/Engine/Classes/Components/SceneComponent.h
--- a/Engine/Source/Runtime/Engine/Classes/Components/SceneComponent.h
+++ b/Engine/Source/Runtime/Engine/Classes/Components/SceneComponent.h
@@ -2,6 +2,7 @@
 
 #pragma once
 
+#include "Components.h"
 #include "SceneComponent.generated.h"
 
 /** Overlap info consisting of the primitive and the body that is overlapping */
@@ -62,7 +63,7 @@ FORCEINLINE void operator|=(EMoveComponentFlags& Dest,EMoveComponentFlags Arg)
 
 
 /** A SceneComponent has a transform and supports attachment, but has no rendering or collision capabilities. Useful as a 'dummy' component in the hierarchy to offset others. */
-UCLASS(HeaderGroup=Component, ClassGroup=Utility, BlueprintType, HideCategories=(Trigger), meta=(BlueprintSpawnableComponent))
+UCLASS(ClassGroup=Utility, BlueprintType, HideCategories=(Trigger), meta=(BlueprintSpawnableComponent))
 class ENGINE_API USceneComponent : public UActorComponent
 {
 	GENERATED_UCLASS_BODY()
diff --git a/Engine/Source/Runtime/Engine/Classes/Components/ShapeComponent.h b/Engine/Source/Runtime/Engine/Classes/Components/ShapeComponent.h
--- a/Engine/Source/Runtime/Engine/Classes/Components/ShapeComponent.h
+++ b/Engine/Source/Runtime/Engine/Classes/Components/ShapeComponent.h
@@ -4,7 +4,7 @@
 #pragma once
 #include "ShapeComponent.generated.h"
 
-UCLASS(abstract, HeaderGroup=Component, hidecategories=(Object,LOD,Lighting,TextureStreaming,Activation,"Components|Activation"), editinlinenew, MinimalAPI, meta=(BlueprintSpawnableComponent), showcategories=(Mobility))
+UCLASS(abstract, hidecategories=(Object,LOD,Lighting,TextureStreaming,Activation,"Components|Activation"), editinlinenew, MinimalAPI, meta=(BlueprintSpawnableComponent), showcategories=(Mobility))
 class UShapeComponent : public UPrimitiveComponent
 {
 	GENERATED_UCLASS_BODY()
diff --git a/Engine/Source/Runtime/Engine/Classes/Components/SkeletalMeshComponent.h b/Engine/Source/Runtime/Engine/Classes/Components/SkeletalMeshComponent.h
--- a/Engine/Source/Runtime/Engine/Classes/Components/SkeletalMeshComponent.h
+++ b/Engine/Source/Runtime/Engine/Classes/Components/SkeletalMeshComponent.h
@@ -218,7 +218,7 @@ struct FSkeletalMeshComponentPreClothTickFunction : public FTickFunction
 /**
  * SkeletalMeshComponent is a mesh that supports skeletal animation and physics.
  */
-UCLASS(HeaderGroup=Component, ClassGroup=(Rendering, Common), hidecategories=Object, config=Engine, editinlinenew, meta=(BlueprintSpawnableComponent))
+UCLASS(ClassGroup=(Rendering, Common), hidecategories=Object, config=Engine, editinlinenew, meta=(BlueprintSpawnableComponent))
 class ENGINE_API USkeletalMeshComponent : public USkinnedMeshComponent
 {
 	GENERATED_UCLASS_BODY()
diff --git a/Engine/Source/Runtime/Engine/Classes/Components/SkinnedMeshComponent.h b/Engine/Source/Runtime/Engine/Classes/Components/SkinnedMeshComponent.h
--- a/Engine/Source/Runtime/Engine/Classes/Components/SkinnedMeshComponent.h
+++ b/Engine/Source/Runtime/Engine/Classes/Components/SkinnedMeshComponent.h
@@ -8,8 +8,14 @@
  */
 
 #pragma once
+#include "SkeletalMeshTypes.h"
 #include "SkinnedMeshComponent.generated.h"
 
+//
+// Forward declarations
+//
+class FSkeletalMeshResource;
+
 //
 // Bone Visibility.
 //
@@ -119,7 +125,7 @@ struct FSkelMeshComponentLODInfo
 	
 };
 
-UCLASS(HeaderGroup=Component, hidecategories=Object, config=Engine, editinlinenew, abstract)
+UCLASS(hidecategories=Object, config=Engine, editinlinenew, abstract)
 class ENGINE_API USkinnedMeshComponent : public UMeshComponent
 {
 	GENERATED_UCLASS_BODY()
diff --git a/Engine/Source/Runtime/Engine/Classes/Components/SkyLightComponent.h b/Engine/Source/Runtime/Engine/Classes/Components/SkyLightComponent.h
--- a/Engine/Source/Runtime/Engine/Classes/Components/SkyLightComponent.h
+++ b/Engine/Source/Runtime/Engine/Classes/Components/SkyLightComponent.h
@@ -81,7 +81,7 @@ enum ESkyLightSourceType
 	SLS_MAX,
 };
 
-UCLASS(HeaderGroup=Component, ClassGroup=Lights, HideCategories=(Trigger,Activation,"Components|Activation",Physics), meta=(BlueprintSpawnableComponent))
+UCLASS(ClassGroup=Lights, HideCategories=(Trigger,Activation,"Components|Activation",Physics), meta=(BlueprintSpawnableComponent))
 class ENGINE_API USkyLightComponent : public ULightComponentBase
 {
 	GENERATED_UCLASS_BODY()
diff --git a/Engine/Source/Runtime/Engine/Classes/Components/SplineComponent.h b/Engine/Source/Runtime/Engine/Classes/Components/SplineComponent.h
--- a/Engine/Source/Runtime/Engine/Classes/Components/SplineComponent.h
+++ b/Engine/Source/Runtime/Engine/Classes/Components/SplineComponent.h
@@ -4,7 +4,7 @@
 #pragma once
 #include "SplineComponent.generated.h"
 
-UCLASS(HeaderGroup=Spline, MinimalAPI)
+UCLASS(MinimalAPI)
 class USplineComponent : public UPrimitiveComponent
 {
 	GENERATED_UCLASS_BODY()
diff --git a/Engine/Source/Runtime/Engine/Classes/Components/SpotLightComponent.h b/Engine/Source/Runtime/Engine/Classes/Components/SpotLightComponent.h
--- a/Engine/Source/Runtime/Engine/Classes/Components/SpotLightComponent.h
+++ b/Engine/Source/Runtime/Engine/Classes/Components/SpotLightComponent.h
@@ -7,7 +7,7 @@
 /**
  * A spot light component emits a directional cone shaped light (Eg a Torch).
  */
-UCLASS(HeaderGroup=Light, ClassGroup=Lights, hidecategories=Object, editinlinenew, meta=(BlueprintSpawnableComponent))
+UCLASS(ClassGroup=Lights, hidecategories=Object, editinlinenew, meta=(BlueprintSpawnableComponent))
 class ENGINE_API USpotLightComponent : public UPointLightComponent
 {
 	GENERATED_UCLASS_BODY()
diff --git a/Engine/Source/Runtime/Engine/Classes/Components/StaticMeshComponent.h b/Engine/Source/Runtime/Engine/Classes/Components/StaticMeshComponent.h
--- a/Engine/Source/Runtime/Engine/Classes/Components/StaticMeshComponent.h
+++ b/Engine/Source/Runtime/Engine/Classes/Components/StaticMeshComponent.h
@@ -2,6 +2,8 @@
 
 
 #pragma once
+#include "LightMap.h"
+#include "ShadowMap.h"
 #include "StaticMeshComponent.generated.h"
 
 /** Cached vertex information at the time the mesh was painted. */
@@ -169,7 +171,7 @@ public:
 };
 
 /** A StaticMeshComponent is a mesh that does not animate. */
-UCLASS(HeaderGroup=Component, ClassGroup=(Rendering, Common), hidecategories=(Object,Activation,"Components|Activation"), ShowCategories=(Mobility), dependson=ULightmassPrimitiveSettingsObject, editinlinenew, meta=(BlueprintSpawnableComponent))
+UCLASS(ClassGroup=(Rendering, Common), hidecategories=(Object,Activation,"Components|Activation"), ShowCategories=(Mobility), dependson=ULightmassPrimitiveSettingsObject, editinlinenew, meta=(BlueprintSpawnableComponent))
 class ENGINE_API UStaticMeshComponent : public UMeshComponent
 {
 	GENERATED_UCLASS_BODY()
diff --git a/Engine/Source/Runtime/Engine/Classes/Components/TextRenderComponent.h b/Engine/Source/Runtime/Engine/Classes/Components/TextRenderComponent.h
--- a/Engine/Source/Runtime/Engine/Classes/Components/TextRenderComponent.h
+++ b/Engine/Source/Runtime/Engine/Classes/Components/TextRenderComponent.h
@@ -31,7 +31,7 @@ enum EVerticalTextAligment
 /**
  * Renders text in the world with given font. Contains usual font related attributes such as Scale, Alignment, Color etc.
  */
-UCLASS(HeaderGroup=Component, ClassGroup=Rendering, hidecategories=(Object,LOD,Physics,TextureStreaming,Activation,"Components|Activation",Collision), editinlinenew, meta=(BlueprintSpawnableComponent = ""), MinimalAPI)
+UCLASS(ClassGroup=Rendering, hidecategories=(Object,LOD,Physics,TextureStreaming,Activation,"Components|Activation",Collision), editinlinenew, meta=(BlueprintSpawnableComponent = ""), MinimalAPI)
 class UTextRenderComponent : public UPrimitiveComponent
 {
 	GENERATED_UCLASS_BODY()
diff --git a/Engine/Source/Runtime/Engine/Classes/Components/VectorFieldComponent.h b/Engine/Source/Runtime/Engine/Classes/Components/VectorFieldComponent.h
--- a/Engine/Source/Runtime/Engine/Classes/Components/VectorFieldComponent.h
+++ b/Engine/Source/Runtime/Engine/Classes/Components/VectorFieldComponent.h
@@ -8,7 +8,7 @@
 /** 
  * A Component referencing a vector field.
  */
-UCLASS(HeaderGroup=Component, ClassGroup=Rendering, hidecategories=Object, meta=(BlueprintSpawnableComponent), MinimalAPI)
+UCLASS(ClassGroup=Rendering, hidecategories=Object, meta=(BlueprintSpawnableComponent), MinimalAPI)
 class UVectorFieldComponent : public UPrimitiveComponent
 {
 	GENERATED_UCLASS_BODY()
diff --git a/Engine/Source/Runtime/Engine/Classes/Components/WindDirectionalSourceComponent.h b/Engine/Source/Runtime/Engine/Classes/Components/WindDirectionalSourceComponent.h
--- a/Engine/Source/Runtime/Engine/Classes/Components/WindDirectionalSourceComponent.h
+++ b/Engine/Source/Runtime/Engine/Classes/Components/WindDirectionalSourceComponent.h
@@ -4,7 +4,7 @@
 #pragma once
 #include "WindDirectionalSourceComponent.generated.h"
 
-UCLASS(HeaderGroup=Component, collapsecategories, hidecategories=(Object, Mobility), editinlinenew)
+UCLASS(collapsecategories, hidecategories=(Object, Mobility), editinlinenew)
 class UWindDirectionalSourceComponent : public USceneComponent
 {
 	GENERATED_UCLASS_BODY()
diff --git a/Engine/Source/Runtime/Engine/Classes/Debug/GameplayDebuggingControllerComponent.h b/Engine/Source/Runtime/Engine/Classes/Debug/GameplayDebuggingControllerComponent.h
--- a/Engine/Source/Runtime/Engine/Classes/Debug/GameplayDebuggingControllerComponent.h
+++ b/Engine/Source/Runtime/Engine/Classes/Debug/GameplayDebuggingControllerComponent.h
@@ -7,7 +7,7 @@
 #pragma once
 #include "GameplayDebuggingControllerComponent.generated.h"
 
-UCLASS(config=Engine, HeaderGroup=Component)
+UCLASS(config=Engine)
 class ENGINE_API UGameplayDebuggingControllerComponent : public UActorComponent
 {
 	GENERATED_UCLASS_BODY()
diff --git a/Engine/Source/Runtime/Engine/Classes/Engine/BoxReflectionCapture.h b/Engine/Source/Runtime/Engine/Classes/Engine/BoxReflectionCapture.h
--- a/Engine/Source/Runtime/Engine/Classes/Engine/BoxReflectionCapture.h
+++ b/Engine/Source/Runtime/Engine/Classes/Engine/BoxReflectionCapture.h
@@ -8,7 +8,7 @@
 #pragma once
 #include "BoxReflectionCapture.generated.h"
 
-UCLASS(hidecategories=(Collision, Attachment, Actor), HeaderGroup=Decal, MinimalAPI)
+UCLASS(hidecategories=(Collision, Attachment, Actor), MinimalAPI)
 class ABoxReflectionCapture : public AReflectionCapture
 {
 	GENERATED_UCLASS_BODY()
diff --git a/Engine/Source/Runtime/Engine/Classes/Engine/ChildConnection.h b/Engine/Source/Runtime/Engine/Classes/Engine/ChildConnection.h
--- a/Engine/Source/Runtime/Engine/Classes/Engine/ChildConnection.h
+++ b/Engine/Source/Runtime/Engine/Classes/Engine/ChildConnection.h
@@ -5,7 +5,7 @@
 #pragma once
 #include "ChildConnection.generated.h"
 
-UCLASS(HeaderGroup=Network, MinimalAPI,transient,config=Engine)
+UCLASS(MinimalAPI,transient,config=Engine)
 class UChildConnection : public UNetConnection
 {
 	GENERATED_UCLASS_BODY()
diff --git a/Engine/Source/Runtime/Engine/Classes/Engine/Console.h b/Engine/Source/Runtime/Engine/Classes/Engine/Console.h
--- a/Engine/Source/Runtime/Engine/Classes/Engine/Console.h
+++ b/Engine/Source/Runtime/Engine/Classes/Engine/Console.h
@@ -59,7 +59,7 @@ struct FAutoCompleteNode
 	
 };
 
-UCLASS(Within=GameViewportClient, transient, config=Input, HeaderGroup=UserInterface)
+UCLASS(Within=GameViewportClient, transient, config=Input)
 class ENGINE_API UConsole : public UObject, public FOutputDevice
 {
 	GENERATED_UCLASS_BODY()
diff --git a/Engine/Source/Runtime/Engine/Classes/Engine/DecalActor.h b/Engine/Source/Runtime/Engine/Classes/Engine/DecalActor.h
--- a/Engine/Source/Runtime/Engine/Classes/Engine/DecalActor.h
+++ b/Engine/Source/Runtime/Engine/Classes/Engine/DecalActor.h
@@ -8,7 +8,12 @@
 #pragma once
 #include "DecalActor.generated.h"
 
-UCLASS(hidecategories=(Collision, Attachment, Actor, Input, Replication), HeaderGroup=Decal, MinimalAPI)
+//
+// Forward declarations.
+//
+class UBoxComponent;
+
+UCLASS(hidecategories=(Collision, Attachment, Actor, Input, Replication), MinimalAPI)
 class ADecalActor : public AActor
 {
 	GENERATED_UCLASS_BODY()
diff --git a/Engine/Source/Runtime/Engine/Classes/Engine/DirectionalLight.h b/Engine/Source/Runtime/Engine/Classes/Engine/DirectionalLight.h
--- a/Engine/Source/Runtime/Engine/Classes/Engine/DirectionalLight.h
+++ b/Engine/Source/Runtime/Engine/Classes/Engine/DirectionalLight.h
@@ -4,7 +4,7 @@
 
 #include "DirectionalLight.generated.h"
 
-UCLASS(HeaderGroup=Light, ClassGroup=(Lights, DirectionalLights), meta=(ChildCanTick))
+UCLASS(ClassGroup=(Lights, DirectionalLights), meta=(ChildCanTick))
 class ENGINE_API ADirectionalLight : public ALight
 {
 	GENERATED_UCLASS_BODY()
diff --git a/Engine/Source/Runtime/Engine/Classes/Engine/Engine.h b/Engine/Source/Runtime/Engine/Classes/Engine/Engine.h
--- a/Engine/Source/Runtime/Engine/Classes/Engine/Engine.h
+++ b/Engine/Source/Runtime/Engine/Classes/Engine/Engine.h
@@ -462,7 +462,7 @@ namespace EMatineeCaptureType
 
 class IAnalyticsProvider;
 
-UCLASS(HeaderGroup=GameEngine, abstract, config=Engine, transient)
+UCLASS(abstract, config=Engine, transient)
 class ENGINE_API UEngine : public UObject, public FExec
 {
 	GENERATED_UCLASS_BODY()
diff --git a/Engine/Source/Runtime/Engine/Classes/Engine/EngineBaseTypes.h b/Engine/Source/Runtime/Engine/Classes/Engine/EngineBaseTypes.h
--- a/Engine/Source/Runtime/Engine/Classes/Engine/EngineBaseTypes.h
+++ b/Engine/Source/Runtime/Engine/Classes/Engine/EngineBaseTypes.h
@@ -545,7 +545,7 @@ struct FExposureSettings
 };
 
 
-UCLASS(HeaderGroup=Base, abstract, config=Engine)
+UCLASS(abstract, config=Engine)
 class UEngineBaseTypes : public UObject
 {
 	GENERATED_UCLASS_BODY()
diff --git a/Engine/Source/Runtime/Engine/Classes/Engine/GameEngine.h b/Engine/Source/Runtime/Engine/Classes/Engine/GameEngine.h
--- a/Engine/Source/Runtime/Engine/Classes/Engine/GameEngine.h
+++ b/Engine/Source/Runtime/Engine/Classes/Engine/GameEngine.h
@@ -7,7 +7,7 @@
 #pragma once
 #include "GameEngine.generated.h"
 
-UCLASS(HeaderGroup=GameEngine, config=Engine, transient)
+UCLASS(config=Engine, transient)
 class ENGINE_API UGameEngine : public UEngine
 {
 	GENERATED_UCLASS_BODY()
diff --git a/Engine/Source/Runtime/Engine/Classes/Engine/GeneratedMeshAreaLight.h b/Engine/Source/Runtime/Engine/Classes/Engine/GeneratedMeshAreaLight.h
--- a/Engine/Source/Runtime/Engine/Classes/Engine/GeneratedMeshAreaLight.h
+++ b/Engine/Source/Runtime/Engine/Classes/Engine/GeneratedMeshAreaLight.h
@@ -8,7 +8,7 @@
 
 #include "GeneratedMeshAreaLight.generated.h"
 
-UCLASS(HeaderGroup=Light, notplaceable, MinimalAPI)
+UCLASS(notplaceable, MinimalAPI)
 class AGeneratedMeshAreaLight : public ASpotLight
 {
 	GENERATED_UCLASS_BODY()
diff --git a/Engine/Source/Runtime/Engine/Classes/Engine/InteractiveFoliageActor.h b/Engine/Source/Runtime/Engine/Classes/Engine/InteractiveFoliageActor.h
--- a/Engine/Source/Runtime/Engine/Classes/Engine/InteractiveFoliageActor.h
+++ b/Engine/Source/Runtime/Engine/Classes/Engine/InteractiveFoliageActor.h
@@ -6,7 +6,7 @@
 #pragma once
 #include "InteractiveFoliageActor.generated.h"
 
-UCLASS(HeaderGroup=Foliage, MinimalAPI)
+UCLASS(MinimalAPI)
 class AInteractiveFoliageActor : public AStaticMeshActor
 {
 	GENERATED_UCLASS_BODY()
diff --git a/Engine/Source/Runtime/Engine/Classes/Engine/LevelScriptActor.h b/Engine/Source/Runtime/Engine/Classes/Engine/LevelScriptActor.h
--- a/Engine/Source/Runtime/Engine/Classes/Engine/LevelScriptActor.h
+++ b/Engine/Source/Runtime/Engine/Classes/Engine/LevelScriptActor.h
@@ -4,7 +4,7 @@
 
 #include "LevelScriptActor.generated.h"
 
-UCLASS(HeaderGroup=LevelScript, notplaceable, meta=(KismetHideOverrides = "ReceiveAnyDamage,ReceivePointDamage,ReceiveRadialDamage,ReceiveActorBeginOverlap,ReceiveActorEndOverlap,ReceiveHit,ReceiveDestroyed,ReceiveActorBeginCursorOver,ReceiveActorEndCursorOver,ReceiveActorOnClicked,ReceiveActorOnReleased,ReceiveActorOnInputTouchBegin,ReceiveActorOnInputTouchEnd,ReceiveActorOnInputTouchEnter,ReceiveActorOnInputTouchLeave"), HideCategories=(Collision,Rendering,"Utilities|Orientation"))
+UCLASS(notplaceable, meta=(KismetHideOverrides = "ReceiveAnyDamage,ReceivePointDamage,ReceiveRadialDamage,ReceiveActorBeginOverlap,ReceiveActorEndOverlap,ReceiveHit,ReceiveDestroyed,ReceiveActorBeginCursorOver,ReceiveActorEndCursorOver,ReceiveActorOnClicked,ReceiveActorOnReleased,ReceiveActorOnInputTouchBegin,ReceiveActorOnInputTouchEnd,ReceiveActorOnInputTouchEnter,ReceiveActorOnInputTouchLeave"), HideCategories=(Collision,Rendering,"Utilities|Orientation"))
 class ENGINE_API ALevelScriptActor : public AActor
 {
 	GENERATED_UCLASS_BODY()
diff --git a/Engine/Source/Runtime/Engine/Classes/Engine/Light.h b/Engine/Source/Runtime/Engine/Classes/Engine/Light.h
--- a/Engine/Source/Runtime/Engine/Classes/Engine/Light.h
+++ b/Engine/Source/Runtime/Engine/Classes/Engine/Light.h
@@ -4,7 +4,7 @@
 
 #include "Light.generated.h"
 
-UCLASS(Abstract, HeaderGroup=Light, ClassGroup=Lights, hidecategories=(Input,Collision,Replication), ConversionRoot, meta=(ChildCanTick))
+UCLASS(Abstract, ClassGroup=Lights, hidecategories=(Input,Collision,Replication), ConversionRoot, meta=(ChildCanTick))
 class ENGINE_API ALight : public AActor
 {
 	GENERATED_UCLASS_BODY()
diff --git a/Engine/Source/Runtime/Engine/Classes/Engine/LocalPlayer.h b/Engine/Source/Runtime/Engine/Classes/Engine/LocalPlayer.h
--- a/Engine/Source/Runtime/Engine/Classes/Engine/LocalPlayer.h
+++ b/Engine/Source/Runtime/Engine/Classes/Engine/LocalPlayer.h
@@ -46,7 +46,7 @@ private:
  *	There may be several spawned in the case of splitscreen/coop
  *	There may be 0 spawned on servers
  */
-UCLASS(HeaderGroup=Network, Within=Engine, config=Engine, transient)
+UCLASS(Within=Engine, config=Engine, transient)
 class ENGINE_API ULocalPlayer : public UPlayer
 {
 	GENERATED_UCLASS_BODY()
diff --git a/Engine/Source/Runtime/Engine/Classes/Engine/NetConnection.h b/Engine/Source/Runtime/Engine/Classes/Engine/NetConnection.h
--- a/Engine/Source/Runtime/Engine/Classes/Engine/NetConnection.h
+++ b/Engine/Source/Runtime/Engine/Classes/Engine/NetConnection.h
@@ -91,7 +91,7 @@ struct DelayedPacket
 #define PING_ACK_DELAY 0.5
 
 
-UCLASS(HeaderGroup=Network, customConstructor, Abstract, MinimalAPI, transient, config=Engine, dependsOn=UGameEngine)
+UCLASS(customConstructor, Abstract, MinimalAPI, transient, config=Engine, dependsOn=UGameEngine)
 class UNetConnection : public UPlayer
 {
 	GENERATED_UCLASS_BODY()
diff --git a/Engine/Source/Runtime/Engine/Classes/Engine/NetDriver.h b/Engine/Source/Runtime/Engine/Classes/Engine/NetDriver.h
--- a/Engine/Source/Runtime/Engine/Classes/Engine/NetDriver.h
+++ b/Engine/Source/Runtime/Engine/Classes/Engine/NetDriver.h
@@ -87,7 +87,7 @@ struct FActorDestructionInfo
 };
 
 
-UCLASS(HeaderGroup=Network, dependson=UEngineTypes, Abstract, customConstructor, transient, MinimalAPI, config=Engine)
+UCLASS(dependson=UEngineTypes, Abstract, customConstructor, transient, MinimalAPI, config=Engine)
 class UNetDriver : public UObject, public FExec
 {
 	GENERATED_UCLASS_BODY()
diff --git a/Engine/Source/Runtime/Engine/Classes/Engine/Player.h b/Engine/Source/Runtime/Engine/Classes/Engine/Player.h
--- a/Engine/Source/Runtime/Engine/Classes/Engine/Player.h
+++ b/Engine/Source/Runtime/Engine/Classes/Engine/Player.h
@@ -7,7 +7,7 @@
 #pragma once
 #include "Player.generated.h"
 
-UCLASS(HeaderGroup=Network, transient, config=Engine)
+UCLASS(transient, config=Engine)
 class UPlayer : public UObject, public FExec
 {
 	GENERATED_UCLASS_BODY()
