#ttp 303393 UObject improvement: Remove header groups

diff --git a/Engine/Source/Runtime/Engine/Classes/Engine/PointLight.h b/Engine/Source/Runtime/Engine/Classes/Engine/PointLight.h
--- a/Engine/Source/Runtime/Engine/Classes/Engine/PointLight.h
+++ b/Engine/Source/Runtime/Engine/Classes/Engine/PointLight.h
@@ -4,7 +4,7 @@
 
 #include "PointLight.generated.h"
 
-UCLASS(HeaderGroup=Light, ClassGroup=(Lights, PointLights), meta=(ChildCanTick))
+UCLASS(ClassGroup=(Lights, PointLights), meta=(ChildCanTick))
 class ENGINE_API APointLight : public ALight
 {
 	GENERATED_UCLASS_BODY()
diff --git a/Engine/Source/Runtime/Engine/Classes/Engine/Scene.h b/Engine/Source/Runtime/Engine/Classes/Engine/Scene.h
--- a/Engine/Source/Runtime/Engine/Classes/Engine/Scene.h
+++ b/Engine/Source/Runtime/Engine/Classes/Engine/Scene.h
@@ -895,7 +895,7 @@ struct FPostProcessSettings
 	}
 };
 
-UCLASS(HeaderGroup=Scene)
+UCLASS()
 class UScene : public UObject
 {
 	GENERATED_UCLASS_BODY()
diff --git a/Engine/Source/Runtime/Engine/Classes/Engine/SceneCapture2D.h b/Engine/Source/Runtime/Engine/Classes/Engine/SceneCapture2D.h
--- a/Engine/Source/Runtime/Engine/Classes/Engine/SceneCapture2D.h
+++ b/Engine/Source/Runtime/Engine/Classes/Engine/SceneCapture2D.h
@@ -8,7 +8,7 @@
 #pragma once
 #include "SceneCapture2D.generated.h"
 
-UCLASS(hidecategories=(Collision, Material, Attachment, Actor), HeaderGroup=Decal, MinimalAPI)
+UCLASS(hidecategories=(Collision, Material, Attachment, Actor), MinimalAPI)
 class ASceneCapture2D : public ASceneCapture
 {
 	GENERATED_UCLASS_BODY()
diff --git a/Engine/Source/Runtime/Engine/Classes/Engine/SceneCaptureCube.h b/Engine/Source/Runtime/Engine/Classes/Engine/SceneCaptureCube.h
--- a/Engine/Source/Runtime/Engine/Classes/Engine/SceneCaptureCube.h
+++ b/Engine/Source/Runtime/Engine/Classes/Engine/SceneCaptureCube.h
@@ -8,7 +8,7 @@
 #pragma once
 #include "SceneCaptureCube.generated.h"
 
-UCLASS(hidecategories = (Collision, Material, Attachment, Actor), HeaderGroup = Decal, MinimalAPI)
+UCLASS(hidecategories = (Collision, Material, Attachment, Actor), MinimalAPI)
 class ASceneCaptureCube : public ASceneCapture
 {
 	GENERATED_UCLASS_BODY()
diff --git a/Engine/Source/Runtime/Engine/Classes/Engine/SkeletalMesh.h b/Engine/Source/Runtime/Engine/Classes/Engine/SkeletalMesh.h
--- a/Engine/Source/Runtime/Engine/Classes/Engine/SkeletalMesh.h
+++ b/Engine/Source/Runtime/Engine/Classes/Engine/SkeletalMesh.h
@@ -414,7 +414,7 @@ struct FSkeletalMaterial
 
 class FSkeletalMeshResource;
 
-UCLASS(HeaderGroup=SkeletalMesh, hidecategories=Object, MinimalAPI, BlueprintType)
+UCLASS(hidecategories=Object, MinimalAPI, BlueprintType)
 class USkeletalMesh : public UObject
 {
 	GENERATED_UCLASS_BODY()
diff --git a/Engine/Source/Runtime/Engine/Classes/Engine/SkeletalMeshSocket.h b/Engine/Source/Runtime/Engine/Classes/Engine/SkeletalMeshSocket.h
--- a/Engine/Source/Runtime/Engine/Classes/Engine/SkeletalMeshSocket.h
+++ b/Engine/Source/Runtime/Engine/Classes/Engine/SkeletalMeshSocket.h
@@ -4,7 +4,7 @@
 #pragma once
 #include "SkeletalMeshSocket.generated.h"
 
-UCLASS(HeaderGroup=SkeletalMesh, hidecategories=Object, hidecategories=Actor, MinimalAPI)
+UCLASS(hidecategories=Object, hidecategories=Actor, MinimalAPI)
 class USkeletalMeshSocket : public UObject
 {
 	GENERATED_UCLASS_BODY()
diff --git a/Engine/Source/Runtime/Engine/Classes/Engine/SkyLight.h b/Engine/Source/Runtime/Engine/Classes/Engine/SkyLight.h
--- a/Engine/Source/Runtime/Engine/Classes/Engine/SkyLight.h
+++ b/Engine/Source/Runtime/Engine/Classes/Engine/SkyLight.h
@@ -4,7 +4,7 @@
 
 #include "SkyLight.generated.h"
 
-UCLASS(HeaderGroup=Light, ClassGroup=Lights, hidecategories=(Input,Collision,Replication,Info), ConversionRoot)
+UCLASS(ClassGroup=Lights, hidecategories=(Input,Collision,Replication,Info), ConversionRoot)
 class ENGINE_API ASkyLight : public AInfo
 {
 	GENERATED_UCLASS_BODY()
diff --git a/Engine/Source/Runtime/Engine/Classes/Engine/SphereReflectionCapture.h b/Engine/Source/Runtime/Engine/Classes/Engine/SphereReflectionCapture.h
--- a/Engine/Source/Runtime/Engine/Classes/Engine/SphereReflectionCapture.h
+++ b/Engine/Source/Runtime/Engine/Classes/Engine/SphereReflectionCapture.h
@@ -8,7 +8,11 @@
 #pragma once
 #include "SphereReflectionCapture.generated.h"
 
-UCLASS(hidecategories=(Collision, Attachment, Actor), HeaderGroup=Decal, MinimalAPI)
+//
+// Forward declarations.
+//
+class UDrawSphereComponent;
+UCLASS(hidecategories=(Collision, Attachment, Actor), MinimalAPI)
 class ASphereReflectionCapture : public AReflectionCapture
 {
 	GENERATED_UCLASS_BODY()
diff --git a/Engine/Source/Runtime/Engine/Classes/Engine/SpotLight.h b/Engine/Source/Runtime/Engine/Classes/Engine/SpotLight.h
--- a/Engine/Source/Runtime/Engine/Classes/Engine/SpotLight.h
+++ b/Engine/Source/Runtime/Engine/Classes/Engine/SpotLight.h
@@ -5,7 +5,7 @@
 
 #include "SpotLight.generated.h"
 
-UCLASS(HeaderGroup=Light, ClassGroup=(Lights, SpotLights), meta=(ChildCanTick))
+UCLASS(ClassGroup=(Lights, SpotLights), meta=(ChildCanTick))
 class ENGINE_API ASpotLight : public ALight
 {
 	GENERATED_UCLASS_BODY()
diff --git a/Engine/Source/Runtime/Engine/Classes/Engine/StaticMesh.h b/Engine/Source/Runtime/Engine/Classes/Engine/StaticMesh.h
--- a/Engine/Source/Runtime/Engine/Classes/Engine/StaticMesh.h
+++ b/Engine/Source/Runtime/Engine/Classes/Engine/StaticMesh.h
@@ -303,7 +303,7 @@ struct FAssetEditorOrbitCameraPosition
 	FRotator CamOrbitRotation;
 };
 
-UCLASS(HeaderGroup=StaticMesh, collapsecategories, hidecategories=Object, customconstructor, MinimalAPI, BlueprintType, config=Engine)
+UCLASS(collapsecategories, hidecategories=Object, customconstructor, MinimalAPI, BlueprintType, config=Engine)
 class UStaticMesh : public UObject, public IInterface_CollisionDataProvider, public IInterface_AssetUserData
 {
 	GENERATED_UCLASS_BODY()
diff --git a/Engine/Source/Runtime/Engine/Classes/Engine/StaticMeshSocket.h b/Engine/Source/Runtime/Engine/Classes/Engine/StaticMeshSocket.h
--- a/Engine/Source/Runtime/Engine/Classes/Engine/StaticMeshSocket.h
+++ b/Engine/Source/Runtime/Engine/Classes/Engine/StaticMeshSocket.h
@@ -4,7 +4,7 @@
 #pragma once
 #include "StaticMeshSocket.generated.h"
 
-UCLASS(HeaderGroup=StaticMesh, hidecategories=Object, hidecategories=Actor, MinimalAPI)
+UCLASS(hidecategories=Object, hidecategories=Actor, MinimalAPI)
 class UStaticMeshSocket : public UObject
 {
 	GENERATED_UCLASS_BODY()
diff --git a/Engine/Source/Runtime/Engine/Classes/Engine/TextRenderActor.h b/Engine/Source/Runtime/Engine/Classes/Engine/TextRenderActor.h
--- a/Engine/Source/Runtime/Engine/Classes/Engine/TextRenderActor.h
+++ b/Engine/Source/Runtime/Engine/Classes/Engine/TextRenderActor.h
@@ -8,7 +8,7 @@
 
 #include "TextRenderActor.generated.h"
 
-UCLASS(MinimalAPI, hidecategories=(Collision, Attachment, Actor), HeaderGroup=Decal)
+UCLASS(MinimalAPI, hidecategories=(Collision, Attachment, Actor))
 class ATextRenderActor : public AActor
 {
 	GENERATED_UCLASS_BODY()
diff --git a/Engine/Source/Runtime/Engine/Classes/Engine/TwitterIntegrationBase.h b/Engine/Source/Runtime/Engine/Classes/Engine/TwitterIntegrationBase.h
--- a/Engine/Source/Runtime/Engine/Classes/Engine/TwitterIntegrationBase.h
+++ b/Engine/Source/Runtime/Engine/Classes/Engine/TwitterIntegrationBase.h
@@ -28,7 +28,7 @@ enum ETwitterIntegrationDelegate
 	TID_MAX,
 };
 
-UCLASS(HeaderGroup=PlatformInterface, config=Engine)
+UCLASS(config=Engine)
 class UTwitterIntegrationBase : public UPlatformInterfaceBase
 {
 	GENERATED_UCLASS_BODY()
diff --git a/Engine/Source/Runtime/Engine/Classes/Foliage/InstancedFoliageActor.h b/Engine/Source/Runtime/Engine/Classes/Foliage/InstancedFoliageActor.h
--- a/Engine/Source/Runtime/Engine/Classes/Foliage/InstancedFoliageActor.h
+++ b/Engine/Source/Runtime/Engine/Classes/Foliage/InstancedFoliageActor.h
@@ -2,9 +2,10 @@
 
 
 #pragma once
+#include "InstancedFoliage.h"
 #include "InstancedFoliageActor.generated.h"
 
-UCLASS(notplaceable, HeaderGroup = Foliage, hidecategories = Object, MinimalAPI, NotBlueprintable)
+UCLASS(notplaceable, hidecategories = Object, MinimalAPI, NotBlueprintable)
 class AInstancedFoliageActor : public AActor
 {
 	GENERATED_UCLASS_BODY()
diff --git a/Engine/Source/Runtime/Engine/Classes/Foliage/InstancedFoliageSettings.h b/Engine/Source/Runtime/Engine/Classes/Foliage/InstancedFoliageSettings.h
--- a/Engine/Source/Runtime/Engine/Classes/Foliage/InstancedFoliageSettings.h
+++ b/Engine/Source/Runtime/Engine/Classes/Foliage/InstancedFoliageSettings.h
@@ -15,7 +15,7 @@ enum FoliageVertexColorMask
 };
 
 
-UCLASS(HeaderGroup=Foliage, hidecategories=Object, editinlinenew, MinimalAPI)
+UCLASS(hidecategories=Object, editinlinenew, MinimalAPI)
 class UInstancedFoliageSettings : public UObject
 {
 	GENERATED_UCLASS_BODY()
diff --git a/Engine/Source/Runtime/Engine/Classes/GameFramework/Character.h b/Engine/Source/Runtime/Engine/Classes/GameFramework/Character.h
--- a/Engine/Source/Runtime/Engine/Classes/GameFramework/Character.h
+++ b/Engine/Source/Runtime/Engine/Classes/GameFramework/Character.h
@@ -4,6 +4,10 @@
 
 #include "Character.generated.h"
 
+//
+// Forward declarations
+//
+struct FAnimMontageInstance;
 /** Replicated data when playing a root motion montage. */
 USTRUCT()
 struct FRepRootMotionMontage
@@ -130,7 +134,7 @@ struct FRepRelativeMovement
 // to a CharacterMovementComponent that handles movement of the collision capsule, and they
 // also have implementations of basic networking and input models.
 //=============================================================================
-UCLASS(abstract, HeaderGroup=Pawn, config=Game, dependson=(AController, UCharacterMovementComponent), BlueprintType)
+UCLASS(abstract, config=Game, dependson=(AController, UCharacterMovementComponent), BlueprintType)
 class ENGINE_API ACharacter : public APawn
 {
 	GENERATED_UCLASS_BODY()
diff --git a/Engine/Source/Runtime/Engine/Classes/GameFramework/CharacterMovementComponent.h b/Engine/Source/Runtime/Engine/Classes/GameFramework/CharacterMovementComponent.h
--- a/Engine/Source/Runtime/Engine/Classes/GameFramework/CharacterMovementComponent.h
+++ b/Engine/Source/Runtime/Engine/Classes/GameFramework/CharacterMovementComponent.h
@@ -137,7 +137,7 @@ struct FRootMotionMovementParams
 };
 
 
-UCLASS(HeaderGroup=Component, dependson=(UNetworkPredictionInterface))
+UCLASS(dependson=(UNetworkPredictionInterface))
 class ENGINE_API UCharacterMovementComponent : public UPawnMovementComponent, public INetworkPredictionInterface
 {
 	GENERATED_UCLASS_BODY()
diff --git a/Engine/Source/Runtime/Engine/Classes/GameFramework/DefaultPawn.h b/Engine/Source/Runtime/Engine/Classes/GameFramework/DefaultPawn.h
--- a/Engine/Source/Runtime/Engine/Classes/GameFramework/DefaultPawn.h
+++ b/Engine/Source/Runtime/Engine/Classes/GameFramework/DefaultPawn.h
@@ -6,9 +6,10 @@
 //=============================================================================
 
 #pragma once
+#include "FloatingPawnMovement.h"
 #include "DefaultPawn.generated.h"
 
-UCLASS(HeaderGroup=Pawn, config=Game, Blueprintable, BlueprintType)
+UCLASS(config=Game, Blueprintable, BlueprintType)
 class ENGINE_API ADefaultPawn : public APawn
 {
 	GENERATED_UCLASS_BODY()
diff --git a/Engine/Source/Runtime/Engine/Classes/GameFramework/ForceFeedbackEffect.h b/Engine/Source/Runtime/Engine/Classes/GameFramework/ForceFeedbackEffect.h
--- a/Engine/Source/Runtime/Engine/Classes/GameFramework/ForceFeedbackEffect.h
+++ b/Engine/Source/Runtime/Engine/Classes/GameFramework/ForceFeedbackEffect.h
@@ -70,7 +70,7 @@ struct FActiveForceFeedbackEffect
 /**
  * A predefined force-feedback effect to be played on a controller
  */
-UCLASS(BlueprintType, headergroup=UserInterface, MinimalAPI)
+UCLASS(BlueprintType, MinimalAPI)
 class UForceFeedbackEffect : public UObject
 {
 	GENERATED_UCLASS_BODY()
diff --git a/Engine/Source/Runtime/Engine/Classes/GameFramework/GameState.h b/Engine/Source/Runtime/Engine/Classes/GameFramework/GameState.h
--- a/Engine/Source/Runtime/Engine/Classes/GameFramework/GameState.h
+++ b/Engine/Source/Runtime/Engine/Classes/GameFramework/GameState.h
@@ -8,7 +8,7 @@
 // GameState is replicated and is valid on servers and clients.
 //=============================================================================
 
-UCLASS(config=Game, HeaderGroup=ReplicationInfo, notplaceable, BlueprintType, Blueprintable)
+UCLASS(config=Game, notplaceable, BlueprintType, Blueprintable)
 class ENGINE_API AGameState : public AInfo
 {
 	GENERATED_UCLASS_BODY()
diff --git a/Engine/Source/Runtime/Engine/Classes/GameFramework/InputSettings.h b/Engine/Source/Runtime/Engine/Classes/GameFramework/InputSettings.h
--- a/Engine/Source/Runtime/Engine/Classes/GameFramework/InputSettings.h
+++ b/Engine/Source/Runtime/Engine/Classes/GameFramework/InputSettings.h
@@ -1,10 +1,10 @@
 // Copyright 1998-2014 Epic Games, Inc. All Rights Reserved.
 
 #pragma once
-
+#include "PlayerInput.h"
 #include "InputSettings.generated.h"
 
-UCLASS(config=Input, HeaderGroup=UserInterface)
+UCLASS(config=Input)
 class ENGINE_API UInputSettings : public UObject
 {
 	GENERATED_UCLASS_BODY()
diff --git a/Engine/Source/Runtime/Engine/Classes/GameFramework/MovementComponent.h b/Engine/Source/Runtime/Engine/Classes/GameFramework/MovementComponent.h
--- a/Engine/Source/Runtime/Engine/Classes/GameFramework/MovementComponent.h
+++ b/Engine/Source/Runtime/Engine/Classes/GameFramework/MovementComponent.h
@@ -5,10 +5,10 @@
  */
 
 #pragma once
+#include "Components/SceneComponent.h"
 #include "MovementComponent.generated.h"
 
-
-UCLASS(HeaderGroup=Component, ClassGroup=Movement, abstract, BlueprintType)
+UCLASS(ClassGroup=Movement, abstract, BlueprintType)
 class ENGINE_API UMovementComponent : public UActorComponent
 {
 	GENERATED_UCLASS_BODY()
diff --git a/Engine/Source/Runtime/Engine/Classes/GameFramework/NavMovementComponent.h b/Engine/Source/Runtime/Engine/Classes/GameFramework/NavMovementComponent.h
--- a/Engine/Source/Runtime/Engine/Classes/GameFramework/NavMovementComponent.h
+++ b/Engine/Source/Runtime/Engine/Classes/GameFramework/NavMovementComponent.h
@@ -8,7 +8,7 @@
 #include "AI/Navigation/NavAgentInterface.h"
 #include "NavMovementComponent.generated.h"
 
-UCLASS(HeaderGroup=Component, abstract, dependson=(UNavigationSystem, UPrimitiveComponent, INavAgentInterface))
+UCLASS(abstract, dependson=(UNavigationSystem, UPrimitiveComponent, INavAgentInterface))
 class ENGINE_API UNavMovementComponent : public UMovementComponent
 {
 	GENERATED_UCLASS_BODY()
diff --git a/Engine/Source/Runtime/Engine/Classes/GameFramework/Pawn.h b/Engine/Source/Runtime/Engine/Classes/GameFramework/Pawn.h
--- a/Engine/Source/Runtime/Engine/Classes/GameFramework/Pawn.h
+++ b/Engine/Source/Runtime/Engine/Classes/GameFramework/Pawn.h
@@ -1,7 +1,7 @@
 // Copyright 1998-2014 Epic Games, Inc. All Rights Reserved.
 
 #pragma once
-
+#include "PawnMovementComponent.h"
 #include "Pawn.generated.h"
 
 ENGINE_API DECLARE_LOG_CATEGORY_EXTERN(LogDamage, Warning, All);
@@ -11,7 +11,7 @@ ENGINE_API DECLARE_LOG_CATEGORY_EXTERN(LogDamage, Warning, All);
  *	Pawns are the physical representations of players and creatures in a level.
  *	They are responsible for all physical interaction between the player or AI and the world.
  */
-UCLASS(abstract, HeaderGroup=Pawn, config=Game, dependson=AController, BlueprintType, Blueprintable, hidecategories=(Navigation, "AI|Navigation"))
+UCLASS(abstract, config=Game, dependson=AController, BlueprintType, Blueprintable, hidecategories=(Navigation, "AI|Navigation"))
 class ENGINE_API APawn : public AActor, public INavAgentInterface
 {
 	GENERATED_UCLASS_BODY()
diff --git a/Engine/Source/Runtime/Engine/Classes/GameFramework/PlayerInput.h b/Engine/Source/Runtime/Engine/Classes/GameFramework/PlayerInput.h
--- a/Engine/Source/Runtime/Engine/Classes/GameFramework/PlayerInput.h
+++ b/Engine/Source/Runtime/Engine/Classes/GameFramework/PlayerInput.h
@@ -7,13 +7,17 @@
 //=============================================================================
 
 #pragma once
-
-#include "KeyState.h"
+#include "Components/InputComponent.h"
 #include "GestureRecognizer.h"
+#include "KeyState.h"
 #include "PlayerInput.generated.h"
 
 ENGINE_API DECLARE_LOG_CATEGORY_EXTERN(LogInput, Log, All);
 
+//
+// Forward declarations.
+//
+struct FDelegateDispatchDetails;
 
 USTRUCT()
 struct FKeyBind
@@ -223,7 +227,7 @@ struct FAxisKeyDetails
 	}
 };
 
-UCLASS(Within=PlayerController, config=Input, dependsOn=UEngineTypes, transient, HeaderGroup=UserInterface)
+UCLASS(Within=PlayerController, config=Input, dependsOn=UEngineTypes, transient)
 class ENGINE_API UPlayerInput : public UObject
 {
 	GENERATED_UCLASS_BODY()
diff --git a/Engine/Source/Runtime/Engine/Classes/GameFramework/PlayerState.h b/Engine/Source/Runtime/Engine/Classes/GameFramework/PlayerState.h
--- a/Engine/Source/Runtime/Engine/Classes/GameFramework/PlayerState.h
+++ b/Engine/Source/Runtime/Engine/Classes/GameFramework/PlayerState.h
@@ -30,7 +30,7 @@ struct PingAvgData
 // such as playername, score, etc.
 //=============================================================================
 
-UCLASS(HeaderGroup=ReplicationInfo, BlueprintType, Blueprintable, notplaceable)
+UCLASS(BlueprintType, Blueprintable, notplaceable)
 class ENGINE_API APlayerState : public AInfo
 {
 	GENERATED_UCLASS_BODY()
diff --git a/Engine/Source/Runtime/Engine/Classes/GameFramework/ProjectileMovementComponent.h b/Engine/Source/Runtime/Engine/Classes/GameFramework/ProjectileMovementComponent.h
--- a/Engine/Source/Runtime/Engine/Classes/GameFramework/ProjectileMovementComponent.h
+++ b/Engine/Source/Runtime/Engine/Classes/GameFramework/ProjectileMovementComponent.h
@@ -8,7 +8,7 @@
  * If the updated component is simulating physics, only the initial launch parameters (when initial velocity is non-zero) will affect the projectile,
  * and the physics sim will take over from there.
  */
-UCLASS(HeaderGroup=Component, ClassGroup=Movement, meta=(BlueprintSpawnableComponent), ShowCategories=(Velocity))
+UCLASS(ClassGroup=Movement, meta=(BlueprintSpawnableComponent), ShowCategories=(Velocity))
 class ENGINE_API UProjectileMovementComponent : public UMovementComponent
 {
 	GENERATED_UCLASS_BODY()
diff --git a/Engine/Source/Runtime/Engine/Classes/GameFramework/RotatingMovementComponent.h b/Engine/Source/Runtime/Engine/Classes/GameFramework/RotatingMovementComponent.h
--- a/Engine/Source/Runtime/Engine/Classes/GameFramework/RotatingMovementComponent.h
+++ b/Engine/Source/Runtime/Engine/Classes/GameFramework/RotatingMovementComponent.h
@@ -6,7 +6,7 @@
 /**
  * Movement component updates position of associated PrimitiveComponent during its tick.
  */
-UCLASS(HeaderGroup=Component, ClassGroup=Movement, meta=(BlueprintSpawnableComponent), HideCategories=(PlanarMovement, "Components|Movement|Planar", Velocity))
+UCLASS(ClassGroup=Movement, meta=(BlueprintSpawnableComponent), HideCategories=(PlanarMovement, "Components|Movement|Planar", Velocity))
 class ENGINE_API URotatingMovementComponent : public UMovementComponent
 {
 	GENERATED_UCLASS_BODY()
diff --git a/Engine/Source/Runtime/Engine/Classes/GameFramework/SpectatorPawn.h b/Engine/Source/Runtime/Engine/Classes/GameFramework/SpectatorPawn.h
--- a/Engine/Source/Runtime/Engine/Classes/GameFramework/SpectatorPawn.h
+++ b/Engine/Source/Runtime/Engine/Classes/GameFramework/SpectatorPawn.h
@@ -10,7 +10,7 @@
 #include "SpectatorPawn.generated.h"
 
 
-UCLASS(HeaderGroup=Pawn, config=Game, Blueprintable, BlueprintType)
+UCLASS(config=Game, Blueprintable, BlueprintType)
 class ENGINE_API ASpectatorPawn : public ADefaultPawn
 {
 	GENERATED_UCLASS_BODY()
diff --git a/Engine/Source/Runtime/Engine/Classes/GameFramework/SpectatorPawnMovement.h b/Engine/Source/Runtime/Engine/Classes/GameFramework/SpectatorPawnMovement.h
--- a/Engine/Source/Runtime/Engine/Classes/GameFramework/SpectatorPawnMovement.h
+++ b/Engine/Source/Runtime/Engine/Classes/GameFramework/SpectatorPawnMovement.h
@@ -7,7 +7,7 @@
 #pragma once
 #include "SpectatorPawnMovement.generated.h"
 
-UCLASS(HeaderGroup=Component, DependsOn=UFloatingPawnMovement)
+UCLASS(DependsOn=UFloatingPawnMovement)
 class ENGINE_API USpectatorPawnMovement : public UFloatingPawnMovement
 {
 	GENERATED_UCLASS_BODY()
diff --git a/Engine/Source/Runtime/Engine/Classes/GameFramework/SpringArmComponent.h b/Engine/Source/Runtime/Engine/Classes/GameFramework/SpringArmComponent.h
--- a/Engine/Source/Runtime/Engine/Classes/GameFramework/SpringArmComponent.h
+++ b/Engine/Source/Runtime/Engine/Classes/GameFramework/SpringArmComponent.h
@@ -11,7 +11,7 @@
  * Example: Use as a 'camera boom' to keep the follow camera for a player from colliding into the world.
  */
 
-UCLASS(HeaderGroup=Component, ClassGroup=Camera, meta=(BlueprintSpawnableComponent), hideCategories=(Mobility))
+UCLASS(ClassGroup=Camera, meta=(BlueprintSpawnableComponent), hideCategories=(Mobility))
 class ENGINE_API USpringArmComponent : public USceneComponent
 {
 	GENERATED_UCLASS_BODY()
diff --git a/Engine/Source/Runtime/Engine/Classes/GameFramework/WheeledVehicle.h b/Engine/Source/Runtime/Engine/Classes/GameFramework/WheeledVehicle.h
--- a/Engine/Source/Runtime/Engine/Classes/GameFramework/WheeledVehicle.h
+++ b/Engine/Source/Runtime/Engine/Classes/GameFramework/WheeledVehicle.h
@@ -3,7 +3,7 @@
 #pragma once
 #include "WheeledVehicle.generated.h"
 
-UCLASS(abstract, HeaderGroup=Pawn, config=Game, dependson=(AController), BlueprintType)
+UCLASS(abstract, config=Game, dependson=(AController), BlueprintType)
 class ENGINE_API AWheeledVehicle : public APawn
 {
 	GENERATED_UCLASS_BODY()
diff --git a/Engine/Source/Runtime/Engine/Classes/GameFramework/WorldSettings.h b/Engine/Source/Runtime/Engine/Classes/GameFramework/WorldSettings.h
--- a/Engine/Source/Runtime/Engine/Classes/GameFramework/WorldSettings.h
+++ b/Engine/Source/Runtime/Engine/Classes/GameFramework/WorldSettings.h
@@ -173,7 +173,7 @@ struct ENGINE_API FNetViewer
 /**
  * Actor containing all script accessible world properties.
  */
-UCLASS(HeaderGroup=GameEngine, config=game, showcategories=(Rendering), hidecategories=(Actor, Advanced, Display, Events, Object, Attachment, Info, Input, Blueprint, Layers), notplaceable, dependson=(APostProcessVolume, UMusicTrackDataStructures))
+UCLASS(config=game, showcategories=(Rendering), hidecategories=(Actor, Advanced, Display, Events, Object, Attachment, Info, Input, Blueprint, Layers), notplaceable, dependson=(APostProcessVolume, UMusicTrackDataStructures))
 class ENGINE_API AWorldSettings : public AInfo
 {
 	GENERATED_UCLASS_BODY()
diff --git a/Engine/Source/Runtime/Engine/Classes/Kismet/GameplayStatics.h b/Engine/Source/Runtime/Engine/Classes/Kismet/GameplayStatics.h
--- a/Engine/Source/Runtime/Engine/Classes/Kismet/GameplayStatics.h
+++ b/Engine/Source/Runtime/Engine/Classes/Kismet/GameplayStatics.h
@@ -5,6 +5,10 @@
 #include "../Sound/DialogueTypes.h"
 #include "GameplayStatics.generated.h"
 
+//
+// Forward declarations.
+//
+class USaveGame;
 struct FDialogueContext;
 
 UENUM()
@@ -18,7 +22,7 @@ namespace ESuggestProjVelocityTraceOption
 	};
 }
 
-UCLASS(HeaderGroup=KismetLibrary)
+UCLASS()
 class ENGINE_API UGameplayStatics : public UBlueprintFunctionLibrary
 {
 	GENERATED_UCLASS_BODY()
diff --git a/Engine/Source/Runtime/Engine/Classes/Kismet/KismetAIHelperLibrary.h b/Engine/Source/Runtime/Engine/Classes/Kismet/KismetAIHelperLibrary.h
--- a/Engine/Source/Runtime/Engine/Classes/Kismet/KismetAIHelperLibrary.h
+++ b/Engine/Source/Runtime/Engine/Classes/Kismet/KismetAIHelperLibrary.h
@@ -8,7 +8,7 @@
 #pragma once
 #include "KismetAIHelperLibrary.generated.h"
 
-UCLASS(HeaderGroup=KismetLibrary)
+UCLASS()
 class ENGINE_API UKismetAIHelperLibrary : public UBlueprintFunctionLibrary
 {
 	GENERATED_UCLASS_BODY()
diff --git a/Engine/Source/Runtime/Engine/Classes/Landscape/Landscape.h b/Engine/Source/Runtime/Engine/Classes/Landscape/Landscape.h
--- a/Engine/Source/Runtime/Engine/Classes/Landscape/Landscape.h
+++ b/Engine/Source/Runtime/Engine/Classes/Landscape/Landscape.h
@@ -17,7 +17,7 @@ enum ELandscapeSetupErrors
 	LSE_MAX,
 };
 
-UCLASS(dependson=ULightComponent, HeaderGroup=Terrain, Placeable, hidecategories=LandscapeProxy, showcategories=(Display, Movement, Collision, Lighting, LOD, Input), MinimalAPI)
+UCLASS(dependson=ULightComponent, Placeable, hidecategories=LandscapeProxy, showcategories=(Display, Movement, Collision, Lighting, LOD, Input), MinimalAPI)
 class ALandscape : public ALandscapeProxy
 {
 	GENERATED_UCLASS_BODY()
diff --git a/Engine/Source/Runtime/Engine/Classes/Landscape/LandscapeComponent.h b/Engine/Source/Runtime/Engine/Classes/Landscape/LandscapeComponent.h
--- a/Engine/Source/Runtime/Engine/Classes/Landscape/LandscapeComponent.h
+++ b/Engine/Source/Runtime/Engine/Classes/Landscape/LandscapeComponent.h
@@ -2,7 +2,8 @@
 
 #pragma once
 #include "LandscapeLayerInfoObject.h"
-
+#include "LightMap.h"
+#include "ShadowMap.h"
 #include "LandscapeComponent.generated.h"
 
 class FLandscapeComponentDerivedData
@@ -80,7 +81,7 @@ struct FWeightmapLayerAllocationInfo
 	}
 };
 
-UCLASS(HeaderGroup=Terrain, hidecategories=(Display, Attachment, Physics, Debug, Collision, Movement, Rendering, PrimitiveComponent, Object, Transform), MinimalAPI)
+UCLASS(hidecategories=(Display, Attachment, Physics, Debug, Collision, Movement, Rendering, PrimitiveComponent, Object, Transform), MinimalAPI)
 class ULandscapeComponent : public UPrimitiveComponent
 {
 	GENERATED_UCLASS_BODY()
diff --git a/Engine/Source/Runtime/Engine/Classes/Landscape/LandscapeGizmoActiveActor.h b/Engine/Source/Runtime/Engine/Classes/Landscape/LandscapeGizmoActiveActor.h
--- a/Engine/Source/Runtime/Engine/Classes/Landscape/LandscapeGizmoActiveActor.h
+++ b/Engine/Source/Runtime/Engine/Classes/Landscape/LandscapeGizmoActiveActor.h
@@ -39,7 +39,7 @@ struct FGizmoSelectData
 	
 };
 
-UCLASS(notplaceable, HeaderGroup=Terrain, MinimalAPI)
+UCLASS(notplaceable, MinimalAPI)
 class ALandscapeGizmoActiveActor : public ALandscapeGizmoActor
 {
 	GENERATED_UCLASS_BODY()
diff --git a/Engine/Source/Runtime/Engine/Classes/Landscape/LandscapeGizmoActor.h b/Engine/Source/Runtime/Engine/Classes/Landscape/LandscapeGizmoActor.h
--- a/Engine/Source/Runtime/Engine/Classes/Landscape/LandscapeGizmoActor.h
+++ b/Engine/Source/Runtime/Engine/Classes/Landscape/LandscapeGizmoActor.h
@@ -4,7 +4,7 @@
 #pragma once
 #include "LandscapeGizmoActor.generated.h"
 
-UCLASS(notplaceable, HeaderGroup=Terrain, MinimalAPI, NotBlueprintable)
+UCLASS(notplaceable, MinimalAPI, NotBlueprintable)
 class ALandscapeGizmoActor : public AActor
 {
 	GENERATED_UCLASS_BODY()
diff --git a/Engine/Source/Runtime/Engine/Classes/Landscape/LandscapeGizmoRenderComponent.h b/Engine/Source/Runtime/Engine/Classes/Landscape/LandscapeGizmoRenderComponent.h
--- a/Engine/Source/Runtime/Engine/Classes/Landscape/LandscapeGizmoRenderComponent.h
+++ b/Engine/Source/Runtime/Engine/Classes/Landscape/LandscapeGizmoRenderComponent.h
@@ -4,7 +4,7 @@
 #pragma once
 #include "LandscapeGizmoRenderComponent.generated.h"
 
-UCLASS(HeaderGroup=Terrain, hidecategories=Object)
+UCLASS(hidecategories=Object)
 class ULandscapeGizmoRenderComponent : public UPrimitiveComponent
 {
 	GENERATED_UCLASS_BODY()
diff --git a/Engine/Source/Runtime/Engine/Classes/Landscape/LandscapeProxy.h b/Engine/Source/Runtime/Engine/Classes/Landscape/LandscapeProxy.h
--- a/Engine/Source/Runtime/Engine/Classes/Landscape/LandscapeProxy.h
+++ b/Engine/Source/Runtime/Engine/Classes/Landscape/LandscapeProxy.h
@@ -6,6 +6,11 @@
 
 #include "LandscapeProxy.generated.h"
 
+//
+// Forward declarations.
+//
+class ULandscapeMaterialInstanceConstant;
+
 /** Structure storing channel usage for weightmap textures */
 USTRUCT()
 struct FLandscapeWeightmapUsage
@@ -164,7 +169,7 @@ namespace ELandscapeLayerPaintingRestriction
 	};
 }
 
-UCLASS(HeaderGroup=Terrain, dependson=UEngineTypes, NotPlaceable, hidecategories=(Display, Attachment, Physics, Debug, Lighting, LOD), showcategories=(Rendering, "Utilities|Orientation"), MinimalAPI)
+UCLASS(dependson=UEngineTypes, NotPlaceable, hidecategories=(Display, Attachment, Physics, Debug, Lighting, LOD), showcategories=(Rendering, "Utilities|Orientation"), MinimalAPI)
 class ALandscapeProxy : public AActor, public INavRelevantActorInterface
 {
 	GENERATED_UCLASS_BODY()
diff --git a/Engine/Source/Runtime/Engine/Classes/Landscape/LandscapeSplineControlPoint.h b/Engine/Source/Runtime/Engine/Classes/Landscape/LandscapeSplineControlPoint.h
--- a/Engine/Source/Runtime/Engine/Classes/Landscape/LandscapeSplineControlPoint.h
+++ b/Engine/Source/Runtime/Engine/Classes/Landscape/LandscapeSplineControlPoint.h
@@ -36,7 +36,7 @@ struct FLandscapeSplineConnection
 	ENGINE_API struct FLandscapeSplineSegmentConnection& GetFarConnection() const;
 };
 
-UCLASS(HeaderGroup=Terrain,Within=LandscapeSplinesComponent,autoExpandCategories=LandscapeSplineControlPoint,MinimalAPI)
+UCLASS(Within=LandscapeSplinesComponent,autoExpandCategories=LandscapeSplineControlPoint,MinimalAPI)
 class ULandscapeSplineControlPoint : public UObject
 {
 	GENERATED_UCLASS_BODY()
diff --git a/Engine/Source/Runtime/Engine/Classes/Landscape/LandscapeSplineSegment.h b/Engine/Source/Runtime/Engine/Classes/Landscape/LandscapeSplineSegment.h
--- a/Engine/Source/Runtime/Engine/Classes/Landscape/LandscapeSplineSegment.h
+++ b/Engine/Source/Runtime/Engine/Classes/Landscape/LandscapeSplineSegment.h
@@ -140,7 +140,7 @@ struct FLandscapeSplineMeshEntry
 };
 
 
-UCLASS(HeaderGroup=Terrain,Within=LandscapeSplinesComponent,autoExpandCategories=(LandscapeSplineSegment,LandscapeSplineMeshes),MinimalAPI)
+UCLASS(Within=LandscapeSplinesComponent,autoExpandCategories=(LandscapeSplineSegment,LandscapeSplineMeshes),MinimalAPI)
 class ULandscapeSplineSegment : public UObject
 {
 	GENERATED_UCLASS_BODY()
diff --git a/Engine/Source/Runtime/Engine/Classes/Landscape/LandscapeSplinesComponent.h b/Engine/Source/Runtime/Engine/Classes/Landscape/LandscapeSplinesComponent.h
--- a/Engine/Source/Runtime/Engine/Classes/Landscape/LandscapeSplinesComponent.h
+++ b/Engine/Source/Runtime/Engine/Classes/Landscape/LandscapeSplinesComponent.h
@@ -4,7 +4,7 @@
 #pragma once
 #include "LandscapeSplinesComponent.generated.h"
 
-UCLASS(HeaderGroup=Terrain,MinimalAPI)
+UCLASS(MinimalAPI)
 class ULandscapeSplinesComponent : public UPrimitiveComponent
 {
 	GENERATED_UCLASS_BODY()
diff --git a/Engine/Source/Runtime/Engine/Classes/Materials/Material.h b/Engine/Source/Runtime/Engine/Classes/Materials/Material.h
--- a/Engine/Source/Runtime/Engine/Classes/Materials/Material.h
+++ b/Engine/Source/Runtime/Engine/Classes/Materials/Material.h
@@ -254,7 +254,7 @@ struct FMaterialParameterCollectionInfo
  * when light from the scene hits the surface, the lighting model of the material is used to calculate how
  * that light interacts with the surface. 
  */
-UCLASS(HeaderGroup=Material, hidecategories=Object, MinimalAPI, BlueprintType)
+UCLASS(hidecategories=Object, MinimalAPI, BlueprintType)
 class UMaterial : public UMaterialInterface
 {
 	GENERATED_UCLASS_BODY()
diff --git a/Engine/Source/Runtime/Engine/Classes/Materials/MaterialExpressionTextureSampleParameter.h b/Engine/Source/Runtime/Engine/Classes/Materials/MaterialExpressionTextureSampleParameter.h
--- a/Engine/Source/Runtime/Engine/Classes/Materials/MaterialExpressionTextureSampleParameter.h
+++ b/Engine/Source/Runtime/Engine/Classes/Materials/MaterialExpressionTextureSampleParameter.h
@@ -4,7 +4,7 @@
 #pragma once
 #include "MaterialExpressionTextureSampleParameter.generated.h"
 
-UCLASS(HeaderGroup=Material, collapsecategories, abstract, hidecategories=Object, MinimalAPI)
+UCLASS(collapsecategories, abstract, hidecategories=Object, MinimalAPI)
 class UMaterialExpressionTextureSampleParameter : public UMaterialExpressionTextureSample
 {
 	GENERATED_UCLASS_BODY()
diff --git a/Engine/Source/Runtime/Engine/Classes/Materials/MaterialInstance.h b/Engine/Source/Runtime/Engine/Classes/Materials/MaterialInstance.h
--- a/Engine/Source/Runtime/Engine/Classes/Materials/MaterialInstance.h
+++ b/Engine/Source/Runtime/Engine/Classes/Materials/MaterialInstance.h
@@ -4,6 +4,11 @@
 #pragma once
 #include "MaterialInstance.generated.h"
 
+//
+// Forward declarations.
+//
+struct FMaterialInstanceBasePropertyOverrides;
+
 /** Editable font parameter. */
 USTRUCT()
 struct FFontParameterValue
@@ -121,7 +126,7 @@ struct FVectorParameterValue
 	static ValueType GetValue(const FVectorParameterValue& Parameter) { return Parameter.ParameterValue; }
 };
 
-UCLASS(abstract, HeaderGroup=Material, BlueprintType,MinimalAPI)
+UCLASS(abstract, BlueprintType,MinimalAPI)
 class UMaterialInstance : public UMaterialInterface
 {
 	GENERATED_UCLASS_BODY()
diff --git a/Engine/Source/Runtime/Engine/Classes/Materials/MaterialInstanceConstant.h b/Engine/Source/Runtime/Engine/Classes/Materials/MaterialInstanceConstant.h
--- a/Engine/Source/Runtime/Engine/Classes/Materials/MaterialInstanceConstant.h
+++ b/Engine/Source/Runtime/Engine/Classes/Materials/MaterialInstanceConstant.h
@@ -8,7 +8,7 @@
  * General modification of the material cannot be supported without recompilation, so the instances are limited to changing the values of
  * predefined material parameters. The parameters are statically defined in the compiled material by a unique name, type and default value.
  */
-UCLASS(HeaderGroup=Material, hidecategories=Object, collapsecategories, BlueprintType,MinimalAPI)
+UCLASS(hidecategories=Object, collapsecategories, BlueprintType,MinimalAPI)
 class UMaterialInstanceConstant : public UMaterialInstance
 {
 	GENERATED_UCLASS_BODY()
diff --git a/Engine/Source/Runtime/Engine/Classes/Materials/MaterialInstanceDynamic.h b/Engine/Source/Runtime/Engine/Classes/Materials/MaterialInstanceDynamic.h
--- a/Engine/Source/Runtime/Engine/Classes/Materials/MaterialInstanceDynamic.h
+++ b/Engine/Source/Runtime/Engine/Classes/Materials/MaterialInstanceDynamic.h
@@ -4,7 +4,7 @@
 #pragma once
 #include "MaterialInstanceDynamic.generated.h"
 
-UCLASS(HeaderGroup=Material, hidecategories=Object, collapsecategories, BlueprintType, MinimalAPI)
+UCLASS(hidecategories=Object, collapsecategories, BlueprintType, MinimalAPI)
 class UMaterialInstanceDynamic : public UMaterialInstance
 {
 	GENERATED_UCLASS_BODY()
diff --git a/Engine/Source/Runtime/Engine/Classes/Materials/MaterialInterface.h b/Engine/Source/Runtime/Engine/Classes/Materials/MaterialInterface.h
--- a/Engine/Source/Runtime/Engine/Classes/Materials/MaterialInterface.h
+++ b/Engine/Source/Runtime/Engine/Classes/Materials/MaterialInterface.h
@@ -90,7 +90,7 @@ struct FLightmassMaterialInterfaceSettings
 	
 };
 
-UCLASS(abstract, HeaderGroup=Material, BlueprintType,MinimalAPI)
+UCLASS(abstract, BlueprintType,MinimalAPI)
 class UMaterialInterface : public UObject, public IBlendableInterface
 {
 	GENERATED_UCLASS_BODY()
diff --git a/Engine/Source/Runtime/Engine/Classes/Matinee/InterpGroup.h b/Engine/Source/Runtime/Engine/Classes/Matinee/InterpGroup.h
--- a/Engine/Source/Runtime/Engine/Classes/Matinee/InterpGroup.h
+++ b/Engine/Source/Runtime/Engine/Classes/Matinee/InterpGroup.h
@@ -1,7 +1,7 @@
 // Copyright 1998-2014 Epic Games, Inc. All Rights Reserved.
 
 #pragma once
-
+#include "Interpolation.h"
 #include "InterpGroup.generated.h"
 
 /**
@@ -74,7 +74,7 @@ struct FInterpEdSelKey
 	
 };
 
-UCLASS(HeaderGroup=Interpolation, collapsecategories, hidecategories=Object, MinimalAPI)
+UCLASS(collapsecategories, hidecategories=Object, MinimalAPI)
 class UInterpGroup : public UObject, public FInterpEdInputInterface
 {
 	GENERATED_UCLASS_BODY()
diff --git a/Engine/Source/Runtime/Engine/Classes/Matinee/InterpTrack.h b/Engine/Source/Runtime/Engine/Classes/Matinee/InterpTrack.h
--- a/Engine/Source/Runtime/Engine/Classes/Matinee/InterpTrack.h
+++ b/Engine/Source/Runtime/Engine/Classes/Matinee/InterpTrack.h
@@ -10,6 +10,12 @@
 
 #include "InterpTrack.generated.h"
 
+//
+// Forward declarations
+//
+class FCanvas;
+
+
 /** Required condition for this track to be enabled */
 UENUM()
 enum ETrackActiveCondition
@@ -81,7 +87,7 @@ struct FSubTrackGroup
 
 };
 
-UCLASS(HeaderGroup=Interpolation, collapsecategories, hidecategories=Object, abstract, MinimalAPI)
+UCLASS(collapsecategories, hidecategories=Object, abstract, MinimalAPI)
 class UInterpTrack : public UObject, public FCurveEdInterface, public FInterpEdInputInterface
 {
 	GENERATED_UCLASS_BODY()
diff --git a/Engine/Source/Runtime/Engine/Classes/Matinee/InterpTrackAnimControl.h b/Engine/Source/Runtime/Engine/Classes/Matinee/InterpTrackAnimControl.h
--- a/Engine/Source/Runtime/Engine/Classes/Matinee/InterpTrackAnimControl.h
+++ b/Engine/Source/Runtime/Engine/Classes/Matinee/InterpTrackAnimControl.h
@@ -51,7 +51,7 @@ struct FAnimControlTrackKey
 
 };
 
-UCLASS(HeaderGroup=Interpolation, MinimalAPI, meta=( DisplayName = "Animation Track" ))
+UCLASS(MinimalAPI, meta=( DisplayName = "Animation Track" ))
 class UInterpTrackAnimControl : public UInterpTrackFloatBase
 {
 	GENERATED_UCLASS_BODY()
diff --git a/Engine/Source/Runtime/Engine/Classes/Matinee/InterpTrackAudioMaster.h b/Engine/Source/Runtime/Engine/Classes/Matinee/InterpTrackAudioMaster.h
--- a/Engine/Source/Runtime/Engine/Classes/Matinee/InterpTrackAudioMaster.h
+++ b/Engine/Source/Runtime/Engine/Classes/Matinee/InterpTrackAudioMaster.h
@@ -4,7 +4,7 @@
 #pragma once
 #include "InterpTrackAudioMaster.generated.h"
 
-UCLASS(HeaderGroup=Interpolation, MinimalAPI, meta=( DisplayName = "Audio Master Track" ) )
+UCLASS(MinimalAPI, meta=( DisplayName = "Audio Master Track" ) )
 class UInterpTrackAudioMaster : public UInterpTrackVectorBase
 {
 	GENERATED_UCLASS_BODY()
diff --git a/Engine/Source/Runtime/Engine/Classes/Matinee/InterpTrackBoolProp.h b/Engine/Source/Runtime/Engine/Classes/Matinee/InterpTrackBoolProp.h
--- a/Engine/Source/Runtime/Engine/Classes/Matinee/InterpTrackBoolProp.h
+++ b/Engine/Source/Runtime/Engine/Classes/Matinee/InterpTrackBoolProp.h
@@ -25,7 +25,7 @@ struct FBoolTrackKey
 
 };
 
-UCLASS(HeaderGroup=Interpolation, MinimalAPI, meta=( DisplayName = "Bool Property Track" ) )
+UCLASS(MinimalAPI, meta=( DisplayName = "Bool Property Track" ) )
 class UInterpTrackBoolProp : public UInterpTrack
 {
 	GENERATED_UCLASS_BODY()
diff --git a/Engine/Source/Runtime/Engine/Classes/Matinee/InterpTrackColorProp.h b/Engine/Source/Runtime/Engine/Classes/Matinee/InterpTrackColorProp.h
--- a/Engine/Source/Runtime/Engine/Classes/Matinee/InterpTrackColorProp.h
+++ b/Engine/Source/Runtime/Engine/Classes/Matinee/InterpTrackColorProp.h
@@ -4,7 +4,7 @@
 #pragma once
 #include "InterpTrackColorProp.generated.h"
 
-UCLASS(HeaderGroup=Interpolation, MinimalAPI, meta=( DisplayName = "Color Property Track" ) )
+UCLASS(MinimalAPI, meta=( DisplayName = "Color Property Track" ) )
 class UInterpTrackColorProp : public UInterpTrackVectorBase
 {
 	GENERATED_UCLASS_BODY()
diff --git a/Engine/Source/Runtime/Engine/Classes/Matinee/InterpTrackColorScale.h b/Engine/Source/Runtime/Engine/Classes/Matinee/InterpTrackColorScale.h
--- a/Engine/Source/Runtime/Engine/Classes/Matinee/InterpTrackColorScale.h
+++ b/Engine/Source/Runtime/Engine/Classes/Matinee/InterpTrackColorScale.h
@@ -8,7 +8,7 @@
 #pragma once
 #include "InterpTrackColorScale.generated.h"
 
-UCLASS(HeaderGroup=Interpolation, meta=( DisplayName = "Color Scale Track" ) )
+UCLASS(meta=( DisplayName = "Color Scale Track" ) )
 class UInterpTrackColorScale : public UInterpTrackVectorBase
 {
 	GENERATED_UCLASS_BODY()
diff --git a/Engine/Source/Runtime/Engine/Classes/Matinee/InterpTrackDirector.h b/Engine/Source/Runtime/Engine/Classes/Matinee/InterpTrackDirector.h
--- a/Engine/Source/Runtime/Engine/Classes/Matinee/InterpTrackDirector.h
+++ b/Engine/Source/Runtime/Engine/Classes/Matinee/InterpTrackDirector.h
@@ -41,7 +41,7 @@ struct FDirectorTrackCut
 
 };
 
-UCLASS(HeaderGroup=Interpolation, MinimalAPI, meta=( DisplayName = "Director Track" ) )
+UCLASS(MinimalAPI, meta=( DisplayName = "Director Track" ) )
 class UInterpTrackDirector : public UInterpTrack
 {
 	GENERATED_UCLASS_BODY()
diff --git a/Engine/Source/Runtime/Engine/Classes/Matinee/InterpTrackEvent.h b/Engine/Source/Runtime/Engine/Classes/Matinee/InterpTrackEvent.h
--- a/Engine/Source/Runtime/Engine/Classes/Matinee/InterpTrackEvent.h
+++ b/Engine/Source/Runtime/Engine/Classes/Matinee/InterpTrackEvent.h
@@ -30,7 +30,7 @@ struct FEventTrackKey
 
 };
 
-UCLASS(HeaderGroup=Interpolation, MinimalAPI, meta=( DisplayName = "Event Track" ) )
+UCLASS(MinimalAPI, meta=( DisplayName = "Event Track" ) )
 class UInterpTrackEvent : public UInterpTrack
 {
 	GENERATED_UCLASS_BODY()
diff --git a/Engine/Source/Runtime/Engine/Classes/Matinee/InterpTrackFade.h b/Engine/Source/Runtime/Engine/Classes/Matinee/InterpTrackFade.h
--- a/Engine/Source/Runtime/Engine/Classes/Matinee/InterpTrackFade.h
+++ b/Engine/Source/Runtime/Engine/Classes/Matinee/InterpTrackFade.h
@@ -3,7 +3,7 @@
 #pragma once
 #include "InterpTrackFade.generated.h"
 
-UCLASS(HeaderGroup=Interpolation, meta=( DisplayName = "Fade Track" ) )
+UCLASS(meta=( DisplayName = "Fade Track" ) )
 class UInterpTrackFade : public UInterpTrackFloatBase
 {
 	GENERATED_UCLASS_BODY()
diff --git a/Engine/Source/Runtime/Engine/Classes/Matinee/InterpTrackFloatMaterialParam.h b/Engine/Source/Runtime/Engine/Classes/Matinee/InterpTrackFloatMaterialParam.h
--- a/Engine/Source/Runtime/Engine/Classes/Matinee/InterpTrackFloatMaterialParam.h
+++ b/Engine/Source/Runtime/Engine/Classes/Matinee/InterpTrackFloatMaterialParam.h
@@ -3,7 +3,7 @@
 #pragma once
 #include "InterpTrackFloatMaterialParam.generated.h"
 
-UCLASS(HeaderGroup=Interpolation, meta=( DisplayName = "Float Material Parameter Track" ) )
+UCLASS(meta=( DisplayName = "Float Material Parameter Track" ) )
 class UInterpTrackFloatMaterialParam : public UInterpTrackFloatBase
 {
 	GENERATED_UCLASS_BODY()
diff --git a/Engine/Source/Runtime/Engine/Classes/Matinee/InterpTrackFloatParticleParam.h b/Engine/Source/Runtime/Engine/Classes/Matinee/InterpTrackFloatParticleParam.h
--- a/Engine/Source/Runtime/Engine/Classes/Matinee/InterpTrackFloatParticleParam.h
+++ b/Engine/Source/Runtime/Engine/Classes/Matinee/InterpTrackFloatParticleParam.h
@@ -3,7 +3,7 @@
 #pragma once
 #include "InterpTrackFloatParticleParam.generated.h"
 
-UCLASS(HeaderGroup=Interpolation, meta=( DisplayName = "Float Particle Param Track" ) )
+UCLASS(meta=( DisplayName = "Float Particle Param Track" ) )
 class UInterpTrackFloatParticleParam : public UInterpTrackFloatBase
 {
 	GENERATED_UCLASS_BODY()
diff --git a/Engine/Source/Runtime/Engine/Classes/Matinee/InterpTrackFloatProp.h b/Engine/Source/Runtime/Engine/Classes/Matinee/InterpTrackFloatProp.h
--- a/Engine/Source/Runtime/Engine/Classes/Matinee/InterpTrackFloatProp.h
+++ b/Engine/Source/Runtime/Engine/Classes/Matinee/InterpTrackFloatProp.h
@@ -3,7 +3,7 @@
 #pragma once
 #include "InterpTrackFloatProp.generated.h"
 
-UCLASS(HeaderGroup=Interpolation, MinimalAPI, meta=( DisplayName = "Float Property Track" ) )
+UCLASS(MinimalAPI, meta=( DisplayName = "Float Property Track" ) )
 class UInterpTrackFloatProp : public UInterpTrackFloatBase
 {
 	GENERATED_UCLASS_BODY()
diff --git a/Engine/Source/Runtime/Engine/Classes/Matinee/InterpTrackLinearColorProp.h b/Engine/Source/Runtime/Engine/Classes/Matinee/InterpTrackLinearColorProp.h
--- a/Engine/Source/Runtime/Engine/Classes/Matinee/InterpTrackLinearColorProp.h
+++ b/Engine/Source/Runtime/Engine/Classes/Matinee/InterpTrackLinearColorProp.h
@@ -4,7 +4,7 @@
 #pragma once
 #include "InterpTrackLinearColorProp.generated.h"
 
-UCLASS(HeaderGroup=Interpolation, MinimalAPI, meta=( DisplayName = "LinearColor Property Track" ) )
+UCLASS(MinimalAPI, meta=( DisplayName = "LinearColor Property Track" ) )
 class UInterpTrackLinearColorProp : public UInterpTrackLinearColorBase
 {
 	GENERATED_UCLASS_BODY()
diff --git a/Engine/Source/Runtime/Engine/Classes/Matinee/InterpTrackMove.h b/Engine/Source/Runtime/Engine/Classes/Matinee/InterpTrackMove.h
--- a/Engine/Source/Runtime/Engine/Classes/Matinee/InterpTrackMove.h
+++ b/Engine/Source/Runtime/Engine/Classes/Matinee/InterpTrackMove.h
@@ -86,7 +86,7 @@ struct FInterpLookupTrack
 };
 
 /** Track containing data for moving an actor around over time. */
-UCLASS(HeaderGroup=Interpolation, MinimalAPI, meta=( DisplayName = "Movement Track" ) )
+UCLASS(MinimalAPI, meta=( DisplayName = "Movement Track" ) )
 class UInterpTrackMove : public UInterpTrack
 {
 	GENERATED_UCLASS_BODY()
diff --git a/Engine/Source/Runtime/Engine/Classes/Matinee/InterpTrackMoveAxis.h b/Engine/Source/Runtime/Engine/Classes/Matinee/InterpTrackMoveAxis.h
--- a/Engine/Source/Runtime/Engine/Classes/Matinee/InterpTrackMoveAxis.h
+++ b/Engine/Source/Runtime/Engine/Classes/Matinee/InterpTrackMoveAxis.h
@@ -21,7 +21,7 @@ enum EInterpMoveAxis
 	AXIS_RotationZ,
 };
 
-UCLASS(dependson=UInterpTrackMove, HeaderGroup=Interpolation, MinimalAPI, meta=( DisplayName = "Move Axis Track" ) )
+UCLASS(dependson=UInterpTrackMove, MinimalAPI, meta=( DisplayName = "Move Axis Track" ) )
 class UInterpTrackMoveAxis : public UInterpTrackFloatBase
 {
 	GENERATED_UCLASS_BODY()
diff --git a/Engine/Source/Runtime/Engine/Classes/Matinee/InterpTrackParticleReplay.h b/Engine/Source/Runtime/Engine/Classes/Matinee/InterpTrackParticleReplay.h
--- a/Engine/Source/Runtime/Engine/Classes/Matinee/InterpTrackParticleReplay.h
+++ b/Engine/Source/Runtime/Engine/Classes/Matinee/InterpTrackParticleReplay.h
@@ -87,7 +87,7 @@ struct FParticleReplayTrackKey
 
 };
 
-UCLASS(HeaderGroup=Interpolation, MinimalAPI, meta=( DisplayName = "Particle Replay Track" ) )
+UCLASS(MinimalAPI, meta=( DisplayName = "Particle Replay Track" ) )
 class UInterpTrackParticleReplay : public UInterpTrack
 {
 	GENERATED_UCLASS_BODY()
diff --git a/Engine/Source/Runtime/Engine/Classes/Matinee/InterpTrackSlomo.h b/Engine/Source/Runtime/Engine/Classes/Matinee/InterpTrackSlomo.h
--- a/Engine/Source/Runtime/Engine/Classes/Matinee/InterpTrackSlomo.h
+++ b/Engine/Source/Runtime/Engine/Classes/Matinee/InterpTrackSlomo.h
@@ -4,7 +4,7 @@
 #pragma once
 #include "InterpTrackSlomo.generated.h"
 
-UCLASS(HeaderGroup=Interpolation, meta=( DisplayName = "Slomo Track" ) )
+UCLASS(meta=( DisplayName = "Slomo Track" ) )
 class UInterpTrackSlomo : public UInterpTrackFloatBase
 {
 	GENERATED_UCLASS_BODY()
diff --git a/Engine/Source/Runtime/Engine/Classes/Matinee/InterpTrackSound.h b/Engine/Source/Runtime/Engine/Classes/Matinee/InterpTrackSound.h
--- a/Engine/Source/Runtime/Engine/Classes/Matinee/InterpTrackSound.h
+++ b/Engine/Source/Runtime/Engine/Classes/Matinee/InterpTrackSound.h
@@ -38,7 +38,7 @@ struct FSoundTrackKey
 	
 };
 
-UCLASS(HeaderGroup=Interpolation, MinimalAPI, meta=( DisplayName = "Sound Track" ) )
+UCLASS(MinimalAPI, meta=( DisplayName = "Sound Track" ) )
 class UInterpTrackSound : public UInterpTrackVectorBase
 {
 	GENERATED_UCLASS_BODY()
diff --git a/Engine/Source/Runtime/Engine/Classes/Matinee/InterpTrackToggle.h b/Engine/Source/Runtime/Engine/Classes/Matinee/InterpTrackToggle.h
--- a/Engine/Source/Runtime/Engine/Classes/Matinee/InterpTrackToggle.h
+++ b/Engine/Source/Runtime/Engine/Classes/Matinee/InterpTrackToggle.h
@@ -42,7 +42,7 @@ struct FToggleTrackKey
 
 };
 
-UCLASS(HeaderGroup=Interpolation, MinimalAPI, meta=( DisplayName = "Toggle Track" ) )
+UCLASS(MinimalAPI, meta=( DisplayName = "Toggle Track" ) )
 class UInterpTrackToggle : public UInterpTrack
 {
 	GENERATED_UCLASS_BODY()
diff --git a/Engine/Source/Runtime/Engine/Classes/Matinee/InterpTrackVectorMaterialParam.h b/Engine/Source/Runtime/Engine/Classes/Matinee/InterpTrackVectorMaterialParam.h
--- a/Engine/Source/Runtime/Engine/Classes/Matinee/InterpTrackVectorMaterialParam.h
+++ b/Engine/Source/Runtime/Engine/Classes/Matinee/InterpTrackVectorMaterialParam.h
@@ -3,7 +3,7 @@
 #pragma once
 #include "InterpTrackVectorMaterialParam.generated.h"
 
-UCLASS(HeaderGroup=Interpolation, MinimalAPI, meta=( DisplayName = "Vector Material Parameter Track" ))
+UCLASS(MinimalAPI, meta=( DisplayName = "Vector Material Parameter Track" ))
 class UInterpTrackVectorMaterialParam : public UInterpTrackVectorBase
 {
 	GENERATED_UCLASS_BODY()
diff --git a/Engine/Source/Runtime/Engine/Classes/Matinee/InterpTrackVectorProp.h b/Engine/Source/Runtime/Engine/Classes/Matinee/InterpTrackVectorProp.h
--- a/Engine/Source/Runtime/Engine/Classes/Matinee/InterpTrackVectorProp.h
+++ b/Engine/Source/Runtime/Engine/Classes/Matinee/InterpTrackVectorProp.h
@@ -4,7 +4,7 @@
 #pragma once
 #include "InterpTrackVectorProp.generated.h"
 
-UCLASS(HeaderGroup=Interpolation, MinimalAPI, meta=( DisplayName = "Vector Property Track" ) )
+UCLASS(MinimalAPI, meta=( DisplayName = "Vector Property Track" ) )
 class UInterpTrackVectorProp : public UInterpTrackVectorBase
 {
 	GENERATED_UCLASS_BODY()
diff --git a/Engine/Source/Runtime/Engine/Classes/Matinee/InterpTrackVisibility.h b/Engine/Source/Runtime/Engine/Classes/Matinee/InterpTrackVisibility.h
--- a/Engine/Source/Runtime/Engine/Classes/Matinee/InterpTrackVisibility.h
+++ b/Engine/Source/Runtime/Engine/Classes/Matinee/InterpTrackVisibility.h
@@ -61,7 +61,7 @@ struct FVisibilityTrackKey
 
 };
 
-UCLASS(HeaderGroup=Interpolation, MinimalAPI, meta=( DisplayName = "Visibility Track" ) )
+UCLASS(MinimalAPI, meta=( DisplayName = "Visibility Track" ) )
 class UInterpTrackVisibility : public UInterpTrack
 {
 	GENERATED_UCLASS_BODY()
diff --git a/Engine/Source/Runtime/Engine/Classes/Particles/Acceleration/ParticleModuleAcceleration.h b/Engine/Source/Runtime/Engine/Classes/Particles/Acceleration/ParticleModuleAcceleration.h
--- a/Engine/Source/Runtime/Engine/Classes/Particles/Acceleration/ParticleModuleAcceleration.h
+++ b/Engine/Source/Runtime/Engine/Classes/Particles/Acceleration/ParticleModuleAcceleration.h
@@ -4,7 +4,7 @@
 #pragma once
 #include "ParticleModuleAcceleration.generated.h"
 
-UCLASS(HeaderGroup=Particle, editinlinenew, hidecategories=Object, meta=(DisplayName = "Acceleration"), MinimalAPI)
+UCLASS(editinlinenew, hidecategories=Object, meta=(DisplayName = "Acceleration"), MinimalAPI)
 class UParticleModuleAcceleration : public UParticleModuleAccelerationBase
 {
 	GENERATED_UCLASS_BODY()
diff --git a/Engine/Source/Runtime/Engine/Classes/Particles/Acceleration/ParticleModuleAccelerationBase.h b/Engine/Source/Runtime/Engine/Classes/Particles/Acceleration/ParticleModuleAccelerationBase.h
--- a/Engine/Source/Runtime/Engine/Classes/Particles/Acceleration/ParticleModuleAccelerationBase.h
+++ b/Engine/Source/Runtime/Engine/Classes/Particles/Acceleration/ParticleModuleAccelerationBase.h
@@ -4,7 +4,7 @@
 #pragma once
 #include "ParticleModuleAccelerationBase.generated.h"
 
-UCLASS(HeaderGroup=Particle, editinlinenew, hidecategories=Object, meta=(DisplayName = "Acceleration"), abstract)
+UCLASS(editinlinenew, hidecategories=Object, meta=(DisplayName = "Acceleration"), abstract)
 class UParticleModuleAccelerationBase : public UParticleModule
 {
 	GENERATED_UCLASS_BODY()
diff --git a/Engine/Source/Runtime/Engine/Classes/Particles/Acceleration/ParticleModuleAccelerationConstant.h b/Engine/Source/Runtime/Engine/Classes/Particles/Acceleration/ParticleModuleAccelerationConstant.h
--- a/Engine/Source/Runtime/Engine/Classes/Particles/Acceleration/ParticleModuleAccelerationConstant.h
+++ b/Engine/Source/Runtime/Engine/Classes/Particles/Acceleration/ParticleModuleAccelerationConstant.h
@@ -8,7 +8,7 @@
 
 #include "ParticleModuleAccelerationConstant.generated.h"
 
-UCLASS(HeaderGroup=Particle, editinlinenew, hidecategories=(Object, Acceleration), meta=(DisplayName = "Const Acceleration"))
+UCLASS(editinlinenew, hidecategories=(Object, Acceleration), meta=(DisplayName = "Const Acceleration"))
 class UParticleModuleAccelerationConstant : public UParticleModuleAccelerationBase
 {
 	GENERATED_UCLASS_BODY()
diff --git a/Engine/Source/Runtime/Engine/Classes/Particles/Acceleration/ParticleModuleAccelerationDrag.h b/Engine/Source/Runtime/Engine/Classes/Particles/Acceleration/ParticleModuleAccelerationDrag.h
--- a/Engine/Source/Runtime/Engine/Classes/Particles/Acceleration/ParticleModuleAccelerationDrag.h
+++ b/Engine/Source/Runtime/Engine/Classes/Particles/Acceleration/ParticleModuleAccelerationDrag.h
@@ -8,7 +8,7 @@
 
 #include "ParticleModuleAccelerationDrag.generated.h"
 
-UCLASS(HeaderGroup=Particle, editinlinenew, hidecategories=(Object, Acceleration), MinimalAPI, meta=(DisplayName = "Drag"))
+UCLASS(editinlinenew, hidecategories=(Object, Acceleration), MinimalAPI, meta=(DisplayName = "Drag"))
 class UParticleModuleAccelerationDrag : public UParticleModuleAccelerationBase
 {
 	GENERATED_UCLASS_BODY()
diff --git a/Engine/Source/Runtime/Engine/Classes/Particles/Acceleration/ParticleModuleAccelerationDragScaleOverLife.h b/Engine/Source/Runtime/Engine/Classes/Particles/Acceleration/ParticleModuleAccelerationDragScaleOverLife.h
--- a/Engine/Source/Runtime/Engine/Classes/Particles/Acceleration/ParticleModuleAccelerationDragScaleOverLife.h
+++ b/Engine/Source/Runtime/Engine/Classes/Particles/Acceleration/ParticleModuleAccelerationDragScaleOverLife.h
@@ -8,7 +8,7 @@
 
 #include "ParticleModuleAccelerationDragScaleOverLife.generated.h"
 
-UCLASS(HeaderGroup=Particle, editinlinenew, hidecategories=(UObject, Acceleration), MinimalAPI, meta=(DisplayName = "Drag Scale/Life"))
+UCLASS(editinlinenew, hidecategories=(UObject, Acceleration), MinimalAPI, meta=(DisplayName = "Drag Scale/Life"))
 class UParticleModuleAccelerationDragScaleOverLife : public UParticleModuleAccelerationBase
 {
 	GENERATED_UCLASS_BODY()
diff --git a/Engine/Source/Runtime/Engine/Classes/Particles/Acceleration/ParticleModuleAccelerationOverLifetime.h b/Engine/Source/Runtime/Engine/Classes/Particles/Acceleration/ParticleModuleAccelerationOverLifetime.h
--- a/Engine/Source/Runtime/Engine/Classes/Particles/Acceleration/ParticleModuleAccelerationOverLifetime.h
+++ b/Engine/Source/Runtime/Engine/Classes/Particles/Acceleration/ParticleModuleAccelerationOverLifetime.h
@@ -4,7 +4,7 @@
 #pragma once
 #include "ParticleModuleAccelerationOverLifetime.generated.h"
 
-UCLASS(HeaderGroup=Particle, editinlinenew, hidecategories=Object, meta=(DisplayName = "Acceleration/Life"), MinimalAPI)
+UCLASS(editinlinenew, hidecategories=Object, meta=(DisplayName = "Acceleration/Life"), MinimalAPI)
 class UParticleModuleAccelerationOverLifetime : public UParticleModuleAccelerationBase
 {
 	GENERATED_UCLASS_BODY()
diff --git a/Engine/Source/Runtime/Engine/Classes/Particles/Attractor/ParticleModuleAttractorBase.h b/Engine/Source/Runtime/Engine/Classes/Particles/Attractor/ParticleModuleAttractorBase.h
--- a/Engine/Source/Runtime/Engine/Classes/Particles/Attractor/ParticleModuleAttractorBase.h
+++ b/Engine/Source/Runtime/Engine/Classes/Particles/Attractor/ParticleModuleAttractorBase.h
@@ -4,7 +4,7 @@
 #pragma once
 #include "ParticleModuleAttractorBase.generated.h"
 
-UCLASS(HeaderGroup=Particle, editinlinenew, hidecategories=Object, abstract, meta=(DisplayName = "Attraction"))
+UCLASS(editinlinenew, hidecategories=Object, abstract, meta=(DisplayName = "Attraction"))
 class UParticleModuleAttractorBase : public UParticleModule
 {
 	GENERATED_UCLASS_BODY()
diff --git a/Engine/Source/Runtime/Engine/Classes/Particles/Attractor/ParticleModuleAttractorLine.h b/Engine/Source/Runtime/Engine/Classes/Particles/Attractor/ParticleModuleAttractorLine.h
--- a/Engine/Source/Runtime/Engine/Classes/Particles/Attractor/ParticleModuleAttractorLine.h
+++ b/Engine/Source/Runtime/Engine/Classes/Particles/Attractor/ParticleModuleAttractorLine.h
@@ -4,7 +4,7 @@
 #pragma once
 #include "ParticleModuleAttractorLine.generated.h"
 
-UCLASS(HeaderGroup=Particle, editinlinenew, hidecategories=Object, MinimalAPI, meta=(DisplayName = "Line Attractor"))
+UCLASS(editinlinenew, hidecategories=Object, MinimalAPI, meta=(DisplayName = "Line Attractor"))
 class UParticleModuleAttractorLine : public UParticleModuleAttractorBase
 {
 	GENERATED_UCLASS_BODY()
diff --git a/Engine/Source/Runtime/Engine/Classes/Particles/Attractor/ParticleModuleAttractorParticle.h b/Engine/Source/Runtime/Engine/Classes/Particles/Attractor/ParticleModuleAttractorParticle.h
--- a/Engine/Source/Runtime/Engine/Classes/Particles/Attractor/ParticleModuleAttractorParticle.h
+++ b/Engine/Source/Runtime/Engine/Classes/Particles/Attractor/ParticleModuleAttractorParticle.h
@@ -12,7 +12,7 @@ enum EAttractorParticleSelectionMethod
 	EAPSM_MAX,
 };
 
-UCLASS(HeaderGroup=Particle, editinlinenew, hidecategories=Object, MinimalAPI, meta=(DisplayName = "Particle Attractor"))
+UCLASS(editinlinenew, hidecategories=Object, MinimalAPI, meta=(DisplayName = "Particle Attractor"))
 class UParticleModuleAttractorParticle : public UParticleModuleAttractorBase
 {
 	GENERATED_UCLASS_BODY()
diff --git a/Engine/Source/Runtime/Engine/Classes/Particles/Attractor/ParticleModuleAttractorPoint.h b/Engine/Source/Runtime/Engine/Classes/Particles/Attractor/ParticleModuleAttractorPoint.h
--- a/Engine/Source/Runtime/Engine/Classes/Particles/Attractor/ParticleModuleAttractorPoint.h
+++ b/Engine/Source/Runtime/Engine/Classes/Particles/Attractor/ParticleModuleAttractorPoint.h
@@ -4,7 +4,7 @@
 #pragma once
 #include "ParticleModuleAttractorPoint.generated.h"
 
-UCLASS(HeaderGroup=Particle, editinlinenew, hidecategories=Object, MinimalAPI, meta=(DisplayName = "Point Attractor"))
+UCLASS(editinlinenew, hidecategories=Object, MinimalAPI, meta=(DisplayName = "Point Attractor"))
 class UParticleModuleAttractorPoint : public UParticleModuleAttractorBase
 {
 	GENERATED_UCLASS_BODY()
diff --git a/Engine/Source/Runtime/Engine/Classes/Particles/Attractor/ParticleModuleAttractorPointGravity.h b/Engine/Source/Runtime/Engine/Classes/Particles/Attractor/ParticleModuleAttractorPointGravity.h
--- a/Engine/Source/Runtime/Engine/Classes/Particles/Attractor/ParticleModuleAttractorPointGravity.h
+++ b/Engine/Source/Runtime/Engine/Classes/Particles/Attractor/ParticleModuleAttractorPointGravity.h
@@ -9,7 +9,7 @@
 
 #include "ParticleModuleAttractorPointGravity.generated.h"
 
-UCLASS(HeaderGroup=Particle, editinlinenew, hidecategories=Object, meta=(DisplayName = "Point Gravity"))
+UCLASS(editinlinenew, hidecategories=Object, meta=(DisplayName = "Point Gravity"))
 class UParticleModuleAttractorPointGravity : public UParticleModuleAttractorBase
 {
 	GENERATED_UCLASS_BODY()
diff --git a/Engine/Source/Runtime/Engine/Classes/Particles/Beam/ParticleModuleBeamBase.h b/Engine/Source/Runtime/Engine/Classes/Particles/Beam/ParticleModuleBeamBase.h
--- a/Engine/Source/Runtime/Engine/Classes/Particles/Beam/ParticleModuleBeamBase.h
+++ b/Engine/Source/Runtime/Engine/Classes/Particles/Beam/ParticleModuleBeamBase.h
@@ -51,7 +51,7 @@ enum Beam2SourceTargetTangentMethod
 	PEB2STTM_MAX,
 };
 
-UCLASS(HeaderGroup=Particle, editinlinenew, hidecategories=Object, abstract, meta=(DisplayName = "Beam"))
+UCLASS(editinlinenew, hidecategories=Object, abstract, meta=(DisplayName = "Beam"))
 class UParticleModuleBeamBase : public UParticleModule
 {
 	GENERATED_UCLASS_BODY()
diff --git a/Engine/Source/Runtime/Engine/Classes/Particles/Beam/ParticleModuleBeamModifier.h b/Engine/Source/Runtime/Engine/Classes/Particles/Beam/ParticleModuleBeamModifier.h
--- a/Engine/Source/Runtime/Engine/Classes/Particles/Beam/ParticleModuleBeamModifier.h
+++ b/Engine/Source/Runtime/Engine/Classes/Particles/Beam/ParticleModuleBeamModifier.h
@@ -48,7 +48,7 @@ struct FBeamModifierOptions
 
 };
 
-UCLASS(HeaderGroup=Particle, editinlinenew, hidecategories=Object, meta=(DisplayName = "Beam Modifier"))
+UCLASS(editinlinenew, hidecategories=Object, meta=(DisplayName = "Beam Modifier"))
 class UParticleModuleBeamModifier : public UParticleModuleBeamBase
 {
 	GENERATED_UCLASS_BODY()
diff --git a/Engine/Source/Runtime/Engine/Classes/Particles/Beam/ParticleModuleBeamNoise.h b/Engine/Source/Runtime/Engine/Classes/Particles/Beam/ParticleModuleBeamNoise.h
--- a/Engine/Source/Runtime/Engine/Classes/Particles/Beam/ParticleModuleBeamNoise.h
+++ b/Engine/Source/Runtime/Engine/Classes/Particles/Beam/ParticleModuleBeamNoise.h
@@ -10,7 +10,7 @@
 #pragma once
 #include "ParticleModuleBeamNoise.generated.h"
 
-UCLASS(HeaderGroup=Particle, editinlinenew, hidecategories=Object, MinimalAPI, meta=(DisplayName = "Noise"))
+UCLASS(editinlinenew, hidecategories=Object, MinimalAPI, meta=(DisplayName = "Noise"))
 class UParticleModuleBeamNoise : public UParticleModuleBeamBase
 {
 	GENERATED_UCLASS_BODY()
diff --git a/Engine/Source/Runtime/Engine/Classes/Particles/Beam/ParticleModuleBeamTarget.h b/Engine/Source/Runtime/Engine/Classes/Particles/Beam/ParticleModuleBeamTarget.h
--- a/Engine/Source/Runtime/Engine/Classes/Particles/Beam/ParticleModuleBeamTarget.h
+++ b/Engine/Source/Runtime/Engine/Classes/Particles/Beam/ParticleModuleBeamTarget.h
@@ -10,7 +10,7 @@
 #pragma once
 #include "ParticleModuleBeamTarget.generated.h"
 
-UCLASS(HeaderGroup=Particle, editinlinenew, hidecategories=Object, MinimalAPI, meta=(DisplayName = "Target"))
+UCLASS(editinlinenew, hidecategories=Object, MinimalAPI, meta=(DisplayName = "Target"))
 class UParticleModuleBeamTarget : public UParticleModuleBeamBase
 {
 	GENERATED_UCLASS_BODY()
diff --git a/Engine/Source/Runtime/Engine/Classes/Particles/Camera/ParticleModuleCameraBase.h b/Engine/Source/Runtime/Engine/Classes/Particles/Camera/ParticleModuleCameraBase.h
--- a/Engine/Source/Runtime/Engine/Classes/Particles/Camera/ParticleModuleCameraBase.h
+++ b/Engine/Source/Runtime/Engine/Classes/Particles/Camera/ParticleModuleCameraBase.h
@@ -4,7 +4,7 @@
 #pragma once
 #include "ParticleModuleCameraBase.generated.h"
 
-UCLASS(HeaderGroup=Particle, editinlinenew, hidecategories=Object, abstract, meta=(DisplayName = "Camera"))
+UCLASS(editinlinenew, hidecategories=Object, abstract, meta=(DisplayName = "Camera"))
 class UParticleModuleCameraBase : public UParticleModule
 {
 	GENERATED_UCLASS_BODY()
diff --git a/Engine/Source/Runtime/Engine/Classes/Particles/Camera/ParticleModuleCameraOffset.h b/Engine/Source/Runtime/Engine/Classes/Particles/Camera/ParticleModuleCameraOffset.h
--- a/Engine/Source/Runtime/Engine/Classes/Particles/Camera/ParticleModuleCameraOffset.h
+++ b/Engine/Source/Runtime/Engine/Classes/Particles/Camera/ParticleModuleCameraOffset.h
@@ -16,7 +16,7 @@ enum EParticleCameraOffsetUpdateMethod
 	EPCOUM_MAX,
 };
 
-UCLASS(HeaderGroup=Particle, editinlinenew, hidecategories=Object, meta=(DisplayName = "Camera Offset"))
+UCLASS(editinlinenew, hidecategories=Object, meta=(DisplayName = "Camera Offset"))
 class UParticleModuleCameraOffset : public UParticleModuleCameraBase
 {
 	GENERATED_UCLASS_BODY()
diff --git a/Engine/Source/Runtime/Engine/Classes/Particles/Collision/ParticleModuleCollision.h b/Engine/Source/Runtime/Engine/Classes/Particles/Collision/ParticleModuleCollision.h
--- a/Engine/Source/Runtime/Engine/Classes/Particles/Collision/ParticleModuleCollision.h
+++ b/Engine/Source/Runtime/Engine/Classes/Particles/Collision/ParticleModuleCollision.h
@@ -4,7 +4,7 @@
 #pragma once
 #include "ParticleModuleCollision.generated.h"
 
-UCLASS(HeaderGroup=Particle, editinlinenew, hidecategories=Object, MinimalAPI, meta=(DisplayName = "Collision"))
+UCLASS(editinlinenew, hidecategories=Object, MinimalAPI, meta=(DisplayName = "Collision"))
 class UParticleModuleCollision : public UParticleModuleCollisionBase
 {
 	GENERATED_UCLASS_BODY()
diff --git a/Engine/Source/Runtime/Engine/Classes/Particles/Collision/ParticleModuleCollisionBase.h b/Engine/Source/Runtime/Engine/Classes/Particles/Collision/ParticleModuleCollisionBase.h
--- a/Engine/Source/Runtime/Engine/Classes/Particles/Collision/ParticleModuleCollisionBase.h
+++ b/Engine/Source/Runtime/Engine/Classes/Particles/Collision/ParticleModuleCollisionBase.h
@@ -25,7 +25,7 @@ enum EParticleCollisionComplete
 	EPCC_MAX,
 };
 
-UCLASS(HeaderGroup=Particle, editinlinenew, hidecategories=Object, abstract, meta=(DisplayName = "Collision"))
+UCLASS(editinlinenew, hidecategories=Object, abstract, meta=(DisplayName = "Collision"))
 class UParticleModuleCollisionBase : public UParticleModule
 {
 	GENERATED_UCLASS_BODY()
diff --git a/Engine/Source/Runtime/Engine/Classes/Particles/Collision/ParticleModuleCollisionGPU.h b/Engine/Source/Runtime/Engine/Classes/Particles/Collision/ParticleModuleCollisionGPU.h
--- a/Engine/Source/Runtime/Engine/Classes/Particles/Collision/ParticleModuleCollisionGPU.h
+++ b/Engine/Source/Runtime/Engine/Classes/Particles/Collision/ParticleModuleCollisionGPU.h
@@ -20,7 +20,7 @@ namespace EParticleCollisionResponse
 	};
 }
 
-UCLASS(HeaderGroup=Particle, editinlinenew, hidecategories=Object, MinimalAPI, meta=(DisplayName = "Collision (Scene Depth)"))
+UCLASS(editinlinenew, hidecategories=Object, MinimalAPI, meta=(DisplayName = "Collision (Scene Depth)"))
 class UParticleModuleCollisionGPU : public UParticleModuleCollisionBase
 {
 	GENERATED_UCLASS_BODY()
diff --git a/Engine/Source/Runtime/Engine/Classes/Particles/Color/ParticleModuleColor.h b/Engine/Source/Runtime/Engine/Classes/Particles/Color/ParticleModuleColor.h
--- a/Engine/Source/Runtime/Engine/Classes/Particles/Color/ParticleModuleColor.h
+++ b/Engine/Source/Runtime/Engine/Classes/Particles/Color/ParticleModuleColor.h
@@ -4,7 +4,7 @@
 #pragma once
 #include "ParticleModuleColor.generated.h"
 
-UCLASS(HeaderGroup=Particle, editinlinenew, hidecategories=Object, MinimalAPI, meta=(DisplayName = "Initial Color"))
+UCLASS(editinlinenew, hidecategories=Object, MinimalAPI, meta=(DisplayName = "Initial Color"))
 class UParticleModuleColor : public UParticleModuleColorBase
 {
 	GENERATED_UCLASS_BODY()
diff --git a/Engine/Source/Runtime/Engine/Classes/Particles/Color/ParticleModuleColorBase.h b/Engine/Source/Runtime/Engine/Classes/Particles/Color/ParticleModuleColorBase.h
--- a/Engine/Source/Runtime/Engine/Classes/Particles/Color/ParticleModuleColorBase.h
+++ b/Engine/Source/Runtime/Engine/Classes/Particles/Color/ParticleModuleColorBase.h
@@ -4,7 +4,7 @@
 #pragma once
 #include "ParticleModuleColorBase.generated.h"
 
-UCLASS(HeaderGroup=Particle, editinlinenew, hidecategories=Object, abstract, meta=(DisplayName = "Color"))
+UCLASS(editinlinenew, hidecategories=Object, abstract, meta=(DisplayName = "Color"))
 class UParticleModuleColorBase : public UParticleModule
 {
 	GENERATED_UCLASS_BODY()
diff --git a/Engine/Source/Runtime/Engine/Classes/Particles/Color/ParticleModuleColorOverLife.h b/Engine/Source/Runtime/Engine/Classes/Particles/Color/ParticleModuleColorOverLife.h
--- a/Engine/Source/Runtime/Engine/Classes/Particles/Color/ParticleModuleColorOverLife.h
+++ b/Engine/Source/Runtime/Engine/Classes/Particles/Color/ParticleModuleColorOverLife.h
@@ -4,7 +4,7 @@
 #pragma once
 #include "ParticleModuleColorOverLife.generated.h"
 
-UCLASS(HeaderGroup=Particle, editinlinenew, hidecategories=Object, MinimalAPI, meta=(DisplayName = "Color Over Life"))
+UCLASS(editinlinenew, hidecategories=Object, MinimalAPI, meta=(DisplayName = "Color Over Life"))
 class UParticleModuleColorOverLife : public UParticleModuleColorBase
 {
 	GENERATED_UCLASS_BODY()
diff --git a/Engine/Source/Runtime/Engine/Classes/Particles/Color/ParticleModuleColorScaleOverLife.h b/Engine/Source/Runtime/Engine/Classes/Particles/Color/ParticleModuleColorScaleOverLife.h
--- a/Engine/Source/Runtime/Engine/Classes/Particles/Color/ParticleModuleColorScaleOverLife.h
+++ b/Engine/Source/Runtime/Engine/Classes/Particles/Color/ParticleModuleColorScaleOverLife.h
@@ -10,7 +10,7 @@
 #pragma once
 #include "ParticleModuleColorScaleOverLife.generated.h"
 
-UCLASS(HeaderGroup=Particle, editinlinenew, hidecategories=Object, MinimalAPI, meta=(DisplayName = "Scale Color / Life"))
+UCLASS(editinlinenew, hidecategories=Object, MinimalAPI, meta=(DisplayName = "Scale Color / Life"))
 class UParticleModuleColorScaleOverLife : public UParticleModuleColorBase
 {
 	GENERATED_UCLASS_BODY()
diff --git a/Engine/Source/Runtime/Engine/Classes/Particles/Color/ParticleModuleColor_Seeded.h b/Engine/Source/Runtime/Engine/Classes/Particles/Color/ParticleModuleColor_Seeded.h
--- a/Engine/Source/Runtime/Engine/Classes/Particles/Color/ParticleModuleColor_Seeded.h
+++ b/Engine/Source/Runtime/Engine/Classes/Particles/Color/ParticleModuleColor_Seeded.h
@@ -4,7 +4,7 @@
 #pragma once
 #include "ParticleModuleColor_Seeded.generated.h"
 
-UCLASS(HeaderGroup=Particle, editinlinenew, hidecategories=Object, meta=(DisplayName = "Init Color (Seed)"))
+UCLASS(editinlinenew, hidecategories=Object, meta=(DisplayName = "Init Color (Seed)"))
 class UParticleModuleColor_Seeded : public UParticleModuleColor
 {
 	GENERATED_UCLASS_BODY()
diff --git a/Engine/Source/Runtime/Engine/Classes/Particles/EmitterCameraLensEffectBase.h b/Engine/Source/Runtime/Engine/Classes/Particles/EmitterCameraLensEffectBase.h
--- a/Engine/Source/Runtime/Engine/Classes/Particles/EmitterCameraLensEffectBase.h
+++ b/Engine/Source/Runtime/Engine/Classes/Particles/EmitterCameraLensEffectBase.h
@@ -9,7 +9,7 @@
 #pragma once
 #include "EmitterCameraLensEffectBase.generated.h"
 
-UCLASS(abstract, Blueprintable, dependson=AEmitter, HeaderGroup=Particle, MinimalAPI)
+UCLASS(abstract, Blueprintable, dependson=AEmitter, MinimalAPI)
 class AEmitterCameraLensEffectBase : public AEmitter
 {
 	GENERATED_UCLASS_BODY()
diff --git a/Engine/Source/Runtime/Engine/Classes/Particles/Event/ParticleModuleEventBase.h b/Engine/Source/Runtime/Engine/Classes/Particles/Event/ParticleModuleEventBase.h
--- a/Engine/Source/Runtime/Engine/Classes/Particles/Event/ParticleModuleEventBase.h
+++ b/Engine/Source/Runtime/Engine/Classes/Particles/Event/ParticleModuleEventBase.h
@@ -4,7 +4,7 @@
 #pragma once
 #include "ParticleModuleEventBase.generated.h"
 
-UCLASS(HeaderGroup=Particle, editinlinenew, hidecategories=Object, abstract, meta=(DisplayName = "Event"))
+UCLASS(editinlinenew, hidecategories=Object, abstract, meta=(DisplayName = "Event"))
 class UParticleModuleEventBase : public UParticleModule
 {
 	GENERATED_UCLASS_BODY()
diff --git a/Engine/Source/Runtime/Engine/Classes/Particles/Event/ParticleModuleEventGenerator.h b/Engine/Source/Runtime/Engine/Classes/Particles/Event/ParticleModuleEventGenerator.h
--- a/Engine/Source/Runtime/Engine/Classes/Particles/Event/ParticleModuleEventGenerator.h
+++ b/Engine/Source/Runtime/Engine/Classes/Particles/Event/ParticleModuleEventGenerator.h
@@ -53,7 +53,7 @@ struct FParticleEvent_GenerateInfo
 	}
 };
 
-UCLASS(HeaderGroup=Particle, editinlinenew, hidecategories=Object, MinimalAPI, meta=(DisplayName = "Event Generator"))
+UCLASS(editinlinenew, hidecategories=Object, MinimalAPI, meta=(DisplayName = "Event Generator"))
 class UParticleModuleEventGenerator : public UParticleModuleEventBase
 {
 	GENERATED_UCLASS_BODY()
diff --git a/Engine/Source/Runtime/Engine/Classes/Particles/Event/ParticleModuleEventReceiverBase.h b/Engine/Source/Runtime/Engine/Classes/Particles/Event/ParticleModuleEventReceiverBase.h
--- a/Engine/Source/Runtime/Engine/Classes/Particles/Event/ParticleModuleEventReceiverBase.h
+++ b/Engine/Source/Runtime/Engine/Classes/Particles/Event/ParticleModuleEventReceiverBase.h
@@ -4,7 +4,7 @@
 #pragma once
 #include "ParticleModuleEventReceiverBase.generated.h"
 
-UCLASS(HeaderGroup=Particle, editinlinenew, hidecategories=Object, abstract)
+UCLASS(editinlinenew, hidecategories=Object, abstract)
 class UParticleModuleEventReceiverBase : public UParticleModuleEventBase
 {
 	GENERATED_UCLASS_BODY()
diff --git a/Engine/Source/Runtime/Engine/Classes/Particles/Event/ParticleModuleEventReceiverKillParticles.h b/Engine/Source/Runtime/Engine/Classes/Particles/Event/ParticleModuleEventReceiverKillParticles.h
--- a/Engine/Source/Runtime/Engine/Classes/Particles/Event/ParticleModuleEventReceiverKillParticles.h
+++ b/Engine/Source/Runtime/Engine/Classes/Particles/Event/ParticleModuleEventReceiverKillParticles.h
@@ -4,7 +4,7 @@
 #pragma once
 #include "ParticleModuleEventReceiverKillParticles.generated.h"
 
-UCLASS(HeaderGroup=Particle, editinlinenew, hidecategories=Object, meta=(DisplayName = "EventReceiver Kill All"))
+UCLASS(editinlinenew, hidecategories=Object, meta=(DisplayName = "EventReceiver Kill All"))
 class UParticleModuleEventReceiverKillParticles : public UParticleModuleEventReceiverBase
 {
 	GENERATED_UCLASS_BODY()
diff --git a/Engine/Source/Runtime/Engine/Classes/Particles/Event/ParticleModuleEventReceiverSpawn.h b/Engine/Source/Runtime/Engine/Classes/Particles/Event/ParticleModuleEventReceiverSpawn.h
--- a/Engine/Source/Runtime/Engine/Classes/Particles/Event/ParticleModuleEventReceiverSpawn.h
+++ b/Engine/Source/Runtime/Engine/Classes/Particles/Event/ParticleModuleEventReceiverSpawn.h
@@ -4,7 +4,7 @@
 #pragma once
 #include "ParticleModuleEventReceiverSpawn.generated.h"
 
-UCLASS(HeaderGroup=Particle, editinlinenew, hidecategories=Object, meta=(DisplayName = "EventReceiver Spawn"))
+UCLASS(editinlinenew, hidecategories=Object, meta=(DisplayName = "EventReceiver Spawn"))
 class UParticleModuleEventReceiverSpawn : public UParticleModuleEventReceiverBase
 {
 	GENERATED_UCLASS_BODY()
diff --git a/Engine/Source/Runtime/Engine/Classes/Particles/Event/ParticleModuleEventSendToGame.h b/Engine/Source/Runtime/Engine/Classes/Particles/Event/ParticleModuleEventSendToGame.h
--- a/Engine/Source/Runtime/Engine/Classes/Particles/Event/ParticleModuleEventSendToGame.h
+++ b/Engine/Source/Runtime/Engine/Classes/Particles/Event/ParticleModuleEventSendToGame.h
@@ -4,7 +4,7 @@
 #pragma once
 #include "ParticleModuleEventSendToGame.generated.h"
 
-UCLASS(HeaderGroup=Particle, abstract, editinlinenew, hidecategories=Object)
+UCLASS(abstract, editinlinenew, hidecategories=Object)
 class UParticleModuleEventSendToGame : public UObject
 {
 	GENERATED_UCLASS_BODY()
diff --git a/Engine/Source/Runtime/Engine/Classes/Particles/Kill/ParticleModuleKillBase.h b/Engine/Source/Runtime/Engine/Classes/Particles/Kill/ParticleModuleKillBase.h
--- a/Engine/Source/Runtime/Engine/Classes/Particles/Kill/ParticleModuleKillBase.h
+++ b/Engine/Source/Runtime/Engine/Classes/Particles/Kill/ParticleModuleKillBase.h
@@ -4,7 +4,7 @@
 #pragma once
 #include "ParticleModuleKillBase.generated.h"
 
-UCLASS(HeaderGroup=Particle, editinlinenew, hidecategories=Object, abstract, meta=(DisplayName = "Kill"))
+UCLASS(editinlinenew, hidecategories=Object, abstract, meta=(DisplayName = "Kill"))
 class UParticleModuleKillBase : public UParticleModule
 {
 	GENERATED_UCLASS_BODY()
diff --git a/Engine/Source/Runtime/Engine/Classes/Particles/Kill/ParticleModuleKillBox.h b/Engine/Source/Runtime/Engine/Classes/Particles/Kill/ParticleModuleKillBox.h
--- a/Engine/Source/Runtime/Engine/Classes/Particles/Kill/ParticleModuleKillBox.h
+++ b/Engine/Source/Runtime/Engine/Classes/Particles/Kill/ParticleModuleKillBox.h
@@ -4,7 +4,7 @@
 #pragma once
 #include "ParticleModuleKillBox.generated.h"
 
-UCLASS(HeaderGroup=Particle, editinlinenew, hidecategories=Object, meta=(DisplayName = "Kill Box"))
+UCLASS(editinlinenew, hidecategories=Object, meta=(DisplayName = "Kill Box"))
 class UParticleModuleKillBox : public UParticleModuleKillBase
 {
 	GENERATED_UCLASS_BODY()
diff --git a/Engine/Source/Runtime/Engine/Classes/Particles/Kill/ParticleModuleKillHeight.h b/Engine/Source/Runtime/Engine/Classes/Particles/Kill/ParticleModuleKillHeight.h
--- a/Engine/Source/Runtime/Engine/Classes/Particles/Kill/ParticleModuleKillHeight.h
+++ b/Engine/Source/Runtime/Engine/Classes/Particles/Kill/ParticleModuleKillHeight.h
@@ -4,7 +4,7 @@
 #pragma once
 #include "ParticleModuleKillHeight.generated.h"
 
-UCLASS(HeaderGroup=Particle, editinlinenew, hidecategories=Object, meta=(DisplayName = "Kill Height"))
+UCLASS(editinlinenew, hidecategories=Object, meta=(DisplayName = "Kill Height"))
 class UParticleModuleKillHeight : public UParticleModuleKillBase
 {
 	GENERATED_UCLASS_BODY()
diff --git a/Engine/Source/Runtime/Engine/Classes/Particles/Lifetime/ParticleModuleLifetime.h b/Engine/Source/Runtime/Engine/Classes/Particles/Lifetime/ParticleModuleLifetime.h
--- a/Engine/Source/Runtime/Engine/Classes/Particles/Lifetime/ParticleModuleLifetime.h
+++ b/Engine/Source/Runtime/Engine/Classes/Particles/Lifetime/ParticleModuleLifetime.h
@@ -4,7 +4,7 @@
 #pragma once
 #include "ParticleModuleLifetime.generated.h"
 
-UCLASS(HeaderGroup=Particle, editinlinenew, hidecategories=Object, MinimalAPI, meta=(DisplayName = "Lifetime"))
+UCLASS(editinlinenew, hidecategories=Object, MinimalAPI, meta=(DisplayName = "Lifetime"))
 class UParticleModuleLifetime : public UParticleModuleLifetimeBase
 {
 	GENERATED_UCLASS_BODY()
diff --git a/Engine/Source/Runtime/Engine/Classes/Particles/Lifetime/ParticleModuleLifetimeBase.h b/Engine/Source/Runtime/Engine/Classes/Particles/Lifetime/ParticleModuleLifetimeBase.h
--- a/Engine/Source/Runtime/Engine/Classes/Particles/Lifetime/ParticleModuleLifetimeBase.h
+++ b/Engine/Source/Runtime/Engine/Classes/Particles/Lifetime/ParticleModuleLifetimeBase.h
@@ -4,7 +4,7 @@
 #pragma once
 #include "ParticleModuleLifetimeBase.generated.h"
 
-UCLASS(HeaderGroup=Particle, editinlinenew, hidecategories=Object, abstract, meta=(DisplayName = "Lifetime"))
+UCLASS(editinlinenew, hidecategories=Object, abstract, meta=(DisplayName = "Lifetime"))
 class UParticleModuleLifetimeBase : public UParticleModule
 {
 	GENERATED_UCLASS_BODY()
diff --git a/Engine/Source/Runtime/Engine/Classes/Particles/Lifetime/ParticleModuleLifetime_Seeded.h b/Engine/Source/Runtime/Engine/Classes/Particles/Lifetime/ParticleModuleLifetime_Seeded.h
--- a/Engine/Source/Runtime/Engine/Classes/Particles/Lifetime/ParticleModuleLifetime_Seeded.h
+++ b/Engine/Source/Runtime/Engine/Classes/Particles/Lifetime/ParticleModuleLifetime_Seeded.h
@@ -4,7 +4,7 @@
 #pragma once
 #include "ParticleModuleLifetime_Seeded.generated.h"
 
-UCLASS(HeaderGroup=Particle, editinlinenew, hidecategories=Object, meta=(DisplayName = "Lifetime (Seed)"))
+UCLASS(editinlinenew, hidecategories=Object, meta=(DisplayName = "Lifetime (Seed)"))
 class UParticleModuleLifetime_Seeded : public UParticleModuleLifetime
 {
 	GENERATED_UCLASS_BODY()
diff --git a/Engine/Source/Runtime/Engine/Classes/Particles/Light/ParticleModuleLight.h b/Engine/Source/Runtime/Engine/Classes/Particles/Light/ParticleModuleLight.h
--- a/Engine/Source/Runtime/Engine/Classes/Particles/Light/ParticleModuleLight.h
+++ b/Engine/Source/Runtime/Engine/Classes/Particles/Light/ParticleModuleLight.h
@@ -4,7 +4,7 @@
 #pragma once
 #include "ParticleModuleLight.generated.h"
 
-UCLASS(HeaderGroup=Particle, editinlinenew, hidecategories=Object, MinimalAPI, meta=(DisplayName = "Light"))
+UCLASS(editinlinenew, hidecategories=Object, MinimalAPI, meta=(DisplayName = "Light"))
 class UParticleModuleLight : public UParticleModuleLightBase
 {
 	GENERATED_UCLASS_BODY()
diff --git a/Engine/Source/Runtime/Engine/Classes/Particles/Light/ParticleModuleLightBase.h b/Engine/Source/Runtime/Engine/Classes/Particles/Light/ParticleModuleLightBase.h
--- a/Engine/Source/Runtime/Engine/Classes/Particles/Light/ParticleModuleLightBase.h
+++ b/Engine/Source/Runtime/Engine/Classes/Particles/Light/ParticleModuleLightBase.h
@@ -4,7 +4,7 @@
 #pragma once
 #include "ParticleModuleLightBase.generated.h"
 
-UCLASS(HeaderGroup=Particle, editinlinenew, hidecategories=Object, abstract, meta=(DisplayName = "Light"))
+UCLASS(editinlinenew, hidecategories=Object, abstract, meta=(DisplayName = "Light"))
 class UParticleModuleLightBase : public UParticleModule
 {
 	GENERATED_UCLASS_BODY()
diff --git a/Engine/Source/Runtime/Engine/Classes/Particles/Light/ParticleModuleLight_Seeded.h b/Engine/Source/Runtime/Engine/Classes/Particles/Light/ParticleModuleLight_Seeded.h
--- a/Engine/Source/Runtime/Engine/Classes/Particles/Light/ParticleModuleLight_Seeded.h
+++ b/Engine/Source/Runtime/Engine/Classes/Particles/Light/ParticleModuleLight_Seeded.h
@@ -4,7 +4,7 @@
 #pragma once
 #include "ParticleModuleLight_Seeded.generated.h"
 
-UCLASS(HeaderGroup=Particle, editinlinenew, hidecategories=Object, MinimalAPI, meta=(DisplayName = "Light (Seed)"))
+UCLASS(editinlinenew, hidecategories=Object, MinimalAPI, meta=(DisplayName = "Light (Seed)"))
 class UParticleModuleLight_Seeded : public UParticleModuleLight
 {
 	GENERATED_UCLASS_BODY()
diff --git a/Engine/Source/Runtime/Engine/Classes/Particles/Location/ParticleModuleLocation.h b/Engine/Source/Runtime/Engine/Classes/Particles/Location/ParticleModuleLocation.h
--- a/Engine/Source/Runtime/Engine/Classes/Particles/Location/ParticleModuleLocation.h
+++ b/Engine/Source/Runtime/Engine/Classes/Particles/Location/ParticleModuleLocation.h
@@ -4,7 +4,7 @@
 #pragma once
 #include "ParticleModuleLocation.generated.h"
 
-UCLASS(HeaderGroup=Particle, editinlinenew, hidecategories=Object, MinimalAPI, meta=(DisplayName = "Initial Location"))
+UCLASS(editinlinenew, hidecategories=Object, MinimalAPI, meta=(DisplayName = "Initial Location"))
 class UParticleModuleLocation : public UParticleModuleLocationBase
 {
 	GENERATED_UCLASS_BODY()
diff --git a/Engine/Source/Runtime/Engine/Classes/Particles/Location/ParticleModuleLocationBase.h b/Engine/Source/Runtime/Engine/Classes/Particles/Location/ParticleModuleLocationBase.h
--- a/Engine/Source/Runtime/Engine/Classes/Particles/Location/ParticleModuleLocationBase.h
+++ b/Engine/Source/Runtime/Engine/Classes/Particles/Location/ParticleModuleLocationBase.h
@@ -4,7 +4,7 @@
 #pragma once
 #include "ParticleModuleLocationBase.generated.h"
 
-UCLASS(HeaderGroup=Particle, editinlinenew, hidecategories=Object, abstract, meta=(DisplayName = "Location"))
+UCLASS(editinlinenew, hidecategories=Object, abstract, meta=(DisplayName = "Location"))
 class UParticleModuleLocationBase : public UParticleModule
 {
 	GENERATED_UCLASS_BODY()
diff --git a/Engine/Source/Runtime/Engine/Classes/Particles/Location/ParticleModuleLocationBoneSocket.h b/Engine/Source/Runtime/Engine/Classes/Particles/Location/ParticleModuleLocationBoneSocket.h
--- a/Engine/Source/Runtime/Engine/Classes/Particles/Location/ParticleModuleLocationBoneSocket.h
+++ b/Engine/Source/Runtime/Engine/Classes/Particles/Location/ParticleModuleLocationBoneSocket.h
@@ -41,7 +41,7 @@ struct FLocationBoneSocketInfo
 
 };
 
-UCLASS(HeaderGroup=Particle, editinlinenew, hidecategories=Object, meta=(DisplayName = "Bone/Socket Location"))
+UCLASS(editinlinenew, hidecategories=Object, meta=(DisplayName = "Bone/Socket Location"))
 class UParticleModuleLocationBoneSocket : public UParticleModuleLocationBase
 {
 	GENERATED_UCLASS_BODY()
diff --git a/Engine/Source/Runtime/Engine/Classes/Particles/Location/ParticleModuleLocationDirect.h b/Engine/Source/Runtime/Engine/Classes/Particles/Location/ParticleModuleLocationDirect.h
--- a/Engine/Source/Runtime/Engine/Classes/Particles/Location/ParticleModuleLocationDirect.h
+++ b/Engine/Source/Runtime/Engine/Classes/Particles/Location/ParticleModuleLocationDirect.h
@@ -10,7 +10,7 @@
 #pragma once
 #include "ParticleModuleLocationDirect.generated.h"
 
-UCLASS(HeaderGroup=Particle, editinlinenew, hidecategories=Object, MinimalAPI, meta=(DisplayName = "Direct Location"))
+UCLASS(editinlinenew, hidecategories=Object, MinimalAPI, meta=(DisplayName = "Direct Location"))
 class UParticleModuleLocationDirect : public UParticleModuleLocationBase
 {
 	GENERATED_UCLASS_BODY()
diff --git a/Engine/Source/Runtime/Engine/Classes/Particles/Location/ParticleModuleLocationEmitter.h b/Engine/Source/Runtime/Engine/Classes/Particles/Location/ParticleModuleLocationEmitter.h
--- a/Engine/Source/Runtime/Engine/Classes/Particles/Location/ParticleModuleLocationEmitter.h
+++ b/Engine/Source/Runtime/Engine/Classes/Particles/Location/ParticleModuleLocationEmitter.h
@@ -19,7 +19,7 @@ enum ELocationEmitterSelectionMethod
 	ELESM_MAX,
 };
 
-UCLASS(HeaderGroup=Particle, editinlinenew, hidecategories=Object, MinimalAPI, meta=(DisplayName = "Emitter Initial Location"))
+UCLASS(editinlinenew, hidecategories=Object, MinimalAPI, meta=(DisplayName = "Emitter Initial Location"))
 class UParticleModuleLocationEmitter : public UParticleModuleLocationBase
 {
 	GENERATED_UCLASS_BODY()
diff --git a/Engine/Source/Runtime/Engine/Classes/Particles/Location/ParticleModuleLocationEmitterDirect.h b/Engine/Source/Runtime/Engine/Classes/Particles/Location/ParticleModuleLocationEmitterDirect.h
--- a/Engine/Source/Runtime/Engine/Classes/Particles/Location/ParticleModuleLocationEmitterDirect.h
+++ b/Engine/Source/Runtime/Engine/Classes/Particles/Location/ParticleModuleLocationEmitterDirect.h
@@ -11,7 +11,7 @@
 #pragma once
 #include "ParticleModuleLocationEmitterDirect.generated.h"
 
-UCLASS(HeaderGroup=Particle, editinlinenew, hidecategories=Object, meta=(DisplayName = "Emitter Direct Location"))
+UCLASS(editinlinenew, hidecategories=Object, meta=(DisplayName = "Emitter Direct Location"))
 class UParticleModuleLocationEmitterDirect : public UParticleModuleLocationBase
 {
 	GENERATED_UCLASS_BODY()
diff --git a/Engine/Source/Runtime/Engine/Classes/Particles/Location/ParticleModuleLocationPrimitiveBase.h b/Engine/Source/Runtime/Engine/Classes/Particles/Location/ParticleModuleLocationPrimitiveBase.h
--- a/Engine/Source/Runtime/Engine/Classes/Particles/Location/ParticleModuleLocationPrimitiveBase.h
+++ b/Engine/Source/Runtime/Engine/Classes/Particles/Location/ParticleModuleLocationPrimitiveBase.h
@@ -8,7 +8,7 @@
 #pragma once
 #include "ParticleModuleLocationPrimitiveBase.generated.h"
 
-UCLASS(HeaderGroup=Particle, editinlinenew, hidecategories=Object)
+UCLASS(editinlinenew, hidecategories=Object)
 class UParticleModuleLocationPrimitiveBase : public UParticleModuleLocationBase
 {
 	GENERATED_UCLASS_BODY()
diff --git a/Engine/Source/Runtime/Engine/Classes/Particles/Location/ParticleModuleLocationPrimitiveCylinder.h b/Engine/Source/Runtime/Engine/Classes/Particles/Location/ParticleModuleLocationPrimitiveCylinder.h
--- a/Engine/Source/Runtime/Engine/Classes/Particles/Location/ParticleModuleLocationPrimitiveCylinder.h
+++ b/Engine/Source/Runtime/Engine/Classes/Particles/Location/ParticleModuleLocationPrimitiveCylinder.h
@@ -17,7 +17,7 @@ enum CylinderHeightAxis
 	PMLPC_HEIGHTAXIS_MAX,
 };
 
-UCLASS(HeaderGroup=Particle, editinlinenew, hidecategories=Object, MinimalAPI)
+UCLASS(editinlinenew, hidecategories=Object, MinimalAPI)
 class UParticleModuleLocationPrimitiveCylinder : public UParticleModuleLocationPrimitiveBase
 {
 	GENERATED_UCLASS_BODY()
diff --git a/Engine/Source/Runtime/Engine/Classes/Particles/Location/ParticleModuleLocationPrimitiveCylinder_Seeded.h b/Engine/Source/Runtime/Engine/Classes/Particles/Location/ParticleModuleLocationPrimitiveCylinder_Seeded.h
--- a/Engine/Source/Runtime/Engine/Classes/Particles/Location/ParticleModuleLocationPrimitiveCylinder_Seeded.h
+++ b/Engine/Source/Runtime/Engine/Classes/Particles/Location/ParticleModuleLocationPrimitiveCylinder_Seeded.h
@@ -4,7 +4,7 @@
 #pragma once
 #include "ParticleModuleLocationPrimitiveCylinder_Seeded.generated.h"
 
-UCLASS(HeaderGroup=Particle, editinlinenew, hidecategories=Object, meta=(DisplayName = "Cylinder (Seed)"))
+UCLASS(editinlinenew, hidecategories=Object, meta=(DisplayName = "Cylinder (Seed)"))
 class UParticleModuleLocationPrimitiveCylinder_Seeded : public UParticleModuleLocationPrimitiveCylinder
 {
 	GENERATED_UCLASS_BODY()
diff --git a/Engine/Source/Runtime/Engine/Classes/Particles/Location/ParticleModuleLocationPrimitiveSphere.h b/Engine/Source/Runtime/Engine/Classes/Particles/Location/ParticleModuleLocationPrimitiveSphere.h
--- a/Engine/Source/Runtime/Engine/Classes/Particles/Location/ParticleModuleLocationPrimitiveSphere.h
+++ b/Engine/Source/Runtime/Engine/Classes/Particles/Location/ParticleModuleLocationPrimitiveSphere.h
@@ -8,7 +8,7 @@
 #pragma once
 #include "ParticleModuleLocationPrimitiveSphere.generated.h"
 
-UCLASS(HeaderGroup=Particle, editinlinenew, hidecategories=Object, MinimalAPI, meta=(DisplayName = "Sphere"))
+UCLASS(editinlinenew, hidecategories=Object, MinimalAPI, meta=(DisplayName = "Sphere"))
 class UParticleModuleLocationPrimitiveSphere : public UParticleModuleLocationPrimitiveBase
 {
 	GENERATED_UCLASS_BODY()
diff --git a/Engine/Source/Runtime/Engine/Classes/Particles/Location/ParticleModuleLocationPrimitiveSphere_Seeded.h b/Engine/Source/Runtime/Engine/Classes/Particles/Location/ParticleModuleLocationPrimitiveSphere_Seeded.h
--- a/Engine/Source/Runtime/Engine/Classes/Particles/Location/ParticleModuleLocationPrimitiveSphere_Seeded.h
+++ b/Engine/Source/Runtime/Engine/Classes/Particles/Location/ParticleModuleLocationPrimitiveSphere_Seeded.h
@@ -4,7 +4,7 @@
 #pragma once
 #include "ParticleModuleLocationPrimitiveSphere_Seeded.generated.h"
 
-UCLASS(HeaderGroup=Particle, editinlinenew, hidecategories=Object, meta=(DisplayName = "Sphere (Seed)"))
+UCLASS(editinlinenew, hidecategories=Object, meta=(DisplayName = "Sphere (Seed)"))
 class UParticleModuleLocationPrimitiveSphere_Seeded : public UParticleModuleLocationPrimitiveSphere
 {
 	GENERATED_UCLASS_BODY()
diff --git a/Engine/Source/Runtime/Engine/Classes/Particles/Location/ParticleModuleLocationPrimitiveTriangle.h b/Engine/Source/Runtime/Engine/Classes/Particles/Location/ParticleModuleLocationPrimitiveTriangle.h
--- a/Engine/Source/Runtime/Engine/Classes/Particles/Location/ParticleModuleLocationPrimitiveTriangle.h
+++ b/Engine/Source/Runtime/Engine/Classes/Particles/Location/ParticleModuleLocationPrimitiveTriangle.h
@@ -8,7 +8,7 @@
 #pragma once
 #include "ParticleModuleLocationPrimitiveTriangle.generated.h"
 
-UCLASS(HeaderGroup=Particle, editinlinenew, hidecategories=Object, MinimalAPI, meta=(DisplayName = "Triangle"))
+UCLASS(editinlinenew, hidecategories=Object, MinimalAPI, meta=(DisplayName = "Triangle"))
 class UParticleModuleLocationPrimitiveTriangle : public UParticleModuleLocationBase
 {
 	GENERATED_UCLASS_BODY()
diff --git a/Engine/Source/Runtime/Engine/Classes/Particles/Location/ParticleModuleLocationSkelVertSurface.h b/Engine/Source/Runtime/Engine/Classes/Particles/Location/ParticleModuleLocationSkelVertSurface.h
--- a/Engine/Source/Runtime/Engine/Classes/Particles/Location/ParticleModuleLocationSkelVertSurface.h
+++ b/Engine/Source/Runtime/Engine/Classes/Particles/Location/ParticleModuleLocationSkelVertSurface.h
@@ -12,7 +12,7 @@ enum ELocationSkelVertSurfaceSource
 	VERTSURFACESOURCE_MAX,
 };
 
-UCLASS(HeaderGroup=Particle, editinlinenew, hidecategories=Object, meta=(DisplayName = "Skel Vert/Surf Location"))
+UCLASS(editinlinenew, hidecategories=Object, meta=(DisplayName = "Skel Vert/Surf Location"))
 class UParticleModuleLocationSkelVertSurface : public UParticleModuleLocationBase
 {
 	GENERATED_UCLASS_BODY()
diff --git a/Engine/Source/Runtime/Engine/Classes/Particles/Location/ParticleModuleLocationWorldOffset.h b/Engine/Source/Runtime/Engine/Classes/Particles/Location/ParticleModuleLocationWorldOffset.h
--- a/Engine/Source/Runtime/Engine/Classes/Particles/Location/ParticleModuleLocationWorldOffset.h
+++ b/Engine/Source/Runtime/Engine/Classes/Particles/Location/ParticleModuleLocationWorldOffset.h
@@ -4,7 +4,7 @@
 #pragma once
 #include "ParticleModuleLocationWorldOffset.generated.h"
 
-UCLASS(HeaderGroup=Particle, editinlinenew, meta=(DisplayName = "World Offset"))
+UCLASS(editinlinenew, meta=(DisplayName = "World Offset"))
 class UParticleModuleLocationWorldOffset : public UParticleModuleLocation
 {
 	GENERATED_UCLASS_BODY()
diff --git a/Engine/Source/Runtime/Engine/Classes/Particles/Location/ParticleModuleLocationWorldOffset_Seeded.h b/Engine/Source/Runtime/Engine/Classes/Particles/Location/ParticleModuleLocationWorldOffset_Seeded.h
--- a/Engine/Source/Runtime/Engine/Classes/Particles/Location/ParticleModuleLocationWorldOffset_Seeded.h
+++ b/Engine/Source/Runtime/Engine/Classes/Particles/Location/ParticleModuleLocationWorldOffset_Seeded.h
@@ -4,7 +4,7 @@
 #pragma once
 #include "ParticleModuleLocationWorldOffset_Seeded.generated.h"
 
-UCLASS(HeaderGroup=Particle, editinlinenew, meta=(DisplayName = "World Offset (Seed)"))
+UCLASS(editinlinenew, meta=(DisplayName = "World Offset (Seed)"))
 class UParticleModuleLocationWorldOffset_Seeded : public UParticleModuleLocationWorldOffset
 {
 	GENERATED_UCLASS_BODY()
diff --git a/Engine/Source/Runtime/Engine/Classes/Particles/Location/ParticleModuleLocation_Seeded.h b/Engine/Source/Runtime/Engine/Classes/Particles/Location/ParticleModuleLocation_Seeded.h
--- a/Engine/Source/Runtime/Engine/Classes/Particles/Location/ParticleModuleLocation_Seeded.h
+++ b/Engine/Source/Runtime/Engine/Classes/Particles/Location/ParticleModuleLocation_Seeded.h
@@ -4,7 +4,7 @@
 #pragma once
 #include "ParticleModuleLocation_Seeded.generated.h"
 
-UCLASS(HeaderGroup=Particle, editinlinenew, hidecategories=Object, meta=(DisplayName = "Initial Location (Seed)"))
+UCLASS(editinlinenew, hidecategories=Object, meta=(DisplayName = "Initial Location (Seed)"))
 class UParticleModuleLocation_Seeded : public UParticleModuleLocation
 {
 	GENERATED_UCLASS_BODY()
diff --git a/Engine/Source/Runtime/Engine/Classes/Particles/Location/ParticleModuleSourceMovement.h b/Engine/Source/Runtime/Engine/Classes/Particles/Location/ParticleModuleSourceMovement.h
--- a/Engine/Source/Runtime/Engine/Classes/Particles/Location/ParticleModuleSourceMovement.h
+++ b/Engine/Source/Runtime/Engine/Classes/Particles/Location/ParticleModuleSourceMovement.h
@@ -4,7 +4,7 @@
 #pragma once
 #include "ParticleModuleSourceMovement.generated.h"
 
-UCLASS(HeaderGroup=Particle, editinlinenew, hidecategories=Object, meta=(DisplayName = "Source Movement"))
+UCLASS(editinlinenew, hidecategories=Object, meta=(DisplayName = "Source Movement"))
 class UParticleModuleSourceMovement : public UParticleModuleLocationBase
 {
 	GENERATED_UCLASS_BODY()
diff --git a/Engine/Source/Runtime/Engine/Classes/Particles/Material/ParticleModuleMaterialBase.h b/Engine/Source/Runtime/Engine/Classes/Particles/Material/ParticleModuleMaterialBase.h
--- a/Engine/Source/Runtime/Engine/Classes/Particles/Material/ParticleModuleMaterialBase.h
+++ b/Engine/Source/Runtime/Engine/Classes/Particles/Material/ParticleModuleMaterialBase.h
@@ -4,7 +4,7 @@
 #pragma once
 #include "ParticleModuleMaterialBase.generated.h"
 
-UCLASS(HeaderGroup=Particle, editinlinenew, hidecategories=Object, abstract, meta=(DisplayName = "Material"))
+UCLASS(editinlinenew, hidecategories=Object, abstract, meta=(DisplayName = "Material"))
 class UParticleModuleMaterialBase : public UParticleModule
 {
 	GENERATED_UCLASS_BODY()
diff --git a/Engine/Source/Runtime/Engine/Classes/Particles/Material/ParticleModuleMeshMaterial.h b/Engine/Source/Runtime/Engine/Classes/Particles/Material/ParticleModuleMeshMaterial.h
--- a/Engine/Source/Runtime/Engine/Classes/Particles/Material/ParticleModuleMeshMaterial.h
+++ b/Engine/Source/Runtime/Engine/Classes/Particles/Material/ParticleModuleMeshMaterial.h
@@ -4,7 +4,7 @@
 #pragma once
 #include "ParticleModuleMeshMaterial.generated.h"
 
-UCLASS(HeaderGroup=Particle, editinlinenew, hidecategories=Object, MinimalAPI, meta=(DisplayName = "Mesh Material"))
+UCLASS(editinlinenew, hidecategories=Object, MinimalAPI, meta=(DisplayName = "Mesh Material"))
 class UParticleModuleMeshMaterial : public UParticleModuleMaterialBase
 {
 	GENERATED_UCLASS_BODY()
diff --git a/Engine/Source/Runtime/Engine/Classes/Particles/Modules/Location/ParticleModulePivotOffset.h b/Engine/Source/Runtime/Engine/Classes/Particles/Modules/Location/ParticleModulePivotOffset.h
--- a/Engine/Source/Runtime/Engine/Classes/Particles/Modules/Location/ParticleModulePivotOffset.h
+++ b/Engine/Source/Runtime/Engine/Classes/Particles/Modules/Location/ParticleModulePivotOffset.h
@@ -3,7 +3,7 @@
 #pragma once
 #include "ParticleModulePivotOffset.generated.h"
 
-UCLASS(HeaderGroup=Particle, editinlinenew, hidecategories=Object, MinimalAPI, DisplayName="Pivot Offset")
+UCLASS(editinlinenew, hidecategories=Object, MinimalAPI, DisplayName="Pivot Offset")
 class UParticleModulePivotOffset : public UParticleModuleLocationBase
 {
 	GENERATED_UCLASS_BODY()
diff --git a/Engine/Source/Runtime/Engine/Classes/Particles/Orbit/ParticleModuleOrbit.h b/Engine/Source/Runtime/Engine/Classes/Particles/Orbit/ParticleModuleOrbit.h
--- a/Engine/Source/Runtime/Engine/Classes/Particles/Orbit/ParticleModuleOrbit.h
+++ b/Engine/Source/Runtime/Engine/Classes/Particles/Orbit/ParticleModuleOrbit.h
@@ -53,7 +53,7 @@ struct FOrbitOptions
 	
 };
 
-UCLASS(HeaderGroup=Particle, editinlinenew, hidecategories=(Object, Orbit), MinimalAPI, meta=(DisplayName = "Orbit"))
+UCLASS(editinlinenew, hidecategories=(Object, Orbit), MinimalAPI, meta=(DisplayName = "Orbit"))
 class UParticleModuleOrbit : public UParticleModuleOrbitBase
 {
 	GENERATED_UCLASS_BODY()
diff --git a/Engine/Source/Runtime/Engine/Classes/Particles/Orbit/ParticleModuleOrbitBase.h b/Engine/Source/Runtime/Engine/Classes/Particles/Orbit/ParticleModuleOrbitBase.h
--- a/Engine/Source/Runtime/Engine/Classes/Particles/Orbit/ParticleModuleOrbitBase.h
+++ b/Engine/Source/Runtime/Engine/Classes/Particles/Orbit/ParticleModuleOrbitBase.h
@@ -4,7 +4,7 @@
 #pragma once
 #include "ParticleModuleOrbitBase.generated.h"
 
-UCLASS(HeaderGroup=Particle, editinlinenew, hidecategories=Object, abstract, meta=(DisplayName = "Orbit"))
+UCLASS(editinlinenew, hidecategories=Object, abstract, meta=(DisplayName = "Orbit"))
 class UParticleModuleOrbitBase : public UParticleModule
 {
 	GENERATED_UCLASS_BODY()
diff --git a/Engine/Source/Runtime/Engine/Classes/Particles/Orientation/ParticleModuleOrientationAxisLock.h b/Engine/Source/Runtime/Engine/Classes/Particles/Orientation/ParticleModuleOrientationAxisLock.h
--- a/Engine/Source/Runtime/Engine/Classes/Particles/Orientation/ParticleModuleOrientationAxisLock.h
+++ b/Engine/Source/Runtime/Engine/Classes/Particles/Orientation/ParticleModuleOrientationAxisLock.h
@@ -31,7 +31,7 @@ enum EParticleAxisLock
 	EPAL_MAX,
 };
 
-UCLASS(HeaderGroup=Particle, editinlinenew, hidecategories=Object, MinimalAPI, meta=(DisplayName = "Lock Axis"))
+UCLASS(editinlinenew, hidecategories=Object, MinimalAPI, meta=(DisplayName = "Lock Axis"))
 class UParticleModuleOrientationAxisLock : public UParticleModuleOrientationBase
 {
 	GENERATED_UCLASS_BODY()
diff --git a/Engine/Source/Runtime/Engine/Classes/Particles/Orientation/ParticleModuleOrientationBase.h b/Engine/Source/Runtime/Engine/Classes/Particles/Orientation/ParticleModuleOrientationBase.h
--- a/Engine/Source/Runtime/Engine/Classes/Particles/Orientation/ParticleModuleOrientationBase.h
+++ b/Engine/Source/Runtime/Engine/Classes/Particles/Orientation/ParticleModuleOrientationBase.h
@@ -4,7 +4,7 @@
 #pragma once
 #include "ParticleModuleOrientationBase.generated.h"
 
-UCLASS(HeaderGroup=Particle, editinlinenew, hidecategories=Object, abstract, meta=(DisplayName = "Orientation"))
+UCLASS(editinlinenew, hidecategories=Object, abstract, meta=(DisplayName = "Orientation"))
 class UParticleModuleOrientationBase : public UParticleModule
 {
 	GENERATED_UCLASS_BODY()
diff --git a/Engine/Source/Runtime/Engine/Classes/Particles/Parameter/ParticleModuleParameterBase.h b/Engine/Source/Runtime/Engine/Classes/Particles/Parameter/ParticleModuleParameterBase.h
--- a/Engine/Source/Runtime/Engine/Classes/Particles/Parameter/ParticleModuleParameterBase.h
+++ b/Engine/Source/Runtime/Engine/Classes/Particles/Parameter/ParticleModuleParameterBase.h
@@ -4,7 +4,7 @@
 #pragma once
 #include "ParticleModuleParameterBase.generated.h"
 
-UCLASS(HeaderGroup=Particle, editinlinenew, hidecategories=Object, abstract, meta=(DisplayName = "Parameter"))
+UCLASS(editinlinenew, hidecategories=Object, abstract, meta=(DisplayName = "Parameter"))
 class UParticleModuleParameterBase : public UParticleModule
 {
 	GENERATED_UCLASS_BODY()
diff --git a/Engine/Source/Runtime/Engine/Classes/Particles/Parameter/ParticleModuleParameterDynamic.h b/Engine/Source/Runtime/Engine/Classes/Particles/Parameter/ParticleModuleParameterDynamic.h
--- a/Engine/Source/Runtime/Engine/Classes/Particles/Parameter/ParticleModuleParameterDynamic.h
+++ b/Engine/Source/Runtime/Engine/Classes/Particles/Parameter/ParticleModuleParameterDynamic.h
@@ -78,7 +78,7 @@ struct FEmitterDynamicParameter
 
 };
 
-UCLASS(HeaderGroup=Particle, editinlinenew, hidecategories=Object, MinimalAPI, meta=(DisplayName = "Dynamic"))
+UCLASS(editinlinenew, hidecategories=Object, MinimalAPI, meta=(DisplayName = "Dynamic"))
 class UParticleModuleParameterDynamic : public UParticleModuleParameterBase
 {
 	GENERATED_UCLASS_BODY()
