#ttp 303393 UObject improvement: Remove header groups

diff --git a/Engine/Source/Runtime/Engine/Classes/Particles/Parameter/ParticleModuleParameterDynamic_Seeded.h b/Engine/Source/Runtime/Engine/Classes/Particles/Parameter/ParticleModuleParameterDynamic_Seeded.h
--- a/Engine/Source/Runtime/Engine/Classes/Particles/Parameter/ParticleModuleParameterDynamic_Seeded.h
+++ b/Engine/Source/Runtime/Engine/Classes/Particles/Parameter/ParticleModuleParameterDynamic_Seeded.h
@@ -4,7 +4,7 @@
 #pragma once
 #include "ParticleModuleParameterDynamic_Seeded.generated.h"
 
-UCLASS(HeaderGroup=Particle, editinlinenew, hidecategories=Object, meta=(DisplayName = "Dynamic (Seed)"))
+UCLASS(editinlinenew, hidecategories=Object, meta=(DisplayName = "Dynamic (Seed)"))
 class UParticleModuleParameterDynamic_Seeded : public UParticleModuleParameterDynamic
 {
 	GENERATED_UCLASS_BODY()
diff --git a/Engine/Source/Runtime/Engine/Classes/Particles/ParticleEmitter.h b/Engine/Source/Runtime/Engine/Classes/Particles/ParticleEmitter.h
--- a/Engine/Source/Runtime/Engine/Classes/Particles/ParticleEmitter.h
+++ b/Engine/Source/Runtime/Engine/Classes/Particles/ParticleEmitter.h
@@ -76,7 +76,7 @@ struct FParticleBurst
 	
 };
 
-UCLASS(HeaderGroup=Particle, dependson=UParticleLODLevel, hidecategories=Object, editinlinenew, abstract, MinimalAPI)
+UCLASS(dependson=UParticleLODLevel, hidecategories=Object, editinlinenew, abstract, MinimalAPI)
 class UParticleEmitter : public UObject
 {
 	GENERATED_UCLASS_BODY()
diff --git a/Engine/Source/Runtime/Engine/Classes/Particles/ParticleEventManager.h b/Engine/Source/Runtime/Engine/Classes/Particles/ParticleEventManager.h
--- a/Engine/Source/Runtime/Engine/Classes/Particles/ParticleEventManager.h
+++ b/Engine/Source/Runtime/Engine/Classes/Particles/ParticleEventManager.h
@@ -4,7 +4,7 @@
 #pragma once
 #include "ParticleEventManager.generated.h"
 
-UCLASS(HeaderGroup=Particle, dependson=UParticleSystemComponent, config=Game, NotBlueprintable, notplaceable)
+UCLASS(dependson=UParticleSystemComponent, config=Game, NotBlueprintable, notplaceable)
 class AParticleEventManager : public AActor
 {
 	GENERATED_UCLASS_BODY()
diff --git a/Engine/Source/Runtime/Engine/Classes/Particles/ParticleModuleRequired.h b/Engine/Source/Runtime/Engine/Classes/Particles/ParticleModuleRequired.h
--- a/Engine/Source/Runtime/Engine/Classes/Particles/ParticleModuleRequired.h
+++ b/Engine/Source/Runtime/Engine/Classes/Particles/ParticleModuleRequired.h
@@ -27,7 +27,7 @@ enum EEmitterNormalsMode
 	ENM_MAX,
 };
 
-UCLASS(HeaderGroup=Particle, editinlinenew, hidecategories=(Object, Cascade), meta=(DisplayName = "Required"), MinimalAPI)
+UCLASS(editinlinenew, hidecategories=(Object, Cascade), meta=(DisplayName = "Required"), MinimalAPI)
 class UParticleModuleRequired : public UParticleModule
 {
 	GENERATED_UCLASS_BODY()
diff --git a/Engine/Source/Runtime/Engine/Classes/Particles/ParticleSystem.h b/Engine/Source/Runtime/Engine/Classes/Particles/ParticleSystem.h
--- a/Engine/Source/Runtime/Engine/Classes/Particles/ParticleSystem.h
+++ b/Engine/Source/Runtime/Engine/Classes/Particles/ParticleSystem.h
@@ -79,7 +79,7 @@ struct FLODSoloTrack
  * in a system, the designer can create elaborate particle effects that are held in a single system. Once created using
  * Cascade, a ParticleSystem can then be inserted into a level or created in script.
  */
-UCLASS(HeaderGroup=Particle, hidecategories=Object, MinimalAPI, BlueprintType)
+UCLASS(hidecategories=Object, MinimalAPI, BlueprintType)
 class UParticleSystem : public UObject
 {
 	GENERATED_UCLASS_BODY()
diff --git a/Engine/Source/Runtime/Engine/Classes/Particles/ParticleSystemComponent.h b/Engine/Source/Runtime/Engine/Classes/Particles/ParticleSystemComponent.h
--- a/Engine/Source/Runtime/Engine/Classes/Particles/ParticleSystemComponent.h
+++ b/Engine/Source/Runtime/Engine/Classes/Particles/ParticleSystemComponent.h
@@ -5,6 +5,17 @@
 #include "ParticleHelper.h"
 #include "ParticleSystemComponent.generated.h"
 
+//
+// Forward declarations.
+//
+class UAnimNotifyState;
+class FParticleDynamicData;
+class FParticleSystemSceneProxy;
+class UAnimNotifyState;
+struct FDynamicEmitterDataBase;
+struct FDynamicEmitterReplayDataBase;
+struct FParticleAnimTrailEmitterInstance;
+
 /** Enum for specifying type of a name instance parameter. */
 UENUM()
 enum EParticleSysParamType
@@ -233,7 +244,7 @@ struct FParticleEventKismetData : public FParticleEventData
 /** 
  * A particle emmitter.
  */
-UCLASS(HeaderGroup=Particle, ClassGroup=Rendering, hidecategories=Object, hidecategories=Physics, hidecategories=Collision, showcategories=Trigger, editinlinenew, dependson=(AEmitter, UParticleSystem), meta=(BlueprintSpawnableComponent))
+UCLASS(ClassGroup=Rendering, hidecategories=Object, hidecategories=Physics, hidecategories=Collision, showcategories=Trigger, editinlinenew, dependson=(AEmitter, UParticleSystem), meta=(BlueprintSpawnableComponent))
 class ENGINE_API UParticleSystemComponent : public UPrimitiveComponent
 {
 	GENERATED_UCLASS_BODY()
diff --git a/Engine/Source/Runtime/Engine/Classes/Particles/ParticleSystemReplay.h b/Engine/Source/Runtime/Engine/Classes/Particles/ParticleSystemReplay.h
--- a/Engine/Source/Runtime/Engine/Classes/Particles/ParticleSystemReplay.h
+++ b/Engine/Source/Runtime/Engine/Classes/Particles/ParticleSystemReplay.h
@@ -58,7 +58,7 @@ struct FParticleSystemReplayFrame
 	
 };
 
-UCLASS(HeaderGroup=Particle, hidecategories=Object, AutoExpandCategories=ParticleSystemReplay)
+UCLASS(hidecategories=Object, AutoExpandCategories=ParticleSystemReplay)
 class UParticleSystemReplay : public UObject
 {
 	GENERATED_UCLASS_BODY()
diff --git a/Engine/Source/Runtime/Engine/Classes/Particles/Rotation/ParticleModuleMeshRotation.h b/Engine/Source/Runtime/Engine/Classes/Particles/Rotation/ParticleModuleMeshRotation.h
--- a/Engine/Source/Runtime/Engine/Classes/Particles/Rotation/ParticleModuleMeshRotation.h
+++ b/Engine/Source/Runtime/Engine/Classes/Particles/Rotation/ParticleModuleMeshRotation.h
@@ -4,7 +4,7 @@
 #pragma once
 #include "ParticleModuleMeshRotation.generated.h"
 
-UCLASS(HeaderGroup=Particle, editinlinenew, hidecategories=Object, MinimalAPI, meta=(DisplayName = "Init Mesh Rotation"))
+UCLASS(editinlinenew, hidecategories=Object, MinimalAPI, meta=(DisplayName = "Init Mesh Rotation"))
 class UParticleModuleMeshRotation : public UParticleModuleRotationBase
 {
 	GENERATED_UCLASS_BODY()
diff --git a/Engine/Source/Runtime/Engine/Classes/Particles/Rotation/ParticleModuleMeshRotation_Seeded.h b/Engine/Source/Runtime/Engine/Classes/Particles/Rotation/ParticleModuleMeshRotation_Seeded.h
--- a/Engine/Source/Runtime/Engine/Classes/Particles/Rotation/ParticleModuleMeshRotation_Seeded.h
+++ b/Engine/Source/Runtime/Engine/Classes/Particles/Rotation/ParticleModuleMeshRotation_Seeded.h
@@ -4,7 +4,7 @@
 #pragma once
 #include "ParticleModuleMeshRotation_Seeded.generated.h"
 
-UCLASS(HeaderGroup=Particle, editinlinenew, hidecategories=Object, meta=(DisplayName = "Init Mesh Rot (Seed)"))
+UCLASS(editinlinenew, hidecategories=Object, meta=(DisplayName = "Init Mesh Rot (Seed)"))
 class UParticleModuleMeshRotation_Seeded : public UParticleModuleMeshRotation
 {
 	GENERATED_UCLASS_BODY()
diff --git a/Engine/Source/Runtime/Engine/Classes/Particles/Rotation/ParticleModuleRotation.h b/Engine/Source/Runtime/Engine/Classes/Particles/Rotation/ParticleModuleRotation.h
--- a/Engine/Source/Runtime/Engine/Classes/Particles/Rotation/ParticleModuleRotation.h
+++ b/Engine/Source/Runtime/Engine/Classes/Particles/Rotation/ParticleModuleRotation.h
@@ -4,7 +4,7 @@
 #pragma once
 #include "ParticleModuleRotation.generated.h"
 
-UCLASS(HeaderGroup=Particle, editinlinenew, hidecategories=Object, MinimalAPI, meta=(DisplayName = "Initial Rotation"))
+UCLASS(editinlinenew, hidecategories=Object, MinimalAPI, meta=(DisplayName = "Initial Rotation"))
 class UParticleModuleRotation : public UParticleModuleRotationBase
 {
 	GENERATED_UCLASS_BODY()
diff --git a/Engine/Source/Runtime/Engine/Classes/Particles/Rotation/ParticleModuleRotationBase.h b/Engine/Source/Runtime/Engine/Classes/Particles/Rotation/ParticleModuleRotationBase.h
--- a/Engine/Source/Runtime/Engine/Classes/Particles/Rotation/ParticleModuleRotationBase.h
+++ b/Engine/Source/Runtime/Engine/Classes/Particles/Rotation/ParticleModuleRotationBase.h
@@ -4,7 +4,7 @@
 #pragma once
 #include "ParticleModuleRotationBase.generated.h"
 
-UCLASS(HeaderGroup=Particle, editinlinenew, hidecategories=Object, abstract, meta=(DisplayName = "Rotation"))
+UCLASS(editinlinenew, hidecategories=Object, abstract, meta=(DisplayName = "Rotation"))
 class UParticleModuleRotationBase : public UParticleModule
 {
 	GENERATED_UCLASS_BODY()
diff --git a/Engine/Source/Runtime/Engine/Classes/Particles/Rotation/ParticleModuleRotationOverLifetime.h b/Engine/Source/Runtime/Engine/Classes/Particles/Rotation/ParticleModuleRotationOverLifetime.h
--- a/Engine/Source/Runtime/Engine/Classes/Particles/Rotation/ParticleModuleRotationOverLifetime.h
+++ b/Engine/Source/Runtime/Engine/Classes/Particles/Rotation/ParticleModuleRotationOverLifetime.h
@@ -4,7 +4,7 @@
 #pragma once
 #include "ParticleModuleRotationOverLifetime.generated.h"
 
-UCLASS(HeaderGroup=Particle, editinlinenew, hidecategories=Object, MinimalAPI, meta=(DisplayName = "Rotation/Life"))
+UCLASS(editinlinenew, hidecategories=Object, MinimalAPI, meta=(DisplayName = "Rotation/Life"))
 class UParticleModuleRotationOverLifetime : public UParticleModuleRotationBase
 {
 	GENERATED_UCLASS_BODY()
diff --git a/Engine/Source/Runtime/Engine/Classes/Particles/Rotation/ParticleModuleRotation_Seeded.h b/Engine/Source/Runtime/Engine/Classes/Particles/Rotation/ParticleModuleRotation_Seeded.h
--- a/Engine/Source/Runtime/Engine/Classes/Particles/Rotation/ParticleModuleRotation_Seeded.h
+++ b/Engine/Source/Runtime/Engine/Classes/Particles/Rotation/ParticleModuleRotation_Seeded.h
@@ -4,7 +4,7 @@
 #pragma once
 #include "ParticleModuleRotation_Seeded.generated.h"
 
-UCLASS(HeaderGroup=Particle, editinlinenew, hidecategories=Object, meta=(DisplayName = "Init Rotation (Seed)"))
+UCLASS(editinlinenew, hidecategories=Object, meta=(DisplayName = "Init Rotation (Seed)"))
 class UParticleModuleRotation_Seeded : public UParticleModuleRotation
 {
 	GENERATED_UCLASS_BODY()
diff --git a/Engine/Source/Runtime/Engine/Classes/Particles/RotationRate/ParticleModuleMeshRotationRate.h b/Engine/Source/Runtime/Engine/Classes/Particles/RotationRate/ParticleModuleMeshRotationRate.h
--- a/Engine/Source/Runtime/Engine/Classes/Particles/RotationRate/ParticleModuleMeshRotationRate.h
+++ b/Engine/Source/Runtime/Engine/Classes/Particles/RotationRate/ParticleModuleMeshRotationRate.h
@@ -4,7 +4,7 @@
 #pragma once
 #include "ParticleModuleMeshRotationRate.generated.h"
 
-UCLASS(HeaderGroup=Particle, editinlinenew, hidecategories=Object, MinimalAPI, meta=(DisplayName = "Init Mesh Rotation Rate"))
+UCLASS(editinlinenew, hidecategories=Object, MinimalAPI, meta=(DisplayName = "Init Mesh Rotation Rate"))
 class UParticleModuleMeshRotationRate : public UParticleModuleRotationRateBase
 {
 	GENERATED_UCLASS_BODY()
diff --git a/Engine/Source/Runtime/Engine/Classes/Particles/RotationRate/ParticleModuleMeshRotationRateMultiplyLife.h b/Engine/Source/Runtime/Engine/Classes/Particles/RotationRate/ParticleModuleMeshRotationRateMultiplyLife.h
--- a/Engine/Source/Runtime/Engine/Classes/Particles/RotationRate/ParticleModuleMeshRotationRateMultiplyLife.h
+++ b/Engine/Source/Runtime/Engine/Classes/Particles/RotationRate/ParticleModuleMeshRotationRateMultiplyLife.h
@@ -4,7 +4,7 @@
 #pragma once
 #include "ParticleModuleMeshRotationRateMultiplyLife.generated.h"
 
-UCLASS(HeaderGroup=Particle, editinlinenew, hidecategories=Object, meta=(DisplayName = "Mesh Rotation Rate * Life"))
+UCLASS(editinlinenew, hidecategories=Object, meta=(DisplayName = "Mesh Rotation Rate * Life"))
 class UParticleModuleMeshRotationRateMultiplyLife : public UParticleModuleRotationRateBase
 {
 	GENERATED_UCLASS_BODY()
diff --git a/Engine/Source/Runtime/Engine/Classes/Particles/RotationRate/ParticleModuleMeshRotationRateOverLife.h b/Engine/Source/Runtime/Engine/Classes/Particles/RotationRate/ParticleModuleMeshRotationRateOverLife.h
--- a/Engine/Source/Runtime/Engine/Classes/Particles/RotationRate/ParticleModuleMeshRotationRateOverLife.h
+++ b/Engine/Source/Runtime/Engine/Classes/Particles/RotationRate/ParticleModuleMeshRotationRateOverLife.h
@@ -4,7 +4,7 @@
 #pragma once
 #include "ParticleModuleMeshRotationRateOverLife.generated.h"
 
-UCLASS(HeaderGroup=Particle, editinlinenew, hidecategories=Object, meta=(DisplayName = "Mesh Rotation Rate over Life"))
+UCLASS(editinlinenew, hidecategories=Object, meta=(DisplayName = "Mesh Rotation Rate over Life"))
 class UParticleModuleMeshRotationRateOverLife : public UParticleModuleRotationRateBase
 {
 	GENERATED_UCLASS_BODY()
diff --git a/Engine/Source/Runtime/Engine/Classes/Particles/RotationRate/ParticleModuleMeshRotationRate_Seeded.h b/Engine/Source/Runtime/Engine/Classes/Particles/RotationRate/ParticleModuleMeshRotationRate_Seeded.h
--- a/Engine/Source/Runtime/Engine/Classes/Particles/RotationRate/ParticleModuleMeshRotationRate_Seeded.h
+++ b/Engine/Source/Runtime/Engine/Classes/Particles/RotationRate/ParticleModuleMeshRotationRate_Seeded.h
@@ -4,7 +4,7 @@
 #pragma once
 #include "ParticleModuleMeshRotationRate_Seeded.generated.h"
 
-UCLASS(HeaderGroup=Particle, editinlinenew, hidecategories=Object, meta=(DisplayName = "Init Mesh Rotation Rate (Seed)"))
+UCLASS(editinlinenew, hidecategories=Object, meta=(DisplayName = "Init Mesh Rotation Rate (Seed)"))
 class UParticleModuleMeshRotationRate_Seeded : public UParticleModuleMeshRotationRate
 {
 	GENERATED_UCLASS_BODY()
diff --git a/Engine/Source/Runtime/Engine/Classes/Particles/RotationRate/ParticleModuleRotationRate.h b/Engine/Source/Runtime/Engine/Classes/Particles/RotationRate/ParticleModuleRotationRate.h
--- a/Engine/Source/Runtime/Engine/Classes/Particles/RotationRate/ParticleModuleRotationRate.h
+++ b/Engine/Source/Runtime/Engine/Classes/Particles/RotationRate/ParticleModuleRotationRate.h
@@ -4,7 +4,7 @@
 #pragma once
 #include "ParticleModuleRotationRate.generated.h"
 
-UCLASS(HeaderGroup=Particle, editinlinenew, hidecategories=Object, MinimalAPI, meta=(DisplayName = "Initial Rotation Rate"))
+UCLASS(editinlinenew, hidecategories=Object, MinimalAPI, meta=(DisplayName = "Initial Rotation Rate"))
 class UParticleModuleRotationRate : public UParticleModuleRotationRateBase
 {
 	GENERATED_UCLASS_BODY()
diff --git a/Engine/Source/Runtime/Engine/Classes/Particles/RotationRate/ParticleModuleRotationRateBase.h b/Engine/Source/Runtime/Engine/Classes/Particles/RotationRate/ParticleModuleRotationRateBase.h
--- a/Engine/Source/Runtime/Engine/Classes/Particles/RotationRate/ParticleModuleRotationRateBase.h
+++ b/Engine/Source/Runtime/Engine/Classes/Particles/RotationRate/ParticleModuleRotationRateBase.h
@@ -4,7 +4,7 @@
 #pragma once
 #include "ParticleModuleRotationRateBase.generated.h"
 
-UCLASS(HeaderGroup=Particle, editinlinenew, hidecategories=Object, abstract, meta=(DisplayName = "Rotation Rate"))
+UCLASS(editinlinenew, hidecategories=Object, abstract, meta=(DisplayName = "Rotation Rate"))
 class UParticleModuleRotationRateBase : public UParticleModule
 {
 	GENERATED_UCLASS_BODY()
diff --git a/Engine/Source/Runtime/Engine/Classes/Particles/RotationRate/ParticleModuleRotationRateMultiplyLife.h b/Engine/Source/Runtime/Engine/Classes/Particles/RotationRate/ParticleModuleRotationRateMultiplyLife.h
--- a/Engine/Source/Runtime/Engine/Classes/Particles/RotationRate/ParticleModuleRotationRateMultiplyLife.h
+++ b/Engine/Source/Runtime/Engine/Classes/Particles/RotationRate/ParticleModuleRotationRateMultiplyLife.h
@@ -4,7 +4,7 @@
 #pragma once
 #include "ParticleModuleRotationRateMultiplyLife.generated.h"
 
-UCLASS(HeaderGroup=Particle, editinlinenew, hidecategories=Object, MinimalAPI, meta=(DisplayName = "Rotation Rate * Life"))
+UCLASS(editinlinenew, hidecategories=Object, MinimalAPI, meta=(DisplayName = "Rotation Rate * Life"))
 class UParticleModuleRotationRateMultiplyLife : public UParticleModuleRotationRateBase
 {
 	GENERATED_UCLASS_BODY()
diff --git a/Engine/Source/Runtime/Engine/Classes/Particles/RotationRate/ParticleModuleRotationRate_Seeded.h b/Engine/Source/Runtime/Engine/Classes/Particles/RotationRate/ParticleModuleRotationRate_Seeded.h
--- a/Engine/Source/Runtime/Engine/Classes/Particles/RotationRate/ParticleModuleRotationRate_Seeded.h
+++ b/Engine/Source/Runtime/Engine/Classes/Particles/RotationRate/ParticleModuleRotationRate_Seeded.h
@@ -4,7 +4,7 @@
 #pragma once
 #include "ParticleModuleRotationRate_Seeded.generated.h"
 
-UCLASS(HeaderGroup=Particle, editinlinenew, hidecategories=Object, meta=(DisplayName = "Init Rotation Rate (Seed)"))
+UCLASS(editinlinenew, hidecategories=Object, meta=(DisplayName = "Init Rotation Rate (Seed)"))
 class UParticleModuleRotationRate_Seeded : public UParticleModuleRotationRate
 {
 	GENERATED_UCLASS_BODY()
diff --git a/Engine/Source/Runtime/Engine/Classes/Particles/Size/ParticleModuleSize.h b/Engine/Source/Runtime/Engine/Classes/Particles/Size/ParticleModuleSize.h
--- a/Engine/Source/Runtime/Engine/Classes/Particles/Size/ParticleModuleSize.h
+++ b/Engine/Source/Runtime/Engine/Classes/Particles/Size/ParticleModuleSize.h
@@ -4,7 +4,7 @@
 #pragma once
 #include "ParticleModuleSize.generated.h"
 
-UCLASS(HeaderGroup=Particle, editinlinenew, hidecategories=Object, MinimalAPI, meta=(DisplayName = "Initial Size"))
+UCLASS(editinlinenew, hidecategories=Object, MinimalAPI, meta=(DisplayName = "Initial Size"))
 class UParticleModuleSize : public UParticleModuleSizeBase
 {
 	GENERATED_UCLASS_BODY()
diff --git a/Engine/Source/Runtime/Engine/Classes/Particles/Size/ParticleModuleSizeBase.h b/Engine/Source/Runtime/Engine/Classes/Particles/Size/ParticleModuleSizeBase.h
--- a/Engine/Source/Runtime/Engine/Classes/Particles/Size/ParticleModuleSizeBase.h
+++ b/Engine/Source/Runtime/Engine/Classes/Particles/Size/ParticleModuleSizeBase.h
@@ -4,7 +4,7 @@
 #pragma once
 #include "ParticleModuleSizeBase.generated.h"
 
-UCLASS(HeaderGroup=Particle, editinlinenew, hidecategories=Object, abstract, meta=(DisplayName = "Size"))
+UCLASS(editinlinenew, hidecategories=Object, abstract, meta=(DisplayName = "Size"))
 class UParticleModuleSizeBase : public UParticleModule
 {
 	GENERATED_UCLASS_BODY()
diff --git a/Engine/Source/Runtime/Engine/Classes/Particles/Size/ParticleModuleSizeMultiplyLife.h b/Engine/Source/Runtime/Engine/Classes/Particles/Size/ParticleModuleSizeMultiplyLife.h
--- a/Engine/Source/Runtime/Engine/Classes/Particles/Size/ParticleModuleSizeMultiplyLife.h
+++ b/Engine/Source/Runtime/Engine/Classes/Particles/Size/ParticleModuleSizeMultiplyLife.h
@@ -4,7 +4,7 @@
 #pragma once
 #include "ParticleModuleSizeMultiplyLife.generated.h"
 
-UCLASS(HeaderGroup=Particle, editinlinenew, hidecategories=Object, MinimalAPI, meta=(DisplayName = "Size By Life"))
+UCLASS(editinlinenew, hidecategories=Object, MinimalAPI, meta=(DisplayName = "Size By Life"))
 class UParticleModuleSizeMultiplyLife : public UParticleModuleSizeBase
 {
 	GENERATED_UCLASS_BODY()
diff --git a/Engine/Source/Runtime/Engine/Classes/Particles/Size/ParticleModuleSizeScale.h b/Engine/Source/Runtime/Engine/Classes/Particles/Size/ParticleModuleSizeScale.h
--- a/Engine/Source/Runtime/Engine/Classes/Particles/Size/ParticleModuleSizeScale.h
+++ b/Engine/Source/Runtime/Engine/Classes/Particles/Size/ParticleModuleSizeScale.h
@@ -4,7 +4,7 @@
 #pragma once
 #include "ParticleModuleSizeScale.generated.h"
 
-UCLASS(HeaderGroup=Particle, editinlinenew, hidecategories=Object, MinimalAPI, meta=(DisplayName = "Size Scale"))
+UCLASS(editinlinenew, hidecategories=Object, MinimalAPI, meta=(DisplayName = "Size Scale"))
 class UParticleModuleSizeScale : public UParticleModuleSizeBase
 {
 	GENERATED_UCLASS_BODY()
diff --git a/Engine/Source/Runtime/Engine/Classes/Particles/Size/ParticleModuleSizeScaleBySpeed.h b/Engine/Source/Runtime/Engine/Classes/Particles/Size/ParticleModuleSizeScaleBySpeed.h
--- a/Engine/Source/Runtime/Engine/Classes/Particles/Size/ParticleModuleSizeScaleBySpeed.h
+++ b/Engine/Source/Runtime/Engine/Classes/Particles/Size/ParticleModuleSizeScaleBySpeed.h
@@ -8,7 +8,7 @@
 
 #include "ParticleModuleSizeScaleBySpeed.generated.h"
 
-UCLASS(HeaderGroup=Particle, editinlinenew, hidecategories=Object, meta=(DisplayName = "Size By Speed"))
+UCLASS(editinlinenew, hidecategories=Object, meta=(DisplayName = "Size By Speed"))
 class UParticleModuleSizeScaleBySpeed : public UParticleModuleSizeBase
 {
 	GENERATED_UCLASS_BODY()
diff --git a/Engine/Source/Runtime/Engine/Classes/Particles/Size/ParticleModuleSize_Seeded.h b/Engine/Source/Runtime/Engine/Classes/Particles/Size/ParticleModuleSize_Seeded.h
--- a/Engine/Source/Runtime/Engine/Classes/Particles/Size/ParticleModuleSize_Seeded.h
+++ b/Engine/Source/Runtime/Engine/Classes/Particles/Size/ParticleModuleSize_Seeded.h
@@ -4,7 +4,7 @@
 #pragma once
 #include "ParticleModuleSize_Seeded.generated.h"
 
-UCLASS(HeaderGroup=Particle, editinlinenew, hidecategories=Object, meta=(DisplayName = "Initial Size (Seed)"))
+UCLASS(editinlinenew, hidecategories=Object, meta=(DisplayName = "Initial Size (Seed)"))
 class UParticleModuleSize_Seeded : public UParticleModuleSize
 {
 	GENERATED_UCLASS_BODY()
diff --git a/Engine/Source/Runtime/Engine/Classes/Particles/Spawn/ParticleModuleSpawn.h b/Engine/Source/Runtime/Engine/Classes/Particles/Spawn/ParticleModuleSpawn.h
--- a/Engine/Source/Runtime/Engine/Classes/Particles/Spawn/ParticleModuleSpawn.h
+++ b/Engine/Source/Runtime/Engine/Classes/Particles/Spawn/ParticleModuleSpawn.h
@@ -4,7 +4,7 @@
 #pragma once
 #include "ParticleModuleSpawn.generated.h"
 
-UCLASS(HeaderGroup=Particle, editinlinenew, hidecategories=Object, hidecategories=ParticleModuleSpawnBase, MinimalAPI, meta=(DisplayName = "Spawn"))
+UCLASS(editinlinenew, hidecategories=Object, hidecategories=ParticleModuleSpawnBase, MinimalAPI, meta=(DisplayName = "Spawn"))
 class UParticleModuleSpawn : public UParticleModuleSpawnBase
 {
 	GENERATED_UCLASS_BODY()
diff --git a/Engine/Source/Runtime/Engine/Classes/Particles/Spawn/ParticleModuleSpawnBase.h b/Engine/Source/Runtime/Engine/Classes/Particles/Spawn/ParticleModuleSpawnBase.h
--- a/Engine/Source/Runtime/Engine/Classes/Particles/Spawn/ParticleModuleSpawnBase.h
+++ b/Engine/Source/Runtime/Engine/Classes/Particles/Spawn/ParticleModuleSpawnBase.h
@@ -4,7 +4,7 @@
 #pragma once
 #include "ParticleModuleSpawnBase.generated.h"
 
-UCLASS(HeaderGroup=Particle, editinlinenew, hidecategories=Object, abstract, meta=(DisplayName = "Spawn"))
+UCLASS(editinlinenew, hidecategories=Object, abstract, meta=(DisplayName = "Spawn"))
 class UParticleModuleSpawnBase : public UParticleModule
 {
 	GENERATED_UCLASS_BODY()
diff --git a/Engine/Source/Runtime/Engine/Classes/Particles/Spawn/ParticleModuleSpawnPerUnit.h b/Engine/Source/Runtime/Engine/Classes/Particles/Spawn/ParticleModuleSpawnPerUnit.h
--- a/Engine/Source/Runtime/Engine/Classes/Particles/Spawn/ParticleModuleSpawnPerUnit.h
+++ b/Engine/Source/Runtime/Engine/Classes/Particles/Spawn/ParticleModuleSpawnPerUnit.h
@@ -4,7 +4,7 @@
 #pragma once
 #include "ParticleModuleSpawnPerUnit.generated.h"
 
-UCLASS(HeaderGroup=Particle, editinlinenew, hidecategories=Object, meta=(DisplayName = "Spawn PerUnit"))
+UCLASS(editinlinenew, hidecategories=Object, meta=(DisplayName = "Spawn PerUnit"))
 class UParticleModuleSpawnPerUnit : public UParticleModuleSpawnBase
 {
 	GENERATED_UCLASS_BODY()
diff --git a/Engine/Source/Runtime/Engine/Classes/Particles/SubUV/ParticleModuleSubUV.h b/Engine/Source/Runtime/Engine/Classes/Particles/SubUV/ParticleModuleSubUV.h
--- a/Engine/Source/Runtime/Engine/Classes/Particles/SubUV/ParticleModuleSubUV.h
+++ b/Engine/Source/Runtime/Engine/Classes/Particles/SubUV/ParticleModuleSubUV.h
@@ -4,7 +4,7 @@
 #pragma once
 #include "ParticleModuleSubUV.generated.h"
 
-UCLASS(HeaderGroup=Particle, editinlinenew, hidecategories=Object, MinimalAPI, meta=(DisplayName = "SubImage Index"))
+UCLASS(editinlinenew, hidecategories=Object, MinimalAPI, meta=(DisplayName = "SubImage Index"))
 class UParticleModuleSubUV : public UParticleModuleSubUVBase
 {
 	GENERATED_UCLASS_BODY()
diff --git a/Engine/Source/Runtime/Engine/Classes/Particles/SubUV/ParticleModuleSubUVBase.h b/Engine/Source/Runtime/Engine/Classes/Particles/SubUV/ParticleModuleSubUVBase.h
--- a/Engine/Source/Runtime/Engine/Classes/Particles/SubUV/ParticleModuleSubUVBase.h
+++ b/Engine/Source/Runtime/Engine/Classes/Particles/SubUV/ParticleModuleSubUVBase.h
@@ -4,7 +4,7 @@
 #pragma once
 #include "ParticleModuleSubUVBase.generated.h"
 
-UCLASS(HeaderGroup=Particle, editinlinenew, hidecategories=Object, abstract, meta=(DisplayName = "SubUV"))
+UCLASS(editinlinenew, hidecategories=Object, abstract, meta=(DisplayName = "SubUV"))
 class UParticleModuleSubUVBase : public UParticleModule
 {
 	GENERATED_UCLASS_BODY()
diff --git a/Engine/Source/Runtime/Engine/Classes/Particles/SubUV/ParticleModuleSubUVMovie.h b/Engine/Source/Runtime/Engine/Classes/Particles/SubUV/ParticleModuleSubUVMovie.h
--- a/Engine/Source/Runtime/Engine/Classes/Particles/SubUV/ParticleModuleSubUVMovie.h
+++ b/Engine/Source/Runtime/Engine/Classes/Particles/SubUV/ParticleModuleSubUVMovie.h
@@ -4,7 +4,7 @@
 #pragma once
 #include "ParticleModuleSubUVMovie.generated.h"
 
-UCLASS(HeaderGroup=Particle, editinlinenew, hidecategories=Object, hidecategories=SubUV, meta=(DisplayName = "SubUV Movie"))
+UCLASS(editinlinenew, hidecategories=Object, hidecategories=SubUV, meta=(DisplayName = "SubUV Movie"))
 class UParticleModuleSubUVMovie : public UParticleModuleSubUV
 {
 	GENERATED_UCLASS_BODY()
diff --git a/Engine/Source/Runtime/Engine/Classes/Particles/Trail/ParticleModuleTrailBase.h b/Engine/Source/Runtime/Engine/Classes/Particles/Trail/ParticleModuleTrailBase.h
--- a/Engine/Source/Runtime/Engine/Classes/Particles/Trail/ParticleModuleTrailBase.h
+++ b/Engine/Source/Runtime/Engine/Classes/Particles/Trail/ParticleModuleTrailBase.h
@@ -9,7 +9,7 @@
 #pragma once
 #include "ParticleModuleTrailBase.generated.h"
 
-UCLASS(HeaderGroup=Particle, editinlinenew, hidecategories=Object, abstract, meta=(DisplayName = "Trail"))
+UCLASS(editinlinenew, hidecategories=Object, abstract, meta=(DisplayName = "Trail"))
 class UParticleModuleTrailBase : public UParticleModule
 {
 	GENERATED_UCLASS_BODY()
diff --git a/Engine/Source/Runtime/Engine/Classes/Particles/TypeData/ParticleModuleTypeDataAnimTrail.h b/Engine/Source/Runtime/Engine/Classes/Particles/TypeData/ParticleModuleTypeDataAnimTrail.h
--- a/Engine/Source/Runtime/Engine/Classes/Particles/TypeData/ParticleModuleTypeDataAnimTrail.h
+++ b/Engine/Source/Runtime/Engine/Classes/Particles/TypeData/ParticleModuleTypeDataAnimTrail.h
@@ -9,7 +9,7 @@
 #pragma once
 #include "ParticleModuleTypeDataAnimTrail.generated.h"
 
-UCLASS(HeaderGroup=Particle, editinlinenew, hidecategories=Object, meta=(DisplayName = "AnimTrail Data"))
+UCLASS(editinlinenew, hidecategories=Object, meta=(DisplayName = "AnimTrail Data"))
 class UParticleModuleTypeDataAnimTrail : public UParticleModuleTypeDataBase
 {
 	GENERATED_UCLASS_BODY()
diff --git a/Engine/Source/Runtime/Engine/Classes/Particles/TypeData/ParticleModuleTypeDataBeam2.h b/Engine/Source/Runtime/Engine/Classes/Particles/TypeData/ParticleModuleTypeDataBeam2.h
--- a/Engine/Source/Runtime/Engine/Classes/Particles/TypeData/ParticleModuleTypeDataBeam2.h
+++ b/Engine/Source/Runtime/Engine/Classes/Particles/TypeData/ParticleModuleTypeDataBeam2.h
@@ -49,7 +49,7 @@ struct FBeamTargetData
 
 };
 
-UCLASS(HeaderGroup=Particle, editinlinenew, dontcollapsecategories, hidecategories=Object, MinimalAPI, meta=(DisplayName = "Beam Data"))
+UCLASS(editinlinenew, dontcollapsecategories, hidecategories=Object, MinimalAPI, meta=(DisplayName = "Beam Data"))
 class UParticleModuleTypeDataBeam2 : public UParticleModuleTypeDataBase
 {
 	GENERATED_UCLASS_BODY()
diff --git a/Engine/Source/Runtime/Engine/Classes/Particles/TypeData/ParticleModuleTypeDataGpu.h b/Engine/Source/Runtime/Engine/Classes/Particles/TypeData/ParticleModuleTypeDataGpu.h
--- a/Engine/Source/Runtime/Engine/Classes/Particles/TypeData/ParticleModuleTypeDataGpu.h
+++ b/Engine/Source/Runtime/Engine/Classes/Particles/TypeData/ParticleModuleTypeDataGpu.h
@@ -377,7 +377,7 @@ struct FGPUSpriteResourceData
 
 };
 
-UCLASS(HeaderGroup=Particle, editinlinenew, hidecategories=Object, MinimalAPI, meta=(DisplayName = "GPU Sprites"))
+UCLASS(editinlinenew, hidecategories=Object, MinimalAPI, meta=(DisplayName = "GPU Sprites"))
 class UParticleModuleTypeDataGpu : public UParticleModuleTypeDataBase
 {
 	GENERATED_UCLASS_BODY()
diff --git a/Engine/Source/Runtime/Engine/Classes/Particles/TypeData/ParticleModuleTypeDataMesh.h b/Engine/Source/Runtime/Engine/Classes/Particles/TypeData/ParticleModuleTypeDataMesh.h
--- a/Engine/Source/Runtime/Engine/Classes/Particles/TypeData/ParticleModuleTypeDataMesh.h
+++ b/Engine/Source/Runtime/Engine/Classes/Particles/TypeData/ParticleModuleTypeDataMesh.h
@@ -46,7 +46,7 @@ enum EMeshCameraFacingOptions
 	EMeshCameraFacingOptions_MAX,
 };
 
-UCLASS(HeaderGroup=Particle, editinlinenew, hidecategories=Object, MinimalAPI, meta=(DisplayName = "Mesh Data"))
+UCLASS(editinlinenew, hidecategories=Object, MinimalAPI, meta=(DisplayName = "Mesh Data"))
 class UParticleModuleTypeDataMesh : public UParticleModuleTypeDataBase
 {
 	GENERATED_UCLASS_BODY()
diff --git a/Engine/Source/Runtime/Engine/Classes/Particles/TypeData/ParticleModuleTypeDataRibbon.h b/Engine/Source/Runtime/Engine/Classes/Particles/TypeData/ParticleModuleTypeDataRibbon.h
--- a/Engine/Source/Runtime/Engine/Classes/Particles/TypeData/ParticleModuleTypeDataRibbon.h
+++ b/Engine/Source/Runtime/Engine/Classes/Particles/TypeData/ParticleModuleTypeDataRibbon.h
@@ -18,7 +18,7 @@ enum ETrailsRenderAxisOption
 	Trails_MAX,
 };
 
-UCLASS(HeaderGroup=Particle, editinlinenew, hidecategories=Object, meta=(DisplayName = "Ribbon Data"))
+UCLASS(editinlinenew, hidecategories=Object, meta=(DisplayName = "Ribbon Data"))
 class UParticleModuleTypeDataRibbon : public UParticleModuleTypeDataBase
 {
 	GENERATED_UCLASS_BODY()
diff --git a/Engine/Source/Runtime/Engine/Classes/Particles/VectorField/ParticleModuleVectorFieldBase.h b/Engine/Source/Runtime/Engine/Classes/Particles/VectorField/ParticleModuleVectorFieldBase.h
--- a/Engine/Source/Runtime/Engine/Classes/Particles/VectorField/ParticleModuleVectorFieldBase.h
+++ b/Engine/Source/Runtime/Engine/Classes/Particles/VectorField/ParticleModuleVectorFieldBase.h
@@ -9,7 +9,7 @@
 
 #include "ParticleModuleVectorFieldBase.generated.h"
 
-UCLASS(HeaderGroup=Particle, editinlinenew, hidecategories=Object, abstract, meta=(DisplayName = "Vector Field"))
+UCLASS(editinlinenew, hidecategories=Object, abstract, meta=(DisplayName = "Vector Field"))
 class UParticleModuleVectorFieldBase : public UParticleModule
 {
 	GENERATED_UCLASS_BODY()
diff --git a/Engine/Source/Runtime/Engine/Classes/Particles/VectorField/ParticleModuleVectorFieldGlobal.h b/Engine/Source/Runtime/Engine/Classes/Particles/VectorField/ParticleModuleVectorFieldGlobal.h
--- a/Engine/Source/Runtime/Engine/Classes/Particles/VectorField/ParticleModuleVectorFieldGlobal.h
+++ b/Engine/Source/Runtime/Engine/Classes/Particles/VectorField/ParticleModuleVectorFieldGlobal.h
@@ -8,7 +8,7 @@
 
 #include "ParticleModuleVectorFieldGlobal.generated.h"
 
-UCLASS(HeaderGroup=Particle, editinlinenew, hidecategories=Object, meta=(DisplayName = "Global Vector Fields"))
+UCLASS(editinlinenew, hidecategories=Object, meta=(DisplayName = "Global Vector Fields"))
 class UParticleModuleVectorFieldGlobal : public UParticleModuleVectorFieldBase
 {
 	GENERATED_UCLASS_BODY()
diff --git a/Engine/Source/Runtime/Engine/Classes/Particles/VectorField/ParticleModuleVectorFieldLocal.h b/Engine/Source/Runtime/Engine/Classes/Particles/VectorField/ParticleModuleVectorFieldLocal.h
--- a/Engine/Source/Runtime/Engine/Classes/Particles/VectorField/ParticleModuleVectorFieldLocal.h
+++ b/Engine/Source/Runtime/Engine/Classes/Particles/VectorField/ParticleModuleVectorFieldLocal.h
@@ -8,7 +8,7 @@
 
 #include "ParticleModuleVectorFieldLocal.generated.h"
 
-UCLASS(HeaderGroup=Particle, editinlinenew, hidecategories=Object, MinimalAPI, meta=(DisplayName = "Local Vector Field"))
+UCLASS(editinlinenew, hidecategories=Object, MinimalAPI, meta=(DisplayName = "Local Vector Field"))
 class UParticleModuleVectorFieldLocal : public UParticleModuleVectorFieldBase
 {
 	GENERATED_UCLASS_BODY()
diff --git a/Engine/Source/Runtime/Engine/Classes/Particles/VectorField/ParticleModuleVectorFieldRotation.h b/Engine/Source/Runtime/Engine/Classes/Particles/VectorField/ParticleModuleVectorFieldRotation.h
--- a/Engine/Source/Runtime/Engine/Classes/Particles/VectorField/ParticleModuleVectorFieldRotation.h
+++ b/Engine/Source/Runtime/Engine/Classes/Particles/VectorField/ParticleModuleVectorFieldRotation.h
@@ -9,7 +9,7 @@
 
 #include "ParticleModuleVectorFieldRotation.generated.h"
 
-UCLASS(HeaderGroup=Particle, editinlinenew, hidecategories=Object, meta=(DisplayName = "VF Init Rotation"))
+UCLASS(editinlinenew, hidecategories=Object, meta=(DisplayName = "VF Init Rotation"))
 class UParticleModuleVectorFieldRotation : public UParticleModuleVectorFieldBase
 {
 	GENERATED_UCLASS_BODY()
diff --git a/Engine/Source/Runtime/Engine/Classes/Particles/VectorField/ParticleModuleVectorFieldRotationRate.h b/Engine/Source/Runtime/Engine/Classes/Particles/VectorField/ParticleModuleVectorFieldRotationRate.h
--- a/Engine/Source/Runtime/Engine/Classes/Particles/VectorField/ParticleModuleVectorFieldRotationRate.h
+++ b/Engine/Source/Runtime/Engine/Classes/Particles/VectorField/ParticleModuleVectorFieldRotationRate.h
@@ -8,7 +8,7 @@
 
 #include "ParticleModuleVectorFieldRotationRate.generated.h"
 
-UCLASS(HeaderGroup=Particle, editinlinenew, hidecategories=Object, meta=(DisplayName = "VF Rotation Rate"))
+UCLASS(editinlinenew, hidecategories=Object, meta=(DisplayName = "VF Rotation Rate"))
 class UParticleModuleVectorFieldRotationRate : public UParticleModuleVectorFieldBase
 {
 	GENERATED_UCLASS_BODY()
diff --git a/Engine/Source/Runtime/Engine/Classes/Particles/VectorField/ParticleModuleVectorFieldScale.h b/Engine/Source/Runtime/Engine/Classes/Particles/VectorField/ParticleModuleVectorFieldScale.h
--- a/Engine/Source/Runtime/Engine/Classes/Particles/VectorField/ParticleModuleVectorFieldScale.h
+++ b/Engine/Source/Runtime/Engine/Classes/Particles/VectorField/ParticleModuleVectorFieldScale.h
@@ -8,7 +8,7 @@
 
 #include "ParticleModuleVectorFieldScale.generated.h"
 
-UCLASS(HeaderGroup=Particle, editinlinenew, hidecategories=Object, MinimalAPI, meta=(DisplayName = "Vector Field Scale"))
+UCLASS(editinlinenew, hidecategories=Object, MinimalAPI, meta=(DisplayName = "Vector Field Scale"))
 class UParticleModuleVectorFieldScale : public UParticleModuleVectorFieldBase
 {
 	GENERATED_UCLASS_BODY()
diff --git a/Engine/Source/Runtime/Engine/Classes/Particles/VectorField/ParticleModuleVectorFieldScaleOverLife.h b/Engine/Source/Runtime/Engine/Classes/Particles/VectorField/ParticleModuleVectorFieldScaleOverLife.h
--- a/Engine/Source/Runtime/Engine/Classes/Particles/VectorField/ParticleModuleVectorFieldScaleOverLife.h
+++ b/Engine/Source/Runtime/Engine/Classes/Particles/VectorField/ParticleModuleVectorFieldScaleOverLife.h
@@ -8,7 +8,7 @@
 
 #include "ParticleModuleVectorFieldScaleOverLife.generated.h"
 
-UCLASS(HeaderGroup=Particle, editinlinenew, hidecategories=Object, MinimalAPI, meta=(DisplayName = "VF Scale/Life"))
+UCLASS(editinlinenew, hidecategories=Object, MinimalAPI, meta=(DisplayName = "VF Scale/Life"))
 class UParticleModuleVectorFieldScaleOverLife : public UParticleModuleVectorFieldBase
 {
 	GENERATED_UCLASS_BODY()
diff --git a/Engine/Source/Runtime/Engine/Classes/Particles/Velocity/ParticleModuleVelocity.h b/Engine/Source/Runtime/Engine/Classes/Particles/Velocity/ParticleModuleVelocity.h
--- a/Engine/Source/Runtime/Engine/Classes/Particles/Velocity/ParticleModuleVelocity.h
+++ b/Engine/Source/Runtime/Engine/Classes/Particles/Velocity/ParticleModuleVelocity.h
@@ -4,7 +4,7 @@
 #pragma once
 #include "ParticleModuleVelocity.generated.h"
 
-UCLASS(HeaderGroup=Particle, editinlinenew, hidecategories=Object, MinimalAPI, meta=(DisplayName = "Initial Velocity"))
+UCLASS(editinlinenew, hidecategories=Object, MinimalAPI, meta=(DisplayName = "Initial Velocity"))
 class UParticleModuleVelocity : public UParticleModuleVelocityBase
 {
 	GENERATED_UCLASS_BODY()
diff --git a/Engine/Source/Runtime/Engine/Classes/Particles/Velocity/ParticleModuleVelocityBase.h b/Engine/Source/Runtime/Engine/Classes/Particles/Velocity/ParticleModuleVelocityBase.h
--- a/Engine/Source/Runtime/Engine/Classes/Particles/Velocity/ParticleModuleVelocityBase.h
+++ b/Engine/Source/Runtime/Engine/Classes/Particles/Velocity/ParticleModuleVelocityBase.h
@@ -4,7 +4,7 @@
 #pragma once
 #include "ParticleModuleVelocityBase.generated.h"
 
-UCLASS(HeaderGroup=Particle, editinlinenew, hidecategories=Object, abstract, meta=(DisplayName = "Velocity"))
+UCLASS(editinlinenew, hidecategories=Object, abstract, meta=(DisplayName = "Velocity"))
 class UParticleModuleVelocityBase : public UParticleModule
 {
 	GENERATED_UCLASS_BODY()
diff --git a/Engine/Source/Runtime/Engine/Classes/Particles/Velocity/ParticleModuleVelocityCone.h b/Engine/Source/Runtime/Engine/Classes/Particles/Velocity/ParticleModuleVelocityCone.h
--- a/Engine/Source/Runtime/Engine/Classes/Particles/Velocity/ParticleModuleVelocityCone.h
+++ b/Engine/Source/Runtime/Engine/Classes/Particles/Velocity/ParticleModuleVelocityCone.h
@@ -4,7 +4,7 @@
 #pragma once
 #include "ParticleModuleVelocityCone.generated.h"
 
-UCLASS(HeaderGroup=Particle, editinlinenew, hidecategories=Object, meta=(DisplayName = "Velocity Cone"))
+UCLASS(editinlinenew, hidecategories=Object, meta=(DisplayName = "Velocity Cone"))
 class UParticleModuleVelocityCone : public UParticleModuleVelocityBase
 {
 	GENERATED_UCLASS_BODY()
diff --git a/Engine/Source/Runtime/Engine/Classes/Particles/Velocity/ParticleModuleVelocityInheritParent.h b/Engine/Source/Runtime/Engine/Classes/Particles/Velocity/ParticleModuleVelocityInheritParent.h
--- a/Engine/Source/Runtime/Engine/Classes/Particles/Velocity/ParticleModuleVelocityInheritParent.h
+++ b/Engine/Source/Runtime/Engine/Classes/Particles/Velocity/ParticleModuleVelocityInheritParent.h
@@ -4,7 +4,7 @@
 #pragma once
 #include "ParticleModuleVelocityInheritParent.generated.h"
 
-UCLASS(HeaderGroup=Particle, editinlinenew, hidecategories=Object, MinimalAPI, meta=(DisplayName = "Inherit Parent Velocity"))
+UCLASS(editinlinenew, hidecategories=Object, MinimalAPI, meta=(DisplayName = "Inherit Parent Velocity"))
 class UParticleModuleVelocityInheritParent : public UParticleModuleVelocityBase
 {
 	GENERATED_UCLASS_BODY()
diff --git a/Engine/Source/Runtime/Engine/Classes/Particles/Velocity/ParticleModuleVelocityOverLifetime.h b/Engine/Source/Runtime/Engine/Classes/Particles/Velocity/ParticleModuleVelocityOverLifetime.h
--- a/Engine/Source/Runtime/Engine/Classes/Particles/Velocity/ParticleModuleVelocityOverLifetime.h
+++ b/Engine/Source/Runtime/Engine/Classes/Particles/Velocity/ParticleModuleVelocityOverLifetime.h
@@ -4,7 +4,7 @@
 #pragma once
 #include "ParticleModuleVelocityOverLifetime.generated.h"
 
-UCLASS(HeaderGroup=Particle, editinlinenew, hidecategories=Object, MinimalAPI, meta=(DisplayName = "Velocity/Life"))
+UCLASS(editinlinenew, hidecategories=Object, MinimalAPI, meta=(DisplayName = "Velocity/Life"))
 class UParticleModuleVelocityOverLifetime : public UParticleModuleVelocityBase
 {
 	GENERATED_UCLASS_BODY()
diff --git a/Engine/Source/Runtime/Engine/Classes/Particles/Velocity/ParticleModuleVelocity_Seeded.h b/Engine/Source/Runtime/Engine/Classes/Particles/Velocity/ParticleModuleVelocity_Seeded.h
--- a/Engine/Source/Runtime/Engine/Classes/Particles/Velocity/ParticleModuleVelocity_Seeded.h
+++ b/Engine/Source/Runtime/Engine/Classes/Particles/Velocity/ParticleModuleVelocity_Seeded.h
@@ -4,7 +4,7 @@
 #pragma once
 #include "ParticleModuleVelocity_Seeded.generated.h"
 
-UCLASS(HeaderGroup=Particle, editinlinenew, hidecategories=Object, meta=(DisplayName = "Initial Velocity (Seed)"))
+UCLASS(editinlinenew, hidecategories=Object, meta=(DisplayName = "Initial Velocity (Seed)"))
 class UParticleModuleVelocity_Seeded : public UParticleModuleVelocity
 {
 	GENERATED_UCLASS_BODY()
diff --git a/Engine/Source/Runtime/Engine/Classes/PhysicsEngine/PhysicsConstraintComponent.h b/Engine/Source/Runtime/Engine/Classes/PhysicsEngine/PhysicsConstraintComponent.h
--- a/Engine/Source/Runtime/Engine/Classes/PhysicsEngine/PhysicsConstraintComponent.h
+++ b/Engine/Source/Runtime/Engine/Classes/PhysicsEngine/PhysicsConstraintComponent.h
@@ -9,7 +9,7 @@ DECLARE_DYNAMIC_MULTICAST_DELEGATE_OneParam(FConstraintBrokenSignature, int32, C
 /**
  *	This is effectively a joint that allows you to connect 2 rigid bodies together. You can create different types of joints using the various parameters of this component.
  */
-UCLASS(HeaderGroup=Component, ClassGroup=Physics, dependson(ConstraintInstance), MinimalAPI, meta=(BlueprintSpawnableComponent), HideCategories=(Activation,"Components|Activation",Physics, Mobility))
+UCLASS(ClassGroup=Physics, dependson(ConstraintInstance), MinimalAPI, meta=(BlueprintSpawnableComponent), HideCategories=(Activation,"Components|Activation",Physics, Mobility))
 class UPhysicsConstraintComponent : public USceneComponent
 {
 	GENERATED_UCLASS_BODY()
diff --git a/Engine/Source/Runtime/Engine/Classes/PhysicsEngine/PhysicsHandleComponent.h b/Engine/Source/Runtime/Engine/Classes/PhysicsEngine/PhysicsHandleComponent.h
--- a/Engine/Source/Runtime/Engine/Classes/PhysicsEngine/PhysicsHandleComponent.h
+++ b/Engine/Source/Runtime/Engine/Classes/PhysicsEngine/PhysicsHandleComponent.h
@@ -14,7 +14,7 @@ namespace physx
 /**
  *	Utility object for moving physics objects around.
  */
-UCLASS(collapsecategories, ClassGroup=Physics, hidecategories=Object, HeaderGroup=Component, MinimalAPI, meta=(BlueprintSpawnableComponent))
+UCLASS(collapsecategories, ClassGroup=Physics, hidecategories=Object, MinimalAPI, meta=(BlueprintSpawnableComponent))
 class UPhysicsHandleComponent : public UActorComponent
 {
 	GENERATED_UCLASS_BODY()
diff --git a/Engine/Source/Runtime/Engine/Classes/PhysicsEngine/PhysicsThrusterComponent.h b/Engine/Source/Runtime/Engine/Classes/PhysicsEngine/PhysicsThrusterComponent.h
--- a/Engine/Source/Runtime/Engine/Classes/PhysicsEngine/PhysicsThrusterComponent.h
+++ b/Engine/Source/Runtime/Engine/Classes/PhysicsEngine/PhysicsThrusterComponent.h
@@ -8,7 +8,7 @@
  *	Used with objects that have physics to apply a force down the negative-X direction
  *	ie. point X in the direction you want the thrust in.
  */
-UCLASS(hidecategories=(Object, Mobility, LOD), ClassGroup=Physics, showcategories=Trigger, HeaderGroup=Component, meta=(BlueprintSpawnableComponent))
+UCLASS(hidecategories=(Object, Mobility, LOD), ClassGroup=Physics, showcategories=Trigger, meta=(BlueprintSpawnableComponent))
 class UPhysicsThrusterComponent : public USceneComponent
 {
 	GENERATED_UCLASS_BODY()
diff --git a/Engine/Source/Runtime/Engine/Classes/PhysicsEngine/RadialForceComponent.h b/Engine/Source/Runtime/Engine/Classes/PhysicsEngine/RadialForceComponent.h
--- a/Engine/Source/Runtime/Engine/Classes/PhysicsEngine/RadialForceComponent.h
+++ b/Engine/Source/Runtime/Engine/Classes/PhysicsEngine/RadialForceComponent.h
@@ -7,7 +7,7 @@
 /**
  *	Used to emit a radial force or impulse that can affect physics objects and or destructible objects.
  */
-UCLASS(hidecategories=(Object, Mobility), ClassGroup=Physics, showcategories=Trigger, HeaderGroup=Component, meta=(BlueprintSpawnableComponent), MinimalAPI)
+UCLASS(hidecategories=(Object, Mobility), ClassGroup=Physics, showcategories=Trigger, meta=(BlueprintSpawnableComponent), MinimalAPI)
 class URadialForceComponent : public USceneComponent
 {
 	GENERATED_UCLASS_BODY()
diff --git a/Engine/Source/Runtime/Engine/Classes/Vehicles/WheeledVehicleMovementComponent.h b/Engine/Source/Runtime/Engine/Classes/Vehicles/WheeledVehicleMovementComponent.h
--- a/Engine/Source/Runtime/Engine/Classes/Vehicles/WheeledVehicleMovementComponent.h
+++ b/Engine/Source/Runtime/Engine/Classes/Vehicles/WheeledVehicleMovementComponent.h
@@ -209,7 +209,7 @@ void ConvertCurve( const TArray<FFloatPair>& Pairs, struct FRichCurve& OutCurve
 /**
  * Component to handle the vehicle simulation for an actor.
  */
-UCLASS(Abstract, HeaderGroup=Component, dependson=UCurveBase, hidecategories=(PlanarMovement, "Components|Movement|Planar", Activation, "Components|Activation"))
+UCLASS(Abstract, dependson=UCurveBase, hidecategories=(PlanarMovement, "Components|Movement|Planar", Activation, "Components|Activation"))
 class ENGINE_API UWheeledVehicleMovementComponent : public UPawnMovementComponent
 {
 	GENERATED_UCLASS_BODY()
diff --git a/Engine/Source/Runtime/Engine/Private/AI/Navigation/NavigationOctree.cpp b/Engine/Source/Runtime/Engine/Private/AI/Navigation/NavigationOctree.cpp
--- a/Engine/Source/Runtime/Engine/Private/AI/Navigation/NavigationOctree.cpp
+++ b/Engine/Source/Runtime/Engine/Private/AI/Navigation/NavigationOctree.cpp
@@ -8,8 +8,6 @@
 #include "RecastNavMeshGenerator.h"
 #endif // NAVOCTREE_CONTAINS_COLLISION_DATA
 
-#include "EngineNavigationClasses.h"
-
 //----------------------------------------------------------------------//
 // FNavigationOctree
 //----------------------------------------------------------------------//
diff --git a/Engine/Source/Runtime/Engine/Private/Actor.cpp b/Engine/Source/Runtime/Engine/Private/Actor.cpp
--- a/Engine/Source/Runtime/Engine/Private/Actor.cpp
+++ b/Engine/Source/Runtime/Engine/Private/Actor.cpp
@@ -4,8 +4,6 @@
 #include "EnginePrivate.h"
 #include "ParticleDefinitions.h"
 #include "SoundDefinitions.h"
-#include "EngineNavigationClasses.h"
-#include "EngineKismetLibraryClasses.h"
 #include "LatentActions.h"
 #include "MessageLog.h"
 #include "Net/UnrealNetwork.h"
diff --git a/Engine/Source/Runtime/Engine/Private/AllEngineHeaders.h b/Engine/Source/Runtime/Engine/Private/AllEngineHeaders.h
--- a/Engine/Source/Runtime/Engine/Private/AllEngineHeaders.h
+++ b/Engine/Source/Runtime/Engine/Private/AllEngineHeaders.h
@@ -4,11 +4,4 @@
 #include "ParticleDefinitions.h"
 #include "SoundDefinitions.h"
 #include "BlueprintUtilities.h"
-#include "EngineSplineClasses.h"
-#include "EnginePlatformInterfaceClasses.h"
-#include "EngineFoliageClasses.h"
-#include "EngineLevelScriptClasses.h"
-#include "EngineDecalClasses.h"
 #include "NavigationPathBuilder.h"
-#include "EngineUserInterfaceClasses.h"
-#include "EngineKismetLibraryClasses.h"
diff --git a/Engine/Source/Runtime/Engine/Private/Animation/AnimNode_BlendSpaceEvaluator.cpp b/Engine/Source/Runtime/Engine/Private/Animation/AnimNode_BlendSpaceEvaluator.cpp
--- a/Engine/Source/Runtime/Engine/Private/Animation/AnimNode_BlendSpaceEvaluator.cpp
+++ b/Engine/Source/Runtime/Engine/Private/Animation/AnimNode_BlendSpaceEvaluator.cpp
@@ -1,8 +1,6 @@
 // Copyright 1998-2014 Epic Games, Inc. All Rights Reserved.
 
 #include "EnginePrivate.h"
-#include "EngineAnimClasses.h"
-#include "EngineAnimationNodeClasses.h"
 
 /////////////////////////////////////////////////////
 // FAnimNode_BlendSpaceEvaluator
diff --git a/Engine/Source/Runtime/Engine/Private/Animation/AnimNotifyState_Trail.cpp b/Engine/Source/Runtime/Engine/Private/Animation/AnimNotifyState_Trail.cpp
--- a/Engine/Source/Runtime/Engine/Private/Animation/AnimNotifyState_Trail.cpp
+++ b/Engine/Source/Runtime/Engine/Private/Animation/AnimNotifyState_Trail.cpp
@@ -1,10 +1,8 @@
 // Copyright 1998-2014 Epic Games, Inc. All Rights Reserved.
 
 #include "EnginePrivate.h"
-#include "EngineAnimClasses.h"
 #include "MessageLog.h"
 #include "ParticleDefinitions.h"
-#include "EngineKismetLibraryClasses.h"
 
 #define LOCTEXT_NAMESPACE "AnimNotifyState_Trail"
 
diff --git a/Engine/Source/Runtime/Engine/Private/CheatManager.cpp b/Engine/Source/Runtime/Engine/Private/CheatManager.cpp
--- a/Engine/Source/Runtime/Engine/Private/CheatManager.cpp
+++ b/Engine/Source/Runtime/Engine/Private/CheatManager.cpp
@@ -1,7 +1,75 @@
 // Copyright 1998-2014 Epic Games, Inc. All Rights Reserved.
 
 #include "EnginePrivate.h"
-#include "EngineComponentClasses.h"
+
+#include "Components/ApplicationLifecycleComponent.h"
+#include "AI/BehaviorTree/BlackboardComponent.h"
+#include "AI/BrainComponent.h"
+#include "AI/BehaviorTree/BehaviorTreeComponent.h"
+#include "Debug/GameplayDebuggingControllerComponent.h"
+#include "Components/InputComponent.h"
+#include "GameFramework/MovementComponent.h"
+#include "Components/SceneComponent.h"
+#include "Components/LightComponentBase.h"
+#include "Components/LightComponent.h"
+#include "Components/PrimitiveComponent.h"
+#include "GameFramework/NavMovementComponent.h"
+#include "GameFramework/PawnMovementComponent.h"
+#include "GameFramework/CharacterMovementComponent.h"
+#include "GameFramework/FloatingPawnMovement.h"
+#include "GameFramework/SpectatorPawnMovement.h"
+#include "Vehicles/WheeledVehicleMovementComponent.h"
+#include "Vehicles/WheeledVehicleMovementComponent4W.h"
+#include "GameFramework/ProjectileMovementComponent.h"
+#include "GameFramework/RotatingMovementComponent.h"
+#include "AI/Navigation/NavigationComponent.h"
+#include "AI/Navigation/PathFollowingComponent.h"
+#include "Components/PawnNoiseEmitterComponent.h"
+#include "Components/PawnSensingComponent.h"
+#include "PhysicsEngine/PhysicsHandleComponent.h"
+#include "Atmosphere/AtmosphericFogComponent.h"
+#include "Components/AudioComponent.h"
+#include "Camera/CameraComponent.h"
+#include "Components/ChildActorComponent.h"
+#include "Components/ExponentialHeightFogComponent.h"
+#include "Components/SkyLightComponent.h"
+#include "AI/Navigation/NavigationGraphNodeComponent.h"
+#include "PhysicsEngine/PhysicsConstraintComponent.h"
+#include "PhysicsEngine/PhysicsThrusterComponent.h"
+#include "Components/PostProcessComponent.h"
+#include "Components/ArrowComponent.h"
+#include "Components/BillboardComponent.h"
+#include "Components/BrushComponent.h"
+#include "Components/DrawFrustumComponent.h"
+#include "AI/EnvironmentQuery/EQSRenderingComponent.h"
+#include "Debug/GameplayDebuggingComponent.h"
+#include "Components/LineBatchComponent.h"
+#include "Components/MaterialBillboardComponent.h"
+#include "Components/MeshComponent.h"
+#include "Components/SkinnedMeshComponent.h"
+#include "Components/DestructibleComponent.h"
+#include "Components/PoseableMeshComponent.h"
+#include "Components/SkeletalMeshComponent.h"
+#include "Components/StaticMeshComponent.h"
+#include "Components/InstancedStaticMeshComponent.h"
+#include "Components/InteractiveFoliageComponent.h"
+#include "Components/SplineMeshComponent.h"
+#include "Components/ModelComponent.h"
+#include "AI/Navigation/NavLinkRenderingComponent.h"
+#include "AI/Navigation/NavMeshRenderingComponent.h"
+#include "AI/Navigation/NavTestRenderingComponent.h"
+#include "Components/NiagaraComponent.h"
+#include "Components/ShapeComponent.h"
+#include "Components/BoxComponent.h"
+#include "Components/CapsuleComponent.h"
+#include "Components/SphereComponent.h"
+#include "Components/DrawSphereComponent.h"
+#include "Components/TextRenderComponent.h"
+#include "Components/VectorFieldComponent.h"
+#include "PhysicsEngine/RadialForceComponent.h"
+#include "GameFramework/SpringArmComponent.h"
+#include "Components/WindDirectionalSourceComponent.h"
+#include "Components/TimelineComponent.h"
 #include "Slate.h"
 #include "AI/NavDataGenerator.h"
 #include "OnlineSubsystemUtils.h"
diff --git a/Engine/Source/Runtime/Engine/Private/ComponentInstanceDataCache.cpp b/Engine/Source/Runtime/Engine/Private/ComponentInstanceDataCache.cpp
--- a/Engine/Source/Runtime/Engine/Private/ComponentInstanceDataCache.cpp
+++ b/Engine/Source/Runtime/Engine/Private/ComponentInstanceDataCache.cpp
@@ -1,7 +1,78 @@
 // Copyright 1998-2014 Epic Games, Inc. All Rights Reserved.
 
 #include "EnginePrivate.h"
-#include "EngineComponentClasses.h"
+
+#include "Components/ApplicationLifecycleComponent.h"
+#include "AI/BehaviorTree/BlackboardComponent.h"
+#include "AI/BrainComponent.h"
+#include "AI/BehaviorTree/BehaviorTreeComponent.h"
+#include "Debug/GameplayDebuggingControllerComponent.h"
+#include "Components/InputComponent.h"
+#include "GameFramework/MovementComponent.h"
+#include "Components/SceneComponent.h"
+#include "Components/LightComponentBase.h"
+#include "Components/LightComponent.h"
+#include "Components/PrimitiveComponent.h"
+#include "GameFramework/NavMovementComponent.h"
+#include "GameFramework/PawnMovementComponent.h"
+#include "GameFramework/CharacterMovementComponent.h"
+#include "GameFramework/FloatingPawnMovement.h"
+#include "GameFramework/SpectatorPawnMovement.h"
+#include "Vehicles/WheeledVehicleMovementComponent.h"
+#include "Vehicles/WheeledVehicleMovementComponent4W.h"
+#include "GameFramework/ProjectileMovementComponent.h"
+#include "GameFramework/RotatingMovementComponent.h"
+#include "AI/Navigation/NavigationComponent.h"
+#include "AI/Navigation/PathFollowingComponent.h"
+#include "Components/PawnNoiseEmitterComponent.h"
+#include "Components/PawnSensingComponent.h"
+#include "PhysicsEngine/PhysicsHandleComponent.h"
+#include "Atmosphere/AtmosphericFogComponent.h"
+#include "Components/AudioComponent.h"
+#include "Camera/CameraComponent.h"
+#include "Components/ChildActorComponent.h"
+#include "Components/ExponentialHeightFogComponent.h"
+#include "Components/SkyLightComponent.h"
+#include "AI/Navigation/NavigationGraphNodeComponent.h"
+#include "PhysicsEngine/PhysicsConstraintComponent.h"
+#include "PhysicsEngine/PhysicsThrusterComponent.h"
+#include "Components/PostProcessComponent.h"
+#include "Components/ArrowComponent.h"
+#include "Components/BillboardComponent.h"
+#include "Components/BrushComponent.h"
+#include "Components/DrawFrustumComponent.h"
+#include "AI/EnvironmentQuery/EQSRenderingComponent.h"
+#include "Debug/GameplayDebuggingComponent.h"
+#include "Components/LineBatchComponent.h"
+#include "Components/MaterialBillboardComponent.h"
+#include "Components/MeshComponent.h"
+#include "Components/SkinnedMeshComponent.h"
+#include "Components/DestructibleComponent.h"
+#include "Components/PoseableMeshComponent.h"
+#include "Components/SkeletalMeshComponent.h"
+#include "Components/StaticMeshComponent.h"
+#include "Components/InstancedStaticMeshComponent.h"
+#include "Components/InteractiveFoliageComponent.h"
+#include "Components/SplineMeshComponent.h"
+#include "Components/ModelComponent.h"
+#include "AI/Navigation/NavLinkRenderingComponent.h"
+#include "AI/Navigation/NavMeshRenderingComponent.h"
+#include "AI/Navigation/NavTestRenderingComponent.h"
+#include "Components/NiagaraComponent.h"
+#include "Components/ShapeComponent.h"
+#include "Components/BoxComponent.h"
+#include "Components/CapsuleComponent.h"
+#include "Components/SphereComponent.h"
+#include "Components/DrawSphereComponent.h"
+#include "Components/TextRenderComponent.h"
+#include "Components/VectorFieldComponent.h"
+#include "PhysicsEngine/RadialForceComponent.h"
+#include "GameFramework/SpringArmComponent.h"
+#include "Components/WindDirectionalSourceComponent.h"
+#include "Components/TimelineComponent.h"
+#include "Slate.h"
+#include "AI/NavDataGenerator.h"
+#include "OnlineSubsystemUtils.h"
 
 void FComponentInstanceDataCache::AddInstanceData(TSharedPtr<FComponentInstanceDataBase> NewData)
 {
diff --git a/Engine/Source/Runtime/Engine/Private/FoliageComponent.cpp b/Engine/Source/Runtime/Engine/Private/FoliageComponent.cpp
--- a/Engine/Source/Runtime/Engine/Private/FoliageComponent.cpp
+++ b/Engine/Source/Runtime/Engine/Private/FoliageComponent.cpp
@@ -5,9 +5,6 @@
 =============================================================================*/
 
 #include "EnginePrivate.h"
-#include "EngineFoliageClasses.h"
-
-
 
 /** Scene proxy class for UInteractiveFoliageComponent. */
 class FInteractiveFoliageSceneProxy : public FStaticMeshSceneProxy
diff --git a/Engine/Source/Runtime/Engine/Private/GameEngine.cpp b/Engine/Source/Runtime/Engine/Private/GameEngine.cpp
--- a/Engine/Source/Runtime/Engine/Private/GameEngine.cpp
+++ b/Engine/Source/Runtime/Engine/Private/GameEngine.cpp
@@ -5,18 +5,14 @@
 =============================================================================*/
 
 #include "EnginePrivate.h"
-#include "EngineUserInterfaceClasses.h"
 #include "ParticleDefinitions.h"
 #include "SoundDefinitions.h"
-#include "EnginePlatformInterfaceClasses.h"
-#include "EngineDecalClasses.h"
 #include "Net/UnrealNetwork.h"
 #include "AllocatorFixedSizeFreeList.h"
 #include "Database.h"
 #include "MallocProfiler.h"
 #include "Net/NetworkProfiler.h"
 #include "ConfigCacheIni.h"
-#include "EngineLevelScriptClasses.h"
 
 #include "AVIWriter.h"
 #include "StreamingPauseRendering.h"
diff --git a/Engine/Source/Runtime/Engine/Private/GameViewportClient.cpp b/Engine/Source/Runtime/Engine/Private/GameViewportClient.cpp
--- a/Engine/Source/Runtime/Engine/Private/GameViewportClient.cpp
+++ b/Engine/Source/Runtime/Engine/Private/GameViewportClient.cpp
@@ -1,8 +1,6 @@
 // Copyright 1998-2014 Epic Games, Inc. All Rights Reserved.
 
 #include "EnginePrivate.h"
-#include "EngineUserInterfaceClasses.h"
-#include "EngineKismetLibraryClasses.h"
 #include "ParticleDefinitions.h"
 #include "SoundDefinitions.h"
 #include "FXSystem.h"
diff --git a/Engine/Source/Runtime/Engine/Private/GameplayStatics.cpp b/Engine/Source/Runtime/Engine/Private/GameplayStatics.cpp
--- a/Engine/Source/Runtime/Engine/Private/GameplayStatics.cpp
+++ b/Engine/Source/Runtime/Engine/Private/GameplayStatics.cpp
@@ -1,8 +1,6 @@
 // Copyright 1998-2014 Epic Games, Inc. All Rights Reserved.
 
 #include "EnginePrivate.h"
-#include "EngineKismetLibraryClasses.h"
-#include "EngineDecalClasses.h"
 #include "ParticleDefinitions.h"
 #include "SoundDefinitions.h"
 #include "PlatformFeatures.h"
diff --git a/Engine/Source/Runtime/Engine/Private/InputComponent.cpp b/Engine/Source/Runtime/Engine/Private/InputComponent.cpp
--- a/Engine/Source/Runtime/Engine/Private/InputComponent.cpp
+++ b/Engine/Source/Runtime/Engine/Private/InputComponent.cpp
@@ -1,8 +1,6 @@
 // Copyright 1998-2014 Epic Games, Inc. All Rights Reserved.
 
 #include "EnginePrivate.h"
-#include "EngineUserInterfaceClasses.h"
-#include "EngineLevelScriptClasses.h"		// For ALevelScriptActor
 
 FInputChord::RelationshipType FInputChord::GetRelationship(const FInputChord& OtherChord) const
 {
diff --git a/Engine/Source/Runtime/Engine/Private/InstancedStaticMesh.cpp b/Engine/Source/Runtime/Engine/Private/InstancedStaticMesh.cpp
--- a/Engine/Source/Runtime/Engine/Private/InstancedStaticMesh.cpp
+++ b/Engine/Source/Runtime/Engine/Private/InstancedStaticMesh.cpp
@@ -6,7 +6,79 @@
 
 #include "EnginePrivate.h"
 #include "ShaderParameters.h"
-#include "EngineComponentClasses.h"
+
+#include "Components/ApplicationLifecycleComponent.h"
+#include "AI/BehaviorTree/BlackboardComponent.h"
+#include "AI/BrainComponent.h"
+#include "AI/BehaviorTree/BehaviorTreeComponent.h"
+#include "Debug/GameplayDebuggingControllerComponent.h"
+#include "Components/InputComponent.h"
+#include "GameFramework/MovementComponent.h"
+#include "Components/SceneComponent.h"
+#include "Components/LightComponentBase.h"
+#include "Components/LightComponent.h"
+#include "Components/PrimitiveComponent.h"
+#include "GameFramework/NavMovementComponent.h"
+#include "GameFramework/PawnMovementComponent.h"
+#include "GameFramework/CharacterMovementComponent.h"
+#include "GameFramework/FloatingPawnMovement.h"
+#include "GameFramework/SpectatorPawnMovement.h"
+#include "Vehicles/WheeledVehicleMovementComponent.h"
+#include "Vehicles/WheeledVehicleMovementComponent4W.h"
+#include "GameFramework/ProjectileMovementComponent.h"
+#include "GameFramework/RotatingMovementComponent.h"
+#include "AI/Navigation/NavigationComponent.h"
+#include "AI/Navigation/PathFollowingComponent.h"
+#include "Components/PawnNoiseEmitterComponent.h"
+#include "Components/PawnSensingComponent.h"
+#include "PhysicsEngine/PhysicsHandleComponent.h"
+#include "Atmosphere/AtmosphericFogComponent.h"
+#include "Components/AudioComponent.h"
+#include "Camera/CameraComponent.h"
+#include "Components/ChildActorComponent.h"
+#include "Components/ExponentialHeightFogComponent.h"
+#include "Components/SkyLightComponent.h"
+#include "AI/Navigation/NavigationGraphNodeComponent.h"
+#include "PhysicsEngine/PhysicsConstraintComponent.h"
+#include "PhysicsEngine/PhysicsThrusterComponent.h"
+#include "Components/PostProcessComponent.h"
+#include "Components/ArrowComponent.h"
+#include "Components/BillboardComponent.h"
+#include "Components/BrushComponent.h"
+#include "Components/DrawFrustumComponent.h"
+#include "AI/EnvironmentQuery/EQSRenderingComponent.h"
+#include "Debug/GameplayDebuggingComponent.h"
+#include "Components/LineBatchComponent.h"
+#include "Components/MaterialBillboardComponent.h"
+#include "Components/MeshComponent.h"
+#include "Components/SkinnedMeshComponent.h"
+#include "Components/DestructibleComponent.h"
+#include "Components/PoseableMeshComponent.h"
+#include "Components/SkeletalMeshComponent.h"
+#include "Components/StaticMeshComponent.h"
+#include "Components/InstancedStaticMeshComponent.h"
+#include "Components/InteractiveFoliageComponent.h"
+#include "Components/SplineMeshComponent.h"
+#include "Components/ModelComponent.h"
+#include "AI/Navigation/NavLinkRenderingComponent.h"
+#include "AI/Navigation/NavMeshRenderingComponent.h"
+#include "AI/Navigation/NavTestRenderingComponent.h"
+#include "Components/NiagaraComponent.h"
+#include "Components/ShapeComponent.h"
+#include "Components/BoxComponent.h"
+#include "Components/CapsuleComponent.h"
+#include "Components/SphereComponent.h"
+#include "Components/DrawSphereComponent.h"
+#include "Components/TextRenderComponent.h"
+#include "Components/VectorFieldComponent.h"
+#include "PhysicsEngine/RadialForceComponent.h"
+#include "GameFramework/SpringArmComponent.h"
+#include "Components/WindDirectionalSourceComponent.h"
+#include "Components/TimelineComponent.h"
+#include "Slate.h"
+#include "AI/NavDataGenerator.h"
+#include "OnlineSubsystemUtils.h"
+
 #include "StaticMeshLight.h"
 
 #if WITH_PHYSX
diff --git a/Engine/Source/Runtime/Engine/Private/Interpolation.cpp b/Engine/Source/Runtime/Engine/Private/Interpolation.cpp
--- a/Engine/Source/Runtime/Engine/Private/Interpolation.cpp
+++ b/Engine/Source/Runtime/Engine/Private/Interpolation.cpp
@@ -6,16 +6,11 @@
 
 #include "EnginePrivate.h"
 #include "EngineClasses.h"
-#include "EngineInterpolationClasses.h"
-#include "EngineKismetLibraryClasses.h"
 #include "ParticleDefinitions.h"
-#include "EngineMaterialClasses.h"
-#include "EngineLevelScriptClasses.h"
 #include "SoundDefinitions.h"
 #include "InterpolationHitProxy.h"
 #include "AVIWriter.h"
 #include "AnimationUtils.h"
-#include "EngineDecalClasses.h"
 #include "MatineeUtils.h"
 
 #if WITH_EDITOR
diff --git a/Engine/Source/Runtime/Engine/Private/KismetSystemLibrary.cpp b/Engine/Source/Runtime/Engine/Private/KismetSystemLibrary.cpp
--- a/Engine/Source/Runtime/Engine/Private/KismetSystemLibrary.cpp
+++ b/Engine/Source/Runtime/Engine/Private/KismetSystemLibrary.cpp
@@ -1,8 +1,6 @@
 // Copyright 1998-2014 Epic Games, Inc. All Rights Reserved.
 
 #include "EnginePrivate.h"
-#include "EngineKismetLibraryClasses.h"
-#include "EngineUserInterfaceClasses.h"
 #include "LatentActions.h"
 #include "DelayAction.h"
 #include "InterpolateComponentToAction.h"
diff --git a/Engine/Source/Runtime/Engine/Private/Landscape/LandscapeEdit.cpp b/Engine/Source/Runtime/Engine/Private/Landscape/LandscapeEdit.cpp
--- a/Engine/Source/Runtime/Engine/Private/Landscape/LandscapeEdit.cpp
+++ b/Engine/Source/Runtime/Engine/Private/Landscape/LandscapeEdit.cpp
@@ -5,7 +5,6 @@ LandscapeEdit.cpp: Landscape editing
 =============================================================================*/
 
 #include "EnginePrivate.h"
-#include "EngineFoliageClasses.h"
 #include "Landscape/LandscapeDataAccess.h"
 #include "Landscape/LandscapeEdit.h"
 #include "Landscape/LandscapeRender.h"
@@ -2677,7 +2676,6 @@ void ULandscapeInfo::ExportLayer(ULandscapeLayerInfoObject* LayerInfo, const FSt
 //////////////////////////////////////////////////////////////////////////
 // Apply splines
 //////////////////////////////////////////////////////////////////////////
-#include "EngineSplineClasses.h"
 #include "Raster.h"
 
 class FLandscapeSplineHeightsRasterPolicy
diff --git a/Engine/Source/Runtime/Engine/Private/Level.cpp b/Engine/Source/Runtime/Engine/Private/Level.cpp
--- a/Engine/Source/Runtime/Engine/Private/Level.cpp
+++ b/Engine/Source/Runtime/Engine/Private/Level.cpp
@@ -7,7 +7,6 @@ Level.cpp: Level-related functions
 #include "EnginePrivate.h"
 #include "Net/UnrealNetwork.h"
 #include "SoundDefinitions.h"
-#include "EngineMaterialClasses.h"
 #include "PrecomputedLightVolume.h"
 #include "TickTaskManagerInterface.h"
 #include "BlueprintUtilities.h"
@@ -15,7 +14,6 @@ Level.cpp: Level-related functions
 #include "Editor/UnrealEd/Public/Kismet2/KismetEditorUtilities.h"
 #include "Editor/UnrealEd/Public/Kismet2/BlueprintEditorUtils.h"
 #endif
-#include "EngineLevelScriptClasses.h"
 #include "Runtime/Engine/Classes/MovieScene/RuntimeMovieScenePlayerInterface.h"
 #include "LevelUtils.h"
 #include "TargetPlatform.h"
diff --git a/Engine/Source/Runtime/Engine/Private/LevelActor.cpp b/Engine/Source/Runtime/Engine/Private/LevelActor.cpp
--- a/Engine/Source/Runtime/Engine/Private/LevelActor.cpp
+++ b/Engine/Source/Runtime/Engine/Private/LevelActor.cpp
@@ -3,8 +3,6 @@
 
 #include "EnginePrivate.h"
 #include "Net/UnrealNetwork.h"
-#include "EngineLevelScriptClasses.h"
-#include "EngineKismetLibraryClasses.h"
 #include "Collision.h"
 
 #if WITH_PHYSX
diff --git a/Engine/Source/Runtime/Engine/Private/LevelScriptActor.cpp b/Engine/Source/Runtime/Engine/Private/LevelScriptActor.cpp
--- a/Engine/Source/Runtime/Engine/Private/LevelScriptActor.cpp
+++ b/Engine/Source/Runtime/Engine/Private/LevelScriptActor.cpp
@@ -2,10 +2,7 @@
 
 
 #include "EnginePrivate.h"
-#include "EngineLevelScriptClasses.h"
 #include "SoundDefinitions.h"
-#include "EngineUserInterfaceClasses.h"
-#include "EngineKismetLibraryClasses.h"
 #include "BlueprintUtilities.h"
 
 
diff --git a/Engine/Source/Runtime/Engine/Private/LevelTick.cpp b/Engine/Source/Runtime/Engine/Private/LevelTick.cpp
--- a/Engine/Source/Runtime/Engine/Private/LevelTick.cpp
+++ b/Engine/Source/Runtime/Engine/Private/LevelTick.cpp
@@ -8,8 +8,6 @@
 #include "Net/UnrealNetwork.h"
 #include "Collision.h"
 
-#include "EngineUserInterfaceClasses.h"
-#include "EngineDecalClasses.h"
 #include "ParticleDefinitions.h"
 //#include "SoundDefinitions.h"
 #include "FXSystem.h"
diff --git a/Engine/Source/Runtime/Engine/Private/Pawn.cpp b/Engine/Source/Runtime/Engine/Private/Pawn.cpp
--- a/Engine/Source/Runtime/Engine/Private/Pawn.cpp
+++ b/Engine/Source/Runtime/Engine/Private/Pawn.cpp
@@ -10,10 +10,7 @@
 #include "Net/UnrealNetwork.h"
 #include "ConfigCacheIni.h"
 #include "NavigationPathBuilder.h"
-#include "EngineInterpolationClasses.h"
-#include "EngineKismetLibraryClasses.h"
 #include "ParticleDefinitions.h"
-#include "EngineComponentClasses.h"
 #include "DisplayDebugHelpers.h"
 
 DEFINE_LOG_CATEGORY(LogDamage);
diff --git a/Engine/Source/Runtime/Engine/Private/PlayerController.cpp b/Engine/Source/Runtime/Engine/Private/PlayerController.cpp
--- a/Engine/Source/Runtime/Engine/Private/PlayerController.cpp
+++ b/Engine/Source/Runtime/Engine/Private/PlayerController.cpp
@@ -3,11 +3,7 @@
 #include "EnginePrivate.h"
 #include "Net/UnrealNetwork.h"
 #include "ConfigCacheIni.h"
-#include "EngineKismetLibraryClasses.h"
 #include "SoundDefinitions.h"
-#include "EngineUserInterfaceClasses.h"
-#include "EngineInterpolationClasses.h"
-#include "EngineLevelScriptClasses.h"
 #include "OnlineSubsystemUtils.h"
 #include "IHeadMountedDisplay.h"
 #include "IForceFeedbackSystem.h"
diff --git a/Engine/Source/Runtime/Engine/Private/UserInterface/Console.cpp b/Engine/Source/Runtime/Engine/Private/UserInterface/Console.cpp
--- a/Engine/Source/Runtime/Engine/Private/UserInterface/Console.cpp
+++ b/Engine/Source/Runtime/Engine/Private/UserInterface/Console.cpp
@@ -5,8 +5,6 @@
 =============================================================================*/
 
 #include "EnginePrivate.h"
-#include "EngineUserInterfaceClasses.h"
-#include "EngineLevelScriptClasses.h"
 #include "Slate.h"
 #include "DefaultValueHelper.h"
 
diff --git a/Engine/Source/Runtime/Engine/Private/World.cpp b/Engine/Source/Runtime/Engine/Private/World.cpp
--- a/Engine/Source/Runtime/Engine/Private/World.cpp
+++ b/Engine/Source/Runtime/Engine/Private/World.cpp
@@ -6,8 +6,6 @@
 
 #include "EnginePrivate.h"
 #include "Net/UnrealNetwork.h"
-#include "EngineLevelScriptClasses.h"
-#include "EngineUserInterfaceClasses.h"
 #include "ParticleDefinitions.h"
 #include "Database.h"
 #include "Net/NetworkProfiler.h"
diff --git a/Engine/Source/Runtime/Engine/Public/Audio.h b/Engine/Source/Runtime/Engine/Public/Audio.h
--- a/Engine/Source/Runtime/Engine/Public/Audio.h
+++ b/Engine/Source/Runtime/Engine/Public/Audio.h
@@ -5,6 +5,7 @@
 =============================================================================*/
 
 #pragma once
+#include "Sound/SoundClass.h"
 
 ENGINE_API DECLARE_LOG_CATEGORY_EXTERN(LogAudio, Warning, All);
 
@@ -97,6 +98,8 @@ enum EAudioSpeakers
 
 // Forward declarations.
 class UAudioComponent;
+class USoundNode;
+struct FWaveInstance;
 struct FReverbSettings;
 struct FSampleLoop;
 
diff --git a/Engine/Source/Runtime/Engine/Public/Engine.h b/Engine/Source/Runtime/Engine/Public/Engine.h
--- a/Engine/Source/Runtime/Engine/Public/Engine.h
+++ b/Engine/Source/Runtime/Engine/Public/Engine.h
@@ -602,7 +602,6 @@ struct FURL;
 #include "RenderUtils.h"					// Render utility classes.
 #include "Distributions.h"					// Distributions
 #include "PhysxUserData.h"
-#include "EngineSceneClasses.h"
 #include "EngineClasses.h"					// All actor classes.
 #include "VisualLog.h"
 #include "MaterialShared.h"					// Shared material definitions.
@@ -615,10 +614,7 @@ struct FURL;
 #include "LightMap.h"						// Light-maps.
 #include "ShadowMap.h"
 #include "Model.h"							// Model class.
-#include "EngineComponentClasses.h"
-#include "EngineMaterialClasses.h"
-#include "EngineStaticMeshClasses.h"
-#include "EngineNavigationClasses.h"
+
 #include "AI/NavDataGenerator.h"
 #include "AI/NavLinkRenderingProxy.h"
 #include "AI/NavigationModifier.h"
@@ -626,16 +622,8 @@ struct FURL;
 #include "StaticMeshResources.h"			// Static meshes.
 #include "AnimTree.h"						// Animation.
 #include "SkeletalMeshTypes.h"				// Skeletal animated mesh.
-#include "EngineSkeletalMeshClasses.h"
-#include "EngineLightClasses.h"
 #include "Animation/SkeletalMeshActor.h"
 #include "Interpolation.h"					// Matinee.
-#include "EngineInterpolationClasses.h"
-#include "EnginePawnClasses.h"
-#include "EngineReplicationInfoClasses.h"
-#include "EngineGameEngineClasses.h"		// Main Unreal engine declarations
-#include "EngineTerrainClasses.h"
-#include "EngineNetworkClasses.h"
 #include "ContentStreaming.h"				// Content streaming class definitions.
 #include "LightingBuildOptions.h"			// Definition of lighting build option struct.
 #include "PixelFormat.h"
@@ -650,8 +638,6 @@ struct FURL;
 #include "TimerManager.h"					// Game play timers
 #include "EngineService.h"
 #include "AI/NavigationSystemHelpers.h"
-#include "EngineAnimationNodeClasses.h"
-#include "EngineNavigationClasses.h"
 #include "HardwareInfo.h"
 
 /** Implements the engine module. */
diff --git a/Engine/Source/Runtime/Engine/Public/InstancedFoliage.h b/Engine/Source/Runtime/Engine/Public/InstancedFoliage.h
--- a/Engine/Source/Runtime/Engine/Public/InstancedFoliage.h
+++ b/Engine/Source/Runtime/Engine/Public/InstancedFoliage.h
@@ -6,6 +6,12 @@
 
 #pragma once
 
+#include "Components/SceneComponent.h"
+//
+// Forward declarations.
+//
+class UInstancedStaticMeshComponent;
+
 /**
  * Flags stored with each instance
  */
diff --git a/Engine/Source/Runtime/Engine/Public/LightMap.h b/Engine/Source/Runtime/Engine/Public/LightMap.h
--- a/Engine/Source/Runtime/Engine/Public/LightMap.h
+++ b/Engine/Source/Runtime/Engine/Public/LightMap.h
@@ -7,6 +7,7 @@
 #pragma once
 
 #include "ChunkedArray.h"
+#include "Engine/LightMapTexture2D.h"
 
 /** Whether to use bilinear filtering on lightmaps */
 extern ENGINE_API bool GUseBilinearLightmaps;
diff --git a/Engine/Source/Runtime/Engine/Public/ParticleDefinitions.h b/Engine/Source/Runtime/Engine/Public/ParticleDefinitions.h
--- a/Engine/Source/Runtime/Engine/Public/ParticleDefinitions.h
+++ b/Engine/Source/Runtime/Engine/Public/ParticleDefinitions.h
@@ -9,4 +9,3 @@
 #include "Engine.h"
 #include "ParticleHelper.h"
 #include "ParticleEmitterInstances.h"
-#include "EngineParticleClasses.h"
diff --git a/Engine/Source/Runtime/Engine/Public/ParticleHelper.h b/Engine/Source/Runtime/Engine/Public/ParticleHelper.h
--- a/Engine/Source/Runtime/Engine/Public/ParticleHelper.h
+++ b/Engine/Source/Runtime/Engine/Public/ParticleHelper.h
@@ -125,6 +125,8 @@ class UParticleModuleOrientationAxisLock;
 
 class UParticleLODLevel;
 
+class USkeletalMeshComponent;
+
 class FParticleSystemSceneProxy;
 class FParticleDynamicData;
 struct FDynamicBeam2EmitterData;
@@ -134,6 +136,8 @@ struct FParticleSpriteEmitterInstance;
 struct FParticleMeshEmitterInstance;
 struct FParticleBeam2EmitterInstance;
 
+struct FStaticMeshLODResources;
+
 // Special module indices...
 #define INDEX_TYPEDATAMODULE	(INDEX_NONE - 1)
 #define INDEX_REQUIREDMODULE	(INDEX_NONE - 2)
diff --git a/Engine/Source/Runtime/Engine/Public/SceneManagement.h b/Engine/Source/Runtime/Engine/Public/SceneManagement.h
--- a/Engine/Source/Runtime/Engine/Public/SceneManagement.h
+++ b/Engine/Source/Runtime/Engine/Public/SceneManagement.h
@@ -9,7 +9,6 @@
 // Includes the draw mesh macros
 #include "SceneUtils.h"
 #include "UniformBuffer.h"
-#include "EngineSceneClasses.h"
 #include "BufferVisualizationData.h"
 #include "ConvexVolume.h"
 #include "SystemSettings.h"
diff --git a/Engine/Source/Runtime/Online/OnlineSubsystemUtils/Private/OnlineSessionClient.cpp b/Engine/Source/Runtime/Online/OnlineSubsystemUtils/Private/OnlineSessionClient.cpp
--- a/Engine/Source/Runtime/Online/OnlineSubsystemUtils/Private/OnlineSessionClient.cpp
+++ b/Engine/Source/Runtime/Online/OnlineSubsystemUtils/Private/OnlineSessionClient.cpp
@@ -6,7 +6,6 @@
 =============================================================================*/
 
 #include "OnlineSubsystemUtilsPrivatePCH.h"
-#include "EngineNetworkClasses.h"
 
 #define INVALID_CONTROLLERID 255
 
diff --git a/Engine/Source/Runtime/SkillSystem/Classes/AttributeComponent.h b/Engine/Source/Runtime/SkillSystem/Classes/AttributeComponent.h
--- a/Engine/Source/Runtime/SkillSystem/Classes/AttributeComponent.h
+++ b/Engine/Source/Runtime/SkillSystem/Classes/AttributeComponent.h
@@ -21,7 +21,7 @@ struct SKILLSYSTEM_API FAttributeDefaults
 /** 
  * 
  */
-UCLASS(HeaderGroup=Component, ClassGroup=SkillSystem, hidecategories=(Object,LOD,Lighting,Transform,Sockets,TextureStreaming), editinlinenew, meta=(BlueprintSpawnableComponent))
+UCLASS(ClassGroup=SkillSystem, hidecategories=(Object,LOD,Lighting,Transform,Sockets,TextureStreaming), editinlinenew, meta=(BlueprintSpawnableComponent))
 class SKILLSYSTEM_API UAttributeComponent : public UActorComponent
 {
 	GENERATED_UCLASS_BODY()
diff --git a/Engine/Source/Runtime/SkillSystem/Private/SkillSystemModulePrivatePCH.h b/Engine/Source/Runtime/SkillSystem/Private/SkillSystemModulePrivatePCH.h
--- a/Engine/Source/Runtime/SkillSystem/Private/SkillSystemModulePrivatePCH.h
+++ b/Engine/Source/Runtime/SkillSystem/Private/SkillSystemModulePrivatePCH.h
@@ -14,7 +14,6 @@
 #include "Engine.h"
 #include "Engine/DataTable.h"
 #include "SkillSystemClasses.h"
-#include "SkillSystemComponentClasses.h"
 
 #include "SkillSystemModule.h"
 #include "SkillSystemLog.h"
