Particle parameter distributions can now be used in conjunction with the ColorOverLife module for GPU particles.

diff --git a/Engine/Shaders/ParticleGPUSpriteVertexFactory.usf b/Engine/Shaders/ParticleGPUSpriteVertexFactory.usf
--- a/Engine/Shaders/ParticleGPUSpriteVertexFactory.usf
+++ b/Engine/Shaders/ParticleGPUSpriteVertexFactory.usf
@@ -213,7 +213,7 @@ FVertexFactoryIntermediates GetVertexFactoryIntermediates( FVertexFactoryInput I
 		EmitterUniforms.ColorCurve.zw * RelativeTime;
 	float4 ColorCurveSample = Texture2DSampleLevel(CurveTexture, CurveTextureSampler, ColorCurveTexCoord, 0 );
 	float4 ColorScale = ColorCurveSample * EmitterUniforms.ColorScale + EmitterUniforms.ColorBias;
-	float4 Color = ColorScale * InitialColor * EmitterDynamicUniforms.ScaleColorOverLife;
+	float4 Color = ColorScale * InitialColor * EmitterDynamicUniforms.DynamicColor;
 
 	// Sample the curve containing misc. attributes.
 	float2 MiscCurveTexCoord = EmitterUniforms.MiscCurve.xy +
diff --git a/Engine/Source/Editor/DetailCustomizations/Private/DetailCustomizations.cpp b/Engine/Source/Editor/DetailCustomizations/Private/DetailCustomizations.cpp
--- a/Engine/Source/Editor/DetailCustomizations/Private/DetailCustomizations.cpp
+++ b/Engine/Source/Editor/DetailCustomizations/Private/DetailCustomizations.cpp
@@ -171,7 +171,6 @@ void FDetailCustomizationsModule::StartupModule()
 	RegisterCustomPropertyLayout( "CollisionProfile", FOnGetDetailCustomizationInstance::CreateStatic(&FCollisionProfileDetails::MakeInstance));
 	RegisterCustomPropertyLayout( "PhysicsSettings", FOnGetDetailCustomizationInstance::CreateStatic(&FPhysicsSettingsDetails::MakeInstance));
 
-	RegisterCustomPropertyLayout( "ParticleModuleColorOverLife", FOnGetDetailCustomizationInstance::CreateStatic( &FParticleModuleColorOverLifeDetails::MakeInstance ) );
 	RegisterCustomPropertyLayout( "ParticleModuleRequired", FOnGetDetailCustomizationInstance::CreateStatic( &FParticleModuleRequiredDetails::MakeInstance ) );
 	RegisterCustomPropertyLayout( "ParticleModuleSubUV", FOnGetDetailCustomizationInstance::CreateStatic( &FParticleModuleSubUVDetails::MakeInstance ) );
 	RegisterCustomPropertyLayout( "ParticleModuleAccelerationDrag", FOnGetDetailCustomizationInstance::CreateStatic( &FParticleModuleAccelerationDragDetails::MakeInstance ) );
diff --git a/Engine/Source/Editor/DetailCustomizations/Private/ParticleModuleDetails.cpp b/Engine/Source/Editor/DetailCustomizations/Private/ParticleModuleDetails.cpp
--- a/Engine/Source/Editor/DetailCustomizations/Private/ParticleModuleDetails.cpp
+++ b/Engine/Source/Editor/DetailCustomizations/Private/ParticleModuleDetails.cpp
@@ -91,24 +91,6 @@ void FParticleModuleRequiredDetails::CustomizeDetails( IDetailLayoutBuilder& Det
 
 //////////////////////////////////////////////////////////////////////////
 
-TSharedRef<IDetailCustomization> FParticleModuleColorOverLifeDetails::MakeInstance()
-{
-	return MakeShareable( new FParticleModuleColorOverLifeDetails );
-}
-
-void FParticleModuleColorOverLifeDetails::CustomizeDetails( IDetailLayoutBuilder& DetailBuilder )
-{
-	TArray<FString> DistributionsToRestrict;
-	DistributionsToRestrict.Add(TEXT("ColorOverLife.Distribution"));
-	DistributionsToRestrict.Add(TEXT("AlphaOverLife.Distribution"));
-
-	TRestrictionList RestrictionList;
-	RestrictionList.Add( GetDistributionsForGPURestriction() );
-	RestrictPropertiesOnGPUEmitter( DetailBuilder, DistributionsToRestrict, RestrictionList );
-}
-
-//////////////////////////////////////////////////////////////////////////
-
 TSharedRef<IDetailCustomization> FParticleModuleSubUVDetails::MakeInstance()
 {
 	return MakeShareable( new FParticleModuleSubUVDetails );
diff --git a/Engine/Source/Editor/DetailCustomizations/Private/ParticleModuleDetails.h b/Engine/Source/Editor/DetailCustomizations/Private/ParticleModuleDetails.h
--- a/Engine/Source/Editor/DetailCustomizations/Private/ParticleModuleDetails.h
+++ b/Engine/Source/Editor/DetailCustomizations/Private/ParticleModuleDetails.h
@@ -22,16 +22,6 @@ public:
 	virtual void CustomizeDetails( IDetailLayoutBuilder& DetailBuilder ) OVERRIDE;
 };
 
-class FParticleModuleColorOverLifeDetails : public FParticleModuleDetailsBase
-{
-public:
-	/** Makes a new instance of this detail layout class for a specific detail view requesting it */
-	static TSharedRef<IDetailCustomization> MakeInstance();
-
-	/** IDetailCustomization interface */
-	virtual void CustomizeDetails( IDetailLayoutBuilder& DetailBuilder ) OVERRIDE;
-};
-
 class FParticleModuleSubUVDetails : public FParticleModuleDetailsBase
 {
 public:
diff --git a/Engine/Source/Runtime/Engine/Classes/Particles/Color/ParticleModuleColorOverLife.h b/Engine/Source/Runtime/Engine/Classes/Particles/Color/ParticleModuleColorOverLife.h
--- a/Engine/Source/Runtime/Engine/Classes/Particles/Color/ParticleModuleColorOverLife.h
+++ b/Engine/Source/Runtime/Engine/Classes/Particles/Color/ParticleModuleColorOverLife.h
@@ -39,14 +39,8 @@ class UParticleModuleColorOverLife : public UParticleModuleColorBase
 	virtual void Update(FParticleEmitterInstance* Owner, int32 Offset, float DeltaTime) OVERRIDE;
 	virtual void CompileModule( struct FParticleEmitterBuildInfo& EmitterInfo ) OVERRIDE;
 	virtual void SetToSensibleDefaults(UParticleEmitter* Owner) OVERRIDE;
-
-#if WITH_EDITOR
-	virtual bool IsValidForLODLevel(UParticleLODLevel* LODLevel, FString& OutErrorString) OVERRIDE;
-#endif
 	//End UParticleModule Interface
 
-protected:
-	friend class FParticleModuleColorOverLifeDetails;
 };
 
 
diff --git a/Engine/Source/Runtime/Engine/Classes/Particles/TypeData/ParticleModuleTypeDataGpu.h b/Engine/Source/Runtime/Engine/Classes/Particles/TypeData/ParticleModuleTypeDataGpu.h
--- a/Engine/Source/Runtime/Engine/Classes/Particles/TypeData/ParticleModuleTypeDataGpu.h
+++ b/Engine/Source/Runtime/Engine/Classes/Particles/TypeData/ParticleModuleTypeDataGpu.h
@@ -165,11 +165,19 @@ struct FGPUSpriteEmitterInfo
 	UPROPERTY()
 	uint32 bEnableCollision : 1;
 
-	/** Dynamic color scale from the ScaleColorOverLife module. */
+	/** Dynamic color scale from the ColorOverLife module. */
+	UPROPERTY()
+	FRawDistributionVector DynamicColor;
+
+	/** Dynamic alpha scale from the ColorOverLife module. */
+	UPROPERTY()
+	FRawDistributionFloat DynamicAlpha;
+
+	/** Dynamic color scale from the ColorScaleOverLife module. */
 	UPROPERTY()
 	FRawDistributionVector DynamicColorScale;
 
-	/* Dynamic alpha scale from the ScaleColorOverLife module.**/
+	/** Dynamic alpha scale from the ColorScaleOverLife module. */
 	UPROPERTY()
 	FRawDistributionFloat DynamicAlphaScale;
 
diff --git a/Engine/Source/Runtime/Engine/Private/Particles/ParticleGpuSimulation.cpp b/Engine/Source/Runtime/Engine/Private/Particles/ParticleGpuSimulation.cpp
--- a/Engine/Source/Runtime/Engine/Private/Particles/ParticleGpuSimulation.cpp
+++ b/Engine/Source/Runtime/Engine/Private/Particles/ParticleGpuSimulation.cpp
@@ -433,7 +433,7 @@ BEGIN_UNIFORM_BUFFER_STRUCT( FGPUSpriteEmitterDynamicUniformParameters, )
 	DECLARE_UNIFORM_BUFFER_STRUCT_MEMBER( FVector2D, LocalToWorldScale )
 	DECLARE_UNIFORM_BUFFER_STRUCT_MEMBER( FVector4, AxisLockRight )
 	DECLARE_UNIFORM_BUFFER_STRUCT_MEMBER( FVector4, AxisLockUp )
-	DECLARE_UNIFORM_BUFFER_STRUCT_MEMBER(FVector4, ScaleColorOverLife )
+	DECLARE_UNIFORM_BUFFER_STRUCT_MEMBER( FVector4, DynamicColor)
 END_UNIFORM_BUFFER_STRUCT( FGPUSpriteEmitterDynamicUniformParameters )
 
 IMPLEMENT_UNIFORM_BUFFER_STRUCT(FGPUSpriteEmitterDynamicUniformParameters,TEXT("EmitterDynamicUniforms"));
@@ -2889,15 +2889,28 @@ public:
 			DynamicData->EmitterDynamicParameters.AxisLockUp.W = 1.0f;
 		}
 
-		DynamicData->EmitterDynamicParameters.ScaleColorOverLife = FVector4(1.0f,1.0f,1.0f,1.0f);
+		
+		// Setup dynamic color parameter. Only set when using particle parameter distributions.
+		FVector4 ColorOverLife(1.0f, 1.0f, 1.0f, 1.0f);
+		FVector4 ColorScaleOverLife(1.0f, 1.0f, 1.0f, 1.0f);
 		if( EmitterInfo.DynamicColorScale.Distribution )
 		{
-			DynamicData->EmitterDynamicParameters.ScaleColorOverLife = EmitterInfo.DynamicColorScale.GetValue(0.0f,Component);
+			ColorScaleOverLife = EmitterInfo.DynamicColorScale.GetValue(0.0f,Component);
 		}
 		if( EmitterInfo.DynamicAlphaScale.Distribution )
 		{
-			DynamicData->EmitterDynamicParameters.ScaleColorOverLife.W = EmitterInfo.DynamicAlphaScale.GetValue(0.0f,Component);
+			ColorScaleOverLife.W = EmitterInfo.DynamicAlphaScale.GetValue(0.0f,Component);
+		}
+
+		if( EmitterInfo.DynamicColor.Distribution )
+		{
+			ColorOverLife = EmitterInfo.DynamicColor.GetValue(0.0f,Component);
+		}
+		if( EmitterInfo.DynamicAlpha.Distribution )
+		{
+			ColorOverLife.W = EmitterInfo.DynamicAlpha.GetValue(0.0f,Component);
 		}
+		DynamicData->EmitterDynamicParameters.DynamicColor = ColorOverLife * ColorScaleOverLife;
 
 		const bool bSimulateGPUParticles = 
 			FXConsoleVariables::bFreezeGPUSimulation == false &&
diff --git a/Engine/Source/Runtime/Engine/Private/Particles/ParticleModules.cpp b/Engine/Source/Runtime/Engine/Private/Particles/ParticleModules.cpp
--- a/Engine/Source/Runtime/Engine/Private/Particles/ParticleModules.cpp
+++ b/Engine/Source/Runtime/Engine/Private/Particles/ParticleModules.cpp
@@ -4452,6 +4452,9 @@ void UParticleModuleTypeDataGpu::Build( FParticleEmitterBuildInfo& EmitterBuildI
 
 	ResourceData.PivotOffset = EmitterBuildInfo.PivotOffset;
 
+    // Store color and scale when using particle parameters.
+	EmitterInfo.DynamicColor = EmitterBuildInfo.DynamicColor;
+	EmitterInfo.DynamicAlpha= EmitterBuildInfo.DynamicAlpha;
 	EmitterInfo.DynamicColorScale = EmitterBuildInfo.DynamicColorScale;
 	EmitterInfo.DynamicAlphaScale = EmitterBuildInfo.DynamicAlphaScale;
 
diff --git a/Engine/Source/Runtime/Engine/Private/Particles/ParticleModules_Color.cpp b/Engine/Source/Runtime/Engine/Private/Particles/ParticleModules_Color.cpp
--- a/Engine/Source/Runtime/Engine/Private/Particles/ParticleModules_Color.cpp
+++ b/Engine/Source/Runtime/Engine/Private/Particles/ParticleModules_Color.cpp
@@ -257,8 +257,42 @@ void UParticleModuleColorOverLife::Serialize(FArchive& Ar)
 
 void UParticleModuleColorOverLife::CompileModule( FParticleEmitterBuildInfo& EmitterInfo )
 {
-	EmitterInfo.ColorScale.Initialize( ColorOverLife.Distribution );
-	EmitterInfo.AlphaScale.Initialize( AlphaOverLife.Distribution );
+	bool ScaleColor = true;
+	bool ScaleAlpha = true;
+	if(IsUsedInGPUEmitter())
+	{
+		if( ColorOverLife.Distribution->IsA( UDistributionVectorParticleParameter::StaticClass() ) )
+		{
+			EmitterInfo.DynamicColor = ColorOverLife;
+#if WITH_EDITOR
+			EmitterInfo.DynamicColor.Distribution->bIsDirty = true;
+			EmitterInfo.DynamicColor.Initialize();
+#endif
+			ScaleColor = false;
+			EmitterInfo.ColorScale.InitializeWithConstant( FVector(1.0f, 1.0f, 1.0f) );
+		}
+
+		if( AlphaOverLife.Distribution->IsA( UDistributionFloatParticleParameter::StaticClass() ) )
+		{
+			EmitterInfo.DynamicAlpha = AlphaOverLife;
+#if WITH_EDITOR
+			EmitterInfo.DynamicAlpha.Distribution->bIsDirty = true;
+			EmitterInfo.DynamicAlpha.Initialize();
+#endif
+			ScaleAlpha = false;
+			EmitterInfo.AlphaScale.InitializeWithConstant(1.0f);
+		}
+	}
+
+	if( ScaleColor )
+	{
+		EmitterInfo.ColorScale.Initialize( ColorOverLife.Distribution );
+	}
+
+	if( ScaleAlpha )
+	{
+		EmitterInfo.AlphaScale.Initialize( AlphaOverLife.Distribution );
+	}
 }
 
 #if WITH_EDITOR
@@ -432,28 +466,6 @@ void UParticleModuleColorOverLife::SetToSensibleDefaults(UParticleEmitter* Owner
 	}
 }
 
-#if WITH_EDITOR
-bool UParticleModuleColorOverLife::IsValidForLODLevel(UParticleLODLevel* LODLevel, FString& OutErrorString)
-{
-	if (LODLevel->TypeDataModule && LODLevel->TypeDataModule->IsA(UParticleModuleTypeDataGpu::StaticClass()))
-	{
-		if(!IsDistributionAllowedOnGPU(ColorOverLife.Distribution))
-		{
-			OutErrorString = GetDistributionNotAllowedOnGPUText(StaticClass()->GetName(), "ColorOverLife" ).ToString();
-			return false;
-		}
-
-		if(!IsDistributionAllowedOnGPU(AlphaOverLife.Distribution))
-		{
-			OutErrorString = GetDistributionNotAllowedOnGPUText(StaticClass()->GetName(), "AlphaOverLife" ).ToString();
-			return false;
-		}
-	}
-
-	return true;
-}
-#endif
-
 /*-----------------------------------------------------------------------------
 	UParticleModuleColorScaleOverLife implementation.
 -----------------------------------------------------------------------------*/
diff --git a/Engine/Source/Runtime/Engine/Public/ParticleEmitterInstances.h b/Engine/Source/Runtime/Engine/Public/ParticleEmitterInstances.h
--- a/Engine/Source/Runtime/Engine/Public/ParticleEmitterInstances.h
+++ b/Engine/Source/Runtime/Engine/Public/ParticleEmitterInstances.h
@@ -77,6 +77,11 @@ struct FParticleEmitterBuildInfo
 	/** The alpha scale of a particle over time. */
 	FComposableFloatDistribution AlphaScale;
 
+	/** An additional color scale for allowing parameters to be used for ColorOverLife modules. */
+	FRawDistributionVector DynamicColor;
+	/** An additional alpha scale for allowing parameters to be used for ColorOverLife modules. */
+	FRawDistributionFloat DynamicAlpha;
+
 	/** An additional color scale for allowing parameters to be used for ColorScaleOverLife modules. */
 	FRawDistributionVector DynamicColorScale;
 	/** An additional alpha scale for allowing parameters to be used for ColorScaleOverLife modules. */
