Use sweep test to determine collision for mesh particles.
Corrected typo.

diff --git a/Engine/Source/Runtime/Engine/Private/Particles/ParticleModules_Collision.cpp b/Engine/Source/Runtime/Engine/Private/Particles/ParticleModules_Collision.cpp
--- a/Engine/Source/Runtime/Engine/Private/Particles/ParticleModules_Collision.cpp
+++ b/Engine/Source/Runtime/Engine/Private/Particles/ParticleModules_Collision.cpp
@@ -336,8 +336,15 @@ void UParticleModuleCollision::Update(FParticleEmitterInstance* Owner, int32 Off
 
 		// Determine the size
 		FVector Size = Particle.Size * ParentScale;
-
 		FVector	Extent(0.0f);
+
+		// Setup extent for mesh particles. 
+		UParticleModuleTypeDataMesh* MeshType = Cast<UParticleModuleTypeDataMesh>(LODLevel->TypeDataModule);
+		if (MeshType && MeshType->Mesh)
+		{
+			Extent = MeshType->Mesh->GetBounds().BoxExtent;
+		}
+		
 		FHitResult Hit;
 
 		Hit.Normal.X = 0.0f;
diff --git a/Engine/Source/Runtime/Engine/Public/ParticleEmitterInstances.h b/Engine/Source/Runtime/Engine/Public/ParticleEmitterInstances.h
--- a/Engine/Source/Runtime/Engine/Public/ParticleEmitterInstances.h
+++ b/Engine/Source/Runtime/Engine/Public/ParticleEmitterInstances.h
@@ -198,9 +198,9 @@ public:
 	int32 CurrentLODLevelIndex;
 	/** The currently set LOD level.									*/
 	UParticleLODLevel* CurrentLODLevel;
-	/** The offset to the TypeDate payload in the particle data.		*/
+	/** The offset to the TypeData payload in the particle data.		*/
 	int32 TypeDataOffset;
-	/** The offset to the TypeDate instance payload.					*/
+	/** The offset to the TypeData instance payload.					*/
 	int32 TypeDataInstanceOffset;
 	/** The offset to the SubUV payload in the particle data.			*/
 	int32 SubUVDataOffset;
