Fix issue that could cause partial bunch sequence id's to wrap incorrectly

diff --git a/Engine/Source/Runtime/Engine/Classes/Engine/NetConnection.h b/Engine/Source/Runtime/Engine/Classes/Engine/NetConnection.h
--- a/Engine/Source/Runtime/Engine/Classes/Engine/NetConnection.h
+++ b/Engine/Source/Runtime/Engine/Classes/Engine/NetConnection.h
@@ -212,7 +212,8 @@ public:
 	int32			InPacketId;				// Full incoming packet index.
 	int32			OutPacketId;			// Most recently sent packet.
 	int32 			OutAckPacketId;			// Most recently acked outgoing packet.
-	int32			PartialPackedId;		// Sequencing number for partial packets
+	int32			PartialPacketId;		// Sequencing number for partial packets
+	int32			LastPartialPacketId;	// Most recently received unrealiable partial packet id
 
 	uint32			PingAckDataCache[MAX_PACKETID/PING_ACK_PACKET_INTERVAL];	// Caches packet data on the server, for verifying pings
 	float			LastPingAck;												// The time of the most recent PingAck on the client
diff --git a/Engine/Source/Runtime/Engine/Private/DataChannel.cpp b/Engine/Source/Runtime/Engine/Private/DataChannel.cpp
--- a/Engine/Source/Runtime/Engine/Private/DataChannel.cpp
+++ b/Engine/Source/Runtime/Engine/Private/DataChannel.cpp
@@ -656,8 +656,8 @@ FPacketIdRange UChannel::SendBunch( FOutBunch* Bunch, bool Merge )
 		if (ThisOutBunch->bPartial && !ThisOutBunch->bReliable)
 		{
 			// If we are reliable, then partial bunches just use the reliable sequences
-			// if not reliable, we hijack ChSequence and use Connection->PartialPackedId to sequence these packets
-			ThisOutBunch->ChSequence = Connection->PartialPackedId++;
+			// if not reliable, we hijack ChSequence and use Connection->PartialPacketId to sequence these packets
+			ThisOutBunch->ChSequence = Connection->PartialPacketId++;
 		}
 
 		if (UE_LOG_ACTIVE(LogNetPartialBunch,Verbose) && (OutgoingBunches.Num() > 1)) // Don't want to call appMemcrc unless we need to
diff --git a/Engine/Source/Runtime/Engine/Private/NetConnection.cpp b/Engine/Source/Runtime/Engine/Private/NetConnection.cpp
--- a/Engine/Source/Runtime/Engine/Private/NetConnection.cpp
+++ b/Engine/Source/Runtime/Engine/Private/NetConnection.cpp
@@ -53,7 +53,8 @@ UNetConnection::UNetConnection(const class FPostConstructInitializeProperties& P
 ,	InPacketId			( -1 )
 ,	OutPacketId			( 0 ) // must be initialized as OutAckPacketId + 1 so loss of first packet can be detected
 ,	OutAckPacketId		( -1 )
-,	PartialPackedId		( 0 )
+,	PartialPacketId		( 0 )
+,	LastPartialPacketId	( -1 )
 ,	LastPingAck			( 0.f )
 ,	LastPingAckPacketId	( -1 )
 ,	ClientWorldPackageName( NAME_None )
@@ -677,10 +678,10 @@ void UNetConnection::ReceivedPacket( FBitReader& Reader )
 	LastReceiveTime = Driver->Time;
 
 	// Check packet ordering.
-	int32 PacketId = MakeRelative(Reader.ReadInt(MAX_PACKETID),InPacketId,MAX_PACKETID);
+	const int32 PacketId = MakeRelative(Reader.ReadInt(MAX_PACKETID),InPacketId,MAX_PACKETID);
 	if( PacketId > InPacketId )
 	{
-		int32 PacketsLost = PacketId - InPacketId - 1;
+		const int32 PacketsLost = PacketId - InPacketId - 1;
 		InPacketsLost += PacketsLost;
 		Driver->InPacketsLost += PacketsLost;
 		InPacketId = PacketId;
@@ -841,7 +842,33 @@ void UNetConnection::ReceivedPacket( FBitReader& Reader )
 			Bunch.ChIndex      = Reader.ReadInt( MAX_CHANNELS );
 			Bunch.bHasGUIDs	   = Reader.ReadBit();
 			Bunch.bPartial	   = Reader.ReadBit();
-			Bunch.ChSequence   = (Bunch.bReliable || Bunch.bPartial) ? MakeRelative(Reader.ReadInt(MAX_CHSEQUENCE),InReliable[Bunch.ChIndex],MAX_CHSEQUENCE) : 0;
+
+			if ( Bunch.bReliable )
+			{
+				// If this is a reliable bunch, use the last processed reliable sequence to read the new reliable sequence
+				Bunch.ChSequence = MakeRelative( Reader.ReadInt( MAX_CHSEQUENCE ), InReliable[Bunch.ChIndex], MAX_CHSEQUENCE );
+			} 
+			else if ( Bunch.bPartial )
+			{
+				// If this is a partial bunch, use the last processed partial sequence to read the new partial sequence
+				Bunch.ChSequence = MakeRelative( Reader.ReadInt( MAX_CHSEQUENCE ), LastPartialPacketId, MAX_CHSEQUENCE );
+				if ( Bunch.ChSequence > LastPartialPacketId )
+				{
+					LastPartialPacketId = Bunch.ChSequence;
+				}
+				else
+				{
+					// Well behaved clients should never hit this, since we've already thrown out packets that were older
+					// Since PartialPacketId should evolve in sync with PacketId, this should be impossible unless the 
+					// client is misbehaving...
+					UE_LOG( LogNetTraffic, Error, TEXT( "UNetConnection::ReceivedPacket: Invalid partial packet id." ) );
+				}
+			}
+			else
+			{
+				Bunch.ChSequence = 0;
+			}
+
 			Bunch.bPartialInitial = Bunch.bPartial ? Reader.ReadBit() : 0;
 			Bunch.bPartialFinal	  = Bunch.bPartial ? Reader.ReadBit() : 0;
 			Bunch.ChType       = (Bunch.bReliable||Bunch.bOpen) ? Reader.ReadInt(CHTYPE_MAX) : CHTYPE_None;
