Reverting back to adding sub-components in UCameraComponent's OnRegister method, but better handling cleaning them up when we reinstance the blueprint.
(TTP#333370)

diff --git a/Engine/Source/Editor/UnrealEd/Private/Kismet2/KismetReinstanceUtilities.cpp b/Engine/Source/Editor/UnrealEd/Private/Kismet2/KismetReinstanceUtilities.cpp
--- a/Engine/Source/Editor/UnrealEd/Private/Kismet2/KismetReinstanceUtilities.cpp
+++ b/Engine/Source/Editor/UnrealEd/Private/Kismet2/KismetReinstanceUtilities.cpp
@@ -322,6 +322,8 @@ void FBlueprintCompileReinstancer::ReplaceInstancesOfClass(UClass* OldClass, UCl
 				NewObject = NewActor;
 
 				OldActor->DestroyConstructedComponents(); // don't want to serialize components from the old actor
+				// Unregister native components so we don't copy any sub-components they generate for themselves (like UCameraComponent does)
+				OldActor->UnregisterAllComponents(); 
 
 				NewActor->UnregisterAllComponents(); // Unregister any native components, might have cached state based on properties we are going to overwrite
 
diff --git a/Engine/Source/Runtime/Engine/Classes/Camera/CameraComponent.h b/Engine/Source/Runtime/Engine/Classes/Camera/CameraComponent.h
--- a/Engine/Source/Runtime/Engine/Classes/Camera/CameraComponent.h
+++ b/Engine/Source/Runtime/Engine/Classes/Camera/CameraComponent.h
@@ -46,6 +46,7 @@ class UCameraComponent : public USceneComponent
 
 	// UActorComponent interface
 	ENGINE_API virtual void OnRegister() OVERRIDE;
+	ENGINE_API virtual void OnUnregister() OVERRIDE;
 #if WITH_EDITOR
 	ENGINE_API virtual void CheckForErrors() OVERRIDE;
 	// End of UActorComponent interface
@@ -66,14 +67,14 @@ protected:
 #if WITH_EDITORONLY_DATA
 	// The frustum component used to show visually where the camera field of view is
 	UPROPERTY(transient)
-	TSubobjectPtr<class UDrawFrustumComponent> DrawFrustum;
+	class UDrawFrustumComponent* DrawFrustum;
 
 	UPROPERTY(transient)
 	class UStaticMesh* CameraMesh;
 
 	// The camera mesh to show visually where the camera is placed
 	UPROPERTY(transient)
-	TSubobjectPtr<class UStaticMeshComponent> ProxyMeshComponent;
+	class UStaticMeshComponent* ProxyMeshComponent;
 
 public:
 	// Refreshes the visual components to match the component state
diff --git a/Engine/Source/Runtime/Engine/Private/Camera/CameraComponent.cpp b/Engine/Source/Runtime/Engine/Private/Camera/CameraComponent.cpp
--- a/Engine/Source/Runtime/Engine/Private/Camera/CameraComponent.cpp
+++ b/Engine/Source/Runtime/Engine/Private/Camera/CameraComponent.cpp
@@ -19,45 +19,71 @@ UCameraComponent::UCameraComponent(const class FPostConstructInitializePropertie
 		static ConstructorHelpers::FObjectFinder<UStaticMesh> EditorCameraMesh(TEXT("/Engine/EditorMeshes/MatineeCam_SM"));
 		CameraMesh = EditorCameraMesh.Object;
 	}
+#endif
 
-	ProxyMeshComponent = PCIP.CreateEditorOnlyDefaultSubobject<UStaticMeshComponent>(GetOuter(), TEXT("CameraProxy"), true);
-	if (ProxyMeshComponent)
+	FieldOfView = 90.0f;
+	AspectRatio = 1.777778f;
+	OrthoWidth = 512.0f;
+	bConstrainAspectRatio = false;
+	PostProcessBlendWeight = 1.0f;
+	bUseControllerViewRotation = true;
+	bAutoActivate = true;
+}
+
+void UCameraComponent::OnRegister()
+{
+#if WITH_EDITORONLY_DATA
+	if (ProxyMeshComponent == NULL)
 	{
-		ProxyMeshComponent->AttachParent = this;
+		ProxyMeshComponent = ConstructObject<UStaticMeshComponent>(UStaticMeshComponent::StaticClass(), GetOuter(), NAME_None, RF_Transactional);
+		ProxyMeshComponent->AttachTo(this);
 		ProxyMeshComponent->StaticMesh = CameraMesh;
 		ProxyMeshComponent->SetCollisionProfileName(UCollisionProfile::NoCollision_ProfileName);
 		ProxyMeshComponent->bHiddenInGame = true;
 		ProxyMeshComponent->CastShadow = false;
 		ProxyMeshComponent->PostPhysicsComponentTick.bCanEverTick = false;
 		ProxyMeshComponent->bCreatedByConstructionScript = bCreatedByConstructionScript;
+		ProxyMeshComponent->RegisterComponentWithWorld(GetWorld());
 	}
 
-	DrawFrustum = PCIP.CreateEditorOnlyDefaultSubobject<UDrawFrustumComponent>(GetOuter(), TEXT("FrustumVisualizer"), true);
-	if (DrawFrustum)
+	if (DrawFrustum == NULL)
 	{
-		DrawFrustum->AttachParent = this;
+		DrawFrustum = ConstructObject<UDrawFrustumComponent>(UDrawFrustumComponent::StaticClass(), GetOuter(), NAME_None, RF_Transactional);
+		DrawFrustum->AttachTo(this);
 		DrawFrustum->AlwaysLoadOnClient = false;
 		DrawFrustum->AlwaysLoadOnServer = false;
 		DrawFrustum->bCreatedByConstructionScript = bCreatedByConstructionScript;
+		DrawFrustum->RegisterComponentWithWorld(GetWorld());
 	}
+
+	RefreshVisualRepresentation();
 #endif
 
-	FieldOfView = 90.0f;
-	AspectRatio = 1.777778f;
-	OrthoWidth = 512.0f;
-	bConstrainAspectRatio = false;
-	PostProcessBlendWeight = 1.0f;
-	bUseControllerViewRotation = true;
-	bAutoActivate = true;
+	Super::OnRegister();
 }
 
-void UCameraComponent::OnRegister()
+void UCameraComponent::OnUnregister()
 {
+	Super::OnUnregister();
+
 #if WITH_EDITORONLY_DATA
-	RefreshVisualRepresentation();
-#endif
+	// have to removed the sub-components that we added in OnRegister (for 
+	// reinstancing, where we CopyPropertiesForUnrelatedObjects()... don't want
+	// these copied since we'll generate them on the next OnRegister)
+	if (ProxyMeshComponent != NULL)
+	{
+		ProxyMeshComponent->DetachFromParent();
+		ProxyMeshComponent->DestroyComponent();
+		ProxyMeshComponent = NULL;
+	}
 
-	Super::OnRegister();
+	if (DrawFrustum != NULL)
+	{
+		DrawFrustum->DetachFromParent();
+		DrawFrustum->DestroyComponent();
+		DrawFrustum = NULL;
+	}
+#endif
 }
 
 #if WITH_EDITORONLY_DATA
