Nan fix for vehicle tire shader

diff --git a/Engine/Source/Runtime/Engine/Private/Vehicles/WheeledVehicleMovementComponent.cpp b/Engine/Source/Runtime/Engine/Private/Vehicles/WheeledVehicleMovementComponent.cpp
--- a/Engine/Source/Runtime/Engine/Private/Vehicles/WheeledVehicleMovementComponent.cpp
+++ b/Engine/Source/Runtime/Engine/Private/Vehicles/WheeledVehicleMovementComponent.cpp
@@ -394,7 +394,7 @@ void PxVehicleComputeTireForceDefault
 	tireAlignMoment=0.0f;
 
 	//If long slip/lat slip/camber are all zero than there will be zero tire force.
-	if((0==latSlip)&&(0==longSlip)&&(0==camber))
+	if (FMath::IsNearlyZero(latSlip) && FMath::IsNearlyZero(longSlip) && FMath::IsNearlyZero(camber))
 	{
 		return;
 	}
