#UE4 added verbose logging of RPC calls during ServerNotifyLoadedWorld for clients
- disabled in shipping

diff --git a/Engine/Source/Runtime/Engine/Private/Actor.cpp b/Engine/Source/Runtime/Engine/Private/Actor.cpp
--- a/Engine/Source/Runtime/Engine/Private/Actor.cpp
+++ b/Engine/Source/Runtime/Engine/Private/Actor.cpp
@@ -14,6 +14,12 @@ DEFINE_LOG_CATEGORY(LogActor);
 
 DEFINE_STAT(STAT_GetComponentsTime);
 
+#if UE_BUILD_SHIPPING
+#define DEBUG_CALLSPACE(Format, ...)
+#else
+#define DEBUG_CALLSPACE(Format, ...) UE_LOG(LogNet, VeryVerbose, Format, __VA_ARGS__);
+#endif
+
 #define LOCTEXT_NAMESPACE "Actor"
 
 FMakeNoiseDelegate AActor::MakeNoiseDelegate = FMakeNoiseDelegate::CreateStatic(&AActor::MakeNoiseImpl);
@@ -2635,12 +2641,14 @@ int32 AActor::GetFunctionCallspace( UFunction* Function, void* Parameters, FFram
 	if ((Function->FunctionFlags & FUNC_Static))
 	{
 		// Call local
+		DEBUG_CALLSPACE(TEXT("GetFunctionCallspace Local1: %s"), *Function->GetName());
 		return FunctionCallspace::Local;
 	}
 
 	if (GAllowActorScriptExecutionInEditor)
 	{
 		// Call local
+		DEBUG_CALLSPACE(TEXT("GetFunctionCallspace Local2: %s"), *Function->GetName());
 		return FunctionCallspace::Local;
 	}
 
@@ -2648,6 +2656,7 @@ int32 AActor::GetFunctionCallspace( UFunction* Function, void* Parameters, FFram
 	if (!World)
 	{
 		// Call local
+		DEBUG_CALLSPACE(TEXT("GetFunctionCallspace Local3: %s"), *Function->GetName());
 		return FunctionCallspace::Local;
 	}
 
@@ -2658,12 +2667,14 @@ int32 AActor::GetFunctionCallspace( UFunction* Function, void* Parameters, FFram
 	{
 		// Never call remote on a pending kill actor. 
 		// We can call it local or absorb it depending on authority/role check above.
+		DEBUG_CALLSPACE(TEXT("GetFunctionCallspace: IsPendingKill %s %s"), *Function->GetName(), FunctionCallspace::ToString(Callspace));
 		return Callspace;
 	}
 
 	if (Function->FunctionFlags & FUNC_NetRequest)
 	{
 		// Call remote
+		DEBUG_CALLSPACE(TEXT("GetFunctionCallspace NetRequest: %s"), *Function->GetName());
 		return FunctionCallspace::Remote;
 	}	
 	
@@ -2672,10 +2683,12 @@ int32 AActor::GetFunctionCallspace( UFunction* Function, void* Parameters, FFram
 		if (Function->RPCId > 0)
 		{
 			// Call local
+			DEBUG_CALLSPACE(TEXT("GetFunctionCallspace NetResponse Local: %s"), *Function->GetName());
 			return FunctionCallspace::Local;
 		}
 
 		// Shouldn't happen, so skip call
+		DEBUG_CALLSPACE(TEXT("GetFunctionCallspace NetResponse Absorbed: %s"), *Function->GetName());
 		return FunctionCallspace::Absorbed;
 	}
 
@@ -2686,6 +2699,7 @@ int32 AActor::GetFunctionCallspace( UFunction* Function, void* Parameters, FFram
 		if (Role < ROLE_Authority && (Function->FunctionFlags & FUNC_NetServer))
 		{
 			// Don't let clients call server functions (in edge cases where NetMode is standalone (net driver is null))
+			DEBUG_CALLSPACE(TEXT("GetFunctionCallspace No Authority Server Call Absorbed: %s"), *Function->GetName());
 			return FunctionCallspace::Absorbed;
 		}
 
@@ -2696,11 +2710,14 @@ int32 AActor::GetFunctionCallspace( UFunction* Function, void* Parameters, FFram
 	// Dedicated servers don't care about "cosmetic" functions.
 	if (NetMode == NM_DedicatedServer && Function->HasAllFunctionFlags(FUNC_BlueprintCosmetic))
 	{
+		DEBUG_CALLSPACE(TEXT("GetFunctionCallspace Blueprint Cosmetic Absorbed: %s"), *Function->GetName());
 		return FunctionCallspace::Absorbed;
 	}
 
 	if (!(Function->FunctionFlags & FUNC_Net))
 	{
+		// Not a network function
+		DEBUG_CALLSPACE(TEXT("GetFunctionCallspace Not Net: %s %s"), *Function->GetName(), FunctionCallspace::ToString(Callspace));
 		return Callspace;
 	}
 	
@@ -2719,13 +2736,17 @@ int32 AActor::GetFunctionCallspace( UFunction* Function, void* Parameters, FFram
 			// Server should execute locally and call remotely
 			if (RemoteRole != ROLE_None)
 			{
+				DEBUG_CALLSPACE(TEXT("GetFunctionCallspace Multicast: %s"), *Function->GetName());
 				return (FunctionCallspace::Local | FunctionCallspace::Remote);
 			}
+
+			DEBUG_CALLSPACE(TEXT("GetFunctionCallspace Multicast NoRemoteRole: %s"), *Function->GetName());
 			return FunctionCallspace::Local;
 		}
 		else
 		{
-			// Client should only execute locally iff it is allowed to (function is not KismetAuthorittyOnly)
+			// Client should only execute locally iff it is allowed to (function is not KismetAuthorityOnly)
+			DEBUG_CALLSPACE(TEXT("GetFunctionCallspace Multicast Client: %s %s"), *Function->GetName(), FunctionCallspace::ToString(Callspace));
 			return Callspace;
 		}
 	}
@@ -2734,12 +2755,14 @@ int32 AActor::GetFunctionCallspace( UFunction* Function, void* Parameters, FFram
 	if (bIsServer && !(Function->FunctionFlags & FUNC_NetClient))
 	{
 		// don't replicate
+		DEBUG_CALLSPACE(TEXT("GetFunctionCallspace Server calling Server function: %s %s"), *Function->GetName(), FunctionCallspace::ToString(Callspace));
 		return Callspace;
 	}
 	// if we aren't the server, and it's not a send-to-server function,
 	if (!bIsServer && !(Function->FunctionFlags & FUNC_NetServer))
 	{
 		// don't replicate
+		DEBUG_CALLSPACE(TEXT("GetFunctionCallspace Client calling Client function: %s %s"), *Function->GetName(), FunctionCallspace::ToString(Callspace));
 		return Callspace;
 	}
 
@@ -2753,17 +2776,23 @@ int32 AActor::GetFunctionCallspace( UFunction* Function, void* Parameters, FFram
 			if (ClientPlayer == NULL)
 			{
 				// No owning player, we must absorb
+				DEBUG_CALLSPACE(TEXT("GetFunctionCallspace Client NonOwner absorbed %s"), *Function->GetName());
 				return FunctionCallspace::Absorbed;
 			}
 			else if (Cast<ULocalPlayer>(ClientPlayer) != NULL)
 			{
 				// This is a local player, call locally
+				DEBUG_CALLSPACE(TEXT("GetFunctionCallspace Client local function: %s %s"), *Function->GetName(), FunctionCallspace::ToString(Callspace));
 				return Callspace;
 			}
 		}
 		else if (!NetConnection->Driver || !NetConnection->Driver->World)
 		{
 			// NetDriver does not have a world, most likely shutting down
+			DEBUG_CALLSPACE(TEXT("GetFunctionCallspace NetConnection with no driver or world absorbed: %s %s %s"),
+				*Function->GetName(), 
+				NetConnection->Driver ? *NetConnection->Driver->GetName() : TEXT("NoNetDriver"),
+				NetConnection->Driver && NetConnection->Driver->World ? *NetConnection->Driver->World->GetName() : TEXT("NoWorld"));
 			return FunctionCallspace::Absorbed;
 		}
 
@@ -2775,13 +2804,15 @@ int32 AActor::GetFunctionCallspace( UFunction* Function, void* Parameters, FFram
 	{
 		if (!bIsServer)
 		{
-			UE_LOG(LogNet, Warning,  TEXT("Client is absorbing remote function %s on actor %s because RemoteRole is ROLE_None"), *Function->GetName(), *GetName() );
+			UE_LOG(LogNet, Warning, TEXT("Client is absorbing remote function %s on actor %s because RemoteRole is ROLE_None"), *Function->GetName(), *GetName() );
 		}
 
+		DEBUG_CALLSPACE(TEXT("GetFunctionCallspace RemoteRole None absorbed %s"), *Function->GetName());
 		return FunctionCallspace::Absorbed;
 	}
 
 	// Call remotely
+	DEBUG_CALLSPACE(TEXT("GetFunctionCallspace RemoteRole Remote %s"), *Function->GetName());
 	return FunctionCallspace::Remote;
 }
 
diff --git a/Engine/Source/Runtime/Engine/Private/NetworkDriver.cpp b/Engine/Source/Runtime/Engine/Private/NetworkDriver.cpp
--- a/Engine/Source/Runtime/Engine/Private/NetworkDriver.cpp
+++ b/Engine/Source/Runtime/Engine/Private/NetworkDriver.cpp
@@ -49,6 +49,12 @@ DEFINE_STAT(STAT_VoicePacketsRecv);
 DEFINE_STAT(STAT_PercentInVoice);
 DEFINE_STAT(STAT_PercentOutVoice);
 
+#if UE_BUILD_SHIPPING
+#define DEBUG_REMOTEFUNCTION(Format, ...)
+#else
+#define DEBUG_REMOTEFUNCTION(Format, ...) UE_LOG(LogNet, VeryVerbose, Format, __VA_ARGS__);
+#endif
+
 /*-----------------------------------------------------------------------------
 	UNetDriver implementation.
 -----------------------------------------------------------------------------*/
@@ -592,6 +598,7 @@ void UNetDriver::InternalProcessRemoteFunction
 	// If saturated and function is unimportant, skip it.
 	if( !(Function->FunctionFlags & FUNC_NetReliable) && !Connection->IsNetReady(0) )
 	{
+		DEBUG_REMOTEFUNCTION(TEXT("Network saturated, not calling %s::%s"), *Actor->GetName(), *Function->GetName());
 		return;
 	}
 
@@ -606,12 +613,19 @@ void UNetDriver::InternalProcessRemoteFunction
 
 	// Make sure this function exists for both parties.
 	FClassNetCache* ClassCache = Connection->PackageMap->GetClassNetCache( TargetObj->GetClass() );
-	if( !ClassCache )
+	if (!ClassCache)
+	{
+		DEBUG_REMOTEFUNCTION(TEXT("ClassNetCache empty, not calling %s::%s"), *Actor->GetName(), *Function->GetName());
 		return;
+	}
+		
 	FFieldNetCache* FieldCache = ClassCache->GetFromField( Function );
-	if( !FieldCache )
+	if (!FieldCache)
+	{
+		DEBUG_REMOTEFUNCTION(TEXT("FieldCache empty, not calling %s::%s"), *Actor->GetName(), *Function->GetName());
 		return;
-
+	}
+		
 	// Get the actor channel.
 	UActorChannel* Ch = Connection->ActorChannels.FindRef(Actor);
 	if( !Ch )
@@ -637,17 +651,24 @@ void UNetDriver::InternalProcessRemoteFunction
 				}
 			}
 		}
-		if( !Ch )
+		if (!Ch)
+		{
 			return;
-		if( IsServer )
-			Ch->SetChannelActor( Actor );
+		}
+		if (IsServer)
+		{
+			Ch->SetChannelActor(Actor);
+		}	
 	}
 
 	// Make sure initial channel-opening replication has taken place.
 	if( Ch->OpenPacketId.First==INDEX_NONE )
 	{
-		if( !IsServer )
+		if (!IsServer)
+		{
+			DEBUG_REMOTEFUNCTION(TEXT("Initial channel replication has not occurred, not calling %s::%s"), *Actor->GetName(), *Function->GetName());
 			return;
+		}
 
 		// triggering replication of an Actor while already in the middle of replication can result in invalid data being sent and is therefore illegal
 		if (Ch->bIsReplicatingActor)
@@ -669,15 +690,19 @@ void UNetDriver::InternalProcessRemoteFunction
 
 	// Reliability.
 	//warning: RPC's might overflow, preventing reliable functions from getting thorough.
-	if( Function->FunctionFlags & FUNC_NetReliable )
+	if (Function->FunctionFlags & FUNC_NetReliable)
+	{
 		Bunch.bReliable = 1;
+	}
 
 	// Queue unreliable multicast 
 	bool QueueBunch = (!Bunch.bReliable && Function->FunctionFlags & FUNC_NetMulticast);
 
 	if (!QueueBunch)
-		Ch->BeginContentBlock( TargetObj, Bunch );
-
+	{
+		Ch->BeginContentBlock(TargetObj, Bunch);
+	}
+		
 	//UE_LOG(LogScript, Log, TEXT("   Call %s"),Function->GetFullName());
 	check( FieldCache->FieldNetIndex <= ClassCache->GetMaxIndex() );
 	Bunch.WriteIntWrapped(FieldCache->FieldNetIndex, ClassCache->GetMaxIndex()+1);
@@ -792,7 +817,6 @@ void UNetDriver::InternalProcessRemoteFunction
 		Ch->EndContentBlock( TargetObj, Bunch );
 	}
 
-
 	static const auto CVar = IConsoleManager::Get().FindTConsoleVariableDataInt(TEXT("net.RPC.Debug"));
 	bool LogAsWarning = (CVar && CVar->GetValueOnGameThread() == 1);
 
diff --git a/Engine/Source/Runtime/Engine/Private/World.cpp b/Engine/Source/Runtime/Engine/Private/World.cpp
--- a/Engine/Source/Runtime/Engine/Private/World.cpp
+++ b/Engine/Source/Runtime/Engine/Private/World.cpp
@@ -4068,7 +4068,7 @@ UWorld* FSeamlessTravelHandler::Tick()
 		}
 		else if ( LoadedWorld->PersistentLevel == NULL)
 		{
-			// Package isnt a level
+			// Package isn't a level
 			FString Error = FString::Printf(TEXT("Unable to travel to '%s' - package is not a level"), *LoadedPackage->GetName());
 			UE_LOG(LogWorld, Error, TEXT("%s"), *Error);
 			// abort
@@ -4293,8 +4293,21 @@ UWorld* FSeamlessTravelHandler::Tick()
 				UE_LOG(LogWorld, Log, TEXT("Sending NotifyLoadedWorld for LP: %s PC: %s"), *It->GetName(), It->PlayerController ? *It->PlayerController->GetName() : TEXT("NoPC"));
 				if (It->PlayerController != NULL)
 				{
+#if !UE_BUILD_SHIPPING
+					LOG_SCOPE_VERBOSITY_OVERRIDE(LogNet, ELogVerbosity::VeryVerbose);
+					LOG_SCOPE_VERBOSITY_OVERRIDE(LogNetTraffic, ELogVerbosity::VeryVerbose);
+					UE_LOG(LogNet, Verbose, TEXT("NotifyLoadedWorld Begin"));
+#endif
 					It->PlayerController->NotifyLoadedWorld(LoadedWorld->GetOutermost()->GetFName(), bSwitchedToDefaultMap);
 					It->PlayerController->ServerNotifyLoadedWorld(LoadedWorld->GetOutermost()->GetFName());
+#if !UE_BUILD_SHIPPING
+					UE_LOG(LogNet, Verbose, TEXT("NotifyLoadedWorld End"));
+#endif
+				}
+				else
+				{
+					UE_LOG(LogNet, Error, TEXT("No Player Controller during seamless travel for LP: %s."), *It->GetName());
+					// @todo add some kind of travel back to main menu
 				}
 			}
 
