Fixed compiler error on all platform without editoronly data.

diff --git a/Engine/Source/Runtime/Engine/Private/Animation/Skeleton.cpp b/Engine/Source/Runtime/Engine/Private/Animation/Skeleton.cpp
--- a/Engine/Source/Runtime/Engine/Private/Animation/Skeleton.cpp
+++ b/Engine/Source/Runtime/Engine/Private/Animation/Skeleton.cpp
@@ -136,11 +136,12 @@ void USkeleton::Serialize( FArchive& Ar )
 	{
 		Ar << Guid;
 	}
-
+#if WITH_EDITORONLY_DATA
 	if (Ar.UE4Ver() < VER_UE4_SKELETON_ASSET_PROPERTY_TYPE_CHANGE)
 	{
 		PreviewAttachedAssetContainer.SaveAttachedObjectsFromDeprecatedProperties();
 	}
+#endif
 }
 
 /** Remove this function when VER_UE4_REFERENCE_SKELETON_REFACTOR is removed. */
diff --git a/Engine/Source/Runtime/Engine/Private/SkeletalMesh.cpp b/Engine/Source/Runtime/Engine/Private/SkeletalMesh.cpp
--- a/Engine/Source/Runtime/Engine/Private/SkeletalMesh.cpp
+++ b/Engine/Source/Runtime/Engine/Private/SkeletalMesh.cpp
@@ -2366,11 +2366,12 @@ void USkeletalMesh::Serialize( FArchive& Ar )
 		// now they're in the Materials array so we need to move them over
 		MoveDeprecatedShadowFlagToMaterials();
 	}
-
+#if WITH_EDITORONLY_DATA
 	if (Ar.UE4Ver() < VER_UE4_SKELETON_ASSET_PROPERTY_TYPE_CHANGE)
 	{
 		PreviewAttachedAssetContainer.SaveAttachedObjectsFromDeprecatedProperties();
 	}
+#endif
 }
 
 void USkeletalMesh::AddReferencedObjects(UObject* InThis, FReferenceCollector& Collector)
