Removed unused and forgotten NEQ functions.

diff --git a/Engine/Source/Runtime/Engine/Classes/Vehicles/WheeledVehicleMovementComponent.h b/Engine/Source/Runtime/Engine/Classes/Vehicles/WheeledVehicleMovementComponent.h
--- a/Engine/Source/Runtime/Engine/Classes/Vehicles/WheeledVehicleMovementComponent.h
+++ b/Engine/Source/Runtime/Engine/Classes/Vehicles/WheeledVehicleMovementComponent.h
@@ -156,15 +156,6 @@ struct FReplicatedVehicleState
 	int32 CurrentGear;
 };
 
-static inline bool NEQ(const FReplicatedVehicleState& A, const FReplicatedVehicleState& B, UPackageMap* Map, UActorChannel* Channel)
-{
-	return (A.SteeringInput != B.SteeringInput) ||
-		(A.ThrottleInput != B.ThrottleInput) ||
-		(A.BrakeInput != B.BrakeInput) ||
-		(A.HandbrakeInput != B.HandbrakeInput) ||
-		(A.CurrentGear != B.CurrentGear);
-}
-
 USTRUCT()
 struct FVehicleInputRate
 {
