Fixed log to reflect function it is being logged from

diff --git a/Engine/Source/Runtime/Engine/Private/DataChannel.cpp b/Engine/Source/Runtime/Engine/Private/DataChannel.cpp
--- a/Engine/Source/Runtime/Engine/Private/DataChannel.cpp
+++ b/Engine/Source/Runtime/Engine/Private/DataChannel.cpp
@@ -86,7 +86,7 @@ void UChannel::CleanUp()
 	// if this is the control channel, make sure we properly killed the connection
 	if (ChIndex == 0 && !Closing)
 	{
-		UE_LOG(LogNet, Log, TEXT("NetConnection::Close() [%s] [%s] [%s] from UChannel::Cleanup(). ChIndex == 0. Closing connection."), 
+		UE_LOG(LogNet, Log, TEXT("UChannel::CleanUp: [%s] [%s] [%s]. ChIndex == 0. Closing connection."), 
 			Connection->Driver ? *Connection->Driver->NetDriverName.ToString() : TEXT("NULL"),
 			Connection->PlayerController ? *Connection->PlayerController->GetName() : TEXT("NoPC"),
 			Connection->OwningActor ? *Connection->OwningActor->GetName() : TEXT("No Owner"));
