Fix Linux CIS.

ProcHandle default constructor should already set it to invalid value, supplying NULL is thus not necessary (and problematic for platforms where it is defined as nullptr).

diff --git a/Engine/Source/Runtime/Engine/Private/Tests/ExternalToolTest.cpp b/Engine/Source/Runtime/Engine/Private/Tests/ExternalToolTest.cpp
--- a/Engine/Source/Runtime/Engine/Private/Tests/ExternalToolTest.cpp
+++ b/Engine/Source/Runtime/Engine/Private/Tests/ExternalToolTest.cpp
@@ -14,7 +14,7 @@ namespace ExternalProcessHelper
 	struct ExternalProcessInfo
 	{
 		ExternalProcessInfo()
-			: ProcessHandle(NULL)
+			: ProcessHandle()
 			, ReadPipe(NULL)
 			, WritePipe(NULL)
 		{
