Spelling

diff --git a/Engine/Source/Runtime/Engine/Private/Texture.cpp b/Engine/Source/Runtime/Engine/Private/Texture.cpp
--- a/Engine/Source/Runtime/Engine/Private/Texture.cpp
+++ b/Engine/Source/Runtime/Engine/Private/Texture.cpp
@@ -307,7 +307,7 @@ void UTexture::BeginDestroy()
 		// Send the rendering thread a release message for the texture's resource.
 		BeginReleaseResource(Resource);
 		Resource->ReleaseFence.BeginFence();
-		// Keep track that we alrady kicked off the async release.
+		// Keep track that we already kicked off the async release.
 		bAsyncResourceReleaseHasBeenStarted = true;
 	}
 }
@@ -324,7 +324,7 @@ bool UTexture::IsReadyForFinishDestroy()
 			// Send the rendering thread a release message for the texture's resource.
 			BeginReleaseResource(Resource);
 			Resource->ReleaseFence.BeginFence();
-			// Keep track that we alrady kicked off the async release.
+			// Keep track that we already kicked off the async release.
 			bAsyncResourceReleaseHasBeenStarted = true;
 		}
 		// Only allow FinishDestroy to be called once the texture resource has finished its rendering thread cleanup.
