Fixed crashes in landscape splines when navigation system is disabled

diff --git a/Engine/Source/Runtime/Engine/Private/Landscape/LandscapeSplines.cpp b/Engine/Source/Runtime/Engine/Private/Landscape/LandscapeSplines.cpp
--- a/Engine/Source/Runtime/Engine/Private/Landscape/LandscapeSplines.cpp
+++ b/Engine/Source/Runtime/Engine/Private/Landscape/LandscapeSplines.cpp
@@ -741,7 +741,10 @@ void ULandscapeSplineControlPoint::UpdateSplinePoints(bool bUpdateCollision)
 		else
 		{
 			// UpdateNavOctree won't work if we destroy any components, ghosts get left in the navmesh, so we have to unregister
-			OuterSplines->GetWorld()->GetNavigationSystem()->UnregisterNavigationRelevantActor(OuterSplines->GetOwner());
+			if (OuterSplines->GetWorld()->GetNavigationSystem() != nullptr)
+			{
+				OuterSplines->GetWorld()->GetNavigationSystem()->UnregisterNavigationRelevantActor(OuterSplines->GetOwner());
+			}
 
 			Modify();
 			MeshComponent->DestroyComponent();
@@ -840,7 +843,7 @@ void ULandscapeSplineControlPoint::UpdateSplinePoints(bool bUpdateCollision)
 
 	OuterSplines->MarkRenderStateDirty();
 
-	if (bNavDirty && bUpdateCollision)
+	if (bNavDirty && bUpdateCollision && OuterSplines->GetWorld()->GetNavigationSystem() != nullptr)
 	{
 		OuterSplines->GetWorld()->GetNavigationSystem()->UpdateNavOctree(OuterSplines->GetOwner());
 		bNavDirty = false;
@@ -864,7 +867,10 @@ void ULandscapeSplineControlPoint::DeleteSplinePoints()
 	OuterSplines->MarkRenderStateDirty();
 
 	// UpdateNavOctree won't work if we destroy any components, ghosts get left in the navmesh, so we have to unregister
-	OuterSplines->GetWorld()->GetNavigationSystem()->UnregisterNavigationRelevantActor(OuterSplines->GetOwner());
+	if (OuterSplines->GetWorld()->GetNavigationSystem() != nullptr)
+	{
+		OuterSplines->GetWorld()->GetNavigationSystem()->UnregisterNavigationRelevantActor(OuterSplines->GetOwner());
+	}
 
 	if (MeshComponent != NULL)
 	{
@@ -874,7 +880,10 @@ void ULandscapeSplineControlPoint::DeleteSplinePoints()
 		MeshComponent = NULL;
 	}
 
-	OuterSplines->GetWorld()->GetNavigationSystem()->UpdateNavOctree(OuterSplines->GetOwner());
+	if (OuterSplines->GetWorld()->GetNavigationSystem() != nullptr)
+	{
+		OuterSplines->GetWorld()->GetNavigationSystem()->UpdateNavOctree(OuterSplines->GetOwner());
+	}
 }
 
 void ULandscapeSplineControlPoint::PostEditUndo()
@@ -1389,7 +1398,10 @@ void ULandscapeSplineSegment::UpdateSplinePoints(bool bUpdateCollision)
 		if (iMesh < MeshComponents.Num())
 		{
 			// UpdateNavOctree won't work if we destroy any components, ghosts get left in the navmesh, so we have to unregister
-			OuterSplines->GetWorld()->GetNavigationSystem()->UnregisterNavigationRelevantActor(OuterSplines->GetOwner());
+			if (OuterSplines->GetWorld()->GetNavigationSystem() != nullptr)
+			{
+				OuterSplines->GetWorld()->GetNavigationSystem()->UnregisterNavigationRelevantActor(OuterSplines->GetOwner());
+			}
 
 			for (int32 i = iMesh; i < MeshComponents.Num(); i++)
 			{
@@ -1578,7 +1590,10 @@ void ULandscapeSplineSegment::UpdateSplinePoints(bool bUpdateCollision)
 		if (MeshComponents.Num() > 0)
 		{
 			// UpdateNavOctree won't work if we destroy any components, ghosts get left in the navmesh, so we have to unregister
-			OuterSplines->GetWorld()->GetNavigationSystem()->UnregisterNavigationRelevantActor(OuterSplines->GetOwner());
+			if (OuterSplines->GetWorld()->GetNavigationSystem() != nullptr)
+			{
+				OuterSplines->GetWorld()->GetNavigationSystem()->UnregisterNavigationRelevantActor(OuterSplines->GetOwner());
+			}
 
 			for (auto It = MeshComponents.CreateConstIterator(); It; It++)
 			{
@@ -1589,7 +1604,7 @@ void ULandscapeSplineSegment::UpdateSplinePoints(bool bUpdateCollision)
 		}
 	}
 
-	if (bNavDirty && bUpdateCollision)
+	if (bNavDirty && bUpdateCollision && OuterSplines->GetWorld()->GetNavigationSystem() != nullptr)
 	{
 		OuterSplines->GetWorld()->GetNavigationSystem()->UpdateNavOctree(OuterSplines->GetOwner());
 		bNavDirty = false;
@@ -1623,7 +1638,10 @@ void ULandscapeSplineSegment::DeleteSplinePoints()
 	OuterSplines->MarkRenderStateDirty();
 
 	// UpdateNavOctree won't work if we destroy any components, ghosts get left in the navmesh, so we have to unregister
-	OuterSplines->GetWorld()->GetNavigationSystem()->UnregisterNavigationRelevantActor(OuterSplines->GetOwner());
+	if (OuterSplines->GetWorld()->GetNavigationSystem() != nullptr)
+	{
+		OuterSplines->GetWorld()->GetNavigationSystem()->UnregisterNavigationRelevantActor(OuterSplines->GetOwner());
+	}
 
 	for (auto It = MeshComponents.CreateConstIterator(); It; It++)
 	{
@@ -1635,7 +1653,10 @@ void ULandscapeSplineSegment::DeleteSplinePoints()
 	}
 	MeshComponents.Empty();
 
-	OuterSplines->GetWorld()->GetNavigationSystem()->UpdateNavOctree(OuterSplines->GetOwner());
+	if (OuterSplines->GetWorld()->GetNavigationSystem() != nullptr)
+	{
+		OuterSplines->GetWorld()->GetNavigationSystem()->UpdateNavOctree(OuterSplines->GetOwner());
+	}
 }
 
 static bool LineIntersect( const FVector2D& L1Start, const FVector2D& L1End, const FVector2D& L2Start, const FVector2D& L2End, FVector2D& Intersect, float Tolerance = KINDA_SMALL_NUMBER)
