Sort ES2 landscape vertices in the order of first use in lowest to highest LOD, to minimize the MaxVertexIndex ranges for drawcalls at each LOD. Improves performance on Mali substantially.

diff --git a/Engine/Source/Runtime/Engine/Classes/Landscape/LandscapeComponent.h b/Engine/Source/Runtime/Engine/Classes/Landscape/LandscapeComponent.h
--- a/Engine/Source/Runtime/Engine/Classes/Landscape/LandscapeComponent.h
+++ b/Engine/Source/Runtime/Engine/Classes/Landscape/LandscapeComponent.h
@@ -38,6 +38,26 @@ public:
 	friend FArchive& operator<<(FArchive& Ar, FLandscapeComponentDerivedData& Data);
 };
 
+/* Used to uniquely reference a landscape vertex in a component, and generate a key suitable for a TMap. */
+struct FLandscapeVertexRef
+{
+	FLandscapeVertexRef(int16 InX, int16 InY, int8 InSubX, int8 InSubY)
+	: X(InX)
+	, Y(InY)
+	, SubX(InSubX)
+	, SubY(InSubY)
+	{}
+	int16 X;
+	int16 Y;
+	int8 SubX;
+	int8 SubY;
+
+	uint64 MakeKey() const
+	{
+		return (uint64)X << 32 | (uint64)Y << 16 | (uint64)SubX << 8 | (uint64)SubY;
+	}
+};
+
 /** Stores information about which weightmap texture and channel each layer is stored */
 USTRUCT()
 struct FWeightmapLayerAllocationInfo
diff --git a/Engine/Source/Runtime/Engine/Private/Landscape/Landscape.cpp b/Engine/Source/Runtime/Engine/Private/Landscape/Landscape.cpp
--- a/Engine/Source/Runtime/Engine/Private/Landscape/Landscape.cpp
+++ b/Engine/Source/Runtime/Engine/Private/Landscape/Landscape.cpp
@@ -1811,7 +1811,7 @@ void ULandscapeComponent::PostDuplicate(bool bDuplicateForPIE)
 }
 
 // Generate a new guid to force a recache of all landscape derived data
-#define LANDSCAPE_FULL_DERIVEDDATA_VER			TEXT("d17ce2a02ccc11e3aa6e0800200c9a66")
+#define LANDSCAPE_FULL_DERIVEDDATA_VER			TEXT("28f0da4d085e48a680482c9835907e2f")
 
 FString FLandscapeComponentDerivedData::GetDDCKeyString(const FGuid& StateId)
 {
diff --git a/Engine/Source/Runtime/Engine/Private/Landscape/LandscapeEdit.cpp b/Engine/Source/Runtime/Engine/Private/Landscape/LandscapeEdit.cpp
--- a/Engine/Source/Runtime/Engine/Private/Landscape/LandscapeEdit.cpp
+++ b/Engine/Source/Runtime/Engine/Private/Landscape/LandscapeEdit.cpp
@@ -4852,7 +4852,7 @@ UMaterialInstance* ULandscapeComponent::GeneratePlatformPixelData(TArray<UTextur
 void ULandscapeComponent::GeneratePlatformVertexData()
 {
 	if (IsTemplate())
-{
+	{
 		return;
 	}
 	check(HeightmapTexture);
@@ -4861,93 +4861,149 @@ void ULandscapeComponent::GeneratePlatformVertexData()
 	TArray<uint8> NewPlatformData;
 	int32 NewPlatformDataSize = 0;
 
-	int32 MaxLOD = FMath::CeilLogTwo(SubsectionSizeQuads+1)-1;
+	int32 SubsectionSizeVerts = SubsectionSizeQuads + 1;
+	int32 MaxLOD = FMath::CeilLogTwo(SubsectionSizeVerts) - 1;
 
-	float HeightmapSubsectionOffsetU = (float)(SubsectionSizeQuads+1) / (float)HeightmapTexture->Source.GetSizeX();
-	float HeightmapSubsectionOffsetV = (float)(SubsectionSizeQuads+1) / (float)HeightmapTexture->Source.GetSizeY();
+	float HeightmapSubsectionOffsetU = (float)(SubsectionSizeVerts) / (float)HeightmapTexture->Source.GetSizeX();
+	float HeightmapSubsectionOffsetV = (float)(SubsectionSizeVerts) / (float)HeightmapTexture->Source.GetSizeY();
 
-	NewPlatformDataSize += sizeof(FLandscapeMobileVertex) * FMath::Square( (SubsectionSizeQuads+1) * NumSubsections );
-	TArray<FColor*> HeightmapMipData;
+	NewPlatformDataSize += sizeof(FLandscapeMobileVertex) * FMath::Square(SubsectionSizeVerts * NumSubsections);
+	NewPlatformData.AddZeroed(NewPlatformDataSize);
 
+	// Get the required mip data
+	TArray<FColor*> HeightmapMipData;
 	for( int32 MipIdx=0; MipIdx < FMath::Min(LANDSCAPE_MAX_ES_LOD, HeightmapTexture->Source.GetNumMips()); MipIdx++ )
 	{
-		int32 MipSubsectionSizeVerts = (SubsectionSizeQuads+1) >> MipIdx;
+		int32 MipSubsectionSizeVerts = (SubsectionSizeVerts) >> MipIdx;
 		if( MipSubsectionSizeVerts > 1)
 		{
 			HeightmapMipData.Add( (FColor*)HeightmapTexture->Source.LockMip(MipIdx) );
 		}
 	}
+	
+	TMap<uint64, int32> VertexMap;
+	TArray<FLandscapeVertexRef> VertexOrder;
+	VertexOrder.Empty(FMath::Square(SubsectionSizeVerts * NumSubsections));
 
-	NewPlatformData.AddZeroed(NewPlatformDataSize);
-	FLandscapeMobileVertex* DstVert = (FLandscapeMobileVertex*)NewPlatformData.GetTypedData();
-
-	int32 SubsectionSizeVerts = SubsectionSizeQuads+1;
-	for( int32 SubY=0;SubY<NumSubsections;SubY++ )
+	// Layout index buffer to determine best vertex order
+	for (int32 Mip = MaxLOD; Mip >= 0; Mip--)
 	{
-		for( int32 SubX=0;SubX<NumSubsections;SubX++ )
-		{
-			float HeightmapScaleBiasZ = HeightmapScaleBias.Z + HeightmapSubsectionOffsetU * (float)SubX;
-			float HeightmapScaleBiasW = HeightmapScaleBias.W + HeightmapSubsectionOffsetV * (float)SubY; 
-			int32 BaseMipOfsX = FMath::Round(HeightmapScaleBiasZ * (float)HeightmapTexture->Source.GetSizeX());
-			int32 BaseMipOfsY = FMath::Round(HeightmapScaleBiasW * (float)HeightmapTexture->Source.GetSizeY());
+		int32 LodSubsectionSizeQuads = (SubsectionSizeVerts >> Mip) - 1;
+		float MipRatio = (float)SubsectionSizeQuads / (float)LodSubsectionSizeQuads; // Morph current MIP to base MIP
 
-			for( int32 Y=0; Y<SubsectionSizeVerts; Y++ )
+		for (int32 SubY = 0; SubY < NumSubsections; SubY++)
+		{
+			for (int32 SubX = 0; SubX < NumSubsections; SubX++)
 			{
-				for( int32 X=0; X<SubsectionSizeVerts; X++ )
+				for (int32 y = 0; y < LodSubsectionSizeQuads; y++)
 				{
-					DstVert->Position[0] = X;
-					DstVert->Position[1] = Y;
-					DstVert->Position[2] = SubX;
-					DstVert->Position[3] = SubY;
+					for (int32 x = 0; x < LodSubsectionSizeQuads; x++)
+					{
+						int32 x0 = FMath::Round((float)x * MipRatio);
+						int32 y0 = FMath::Round((float)y * MipRatio);
+						int32 x1 = FMath::Round((float)(x + 1) * MipRatio);
+						int32 y1 = FMath::Round((float)(y + 1) * MipRatio);
+
+						FLandscapeVertexRef V1(x0, y0, SubX, SubY);
+						FLandscapeVertexRef V2(x1, y0, SubX, SubY);
+						FLandscapeVertexRef V3(x1, y1, SubX, SubY);
+						FLandscapeVertexRef V4(x0, y1, SubX, SubY);
+
+						uint64 Key1 = V1.MakeKey();
+						if (VertexMap.Find(Key1) == NULL)
+						{
+							VertexMap.Add(Key1, VertexOrder.Num());
+							VertexOrder.Add(V1);
+						}
+						uint64 Key2 = V2.MakeKey();
+						if (VertexMap.Find(Key2) == NULL)
+						{
+							VertexMap.Add(Key2, VertexOrder.Num());
+							VertexOrder.Add(V2);
+						}
+						uint64 Key3 = V3.MakeKey();
+						if (VertexMap.Find(Key3) == NULL)
+						{
+							VertexMap.Add(Key3, VertexOrder.Num());
+							VertexOrder.Add(V3);
+						}
+						uint64 Key4 = V4.MakeKey();
+						if (VertexMap.Find(Key4) == NULL)
+						{
+							VertexMap.Add(Key4, VertexOrder.Num());
+							VertexOrder.Add(V4);
+						}
+					}
+				}
+			}
+		}
+	}
+	check(VertexOrder.Num() == FMath::Square(SubsectionSizeVerts) * FMath::Square(NumSubsections));
 
-					TArray<int32> MipHeights;
-					MipHeights.AddZeroed(HeightmapMipData.Num());
-					int32 LastIndex = 0;
-					uint16 MaxHeight = 0, MinHeight = 65535;
+	// Fill in the vertices in the specified order
+	FLandscapeMobileVertex* DstVert = (FLandscapeMobileVertex*)NewPlatformData.GetTypedData();
+	for (int32 Idx = 0; Idx < VertexOrder.Num(); Idx++)
+	{
+		int32 X = VertexOrder[Idx].X;
+		int32 Y = VertexOrder[Idx].Y;
+		int32 SubX = VertexOrder[Idx].SubX;
+		int32 SubY = VertexOrder[Idx].SubY;
 
-					for (int32 Mip = 0; Mip < HeightmapMipData.Num(); ++Mip)
-					{
-						int32 MipSizeX = HeightmapTexture->Source.GetSizeX() >> Mip;
+		float HeightmapScaleBiasZ = HeightmapScaleBias.Z + HeightmapSubsectionOffsetU * (float)SubX;
+		float HeightmapScaleBiasW = HeightmapScaleBias.W + HeightmapSubsectionOffsetV * (float)SubY;
+		int32 BaseMipOfsX = FMath::Round(HeightmapScaleBiasZ * (float)HeightmapTexture->Source.GetSizeX());
+		int32 BaseMipOfsY = FMath::Round(HeightmapScaleBiasW * (float)HeightmapTexture->Source.GetSizeY());
 
-						int32 MipSubsectionSizeVerts = (SubsectionSizeVerts >> Mip);
-						int32 MipSubsectionSizeQuads = MipSubsectionSizeVerts - 1;
+		DstVert->Position[0] = X;
+		DstVert->Position[1] = Y;
+		DstVert->Position[2] = SubX;
+		DstVert->Position[3] = SubY;
 
-						int32 CurrentMipOfsX = BaseMipOfsX >> Mip;
-						int32 CurrentMipOfsY = BaseMipOfsY >> Mip;
+		TArray<int32> MipHeights;
+		MipHeights.AddZeroed(HeightmapMipData.Num());
+		int32 LastIndex = 0;
+		uint16 MaxHeight = 0, MinHeight = 65535;
 
-						float MipRatio = (float)MipSubsectionSizeQuads / (float)SubsectionSizeQuads; // Morph Base to current MIP
-						int32 MipX = FMath::Round( (float)X * MipRatio );
-						int32 MipY = FMath::Round( (float)Y * MipRatio );
+		for (int32 Mip = 0; Mip < HeightmapMipData.Num(); ++Mip)
+		{
+			int32 MipSizeX = HeightmapTexture->Source.GetSizeX() >> Mip;
 
-						FColor* CurrentMipSrcRow = HeightmapMipData[Mip] + (CurrentMipOfsY + MipY) * MipSizeX + CurrentMipOfsX;
-						uint16 Height = CurrentMipSrcRow[MipX].R << 8 | CurrentMipSrcRow[MipX].G;
+			int32 MipSubsectionSizeVerts = (SubsectionSizeVerts >> Mip);
+			int32 MipSubsectionSizeQuads = MipSubsectionSizeVerts - 1;
 
-						MipHeights[Mip] = Height;
-						MaxHeight = FMath::Max(MaxHeight, Height);
-						MinHeight = FMath::Min(MinHeight, Height);
-					}
+			int32 CurrentMipOfsX = BaseMipOfsX >> Mip;
+			int32 CurrentMipOfsY = BaseMipOfsY >> Mip;
 
-					DstVert->LODHeights[0] = MinHeight >> 8;
-					DstVert->LODHeights[1] = (MinHeight & 255);
-					DstVert->LODHeights[2] = MaxHeight >> 8;
-					DstVert->LODHeights[3] = (MaxHeight & 255);
+			float MipRatio = (float)MipSubsectionSizeQuads / (float)SubsectionSizeQuads; // Morph Base to current MIP
+			int32 MipX = FMath::Round((float)X * MipRatio);
+			int32 MipY = FMath::Round((float)Y * MipRatio);
 
-					for (int32 Mip = 0; Mip < HeightmapMipData.Num(); ++Mip)
-					{
-						if (Mip < 4)
-						{
-							DstVert->LODHeights[4+Mip] = FMath::Round(float(MipHeights[Mip] - MinHeight) / (MaxHeight - MinHeight) * 255);
-						}
-						else // Mip 4 5 packed into SubX, SubY
-						{
-							DstVert->Position[Mip-2] += (FMath::Round(float(MipHeights[Mip] - MinHeight) / (MaxHeight - MinHeight) * 255)) & (0xfffe);
-						}
-					}
+			FColor* CurrentMipSrcRow = HeightmapMipData[Mip] + (CurrentMipOfsY + MipY) * MipSizeX + CurrentMipOfsX;
+			uint16 Height = CurrentMipSrcRow[MipX].R << 8 | CurrentMipSrcRow[MipX].G;
 
-					DstVert++;
-				}
+			MipHeights[Mip] = Height;
+			MaxHeight = FMath::Max(MaxHeight, Height);
+			MinHeight = FMath::Min(MinHeight, Height);
+		}
+
+		DstVert->LODHeights[0] = MinHeight >> 8;
+		DstVert->LODHeights[1] = (MinHeight & 255);
+		DstVert->LODHeights[2] = MaxHeight >> 8;
+		DstVert->LODHeights[3] = (MaxHeight & 255);
+
+		for (int32 Mip = 0; Mip < HeightmapMipData.Num(); ++Mip)
+		{
+			if (Mip < 4)
+			{
+				DstVert->LODHeights[4 + Mip] = FMath::Round(float(MipHeights[Mip] - MinHeight) / (MaxHeight - MinHeight) * 255);
+			}
+			else // Mip 4 5 packed into SubX, SubY
+			{
+				DstVert->Position[Mip - 2] += (FMath::Round(float(MipHeights[Mip] - MinHeight) / (MaxHeight - MinHeight) * 255)) & (0xfffe);
 			}
 		}
+
+		DstVert++;
 	}
 
 	for( int32 MipIdx=0;MipIdx<HeightmapTexture->Source.GetNumMips();MipIdx++ )
diff --git a/Engine/Source/Runtime/Engine/Private/Landscape/LandscapeRender.cpp b/Engine/Source/Runtime/Engine/Private/Landscape/LandscapeRender.cpp
--- a/Engine/Source/Runtime/Engine/Private/Landscape/LandscapeRender.cpp
+++ b/Engine/Source/Runtime/Engine/Private/Landscape/LandscapeRender.cpp
@@ -660,16 +660,12 @@ FLandscapeComponentSceneProxy::FLandscapeComponentSceneProxy(ULandscapeComponent
 			BatchElementParams->SubX = SubX;
 			BatchElementParams->SubY = SubY;
 
-			BatchElement->MinVertexIndex = 0;
-			BatchElement->MaxVertexIndex = SubsectionSizeVerts * FMath::Square(NumSubsections);
-
 #if WITH_EDITOR
 			FMeshBatchElement* ToolBatchElement = (SubX==0 && SubY==0) ? DynamicMeshTools.Elements.GetTypedData() : new(DynamicMeshTools.Elements) FMeshBatchElement;
 			*ToolBatchElement = *BatchElement;
 #endif
 		}
 	}
-
 }
 
 void FLandscapeComponentSceneProxy::CreateRenderThreadResources()
@@ -1049,7 +1045,9 @@ void FLandscapeComponentSceneProxy::DrawStaticElements(FStaticPrimitiveDrawInter
 
 	for( int32 LOD = FirstLOD; LOD <= LastLOD; LOD++ )
 	{
-		if( ForcedLOD < 0 && NumSubsections > 1 )
+		int32 LodSubsectionSizeVerts = SubsectionSizeVerts >> LOD;
+
+		if (ForcedLOD < 0 && NumSubsections > 1)
 		{
 			// Per-subsection batch elements
 			for( int32 SubY=0;SubY<NumSubsections;SubY++ )
@@ -1070,10 +1068,10 @@ void FLandscapeComponentSceneProxy::DrawStaticElements(FStaticPrimitiveDrawInter
 					BatchElementParams->CurrentLOD = LOD;
 
 					BatchElement->IndexBuffer = SharedBuffers->IndexBuffers[LOD];
-					BatchElement->NumPrimitives = FMath::Square(((SubsectionSizeVerts >> LOD) - 1)) * 2;
+					BatchElement->NumPrimitives = FMath::Square((LodSubsectionSizeVerts-1)) * 2;
 					BatchElement->FirstIndex = (SubX + SubY * NumSubsections) * BatchElement->NumPrimitives * 3;
-					BatchElement->MinVertexIndex = 0;
-					BatchElement->MaxVertexIndex = SubsectionSizeVerts * FMath::Square(NumSubsections);
+					BatchElement->MinVertexIndex = SharedBuffers->IndexRanges[LOD].MinIndex[SubX][SubY];
+					BatchElement->MaxVertexIndex = SharedBuffers->IndexRanges[LOD].MaxIndex[SubX][SubY];
 				}
 			}
 		}
@@ -1100,10 +1098,10 @@ void FLandscapeComponentSceneProxy::DrawStaticElements(FStaticPrimitiveDrawInter
 			BatchElement->IndexBuffer = SharedBuffers->IndexBuffers[LOD];
 		}
 
-		BatchElement->NumPrimitives = FMath::Square(((SubsectionSizeVerts >> LOD) - 1)) * FMath::Square(NumSubsections) * 2;
+		BatchElement->NumPrimitives = FMath::Square((LodSubsectionSizeVerts-1)) * FMath::Square(NumSubsections) * 2;
 		BatchElement->FirstIndex = 0;
-		BatchElement->MinVertexIndex = 0;
-		BatchElement->MaxVertexIndex = SubsectionSizeVerts * FMath::Square(NumSubsections);
+		BatchElement->MinVertexIndex = SharedBuffers->IndexRanges[LOD].MinIndexFull;
+		BatchElement->MaxVertexIndex = SharedBuffers->IndexRanges[LOD].MaxIndexFull;
 	}
 
 	PDI->DrawMesh(MeshBatch,0,FLT_MAX);				
@@ -1312,9 +1310,13 @@ void FLandscapeComponentSceneProxy::DrawDynamicElements(FPrimitiveDrawInterface*
 				BatchElement.IndexBuffer = SharedBuffers->IndexBuffers[CurrentLOD];
 			}
 
-			int32 NumPrimitives = FMath::Square(((SubsectionSizeVerts >> CurrentLOD) - 1)) * 2;
+			int32 LodSubsectionSizeVerts = (SubsectionSizeVerts >> CurrentLOD);
+
+			int32 NumPrimitives = FMath::Square((LodSubsectionSizeVerts - 1)) * 2;
 			BatchElement.NumPrimitives = NumPrimitives;
 			BatchElement.FirstIndex = (SubX + SubY * NumSubsections) * NumPrimitives * 3;
+			BatchElement.MinVertexIndex = SharedBuffers->IndexRanges[CurrentLOD].MinIndex[SubX][SubY];
+			BatchElement.MaxVertexIndex = SharedBuffers->IndexRanges[CurrentLOD].MaxIndex[SubX][SubY];
 
 #if WITH_EDITOR
 			// Tools never use tessellation
@@ -1323,6 +1325,8 @@ void FLandscapeComponentSceneProxy::DrawDynamicElements(FPrimitiveDrawInterface*
 			BatchElementTools.IndexBuffer = SharedBuffers->IndexBuffers[CurrentLOD];
 			BatchElementTools.NumPrimitives = NumPrimitives;
 			BatchElementTools.FirstIndex = BatchElement.FirstIndex;
+			BatchElementTools.MinVertexIndex = SharedBuffers->IndexRanges[CurrentLOD].MinIndex[SubX][SubY];
+			BatchElementTools.MaxVertexIndex = SharedBuffers->IndexRanges[CurrentLOD].MaxIndex[SubX][SubY];
 #endif
 		}
 	}
@@ -1528,95 +1532,216 @@ void FLandscapeVertexBuffer::InitRHI()
 }
 
 //
-// FLandscapeVertexBuffer
+// FLandscapeSharedBuffers
 //
 
 template <typename INDEX_TYPE>
-class FLandscapeIndexBuffer : public FRawStaticIndexBuffer16or32<INDEX_TYPE>
+void FLandscapeSharedBuffers::CreateIndexBuffers()
 {
-public:
-	FLandscapeIndexBuffer(int32 ThisLODSubsectionSizeQuads, int32 NumSubsections, int32 FullLODSubsectionSizeVerts)
-	:	FRawStaticIndexBuffer16or32<INDEX_TYPE>(false)
+	if (GRHIFeatureLevel == ERHIFeatureLevel::ES2)
+	{
+		if (!bVertexScoresComputed)
+		{
+			bVertexScoresComputed = ComputeVertexScores();
+		}
+	}
+
+	TMap<uint64, INDEX_TYPE> VertexMap;
+	INDEX_TYPE VertexCount = 0;
+	int32 SubsectionSizeQuads = SubsectionSizeVerts - 1;
+
+	// Layout index buffer to determine best vertex order
+	int32 MaxLOD = NumIndexBuffers - 1;
+	for (int32 Mip = MaxLOD; Mip >= 0; Mip--)
 	{
+		int32 LodSubsectionSizeQuads = (SubsectionSizeVerts >> Mip) - 1;
+
 		TArray<INDEX_TYPE> NewIndices;
-		NewIndices.Empty(ThisLODSubsectionSizeQuads*ThisLODSubsectionSizeQuads*6);
+		int32 ExpectedNumIndices = FMath::Square(NumSubsections) * FMath::Square(LodSubsectionSizeQuads) * 6;
+		NewIndices.Empty(ExpectedNumIndices);
 
-		for( int32 SubY=0;SubY<NumSubsections;SubY++ )
+		int32& MaxIndexFull = IndexRanges[Mip].MaxIndexFull;
+		int32& MinIndexFull = IndexRanges[Mip].MinIndexFull;
+		MaxIndexFull = 0;
+		MinIndexFull = MAX_int32;
+
+		if (GRHIFeatureLevel == ERHIFeatureLevel::ES2)
 		{
-			for( int32 SubX=0;SubX<NumSubsections;SubX++ )
-			{
-				int32 SubOffset = (SubX + SubY * NumSubsections) * FMath::Square(FullLODSubsectionSizeVerts);
-				TArray<INDEX_TYPE> SubIndices;
+			// ES2 version
+			float MipRatio = (float)SubsectionSizeQuads / (float)LodSubsectionSizeQuads; // Morph current MIP to base MIP
 
-				for( int32 y=0;y<ThisLODSubsectionSizeQuads;y++ )
+			for (int32 SubY = 0; SubY < NumSubsections; SubY++)
+			{
+				for (int32 SubX = 0; SubX < NumSubsections; SubX++)
 				{
-					for( int32 x=0;x<ThisLODSubsectionSizeQuads;x++ )
-					{
-						int32 x0 = x;
-						int32 y0 = y;
-						int32 x1 = x + 1;
-						int32 y1 = y + 1;
+					TArray<INDEX_TYPE> SubIndices;
+					SubIndices.Empty(FMath::Square(LodSubsectionSizeQuads) * 6);
 
-						if (GRHIFeatureLevel == ERHIFeatureLevel::ES2)
-						{
-							float MipRatio = (float)(FullLODSubsectionSizeVerts-1) / (float)ThisLODSubsectionSizeQuads; // Morph current MIP to base MIP
-							x0 = FMath::Round( (float)x0 * MipRatio );
-							y0 = FMath::Round( (float)y0 * MipRatio );
-							x1 = FMath::Round( (float)x1 * MipRatio );
-							y1 = FMath::Round( (float)y1 * MipRatio );
-
-							SubIndices.Add( x0 + y0 * FullLODSubsectionSizeVerts + SubOffset );
-							SubIndices.Add( x1 + y1 * FullLODSubsectionSizeVerts + SubOffset );
-							SubIndices.Add( x1 + y0 * FullLODSubsectionSizeVerts + SubOffset );
-
-							SubIndices.Add( x0 + y0 * FullLODSubsectionSizeVerts + SubOffset );
-							SubIndices.Add( x0 + y1 * FullLODSubsectionSizeVerts + SubOffset );
-							SubIndices.Add( x1 + y1 * FullLODSubsectionSizeVerts + SubOffset );
-						}
-						else
+					int32& MaxIndex = IndexRanges[Mip].MaxIndex[SubX][SubY];
+					int32& MinIndex = IndexRanges[Mip].MinIndex[SubX][SubY];
+					MaxIndex = 0;
+					MinIndex = MAX_int32;
+
+					for (int32 y = 0; y < LodSubsectionSizeQuads; y++)
+					{
+						for (int32 x = 0; x < LodSubsectionSizeQuads; x++)
 						{
-							NewIndices.Add( x0 + y0 * FullLODSubsectionSizeVerts + SubOffset );
-							NewIndices.Add( x1 + y1 * FullLODSubsectionSizeVerts + SubOffset );
-							NewIndices.Add( x1 + y0 * FullLODSubsectionSizeVerts + SubOffset );
+							int32 x0 = FMath::Round((float)x * MipRatio);
+							int32 y0 = FMath::Round((float)y * MipRatio);
+							int32 x1 = FMath::Round((float)(x + 1) * MipRatio);
+							int32 y1 = FMath::Round((float)(y + 1) * MipRatio);
+
+							FLandscapeVertexRef V00(x0, y0, SubX, SubY);
+							FLandscapeVertexRef V10(x1, y0, SubX, SubY);
+							FLandscapeVertexRef V11(x1, y1, SubX, SubY);
+							FLandscapeVertexRef V01(x0, y1, SubX, SubY);
+
+							uint64 Key00 = V00.MakeKey();
+							uint64 Key10 = V10.MakeKey();
+							uint64 Key11 = V11.MakeKey();
+							uint64 Key01 = V01.MakeKey();
+							
+							INDEX_TYPE i00;
+							INDEX_TYPE i10;
+							INDEX_TYPE i11;
+							INDEX_TYPE i01;
+
+							INDEX_TYPE* KeyPtr = VertexMap.Find(Key00);
+							if (KeyPtr == NULL)
+							{
+								i00 = VertexCount++;
+								VertexMap.Add(Key00, i00);
+							}
+							else
+							{
+								i00 = *KeyPtr;
+							}
 
-							NewIndices.Add( x0 + y0 * FullLODSubsectionSizeVerts + SubOffset );
-							NewIndices.Add( x0 + y1 * FullLODSubsectionSizeVerts + SubOffset );
-							NewIndices.Add( x1 + y1 * FullLODSubsectionSizeVerts + SubOffset );
-						}
+							KeyPtr = VertexMap.Find(Key10);
+							if (KeyPtr == NULL)
+							{
+								i10 = VertexCount++;
+								VertexMap.Add(Key10, i10);
+							}
+							else
+							{
+								i10 = *KeyPtr;
+							}
+							
+							KeyPtr = VertexMap.Find(Key11);
+							if (KeyPtr == NULL)
+							{
+								i11 = VertexCount++;
+								VertexMap.Add(Key11, i11);
+							}
+							else
+							{
+								i11 = *KeyPtr;
+							}
 
-						check( x1 + y1 * FullLODSubsectionSizeVerts + SubOffset <= (uint32)((INDEX_TYPE)(~(INDEX_TYPE)0)) );
-					}
-				}
+							KeyPtr = VertexMap.Find(Key01);
+							if (KeyPtr == NULL)
+							{
+								i01 = VertexCount++;
+								VertexMap.Add(Key01, i01);
+							}
+							else
+							{
+								i01 = *KeyPtr;
+							}
 
-				if (GRHIFeatureLevel == ERHIFeatureLevel::ES2) // Post-transform vertex cache optimize only for ES2
-				{
-					if (!bVertexScoresComputed)
-					{
-						bVertexScoresComputed = ComputeVertexScores();
+							// Update the min/max index ranges
+							MaxIndex = FMath::Max<int32>(MaxIndex, i00);
+							MinIndex = FMath::Min<int32>(MinIndex, i00);
+							MaxIndex = FMath::Max<int32>(MaxIndex, i10);
+							MinIndex = FMath::Min<int32>(MinIndex, i10);
+							MaxIndex = FMath::Max<int32>(MaxIndex, i11);
+							MinIndex = FMath::Min<int32>(MinIndex, i11);
+							MaxIndex = FMath::Max<int32>(MaxIndex, i01);
+							MinIndex = FMath::Min<int32>(MinIndex, i01);
+
+							SubIndices.Add(i00);
+							SubIndices.Add(i11);
+							SubIndices.Add(i10);
+
+							SubIndices.Add(i00);
+							SubIndices.Add(i01);
+							SubIndices.Add(i11);
+						}
 					}
+
+					// update min/max for full subsection
+					MaxIndexFull = FMath::Max<int32>(MaxIndexFull, MaxIndex);
+					MinIndexFull = FMath::Min<int32>(MinIndexFull, MinIndex);
+
 					TArray<INDEX_TYPE> NewSubIndices;
 					::OptimizeFaces<INDEX_TYPE>(SubIndices, NewSubIndices, 32);
 					NewIndices.Append(NewSubIndices);
 				}
 			}
 		}
+		else
+		{
+			// non-ES2 version
+			for (int32 SubY = 0; SubY < NumSubsections; SubY++)
+			{
+				for (int32 SubX = 0; SubX < NumSubsections; SubX++)
+				{
+					int32 SubOffset = (SubX + SubY * NumSubsections) * FMath::Square(SubsectionSizeVerts);
 
-		this->AssignNewBuffer(NewIndices);
+					int32& MaxIndex = IndexRanges[Mip].MaxIndex[SubX][SubY];
+					int32& MinIndex = IndexRanges[Mip].MinIndex[SubX][SubY];
+					MaxIndex = 0;
+					MinIndex = MAX_int32;
 
-		this->InitResource();
-	}
+					for (int32 y = 0; y < LodSubsectionSizeQuads; y++)
+					{
+						for (int32 x = 0; x < LodSubsectionSizeQuads; x++)
+						{
+							INDEX_TYPE i00 = (x+0) + (y+0) * SubsectionSizeVerts + SubOffset;
+							INDEX_TYPE i10 = (x+1) + (y+0) * SubsectionSizeVerts + SubOffset;
+							INDEX_TYPE i11 = (x+1) + (y+1) * SubsectionSizeVerts + SubOffset;
+							INDEX_TYPE i01 = (x+0) + (y+1) * SubsectionSizeVerts + SubOffset;
+
+							NewIndices.Add(i00);
+							NewIndices.Add(i11);
+							NewIndices.Add(i10);
+
+							NewIndices.Add(i00);
+							NewIndices.Add(i01);
+							NewIndices.Add(i11);
+
+							// Update the min/max index ranges
+							MaxIndex = FMath::Max<int32>(MaxIndex, i00);
+							MinIndex = FMath::Min<int32>(MinIndex, i00);
+							MaxIndex = FMath::Max<int32>(MaxIndex, i10);
+							MinIndex = FMath::Min<int32>(MinIndex, i10);
+							MaxIndex = FMath::Max<int32>(MaxIndex, i11);
+							MinIndex = FMath::Min<int32>(MinIndex, i11);
+							MaxIndex = FMath::Max<int32>(MaxIndex, i01);
+							MinIndex = FMath::Min<int32>(MinIndex, i01);
+						}
+					}
 
-	/** Destructor. */
-	virtual ~FLandscapeIndexBuffer()
-	{
-		this->ReleaseResource();
-	}
-};
+					// update min/max for full subsection
+					MaxIndexFull = FMath::Max<int32>(MaxIndexFull, MaxIndex);
+					MinIndexFull = FMath::Min<int32>(MinIndexFull, MinIndex);
+				}
+			}
 
+			check(MinIndexFull <= (uint32)((INDEX_TYPE)(~(INDEX_TYPE)0)));
+			check(NewIndices.Num() == ExpectedNumIndices);
+		}
+
+		// Create and init new index buffer with index data
+		auto IndexBuffer = new FRawStaticIndexBuffer16or32<INDEX_TYPE>(false);
+		IndexBuffer->AssignNewBuffer(NewIndices);
+		IndexBuffer->InitResource();
+
+		IndexBuffers[Mip] = IndexBuffer;
+	}
+}
 
-//
-// FLandscapeSharedBuffers
-//
 FLandscapeSharedBuffers::FLandscapeSharedBuffers(int32 InSharedBuffersKey, int32 InSubsectionSizeQuads, int32 InNumSubsections)
 :	SharedBuffersKey(InSharedBuffersKey)
 ,	NumIndexBuffers(FMath::CeilLogTwo(InSubsectionSizeQuads+1))
@@ -1631,21 +1756,16 @@ FLandscapeSharedBuffers::FLandscapeSharedBuffers(int32 InSharedBuffersKey, int32
 		VertexBuffer = new FLandscapeVertexBuffer(SubsectionSizeVerts, NumSubsections);
 	}
 	IndexBuffers = new FIndexBuffer*[NumIndexBuffers];
+	IndexRanges = new FLandscapeIndexRanges[NumIndexBuffers]();
 
 	// See if we need to use 16 or 32-bit index buffers
 	if( FMath::Square(SubsectionSizeVerts) * FMath::Square(NumSubsections) > 65535 )
 	{
-		for( int32 i=0;i<NumIndexBuffers;i++ )
-		{
-			IndexBuffers[i] = new FLandscapeIndexBuffer<uint32>((SubsectionSizeVerts>>i)-1, NumSubsections, SubsectionSizeVerts);
-		}
+		CreateIndexBuffers<uint32>();
 	}
 	else
 	{
-		for( int32 i=0;i<NumIndexBuffers;i++ )
-		{
-			IndexBuffers[i] = new FLandscapeIndexBuffer<uint16>((SubsectionSizeVerts>>i)-1, NumSubsections, SubsectionSizeVerts);
-		}
+		CreateIndexBuffers<uint16>();
 	}
 }
 
@@ -1657,9 +1777,12 @@ FLandscapeSharedBuffers::~FLandscapeSharedBuffers()
 	}
 	for( int32 i=0;i<NumIndexBuffers;i++ )
 	{
+		IndexBuffers[i]->ReleaseResource();
 		delete IndexBuffers[i];
 	}
 	delete[] IndexBuffers;
+	delete[] IndexRanges;
+
 	if (AdjacencyIndexBuffers)
 	{
 		if (AdjacencyIndexBuffers->Release() == 0)
diff --git a/Engine/Source/Runtime/Engine/Public/Landscape/LandscapeRender.h b/Engine/Source/Runtime/Engine/Public/Landscape/LandscapeRender.h
--- a/Engine/Source/Runtime/Engine/Public/Landscape/LandscapeRender.h
+++ b/Engine/Source/Runtime/Engine/Public/Landscape/LandscapeRender.h
@@ -270,6 +270,14 @@ public:
 class FLandscapeSharedBuffers : public FRefCountedObject
 {
 public:
+	struct FLandscapeIndexRanges
+	{
+		int32 MinIndex[LANDSCAPE_MAX_SUBSECTION_NUM][LANDSCAPE_MAX_SUBSECTION_NUM];
+		int32 MaxIndex[LANDSCAPE_MAX_SUBSECTION_NUM][LANDSCAPE_MAX_SUBSECTION_NUM];
+		int32 MinIndexFull;
+		int32 MaxIndexFull;
+	};
+
 	int32 SharedBuffersKey;
 	int32 NumIndexBuffers;
 	int32 SubsectionSizeVerts;
@@ -278,9 +286,14 @@ public:
 	FLandscapeVertexFactory* VertexFactory;
 	FLandscapeVertexBuffer* VertexBuffer;
 	FIndexBuffer** IndexBuffers;
+	FLandscapeIndexRanges* IndexRanges;
 	FLandscapeSharedAdjacencyIndexBuffer* AdjacencyIndexBuffers;
 
 	FLandscapeSharedBuffers(int32 SharedBuffersKey, int32 SubsectionSizeQuads, int32 NumSubsections);
+
+	template <typename INDEX_TYPE>
+	void CreateIndexBuffers();
+
 	virtual ~FLandscapeSharedBuffers();
 };
 
