Expose PrimitiveComponent SetPhysMaterialOverride to BP

diff --git a/Engine/Source/Runtime/Engine/Classes/Components/PrimitiveComponent.h b/Engine/Source/Runtime/Engine/Classes/Components/PrimitiveComponent.h
--- a/Engine/Source/Runtime/Engine/Classes/Components/PrimitiveComponent.h
+++ b/Engine/Source/Runtime/Engine/Classes/Components/PrimitiveComponent.h
@@ -1359,8 +1359,9 @@ public:
 	/** 
 	 *	Changes the current PhysMaterialOverride for this component. 
 	 *	Note that if physics is already running on this component, this will _not_ alter its mass/inertia etc,  
-	 *	it will only change its surface properties like friction and the damping.
+	 *	it will only change its surface properties like friction.
 	 */
+	UFUNCTION(BlueprintCallable, Category="Physics", meta=(FriendlyName="Set PhysicalMaterial Override"))
 	virtual void SetPhysMaterialOverride(class UPhysicalMaterial* NewPhysMaterial);
 
 	/** 
