Prevent TextRenderSceneProxy from initializing RHI resources when no vertices are generated. Fixes crash on some android devices.

diff --git a/Engine/Source/Runtime/Engine/Private/TextRenderComponent.cpp b/Engine/Source/Runtime/Engine/Private/TextRenderComponent.cpp
--- a/Engine/Source/Runtime/Engine/Private/TextRenderComponent.cpp
+++ b/Engine/Source/Runtime/Engine/Private/TextRenderComponent.cpp
@@ -658,7 +658,9 @@ bool  FTextRenderSceneProxy::BuildStringMesh( TArray<FDynamicMeshVertex>& OutVer
 		// Move Y position down to next line. If the current line is empty, move by max char height in font
 		StartY += LineSize.Y > 0 ? LineSize.Y : Font->GetMaxCharHeight();
 	}
-	return true;
+
+	// Avoid initializing RHI resources when no vertices are generated.
+	return (OutVertices.Num() > 0);
 }
 
 // ------------------------------------------------------
