Fully expose TextRenderComponent to be called from other modules

diff --git a/Engine/Source/Runtime/Engine/Classes/Components/TextRenderComponent.h b/Engine/Source/Runtime/Engine/Classes/Components/TextRenderComponent.h
--- a/Engine/Source/Runtime/Engine/Classes/Components/TextRenderComponent.h
+++ b/Engine/Source/Runtime/Engine/Classes/Components/TextRenderComponent.h
@@ -31,8 +31,8 @@ enum EVerticalTextAligment
 /**
  * Renders text in the world with given font. Contains usual font related attributes such as Scale, Alignment, Color etc.
  */
-UCLASS(ClassGroup=Rendering, hidecategories=(Object,LOD,Physics,TextureStreaming,Activation,"Components|Activation",Collision), editinlinenew, meta=(BlueprintSpawnableComponent = ""), MinimalAPI)
-class UTextRenderComponent : public UPrimitiveComponent
+UCLASS(ClassGroup=Rendering, hidecategories=(Object,LOD,Physics,TextureStreaming,Activation,"Components|Activation",Collision), editinlinenew, meta=(BlueprintSpawnableComponent = ""))
+class ENGINE_API UTextRenderComponent : public UPrimitiveComponent
 {
 	GENERATED_UCLASS_BODY()
 
