PS4 will pad textures which will cause a mismatch in DDC size and calculated size.
Converted this error into a warning.

diff --git a/Engine/Source/Runtime/Engine/Private/TextureDerivedData.cpp b/Engine/Source/Runtime/Engine/Private/TextureDerivedData.cpp
--- a/Engine/Source/Runtime/Engine/Private/TextureDerivedData.cpp
+++ b/Engine/Source/Runtime/Engine/Private/TextureDerivedData.cpp
@@ -994,7 +994,7 @@ static void CheckMipSize(FTexture2DMipMap& Mip, EPixelFormat PixelFormat, int32
 {
 	if (MipSize != CalcTextureMipMapSize(Mip.SizeX, Mip.SizeY, PixelFormat, 0))
 	{
-		UE_LOG(LogTexture, Fatal,
+		UE_LOG(LogTexture, Warning,
 			TEXT("%dx%d mip of %s texture has invalid data in the DDC. Got %d bytes, expected %d. Key=%s"),
 			Mip.SizeX,
 			Mip.SizeY,
