WinRT and XboxOne build fixes.

diff --git a/Engine/Source/Programs/UnrealBuildTool/Configuration/UEBuildTarget.cs b/Engine/Source/Programs/UnrealBuildTool/Configuration/UEBuildTarget.cs
--- a/Engine/Source/Programs/UnrealBuildTool/Configuration/UEBuildTarget.cs
+++ b/Engine/Source/Programs/UnrealBuildTool/Configuration/UEBuildTarget.cs
@@ -1484,8 +1484,9 @@ namespace UnrealBuildTool
 				string LinkerFixupsName = "UELinkerFixups";
 
 				// Include an empty header so UEBuildModule.Compile does not complain about a lack of PCH
-				string HeaderFilename = LinkerFixupsName + "Name.h";
-				LinkerFixupsFileContents.Add("#include \"" + HeaderFilename + "\"");
+				string HeaderFilename                    = LinkerFixupsName + "Name.h";
+				string LinkerFixupHeaderFilenameWithPath = Path.Combine(GlobalCompileEnvironment.Config.OutputDirectory, HeaderFilename);
+				LinkerFixupsFileContents.Add("#include \"" + LinkerFixupHeaderFilenameWithPath + "\"");
 
 				// Add a function that is not referenced by anything that invokes all the empty functions in the different static libraries
 				LinkerFixupsFileContents.Add("void " + LinkerFixupsName + "()");
@@ -1517,7 +1518,6 @@ namespace UnrealBuildTool
 				// If we determined that we should write the file, write it now.
 				if (bShouldWriteFile)
 				{
-					string LinkerFixupHeaderFilenameWithPath = Path.Combine(GlobalCompileEnvironment.Config.OutputDirectory, HeaderFilename);
 					ResponseFile.Create(LinkerFixupHeaderFilenameWithPath, new List<string>());
 					ResponseFile.Create(LinkerFixupCPPFilename, LinkerFixupsFileContents);
 				}
diff --git a/Engine/Source/Programs/UnrealBuildTool/WinRT/WinRTToolChain.cs b/Engine/Source/Programs/UnrealBuildTool/WinRT/WinRTToolChain.cs
--- a/Engine/Source/Programs/UnrealBuildTool/WinRT/WinRTToolChain.cs
+++ b/Engine/Source/Programs/UnrealBuildTool/WinRT/WinRTToolChain.cs
@@ -563,7 +563,10 @@ namespace UnrealBuildTool
 						FileArguments += string.Format(" /Yu\"{0}\"", CompileEnvironment.Config.PCHHeaderNameInCode);
 						FileArguments += string.Format(" /Fp\"{0}\"", CompileEnvironment.PrecompiledHeaderFile.AbsolutePath);
 
-						if (CompileEnvironment.Config.bForceIncludePrecompiledHeader)
+						// Is it unsafe to always force inclusion?  Clang is doing it, and .generated.cpp files
+						// won't work otherwise, because they're not located in the context of the module,
+						// so they can't access the module's PCH without an absolute path.
+						//if (CompileEnvironment.Config.bForceIncludePrecompiledHeader)
 						{
 							// Force include the precompiled header file.  This is needed because we may have selected a
 							// precompiled header that is different than the first direct include in the C++ source file, but
@@ -879,7 +882,7 @@ namespace UnrealBuildTool
 					string BaseVSToolPath = WindowsPlatform.GetVSComnToolsPath(WindowsCompiler.VisualStudio2012);
 					if (string.IsNullOrEmpty(BaseVSToolPath))
 					{
-						throw new BuildException("Visual Studio 2011 must be installed in order to build this target.");
+						throw new BuildException("Visual Studio 2012 must be installed in order to build this target.");
 					}
 
 					if (Platform == CPPTargetPlatform.WinRT_ARM)
@@ -935,7 +938,7 @@ namespace UnrealBuildTool
 					BaseVSToolPath = "C:/Program Files (x86)/Microsoft Visual Studio 11.0/Common7/Tools/";
 					if (Directory.Exists("C:/Program Files (x86)/Microsoft Visual Studio 11.0/Common7/Tools/") == false)
 					{
-						throw new BuildException("Visual Studio 2011 must be installed in order to build this target.");
+						throw new BuildException("Visual Studio 2012 must be installed in order to build this target.");
 					}
 				}
 
diff --git a/Engine/Source/Runtime/Engine/Private/Vehicles/WheeledVehicleMovementComponent.cpp b/Engine/Source/Runtime/Engine/Private/Vehicles/WheeledVehicleMovementComponent.cpp
--- a/Engine/Source/Runtime/Engine/Private/Vehicles/WheeledVehicleMovementComponent.cpp
+++ b/Engine/Source/Runtime/Engine/Private/Vehicles/WheeledVehicleMovementComponent.cpp
@@ -1066,22 +1066,6 @@ bool UWheeledVehicleMovementComponent::CheckSlipThreshold(float AbsLongSlipThres
 	return false;
 }
 
-void UWheeledVehicleMovementComponent::Serialize(FArchive & Ar)
-{
-	Super::Serialize(Ar);
-	if (Ar.IsLoading() && Ar.UE4Ver() < VER_UE4_VEHICLES_UNIT_CHANGE)
-	{
-		for (int32 WheelIdx = 0; WheelIdx < WheelSetups.Num(); ++WheelIdx)
-		{
-			//this backwards compatable code breaks if they were using some strange value that is exactly the new default in the new units
-			FWheelSetup & WheelSetup = WheelSetups[WheelIdx];
-			WheelSetup.MaxHandBrakeTorque = WheelSetup.MaxHandBrakeTorque != 1500.f ? Cm2ToM2(WheelSetup.MaxHandBrakeTorque) : 1500.f;	//need to convert from cm^2 to m^2
-			WheelSetup.MaxBrakeTorque = WheelSetup.MaxHandBrakeTorque != 3000.f ? Cm2ToM2(WheelSetup.MaxBrakeTorque) : 3000.f;			//need to convert from cm^2 to m^2
-		}
-
-	}
-}
-
 float UWheeledVehicleMovementComponent::GetMaxSpringForce() const
 {
 	if (PVehicle == NULL)
@@ -1105,13 +1089,6 @@ float UWheeledVehicleMovementComponent::GetMaxSpringForce() const
 
 }
 
-void UWheeledVehicleMovementComponent::ComputeConstants()
-{
-	DragArea = ChassisWidth * ChassisHeight;
-	MaxEngineRPM = 5000.f;
-}
-
-
 void UWheeledVehicleMovementComponent::DrawDebug( UCanvas* Canvas, float& YL, float& YPos )
 {
 	if ( PVehicle == NULL )
@@ -1312,4 +1289,26 @@ void UWheeledVehicleMovementComponent::GetLifetimeReplicatedProps( TArray< FLife
 	Super::GetLifetimeReplicatedProps( OutLifetimeProps );
 
 	DOREPLIFETIME( UWheeledVehicleMovementComponent, ReplicatedState );
-}
\ No newline at end of file
+}
+
+void UWheeledVehicleMovementComponent::ComputeConstants()
+{
+	DragArea = ChassisWidth * ChassisHeight;
+	MaxEngineRPM = 5000.f;
+}
+
+void UWheeledVehicleMovementComponent::Serialize(FArchive & Ar)
+{
+	Super::Serialize(Ar);
+	if (Ar.IsLoading() && Ar.UE4Ver() < VER_UE4_VEHICLES_UNIT_CHANGE)
+	{
+		for (int32 WheelIdx = 0; WheelIdx < WheelSetups.Num(); ++WheelIdx)
+		{
+			//this backwards compatable code breaks if they were using some strange value that is exactly the new default in the new units
+			FWheelSetup & WheelSetup = WheelSetups[WheelIdx];
+			WheelSetup.MaxHandBrakeTorque = WheelSetup.MaxHandBrakeTorque != 1500.f ? Cm2ToM2(WheelSetup.MaxHandBrakeTorque) : 1500.f;	//need to convert from cm^2 to m^2
+			WheelSetup.MaxBrakeTorque = WheelSetup.MaxHandBrakeTorque != 3000.f ? Cm2ToM2(WheelSetup.MaxBrakeTorque) : 3000.f;			//need to convert from cm^2 to m^2
+		}
+
+	}
+}
