Hiding Opus code from platforms that don't have compiled libs

Also added opus to third party static dependencies for mac and linux so that mac will hopefully work.

diff --git a/Engine/Source/Runtime/Engine/Engine.Build.cs b/Engine/Source/Runtime/Engine/Engine.Build.cs
--- a/Engine/Source/Runtime/Engine/Engine.Build.cs
+++ b/Engine/Source/Runtime/Engine/Engine.Build.cs
@@ -201,7 +201,8 @@ public class Engine : ModuleRules
 		{
 			AddThirdPartyPrivateStaticDependencies(Target, 
 				"UEOgg",
-				"Vorbis"
+				"Vorbis",
+				"libOpus"
 				);
 		}
 
@@ -219,7 +220,8 @@ public class Engine : ModuleRules
 			AddThirdPartyPrivateStaticDependencies(Target,
 				"UEOgg",
 				"Vorbis",
-				"VorbisFile"
+				"VorbisFile",
+				"libOpus"
 				);
 		}
 
diff --git a/Engine/Source/Runtime/Engine/Private/AudioDecompress.cpp b/Engine/Source/Runtime/Engine/Private/AudioDecompress.cpp
--- a/Engine/Source/Runtime/Engine/Private/AudioDecompress.cpp
+++ b/Engine/Source/Runtime/Engine/Private/AudioDecompress.cpp
@@ -33,7 +33,11 @@
 #define VORBIS_BYTE_ORDER 1
 #endif
 
+#define WITH_OPUS (PLATFORM_WINDOWS || PLATFORM_MAC)
+
+#if WITH_OPUS
 #include "opus.h"
+#endif
 
 /*------------------------------------------------------------------------------------
 	FVorbisFileWrapper. Hides Vorbis structs from public headers.
@@ -321,25 +325,31 @@ struct FOpusDecoderWrapper
 {
 	FOpusDecoderWrapper(uint8 NumChannels)
 	{
+#if WITH_OPUS
 #if USE_UE4_MEM_ALLOC
 		int32 DecSize = opus_decoder_get_size(NumChannels);
 		Decoder = (OpusDecoder*)FMemory::Malloc(DecSize);
 		DecError = opus_decoder_init(Decoder, OPUS_SAMPLE_RATE, NumChannels);
 #else
 		Decoder = opus_decoder_create(OPUS_SAMPLE_RATE, NumChannels, &DecError);
+#endif
 #endif
 	}
 
 	~FOpusDecoderWrapper()
 	{
+#if WITH_OPUS
 #if USE_UE4_MEM_ALLOC
 		FMemory::Free(Decoder);
 #else
 		opus_encoder_destroy(Decoder);
+#endif
 #endif
 	}
 
+#if WITH_OPUS
 	OpusDecoder* Decoder;
+#endif
 	int32 DecError;
 };
 
@@ -393,11 +403,15 @@ bool FOpusAudioInfo::ReadCompressedInfo(const uint8* InSrcBufferData, uint32 InS
 	OpusDecoderWrapper = new FOpusDecoderWrapper(NumChannels);
 	LastDecodedPCM.Empty();
 	LastDecodedPCM.AddUninitialized(OPUS_MAX_FRAME_SIZE * SAMPLE_SIZE * NumChannels);
+#if WITH_OPUS
 	if (OpusDecoderWrapper->DecError != OPUS_OK)
 	{
 		delete OpusDecoderWrapper;
 		return false;
 	}
+#else
+	return false;
+#endif
 
 	if (QualityInfo)
 	{
@@ -433,7 +447,11 @@ bool FOpusAudioInfo::ReadCompressedData(uint8* Destination, bool bLooping, uint3
 	while (RawPCMOffset < BufferSize)
 	{
 		Read(&FrameSize, sizeof(uint16));
+#if WITH_OPUS
 		int32 DecodedSamples = opus_decode(OpusDecoderWrapper->Decoder, SrcBufferData + SrcBufferOffset, FrameSize, (int16*)(LastDecodedPCM.GetTypedData()), OPUS_MAX_FRAME_SIZE, 0);
+#else
+		int32 DecodedSamples = -1;
+#endif
 		SrcBufferOffset += FrameSize;
 		if (DecodedSamples < 0)
 		{
@@ -498,7 +516,11 @@ void FOpusAudioInfo::ExpandFile(uint8* DstBuffer, struct FSoundQualityInfo* Qual
 	while (RawPCMOffset < QualityInfo->SampleDataSize)
 	{
 		Read(&FrameSize, sizeof(uint16));
+#if WITH_OPUS
 		int32 DecodedSamples = opus_decode(OpusDecoderWrapper->Decoder, (const unsigned char*)(SrcBufferData + SrcBufferOffset), FrameSize, (int16*)(TempPCMData.GetTypedData()), OPUS_MAX_FRAME_SIZE, 0);
+#else
+		int32 DecodedSamples = -1;
+#endif
 		if (DecodedSamples < 0)
 		{
 			break;
