Fix for Shaders being compiled while the engine is shutdown causing crash on exit.

diff --git a/Engine/Source/Runtime/Engine/Private/ShaderCompiler/ShaderCompiler.cpp b/Engine/Source/Runtime/Engine/Private/ShaderCompiler/ShaderCompiler.cpp
--- a/Engine/Source/Runtime/Engine/Private/ShaderCompiler/ShaderCompiler.cpp
+++ b/Engine/Source/Runtime/Engine/Private/ShaderCompiler/ShaderCompiler.cpp
@@ -427,7 +427,9 @@ struct FShaderCompileWorkerInfo
 FShaderCompileThreadRunnable::FShaderCompileThreadRunnable(FShaderCompilingManager* InManager) :
 	Manager(InManager),
 	Thread(NULL),
-	bTerminatedByError(false)
+	bTerminatedByError(false),
+	bForceFinish( false ),
+	bIsRunning( false )
 {
 	LastCheckForWorkersTime = 0;
 
@@ -464,7 +466,7 @@ uint32 FShaderCompileThreadRunnable::Run()
 		{
 			check(Manager->bAllowAsynchronousShaderCompiling);
 			// Do the work
-			while (true)
+			while (!bForceFinish)
 			{
 				CompilingLoop();
 			}
@@ -493,7 +495,7 @@ uint32 FShaderCompileThreadRunnable::Run()
 #endif
 	{
 		check(Manager->bAllowAsynchronousShaderCompiling);
-		while (true)
+		while (!bForceFinish)
 		{
 			CompilingLoop();
 		}
@@ -1472,6 +1474,22 @@ void FShaderCompilingManager::ProcessCompiledShaderMaps(
 	}
 }
 
+
+/**
+ * Shutdown the shader compile manager
+ * this function should be used when ending the game to shutdown shader compile threads
+ * will not complete current pending shader compilation
+ */
+void FShaderCompilingManager::Shutdown()
+{
+	Thread->Stop();
+	while ( Thread->IsRunning() )
+	{
+		FPlatformProcess::Sleep(0.01f);
+	}
+}
+
+
 bool FShaderCompilingManager::HandlePotentialRetryOnError(TMap<int32, FShaderMapFinalizeResults>& CompletedShaderMaps)
 {
 	bool bRetryCompile = false;
diff --git a/Engine/Source/Runtime/Engine/Public/ShaderCompiler.h b/Engine/Source/Runtime/Engine/Public/ShaderCompiler.h
--- a/Engine/Source/Runtime/Engine/Public/ShaderCompiler.h
+++ b/Engine/Source/Runtime/Engine/Public/ShaderCompiler.h
@@ -70,15 +70,17 @@ private:
 	/** Tracks the last time that this thread checked if the workers were still active. */
 	double LastCheckForWorkersTime;
 
+	volatile bool bForceFinish;
+	volatile bool bIsRunning;
 public:
 	/** Initialization constructor. */
 	FShaderCompileThreadRunnable(class FShaderCompilingManager* InManager);
 	virtual ~FShaderCompileThreadRunnable();
 
 	// FRunnable interface.
-	virtual bool Init(void) { return true; }
-	virtual void Exit(void) {}
-	virtual void Stop(void) {}
+	virtual bool Init(void) { return true; bIsRunning = true; }
+	virtual void Exit(void) { bIsRunning = false; }
+	virtual void Stop(void) { bForceFinish = true; }
 	virtual uint32 Run(void);
 
 	/** Checks the thread's health, and passes on any errors that have occured.  Called by the main thread. */
@@ -109,6 +111,8 @@ private:
 
 	/** Used when compiling through workers, launches worker processes if needed. */
 	void LaunchWorkerIfNeeded(FShaderCompileWorkerInfo & CurrentWorkerInfo, uint32 WorkerIndex);
+
+	inline bool IsRunning() { return bIsRunning; }
 };
 
 /** Results for a single compiled shader map. */
@@ -266,6 +270,7 @@ public:
 		return AbsoluteShaderDebugInfoDirectory;
 	}
 
+
 	/** 
 	 * Adds shader jobs to be asynchronously compiled. 
 	 * FinishCompilation or ProcessAsyncResults must be used to get the results.
@@ -284,6 +289,13 @@ public:
 	 */
 	ENGINE_API void FinishAllCompilation();
 
+
+	/** 
+	 * Shutdown the shader compiler manager, this will shutdown immediately and not process any more shader compile requests. 
+	 */
+	ENGINE_API void Shutdown();
+
+
 	/** 
 	 * Processes completed asynchronous shader maps, and assigns them to relevant materials.
 	 * @param bLimitExecutionTime - When enabled, ProcessAsyncResults will be bandwidth throttled by ProcessGameThreadTargetTime, to limit hitching.
diff --git a/Engine/Source/Runtime/Launch/Private/LaunchEngineLoop.cpp b/Engine/Source/Runtime/Launch/Private/LaunchEngineLoop.cpp
--- a/Engine/Source/Runtime/Launch/Private/LaunchEngineLoop.cpp
+++ b/Engine/Source/Runtime/Launch/Private/LaunchEngineLoop.cpp
@@ -2482,6 +2482,13 @@ void FEngineLoop::AppPreExit( )
 	{
 		GThreadPool->Destroy();
 	}
+
+#if WITH_ENGINE
+	if ( GShaderCompilingManager )
+	{
+		GShaderCompilingManager->Shutdown();
+	}
+#endif
 }
 
 
