Partially roll back CL#2056352 - TargetPlatformManager doesn't exist in cooked builds, force Vorbis again for now

diff --git a/Engine/Source/Runtime/Android/AndroidAudio/Private/AndroidAudioBuffer.cpp b/Engine/Source/Runtime/Android/AndroidAudio/Private/AndroidAudioBuffer.cpp
--- a/Engine/Source/Runtime/Android/AndroidAudio/Private/AndroidAudioBuffer.cpp
+++ b/Engine/Source/Runtime/Android/AndroidAudio/Private/AndroidAudioBuffer.cpp
@@ -5,7 +5,6 @@
 #include "Engine.h"
 #include "IAudioFormat.h"
 #include "AudioDecompress.h"
-#include "TargetPlatform.h"
 
 /*------------------------------------------------------------------------------------
 	FSLESSoundBuffer.
@@ -52,13 +51,7 @@ FSLESSoundBuffer* FSLESSoundBuffer::CreateQueuedBuffer( FSLESAudioDevice* AudioD
 	// Prime the first two buffers and prepare the decompression
 	FSoundQualityInfo QualityInfo = { 0 };
 	
-	const IAudioFormat* Format = NULL;
-	ITargetPlatformManagerModule* TPM = GetTargetPlatformManager();
-	if (TPM)
-	{
-		Format = TPM->FindAudioFormat(AudioDevice->GetRuntimeFormat());
-	}
-	Buffer->DecompressionState = Format->CreateCompressedAudioInfo();
+	Buffer->DecompressionState = new FVorbisAudioInfo();
 
 	InWave->InitAudioResource( AudioDevice->GetRuntimeFormat() );
 	
diff --git a/Engine/Source/Runtime/Engine/Private/AudioDecompress.cpp b/Engine/Source/Runtime/Engine/Private/AudioDecompress.cpp
--- a/Engine/Source/Runtime/Engine/Private/AudioDecompress.cpp
+++ b/Engine/Source/Runtime/Engine/Private/AudioDecompress.cpp
@@ -541,13 +541,7 @@ void FOpusAudioInfo::ExpandFile(uint8* DstBuffer, struct FSoundQualityInfo* Qual
  */
 void FAsyncAudioDecompressWorker::DoWork( void )
 {
-	const IAudioFormat* Format = NULL;
-	ITargetPlatformManagerModule* TPM = GetTargetPlatformManager();
-	if (TPM)
-	{
-		Format = TPM->FindAudioFormat(GEngine->GetAudioDevice()->GetRuntimeFormat());
-	}
-	ICompressedAudioInfo* AudioInfo = Format->CreateCompressedAudioInfo();
+	ICompressedAudioInfo* AudioInfo = new FVorbisAudioInfo;
 	FSoundQualityInfo	QualityInfo = { 0 };
 
 	// Parse the audio header for the relevant information
diff --git a/Engine/Source/Runtime/Engine/Private/AudioDevice.cpp b/Engine/Source/Runtime/Engine/Private/AudioDevice.cpp
--- a/Engine/Source/Runtime/Engine/Private/AudioDevice.cpp
+++ b/Engine/Source/Runtime/Engine/Private/AudioDevice.cpp
@@ -2252,13 +2252,6 @@ void FAudioDevice::Precache(USoundWave* SoundWave, bool bSynchronous, bool bTrac
 		return;
 	}
 
-	const IAudioFormat* Format = NULL;
-	ITargetPlatformManagerModule* TPM = GetTargetPlatformManager();
-	if (TPM)
-	{
-		Format = TPM->FindAudioFormat(GetRuntimeFormat());
-	}
-
 	// calculate the decompression type
 	// @todo audio: maybe move this into SoundWave?
 	if( SoundWave->NumChannels == 0 )
@@ -2276,7 +2269,8 @@ void FAudioDevice::Precache(USoundWave* SoundWave, bool bSynchronous, bool bTrac
 		// Streaming data, created programmatically.
 		SoundWave->DecompressionType = DTYPE_Procedural;
 	}
-	else if (Format->HasCompressedAudioInfoClass())
+#if WITH_OGGVORBIS
+	else
 	{
 		const FSoundGroup& SoundGroup = GetDefault<USoundGroups>()->GetSoundGroup(SoundWave->SoundGroup);
 
@@ -2327,11 +2321,13 @@ void FAudioDevice::Precache(USoundWave* SoundWave, bool bSynchronous, bool bTrac
 			bTrackMemory = false;
 		}
 	}
+#else
 	else
 	{
 		// Preserve old behavior if there is no compressed audio info class for this audio format
 		SoundWave->DecompressionType = DTYPE_Native;
 	}
+#endif
 
 	if (bTrackMemory)
 	{
diff --git a/Engine/Source/Runtime/Mac/CoreAudio/Private/CoreAudioBuffer.cpp b/Engine/Source/Runtime/Mac/CoreAudio/Private/CoreAudioBuffer.cpp
--- a/Engine/Source/Runtime/Mac/CoreAudio/Private/CoreAudioBuffer.cpp
+++ b/Engine/Source/Runtime/Mac/CoreAudio/Private/CoreAudioBuffer.cpp
@@ -14,7 +14,6 @@
 #include "CoreAudioEffects.h"
 #include "Engine.h"
 #include "IAudioFormat.h"
-#include "TargetPlatform.h"
 
 /*------------------------------------------------------------------------------------
  FCoreAudioSoundBuffer.
@@ -148,13 +147,7 @@ FCoreAudioSoundBuffer* FCoreAudioSoundBuffer::CreateQueuedBuffer( FCoreAudioDevi
 	// Prime the first two buffers and prepare the decompression
 	FSoundQualityInfo QualityInfo = { 0 };
 
-	const IAudioFormat* Format = NULL;
-	ITargetPlatformManagerModule* TPM = GetTargetPlatformManager();
-	if (TPM)
-	{
-		Format = TPM->FindAudioFormat(CoreAudioDevice->GetRuntimeFormat());
-	}
-	Buffer->DecompressionState = Format->CreateCompressedAudioInfo();
+	Buffer->DecompressionState = new FVorbisAudioInfo();
 	
 	Wave->InitAudioResource( CoreAudioDevice->GetRuntimeFormat() );
 	
diff --git a/Engine/Source/Runtime/Windows/XAudio2/Private/XAudio2Buffer.cpp b/Engine/Source/Runtime/Windows/XAudio2/Private/XAudio2Buffer.cpp
--- a/Engine/Source/Runtime/Windows/XAudio2/Private/XAudio2Buffer.cpp
+++ b/Engine/Source/Runtime/Windows/XAudio2/Private/XAudio2Buffer.cpp
@@ -326,13 +326,7 @@ FXAudio2SoundBuffer* FXAudio2SoundBuffer::CreateQueuedBuffer( FXAudio2Device* XA
 	// Prime the first two buffers and prepare the decompression
 	FSoundQualityInfo QualityInfo = { 0 };
 
-	const IAudioFormat* Format = NULL;
-	ITargetPlatformManagerModule* TPM = GetTargetPlatformManager();
-	if (TPM)
-	{
-		Format = TPM->FindAudioFormat(XAudio2Device->GetRuntimeFormat());
-	}
-	Buffer->DecompressionState = Format->CreateCompressedAudioInfo();
+	Buffer->DecompressionState = new FVorbisAudioInfo();
 
 	Wave->InitAudioResource(XAudio2Device->GetRuntimeFormat());
 
