Fix compilation when WITH_OGGVORBIS not defined

diff --git a/Engine/Source/Runtime/Engine/Private/AudioDecompress.cpp b/Engine/Source/Runtime/Engine/Private/AudioDecompress.cpp
--- a/Engine/Source/Runtime/Engine/Private/AudioDecompress.cpp
+++ b/Engine/Source/Runtime/Engine/Private/AudioDecompress.cpp
@@ -541,6 +541,7 @@ void FOpusAudioInfo::ExpandFile(uint8* DstBuffer, struct FSoundQualityInfo* Qual
  */
 void FAsyncAudioDecompressWorker::DoWork( void )
 {
+#if WITH_OGGVORBIS
 	ICompressedAudioInfo* AudioInfo = new FVorbisAudioInfo;
 	FSoundQualityInfo	QualityInfo = { 0 };
 
@@ -581,6 +582,9 @@ void FAsyncAudioDecompressWorker::DoWork( void )
 	Wave->RemoveAudioResource();
 
 	delete AudioInfo;
+#else
+	check(0);
+#endif
 }
 
 // end
