Make motion axes bindable
Begin generalization of some of the input binding common components to reduce duplicated code

diff --git a/Engine/Source/Editor/BlueprintGraph/Classes/K2Node_GetInputVectorAxisValue.h b/Engine/Source/Editor/BlueprintGraph/Classes/K2Node_GetInputVectorAxisValue.h
new file mode 100644
--- /dev/null
+++ b/Engine/Source/Editor/BlueprintGraph/Classes/K2Node_GetInputVectorAxisValue.h
@@ -0,0 +1,23 @@
+// Copyright 1998-2014 Epic Games, Inc. All Rights Reserved.
+
+#pragma once
+#include "K2Node_GetInputVectorAxisValue.generated.h"
+
+UCLASS(MinimalAPI)
+class UK2Node_GetInputVectorAxisValue : public UK2Node_GetInputAxisKeyValue
+{
+	GENERATED_UCLASS_BODY()
+
+	// Begin EdGraphNode interface
+	virtual FString GetTooltip() const OVERRIDE;
+	// End EdGraphNode interface
+
+	// Begin UK2Node interface
+	virtual void ValidateNodeDuringCompilation(class FCompilerResultsLog& MessageLog) const OVERRIDE;
+	virtual bool ShouldShowNodeProperties() const OVERRIDE { return true; }
+	virtual UClass* GetDynamicBindingClass() const OVERRIDE;
+	virtual void RegisterDynamicBinding(UDynamicBlueprintBinding* BindingObject) const OVERRIDE;
+	// End UK2Node interface
+	
+	void Initialize(const FKey AxisKey);
+};
\ No newline at end of file
diff --git a/Engine/Source/Editor/BlueprintGraph/Classes/K2Node_InputVectorAxisEvent.h b/Engine/Source/Editor/BlueprintGraph/Classes/K2Node_InputVectorAxisEvent.h
new file mode 100644
--- /dev/null
+++ b/Engine/Source/Editor/BlueprintGraph/Classes/K2Node_InputVectorAxisEvent.h
@@ -0,0 +1,16 @@
+// Copyright 1998-2014 Epic Games, Inc. All Rights Reserved.
+
+#pragma once
+#include "K2Node_InputVectorAxisEvent.generated.h"
+
+UCLASS(MinimalAPI)
+class UK2Node_InputVectorAxisEvent : public UK2Node_InputAxisKeyEvent
+{
+	GENERATED_UCLASS_BODY()
+
+	// Begin UK2Node interface
+	virtual void ValidateNodeDuringCompilation(class FCompilerResultsLog& MessageLog) const OVERRIDE;
+	virtual UClass* GetDynamicBindingClass() const OVERRIDE;
+	virtual void RegisterDynamicBinding(UDynamicBlueprintBinding* BindingObject) const OVERRIDE;
+	// End UK2Node interface
+};
\ No newline at end of file
diff --git a/Engine/Source/Editor/BlueprintGraph/Private/K2ActionMenuBuilder.cpp b/Engine/Source/Editor/BlueprintGraph/Private/K2ActionMenuBuilder.cpp
--- a/Engine/Source/Editor/BlueprintGraph/Private/K2ActionMenuBuilder.cpp
+++ b/Engine/Source/Editor/BlueprintGraph/Private/K2ActionMenuBuilder.cpp
@@ -177,7 +177,7 @@ static void GetInputNodes(FGraphActionListBuilderBase& ActionMenuBuilder, const
 					Action->Keywords = TEXT("Input Key InputKey Keyboard");
 				}
 
-				if (Key.IsAxis())
+				if (Key.IsFloatAxis())
 				{
 					UK2Node_InputAxisKeyEvent* InputKeyAxisNode = ActionMenuBuilder.CreateTemplateNode<UK2Node_InputAxisKeyEvent>();
 					InputKeyAxisNode->Initialize(Key);
@@ -191,6 +191,20 @@ static void GetInputNodes(FGraphActionListBuilderBase& ActionMenuBuilder, const
 					GetAction->NodeTemplate = GetInputAxisKeyValue;
 					GetAction->Keywords = Action->Keywords;
 				}
+				else if (Key.IsVectorAxis())
+				{
+					UK2Node_InputVectorAxisEvent* InputVectorAxisNode = ActionMenuBuilder.CreateTemplateNode<UK2Node_InputVectorAxisEvent>();
+					InputVectorAxisNode->Initialize(Key);
+					Action->NodeTemplate = InputVectorAxisNode;
+					Action->SearchTitle = Action->NodeTemplate->GetNodeSearchTitle();
+
+					const FString& KeyValuesCategory = (bGamepad ? K2ActionCategories::GamepadValuesCategory : (bMouse ? K2ActionCategories::MouseValuesCategory : K2ActionCategories::KeyValuesCategory));
+					GetAction = FK2ActionMenuBuilder::AddNewNodeAction(ActionMenuBuilder, KeyValuesCategory, KeyName, KeyName.ToString());
+					UK2Node_GetInputVectorAxisValue* GetInputVectorAxisValue = ActionMenuBuilder.CreateTemplateNode<UK2Node_GetInputVectorAxisValue>();
+					GetInputVectorAxisValue->Initialize(Key);
+					GetAction->NodeTemplate = GetInputVectorAxisValue;
+					GetAction->Keywords = Action->Keywords;
+				}
 				else
 				{
 					UK2Node_InputKey* InputKeyNode = ActionMenuBuilder.CreateTemplateNode<UK2Node_InputKey>();
diff --git a/Engine/Source/Editor/BlueprintGraph/Private/K2Node_GetInputAxisKeyValue.cpp b/Engine/Source/Editor/BlueprintGraph/Private/K2Node_GetInputAxisKeyValue.cpp
--- a/Engine/Source/Editor/BlueprintGraph/Private/K2Node_GetInputAxisKeyValue.cpp
+++ b/Engine/Source/Editor/BlueprintGraph/Private/K2Node_GetInputAxisKeyValue.cpp
@@ -46,9 +46,9 @@ void UK2Node_GetInputAxisKeyValue::ValidateNodeDuringCompilation(class FCompiler
 	{
 		MessageLog.Warning(*FText::Format(NSLOCTEXT("KismetCompiler", "Invalid_GetInputAxisKey_Warning", "GetInputAxisKey Value specifies invalid FKey'{0}' for @@"), FText::FromString(InputAxisKey.ToString())).ToString(), this);
 	}
-	else if (!InputAxisKey.IsAxis())
+	else if (!InputAxisKey.IsFloatAxis())
 	{
-		MessageLog.Warning(*FText::Format(NSLOCTEXT("KismetCompiler", "NotAxis_GetInputAxisKey_Warning", "GetInputAxisKey Value specifies non-axis FKey'{0}' for @@"), FText::FromString(InputAxisKey.ToString())).ToString(), this);
+		MessageLog.Warning(*FText::Format(NSLOCTEXT("KismetCompiler", "NotAxis_GetInputAxisKey_Warning", "GetInputAxisKey Value specifies FKey'{0}' which is not a float axis for @@"), FText::FromString(InputAxisKey.ToString())).ToString(), this);
 	}
 	else if (!InputAxisKey.IsBindableInBlueprints())
 	{
diff --git a/Engine/Source/Editor/BlueprintGraph/Private/K2Node_GetInputVectorAxisValue.cpp b/Engine/Source/Editor/BlueprintGraph/Private/K2Node_GetInputVectorAxisValue.cpp
new file mode 100644
--- /dev/null
+++ b/Engine/Source/Editor/BlueprintGraph/Private/K2Node_GetInputVectorAxisValue.cpp
@@ -0,0 +1,60 @@
+// Copyright 1998-2014 Epic Games, Inc. All Rights Reserved.
+
+#include "BlueprintGraphPrivatePCH.h"
+
+#include "CompilerResultsLog.h"
+
+UK2Node_GetInputVectorAxisValue::UK2Node_GetInputVectorAxisValue(const class FPostConstructInitializeProperties& PCIP)
+	: Super(PCIP)
+{
+	bConsumeInput = true;
+}
+
+void UK2Node_GetInputVectorAxisValue::Initialize(const FKey AxisKey)
+{
+	InputAxisKey = AxisKey;
+	SetFromFunction(AActor::StaticClass()->FindFunctionByName(TEXT("GetInputVectorAxisValue")));
+}
+
+FString UK2Node_GetInputVectorAxisValue::GetTooltip() const
+{
+	FFormatNamedArguments Args;
+	Args.Add(TEXT("AxisKey"), InputAxisKey.GetDisplayName());
+	return FText::Format(NSLOCTEXT("K2Node", "GetInputVectorAxis_Tooltip", "Returns the current value of input axis key {AxisKey}.  If input is disabled for the actor the value will be (0, 0, 0)."), Args).ToString();
+}
+
+void UK2Node_GetInputVectorAxisValue::ValidateNodeDuringCompilation(class FCompilerResultsLog& MessageLog) const
+{
+	// Skip GetInputKeyAxisValue validation
+	UK2Node_CallFunction::ValidateNodeDuringCompilation(MessageLog);
+
+	if (!InputAxisKey.IsValid())
+	{
+		MessageLog.Warning(*FText::Format(NSLOCTEXT("KismetCompiler", "Invalid_GetInputVectorAxis_Warning", "GetInputVectorAxis Value specifies invalid FKey'{0}' for @@"), FText::FromString(InputAxisKey.ToString())).ToString(), this);
+	}
+	else if (!InputAxisKey.IsVectorAxis())
+	{
+		MessageLog.Warning(*FText::Format(NSLOCTEXT("KismetCompiler", "NotAxis_GetInputVectorAxis_Warning", "GetInputVectorAxis Value specifies FKey'{0}' which is not a vector axis for @@"), FText::FromString(InputAxisKey.ToString())).ToString(), this);
+	}
+	else if (!InputAxisKey.IsBindableInBlueprints())
+	{
+		MessageLog.Warning(*FText::Format(NSLOCTEXT("KismetCompiler", "NotBindanble_GetInputVectorAxis_Warning", "GetInputVectorAxis Value specifies FKey'{0}' that is not blueprint bindable for @@"), FText::FromString(InputAxisKey.ToString())).ToString(), this);
+	}
+}
+
+UClass* UK2Node_GetInputVectorAxisValue::GetDynamicBindingClass() const
+{
+	return UInputVectorAxisDelegateBinding::StaticClass();
+}
+
+void UK2Node_GetInputVectorAxisValue::RegisterDynamicBinding(UDynamicBlueprintBinding* BindingObject) const
+{
+	UInputVectorAxisDelegateBinding* InputVectorAxisBindingObject = CastChecked<UInputVectorAxisDelegateBinding>(BindingObject);
+
+	FBlueprintInputAxisKeyDelegateBinding Binding;
+	Binding.AxisKey = InputAxisKey;
+	Binding.bConsumeInput = bConsumeInput;
+	Binding.bExecuteWhenPaused = bExecuteWhenPaused;
+
+	InputVectorAxisBindingObject->InputAxisKeyDelegateBindings.Add(Binding);
+}
\ No newline at end of file
diff --git a/Engine/Source/Editor/BlueprintGraph/Private/K2Node_InputAxisKeyEvent.cpp b/Engine/Source/Editor/BlueprintGraph/Private/K2Node_InputAxisKeyEvent.cpp
--- a/Engine/Source/Editor/BlueprintGraph/Private/K2Node_InputAxisKeyEvent.cpp
+++ b/Engine/Source/Editor/BlueprintGraph/Private/K2Node_InputAxisKeyEvent.cpp
@@ -45,9 +45,9 @@ void UK2Node_InputAxisKeyEvent::ValidateNodeDuringCompilation(class FCompilerRes
 	{
 		MessageLog.Warning(*FText::Format(NSLOCTEXT("KismetCompiler", "Invalid_InputAxisKey_Warning", "InputAxisKey Event specifies invalid FKey'{0}' for @@"), FText::FromString(AxisKey.ToString())).ToString(), this);
 	}
-	else if (!AxisKey.IsAxis())
+	else if (!AxisKey.IsFloatAxis())
 	{
-		MessageLog.Warning(*FText::Format(NSLOCTEXT("KismetCompiler", "NotAxis_InputAxisKey_Warning", "InputAxisKey Event specifies non-axis FKey'{0}' for @@"), FText::FromString(AxisKey.ToString())).ToString(), this);
+		MessageLog.Warning(*FText::Format(NSLOCTEXT("KismetCompiler", "NotAxis_InputAxisKey_Warning", "InputAxisKey Event specifies FKey'{0}' which is not a float axis for @@"), FText::FromString(AxisKey.ToString())).ToString(), this);
 	}
 	else if (!AxisKey.IsBindableInBlueprints())
 	{
diff --git a/Engine/Source/Editor/BlueprintGraph/Private/K2Node_InputKey.cpp b/Engine/Source/Editor/BlueprintGraph/Private/K2Node_InputKey.cpp
--- a/Engine/Source/Editor/BlueprintGraph/Private/K2Node_InputKey.cpp
+++ b/Engine/Source/Editor/BlueprintGraph/Private/K2Node_InputKey.cpp
@@ -203,7 +203,7 @@ void UK2Node_InputKey::ValidateNodeDuringCompilation(class FCompilerResultsLog&
 	{
 		MessageLog.Warning(*FText::Format(NSLOCTEXT("KismetCompiler", "Invalid_InputKey_Warning", "InputKey Event specifies invalid FKey'{0}' for @@"), FText::FromString(InputKey.ToString())).ToString(), this);
 	}
-	else if (InputKey.IsAxis())
+	else if (InputKey.IsFloatAxis())
 	{
 		MessageLog.Warning(*FText::Format(NSLOCTEXT("KismetCompiler", "Axis_InputKey_Warning", "InputKey Event specifies axis FKey'{0}' for @@"), FText::FromString(InputKey.ToString())).ToString(), this);
 	}
diff --git a/Engine/Source/Editor/BlueprintGraph/Private/K2Node_InputVectorAxisEvent.cpp b/Engine/Source/Editor/BlueprintGraph/Private/K2Node_InputVectorAxisEvent.cpp
new file mode 100644
--- /dev/null
+++ b/Engine/Source/Editor/BlueprintGraph/Private/K2Node_InputVectorAxisEvent.cpp
@@ -0,0 +1,49 @@
+// Copyright 1998-2014 Epic Games, Inc. All Rights Reserved.
+
+#include "BlueprintGraphPrivatePCH.h"
+
+#include "CompilerResultsLog.h"
+
+UK2Node_InputVectorAxisEvent::UK2Node_InputVectorAxisEvent(const class FPostConstructInitializeProperties& PCIP)
+: Super(PCIP)
+{
+	EventSignatureName = TEXT("InputVectorAxisHandlerDynamicSignature__DelegateSignature");
+}
+
+void UK2Node_InputVectorAxisEvent::ValidateNodeDuringCompilation(class FCompilerResultsLog& MessageLog) const
+{
+	// Skip AxisKeyEvent validation
+	UK2Node_Event::ValidateNodeDuringCompilation(MessageLog);
+
+	if (!AxisKey.IsValid())
+	{
+		MessageLog.Warning(*FText::Format(NSLOCTEXT("KismetCompiler", "Invalid_InputVectorAxis_Warning", "InputVectorAxis Event specifies invalid FKey'{0}' for @@"), FText::FromString(AxisKey.ToString())).ToString(), this);
+	}
+	else if (!AxisKey.IsVectorAxis())
+	{
+		MessageLog.Warning(*FText::Format(NSLOCTEXT("KismetCompiler", "NotAxis_InputVectorAxis_Warning", "InputVectorAxis Event specifies FKey'{0}' which is not a vector axis for @@"), FText::FromString(AxisKey.ToString())).ToString(), this);
+	}
+	else if (!AxisKey.IsBindableInBlueprints())
+	{
+		MessageLog.Warning(*FText::Format(NSLOCTEXT("KismetCompiler", "NotBindable_InputVectorAxis_Warning", "InputVectorAxis Event specifies FKey'{0}' that is not blueprint bindable for @@"), FText::FromString(AxisKey.ToString())).ToString(), this);
+	}
+}
+
+UClass* UK2Node_InputVectorAxisEvent::GetDynamicBindingClass() const
+{
+	return UInputVectorAxisDelegateBinding::StaticClass();
+}
+
+void UK2Node_InputVectorAxisEvent::RegisterDynamicBinding(UDynamicBlueprintBinding* BindingObject) const
+{
+	UInputVectorAxisDelegateBinding* InputVectorAxisBindingObject = CastChecked<UInputVectorAxisDelegateBinding>(BindingObject);
+
+	FBlueprintInputAxisKeyDelegateBinding Binding;
+	Binding.AxisKey = AxisKey;
+	Binding.bConsumeInput = bConsumeInput;
+	Binding.bExecuteWhenPaused = bExecuteWhenPaused;
+	Binding.bOverrideParentBinding = bOverrideParentBinding;
+	Binding.FunctionNameToBind = CustomFunctionName;
+
+	InputVectorAxisBindingObject->InputAxisKeyDelegateBindings.Add(Binding);
+}
\ No newline at end of file
diff --git a/Engine/Source/Runtime/Engine/Classes/Components/InputComponent.h b/Engine/Source/Runtime/Engine/Classes/Components/InputComponent.h
--- a/Engine/Source/Runtime/Engine/Classes/Components/InputComponent.h
+++ b/Engine/Source/Runtime/Engine/Classes/Components/InputComponent.h
@@ -133,8 +133,22 @@
 	}
 };
 
+struct FInputBinding
+{
+	/** Whether the binding should consume the input or allow it to pass to another component */
+	uint32 bConsumeInput:1;
+
+	/** Whether the binding should execute while paused */
+	uint32 bExecuteWhenPaused:1;
+
+	FInputBinding()
+		: bConsumeInput(true)
+		, bExecuteWhenPaused(false)
+	{}
+};
+
 /** Binds a game-defined action. */
-struct FInputActionBinding
+struct FInputActionBinding : public FInputBinding
 {
 	/** Friendly name of action, e.g "jump" */
 	FName ActionName;
@@ -142,35 +156,27 @@
 	/** Key event to bind it to, e.g. pressed, released, dblclick */
 	TEnumAsByte<EInputEvent> KeyEvent;
 
-	/** Whether the binding should consume the input or allow it to pass to another component */
-	uint32 bConsumeInput:1;
-
-	/** Whether the binding should execute while paused */
-	uint32 bExecuteWhenPaused:1;
-
 	/** Whether the binding is part of a paired (both pressed and released events bound) action */
 	uint32 bPaired:1;
 
 	FInputActionUnifiedDelegate ActionDelegate;
 
 	FInputActionBinding()
-		: ActionName(NAME_None)
+		: FInputBinding()
+		, ActionName(NAME_None)
 		, KeyEvent(EInputEvent::IE_Pressed)
-		, bConsumeInput(true)
-		, bExecuteWhenPaused(false)
 		, bPaired(false)
 	{}
 
 	FInputActionBinding(const FName InActionName, const enum EInputEvent InKeyEvent)
-		: ActionName(InActionName)
+		: FInputBinding()
+		, ActionName(InActionName)
 		, KeyEvent(InKeyEvent)
-		, bConsumeInput(true)
-		, bExecuteWhenPaused(false)
 		, bPaired(false)
 	{}
 };
 
-struct FInputKeyBinding
+struct FInputKeyBinding : public FInputBinding
 {
 	/** Raw key to bind to */
 	FInputChord Chord;
@@ -178,23 +184,17 @@
 	/** Key event to bind it to, e.g. pressed, released, dblclick */
 	TEnumAsByte<EInputEvent> KeyEvent;
 
-	uint32 bConsumeInput:1;
-
-	uint32 bExecuteWhenPaused:1;
-
 	FInputActionUnifiedDelegate KeyDelegate;
 
 	FInputKeyBinding()
-		: KeyEvent(EInputEvent::IE_Pressed)
-		, bConsumeInput(true)
-		, bExecuteWhenPaused(false)
+		: FInputBinding()
+		, KeyEvent(EInputEvent::IE_Pressed)
 	{}
 
 	FInputKeyBinding(const FInputChord InChord, const enum EInputEvent InKeyEvent)
-		: Chord(InChord)
+		: FInputBinding()
+		, Chord(InChord)
 		, KeyEvent(InKeyEvent)
-		, bConsumeInput(true)
-		, bExecuteWhenPaused(false)
 	{}
 };
 
@@ -222,27 +222,21 @@
 };
 
 
-struct FInputTouchBinding
+struct FInputTouchBinding : public FInputBinding
 {
 	/** Key event to bind it to, e.g. pressed, released, dblclick */
 	TEnumAsByte<EInputEvent> KeyEvent;
 
-	uint32 bConsumeInput:1;
-
-	uint32 bExecuteWhenPaused:1;
-
 	FInputTouchUnifiedDelegate TouchDelegate;
 
 	FInputTouchBinding()
-		: KeyEvent(EInputEvent::IE_Pressed)
-		, bConsumeInput(true)
-		, bExecuteWhenPaused(false)
+		: FInputBinding()
+		, KeyEvent(EInputEvent::IE_Pressed)
 	{}
 
 	FInputTouchBinding(const enum EInputEvent InKeyEvent)
-		: KeyEvent(InKeyEvent)
-		, bConsumeInput(true)
-		, bExecuteWhenPaused(false)
+		: FInputBinding()
+		, KeyEvent(InKeyEvent)
 	{}
 };
 
@@ -272,60 +266,91 @@
 
 
 /** Binds a game-defined axis to a function */
-struct FInputAxisBinding
+struct FInputAxisBinding : public FInputBinding
 {
 	FName AxisName;
 
 	FInputAxisUnifiedDelegate AxisDelegate;
 
-	uint32 bConsumeInput:1;
-
-	uint32 bExecuteWhenPaused:1;
-
 	float AxisValue;
 
 	FInputAxisBinding()
-		: AxisName(NAME_None)
-		, bConsumeInput(true)
-		, bExecuteWhenPaused(false)
+		: FInputBinding()
+		, AxisName(NAME_None)
 		, AxisValue(0.f)
 	{}
 
 	FInputAxisBinding(const FName InAxisName)
-		: AxisName(InAxisName)
-		, bConsumeInput(true)
-		, bExecuteWhenPaused(false)
+		: FInputBinding()
+		, AxisName(InAxisName)
 		, AxisValue(0.f)
 	{}
 };
 
 /** Binds a raw axis to a function */
-struct FInputAxisKeyBinding
+struct FInputAxisKeyBinding : public FInputBinding
 {
 	FKey AxisKey;
 
 	FInputAxisUnifiedDelegate AxisDelegate;
 
-	uint32 bConsumeInput : 1;
-
-	uint32 bExecuteWhenPaused : 1;
-
 	float AxisValue;
 
 	FInputAxisKeyBinding()
-		: bConsumeInput(true)
-		, bExecuteWhenPaused(false)
+		: FInputBinding()
 		, AxisValue(0.f)
 	{}
 
 	FInputAxisKeyBinding(const FKey InAxisKey)
-		: AxisKey(InAxisKey)
-		, bConsumeInput(true)
-		, bExecuteWhenPaused(false)
+		: FInputBinding()
+		, AxisKey(InAxisKey)
 		, AxisValue(0.f)
-	{}
+	{
+		ensure(AxisKey.IsFloatAxis());
+	}
 };
 
+/**
+* Delegate signature for vector axis handlers.
+* @AxisValue: "Value" to pass to the axis.
+*/
+DECLARE_DELEGATE_OneParam(FInputVectorAxisHandlerSignature, FVector);
+DECLARE_DYNAMIC_DELEGATE_OneParam(FInputVectorAxisHandlerDynamicSignature, FVector, AxisValue);
 
+struct FInputVectorAxisUnifiedDelegate : public TInputUnifiedDelegate<FInputVectorAxisHandlerSignature, FInputVectorAxisHandlerDynamicSignature>
+{
+	inline void Execute(const FVector AxisValue) const
+	{
+		if (FuncDelegate.IsBound())
+		{
+			FuncDelegate.Execute(AxisValue);
+		}
+		else if (FuncDynDelegate.IsBound())
+		{
+			FuncDynDelegate.Execute(AxisValue);
+		}
+	}
+};
+
+struct FInputVectorAxisBinding : public FInputBinding
+{
+	FKey AxisKey;
+
+	FInputVectorAxisUnifiedDelegate AxisDelegate;
+
+	FVector AxisValue;
+
+	FInputVectorAxisBinding()
+		: FInputBinding()
+	{}
+
+	FInputVectorAxisBinding(const FKey InAxisKey)
+		: FInputBinding()
+		, AxisKey(InAxisKey)
+	{
+		ensure(AxisKey.IsVectorAxis());
+	}
+};
+
 /** 
  * Delegate signature for gesture handlers. 
@@ -351,31 +376,24 @@
 };
 
 /** Binds a gesture to a function. */
-struct FInputGestureBinding
+struct FInputGestureBinding : public FInputBinding
 {
 	FKey GestureKey;
 
 	FInputGestureUnifiedDelegate GestureDelegate;
 
-	uint32 bConsumeInput : 1;
-
-	uint32 bExecuteWhenPaused : 1;
-
 	/** Value parameter, meaning is dependent on the gesture. */
 	float GestureValue;
 
 	FInputGestureBinding()
-		: bConsumeInput(true)
-		, bExecuteWhenPaused(false)
+		: FInputBinding()
 		, GestureValue(0.f)
 	{}
 
 	FInputGestureBinding(const FKey InGestureKey)
-		: GestureKey(InGestureKey)
-		, bConsumeInput(true)
-		, bExecuteWhenPaused(false)
+		: FInputBinding()
+		, GestureKey(InGestureKey)
 		, GestureValue(0.f)
-
 	{}
 };
 
@@ -424,6 +442,7 @@
 	TArray<struct FInputTouchBinding> TouchBindings;
 	TArray<struct FInputAxisBinding> AxisBindings;
 	TArray<struct FInputAxisKeyBinding> AxisKeyBindings;
+	TArray<struct FInputVectorAxisBinding> VectorAxisBindings;
 	TArray<struct FInputGestureBinding> GestureBindings;
 
 	// Whether any components lower on the input stack should be allowed to receive input
@@ -431,6 +450,7 @@
 
 	float GetAxisValue(const FName AxisName) const;
 	float GetAxisKeyValue(const FKey AxisKey) const;
+	FVector GetVectorAxisValue(const FKey AxisKey) const;
 
 	bool HasBindings() const;
 
@@ -461,5 +481,14 @@
 	}
 
 	template<class UserClass>
+	FInputVectorAxisBinding& BindVectorAxis(const FKey AxisKey, UserClass* Object, typename FInputVectorAxisHandlerSignature::TUObjectMethodDelegate< UserClass >::FMethodPtr Func)
+	{
+		FInputVectorAxisBinding AB(AxisKey);
+		AB.AxisDelegate.BindDelegate(Object, Func);
+		VectorAxisBindings.Add(AB);
+		return VectorAxisBindings.Last();
+	}
+
+	template<class UserClass>
 	FInputKeyBinding& BindKey(const FInputChord Chord, const EInputEvent KeyEvent, UserClass* Object, typename FInputActionHandlerSignature::TUObjectMethodDelegate< UserClass >::FMethodPtr Func)
 	{
diff --git a/Engine/Source/Runtime/Engine/Classes/Engine/InputActionDelegateBinding.h b/Engine/Source/Runtime/Engine/Classes/Engine/InputActionDelegateBinding.h
--- a/Engine/Source/Runtime/Engine/Classes/Engine/InputActionDelegateBinding.h
+++ b/Engine/Source/Runtime/Engine/Classes/Engine/InputActionDelegateBinding.h
@@ -4,7 +4,7 @@
 #include "InputActionDelegateBinding.generated.h"
 
 USTRUCT()
-struct ENGINE_API FBlueprintInputActionDelegateBinding
+struct ENGINE_API FBlueprintInputActionDelegateBinding : public FBlueprintInputDelegateBinding
 {
 	GENERATED_USTRUCT_BODY()
 
@@ -14,24 +14,13 @@ struct ENGINE_API FBlueprintInputActionDelegateBinding
 	UPROPERTY()
 	TEnumAsByte<EInputEvent> InputKeyEvent;
 
-	UPROPERTY()
-	uint32 bConsumeInput:1;
-
-	UPROPERTY()
-	uint32 bExecuteWhenPaused:1;
-
-	UPROPERTY()
-	uint32 bOverrideParentBinding:1;
-
 	UPROPERTY()
 	FName FunctionNameToBind;
 
 	FBlueprintInputActionDelegateBinding()
-		: InputActionName(NAME_None)
+		: FBlueprintInputDelegateBinding()
+		, InputActionName(NAME_None)
 		, InputKeyEvent(IE_Pressed)
-		, bConsumeInput(true)
-		, bExecuteWhenPaused(false)
-		, bOverrideParentBinding(true)
 		, FunctionNameToBind(NAME_None)
 	{
 	}
diff --git a/Engine/Source/Runtime/Engine/Classes/Engine/InputAxisDelegateBinding.h b/Engine/Source/Runtime/Engine/Classes/Engine/InputAxisDelegateBinding.h
--- a/Engine/Source/Runtime/Engine/Classes/Engine/InputAxisDelegateBinding.h
+++ b/Engine/Source/Runtime/Engine/Classes/Engine/InputAxisDelegateBinding.h
@@ -4,7 +4,7 @@
 #include "InputAxisDelegateBinding.generated.h"
 
 USTRUCT()
-struct ENGINE_API FBlueprintInputAxisDelegateBinding
+struct ENGINE_API FBlueprintInputAxisDelegateBinding : public FBlueprintInputDelegateBinding
 {
 	GENERATED_USTRUCT_BODY()
 
@@ -14,21 +14,10 @@ struct ENGINE_API FBlueprintInputAxisDelegateBinding
 	UPROPERTY()
 	FName FunctionNameToBind;
 
-	UPROPERTY()
-	uint32 bConsumeInput:1;
-
-	UPROPERTY()
-	uint32 bExecuteWhenPaused:1;
-
-	UPROPERTY()
-	uint32 bOverrideParentBinding:1;
-
 	FBlueprintInputAxisDelegateBinding()
-		: InputAxisName(NAME_None)
+		: FBlueprintInputDelegateBinding()
+		, InputAxisName(NAME_None)
 		, FunctionNameToBind(NAME_None)
-		, bConsumeInput(true)
-		, bExecuteWhenPaused(false)
-		, bOverrideParentBinding(true)
 	{
 	}
 };
diff --git a/Engine/Source/Runtime/Engine/Classes/Engine/InputAxisKeyDelegateBinding.h b/Engine/Source/Runtime/Engine/Classes/Engine/InputAxisKeyDelegateBinding.h
--- a/Engine/Source/Runtime/Engine/Classes/Engine/InputAxisKeyDelegateBinding.h
+++ b/Engine/Source/Runtime/Engine/Classes/Engine/InputAxisKeyDelegateBinding.h
@@ -4,7 +4,7 @@
 #include "InputAxisKeyDelegateBinding.generated.h"
 
 USTRUCT()
-struct ENGINE_API FBlueprintInputAxisKeyDelegateBinding
+struct ENGINE_API FBlueprintInputAxisKeyDelegateBinding : public FBlueprintInputDelegateBinding
 {
 	GENERATED_USTRUCT_BODY()
 
@@ -14,20 +14,9 @@ struct ENGINE_API FBlueprintInputAxisKeyDelegateBinding
 	UPROPERTY()
 	FName FunctionNameToBind;
 
-	UPROPERTY()
-	uint32 bConsumeInput:1;
-
-	UPROPERTY()
-	uint32 bExecuteWhenPaused:1;
-
-	UPROPERTY()
-	uint32 bOverrideParentBinding:1;
-
 	FBlueprintInputAxisKeyDelegateBinding()
-		: FunctionNameToBind(NAME_None)
-		, bConsumeInput(true)
-		, bExecuteWhenPaused(false)
-		, bOverrideParentBinding(true)
+		: FBlueprintInputDelegateBinding()
+		, FunctionNameToBind(NAME_None)
 	{
 	}
 };
diff --git a/Engine/Source/Runtime/Engine/Classes/Engine/InputDelegateBinding.h b/Engine/Source/Runtime/Engine/Classes/Engine/InputDelegateBinding.h
--- a/Engine/Source/Runtime/Engine/Classes/Engine/InputDelegateBinding.h
+++ b/Engine/Source/Runtime/Engine/Classes/Engine/InputDelegateBinding.h
@@ -3,6 +3,28 @@
 #pragma once
 #include "InputDelegateBinding.generated.h"
 
+USTRUCT()
+struct ENGINE_API FBlueprintInputDelegateBinding
+{
+	GENERATED_USTRUCT_BODY()
+
+	UPROPERTY()
+	uint32 bConsumeInput:1;
+
+	UPROPERTY()
+	uint32 bExecuteWhenPaused:1;
+
+	UPROPERTY()
+	uint32 bOverrideParentBinding:1;
+
+	FBlueprintInputDelegateBinding()
+		: bConsumeInput(true)
+		, bExecuteWhenPaused(false)
+		, bOverrideParentBinding(true)
+	{
+	}
+};
+
 UCLASS(abstract)
 class ENGINE_API UInputDelegateBinding : public UDynamicBlueprintBinding
 {
diff --git a/Engine/Source/Runtime/Engine/Classes/Engine/InputKeyDelegateBinding.h b/Engine/Source/Runtime/Engine/Classes/Engine/InputKeyDelegateBinding.h
--- a/Engine/Source/Runtime/Engine/Classes/Engine/InputKeyDelegateBinding.h
+++ b/Engine/Source/Runtime/Engine/Classes/Engine/InputKeyDelegateBinding.h
@@ -4,7 +4,7 @@
 #include "InputKeyDelegateBinding.generated.h"
 
 USTRUCT()
-struct ENGINE_API FBlueprintInputKeyDelegateBinding
+struct ENGINE_API FBlueprintInputKeyDelegateBinding : public FBlueprintInputDelegateBinding
 {
 	GENERATED_USTRUCT_BODY()
 
@@ -14,23 +14,12 @@ struct ENGINE_API FBlueprintInputKeyDelegateBinding
 	UPROPERTY()
 	TEnumAsByte<EInputEvent> InputKeyEvent;
 
-	UPROPERTY()
-	uint32 bConsumeInput:1;
-
-	UPROPERTY()
-	uint32 bExecuteWhenPaused:1;
-
-	UPROPERTY()
-	uint32 bOverrideParentBinding:1;
-
 	UPROPERTY()
 	FName FunctionNameToBind;
 
 	FBlueprintInputKeyDelegateBinding()
-		: InputKeyEvent(IE_Pressed)
-		, bConsumeInput(true)
-		, bExecuteWhenPaused(false)
-		, bOverrideParentBinding(true)
+		: FBlueprintInputDelegateBinding()
+		, InputKeyEvent(IE_Pressed)
 		, FunctionNameToBind(NAME_None)
 	{
 	}
diff --git a/Engine/Source/Runtime/Engine/Classes/Engine/InputTouchDelegateBinding.h b/Engine/Source/Runtime/Engine/Classes/Engine/InputTouchDelegateBinding.h
--- a/Engine/Source/Runtime/Engine/Classes/Engine/InputTouchDelegateBinding.h
+++ b/Engine/Source/Runtime/Engine/Classes/Engine/InputTouchDelegateBinding.h
@@ -4,30 +4,19 @@
 #include "InputTouchDelegateBinding.generated.h"
 
 USTRUCT()
-struct ENGINE_API FBlueprintInputTouchDelegateBinding
+struct ENGINE_API FBlueprintInputTouchDelegateBinding : public FBlueprintInputDelegateBinding
 {
 	GENERATED_USTRUCT_BODY()
 
 	UPROPERTY()
 	TEnumAsByte<EInputEvent> InputKeyEvent;
 
-	UPROPERTY()
-	uint32 bConsumeInput:1;
-
-	UPROPERTY()
-	uint32 bExecuteWhenPaused:1;
-
-	UPROPERTY()
-	uint32 bOverrideParentBinding:1;
-
 	UPROPERTY()
 	FName FunctionNameToBind;
 
 	FBlueprintInputTouchDelegateBinding()
-		: InputKeyEvent(IE_Pressed)
-		, bConsumeInput(true)
-		, bExecuteWhenPaused(false)
-		, bOverrideParentBinding(true)
+		: FBlueprintInputDelegateBinding()
+		, InputKeyEvent(IE_Pressed)
 		, FunctionNameToBind(NAME_None)
 	{
 	}
diff --git a/Engine/Source/Runtime/Engine/Classes/Engine/InputVectorAxisDelegateBinding.h b/Engine/Source/Runtime/Engine/Classes/Engine/InputVectorAxisDelegateBinding.h
new file mode 100644
--- /dev/null
+++ b/Engine/Source/Runtime/Engine/Classes/Engine/InputVectorAxisDelegateBinding.h
@@ -0,0 +1,14 @@
+// Copyright 1998-2014 Epic Games, Inc. All Rights Reserved.
+
+#pragma once
+#include "InputVectorAxisDelegateBinding.generated.h"
+
+UCLASS()
+class ENGINE_API UInputVectorAxisDelegateBinding : public UInputAxisKeyDelegateBinding
+{
+	GENERATED_UCLASS_BODY()
+
+	// Begin UInputDelegateBinding interface
+	virtual void BindToInputComponent(UInputComponent* InputComponent) const OVERRIDE;
+	// End UInputDelegateBinding interface
+};
diff --git a/Engine/Source/Runtime/Engine/Classes/GameFramework/Actor.h b/Engine/Source/Runtime/Engine/Classes/GameFramework/Actor.h
--- a/Engine/Source/Runtime/Engine/Classes/GameFramework/Actor.h
+++ b/Engine/Source/Runtime/Engine/Classes/GameFramework/Actor.h
@@ -575,6 +575,9 @@ public:
 	UFUNCTION(BlueprintCallable, meta=(BlueprintInternalUseOnly="true", BlueprintProtected = "true", HidePin="InputAxisKey", ToolTip="Gets the value of the input axis key if input is enabled for this actor"))
 	float GetInputAxisKeyValue(const FKey InputAxisKey) const;
 
+	UFUNCTION(BlueprintCallable, meta=(BlueprintInternalUseOnly="true", BlueprintProtected = "true", HidePin="InputAxisKey", ToolTip="Gets the value of the input axis key if input is enabled for this actor"))
+	FVector GetInputVectorAxisValue(const FKey InputAxisKey) const;
+
 	/** Returns the instigator for this actor, or NULL if there is none. */
 	UFUNCTION(BlueprintCallable, meta=(BlueprintProtected = "true"), Category="Game|Damage")
 	APawn* GetInstigator() const;
diff --git a/Engine/Source/Runtime/Engine/Private/Actor.cpp b/Engine/Source/Runtime/Engine/Private/Actor.cpp
--- a/Engine/Source/Runtime/Engine/Private/Actor.cpp
+++ b/Engine/Source/Runtime/Engine/Private/Actor.cpp
@@ -2348,6 +2348,17 @@ float AActor::GetInputAxisKeyValue(const FKey InputAxisKey) const
 	return Value;
 }
 
+FVector AActor::GetInputVectorAxisValue(const FKey InputAxisKey) const
+{
+	FVector Value;
+
+	if (InputComponent)
+	{
+		Value = InputComponent->GetVectorAxisValue(InputAxisKey);
+	}
+
+	return Value;
+}
 
 bool AActor::SetActorLocation(const FVector& NewLocation, bool bSweep)
 {
diff --git a/Engine/Source/Runtime/Engine/Private/InputComponent.cpp b/Engine/Source/Runtime/Engine/Private/InputComponent.cpp
--- a/Engine/Source/Runtime/Engine/Private/InputComponent.cpp
+++ b/Engine/Source/Runtime/Engine/Private/InputComponent.cpp
@@ -39,35 +39,60 @@ UInputComponent::UInputComponent(const class FPostConstructInitializeProperties&
 
 float UInputComponent::GetAxisValue(const FName AxisName) const
 {
-	for (int32 Index = 0; Index < AxisBindings.Num(); ++Index)
+	float AxisValue = 0.f;
+
+	for (const FInputAxisBinding& AxisBinding : AxisBindings)
 	{
-		const FInputAxisBinding& AxisBinding = AxisBindings[Index];
 		if (AxisBinding.AxisName == AxisName)
 		{
-			return AxisBinding.AxisValue;
+			AxisValue = AxisBinding.AxisValue;
+			break;
 		}
 	}
 
-	return 0.f;
+	return AxisValue;
 }
 
 float UInputComponent::GetAxisKeyValue(const FKey AxisKey) const
 {
-	for (int32 Index = 0; Index < AxisKeyBindings.Num(); ++Index)
+	float AxisValue = 0.f;
+
+	for (const FInputAxisKeyBinding& AxisBinding : AxisKeyBindings)
+	{
+		if (AxisBinding.AxisKey == AxisKey)
+		{
+			AxisValue = AxisBinding.AxisValue;
+			break;
+		}
+	}
+
+	return AxisValue;
+}
+
+FVector UInputComponent::GetVectorAxisValue(const FKey AxisKey) const
+{
+	FVector AxisValue;
+	for (const FInputVectorAxisBinding& AxisBinding : VectorAxisBindings)
 	{
-		const FInputAxisKeyBinding& AxisBinding = AxisKeyBindings[Index];
 		if (AxisBinding.AxisKey == AxisKey)
 		{
-			return AxisBinding.AxisValue;
+			AxisValue = AxisBinding.AxisValue;
+			break;
 		}
 	}
 
-	return 0.f;
+	return AxisValue;
 }
 
 bool UInputComponent::HasBindings() const
 {
-	return ((ActionBindings.Num() > 0) || (AxisBindings.Num() > 0) || (AxisKeyBindings.Num() > 0) || (KeyBindings.Num() > 0) || (TouchBindings.Num() > 0));
+	return (   (ActionBindings.Num() > 0) 
+			|| (AxisBindings.Num() > 0) 
+			|| (AxisKeyBindings.Num() > 0) 
+			|| (KeyBindings.Num() > 0) 
+			|| (TouchBindings.Num() > 0) 
+			|| (GestureBindings.Num() > 0)
+			|| (VectorAxisBindings.Num() > 0));
 }
 
 FInputActionBinding& UInputComponent::AddActionBinding(const FInputActionBinding& Binding)
diff --git a/Engine/Source/Runtime/Engine/Private/InputDelegateBinding.cpp b/Engine/Source/Runtime/Engine/Private/InputDelegateBinding.cpp
--- a/Engine/Source/Runtime/Engine/Private/InputDelegateBinding.cpp
+++ b/Engine/Source/Runtime/Engine/Private/InputDelegateBinding.cpp
@@ -17,6 +17,7 @@ void UInputDelegateBinding::BindInputDelegates(UBlueprintGeneratedClass* BGClass
 												UInputKeyDelegateBinding::StaticClass(),
 												UInputTouchDelegateBinding::StaticClass(),
 												UInputAxisKeyDelegateBinding::StaticClass(),
+												UInputVectorAxisDelegateBinding::StaticClass(),
 										   };
 
 	if (BGClass)
diff --git a/Engine/Source/Runtime/Engine/Private/InputVectorAxisDelegateBinding.cpp b/Engine/Source/Runtime/Engine/Private/InputVectorAxisDelegateBinding.cpp
new file mode 100644
--- /dev/null
+++ b/Engine/Source/Runtime/Engine/Private/InputVectorAxisDelegateBinding.cpp
@@ -0,0 +1,45 @@
+// Copyright 1998-2014 Epic Games, Inc. All Rights Reserved.
+
+#include "EnginePrivate.h"
+#include "BlueprintUtilities.h"
+
+
+UInputVectorAxisDelegateBinding::UInputVectorAxisDelegateBinding(const class FPostConstructInitializeProperties& PCIP)
+	: Super(PCIP)
+{
+}
+
+void UInputVectorAxisDelegateBinding::BindToInputComponent(UInputComponent* InputComponent) const
+{
+	TArray<FInputVectorAxisBinding> BindsToAdd;
+
+	for (int32 BindIndex=0; BindIndex<InputAxisKeyDelegateBindings.Num(); ++BindIndex)
+	{
+		const FBlueprintInputAxisKeyDelegateBinding& Binding = InputAxisKeyDelegateBindings[BindIndex];
+
+		FInputVectorAxisBinding VAB( Binding.AxisKey );
+		VAB.bConsumeInput = Binding.bConsumeInput;
+		VAB.bExecuteWhenPaused = Binding.bExecuteWhenPaused;
+		VAB.AxisDelegate.BindDelegate(InputComponent->GetOwner(), Binding.FunctionNameToBind);
+
+		if (Binding.bOverrideParentBinding)
+		{
+			for (int32 ExistingIndex = InputComponent->VectorAxisBindings.Num() - 1; ExistingIndex >= 0; --ExistingIndex)
+			{
+				const FInputVectorAxisBinding& ExistingBind = InputComponent->VectorAxisBindings[ExistingIndex];
+				if (ExistingBind.AxisKey == VAB.AxisKey)
+				{
+					InputComponent->AxisKeyBindings.RemoveAt(ExistingIndex);
+				}
+			}
+		}
+
+		// To avoid binds in the same layer being removed by the parent override temporarily put them in this array and add later
+		BindsToAdd.Add(VAB);
+	}
+
+	for (int32 Index=0; Index < BindsToAdd.Num(); ++Index)
+	{
+		InputComponent->VectorAxisBindings.Add(BindsToAdd[Index]);
+	}
+}
\ No newline at end of file
diff --git a/Engine/Source/Runtime/Engine/Private/PlayerController.cpp b/Engine/Source/Runtime/Engine/Private/PlayerController.cpp
--- a/Engine/Source/Runtime/Engine/Private/PlayerController.cpp
+++ b/Engine/Source/Runtime/Engine/Private/PlayerController.cpp
@@ -3741,6 +3741,10 @@ bool APlayerController::PopInputComponent(UInputComponent* InputComponent)
 			{
 				AxisKeyBinding.AxisValue = 0.f;
 			}
+			for (FInputVectorAxisBinding& VectorAxisBinding : InputComponent->VectorAxisBindings)
+			{
+				VectorAxisBinding.AxisValue = FVector::ZeroVector;
+			}
 
 			return true;
 		}
diff --git a/Engine/Source/Runtime/Engine/Private/UserInterface/PlayerInput.cpp b/Engine/Source/Runtime/Engine/Private/UserInterface/PlayerInput.cpp
--- a/Engine/Source/Runtime/Engine/Private/UserInterface/PlayerInput.cpp
+++ b/Engine/Source/Runtime/Engine/Private/UserInterface/PlayerInput.cpp
@@ -784,7 +784,19 @@ void UPlayerInput::ProcessInputStack(const TArray<UInputComponent*>& InputCompon
 		{
 		}
 	};
+	struct FVectorAxisDelegateDetails
+	{
+		FInputVectorAxisUnifiedDelegate Delegate;
+		FVector Value;
+
+		FVectorAxisDelegateDetails(const FInputVectorAxisUnifiedDelegate& InDelegate, const FVector InValue)
+			: Delegate(InDelegate)
+			, Value(InValue)
+		{
+		}
+	};
 	TArray<FAxisDelegateDetails> AxisDelegates;
+	TArray<FVectorAxisDelegateDetails> VectorAxisDelegates;
 	TArray<FDelegateDispatchDetails> NonAxisDelegates;
 	struct FDelegateDispatchDetailsSorter
 	{
@@ -903,20 +915,16 @@ void UPlayerInput::ProcessInputStack(const TArray<UInputComponent*>& InputCompon
 			}
 
 			// Run though game axis bindings and accumulate axis values
-			for (int32 AxisIndex=0; AxisIndex<IC->AxisBindings.Num(); ++AxisIndex)
+			for (FInputAxisBinding& AB : IC->AxisBindings)
 			{
-				FInputAxisBinding& AB = IC->AxisBindings[AxisIndex];
-
 				AB.AxisValue = DetermineAxisValue(AB, bGamePaused, KeysToConsume);
 				if (AB.AxisDelegate.IsBound())
 				{
 					AxisDelegates.Add(FAxisDelegateDetails(AB.AxisDelegate, AB.AxisValue));
 				}
 			}
-			for (int32 AxisIndex = 0; AxisIndex < IC->AxisKeyBindings.Num(); ++AxisIndex)
+			for (FInputAxisKeyBinding& AxisKeyBinding : IC->AxisKeyBindings)
 			{
-				FInputAxisKeyBinding& AxisKeyBinding = IC->AxisKeyBindings[AxisIndex];
-
 				if (!IsKeyConsumed(AxisKeyBinding.AxisKey))
 				{
 					if (!bGamePaused || AxisKeyBinding.bExecuteWhenPaused)
@@ -939,6 +947,30 @@ void UPlayerInput::ProcessInputStack(const TArray<UInputComponent*>& InputCompon
 					AxisDelegates.Add(FAxisDelegateDetails(AxisKeyBinding.AxisDelegate, AxisKeyBinding.AxisValue));
 				}
 			}
+			for (FInputVectorAxisBinding& VectorAxisBinding : IC->VectorAxisBindings)
+			{
+				if (!IsKeyConsumed(VectorAxisBinding.AxisKey))
+				{
+					if (!bGamePaused || VectorAxisBinding.bExecuteWhenPaused)
+					{
+						VectorAxisBinding.AxisValue = GetVectorKeyValue(VectorAxisBinding.AxisKey);
+					}
+					else
+					{
+						VectorAxisBinding.AxisValue = FVector::ZeroVector;
+					}
+
+					if (VectorAxisBinding.bConsumeInput)
+					{
+						KeysToConsume.AddUnique(VectorAxisBinding.AxisKey);
+					}
+				}
+
+				if (VectorAxisBinding.AxisDelegate.IsBound())
+				{
+					VectorAxisDelegates.Add(FVectorAxisDelegateDetails(VectorAxisBinding.AxisDelegate, VectorAxisBinding.AxisValue));
+				}
+			}
 
 			if (IC->bBlockInput)
 			{
@@ -970,6 +1002,10 @@ void UPlayerInput::ProcessInputStack(const TArray<UInputComponent*>& InputCompon
 			{
 				AxisKeyBinding.AxisValue = 0.f;
 			}
+			for (FInputVectorAxisBinding& VectorAxisBinding : IC->VectorAxisBindings)
+			{
+				VectorAxisBinding.AxisValue = FVector::ZeroVector;
+			}
 		}
 	}
 
@@ -998,6 +1034,13 @@ void UPlayerInput::ProcessInputStack(const TArray<UInputComponent*>& InputCompon
 			Details.Delegate.Execute(Details.Value);
 		}
 	}
+	for (const FVectorAxisDelegateDetails& Details : VectorAxisDelegates)
+	{
+		if (Details.Delegate.IsBound())
+		{
+			Details.Delegate.Execute(Details.Value);
+		}
+	}
 
 	PlayerController->PostProcessInput(DeltaTime, bGamePaused);
 
@@ -1115,7 +1158,7 @@ void UPlayerInput::DisplayDebug(class UCanvas* Canvas, const FDebugDisplayInfo&
 
 			if ( KeyState->bDown || (KeyState->Value.X != 0.f) )
 			{
-				if (!Key.IsAxis())
+				if (!Key.IsFloatAxis())
 				{
 					Str += FString::Printf(TEXT(" time: %.2f"), World->TimeSince(KeyState->LastUpDownTransitionTime));
 				}
diff --git a/Engine/Source/Runtime/InputCore/Classes/InputCoreTypes.h b/Engine/Source/Runtime/InputCore/Classes/InputCoreTypes.h
--- a/Engine/Source/Runtime/InputCore/Classes/InputCoreTypes.h
+++ b/Engine/Source/Runtime/InputCore/Classes/InputCoreTypes.h
@@ -32,7 +32,8 @@ struct INPUTCORE_API FKey
 	bool IsModifierKey() const;
 	bool IsGamepadKey() const;
 	bool IsMouseButton() const;
-	bool IsAxis() const;
+	bool IsFloatAxis() const;
+	bool IsVectorAxis() const;
 	bool IsBindableInBlueprints() const;
 	FText GetDisplayName() const;
 	FString ToString() const;
@@ -79,10 +80,11 @@ struct INPUTCORE_API FKeyDetails
 	enum EKeyFlags
 	{
 		GamepadKey				= 0x01,
-		ModifierKey				= 0x02,
-		NotBlueprintBindableKey	= 0x04,
-		Axis					= 0x08,
-		MouseButton				= 0x10,
+		MouseButton				= 0x02,
+		ModifierKey				= 0x04,
+		NotBlueprintBindableKey	= 0x08,
+		FloatAxis				= 0x10,
+		VectorAxis				= 0x20,
 
 		NoFlags                 = 0,
 	};
@@ -93,9 +95,18 @@ struct INPUTCORE_API FKeyDetails
 		, bIsModifierKey((InKeyFlags & EKeyFlags::ModifierKey) != 0)
 		, bIsGamepadKey((InKeyFlags & EKeyFlags::GamepadKey) != 0)
 		, bIsMouseButton((InKeyFlags & EKeyFlags::MouseButton) != 0)
-		, bIsAxis((InKeyFlags & EKeyFlags::Axis) != 0)
 		, bIsBindableInBlueprints((~InKeyFlags & EKeyFlags::NotBlueprintBindableKey) != 0)
+		, AxisType(EInputAxisType::None)
 	{
+		if ((InKeyFlags & EKeyFlags::FloatAxis) != 0)
+		{
+			ensure((InKeyFlags & EKeyFlags::VectorAxis) == 0);
+			AxisType = EInputAxisType::Float;
+		}
+		else if ((InKeyFlags & EKeyFlags::VectorAxis) != 0)
+		{
+			AxisType = EInputAxisType::Vector;
+		}
 	}
 
 	FKeyDetails(const FKey InKey, const FGetKeyDisplayNameSignature InGetDisplayNameDelegate, const uint8 InKeyFlags = 0)
@@ -104,21 +115,29 @@ struct INPUTCORE_API FKeyDetails
 		, bIsModifierKey((InKeyFlags & EKeyFlags::ModifierKey) != 0)
 		, bIsGamepadKey((InKeyFlags & EKeyFlags::GamepadKey) != 0)
 		, bIsMouseButton((InKeyFlags & EKeyFlags::MouseButton) != 0)
-		, bIsAxis((InKeyFlags & EKeyFlags::Axis) != 0)
 		, bIsBindableInBlueprints((~InKeyFlags & EKeyFlags::NotBlueprintBindableKey) != 0)
+		, AxisType(EInputAxisType::None)
 	{
 	}
 
 	bool IsModifierKey() const { return bIsModifierKey != 0; }
 	bool IsGamepadKey() const { return bIsGamepadKey != 0; }
 	bool IsMouseButton() const { return bIsMouseButton != 0; }
-	bool IsAxis() const { return bIsAxis != 0; }
+	bool IsFloatAxis() const { return AxisType == EInputAxisType::Float; }
+	bool IsVectorAxis() const { return AxisType == EInputAxisType::Vector; }
 	bool IsBindableInBlueprints() const { return bIsBindableInBlueprints != 0; }
 	FText GetDisplayName() const;
 	const FKey& GetKey() const { return Key; }
 
 private:
 
+	enum class EInputAxisType : uint8
+	{
+		None,
+		Float,
+		Vector
+	};
+
 	FKey  Key;
 	
 	FText DisplayName;
@@ -127,8 +146,9 @@ private:
 	int32 bIsModifierKey:1;
 	int32 bIsGamepadKey:1;
 	int32 bIsMouseButton:1;
-	int32 bIsAxis:1;
 	int32 bIsBindableInBlueprints:1;
+	EInputAxisType AxisType;
+
 };
 
 UENUM(BlueprintType)
@@ -378,7 +398,7 @@ struct INPUTCORE_API EKeys
 	// These exist for backwards compatibility reasons only
 	static bool IsModifierKey(FKey Key) { return Key.IsModifierKey(); }
 	static bool IsGamepadKey(FKey Key) { return Key.IsGamepadKey(); }
-	static bool IsAxis(FKey Key) { return Key.IsAxis(); }
+	static bool IsAxis(FKey Key) { return Key.IsFloatAxis(); }
 	static bool IsBindableInBlueprints(const FKey Key) { return Key.IsBindableInBlueprints(); }
 	static void SetConsoleForGamepadLabels(const EConsoleForGamepadLabels::Type Console) { ConsoleForGamepadLabels = Console; }
 
diff --git a/Engine/Source/Runtime/InputCore/Private/InputCoreTypes.cpp b/Engine/Source/Runtime/InputCore/Private/InputCoreTypes.cpp
--- a/Engine/Source/Runtime/InputCore/Private/InputCoreTypes.cpp
+++ b/Engine/Source/Runtime/InputCore/Private/InputCoreTypes.cpp
@@ -229,8 +229,8 @@ void EKeys::Initialize()
 	if (bInitialized) return;
 	bInitialized = true;
 
-	AddKey(FKeyDetails(EKeys::MouseX, LOCTEXT("MouseX", "Mouse X"), FKeyDetails::Axis | FKeyDetails::MouseButton));
-	AddKey(FKeyDetails(EKeys::MouseY, LOCTEXT("MouseY", "Mouse Y"), FKeyDetails::Axis | FKeyDetails::MouseButton));
+	AddKey(FKeyDetails(EKeys::MouseX, LOCTEXT("MouseX", "Mouse X"), FKeyDetails::FloatAxis | FKeyDetails::MouseButton));
+	AddKey(FKeyDetails(EKeys::MouseY, LOCTEXT("MouseY", "Mouse Y"), FKeyDetails::FloatAxis | FKeyDetails::MouseButton));
 	AddKey(FKeyDetails(EKeys::MouseScrollUp, LOCTEXT("MouseScrollUp", "Mouse Wheel Up"), FKeyDetails::MouseButton));
 	AddKey(FKeyDetails(EKeys::MouseScrollDown, LOCTEXT("MouseScrollDown", "Mouse Wheel Down"), FKeyDetails::MouseButton));
 
@@ -242,7 +242,7 @@ void EKeys::Initialize()
 	AddKey(FKeyDetails(EKeys::MiddleMouseButton, LOCTEXT("MiddleMouseButton", "Middle Mouse Button"), FKeyDetails::MouseButton));
 	AddKey(FKeyDetails(EKeys::ThumbMouseButton, LOCTEXT("ThumbMouseButton", "Thumb Mouse Button"), FKeyDetails::MouseButton));
 	AddKey(FKeyDetails(EKeys::ThumbMouseButton2, LOCTEXT("ThumbMouseButton2", "Thumb Mouse Button 2"), FKeyDetails::MouseButton));
-	
+    
 	AddKey(FKeyDetails(EKeys::Tab, LOCTEXT("Tab", "Tab")));
 	AddKey(FKeyDetails(EKeys::Enter, LOCTEXT("Enter", "Enter")));
 	AddKey(FKeyDetails(EKeys::Pause, LOCTEXT("Pause", "Pause")));
@@ -261,13 +261,13 @@ void EKeys::Initialize()
 	AddKey(FKeyDetails(EKeys::Down, LOCTEXT("Down", "Down")));
 
 	AddKey(FKeyDetails(EKeys::Insert, LOCTEXT("Insert", "Insert")));
-	
+    
 #if PLATFORM_MAC
-	AddKey(FKeyDetails(EKeys::BackSpace, LOCTEXT("Delete", "Delete")));
-	AddKey(FKeyDetails(EKeys::Delete, LOCTEXT("ForwardDelete", "Fn+Delete")));
+    AddKey(FKeyDetails(EKeys::BackSpace, LOCTEXT("Delete", "Delete")));
+    AddKey(FKeyDetails(EKeys::Delete, LOCTEXT("ForwardDelete", "Fn+Delete")));
 #else
-	AddKey(FKeyDetails(EKeys::BackSpace, LOCTEXT("BackSpace", "Backspace")));
-	AddKey(FKeyDetails(EKeys::Delete, LOCTEXT("Delete", "Delete")));
+    AddKey(FKeyDetails(EKeys::BackSpace, LOCTEXT("BackSpace", "Backspace")));
+    AddKey(FKeyDetails(EKeys::Delete, LOCTEXT("Delete", "Delete")));
 #endif
 
 	AddKey(FKeyDetails(EKeys::Zero, FText::FromString("0")));
@@ -363,13 +363,13 @@ void EKeys::Initialize()
 	AddKey(FKeyDetails(EKeys::Quote, LOCTEXT("Quote", "'")));
 
 	// Setup Gamepad keys
-	AddKey(FKeyDetails(EKeys::Gamepad_LeftX, LOCTEXT("Gamepad_LeftX", "Gamepad Left X"), FKeyDetails::GamepadKey | FKeyDetails::Axis));
-	AddKey(FKeyDetails(EKeys::Gamepad_LeftY, LOCTEXT("Gamepad_LeftY", "Gamepad Left Y"), FKeyDetails::GamepadKey | FKeyDetails::Axis));
-	AddKey(FKeyDetails(EKeys::Gamepad_RightX, LOCTEXT("Gamepad_RightX", "Gamepad Right X"), FKeyDetails::GamepadKey | FKeyDetails::Axis));
-	AddKey(FKeyDetails(EKeys::Gamepad_RightY, LOCTEXT("Gamepad_RightY", "Gamepad Right Y"), FKeyDetails::GamepadKey | FKeyDetails::Axis));
+	AddKey(FKeyDetails(EKeys::Gamepad_LeftX, LOCTEXT("Gamepad_LeftX", "Gamepad Left X"), FKeyDetails::GamepadKey | FKeyDetails::FloatAxis));
+	AddKey(FKeyDetails(EKeys::Gamepad_LeftY, LOCTEXT("Gamepad_LeftY", "Gamepad Left Y"), FKeyDetails::GamepadKey | FKeyDetails::FloatAxis));
+	AddKey(FKeyDetails(EKeys::Gamepad_RightX, LOCTEXT("Gamepad_RightX", "Gamepad Right X"), FKeyDetails::GamepadKey | FKeyDetails::FloatAxis));
+	AddKey(FKeyDetails(EKeys::Gamepad_RightY, LOCTEXT("Gamepad_RightY", "Gamepad Right Y"), FKeyDetails::GamepadKey | FKeyDetails::FloatAxis));
 
-	AddKey(FKeyDetails(EKeys::Gamepad_LeftTriggerAxis, LOCTEXT("Gamepad_LeftTriggerAxis", "Gamepad Left Trigger Axis"), FKeyDetails::GamepadKey | FKeyDetails::Axis));
-	AddKey(FKeyDetails(EKeys::Gamepad_RightTriggerAxis, LOCTEXT("Gamepad_RightTriggerAxis", "Gamepad Right Trigger Axis"), FKeyDetails::GamepadKey | FKeyDetails::Axis));
+	AddKey(FKeyDetails(EKeys::Gamepad_LeftTriggerAxis, LOCTEXT("Gamepad_LeftTriggerAxis", "Gamepad Left Trigger Axis"), FKeyDetails::GamepadKey | FKeyDetails::FloatAxis));
+	AddKey(FKeyDetails(EKeys::Gamepad_RightTriggerAxis, LOCTEXT("Gamepad_RightTriggerAxis", "Gamepad Right Trigger Axis"), FKeyDetails::GamepadKey | FKeyDetails::FloatAxis));
 
 	AddKey(FKeyDetails(EKeys::Gamepad_LeftThumbstick, LOCTEXT("Gamepad_LeftThumbstick", "Gamepad Left Thumbstick"), FKeyDetails::GamepadKey));
 	AddKey(FKeyDetails(EKeys::Gamepad_RightThumbstick, LOCTEXT("Gamepad_RightThumbstick", "Gamepad Right Thumbstick"), FKeyDetails::GamepadKey));
@@ -392,7 +392,7 @@ void EKeys::Initialize()
 	AddKey(FKeyDetails(EKeys::Gamepad_RightStick_Up, LOCTEXT("Gamepad_RightStick_Up", "Gamepad Right Stick Up"), FKeyDetails::GamepadKey));
 	AddKey(FKeyDetails(EKeys::Gamepad_RightStick_Down, LOCTEXT("Gamepad_RightStick_Down", "Gamepad Right Stick Down"), FKeyDetails::GamepadKey));
 	AddKey(FKeyDetails(EKeys::Gamepad_RightStick_Right, LOCTEXT("Gamepad_RightStick_Right", "Gamepad Right Stick Right"), FKeyDetails::GamepadKey));
-	AddKey(FKeyDetails(EKeys::Gamepad_RightStick_Left, LOCTEXT("Gamepad_RightStick_Left", "Gamepad Right Stick Left"), FKeyDetails::GamepadKey));
+	AddKey(FKeyDetails(EKeys::Gamepad_RightStick_Left, LOCTEXT("Gamepad_RightStick_Left", "Gamepad Right StickLeft"), FKeyDetails::GamepadKey));
 
 	FGetKeyDisplayNameSignature GetKeyNameDelegate = FGetKeyDisplayNameSignature::CreateStatic(&EKeys::GetGamepadDisplayName);
 
@@ -404,10 +404,10 @@ void EKeys::Initialize()
 	AddKey(FKeyDetails(EKeys::Gamepad_FaceButton_Top, GetKeyNameDelegate, FKeyDetails::GamepadKey));
 
 	// Vector axes (FVector, not float)
-	AddKey(FKeyDetails(EKeys::Tilt, LOCTEXT("Tilt", "Tilt"), FKeyDetails::NotBlueprintBindableKey));
-	AddKey(FKeyDetails(EKeys::RotationRate, LOCTEXT("RotationRate", "Rotation Rate"), FKeyDetails::NotBlueprintBindableKey));
-	AddKey(FKeyDetails(EKeys::Gravity, LOCTEXT("Gravity", "Gravity"), FKeyDetails::NotBlueprintBindableKey));
-	AddKey(FKeyDetails(EKeys::Acceleration, LOCTEXT("Acceleration", "Acceleration"), FKeyDetails::NotBlueprintBindableKey));
+	AddKey(FKeyDetails(EKeys::Tilt, LOCTEXT("Tilt", "Tilt"), FKeyDetails::VectorAxis));
+	AddKey(FKeyDetails(EKeys::RotationRate, LOCTEXT("RotationRate", "Rotation Rate"), FKeyDetails::VectorAxis));
+	AddKey(FKeyDetails(EKeys::Gravity, LOCTEXT("Gravity", "Gravity"), FKeyDetails::VectorAxis));
+	AddKey(FKeyDetails(EKeys::Acceleration, LOCTEXT("Acceleration", "Acceleration"), FKeyDetails::VectorAxis));
 
 	// Fingers
 	AddKey(FKeyDetails(EKeys::TouchKeys[ETouchIndex::Touch1], LOCTEXT("Touch1", "Touch 1"), FKeyDetails::NotBlueprintBindableKey));
@@ -588,10 +588,16 @@ bool FKey::IsMouseButton() const
 	return (KeyDetails.IsValid() ? KeyDetails->IsMouseButton() : false);
 }
 
-bool FKey::IsAxis() const
+bool FKey::IsFloatAxis() const
 {
 	ConditionalLookupKeyDetails();
-	return (KeyDetails.IsValid() ? KeyDetails->IsAxis() : false);
+	return (KeyDetails.IsValid() ? KeyDetails->IsFloatAxis() : false);
+}
+
+bool FKey::IsVectorAxis() const
+{
+	ConditionalLookupKeyDetails();
+	return (KeyDetails.IsValid() ? KeyDetails->IsVectorAxis() : false);
 }
 
 bool FKey::IsBindableInBlueprints() const
