- Add options to SpringArmComponent for rotation lag and inheriting only certain rotation components

diff --git a/Engine/Source/Runtime/Engine/Classes/GameFramework/SpringArmComponent.h b/Engine/Source/Runtime/Engine/Classes/GameFramework/SpringArmComponent.h
--- a/Engine/Source/Runtime/Engine/Classes/GameFramework/SpringArmComponent.h
+++ b/Engine/Source/Runtime/Engine/Classes/GameFramework/SpringArmComponent.h
@@ -40,17 +40,39 @@ class ENGINE_API USpringArmComponent : public USceneComponent
 	UPROPERTY(EditAnywhere, BlueprintReadWrite, Category=CameraSettings)
 	uint32 bUseControllerViewRotation:1;
 
+	/** Should we inherit pitch from parent component. Does nothing if using Absolute Rotation. */
+	UPROPERTY(EditAnywhere, BlueprintReadWrite, Category=CameraSettings)
+	uint32 bInheritPitch : 1;
+
+	/** Should we inherit yaw from parent component. Does nothing if using Absolute Rotation. */
+	UPROPERTY(EditAnywhere, BlueprintReadWrite, Category=CameraSettings)
+	uint32 bInheritYaw : 1;
+
+	/** Should we inherit roll from parent component. Does nothing if using Absolute Rotation. */
+	UPROPERTY(EditAnywhere, BlueprintReadWrite, Category=CameraSettings)
+	uint32 bInheritRoll : 1;
+
 	/** If true, camera lags behind target position to smooth its movement */
 	UPROPERTY(EditAnywhere, BlueprintReadWrite, Category=Lag)
 	uint32 bEnableCameraLag : 1;
 
+	/** If true, camera lags behind target position to smooth its movement */
+	UPROPERTY(EditAnywhere, BlueprintReadWrite, Category=Lag)
+	uint32 bEnableCameraRotationLag : 1;
+
 	/** If bEnableCameraLag is true, controls how quickly camera reaches target position */
 	UPROPERTY(EditAnywhere, BlueprintReadWrite, Category=Lag, meta=(editcondition="bEnableCameraLag"))
 	float CameraLagSpeed;
 
-	/** Temporary variable when using camera log, to record previous camera position */
-	FVector PreviousDesiredLoc;
+	/** If bEnableCameraLag is true, controls how quickly camera reaches target position */
+	UPROPERTY(EditAnywhere, BlueprintReadWrite, Category=Lag, meta=(editcondition = "bEnableCameraRotationLag"))
+	float CameraRotationLagSpeed;
+
 
+	/** Temporary variable when using camera lag, to record previous camera position */
+	FVector PreviousDesiredLoc;
+	/** Temporary variable for lagging camera rotation, for previous rotation */
+	FRotator PreviousDesiredRot;
 
 	// UActorComponent interface
 	virtual void OnRegister() OVERRIDE;
@@ -68,10 +90,12 @@ class ENGINE_API USpringArmComponent : public USceneComponent
 protected:
 	/** Cached component-space socket location */
 	FVector RelativeSocketLocation;
+	/** Cached component-space socket rotation */
+	FQuat RelativeSocketRotation;
 
 protected:
 	/** Updates the desired arm location, calling BlendLocations to do the actual blending if a trace is done */
-	virtual void UpdateDesiredArmLocation(const FVector& Origin, const FRotator& Direction, bool bDoTrace, bool bAllowLag, float DeltaTime);
+	virtual void UpdateDesiredArmLocation(bool bDoTrace, bool bDoLocationLag, bool bDoRotationLag, float DeltaTime);
 
 	/**
 	 * This function allows subclasses to blend the trace hit location with the desired arm location;
diff --git a/Engine/Source/Runtime/Engine/Private/GameFramework/SpringArmComponent.cpp b/Engine/Source/Runtime/Engine/Private/GameFramework/SpringArmComponent.cpp
--- a/Engine/Source/Runtime/Engine/Private/GameFramework/SpringArmComponent.cpp
+++ b/Engine/Source/Runtime/Engine/Private/GameFramework/SpringArmComponent.cpp
@@ -7,6 +7,8 @@
 
 const FName USpringArmComponent::SocketName(TEXT("SpringEndpoint"));
 
+extern float GAverageMS;
+
 USpringArmComponent::USpringArmComponent(const class FPostConstructInitializeProperties& PCIP)
 	: Super(PCIP)
 {
@@ -17,21 +19,54 @@ USpringArmComponent::USpringArmComponent(const class FPostConstructInitializePro
 	bUseControllerViewRotation = false;
 	bDoCollisionTest = true;
 
+	bInheritPitch = true;
+	bInheritYaw = true;
+	bInheritRoll = true;
+
 	TargetArmLength = 300.0f;
 	ProbeSize = 12.0f;
 	ProbeChannel = ECC_Camera;
 
 	CameraLagSpeed = 10.f;
+
+	RelativeSocketRotation = FQuat::Identity;
 }
 
-void USpringArmComponent::UpdateDesiredArmLocation(const FVector& Origin, const FRotator& Direction, bool bDoTrace, bool bDoLag, float DeltaTime)
+void USpringArmComponent::UpdateDesiredArmLocation(bool bDoTrace, bool bDoLocationLag, bool bDoRotationLag, float DeltaTime)
 {
-	FRotator DesiredRot = Direction;
+	FRotator DesiredRot = GetComponentRotation();
+
+	// If inheriting rotation, check options for which components to inherit
+	if(!bAbsoluteRotation)
+	{
+		if(!bInheritPitch)
+		{
+			DesiredRot.Pitch = RelativeRotation.Pitch;
+		}
+
+		if (!bInheritYaw)
+		{
+			DesiredRot.Yaw = RelativeRotation.Yaw;
+		}
 
-	FVector DesiredLoc = Origin;
-	if (bDoLag)
+		if (!bInheritRoll)
+		{
+			DesiredRot.Roll = RelativeRotation.Roll;
+		}
+	}
+
+	if(bDoRotationLag)
+	{
+		DesiredRot = FMath::RInterpTo(PreviousDesiredRot, DesiredRot, GAverageMS/1000.f, CameraRotationLagSpeed);
+	}
+	PreviousDesiredRot = DesiredRot;
+
+
+	FVector ArmOrigin = GetComponentLocation();
+	FVector DesiredLoc = ArmOrigin;
+	if (bDoLocationLag)
 	{
-		DesiredLoc = FMath::VInterpTo(PreviousDesiredLoc, DesiredLoc, DeltaTime, CameraLagSpeed);
+		DesiredLoc = FMath::VInterpTo(PreviousDesiredLoc, DesiredLoc, GAverageMS/1000.f, CameraLagSpeed);
 	}
 	PreviousDesiredLoc = DesiredLoc;
 
@@ -45,7 +80,7 @@ void USpringArmComponent::UpdateDesiredArmLocation(const FVector& Origin, const
 		FCollisionQueryParams QueryParams(TraceTagName, false, GetOwner());
 
 		FHitResult Result;
-		GetWorld()->SweepSingle(Result, Origin, DesiredLoc, FQuat::Identity, ProbeChannel, FCollisionShape::MakeSphere(ProbeSize), QueryParams);
+		GetWorld()->SweepSingle(Result, ArmOrigin, DesiredLoc, FQuat::Identity, ProbeChannel, FCollisionShape::MakeSphere(ProbeSize), QueryParams);
 
 		ResultLoc = BlendLocations(DesiredLoc, Result.Location, Result.bBlockingHit, DeltaTime);
 	}
@@ -54,7 +89,15 @@ void USpringArmComponent::UpdateDesiredArmLocation(const FVector& Origin, const
 		ResultLoc = DesiredLoc;
 	}
 
-	RelativeSocketLocation = ComponentToWorld.ToInverseMatrixWithScale().TransformPosition(ResultLoc);
+	// Form a transform for new world transform for camera
+	FTransform WorldCamTM(DesiredRot, ResultLoc);
+	// Convert to relative to component
+	FTransform RelCamTM = WorldCamTM.GetRelativeTransform(ComponentToWorld);
+
+	// Update socket location/rotation
+	RelativeSocketLocation = RelCamTM.GetLocation();
+	RelativeSocketRotation = RelCamTM.GetRotation();
+
 	UpdateChildTransforms();
 }
 
@@ -66,7 +109,7 @@ FVector USpringArmComponent::BlendLocations(const FVector& DesiredArmLocation, c
 void USpringArmComponent::OnRegister()
 {
 	Super::OnRegister();
-	UpdateDesiredArmLocation(GetComponentLocation(), GetComponentRotation(), false, false, 0.0f);
+	UpdateDesiredArmLocation(false, false, false, 0.f);
 }
 
 void USpringArmComponent::TickComponent(float DeltaTime, enum ELevelTick TickType, FActorComponentTickFunction* ThisTickFunction)
@@ -85,12 +128,11 @@ void USpringArmComponent::TickComponent(float DeltaTime, enum ELevelTick TickTyp
 		}
 	}
 
-	UpdateDesiredArmLocation(GetComponentLocation(), GetComponentRotation(), bDoCollisionTest, bEnableCameraLag, DeltaTime);
+	UpdateDesiredArmLocation(bDoCollisionTest, bEnableCameraLag, bEnableCameraRotationLag, DeltaTime);
 }
 
 FTransform USpringArmComponent::GetSocketTransform(FName InSocketName, ERelativeTransformSpace TransformSpace) const
 {
-	const FQuat RelativeSocketRotation(FRotator(0.0f, 0.0f, 0.0f));
 	FTransform RelativeTransform(RelativeSocketRotation, RelativeSocketLocation);
 
 	switch(TransformSpace)
