Define cloth tick in regardless of WITH_APEX_CLOTHING. This is needed because we also do some physics blending here regardless of cloth

diff --git a/Engine/Source/Runtime/Engine/Classes/Components/SkeletalMeshComponent.h b/Engine/Source/Runtime/Engine/Classes/Components/SkeletalMeshComponent.h
--- a/Engine/Source/Runtime/Engine/Classes/Components/SkeletalMeshComponent.h
+++ b/Engine/Source/Runtime/Engine/Classes/Components/SkeletalMeshComponent.h
@@ -196,7 +196,6 @@ struct FSkeletalMeshComponentPreClothTickFunction : public FTickFunction
 {
 	GENERATED_USTRUCT_BODY()
 
-#if WITH_APEX_CLOTHING
 	/** World this tick function belongs to **/
 	class USkeletalMeshComponent*	Target;
 
@@ -210,8 +209,6 @@ struct FSkeletalMeshComponentPreClothTickFunction : public FTickFunction
 	virtual void ExecuteTick(float DeltaTime, enum ELevelTick TickType, ENamedThreads::Type CurrentThread, const FGraphEventRef& MyCompletionGraphEvent) OVERRIDE;
 	/** Abstract function to describe this tick. Used to print messages about illegal cycles in the dependency graph **/
 	virtual FString DiagnosticMessage();
-
-#endif
 };
 
 
@@ -556,9 +553,9 @@ public:
 
 #endif	//WITH_PHYSX
 
-#if WITH_APEX_CLOTHING
-
 	FSkeletalMeshComponentPreClothTickFunction PreClothTickFunction;
+
+#if WITH_APEX_CLOTHING
 	/** 
 	* clothing actors will be created from clothing assets for cloth simulation 
 	* 1 actor should correspond to 1 asset
@@ -920,6 +917,9 @@ public:
 	 */
 	FVector GetClosestCollidingRigidBodyLocation(const FVector& TestLocation) const;
 
+	/** Calls needed cloth updates */
+	void PreClothTick(float DeltaTime);
+
 #if WITH_APEX_CLOTHING
 	/** 
 	* APEX clothing actor is created from APEX clothing asset for cloth simulation 
@@ -934,8 +934,6 @@ public:
 	void SetClothingLOD(int32 LODIndex);
 	/** check whether clothing teleport is needed or not to avoid a weird simulation result */
 	void CheckClothTeleport(float DeltaTime);
-	/** Calls needed cloth updates */
-	void PreClothTick(float DeltaTime);
 	/** 
 	 * Updates all clothing animation states including ComponentToWorld-related states.
 	 */
diff --git a/Engine/Source/Runtime/Engine/Private/SkeletalMeshComponentPhysics.cpp b/Engine/Source/Runtime/Engine/Private/SkeletalMeshComponentPhysics.cpp
--- a/Engine/Source/Runtime/Engine/Private/SkeletalMeshComponentPhysics.cpp
+++ b/Engine/Source/Runtime/Engine/Private/SkeletalMeshComponentPhysics.cpp
@@ -2860,6 +2860,41 @@ void USkeletalMeshComponent::ProcessClothCollisionWithEnvironment()
 
 #endif// #if WITH_CLOTH_COLLISION_DETECTION
 
+void USkeletalMeshComponent::PreClothTick(float DeltaTime)
+{
+	// if physics is disabled on dedicated server, no reason to be here. 
+	if (!bEnablePhysicsOnDedicatedServer && IsRunningDedicatedServer())
+	{
+		return;
+	}
+
+	if (IsRegistered() && IsSimulatingPhysics())
+	{
+		SyncComponentToRBPhysics();
+	}
+
+	// this used to not run if not rendered, but that causes issues such as bounds not updated
+	// causing it to not rendered, at the end, I think we should blend body positions
+	// for example if you're only simulating, this has to happen all the time
+	// whether looking at it or not, otherwise
+	// @todo better solution is to check if it has moved by changing SyncComponentToRBPhysics to return true if anything modified
+	// and run this if that is true or rendered
+	// that will at least reduce the chance of mismatch
+	// generally if you move your actor position, this has to happen to approximately match their bounds
+	if (Bodies.Num() > 0 && IsRegistered())
+	{
+		BlendInPhysics();
+	}
+
+#if WITH_APEX_CLOTHING
+	// if skeletal mesh has clothing assets, call TickClothing
+	if (SkeletalMesh && SkeletalMesh->ClothingAssets.Num() > 0)
+	{
+		TickClothing(DeltaTime + SkippedTickDeltaTime);
+	}
+#endif
+}
+
 #if WITH_APEX_CLOTHING
 
 void USkeletalMeshComponent::UpdateClothTransform()
@@ -2950,39 +2985,6 @@ void USkeletalMeshComponent::CheckClothTeleport(float DeltaTime)
 	PrevRootBoneMatrix = CurRootBoneMat;
 }
 
-void USkeletalMeshComponent::PreClothTick(float DeltaTime)
-{
-	// if physics is disabled on dedicated server, no reason to be here. 
-	if (!bEnablePhysicsOnDedicatedServer && IsRunningDedicatedServer())
-	{
-		return;
-	}
-
-	if (IsRegistered() && IsSimulatingPhysics())
-	{
-		SyncComponentToRBPhysics();
-	}
-
-	// this used to not run if not rendered, but that causes issues such as bounds not updated
-	// causing it to not rendered, at the end, I think we should blend body positions
-	// for example if you're only simulating, this has to happen all the time
-	// whether looking at it or not, otherwise
-	// @todo better solution is to check if it has moved by changing SyncComponentToRBPhysics to return true if anything modified
-	// and run this if that is true or rendered
-	// that will at least reduce the chance of mismatch
-	// generally if you move your actor position, this has to happen to approximately match their bounds
-	if (Bodies.Num() > 0 && IsRegistered())
-	{
-		BlendInPhysics();
-	}
-
-	// if skeletal mesh has clothing assets, call TickClothing
-	if (SkeletalMesh && SkeletalMesh->ClothingAssets.Num() > 0)
-		{
-		TickClothing(DeltaTime + SkippedTickDeltaTime);
-		}
-	}
-
 void USkeletalMeshComponent::UpdateClothState(float DeltaTime)
 {
 	int32 NumActors = ClothingActors.Num();
