* improved game console auto completion

diff --git a/Engine/Source/Runtime/Engine/Private/UserInterface/Console.cpp b/Engine/Source/Runtime/Engine/Private/UserInterface/Console.cpp
--- a/Engine/Source/Runtime/Engine/Private/UserInterface/Console.cpp
+++ b/Engine/Source/Runtime/Engine/Private/UserInterface/Console.cpp
@@ -279,13 +279,22 @@ void UConsole::UpdateCompleteIndices()
 			}
 			else
 			{
-				break;
+				if(Idx < TypedStr.Len())
+				{
+					// thre is more text behind the auto completed text, we don't need auto completion
+					return;
+				}
+				else
+				{
+					break;
+				}
 			}
 		}
 	}
 	if (Node != &AutoCompleteTree)
 	{
 		AutoCompleteIndices = Node->AutoCompleteListIndices;
+		bNavigatingHistory = false;
 	}
 }
 
@@ -639,7 +648,7 @@ bool UConsole::InputKey_InputLine( int32 ControllerId, FKey Key, EInputEvent Eve
 		}
 		else if( Key == EKeys::Up )
 		{
-			if (!bNavigatingHistory && !bCtrl && AutoCompleteIndices.Num() > 1)
+			if (!bNavigatingHistory && !bCtrl && AutoCompleteIndices.Num())
 			{
 				if(AutoCompleteCursor + 1 < FMath::Min((int32)MAX_AUTOCOMPLETION_LINES, AutoCompleteIndices.Num()))
 				{
@@ -672,7 +681,12 @@ bool UConsole::InputKey_InputLine( int32 ControllerId, FKey Key, EInputEvent Eve
 
 					SetInputText(History[HistoryCur]);
 					SetCursorPos(History[HistoryCur].Len());
-					UpdateCompleteIndices();
+					//testUpdateCompleteIndices();
+					AutoCompleteIndex = 0;
+					AutoCompleteCursor = -1;
+					AutoCompleteIndices.Empty();
+
+
 					bNavigatingHistory = true;
 				}
 				return true;
@@ -706,7 +720,11 @@ bool UConsole::InputKey_InputLine( int32 ControllerId, FKey Key, EInputEvent Eve
 
 					SetInputText(History[HistoryCur]);
 					SetCursorPos(History[HistoryCur].Len());
-					UpdateCompleteIndices();
+					//testUpdateCompleteIndices();
+					AutoCompleteIndex = 0;
+					AutoCompleteCursor = -1;
+					AutoCompleteIndices.Empty();
+
 					bNavigatingHistory = true;
 				}
 
