* added r.AmbientOcclusion.FadeRadiusScale to give control over SSAO look

diff --git a/Engine/Source/Runtime/Engine/Private/SceneView.cpp b/Engine/Source/Runtime/Engine/Private/SceneView.cpp
--- a/Engine/Source/Runtime/Engine/Private/SceneView.cpp
+++ b/Engine/Source/Runtime/Engine/Private/SceneView.cpp
@@ -28,6 +28,13 @@ static TAutoConsoleVariable<float> CVarSSRMaxRoughness(
 	ECVF_RenderThreadSafe);
 #endif
 
+static TAutoConsoleVariable<float> CVarSSAOFadeRadiusScale(
+	TEXT("r.AmbientOcclusion.FadeRadiusScale"),
+	1.0f,
+	TEXT("Allows to scale the ambient occlusion fade radius (SSAO).\n")
+	TEXT(" 0.01:smallest .. 1.0:normal (default), <1:smaller, >1:larger"),
+	ECVF_Cheat | ECVF_RenderThreadSafe);
+
 /** Global vertex color view mode setting when SHOW_VertexColors show flag is set */
 EVertexColorViewMode::Type GVertexColorViewMode = EVertexColorViewMode::Color;
 
@@ -919,6 +926,12 @@ void FSceneView::EndFinalPostprocessSettings()
 		FinalPostProcessSettings.AmbientOcclusionRadius *= Scale;
 	}
 
+	{
+		float Scale = FMath::Clamp(CVarSSAOFadeRadiusScale.GetValueOnGameThread(), 0.01f, 50.0f);
+
+		FinalPostProcessSettings.AmbientOcclusionDistance *= Scale;
+	}
+
 	if(!Family->EngineShowFlags.Lighting || !Family->EngineShowFlags.GlobalIllumination)
 	{
 		FinalPostProcessSettings.IndirectLightingColor = FLinearColor(0,0,0,0);
