* tweaked SSAO Occlusion default settings to have less haloing

diff --git a/Engine/Source/Runtime/Engine/Classes/Engine/Scene.h b/Engine/Source/Runtime/Engine/Classes/Engine/Scene.h
--- a/Engine/Source/Runtime/Engine/Classes/Engine/Scene.h
+++ b/Engine/Source/Runtime/Engine/Classes/Engine/Scene.h
@@ -839,7 +839,7 @@ struct FPostProcessSettings
 		AmbientOcclusionIntensity = .5f;
 		AmbientOcclusionStaticFraction = 1.0f;
 		AmbientOcclusionRadius = 40.0f;
-		AmbientOcclusionDistance = 800.0f;
+		AmbientOcclusionDistance = 80.0f;
 		AmbientOcclusionFadeDistance = 8000.0f;
 		AmbientOcclusionFadeRadius = 5000.0f;
 		AmbientOcclusionPower = 2.0f;
