#UE4 - dedicated server autlogin fixed
- removed OnLoginStatusChanged from autologin flow, its not correct, we are actively trying to login
- OnLoginChanged can be setup later by game code
- client auth on some platforms was not ready at time autologin completed, halting delegates and preventing RegisterServer()

diff --git a/Engine/Source/Runtime/Engine/Classes/GameFramework/GameSession.h b/Engine/Source/Runtime/Engine/Classes/GameFramework/GameSession.h
--- a/Engine/Source/Runtime/Engine/Classes/GameFramework/GameSession.h
+++ b/Engine/Source/Runtime/Engine/Classes/GameFramework/GameSession.h
@@ -58,9 +58,6 @@ class ENGINE_API AGameSession : public AInfo
     /** Delegate triggered on auto login completion */
 	virtual void OnLoginComplete(int32 LocalUserNum, bool bWasSuccessful, const FUniqueNetId& UserId, const FString& Error);
 
-    /** Delegate triggered on auto login success */ 
-	virtual void OnLoginChanged(int32 LocalUserNum);
-
 	/** 
 	 * Called from GameMode.PreLogin() and Login().
 	 * @param	Options	The URL options (e.g. name/spectator) the player has passed
diff --git a/Engine/Source/Runtime/Engine/Private/GameSession.cpp b/Engine/Source/Runtime/Engine/Private/GameSession.cpp
--- a/Engine/Source/Runtime/Engine/Private/GameSession.cpp
+++ b/Engine/Source/Runtime/Engine/Private/GameSession.cpp
@@ -66,7 +66,6 @@ bool AGameSession::ProcessAutoLogin()
 	IOnlineIdentityPtr IdentityInt = Online::GetIdentityInterface(World);
 	if (IdentityInt.IsValid())
 	{
-		IdentityInt->AddOnLoginChangedDelegate(FOnLoginChangedDelegate::CreateUObject(this, &AGameSession::OnLoginChanged));
 		IdentityInt->AddOnLoginCompleteDelegate(0, FOnLoginCompleteDelegate::CreateUObject(this, &AGameSession::OnLoginComplete));
 		if (!IdentityInt->AutoLogin(0))
 		{
@@ -87,23 +86,15 @@ void AGameSession::OnLoginComplete(int32 LocalUserNum, bool bWasSuccessful, cons
 	IOnlineIdentityPtr IdentityInt = Online::GetIdentityInterface(World);
 	if (IdentityInt.IsValid())
 	{
-		IdentityInt->ClearOnLoginChangedDelegate(FOnLoginChangedDelegate::CreateUObject(this, &AGameSession::OnLoginChanged));
-		IdentityInt->ClearOnLoginCompleteDelegate(0, FOnLoginCompleteDelegate::CreateUObject(this, &AGameSession::OnLoginComplete));
-	}
-}
-
-void AGameSession::OnLoginChanged(int32 LocalUserNum)
-{
-	UWorld* World = GetWorld();
-	IOnlineIdentityPtr IdentityInt = Online::GetIdentityInterface(World);
-	if (IdentityInt.IsValid())
-	{
-		IdentityInt->ClearOnLoginChangedDelegate(FOnLoginChangedDelegate::CreateUObject(this, &AGameSession::OnLoginChanged));
 		IdentityInt->ClearOnLoginCompleteDelegate(0, FOnLoginCompleteDelegate::CreateUObject(this, &AGameSession::OnLoginComplete));
 		if (IdentityInt->GetLoginStatus(0) == ELoginStatus::LoggedIn)
 		{
 			RegisterServer();
 		}
+		else
+		{
+			UE_LOG(LogGameMode, Warning, TEXT("Autologin attempt failed, unable to register server!"));
+		}
 	}
 }
 
