Reset connection timeout when initial connection completes to avoid instant disconnections when initial connection timeout is longer than connection timeout

#UE4

diff --git a/Engine/Source/Runtime/Engine/Private/NetConnection.cpp b/Engine/Source/Runtime/Engine/Private/NetConnection.cpp
--- a/Engine/Source/Runtime/Engine/Private/NetConnection.cpp
+++ b/Engine/Source/Runtime/Engine/Private/NetConnection.cpp
@@ -1486,6 +1486,7 @@ void UNetConnection::HandleClientPlayer( APlayerController *PC, UNetConnection*
 	PC->SetPlayer(LocalPlayer);
 	PC->NetConnection = NetConnection;
 	UE_LOG(LogNet, Log, TEXT("%s setplayer %s"),*PC->GetName(),*LocalPlayer->GetName());
+	LastReceiveTime = Driver->Time;
 	State = USOCK_Open;
 	PlayerController = PC;
 	OwningActor = PC;
