#ue4 - Fix jittery client rotation in network games when quickly turning.

This was caused by move combining failing to update the rotation of the pawn to match the control rotation (via FaceRotation) when that codepath is used.

diff --git a/Engine/Source/Runtime/Engine/Private/CharacterMovementComponent.cpp b/Engine/Source/Runtime/Engine/Private/CharacterMovementComponent.cpp
--- a/Engine/Source/Runtime/Engine/Private/CharacterMovementComponent.cpp
+++ b/Engine/Source/Runtime/Engine/Private/CharacterMovementComponent.cpp
@@ -4658,8 +4658,16 @@ void UCharacterMovementComponent::ReplicateMoveToServer(float DeltaTime, const F
 			// Now that we have reverted to the old position, prepare a new move from that position,
 			// using our current velocity, acceleration, and rotation, but applied over the combined time from the old and new move.
 
-			SaveBaseLocation();
 			NewMove->DeltaTime += ClientData->PendingMove->DeltaTime;
+			
+			if (PC)
+			{
+				// We reverted position to that at the start of the pending move (above), however some code paths expect rotation to be set correctly
+				// before character movement occurs (via FaceRotation), so try that now. The bOrientRotationToMovement path happens later as part of PerformMovement() and PhysicsRotation().
+				CharacterOwner->FaceRotation(PC->GetControlRotation(), NewMove->DeltaTime);
+			}
+
+			SaveBaseLocation();
 			NewMove->SetInitialPosition(CharacterOwner);
 
 			// Remove pending move from move list. It would have to be the last move on the list.
