#ue4 - Clean up some CharacterMovementComponent comments.

- bUseControllerDesiredRotation : explain that is is a smooth update, and that bOrientRotationToMovement overrides it.
- Fix misnamed network interface in comment (should have been INetworkPredictionInterface)

diff --git a/Engine/Source/Runtime/Engine/Classes/GameFramework/CharacterMovementComponent.h b/Engine/Source/Runtime/Engine/Classes/GameFramework/CharacterMovementComponent.h
--- a/Engine/Source/Runtime/Engine/Classes/GameFramework/CharacterMovementComponent.h
+++ b/Engine/Source/Runtime/Engine/Classes/GameFramework/CharacterMovementComponent.h
@@ -245,7 +245,7 @@ public:
 	UPROPERTY(Category="Character Movement", EditAnywhere, BlueprintReadWrite)
 	FRotator RotationRate;
 
-	/** If true, rotate the Character toward the Controller's desired rotation, using RotationRate as the rate of rotation change. */
+	/** If true, smoothly rotate the Character toward the Controller's desired rotation, using RotationRate as the rate of rotation change. Overridden by OrientRotationToMovement. */
 	UPROPERTY(Category="Character Movement", EditAnywhere, BlueprintReadWrite, AdvancedDisplay)
 	uint32 bUseControllerDesiredRotation:1;
 
@@ -1130,7 +1130,7 @@ public:
 	virtual float GetNetworkSafeRandomAngleDegrees() const;
 
 	//--------------------------------
-	// INetworkPredictable implementation
+	// INetworkPredictionInterface implementation
 
 	//--------------------------------
 	// Server hook
