Filter displayall output by current world. PIE with multiple clients can cause displayall to be crowded.

#UE4

diff --git a/Engine/Source/Runtime/Engine/Private/UnrealEngine.cpp b/Engine/Source/Runtime/Engine/Private/UnrealEngine.cpp
--- a/Engine/Source/Runtime/Engine/Private/UnrealEngine.cpp
+++ b/Engine/Source/Runtime/Engine/Private/UnrealEngine.cpp
@@ -8101,6 +8101,11 @@ void DrawStatsHUD( UWorld* World, FViewport* Viewport, FCanvas* Canvas, UCanvas*
 		{
 			for (TObjectIterator<UObject> It(true); It; ++It)
 			{
+				if (It->GetWorld() && It->GetWorld() != World)
+				{
+					continue;
+				}
+
 				for (int32 i = 0; i < DebugClasses.Num(); i++)
 				{
 					if ( It->IsA(DebugClasses[i].Class) && !It->IsTemplate() &&
