Start UI customization for WheeledVehicleComponent4W and refactor some tuning variables

diff --git a/Engine/Source/Editor/DetailCustomizations/Private/DetailCustomizations.cpp b/Engine/Source/Editor/DetailCustomizations/Private/DetailCustomizations.cpp
--- a/Engine/Source/Editor/DetailCustomizations/Private/DetailCustomizations.cpp
+++ b/Engine/Source/Editor/DetailCustomizations/Private/DetailCustomizations.cpp
@@ -12,6 +12,7 @@
 #include "SkinnedMeshComponentDetails.h"
 #include "SkeletalMeshComponentDetails.h"
 #include "MeshComponentDetails.h"
+#include "WheeledVehicleMovementComponent4WDetails.h"
 #include "MatineeActorDetails.h"
 #include "ReflectionCaptureDetails.h"
 #include "SkyLightComponentDetails.h"
@@ -42,6 +43,7 @@
 #include "CurveTableCustomization.h"
 #include "DialogueWaveDetails.h"
 #include "BodyInstanceCustomization.h"
+#include "VehicleTransmissionDataCustomization.h"
 #include "BodySetupDetails.h"
 #include "Runtime/Engine/Classes/PhysicsEngine/BodySetup.h"
 #include "SlateBrushCustomization.h"
@@ -101,6 +103,8 @@ void FDetailCustomizationsModule::StartupModule()
 	RegisterStructPropertyLayout( "DialogueContext", FOnGetStructCustomizationInstance::CreateStatic( &FDialogueContextStructCustomization::MakeInstance ) );
 	RegisterStructPropertyLayout( "DialogueWaveParameter", FOnGetStructCustomizationInstance::CreateStatic( &FDialogueWaveParameterStructCustomization::MakeInstance ) );
 	RegisterStructPropertyLayout( "BodyInstance", FOnGetStructCustomizationInstance::CreateStatic( &FBodyInstanceCustomization::MakeInstance ) );
+	RegisterStructPropertyLayout( "BodyInstance", FOnGetStructCustomizationInstance::CreateStatic(&FBodyInstanceCustomization::MakeInstance));
+	RegisterStructPropertyLayout( "VehicleTransmissionData", FOnGetStructCustomizationInstance::CreateStatic(&FVehicleTransmissionDataCustomization::MakeInstance));
 	RegisterStructPropertyLayout( "SlateBrush", FOnGetStructCustomizationInstance::CreateStatic( &FSlateBrushStructCustomization::MakeInstance ) );
 	RegisterStructPropertyLayout( "SlateSound", FOnGetStructCustomizationInstance::CreateStatic( &FSlateSoundStructCustomization::MakeInstance ) );
 	RegisterStructPropertyLayout( "Guid", FOnGetStructCustomizationInstance::CreateStatic( &FGuidStructCustomization::MakeInstance ) );
@@ -134,6 +138,7 @@ void FDetailCustomizationsModule::StartupModule()
 	RegisterCustomPropertyLayout( "StaticMeshComponent", FOnGetDetailCustomizationInstance::CreateStatic( &FStaticMeshComponentDetails::MakeInstance ) );
 	RegisterCustomPropertyLayout( "SkeletalMeshComponent", FOnGetDetailCustomizationInstance::CreateStatic( &FSkeletalMeshComponentDetails::MakeInstance ) );
 	RegisterCustomPropertyLayout( "SkinnedMeshComponent", FOnGetDetailCustomizationInstance::CreateStatic( &FSkinnedMeshComponentDetails::MakeInstance ) );
+	RegisterCustomPropertyLayout( "WheeledVehicleMovementComponent4W", FOnGetDetailCustomizationInstance::CreateStatic(&FWheeledVehicleMovementComponent4WDetails::MakeInstance));
 	RegisterCustomPropertyLayout( "LightComponent", FOnGetDetailCustomizationInstance::CreateStatic( &FLightComponentDetails::MakeInstance ) );
 	RegisterCustomPropertyLayout( "PointLightComponent", FOnGetDetailCustomizationInstance::CreateStatic( &FPointLightComponentDetails::MakeInstance ) );
 	RegisterCustomPropertyLayout( "DirectionalLightComponent", FOnGetDetailCustomizationInstance::CreateStatic( &FDirectionalLightComponentDetails::MakeInstance ) );
diff --git a/Engine/Source/Editor/DetailCustomizations/Private/VehicleTransmissionDataCustomization.cpp b/Engine/Source/Editor/DetailCustomizations/Private/VehicleTransmissionDataCustomization.cpp
new file mode 100644
--- /dev/null
+++ b/Engine/Source/Editor/DetailCustomizations/Private/VehicleTransmissionDataCustomization.cpp
@@ -0,0 +1,286 @@
+// Copyright 1998-2014 Epic Games, Inc. All Rights Reserved.
+
+#include "DetailCustomizationsPrivatePCH.h"
+#include "VehicleTransmissionDataCustomization.h"
+#include "ScopedTransaction.h"
+#include "Editor/Documentation/Public/IDocumentation.h"
+
+#define LOCTEXT_NAMESPACE "VehicleTransmissionDataCustomization"
+
+#define RowWidth_Customization 50
+#define GearColumnsWidth (75.f * 3.f)
+
+////////////////////////////////////////////////////////////////
+
+bool FVehicleTransmissionDataCustomization::IsAutomaticEnabled() const
+{
+	return SelectedTransmission && SelectedTransmission->bUseGearAutoBox;
+}
+
+//Helper function so we can make neutral and reverse look the same as forward gears
+void FVehicleTransmissionDataCustomization::CreateGearUIHelper(FDetailWidgetRow & GearsSetup, FText Label, TSharedRef<IPropertyHandle> GearHandle, EGearType GearType)
+{
+	uint32 NumChildren = 0;
+	GearHandle->GetNumChildren(NumChildren);	//we use num of children to determine if we are dealing with a full gear that has ratio, down, up - or just a single value
+
+	TSharedRef<SWidget> RatioWidget = (NumChildren > 1 ? GearHandle->GetChildHandle("Ratio")->CreatePropertyValueWidget() : GearHandle->CreatePropertyValueWidget());
+	TSharedRef<SWidget> DownRatioWidget = (NumChildren > 1 ? GearHandle->GetChildHandle("DownRatio")->CreatePropertyValueWidget() : GearHandle->CreatePropertyValueWidget());
+	TSharedRef<SWidget> UpRatioWidget = (NumChildren > 1 ? GearHandle->GetChildHandle("UpRatio")->CreatePropertyValueWidget() : GearHandle->CreatePropertyValueWidget());
+	
+	RatioWidget->SetEnabled(GearType != NeutralGear);
+	switch (GearType)
+	{
+	case ForwardGear:
+		{
+			DownRatioWidget->SetEnabled(TAttribute<bool>(this, &FVehicleTransmissionDataCustomization::IsAutomaticEnabled));
+			UpRatioWidget->SetEnabled(TAttribute<bool>(this, &FVehicleTransmissionDataCustomization::IsAutomaticEnabled));
+			break;
+		}
+	case ReverseGear:
+		{
+			DownRatioWidget->SetEnabled(false);
+			UpRatioWidget->SetEnabled(false);
+			break;
+		}
+	case NeutralGear:
+		{
+			DownRatioWidget->SetEnabled(false);
+			UpRatioWidget->SetEnabled(TAttribute<bool>(this, &FVehicleTransmissionDataCustomization::IsAutomaticEnabled));
+			break;
+		}
+	}
+
+	TSharedRef<SWidget> RemoveWidget = PropertyCustomizationHelpers::MakeDeleteButton(FSimpleDelegate::CreateSP(this, &FVehicleTransmissionDataCustomization::RemoveGear, GearHandle), LOCTEXT("RemoveGearToolTip", "Removes gear"));
+	RemoveWidget->SetEnabled(NumChildren > 1);
+	
+	GearsSetup
+	.NameContent()
+	[
+		SNew(STextBlock)
+		.Text(Label)
+		.Font(IDetailLayoutBuilder::GetDetailFont())
+	]
+	.ValueContent()
+		.MaxDesiredWidth(GearColumnsWidth)
+		.MinDesiredWidth(GearColumnsWidth)
+		[
+			SNew(SHorizontalBox)
+			+ SHorizontalBox::Slot()
+			.FillWidth(0.3333f)
+			[
+				RatioWidget
+			]
+			+ SHorizontalBox::Slot()
+			.FillWidth(0.3333f)
+			.Padding(4.f)
+			[
+				DownRatioWidget
+			]
+
+			+ SHorizontalBox::Slot()
+			.FillWidth(0.3333f)
+			.Padding(4.f)
+			[
+				UpRatioWidget
+			]
+
+			+ SHorizontalBox::Slot()
+			.Padding(4.f)
+			.AutoWidth()
+			[
+				RemoveWidget
+			]
+		];
+}
+
+void FVehicleTransmissionDataCustomization::BuildColumnsHeaderHelper(TSharedRef<class IPropertyHandle> StructPropertyHandle, FDetailWidgetRow& GearsSetup)
+{
+	
+	GearsSetup
+	.NameContent()
+		[
+			SNew(STextBlock)
+			.Text(LOCTEXT("GearSetup", "Gear Setup"))
+			.Font(IDetailLayoutBuilder::GetDetailFontBold())
+		]
+	.ValueContent()
+		.MaxDesiredWidth(GearColumnsWidth)
+		.MinDesiredWidth(GearColumnsWidth)
+		[
+			SNew(SHorizontalBox)
+			+ SHorizontalBox::Slot()
+			.FillWidth(0.3333f)
+			.VAlign(VAlign_Center)
+			[
+				SNew(SBox)
+				.WidthOverride(RowWidth_Customization)
+				.HAlign(HAlign_Left)
+				.Content()
+				[
+					SNew(STextBlock)
+					.Text(LOCTEXT("RatioLabel", "Gear Ratio"))
+					.Font(IDetailLayoutBuilder::GetDetailFontBold())
+				]
+			]
+			+ SHorizontalBox::Slot()
+				.FillWidth(0.3333f)
+				.VAlign(VAlign_Center)
+				[
+					SNew(SBox)
+					.HAlign(HAlign_Left)
+					.WidthOverride(RowWidth_Customization)
+					.Content()
+					[
+						SNew(STextBlock)
+						.Text(LOCTEXT("LowRPMLabel", "Down Ratio"))
+						.Font(IDetailLayoutBuilder::GetDetailFontBold())
+					]
+				]
+			+ SHorizontalBox::Slot()
+				.FillWidth(0.3333f)
+				.VAlign(VAlign_Center)
+				[
+					SNew(STextBlock)
+					.Text(LOCTEXT("HighRPMLabel", "Up Ratio"))
+					.Font(IDetailLayoutBuilder::GetDetailFontBold())
+				]
+			+ SHorizontalBox::Slot()
+				.AutoWidth()
+				.VAlign(VAlign_Center)
+				[
+					PropertyCustomizationHelpers::MakeAddButton(FSimpleDelegate::CreateSP(this, &FVehicleTransmissionDataCustomization::AddGear, StructPropertyHandle), LOCTEXT("RemoveGearToolTip", "Adds a new gear"))
+				]
+			+ SHorizontalBox::Slot()
+				.AutoWidth()
+				.VAlign(VAlign_Center)
+				[
+					PropertyCustomizationHelpers::MakeEmptyButton(FSimpleDelegate::CreateSP(this, &FVehicleTransmissionDataCustomization::EmptyGears, StructPropertyHandle), LOCTEXT("EmptyGearToolTip", "Removes all gears"))
+				]
+			
+		];
+}
+
+void FVehicleTransmissionDataCustomization::CreateGearUIDelegate(TSharedRef<IPropertyHandle> GearProperty, int32 GearIdx, IDetailChildrenBuilder& ChildrenBuilder)
+{
+	FText Label = FText::Format(LOCTEXT("TransmissionGear", "Gear {0}"), FText::AsNumber(GearIdx + 1));
+	CreateGearUIHelper(ChildrenBuilder.AddChildContent(Label.ToString()), Label, GearProperty, ForwardGear);
+}
+
+void FVehicleTransmissionDataCustomization::CustomizeStructChildren(TSharedRef<class IPropertyHandle> StructPropertyHandle, class IDetailChildrenBuilder& StructBuilder, IStructCustomizationUtils& StructCustomizationUtils)
+{
+	// copy all transmision instances I'm accessing right now
+	TArray<void*> StructPtrs;
+	StructPropertyHandle->AccessRawData(StructPtrs);
+	
+	if (StructPtrs.Num() == 1)
+	{
+		SelectedTransmission = (FVehicleTransmissionData *)StructPtrs[0];
+	}
+	else
+	{
+		SelectedTransmission = NULL;
+	}
+
+
+	uint32 NumChildren = 0;
+	StructPropertyHandle->GetNumChildren(NumChildren);
+
+	FName GearsSetupGroupName = TEXT("GearsSetup");
+	FString GearsSetupGroupLabel = LOCTEXT("GearSetupLabel", "Gears Setup").ToString();
+	IDetailGroup * GearSetupGroup = NULL;
+	
+	for (uint32 ChildIdx = 0; ChildIdx < NumChildren; ++ChildIdx)
+	{
+		TSharedRef<IPropertyHandle> ChildProperty = StructPropertyHandle->GetChildHandle(ChildIdx).ToSharedRef();
+		const FString PropertyName = ChildProperty->GetProperty() ? ChildProperty->GetProperty()->GetName() : TEXT("");
+
+		if (PropertyName == TEXT("ForwardGears") || PropertyName == TEXT("NeutralGearUpRatio") || PropertyName == TEXT("ReverseGearRatio"))
+		{
+			if (GearSetupGroup == NULL)
+			{
+				GearSetupGroup = &StructBuilder.AddChildGroup(GearsSetupGroupName, GearsSetupGroupLabel);
+				BuildColumnsHeaderHelper(StructPropertyHandle, StructBuilder.AddChildContent(TEXT("GearSetup")));
+			}
+
+			//determine which gear we're showing
+			EGearType GearType = ForwardGear;
+			if (PropertyName == TEXT("NeutralGearUpRatio"))
+			{
+				GearType = NeutralGear;
+			}
+			else if
+			(PropertyName == TEXT("ReverseGearRatio"))
+			{
+				GearType = ReverseGear;
+			}
+						
+			if (GearType == ForwardGear)
+			{
+				TSharedRef<FDetailArrayBuilder> GearsArrayBuilder = MakeShareable(new FDetailArrayBuilder(ChildProperty, false));
+				GearsArrayBuilder->OnGenerateArrayElementWidget(FOnGenerateArrayElementWidget::CreateSP(this, &FVehicleTransmissionDataCustomization::CreateGearUIDelegate));
+				StructBuilder.AddChildCustomBuilder(GearsArrayBuilder);
+			}
+			else
+			{	
+				CreateGearUIHelper(StructBuilder.AddChildContent(PropertyName), FText::FromString(PropertyName), ChildProperty, GearType);
+			}
+
+
+		}
+		else
+		{	//Add all other properties
+			StructBuilder.AddChildProperty(ChildProperty);
+		}
+	}
+
+}
+
+void FVehicleTransmissionDataCustomization::AddGear(TSharedRef<IPropertyHandle> StructPropertyHandle)
+{
+	TSharedPtr<IPropertyHandle> GearsHandle = StructPropertyHandle->GetChildHandle("ForwardGears");
+	if (GearsHandle->IsValidHandle())
+	{
+		TSharedPtr<IPropertyHandleArray> GearsArray = GearsHandle->AsArray();
+		GearsArray->AddItem();
+	}
+}
+
+void FVehicleTransmissionDataCustomization::EmptyGears(TSharedRef<IPropertyHandle> StructPropertyHandle)
+{
+	TSharedPtr<IPropertyHandle> GearsHandle = StructPropertyHandle->GetChildHandle("ForwardGears");
+	if (GearsHandle->IsValidHandle())
+	{
+		TSharedPtr<IPropertyHandleArray> GearsArray = GearsHandle->AsArray();
+		GearsArray->EmptyArray();
+	}
+}
+
+void FVehicleTransmissionDataCustomization::RemoveGear(TSharedRef<IPropertyHandle> GearHandle)
+{
+	if (GearHandle->IsValidHandle())
+	{
+		const TSharedPtr<IPropertyHandle> ParentHandle = GearHandle->GetParentHandle();
+		const TSharedPtr<IPropertyHandleArray> ParentArrayHandle = ParentHandle->AsArray();
+
+		ParentArrayHandle->DeleteItem(GearHandle->GetIndexInArray());
+		
+	}
+}
+
+
+void FVehicleTransmissionDataCustomization::CustomizeStructHeader(TSharedRef<IPropertyHandle> StructPropertyHandle, FDetailWidgetRow& HeaderRow, IStructCustomizationUtils& StructCustomizationUtils)
+{
+	HeaderRow.
+	NameContent()
+	[
+		StructPropertyHandle->CreatePropertyNameWidget()
+	]
+	.ValueContent()
+		[
+			StructPropertyHandle->CreatePropertyValueWidget()
+		];
+}
+
+#undef LOCTEXT_NAMESPACE
+#undef GearColumnsWidth
+#undef RowWidth_Customization
diff --git a/Engine/Source/Editor/DetailCustomizations/Private/VehicleTransmissionDataCustomization.h b/Engine/Source/Editor/DetailCustomizations/Private/VehicleTransmissionDataCustomization.h
new file mode 100644
--- /dev/null
+++ b/Engine/Source/Editor/DetailCustomizations/Private/VehicleTransmissionDataCustomization.h
@@ -0,0 +1,35 @@
+// Copyright 1998-2014 Epic Games, Inc. All Rights Reserved.
+
+#pragma once
+
+class FVehicleTransmissionDataCustomization : public IStructCustomization
+{
+public:
+	
+	static TSharedRef<IStructCustomization> MakeInstance() 
+	{
+		return MakeShareable(new FVehicleTransmissionDataCustomization);
+	}
+
+	/** IStructCustomization interface */
+	virtual void CustomizeStructHeader(TSharedRef<IPropertyHandle> StructPropertyHandle, class FDetailWidgetRow& HeaderRow, IStructCustomizationUtils& StructCustomizationUtils) OVERRIDE;
+	virtual void CustomizeStructChildren( TSharedRef<class IPropertyHandle> StructPropertyHandle, class IDetailChildrenBuilder& StructBuilder, IStructCustomizationUtils& StructCustomizationUtils ) OVERRIDE;
+
+private:
+	FVehicleTransmissionData * SelectedTransmission;
+
+	enum EGearType
+	{
+		ForwardGear,
+		ReverseGear,
+		NeutralGear
+	};
+
+	void CreateGearUIHelper(FDetailWidgetRow& GearsSetup, FText Label, TSharedRef<IPropertyHandle> GearHandle, EGearType GearType);
+	void BuildColumnsHeaderHelper(TSharedRef<IPropertyHandle> StructPropertyHandle, FDetailWidgetRow & GearsSetup);
+	void RemoveGear(TSharedRef<IPropertyHandle> GearHandle);
+	void AddGear(TSharedRef<IPropertyHandle> StructPropertyHandle);
+	void EmptyGears(TSharedRef<IPropertyHandle> StructPropertyHandle);
+	void CreateGearUIDelegate(TSharedRef<IPropertyHandle> PropertyHandle, int32 GearyIndex, IDetailChildrenBuilder& ChildrenBuilder);
+	bool IsAutomaticEnabled() const;
+};
diff --git a/Engine/Source/Editor/DetailCustomizations/Private/WheeledVehicleMovementComponent4WDetails.cpp b/Engine/Source/Editor/DetailCustomizations/Private/WheeledVehicleMovementComponent4WDetails.cpp
new file mode 100644
--- /dev/null
+++ b/Engine/Source/Editor/DetailCustomizations/Private/WheeledVehicleMovementComponent4WDetails.cpp
@@ -0,0 +1,174 @@
+﻿// Copyright 1998-2014 Epic Games, Inc. All Rights Reserved.
+
+#include "DetailCustomizationsPrivatePCH.h"
+#include "WheeledVehicleMovementComponent4WDetails.h"
+#include "ScopedTransaction.h"
+#include "ObjectEditorUtils.h"
+#include "Vehicles/WheeledVehicleMovementComponent4W.h"
+#include "IDocumentation.h"
+
+#define LOCTEXT_NAMESPACE "WheeledVehicleMovementComponent4WDetails"
+
+//////////////////////////////////////////////////////////////
+// This class customizes various settings in WheeledVehicleMovementComponent4W
+//////////////////////////////////////////////////////////////
+
+TSharedRef<IDetailCustomization> FWheeledVehicleMovementComponent4WDetails::MakeInstance()
+{
+	return MakeShareable(new FWheeledVehicleMovementComponent4WDetails);
+}
+
+void FWheeledVehicleMovementComponent4WDetails::FillAutoBoxWidget(UWheeledVehicleMovementComponent4W * VehicleMovement, FDetailWidgetRow & AutoBoxSetup)
+{
+	/*FSlateFontInfo DetailFont = IDetailLayoutBuilder::GetDetailFont();
+
+	AutoBoxSetup.NameContent()
+	[
+		SNew(STextBlock)
+		.Text(LOCTEXT("AutoBoxSetupHeader", "Auto Box Setup"))
+		.ToolTipText(LOCTEXT("AutoBoxSetupHeaderTip", "Describes the automatic transition of the vehicle - how it shifts gears etc."))
+		.Font(DetailFont)
+	];
+
+	TSharedRef<SVerticalBox> AutoBoxWidget = SNew(SVerticalBox);
+
+	AutoBoxWidget
+	[
+		SNew(STextBlock)
+		.Text(LOCTEXT("ASDS", "Testing"))
+	];
+
+	AutoBoxSetup.ValueContent()
+	[
+		AutoBoxWidget
+	];*/
+	
+	//We want to inline the gear ratios with Neutral and Latency.
+	//We only want to show as many ratios as there are gears
+	/*FVehicleAutoBoxData AutoBox = VehicleMovement->AutoBoxSetup;
+	FVehicleGearData GearData = VehicleMovement->GearSetup;
+
+	uint32 NumChildren = 0;
+	AutoBoxSetup->GetNumChildren(NumChildren);
+	for (uint32 ChildIndex = 0; ChildIndex < NumChildren; ++ChildIndex)
+	{
+		TSharedPtr<IPropertyHandle> ChildProperty = AutoBoxSetup->GetChildHandle(ChildIndex);
+		// Only permit modifying bSimulatePhysics when the body has some geometry.
+		if (ChildProperty->GetProperty()->GetName() == TEXT("ForwardGearAutoBox"))
+		{
+			for (int32 GearIdx = 0; GearIdx != GearData.ForwardGearRatios.Num(); ++GearIdx)
+			{
+				const FGearUpDownRatio & GearUpDownRatio = AutoBox.ForwardGearAutoBox[GearIdx + FVehicleGearData::eFIRST];
+				FText Label = FText::Format(LOCTEXT("AutoBoxGear", "Gear {0}"), FText::AsNumber(GearIdx + 1));
+
+				VehicleSetupCategory.AddCustomRow(Label.ToString(), false)
+				.NameContent()
+					[
+						SNew(STextBlock)
+						.Text(Label)
+						.ToolTipText(LOCTEXT("MassInKG", "Mass of the body in KG"))
+						.Font(DetailFont)
+					]
+				.ValueContent()
+					[
+						SNew(SHorizontalBox)
+						+ SHorizontalBox::Slot()
+						.FillWidth(0.5f)
+						[
+							SNew(SNumericEntryBox<float>)
+							.Value(GearUpDownRatio.DownRatio)
+							.MinValue(0.f)
+							.MaxValue(1.f)
+							.Font(DetailFont)
+							.OnValueCommitted(this, &FWheeledVehicleMovementComponent4WDetails::SetGearShiftRatio, false, GearIdx)
+							.LabelVAlign(VAlign_Center)
+							.Label()
+								[
+									SNew(STextBlock)
+									.Font(DetailFont)
+									.Text(LOCTEXT("AutoBoxShiftMin", "Min"))
+								]
+							.ToolTipText(LOCTEXT("AutoBoxShiftMinTip", "Minimum currentRPM/maxRPM ratio required to stay in this gear. Going lower will cause autobox to gear down"))
+						]	
+						+ SHorizontalBox::Slot()
+						.FillWidth(0.5f)
+						.Padding(4, 0)
+						[
+							SNew(SNumericEntryBox<float>)
+							.Value(GearUpDownRatio.UpRatio)
+							.MinValue(0.f)
+							.MaxValue(1.f)
+							.Font(DetailFont)
+							.OnValueCommitted(this, &FWheeledVehicleMovementComponent4WDetails::SetGearShiftRatio, true, GearIdx)
+							.LabelVAlign(VAlign_Center)
+							.Label()
+							[
+								SNew(STextBlock)
+								.Font(DetailFont)
+								.Text(LOCTEXT("AutoBoxShiftMax", "Max"))
+							]
+							.ToolTipText(LOCTEXT("AutoBoxShiftMinTip", "Maximum currentRPM/maxRPM ratio required to stay in this gear. Going higher will cause autobox to gear up"))
+						]
+					];
+
+			}
+		}
+		else
+		{
+			VehicleSetupCategory.AddProperty(ChildProperty);
+		}
+	}*/
+}
+
+void FWheeledVehicleMovementComponent4WDetails::CustomizeDetails(IDetailLayoutBuilder& DetailBuilder)
+{
+	DetailBuilder.GetObjectsBeingCustomized(SelectedObjects);
+	if (SelectedObjects.Num() != 1)	//only do customization if we have one object selected
+	{
+		return;
+	}
+
+	UWheeledVehicleMovementComponent4W * VehicleMovement = SelectedObjects[0].IsValid() ? Cast<UWheeledVehicleMovementComponent4W>(SelectedObjects[0].Get()) : NULL;
+
+	// See which categories are hidden
+	TArray<FString> HideCategories;
+	DetailBuilder.GetDetailsView().GetBaseClass()->GetHideCategories(HideCategories);
+	
+
+	//Vehicle Setup
+	if (!HideCategories.Contains(TEXT("VehicleSetup")))
+	{
+		IDetailCategoryBuilder& VehicleSetupCategory = DetailBuilder.EditCategory("VehicleSetup");
+		//Autobox
+		//FDetailWidgetRow & AutoBoxRow = VehicleSetupCategory.AddProperty(DetailBuilder.GetProperty("AutoBoxSetup")).CustomWidget();
+		//FillAutoBoxWidget(VehicleMovement, AutoBoxRow);
+	}
+}
+
+void FWheeledVehicleMovementComponent4WDetails::SetGearShiftRatio(float NewValue, ETextCommit::Type, bool bShiftUp, int32 GearIdx)
+{
+	/*if (SelectedObjects.Num() != 1)	//only do customization if we have one object selected
+	{
+		return;
+	}
+
+	if (UWheeledVehicleMovementComponent4W * VehicleMovement = SelectedObjects[0].IsValid() ? Cast<UWheeledVehicleMovementComponent4W>(SelectedObjects[0].Get()) : NULL)
+	{
+		FVehicleAutoBoxData & AutoBox = VehicleMovement->AutoBoxSetup;
+		FGearUpDownRatio & GearUpDownRatio = AutoBox.ForwardGearAutoBox[GearIdx + FVehicleGearData::eFIRST];
+		
+		if (bShiftUp)
+		{
+			GearUpDownRatio.UpRatio = NewValue;
+		}
+		else
+		{
+			GearUpDownRatio.DownRatio = NewValue;
+		}
+
+	}*/
+
+}
+
+#undef LOCTEXT_NAMESPACE
+
diff --git a/Engine/Source/Editor/DetailCustomizations/Private/WheeledVehicleMovementComponent4WDetails.h b/Engine/Source/Editor/DetailCustomizations/Private/WheeledVehicleMovementComponent4WDetails.h
new file mode 100644
--- /dev/null
+++ b/Engine/Source/Editor/DetailCustomizations/Private/WheeledVehicleMovementComponent4WDetails.h
@@ -0,0 +1,19 @@
+// Copyright 1998-2014 Epic Games, Inc. All Rights Reserved.
+
+#pragma once
+
+class FWheeledVehicleMovementComponent4WDetails : public IDetailCustomization
+{
+public:
+	/** Makes a new instance of this detail layout class for a specific detail view requesting it */
+	static TSharedRef<IDetailCustomization> MakeInstance();
+
+	/** IDetailCustomization interface */
+	virtual void CustomizeDetails( IDetailLayoutBuilder& DetailBuilder ) OVERRIDE;
+private:
+	void SetGearShiftRatio(float NewValue, ETextCommit::Type CommitInfo, bool bShiftUp, int32 GearIdx);
+	void FillAutoBoxWidget(UWheeledVehicleMovementComponent4W * VehicleMovement, FDetailWidgetRow & AutoBoxSetup);
+
+	TArray<TWeakObjectPtr<UObject> > SelectedObjects;	//the objects we're showing details for
+};
+
diff --git a/Engine/Source/Runtime/Engine/Classes/Vehicles/WheeledVehicleMovementComponent4W.h b/Engine/Source/Runtime/Engine/Classes/Vehicles/WheeledVehicleMovementComponent4W.h
--- a/Engine/Source/Runtime/Engine/Classes/Vehicles/WheeledVehicleMovementComponent4W.h
+++ b/Engine/Source/Runtime/Engine/Classes/Vehicles/WheeledVehicleMovementComponent4W.h
@@ -59,10 +59,6 @@ struct FVehicleEngineData
 {
 	GENERATED_USTRUCT_BODY()
 
-	/** Moment of inertia of the engine around the axis of rotation. */
-	UPROPERTY(EditAnywhere, Category=Setup)
-	float MOI;
-
 	/** Maximum torque available to apply to the engine (Nm). */ 
 	UPROPERTY(EditAnywhere, Category=Setup)
 	float PeakTorque;
@@ -71,6 +67,10 @@ struct FVehicleEngineData
 	UPROPERTY(EditAnywhere, Category=Setup)
 	float MaxRPM;
 
+	/** Moment of inertia of the engine around the axis of rotation. */
+	UPROPERTY(EditAnywhere, Category = Setup)
+	float MOI;
+
 	/** Damping rate of engine when full throttle is applied */
 	UPROPERTY(EditAnywhere, Category=Setup, AdvancedDisplay)
 	float DampingRateFullThrottle;
@@ -84,58 +84,60 @@ struct FVehicleEngineData
 	float DampingRateZeroThrottleClutchDisengaged;
 };
 
+
 USTRUCT()
 struct FVehicleGearData
 {
 	GENERATED_USTRUCT_BODY()
 
-	/** Time it takes to switch gear (seconds) */
-	UPROPERTY(EditAnywhere, Category=Setup)
-	float GearSwitchTime;
+	/** Determines the amount of torque multiplication*/
+	UPROPERTY(EditAnywhere, Category = Setup)
+	float Ratio;
 
-	/** Reverse gear ratio */
-	UPROPERTY(EditAnywhere, Category=Setup)
-	float ReverseGearRatio;
-
-	/** Forward gear ratios (up to 30) */
-	UPROPERTY(EditAnywhere, Category=Setup)
-	TArray<float> ForwardGearRatios;
+	/** Value of engineRevs/maxEngineRevs that is low enough to gear down*/
+	UPROPERTY(EditAnywhere, meta = (ClampMin = "0.0", UIMin = "0.0", ClampMax = "1.0", UIMax = "1.0"), Category = Setup)
+	float DownRatio;
 
-	/** Gear ratio multiplier */
-	UPROPERTY(EditAnywhere, Category=Setup)
-	float RatioMultiplier;
+	/** Value of engineRevs/maxEngineRevs that is high enough to gear up*/
+	UPROPERTY(EditAnywhere, meta = (ClampMin = "0.0", UIMin = "0.0", ClampMax = "1.0", UIMax = "1.0"), Category = Setup)
+	float UpRatio;
 };
 
 USTRUCT()
-struct FGearUpDownRatio
+struct FVehicleTransmissionData
 {
 	GENERATED_USTRUCT_BODY()
+	/** Whether to use automatic transmission */
+	UPROPERTY(EditAnywhere, Category = VehicleSetup, meta=(FriendlyName = "Automatic Transmission"))
+	bool bUseGearAutoBox;
 
-	/** Value of engineRevs/maxEngineRevs that is high enough to increment gear (range 0..1) */
-	UPROPERTY(EditAnywhere, Category=Setup)
-	float UpRatio;
+	/** Time it takes to switch gears (seconds) */
+	UPROPERTY(EditAnywhere, Category = Setup)
+	float GearSwitchTime;
+	
+	/** Minimum time it takes the automatic transmission to initiate a gear change (seconds)*/
+	UPROPERTY(EditAnywhere, Category = Setup, meta = (editcondition = "bUseGearAutoBox"))
+	float GearAutoBoxLatency;
 
-	/** Value of engineRevs/maxEngineRevs that is low enough to decrement gear (range 0..1) */
-	UPROPERTY(EditAnywhere, Category=Setup)
-	float DownRatio;
-};
+	/** The final gear ratio multiplies the transmission gear ratios.*/
+	UPROPERTY(EditAnywhere, AdvancedDisplay, Category = Setup)
+	float FinalRatio;
 
-USTRUCT()
-struct FVehicleAutoBoxData
-{
-	GENERATED_USTRUCT_BODY()
+	/** Forward gear ratios (up to 30) */
+	UPROPERTY(EditAnywhere, Category = Setup, AdvancedDisplay)
+	TArray<FVehicleGearData> ForwardGears;
 
-	/** Value of engineRevs/maxEngineRevs for auto gear change (range 0..1) */ 
-	UPROPERTY(EditAnywhere, Category=Setup, EditFixedSize)
-	TArray<FGearUpDownRatio> ForwardGearAutoBox;
+	/** Reverse gear ratio */
+	UPROPERTY(EditAnywhere, AdvancedDisplay, Category = Setup)
+	float ReverseGearRatio;
 
-	/** Value of engineRevs/maxEngineRevs that is high enough to increment gear (range 0..1) */ 
-	UPROPERTY(EditAnywhere, Category=Setup)
+	/** Value of engineRevs/maxEngineRevs that is high enough to increment gear*/
+	UPROPERTY(EditAnywhere, AdvancedDisplay, Category = Setup, meta = (ClampMin = "0.0", UIMin = "0.0", ClampMax = "1.0", UIMax = "1.0"))
 	float NeutralGearUpRatio;
 
-	/** Latency time is the minimum time that must pass between each gear change that is initiated by the autobox */ 
-	UPROPERTY(EditAnywhere, Category=Setup)
-	float GearAutoBoxLatency;
+	/** Strength of clutch */
+	UPROPERTY(EditAnywhere, Category = Setup, AdvancedDisplay)
+	float ClutchStrength;
 };
 
 UCLASS()
@@ -144,43 +146,32 @@ class ENGINE_API UWheeledVehicleMovementComponent4W : public UWheeledVehicleMove
 	GENERATED_UCLASS_BODY()
 
 	/** Engine */
-	UPROPERTY(EditAnywhere, Category=VehicleSetup)
+	UPROPERTY(EditAnywhere, Category = MechanicalSetup)
 	FVehicleEngineData EngineSetup;
 
 	/** Differential */
-	UPROPERTY(EditAnywhere, Category=VehicleSetup, AdvancedDisplay)
+	UPROPERTY(EditAnywhere, Category = MechanicalSetup)
 	FVehicleDifferential4WData DifferentialSetup;
 
-	/** Gears */
-	UPROPERTY(EditAnywhere, Category=VehicleSetup, AdvancedDisplay)
-	FVehicleGearData GearSetup;
-	
-	/** Automatic gear switching */
-	UPROPERTY(EditAnywhere, Category=VehicleSetup, AdvancedDisplay)
-	FVehicleAutoBoxData AutoBoxSetup;
-
-	/** Should use auto gear switching by default? */
-	UPROPERTY(EditAnywhere, Category=VehicleSetup)
-	bool bUseGearAutoBox;
+	/** Transmission data */
+	UPROPERTY(EditAnywhere, Category = MechanicalSetup)
+	FVehicleTransmissionData TransmissionSetup;
 
-	/** Strength of clutch */
-	UPROPERTY(EditAnywhere, Category=VehicleSetup, AdvancedDisplay)
-	float ClutchStrength;
+	/** Max speed value in steering curve (km/h)*/
+	UPROPERTY(EditAnywhere, Category=SteeringSetup)
+	float MaxSteeringSpeed;
 
 	/** Accuracy of Ackermann steer calculation (range: 0..1) */
-	UPROPERTY(EditAnywhere, Category=VehicleSetup, AdvancedDisplay)
+	UPROPERTY(EditAnywhere, Category = SteeringSetup, AdvancedDisplay)
 	float AckermannAccuracy;
 
-	/** Max speed value in steering curve (km/h)*/
-	UPROPERTY(EditAnywhere, Category=VehicleSetup)
-	float MaxSteeringSpeed;
-
 	/** Steering strength defined for speed in normalized range [0 .. 1] */
-	UPROPERTY(EditAnywhere, Category=VehicleSetup, AdvancedDisplay)
+	UPROPERTY(EditAnywhere, Category=SteeringSetup)
 	TArray<FFloatPair> SteeringCurve;
 
 	virtual void Serialize(FArchive & Ar) OVERRIDE;
 	virtual void ComputeConstants() OVERRIDE;
+	virtual void PostEditChangeProperty(struct FPropertyChangedEvent& PropertyChangedEvent) OVERRIDE;
 
 protected:
 
@@ -201,10 +192,7 @@ protected:
 	/** update simulation data: differential */
 	void UpdateDifferentialSetup(const FVehicleDifferential4WData& NewDifferentialSetup);
 
-	/** update simulation data: gears */
-	void UpdateGearSetup(const FVehicleGearData& NewGearSetup);
-
-	/** update simulation data: auto box */
-	void UpdateAutoBoxSetup(const FVehicleAutoBoxData& NewAutoBoxSetup);
+	/** update simulation data: transmission */
+	void UpdateTransmissionSetup(const FVehicleTransmissionData& NewGearSetup);
 };
 
diff --git a/Engine/Source/Runtime/Engine/Private/Vehicles/WheeledVehicleMovementComponent4W.cpp b/Engine/Source/Runtime/Engine/Private/Vehicles/WheeledVehicleMovementComponent4W.cpp
--- a/Engine/Source/Runtime/Engine/Private/Vehicles/WheeledVehicleMovementComponent4W.cpp
+++ b/Engine/Source/Runtime/Engine/Private/Vehicles/WheeledVehicleMovementComponent4W.cpp
@@ -20,43 +20,39 @@ UWheeledVehicleMovementComponent4W::UWheeledVehicleMovementComponent4W(const cla
 	DifferentialSetup.FrontBias = DefDifferentialSetup.mFrontBias;
 	DifferentialSetup.RearBias = DefDifferentialSetup.mRearBias;
 
-	const PxReal LengthScale = 100.f; // Convert default from m to cm
-
 	PxVehicleEngineData DefEngineData;
-	EngineSetup.MOI = DefEngineData.mMOI * LengthScale * LengthScale;
-	EngineSetup.PeakTorque = 1000.f;
-	EngineSetup.MaxRPM = 7000.f;
-	EngineSetup.DampingRateFullThrottle = 0.5f * LengthScale * LengthScale;
-	EngineSetup.DampingRateZeroThrottleClutchEngaged = 2.0f * LengthScale * LengthScale;
-	EngineSetup.DampingRateZeroThrottleClutchDisengaged = 0.5f * LengthScale * LengthScale;
+	EngineSetup.MOI = M2ToCm2(DefEngineData.mMOI);
+	EngineSetup.PeakTorque = DefEngineData.mPeakTorque;
+	EngineSetup.MaxRPM = OmegaToRPM(DefEngineData.mMaxOmega);
+	EngineSetup.DampingRateFullThrottle = M2ToCm2(DefEngineData.mDampingRateFullThrottle);
+	EngineSetup.DampingRateZeroThrottleClutchEngaged = M2ToCm2(DefEngineData.mDampingRateZeroThrottleClutchEngaged);
+	EngineSetup.DampingRateZeroThrottleClutchDisengaged = M2ToCm2(DefEngineData.mDampingRateZeroThrottleClutchDisengaged);
 
 	PxVehicleClutchData DefClutchData;
-	ClutchStrength = DefClutchData.mStrength * LengthScale * LengthScale; // convert from m to cm scale
+	TransmissionSetup.ClutchStrength = M2ToCm2(DefClutchData.mStrength); // convert from m to cm scale
 
 	PxVehicleAckermannGeometryData DefAckermannSetup;
 	AckermannAccuracy = DefAckermannSetup.mAccuracy;
 
 	PxVehicleGearsData DefGearSetup;
-	GearSetup.GearSwitchTime = DefGearSetup.mSwitchTime;
-	GearSetup.ReverseGearRatio = DefGearSetup.mRatios[PxVehicleGearsData::eREVERSE];
-	for (uint32 i = PxVehicleGearsData::eFIRST; i < DefGearSetup.mNbRatios; i++)
-	{
-		GearSetup.ForwardGearRatios.Add(DefGearSetup.mRatios[i]);
-	}
-	GearSetup.RatioMultiplier =  DefGearSetup.mFinalRatio; 
+	TransmissionSetup.GearSwitchTime = DefGearSetup.mSwitchTime;
+	TransmissionSetup.ReverseGearRatio = DefGearSetup.mRatios[PxVehicleGearsData::eREVERSE];
+	TransmissionSetup.FinalRatio = DefGearSetup.mFinalRatio;
 
 	PxVehicleAutoBoxData DefAutoBoxSetup;
-	for (uint32 i = PxVehicleGearsData::eFIRST; i < PxVehicleGearsData::eGEARSRATIO_COUNT; i++)
+	TransmissionSetup.NeutralGearUpRatio = DefAutoBoxSetup.mUpRatios[PxVehicleGearsData::eNEUTRAL];
+	TransmissionSetup.GearAutoBoxLatency = DefAutoBoxSetup.getLatency();
+	TransmissionSetup.bUseGearAutoBox = true;
+
+	for (uint32 i = PxVehicleGearsData::eFIRST; i < DefGearSetup.mNbRatios; i++)
 	{
-		FGearUpDownRatio RatioData;
-		RatioData.UpRatio = DefAutoBoxSetup.mUpRatios[i];
-		RatioData.DownRatio = DefAutoBoxSetup.mUpRatios[i];
-		AutoBoxSetup.ForwardGearAutoBox.Add(RatioData);
+		FVehicleGearData GearData;
+		GearData.DownRatio = DefAutoBoxSetup.mDownRatios[i];
+		GearData.UpRatio = DefAutoBoxSetup.mUpRatios[i];
+		GearData.Ratio = DefGearSetup.mRatios[i];
+		TransmissionSetup.ForwardGears.Add(GearData);
 	}
-	AutoBoxSetup.NeutralGearUpRatio = DefAutoBoxSetup.mUpRatios[PxVehicleGearsData::eNEUTRAL];
-	AutoBoxSetup.GearAutoBoxLatency = DefAutoBoxSetup.getLatency();
-	bUseGearAutoBox = true;
-
+	
 	MaxSteeringSpeed = 100.f; // editable in vehicle blueprint
 	SteeringCurve.AddZeroed(4);
 	SteeringCurve[0].InVal = 0.0f;
@@ -76,6 +72,28 @@ UWheeledVehicleMovementComponent4W::UWheeledVehicleMovementComponent4W(const cla
 #endif // WITH_PHYSX
 }
 
+void UWheeledVehicleMovementComponent4W::PostEditChangeProperty(struct FPropertyChangedEvent& PropertyChangedEvent)
+{
+	const FName PropertyName = PropertyChangedEvent.Property ? PropertyChangedEvent.Property->GetFName() : NAME_None;
+
+	if (PropertyName == TEXT("DownRatio"))
+	{
+		for (int32 GearIdx = 0; GearIdx < TransmissionSetup.ForwardGears.Num(); ++GearIdx)
+		{
+			FVehicleGearData & GearData = TransmissionSetup.ForwardGears[GearIdx];
+			GearData.DownRatio = FMath::Min(GearData.DownRatio, GearData.UpRatio);
+		}
+	}
+	else if (PropertyName == TEXT("UpRatio"))
+	{
+		for (int32 GearIdx = 0; GearIdx < TransmissionSetup.ForwardGears.Num(); ++GearIdx)
+		{
+			FVehicleGearData & GearData = TransmissionSetup.ForwardGears[GearIdx];
+			GearData.UpRatio = FMath::Max(GearData.DownRatio, GearData.UpRatio);
+		}
+	}
+}
+
 #if WITH_PHYSX
 
 static void GetVehicleDifferential4WSetup(const FVehicleDifferential4WData& Setup, PxVehicleDifferential4WData& PxSetup)
@@ -122,24 +140,25 @@ static void GetVehicleEngineSetup(const FVehicleEngineData& Setup, PxVehicleEngi
 	PxSetup.mDampingRateZeroThrottleClutchDisengaged = Setup.DampingRateZeroThrottleClutchDisengaged;
 }
 
-static void GetVehicleGearSetup(const FVehicleGearData& Setup, PxVehicleGearsData& PxSetup)
+static void GetVehicleGearSetup(const FVehicleTransmissionData& Setup, PxVehicleGearsData& PxSetup)
 {
 	PxSetup.mSwitchTime = Setup.GearSwitchTime;
 	PxSetup.mRatios[PxVehicleGearsData::eREVERSE] = Setup.ReverseGearRatio;
-	for (int32 i = 0; i < Setup.ForwardGearRatios.Num(); i++)
+	for (int32 i = 0; i < Setup.ForwardGears.Num(); i++)
 	{
-		PxSetup.mRatios[i + PxVehicleGearsData::eFIRST] = Setup.ForwardGearRatios[i];
+		PxSetup.mRatios[i + PxVehicleGearsData::eFIRST] = Setup.ForwardGears[i].Ratio;
 	}
-	PxSetup.mFinalRatio = Setup.RatioMultiplier;
+	PxSetup.mFinalRatio = Setup.FinalRatio;
+	PxSetup.mNbRatios = Setup.ForwardGears.Num() + PxVehicleGearsData::eFIRST;
 }
 
-static void GetVehicleAutoBoxSetup(const FVehicleAutoBoxData& Setup, PxVehicleAutoBoxData& PxSetup)
+static void GetVehicleAutoBoxSetup(const FVehicleTransmissionData& Setup, PxVehicleAutoBoxData& PxSetup)
 {
-	for (uint32 i = PxVehicleGearsData::eFIRST; i < PxVehicleGearsData::eGEARSRATIO_COUNT; i++)
+	for (int32 i = 0; i < Setup.ForwardGears.Num(); i++)
 	{
-		const FGearUpDownRatio& RatioData = Setup.ForwardGearAutoBox[i - PxVehicleGearsData::eFIRST];
-		PxSetup.mUpRatios[i] = RatioData.UpRatio;
-		PxSetup.mDownRatios[i] = RatioData.DownRatio;
+		const FVehicleGearData& GearData = Setup.ForwardGears[i];
+		PxSetup.mUpRatios[i] = GearData.UpRatio;
+		PxSetup.mDownRatios[i] = GearData.DownRatio;
 	}
 	PxSetup.mUpRatios[PxVehicleGearsData::eNEUTRAL] = Setup.NeutralGearUpRatio;
 	PxSetup.setLatency(Setup.GearAutoBoxLatency);
@@ -157,7 +176,7 @@ void SetupDriveHelper(const UWheeledVehicleMovementComponent4W* VehicleData, con
 	DriveData.setEngineData(EngineSetup);
 
 	PxVehicleClutchData ClutchSetup;
-	ClutchSetup.mStrength = VehicleData->ClutchStrength;
+	ClutchSetup.mStrength = VehicleData->TransmissionSetup.ClutchStrength;
 	DriveData.setClutchData(ClutchSetup);
 
 	FVector WheelCentreOffsets[4];
@@ -174,11 +193,11 @@ void SetupDriveHelper(const UWheeledVehicleMovementComponent4W* VehicleData, con
 	DriveData.setAckermannGeometryData(AckermannSetup);
 
 	PxVehicleGearsData GearSetup;
-	GetVehicleGearSetup(VehicleData->GearSetup, GearSetup);	
+	GetVehicleGearSetup(VehicleData->TransmissionSetup, GearSetup);	
 	DriveData.setGearsData(GearSetup);
 
 	PxVehicleAutoBoxData AutoBoxSetup;
-	GetVehicleAutoBoxSetup(VehicleData->AutoBoxSetup, AutoBoxSetup);
+	GetVehicleAutoBoxSetup(VehicleData->TransmissionSetup, AutoBoxSetup);
 	DriveData.setAutoBoxData(AutoBoxSetup);
 }
 
@@ -290,44 +309,39 @@ void UWheeledVehicleMovementComponent4W::UpdateDifferentialSetup(const FVehicleD
 #endif
 }
 
-void UWheeledVehicleMovementComponent4W::UpdateGearSetup(const FVehicleGearData& NewGearSetup)
+void UWheeledVehicleMovementComponent4W::UpdateTransmissionSetup(const FVehicleTransmissionData& NewTransmissionSetup)
 {
 #if WITH_PHYSX
 	if (PVehicleDrive)
 	{
 		PxVehicleGearsData GearData;
-		GetVehicleGearSetup(NewGearSetup, GearData);
-
-		PxVehicleDrive4W* PVehicleDrive4W = (PxVehicleDrive4W*)PVehicleDrive;
-		PVehicleDrive4W->mDriveSimData.setGearsData(GearData);
-	}
-#endif
-}
+		GetVehicleGearSetup(NewTransmissionSetup, GearData);
 
-void UWheeledVehicleMovementComponent4W::UpdateAutoBoxSetup(const FVehicleAutoBoxData& NewAutoBoxSetup)
-{
-#if WITH_PHYSX
-	if (PVehicleDrive)
-	{
 		PxVehicleAutoBoxData AutoBoxData;
-		GetVehicleAutoBoxSetup(NewAutoBoxSetup, AutoBoxData);
+		GetVehicleAutoBoxSetup(NewTransmissionSetup, AutoBoxData);
 
 		PxVehicleDrive4W* PVehicleDrive4W = (PxVehicleDrive4W*)PVehicleDrive;
+		PVehicleDrive4W->mDriveSimData.setGearsData(GearData);
 		PVehicleDrive4W->mDriveSimData.setAutoBoxData(AutoBoxData);
 	}
 #endif
 }
-
 void UWheeledVehicleMovementComponent4W::Serialize(FArchive & Ar)
 {
 	Super::Serialize(Ar);
+#if WITH_PHYSX
 	if (Ar.IsLoading() && Ar.UE4Ver() < VER_UE4_VEHICLES_UNIT_CHANGE)
 	{
+		PxVehicleEngineData DefEngineData;
+		float DefaultTorque = DefEngineData.mPeakTorque;
+		float DefaultRPM = OmegaToRPM(DefEngineData.mMaxOmega);
+		
 		//we need to convert from old units to new. This backwards compatable code fails in the rare case that they were using very strange values that are the new defaults in the correct units.
-		EngineSetup.PeakTorque = EngineSetup.PeakTorque != 1000.f ? Cm2ToM2(EngineSetup.PeakTorque) : 1000.f;	//need to convert kg cm^2/s^2 into Nm
-		EngineSetup.MaxRPM = EngineSetup.MaxRPM != 7000.f ? OmegaToRPM(EngineSetup.MaxRPM) : 7000.f;	//need to convert from rad/s to RPM
+		EngineSetup.PeakTorque = EngineSetup.PeakTorque != DefaultTorque ? Cm2ToM2(EngineSetup.PeakTorque) : DefaultTorque;	//need to convert kg cm^2/s^2 into Nm
+		EngineSetup.MaxRPM = EngineSetup.MaxRPM != DefaultRPM ? OmegaToRPM(EngineSetup.MaxRPM) : DefaultRPM;	//need to convert from rad/s to RPM
 		MaxSteeringSpeed = MaxSteeringSpeed != 100.f ? CmSToKmH(MaxSteeringSpeed) : 100.f;	//we now store as km/h instead of cm/s
 	}
+#endif
 }
 
 void UWheeledVehicleMovementComponent4W::ComputeConstants()
