fix build

diff --git a/Engine/Source/Runtime/Engine/Classes/Vehicles/WheeledVehicleMovementComponent4W.h b/Engine/Source/Runtime/Engine/Classes/Vehicles/WheeledVehicleMovementComponent4W.h
--- a/Engine/Source/Runtime/Engine/Classes/Vehicles/WheeledVehicleMovementComponent4W.h
+++ b/Engine/Source/Runtime/Engine/Classes/Vehicles/WheeledVehicleMovementComponent4W.h
@@ -171,7 +171,9 @@ class ENGINE_API UWheeledVehicleMovementComponent4W : public UWheeledVehicleMove
 
 	virtual void Serialize(FArchive & Ar) OVERRIDE;
 	virtual void ComputeConstants() OVERRIDE;
+#if WITH_EDITOR
 	virtual void PostEditChangeProperty(struct FPropertyChangedEvent& PropertyChangedEvent) OVERRIDE;
+#endif
 
 protected:
 
diff --git a/Engine/Source/Runtime/Engine/Private/Vehicles/WheeledVehicleMovementComponent4W.cpp b/Engine/Source/Runtime/Engine/Private/Vehicles/WheeledVehicleMovementComponent4W.cpp
--- a/Engine/Source/Runtime/Engine/Private/Vehicles/WheeledVehicleMovementComponent4W.cpp
+++ b/Engine/Source/Runtime/Engine/Private/Vehicles/WheeledVehicleMovementComponent4W.cpp
@@ -72,6 +72,7 @@ UWheeledVehicleMovementComponent4W::UWheeledVehicleMovementComponent4W(const cla
 #endif // WITH_PHYSX
 }
 
+#if WITH_EDITOR
 void UWheeledVehicleMovementComponent4W::PostEditChangeProperty(struct FPropertyChangedEvent& PropertyChangedEvent)
 {
 	const FName PropertyName = PropertyChangedEvent.Property ? PropertyChangedEvent.Property->GetFName() : NAME_None;
@@ -93,6 +94,7 @@ void UWheeledVehicleMovementComponent4W::PostEditChangeProperty(struct FProperty
 		}
 	}
 }
+#endif
 
 #if WITH_PHYSX
 
