#ue4 Fix for particle rendering in stereoscopic views.  Beams now properly populate all views in stereo rendering modes

diff --git a/Engine/Source/Runtime/Engine/Private/Particles/ParticleSystemRender.cpp b/Engine/Source/Runtime/Engine/Private/Particles/ParticleSystemRender.cpp
--- a/Engine/Source/Runtime/Engine/Private/Particles/ParticleSystemRender.cpp
+++ b/Engine/Source/Runtime/Engine/Private/Particles/ParticleSystemRender.cpp
@@ -2751,7 +2751,7 @@ void FDynamicBeam2EmitterData::PreRenderView(FParticleSystemSceneProxy* Proxy, c
 	// Only need to do this once per-view
 	if (LastFramePreRendered < FrameNumber)
 	{
-		bool bOnlyOneView = !GIsEditor && ((GEngine && GEngine->GameViewport && (GEngine->GameViewport->GetCurrentSplitscreenConfiguration() == ESplitScreenType::None)) ? true : false);
+		bool bOnlyOneView = !GIsEditor && ((GEngine && GEngine->GameViewport && (GEngine->GameViewport->GetCurrentSplitscreenConfiguration() == ESplitScreenType::None) && !GEngine->IsStereoscopic3D()) ? true : false);
 
 		BuildViewFillDataAndSubmit(Proxy,ViewFamily,VisibilityMap,bOnlyOneView,Source.VertexCount,sizeof(FParticleBeamTrailVertex),0);
 
