#ue4
- fixed undefined log category

diff --git a/Engine/Source/Runtime/Engine/Private/GameSession.cpp b/Engine/Source/Runtime/Engine/Private/GameSession.cpp
--- a/Engine/Source/Runtime/Engine/Private/GameSession.cpp
+++ b/Engine/Source/Runtime/Engine/Private/GameSession.cpp
@@ -93,7 +93,7 @@ void AGameSession::OnLoginComplete(int32 LocalUserNum, bool bWasSuccessful, cons
 		}
 		else
 		{
-			UE_LOG(LogGameMode, Warning, TEXT("Autologin attempt failed, unable to register server!"));
+			UE_LOG(LogGameSession, Warning, TEXT("Autologin attempt failed, unable to register server!"));
 		}
 	}
 }
