#TTP 331766 - PERSONA: Derived blueprints do not show inherited properties

#proj UE4

#summary Added code to clear erroneously set flag CPF_DisableEditOnTemplate

diff --git a/Engine/Source/Runtime/Core/Public/UObject/ObjectVersion.h b/Engine/Source/Runtime/Core/Public/UObject/ObjectVersion.h
--- a/Engine/Source/Runtime/Core/Public/UObject/ObjectVersion.h
+++ b/Engine/Source/Runtime/Core/Public/UObject/ObjectVersion.h
@@ -529,7 +529,10 @@ enum EUnrealEngineObjectUE4Version
 	VER_UE4_ANIMATION_REMOVE_NANS,
 	// Change skeleton preview attached assets property type
 	VER_UE4_SKELETON_ASSET_PROPERTY_TYPE_CHANGE,
- 
+ 	// Fix some blueprint variables that have the CPF_DisableEditOnTemplate flag set
+	// when they shouldn't
+	VER_UE4_FIX_BLUEPRINT_VARIABLE_FLAGS,
+
 	// -----<new versions can be added before this line>-------------------------------------------------
 	// - this needs to be the last line (see note below)
 	VER_UE4_AUTOMATIC_VERSION_PLUS_ONE,
diff --git a/Engine/Source/Runtime/Engine/Private/Blueprint.cpp b/Engine/Source/Runtime/Engine/Private/Blueprint.cpp
--- a/Engine/Source/Runtime/Engine/Private/Blueprint.cpp
+++ b/Engine/Source/Runtime/Engine/Private/Blueprint.cpp
@@ -191,7 +191,7 @@ void UBlueprint::Serialize(FArchive& Ar)
 		}
 	}
 
-	if (Ar.UE4Ver() < VER_UE4_BP_ACTOR_VARIABLE_DEFAULT_PREVENTING)
+	if (Ar.UE4Ver() < VER_UE4_FIX_BLUEPRINT_VARIABLE_FLAGS)
 	{
 		// Actor variables can't have default values (because Blueprint templates are library elements that can 
 		// bridge multiple levels and different levels might not have the actor that the default is referencing).
@@ -200,18 +200,27 @@ void UBlueprint::Serialize(FArchive& Ar)
 			FBPVariableDescription& Variable = NewVariables[i];
 
 			const FEdGraphPinType& VarType = Variable.VarType;
-			if (!VarType.PinSubCategoryObject.IsValid()) // ignore variables that don't have associated objects
+
+			bool bDisableEditOnTemplate = false;
+			if (VarType.PinSubCategoryObject.IsValid()) // ignore variables that don't have associated objects
 			{
-				continue;
+				const UClass* ClassObject = Cast<UClass>(VarType.PinSubCategoryObject.Get());
+				// if the object type is an actor...
+				if ((ClassObject != NULL) && ClassObject->IsChildOf(AActor::StaticClass()))
+				{
+					// hide the default value field
+					bDisableEditOnTemplate = true;
+				}
 			}
 
-			const UClass* ClassObject = Cast<UClass>(VarType.PinSubCategoryObject.Get());
-			// if the object type is an actor...
-			if ((ClassObject != NULL) && ClassObject->IsChildOf(AActor::StaticClass()))
+			if(bDisableEditOnTemplate)
 			{
-				// hide the default value field
 				Variable.PropertyFlags |= CPF_DisableEditOnTemplate;
 			}
+			else
+			{
+				Variable.PropertyFlags &= ~CPF_DisableEditOnTemplate;
+			}
 		}
 	}
 #endif // WITH_EDITORONLY_DATA
