Corrected Multi-Line trace blueprint node documentation errors

Fixes TTP#316609 - BP: DOCUMENTATION: Multiple Tooltip typos in Multi Line Trace nodes

diff --git a/Engine/Source/Runtime/Engine/Classes/Kismet/KismetSystemLibrary.h b/Engine/Source/Runtime/Engine/Classes/Kismet/KismetSystemLibrary.h
--- a/Engine/Source/Runtime/Engine/Classes/Kismet/KismetSystemLibrary.h
+++ b/Engine/Source/Runtime/Engine/Classes/Kismet/KismetSystemLibrary.h
@@ -530,13 +530,12 @@ class UKismetSystemLibrary : public UBlueprintFunctionLibrary
 	
 	/**
 	 * Does a collision trace along the given line and returns all hits encountered up to and including the first blocking hit.
-	 * along the given line and returns the first blocking hit encountered.
-	 * This trace finds the objects that RESPONDS to the given TraceChannel
+	 * This trace finds the objects that RESPOND to the given TraceChannel
 	 * 
 	 * @param WorldContext	World context
 	 * @param Start			Start of line segment.
 	 * @param End			End of line segment.
-	 * @param TraceChannel	
+	 * @param TraceChannel	The channel to trace
 	 * @param bTraceComplex	True to test against complex collision, false to test against simplified collision.
 	 * @param OutHit		Properties of the trace hit.
 	 * @return				True if there was a hit, false otherwise.
@@ -561,7 +560,7 @@ class UKismetSystemLibrary : public UBlueprintFunctionLibrary
 
 	/**
 	 * Sweeps a sphere along the given line and returns all hits encountered up to and including the first blocking hit.
-	 * This trace finds the objects that RESPONDS to the given TraceChannel
+	 * This trace finds the objects that RESPOND to the given TraceChannel
 	 * 
 	 * @param WorldContext	World context
 	 * @param Start			Start of line segment.
@@ -577,7 +576,7 @@ class UKismetSystemLibrary : public UBlueprintFunctionLibrary
 	
 	/**
 	 * Sweeps a capsule along the given line and returns the first blocking hit encountered.
-	 * This trace finds the objects that RESPONDS to the given TraceChannel
+	 * This trace finds the objects that RESPOND to the given TraceChannel
 	 * 
 	 * @param WorldContext	World context
 	 * @param Start			Start of line segment.
@@ -594,7 +593,7 @@ class UKismetSystemLibrary : public UBlueprintFunctionLibrary
 
 	/**
 	 * Sweeps a capsule along the given line and returns all hits encountered up to and including the first blocking hit.
-	 * This trace finds the objects that RESPONDS to the given TraceChannel
+	 * This trace finds the objects that RESPOND to the given TraceChannel
 	 * 
 	 * @param WorldContext	World context
 	 * @param Start			Start of line segment.
@@ -611,7 +610,7 @@ class UKismetSystemLibrary : public UBlueprintFunctionLibrary
 
 	/**
 	 * Does a collision trace along the given line and returns the first hit encountered.
-	 * This finds objects with all Movement Channel of ObjectTypes. 
+	 * This only finds objects that are of a type specified by ObjectTypes.
 	 * 
 	 * @param WorldContext	World context
 	 * @param Start			Start of line segment.
@@ -626,8 +625,7 @@ class UKismetSystemLibrary : public UBlueprintFunctionLibrary
 	
 	/**
 	 * Does a collision trace along the given line and returns all hits encountered.
-	 * along the given line and returns the first blocking hit encountered.
-	 * This finds objects with all Movement Channel of ObjectTypes. 
+	 * This only finds objects that are of a type specified by ObjectTypes.
 	 * 
 	 * @param WorldContext	World context
 	 * @param Start			Start of line segment.
@@ -642,7 +640,7 @@ class UKismetSystemLibrary : public UBlueprintFunctionLibrary
 
 	/**
 	 * Sweeps a sphere along the given line and returns the first hit encountered.
-	 * This finds objects with all Movement Channel of ObjectTypes. 
+	 * This only finds objects that are of a type specified by ObjectTypes.
 	 * 
 	 * @param Start			Start of line segment.
 	 * @param End			End of line segment.
@@ -657,7 +655,7 @@ class UKismetSystemLibrary : public UBlueprintFunctionLibrary
 
 	/**
 	 * Sweeps a sphere along the given line and returns all hits encountered.
-	 * This finds objects with all Movement Channel of ObjectTypes. 
+	 * This only finds objects that are of a type specified by ObjectTypes.
 	 * 
 	 * @param WorldContext	World context
 	 * @param Start			Start of line segment.
@@ -673,7 +671,7 @@ class UKismetSystemLibrary : public UBlueprintFunctionLibrary
 	
 	/**
 	 * Sweeps a capsule along the given line and returns the first hit encountered.
-	 * This finds objects with all Movement Channel of ObjectTypes. 
+	 * This only finds objects that are of a type specified by ObjectTypes.
 	 * 
 	 * @param WorldContext	World context
 	 * @param Start			Start of line segment.
@@ -690,7 +688,7 @@ class UKismetSystemLibrary : public UBlueprintFunctionLibrary
 
 	/**
 	 * Sweeps a capsule along the given line and returns all hits encountered.
-	 * This finds objects with all Movement Channel of ObjectTypes. 
+	 * This only finds objects that are of a type specified by ObjectTypes.
 	 * 
 	 * @param WorldContext	World context
 	 * @param Start			Start of line segment.
@@ -707,7 +705,7 @@ class UKismetSystemLibrary : public UBlueprintFunctionLibrary
 
 	/**
 	 * Does a collision trace along the given line and returns the first blocking hit encountered.
-	 * This trace finds the objects that RESPONDS to the given TraceChannel
+	 * This trace finds the objects that RESPOND to the given TraceChannel
 	 * 
 	 * @param WorldContext	World context
 	 * @param Start			Start of line segment.
@@ -722,8 +720,7 @@ class UKismetSystemLibrary : public UBlueprintFunctionLibrary
 	
 	/**
 	 * Does a collision trace along the given line and returns all hits encountered up to and including the first blocking hit.
-	 * along the given line and returns the first blocking hit encountered.
-	 * This trace finds the objects that RESPONDS to the given TraceChannel
+	 * This trace finds the objects that RESPOND to the given TraceChannel
 	 * 
 	 * @param WorldContext	World context
 	 * @param Start			Start of line segment.
@@ -803,7 +800,7 @@ class UKismetSystemLibrary : public UBlueprintFunctionLibrary
 
 	/**
 	 * Does a collision trace along the given line and returns the first hit encountered.
-	 * This finds objects with all Movement Channel of ObjectsToTrace. 
+	 * This finds objects belonging to the channels specified in the ObjectsToTrace input.
 	 * 
 	 * @param WorldContext	World context
 	 * @param Start			Start of line segment.
@@ -818,8 +815,7 @@ class UKismetSystemLibrary : public UBlueprintFunctionLibrary
 	
 	/**
 	 * Does a collision trace along the given line and returns all hits encountered.
-	 * along the given line and returns the first blocking hit encountered.
-	 * This finds objects with all Movement Channel of ObjectsToTrace. 
+	 * This finds objects belonging to the channels specified in the ObjectsToTrace input.
 	 * 
 	 * @param WorldContext	World context
 	 * @param Start			Start of line segment.
@@ -834,7 +830,7 @@ class UKismetSystemLibrary : public UBlueprintFunctionLibrary
 
 	/**
 	 * Sweeps a sphere along the given line and returns the first hit encountered.
-	 * This finds objects with all Movement Channel of ObjectsToTrace. 
+	 * This finds objects belonging to the channels specified in the ObjectsToTrace input.
 	 * 
 	 * @param Start			Start of line segment.
 	 * @param End			End of line segment.
@@ -849,7 +845,7 @@ class UKismetSystemLibrary : public UBlueprintFunctionLibrary
 
 	/**
 	 * Sweeps a sphere along the given line and returns all hits encountered.
-	 * This finds objects with all Movement Channel of ObjectsToTrace. 
+	 * This finds objects belonging to the channels specified in the ObjectsToTrace input.
 	 * 
 	 * @param WorldContext	World context
 	 * @param Start			Start of line segment.
@@ -865,7 +861,7 @@ class UKismetSystemLibrary : public UBlueprintFunctionLibrary
 	
 	/**
 	 * Sweeps a capsule along the given line and returns the first hit encountered.
-	 * This finds objects with all Movement Channel of ObjectsToTrace. 
+	 * This finds objects belonging to the channels specified in the ObjectsToTrace input.
 	 * 
 	 * @param WorldContext	World context
 	 * @param Start			Start of line segment.
@@ -882,7 +878,7 @@ class UKismetSystemLibrary : public UBlueprintFunctionLibrary
 
 	/**
 	 * Sweeps a capsule along the given line and returns all hits encountered.
-	 * This finds objects with all Movement Channel of ObjectsToTrace. 
+	 * This finds objects belonging to the channels specified in the ObjectsToTrace input.
 	 * 
 	 * @param WorldContext	World context
 	 * @param Start			Start of line segment.
