vehicle unit change part two

diff --git a/Engine/Source/Runtime/Core/Public/UObject/ObjectVersion.h b/Engine/Source/Runtime/Core/Public/UObject/ObjectVersion.h
--- a/Engine/Source/Runtime/Core/Public/UObject/ObjectVersion.h
+++ b/Engine/Source/Runtime/Core/Public/UObject/ObjectVersion.h
@@ -532,7 +532,9 @@ enum EUnrealEngineObjectUE4Version
  	// Fix some blueprint variables that have the CPF_DisableEditOnTemplate flag set
 	// when they shouldn't
 	VER_UE4_FIX_BLUEPRINT_VARIABLE_FLAGS,
-
+	// Vehicles use Nm for Torque instead of cm and RPM instead of rad/s part two (missed conversion for some variables
+	VER_UE4_VEHICLES_UNIT_CHANGE2,
+ 
 	// -----<new versions can be added before this line>-------------------------------------------------
 	// - this needs to be the last line (see note below)
 	VER_UE4_AUTOMATIC_VERSION_PLUS_ONE,
diff --git a/Engine/Source/Runtime/Engine/Classes/Vehicles/WheeledVehicleMovementComponent.h b/Engine/Source/Runtime/Engine/Classes/Vehicles/WheeledVehicleMovementComponent.h
--- a/Engine/Source/Runtime/Engine/Classes/Vehicles/WheeledVehicleMovementComponent.h
+++ b/Engine/Source/Runtime/Engine/Classes/Vehicles/WheeledVehicleMovementComponent.h
@@ -109,7 +109,7 @@ struct FWheelSetup
 	UPROPERTY(EditAnywhere, Category=WheelsSetup)
 	float MaxHandBrakeTorque;
 
-	// damping rate for this wheel
+	// damping rate for this wheel (Kgm^2/s)
 	UPROPERTY(EditAnywhere, Category=WheelsSetup)
 	float DampingRate;
 
@@ -124,7 +124,7 @@ struct FWheelSetup
 		, SteerAngle(70.0f)
 		, MaxBrakeTorque(1500.f)
 		, MaxHandBrakeTorque(3000.f)
-		, DampingRate(0.25f * 100.f * 100.f)
+		, DampingRate(0.25f)
 		, Mass(20.f)
 	{
 	}
diff --git a/Engine/Source/Runtime/Engine/Classes/Vehicles/WheeledVehicleMovementComponent4W.h b/Engine/Source/Runtime/Engine/Classes/Vehicles/WheeledVehicleMovementComponent4W.h
--- a/Engine/Source/Runtime/Engine/Classes/Vehicles/WheeledVehicleMovementComponent4W.h
+++ b/Engine/Source/Runtime/Engine/Classes/Vehicles/WheeledVehicleMovementComponent4W.h
@@ -67,19 +67,19 @@ struct FVehicleEngineData
 	UPROPERTY(EditAnywhere, Category=Setup)
 	float MaxRPM;
 
-	/** Moment of inertia of the engine around the axis of rotation. */
+	/** Moment of inertia of the engine around the axis of rotation (Kgm^2). */
 	UPROPERTY(EditAnywhere, Category = Setup)
 	float MOI;
 
-	/** Damping rate of engine when full throttle is applied */
+	/** Damping rate of engine when full throttle is applied (Kgm^2/s) */
 	UPROPERTY(EditAnywhere, Category=Setup, AdvancedDisplay)
 	float DampingRateFullThrottle;
 
-	/** Damping rate of engine in at zero throttle when the clutch is engaged */
+	/** Damping rate of engine in at zero throttle when the clutch is engaged (Kgm^2/s)*/
 	UPROPERTY(EditAnywhere, Category=Setup, AdvancedDisplay)
 	float DampingRateZeroThrottleClutchEngaged;
 
-	/** Damping rate of engine in at zero throttle when the clutch is disengaged (in neutral gear) */
+	/** Damping rate of engine in at zero throttle when the clutch is disengaged (in neutral gear) (Kgm^2/s)*/
 	UPROPERTY(EditAnywhere, Category=Setup, AdvancedDisplay)
 	float DampingRateZeroThrottleClutchDisengaged;
 };
@@ -135,7 +135,7 @@ struct FVehicleTransmissionData
 	UPROPERTY(EditAnywhere, AdvancedDisplay, Category = Setup, meta = (ClampMin = "0.0", UIMin = "0.0", ClampMax = "1.0", UIMax = "1.0"))
 	float NeutralGearUpRatio;
 
-	/** Strength of clutch */
+	/** Strength of clutch (Kgm^2/s)*/
 	UPROPERTY(EditAnywhere, Category = Setup, AdvancedDisplay)
 	float ClutchStrength;
 };
diff --git a/Engine/Source/Runtime/Engine/Private/Vehicles/WheeledVehicleMovementComponent.cpp b/Engine/Source/Runtime/Engine/Private/Vehicles/WheeledVehicleMovementComponent.cpp
--- a/Engine/Source/Runtime/Engine/Private/Vehicles/WheeledVehicleMovementComponent.cpp
+++ b/Engine/Source/Runtime/Engine/Private/Vehicles/WheeledVehicleMovementComponent.cpp
@@ -288,9 +288,7 @@ void UWheeledVehicleMovementComponent::SetupWheels( PxVehicleWheelsSimData* PWhe
 		PWheelData.mMaxBrakeTorque = M2ToCm2(WheelSetups[WheelIdx].MaxBrakeTorque);
 		PWheelData.mMaxHandBrakeTorque = Wheel->bAffectedByHandbrake ? M2ToCm2(WheelSetups[WheelIdx].MaxHandBrakeTorque) : 0.0f;
 
-//		PWheelData.mMass = 1.0f;
-//		PWheelData.mMOI = Wheel->Inertia * Mass; // Multiply by vehicle mass so that as vehicle mass shrinks, the torque and inertia can shrink as well
-		PWheelData.mDampingRate = WheelSetups[WheelIdx].DampingRate;
+		PWheelData.mDampingRate = M2ToCm2(WheelSetups[WheelIdx].DampingRate);
 		PWheelData.mMass = WheelSetups[WheelIdx].Mass;
 		PWheelData.mMOI = 0.5f * PWheelData.mMass * FMath::Square(PWheelData.mRadius);
 		
@@ -1311,4 +1309,15 @@ void UWheeledVehicleMovementComponent::Serialize(FArchive & Ar)
 		}
 
 	}
+
+	if (Ar.IsLoading() && Ar.UE4Ver() < VER_UE4_VEHICLES_UNIT_CHANGE2)
+	{
+		for (int32 WheelIdx = 0; WheelIdx < WheelSetups.Num(); ++WheelIdx)
+		{
+			//this backwards compatable code breaks if they were using some strange value that is exactly the new default in the new units
+			FWheelSetup & WheelSetup = WheelSetups[WheelIdx];
+			WheelSetup.DampingRate = WheelSetup.DampingRate != 0.25f ? Cm2ToM2(WheelSetup.DampingRate) : 0.25;
+		}
+
+	}
 }
diff --git a/Engine/Source/Runtime/Engine/Private/Vehicles/WheeledVehicleMovementComponent4W.cpp b/Engine/Source/Runtime/Engine/Private/Vehicles/WheeledVehicleMovementComponent4W.cpp
--- a/Engine/Source/Runtime/Engine/Private/Vehicles/WheeledVehicleMovementComponent4W.cpp
+++ b/Engine/Source/Runtime/Engine/Private/Vehicles/WheeledVehicleMovementComponent4W.cpp
@@ -21,15 +21,15 @@ UWheeledVehicleMovementComponent4W::UWheeledVehicleMovementComponent4W(const cla
 	DifferentialSetup.RearBias = DefDifferentialSetup.mRearBias;
 
 	PxVehicleEngineData DefEngineData;
-	EngineSetup.MOI = M2ToCm2(DefEngineData.mMOI);
+	EngineSetup.MOI = DefEngineData.mMOI;
 	EngineSetup.PeakTorque = DefEngineData.mPeakTorque;
 	EngineSetup.MaxRPM = OmegaToRPM(DefEngineData.mMaxOmega);
-	EngineSetup.DampingRateFullThrottle = M2ToCm2(DefEngineData.mDampingRateFullThrottle);
-	EngineSetup.DampingRateZeroThrottleClutchEngaged = M2ToCm2(DefEngineData.mDampingRateZeroThrottleClutchEngaged);
-	EngineSetup.DampingRateZeroThrottleClutchDisengaged = M2ToCm2(DefEngineData.mDampingRateZeroThrottleClutchDisengaged);
+	EngineSetup.DampingRateFullThrottle = DefEngineData.mDampingRateFullThrottle;
+	EngineSetup.DampingRateZeroThrottleClutchEngaged = DefEngineData.mDampingRateZeroThrottleClutchEngaged;
+	EngineSetup.DampingRateZeroThrottleClutchDisengaged = DefEngineData.mDampingRateZeroThrottleClutchDisengaged;
 
 	PxVehicleClutchData DefClutchData;
-	TransmissionSetup.ClutchStrength = M2ToCm2(DefClutchData.mStrength); // convert from m to cm scale
+	TransmissionSetup.ClutchStrength = DefClutchData.mStrength;
 
 	PxVehicleAckermannGeometryData DefAckermannSetup;
 	AckermannAccuracy = DefAckermannSetup.mAccuracy;
@@ -75,6 +75,7 @@ UWheeledVehicleMovementComponent4W::UWheeledVehicleMovementComponent4W(const cla
 #if WITH_EDITOR
 void UWheeledVehicleMovementComponent4W::PostEditChangeProperty(struct FPropertyChangedEvent& PropertyChangedEvent)
 {
+	Super::PostEditChangeProperty(PropertyChangedEvent);
 	const FName PropertyName = PropertyChangedEvent.Property ? PropertyChangedEvent.Property->GetFName() : NAME_None;
 
 	if (PropertyName == TEXT("DownRatio"))
@@ -134,12 +135,12 @@ static void GetVehicleDifferential4WSetup(const FVehicleDifferential4WData& Setu
 
 static void GetVehicleEngineSetup(const FVehicleEngineData& Setup, PxVehicleEngineData& PxSetup)
 {
-	PxSetup.mMOI = Setup.MOI; // convert from m to cm scale
+	PxSetup.mMOI = M2ToCm2(Setup.MOI);
 	PxSetup.mPeakTorque = M2ToCm2(Setup.PeakTorque);	//convert Nm to (kg cm^2/s^2)
 	PxSetup.mMaxOmega = RPMToOmega(Setup.MaxRPM);
-	PxSetup.mDampingRateFullThrottle = Setup.DampingRateFullThrottle;
-	PxSetup.mDampingRateZeroThrottleClutchEngaged = Setup.DampingRateZeroThrottleClutchEngaged;
-	PxSetup.mDampingRateZeroThrottleClutchDisengaged = Setup.DampingRateZeroThrottleClutchDisengaged;
+	PxSetup.mDampingRateFullThrottle = M2ToCm2(Setup.DampingRateFullThrottle);
+	PxSetup.mDampingRateZeroThrottleClutchEngaged = M2ToCm2(Setup.DampingRateZeroThrottleClutchEngaged);
+	PxSetup.mDampingRateZeroThrottleClutchDisengaged = M2ToCm2(Setup.DampingRateZeroThrottleClutchDisengaged);
 }
 
 static void GetVehicleGearSetup(const FVehicleTransmissionData& Setup, PxVehicleGearsData& PxSetup)
@@ -178,7 +179,7 @@ void SetupDriveHelper(const UWheeledVehicleMovementComponent4W* VehicleData, con
 	DriveData.setEngineData(EngineSetup);
 
 	PxVehicleClutchData ClutchSetup;
-	ClutchSetup.mStrength = VehicleData->TransmissionSetup.ClutchStrength;
+	ClutchSetup.mStrength = M2ToCm2(VehicleData->TransmissionSetup.ClutchStrength);
 	DriveData.setClutchData(ClutchSetup);
 
 	FVector WheelCentreOffsets[4];
@@ -328,6 +329,15 @@ void UWheeledVehicleMovementComponent4W::UpdateTransmissionSetup(const FVehicleT
 	}
 #endif
 }
+
+void BackwardsConvertCm2ToM2(float & val, float defaultValue)
+{
+	if (val != defaultValue)
+	{
+		val = Cm2ToM2(val);
+	}
+}
+
 void UWheeledVehicleMovementComponent4W::Serialize(FArchive & Ar)
 {
 	Super::Serialize(Ar);
@@ -335,14 +345,27 @@ void UWheeledVehicleMovementComponent4W::Serialize(FArchive & Ar)
 	if (Ar.IsLoading() && Ar.UE4Ver() < VER_UE4_VEHICLES_UNIT_CHANGE)
 	{
 		PxVehicleEngineData DefEngineData;
-		float DefaultTorque = DefEngineData.mPeakTorque;
 		float DefaultRPM = OmegaToRPM(DefEngineData.mMaxOmega);
 		
 		//we need to convert from old units to new. This backwards compatable code fails in the rare case that they were using very strange values that are the new defaults in the correct units.
-		EngineSetup.PeakTorque = EngineSetup.PeakTorque != DefaultTorque ? Cm2ToM2(EngineSetup.PeakTorque) : DefaultTorque;	//need to convert kg cm^2/s^2 into Nm
+		BackwardsConvertCm2ToM2(EngineSetup.PeakTorque, DefEngineData.mPeakTorque);
 		EngineSetup.MaxRPM = EngineSetup.MaxRPM != DefaultRPM ? OmegaToRPM(EngineSetup.MaxRPM) : DefaultRPM;	//need to convert from rad/s to RPM
 		MaxSteeringSpeed = MaxSteeringSpeed != 100.f ? CmSToKmH(MaxSteeringSpeed) : 100.f;	//we now store as km/h instead of cm/s
 	}
+
+	if (Ar.IsLoading() && Ar.UE4Ver() < VER_UE4_VEHICLES_UNIT_CHANGE2)
+	{
+		PxVehicleEngineData DefEngineData;
+		PxVehicleClutchData DefClutchData;
+		float DefaultRPM = OmegaToRPM(DefEngineData.mMaxOmega);
+
+		//we need to convert from old units to new. This backwards compatable code fails in the rare case that they were using very strange values that are the new defaults in the correct units.
+		BackwardsConvertCm2ToM2(EngineSetup.DampingRateFullThrottle, DefEngineData.mDampingRateFullThrottle);
+		BackwardsConvertCm2ToM2(EngineSetup.DampingRateZeroThrottleClutchDisengaged, DefEngineData.mDampingRateZeroThrottleClutchDisengaged);
+		BackwardsConvertCm2ToM2(EngineSetup.DampingRateZeroThrottleClutchEngaged, DefEngineData.mDampingRateZeroThrottleClutchEngaged);
+		BackwardsConvertCm2ToM2(EngineSetup.MOI, DefEngineData.mMOI);
+		BackwardsConvertCm2ToM2(TransmissionSetup.ClutchStrength, DefClutchData.mStrength);
+	}
 #endif
 }
 
