fixed misspelled method name

diff --git a/Engine/Source/Runtime/Engine/Classes/Components/SceneComponent.h b/Engine/Source/Runtime/Engine/Classes/Components/SceneComponent.h
--- a/Engine/Source/Runtime/Engine/Classes/Components/SceneComponent.h
+++ b/Engine/Source/Runtime/Engine/Classes/Components/SceneComponent.h
@@ -540,11 +540,11 @@ public:
 	/** Get the current component-to-world transform for this component */
 	FORCEINLINE FTransform GetComponentToWorld() const 
 	{ 
-		return ComponentToWorld; // TODO: probably deprecate this in favor of GetComponenTransform
+		return ComponentToWorld; // TODO: probably deprecate this in favor of GetComponentTransform
 	}
 
 	/** Get the current component-to-world transform for this component */
-	FORCEINLINE FTransform GetComponenTransform() const
+	FORCEINLINE FTransform GetComponentTransform() const
 	{
 		return ComponentToWorld;
 	}
