Move GetVehicleMovementComponent and vehicle DisplayDebug function up into WheeledVehicle
Add 'Vehicle' collision profile name to UCollisionProfile
Default UVehicleWheel::TireType to DefaultTireType
Default FWheelSetup::WheelClass to UVehicleWheel
Actually call FWheelSetup constructor when allocating wheels array, so you get the defaults
Rework Vehicle C++ template to not use any BPs (all settings done in C++)
Add TargetOffset option to SpringArm

diff --git a/Engine/Source/Runtime/Engine/Classes/Engine/CollisionProfile.h b/Engine/Source/Runtime/Engine/Classes/Engine/CollisionProfile.h
--- a/Engine/Source/Runtime/Engine/Classes/Engine/CollisionProfile.h
+++ b/Engine/Source/Runtime/Engine/Classes/Engine/CollisionProfile.h
@@ -137,6 +137,7 @@ public:
 	static ENGINE_API const FName PhysicsActor_ProfileName;
 	static ENGINE_API const FName BlockAllDynamic_ProfileName;
 	static ENGINE_API const FName Pawn_ProfileName;
+	static ENGINE_API const FName Vehicle_ProfileName;
 
 	/** Accessor and initializer **/
 	ENGINE_API static UCollisionProfile * Get();
diff --git a/Engine/Source/Runtime/Engine/Classes/GameFramework/SpringArmComponent.h b/Engine/Source/Runtime/Engine/Classes/GameFramework/SpringArmComponent.h
--- a/Engine/Source/Runtime/Engine/Classes/GameFramework/SpringArmComponent.h
+++ b/Engine/Source/Runtime/Engine/Classes/GameFramework/SpringArmComponent.h
@@ -24,6 +24,10 @@ class ENGINE_API USpringArmComponent : public USceneComponent
 	UPROPERTY(EditAnywhere, BlueprintReadWrite, Category=Camera)
 	FVector SocketOffset;
 
+	/** Offset at start of spring, applied in world space. Use this if you want a world-space offset from the parent component instead of the usual relative-space offset. */
+	UPROPERTY(EditAnywhere, BlueprintReadWrite, Category=Camera)
+	FVector TargetOffset;
+
 	/** How big should the query probe sphere be (in unreal units) */
 	UPROPERTY(EditAnywhere, BlueprintReadWrite, Category=CameraCollision, meta=(editcondition="bDoCollisionTest"))
 	float ProbeSize;
diff --git a/Engine/Source/Runtime/Engine/Classes/GameFramework/WheeledVehicle.h b/Engine/Source/Runtime/Engine/Classes/GameFramework/WheeledVehicle.h
--- a/Engine/Source/Runtime/Engine/Classes/GameFramework/WheeledVehicle.h
+++ b/Engine/Source/Runtime/Engine/Classes/GameFramework/WheeledVehicle.h
@@ -21,4 +21,14 @@ class ENGINE_API AWheeledVehicle : public APawn
 
 	/** Name of the VehicleMovement. Use this name if you want to use a different class (with PCIP.SetDefaultSubobjectClass). */
 	static FName VehicleMovementComponentName;
+
+	/** Util to get the wheeled vehicle movement component */
+	class UWheeledVehicleMovementComponent* GetVehicleMovementComponent() const 
+	{ 
+		return VehicleMovement; 
+	}
+
+	// Begin AActor interface
+	virtual void DisplayDebug(class UCanvas* Canvas, const FDebugDisplayInfo& DebugDisplay, float& YL, float& YPos) OVERRIDE;
+	// End Actor interface
 };
diff --git a/Engine/Source/Runtime/Engine/Classes/Vehicles/WheeledVehicleMovementComponent.h b/Engine/Source/Runtime/Engine/Classes/Vehicles/WheeledVehicleMovementComponent.h
--- a/Engine/Source/Runtime/Engine/Classes/Vehicles/WheeledVehicleMovementComponent.h
+++ b/Engine/Source/Runtime/Engine/Classes/Vehicles/WheeledVehicleMovementComponent.h
@@ -2,6 +2,7 @@
 
 #pragma once
 
+#include "VehicleWheel.h"
 #include "WheeledVehicleMovementComponent.generated.h"
 
 #if WITH_PHYSX
@@ -86,7 +87,7 @@ struct FWheelSetup
 
 	// The wheel class to use
 	UPROPERTY(EditAnywhere, Category=WheelSetup)
-	TSubclassOf<class UVehicleWheel> WheelClass;
+	TSubclassOf<UVehicleWheel> WheelClass;
 
 	// Bone name on mesh to create wheel at
 	UPROPERTY(EditAnywhere, Category=WheelsSetup)
@@ -118,7 +119,7 @@ struct FWheelSetup
 	float Mass;
 
 	FWheelSetup()
-		: WheelClass(NULL)
+		: WheelClass(UVehicleWheel::StaticClass())
 		, BoneName(NAME_None)
 		, AdditionalOffset(0.0f)
 		, SteerAngle(70.0f)
diff --git a/Engine/Source/Runtime/Engine/Private/Collision/CollisionProfile.cpp b/Engine/Source/Runtime/Engine/Private/Collision/CollisionProfile.cpp
--- a/Engine/Source/Runtime/Engine/Private/Collision/CollisionProfile.cpp
+++ b/Engine/Source/Runtime/Engine/Private/Collision/CollisionProfile.cpp
@@ -36,6 +36,7 @@ ENGINE_API const FName UCollisionProfile::BlockAll_ProfileName = FName(TEXT("Blo
 ENGINE_API const FName UCollisionProfile::PhysicsActor_ProfileName = FName(TEXT("PhysicsActor"));
 ENGINE_API const FName UCollisionProfile::BlockAllDynamic_ProfileName = FName(TEXT("BlockAllDynamic"));
 ENGINE_API const FName UCollisionProfile::Pawn_ProfileName = FName(TEXT("Pawn"));
+ENGINE_API const FName UCollisionProfile::Vehicle_ProfileName = FName(TEXT("Vehicle"));
 
 UCollisionProfile::UCollisionProfile(const class FPostConstructInitializeProperties& PCIP)
 	: Super(PCIP)
diff --git a/Engine/Source/Runtime/Engine/Private/GameFramework/SpringArmComponent.cpp b/Engine/Source/Runtime/Engine/Private/GameFramework/SpringArmComponent.cpp
--- a/Engine/Source/Runtime/Engine/Private/GameFramework/SpringArmComponent.cpp
+++ b/Engine/Source/Runtime/Engine/Private/GameFramework/SpringArmComponent.cpp
@@ -62,7 +62,7 @@ void USpringArmComponent::UpdateDesiredArmLocation(bool bDoTrace, bool bDoLocati
 	PreviousDesiredRot = DesiredRot;
 
 
-	FVector ArmOrigin = GetComponentLocation();
+	FVector ArmOrigin = GetComponentLocation() + TargetOffset;
 	FVector DesiredLoc = ArmOrigin;
 	if (bDoLocationLag)
 	{
diff --git a/Engine/Source/Runtime/Engine/Private/Vehicles/VehicleWheel.cpp b/Engine/Source/Runtime/Engine/Private/Vehicles/VehicleWheel.cpp
--- a/Engine/Source/Runtime/Engine/Private/Vehicles/VehicleWheel.cpp
+++ b/Engine/Source/Runtime/Engine/Private/Vehicles/VehicleWheel.cpp
@@ -14,6 +14,9 @@ UVehicleWheel::UVehicleWheel(const class FPostConstructInitializeProperties& PCI
 	static ConstructorHelpers::FObjectFinder<UStaticMesh> CollisionMeshObj(TEXT("/Engine/EngineMeshes/Cylinder"));
 	CollisionMesh = CollisionMeshObj.Object;
 
+	static ConstructorHelpers::FObjectFinder<UTireType> TireTypeObj(TEXT("/Engine/EngineTireTypes/DefaultTireType.DefaultTireType"));
+	TireType = TireTypeObj.Object;
+
 	ShapeRadius = 30.0f;
 	ShapeWidth = 10.0f;
 	bAutoAdjustCollisionSize = true;
diff --git a/Engine/Source/Runtime/Engine/Private/Vehicles/WheeledVehicleMovementComponent4W.cpp b/Engine/Source/Runtime/Engine/Private/Vehicles/WheeledVehicleMovementComponent4W.cpp
--- a/Engine/Source/Runtime/Engine/Private/Vehicles/WheeledVehicleMovementComponent4W.cpp
+++ b/Engine/Source/Runtime/Engine/Private/Vehicles/WheeledVehicleMovementComponent4W.cpp
@@ -64,11 +64,7 @@ UWheeledVehicleMovementComponent4W::UWheeledVehicleMovementComponent4W(const cla
 	SteeringCurve[3].InVal = 1.0f;
 	SteeringCurve[3].OutVal = 0.1f;
 
-	WheelSetups.AddZeroed(4);
-	for (int i = 0 ; i < WheelSetups.Num(); ++i)
-	{
-		WheelSetups[i].WheelClass = UVehicleWheel::StaticClass();
-	}
+	WheelSetups.SetNum(4);
 #endif // WITH_PHYSX
 }
 
diff --git a/Engine/Source/Runtime/Engine/Private/WheeledVehicle.cpp b/Engine/Source/Runtime/Engine/Private/WheeledVehicle.cpp
--- a/Engine/Source/Runtime/Engine/Private/WheeledVehicle.cpp
+++ b/Engine/Source/Runtime/Engine/Private/WheeledVehicle.cpp
@@ -6,6 +6,7 @@
 =============================================================================*/
 
 #include "EnginePrivate.h"
+#include "DisplayDebugHelpers.h"
 
 FName AWheeledVehicle::VehicleMovementComponentName(TEXT("MovementComp"));
 FName AWheeledVehicle::VehicleMeshComponentName(TEXT("VehicleMesh"));
@@ -14,8 +15,7 @@ AWheeledVehicle::AWheeledVehicle(const class FPostConstructInitializeProperties&
 	: Super(PCIP)
 {
 	Mesh = PCIP.CreateDefaultSubobject<USkeletalMeshComponent>(this, VehicleMeshComponentName);
-	static FName CollisionProfileName(TEXT("Vehicle"));
-	Mesh->SetCollisionProfileName(CollisionProfileName);
+	Mesh->SetCollisionProfileName(UCollisionProfile::Vehicle_ProfileName);
 	Mesh->BodyInstance.bSimulatePhysics = true;
 	Mesh->BodyInstance.bNotifyRigidBodyCollision = true;
 	Mesh->BodyInstance.bUseCCD = true;
@@ -29,3 +29,16 @@ AWheeledVehicle::AWheeledVehicle(const class FPostConstructInitializeProperties&
 	VehicleMovement->UpdatedComponent = Mesh;
 }
 
+void AWheeledVehicle::DisplayDebug(UCanvas* Canvas, const FDebugDisplayInfo& DebugDisplay, float& YL, float& YPos)
+{
+	static FName NAME_Vehicle = FName(TEXT("Vehicle"));
+
+	Super::DisplayDebug(Canvas, DebugDisplay, YL, YPos);
+
+#if WITH_PHYSX
+	if (DebugDisplay.IsDisplayOn(NAME_Vehicle))
+	{
+		GetVehicleMovementComponent()->DrawDebug(Canvas, YL, YPos);
+	}
+#endif
+}
\ No newline at end of file
