Properly handle ActiveSound WaveInstances when stopping sounds for reimporting
TTP #333809

diff --git a/Engine/Source/Runtime/Engine/Private/AudioDevice.cpp b/Engine/Source/Runtime/Engine/Private/AudioDevice.cpp
--- a/Engine/Source/Runtime/Engine/Private/AudioDevice.cpp
+++ b/Engine/Source/Runtime/Engine/Private/AudioDevice.cpp
@@ -1945,9 +1945,9 @@ void FAudioDevice::StopAllSounds( bool bShouldStopUISounds )
 			// Then iterate across the wave instances.  If any of the wave instances is not a UI sound
 			// then we will stop the entire active sound because it makes less sense to leave it half
 			// executing
-			for (int32 WaveIndex=0; WaveIndex < ActiveSound->WaveInstances.Num(); ++WaveIndex)
+			for (auto WaveInstanceIt(ActiveSound->WaveInstances.CreateConstIterator()); WaveInstanceIt; ++WaveInstanceIt)
 			{
-				FWaveInstance* WaveInstance = ActiveSound->WaveInstances[WaveIndex];
+				FWaveInstance* WaveInstance = WaveInstanceIt.Value();
 				if (WaveInstance && !WaveInstance->bIsUISound)
 				{
 					ActiveSound->Stop(this);
@@ -2513,10 +2513,11 @@ void FAudioDevice::StopSoundsForReimport(USoundWave* ReimportedSoundWave, TArray
 		{
 			check(ActiveSound->Sound->IsA(USoundCue::StaticClass()));
 
-			for (int32 WaveInstanceIndex = 0; WaveInstanceIndex < ActiveSound->WaveInstances.Num(); ++WaveInstanceIndex)
+			for (auto WaveInstanceIt(ActiveSound->WaveInstances.CreateConstIterator()); WaveInstanceIt; ++WaveInstanceIt)
 			{
 				// If anything in the SoundCue uses the wave we're going to restart the whole thing for simplicity
-				if (ActiveSound->WaveInstances[WaveInstanceIndex]->WaveData == ReimportedSoundWave)
+				FWaveInstance* WaveInstance = WaveInstanceIt.Value();
+				if (WaveInstance->WaveData == ReimportedSoundWave)
 				{
 					if (ActiveSound->AudioComponent.IsValid())
 					{
