Avoid division by zero in FSceneView::Project()

diff --git a/Engine/Source/Runtime/Engine/Private/SceneView.cpp b/Engine/Source/Runtime/Engine/Private/SceneView.cpp
--- a/Engine/Source/Runtime/Engine/Private/SceneView.cpp
+++ b/Engine/Source/Runtime/Engine/Private/SceneView.cpp
@@ -437,6 +437,12 @@ FVector4 FSceneView::PixelToWorld(float X,float Y,float Z) const
 FPlane FSceneView::Project(const FVector& WorldPoint) const
 {
 	FPlane Result = WorldToScreen(WorldPoint);
+
+	if (Result.W == 0)
+	{
+		Result.W = KINDA_SMALL_NUMBER;
+	}
+
 	const float RHW = 1.0f / Result.W;
 
 	return FPlane(Result.X * RHW,Result.Y * RHW,Result.Z * RHW,Result.W);
diff --git a/Engine/Source/Runtime/Engine/Private/UserInterface/Canvas.cpp b/Engine/Source/Runtime/Engine/Private/UserInterface/Canvas.cpp
--- a/Engine/Source/Runtime/Engine/Private/UserInterface/Canvas.cpp
+++ b/Engine/Source/Runtime/Engine/Private/UserInterface/Canvas.cpp
@@ -1371,6 +1371,7 @@ FVector UCanvas::Project( FVector Location )
 
 	if (SceneView != NULL)
 	{
+		Location.DiagnosticCheckNaN();
 		V = SceneView->Project(Location);
 	}
 
