Add GetDelegateForManualSet for allowing payload parameters to be manually bound to input events

diff --git a/Engine/Source/Runtime/Engine/Classes/Components/InputComponent.h b/Engine/Source/Runtime/Engine/Classes/Components/InputComponent.h
--- a/Engine/Source/Runtime/Engine/Classes/Components/InputComponent.h
+++ b/Engine/Source/Runtime/Engine/Classes/Components/InputComponent.h
@@ -45,6 +45,12 @@ struct TInputUnifiedDelegate
 		FuncDynDelegate.SetFunctionName(FuncName);
 	}
 
+	DelegateType & GetDelegateForManualSet()
+	{
+		FuncDynDelegate.Unbind();
+		return FuncDelegate;
+	}
+
 	inline void Unbind()
 	{
 		FuncDelegate.Unbind();
