Allow actor components to replicate other subobjects on the owning actor channel.

diff --git a/Engine/Source/Runtime/Engine/Classes/Components/ActorComponent.h b/Engine/Source/Runtime/Engine/Classes/Components/ActorComponent.h
--- a/Engine/Source/Runtime/Engine/Classes/Components/ActorComponent.h
+++ b/Engine/Source/Runtime/Engine/Classes/Components/ActorComponent.h
@@ -136,6 +136,9 @@ public:
 	/** Returns whether replication is enabled or not. */
 	bool GetIsReplicated() const;
 
+	/** Allows a component to replicate other subobject on the actor  */
+	virtual bool ReplicateSubobjects(class UActorChannel *Channel, class FOutBunch *Bunch, FReplicationFlags *RepFlags);
+
 	virtual bool GetComponentClassCanReplicate() const;
 
 	/** Returns whether this component is an editor-only object or not */
diff --git a/Engine/Source/Runtime/Engine/Private/ActorComponent.cpp b/Engine/Source/Runtime/Engine/Private/ActorComponent.cpp
--- a/Engine/Source/Runtime/Engine/Private/ActorComponent.cpp
+++ b/Engine/Source/Runtime/Engine/Private/ActorComponent.cpp
@@ -1072,6 +1072,11 @@ bool UActorComponent::GetIsReplicated() const
 	return bReplicates;
 }
 
+bool UActorComponent::ReplicateSubobjects(class UActorChannel *Channel, class FOutBunch *Bunch, FReplicationFlags *RepFlags)
+{
+	return false;
+}
+
 bool UActorComponent::GetComponentClassCanReplicate() const
 {
 	return true;
diff --git a/Engine/Source/Runtime/Engine/Private/ActorReplication.cpp b/Engine/Source/Runtime/Engine/Private/ActorReplication.cpp
--- a/Engine/Source/Runtime/Engine/Private/ActorReplication.cpp
+++ b/Engine/Source/Runtime/Engine/Private/ActorReplication.cpp
@@ -249,6 +249,7 @@ bool AActor::ReplicateSubobjects(UActorChannel *Channel, FOutBunch *Bunch, FRepl
 		if (ActorComp && ActorComp->GetIsReplicated())
 		{
 			WroteSomething |= Channel->ReplicateSubobject(ActorComp, *Bunch, *RepFlags);
+			WroteSomething |= ActorComp->ReplicateSubobjects(Channel, Bunch, RepFlags);
 		}
 	}
 	return WroteSomething;
