Added RandomUnitVectorInCone and FInterpEaseInOut math functions to Blueprint math library

diff --git a/Engine/Source/Runtime/Engine/Classes/Kismet/KismetMathLibrary.h b/Engine/Source/Runtime/Engine/Classes/Kismet/KismetMathLibrary.h
--- a/Engine/Source/Runtime/Engine/Classes/Kismet/KismetMathLibrary.h
+++ b/Engine/Source/Runtime/Engine/Classes/Kismet/KismetMathLibrary.h
@@ -375,6 +375,11 @@ class UKismetMathLibrary : public UBlueprintFunctionLibrary
 	UFUNCTION(BlueprintPure, meta=(Keywords = "* multiply"), Category="Math|Float")
 	static float MultiplyByPi(float Value);
 
+	/** Interpolate between A and B, applying an ease in/out function.  Exp controls the degree of the curve. */
+	UFUNCTION(BlueprintPure, Category = "Math|Float")
+	static float FInterpEaseInOut(float A, float B, float Alpha, float Exponent);
+
+
 	//
 	// Vector functions.
 	//
@@ -461,6 +466,14 @@ class UKismetMathLibrary : public UBlueprintFunctionLibrary
 	UFUNCTION(BlueprintPure, Category="Math|Random")
 	static FVector RandomUnitVector();
 
+	/* 
+	 * Returns a random vector with length of 1, within the specified cone, with uniform random distribution. 
+	 * @param ConeDir	The base "center" direction of the cone.
+	 * @param ConeHalfAngle		The half-angle of the cone (from ConeDir to edge), in degrees.
+	 */
+	UFUNCTION(BlueprintPure, Category = "Math|Random")
+	static FVector RandomUnitVectorInCone(FVector ConeDir, float ConeHalfAngle);
+
 	// Mirrors a vector by a normal
 	UFUNCTION(BlueprintPure, Category="Math|Vector")
 	static FVector MirrorVectorByNormal(FVector InVect, FVector InNormal);
diff --git a/Engine/Source/Runtime/Engine/Private/KismetMathLibrary.cpp b/Engine/Source/Runtime/Engine/Private/KismetMathLibrary.cpp
--- a/Engine/Source/Runtime/Engine/Private/KismetMathLibrary.cpp
+++ b/Engine/Source/Runtime/Engine/Private/KismetMathLibrary.cpp
@@ -479,6 +479,11 @@ float UKismetMathLibrary::MultiplyByPi(float Value)
 	return Value * PI;
 }
 
+float UKismetMathLibrary::FInterpEaseInOut(float A, float B, float Alpha, float Exponent)
+{
+	return FMath::InterpEaseInOut<float>(A, B, Alpha, Exponent);
+}
+
 float UKismetMathLibrary::RandomFloat()
 {
 	return FMath::FRand();
@@ -639,6 +644,11 @@ FVector UKismetMathLibrary::RandomUnitVector()
 	return FMath::VRand();
 }
 
+FVector UKismetMathLibrary::RandomUnitVectorInCone(FVector ConeDir, float ConeHalfAngle)
+{
+	return FMath::VRandCone(ConeDir, ConeHalfAngle);
+}
+
 FVector UKismetMathLibrary::MirrorVectorByNormal(FVector A, FVector B)
 {
 	B = B.SafeNormal();
