Move physical material into collision category which is exposed by blocking volume

diff --git a/Engine/Source/Runtime/Engine/Classes/PhysicsEngine/BodyInstance.h b/Engine/Source/Runtime/Engine/Classes/PhysicsEngine/BodyInstance.h
--- a/Engine/Source/Runtime/Engine/Classes/PhysicsEngine/BodyInstance.h
+++ b/Engine/Source/Runtime/Engine/Classes/PhysicsEngine/BodyInstance.h
@@ -159,7 +159,7 @@ protected:
 public:
 
 	/**	Allows you to override the PhysicalMaterial to use for simple collision on this body. */
-	UPROPERTY(EditAnywhere, AdvancedDisplay, BlueprintReadOnly, Category=Physics)
+	UPROPERTY(EditAnywhere, AdvancedDisplay, BlueprintReadOnly, Category=Collision)
 	class UPhysicalMaterial* PhysMaterialOverride;
 
 	/** User specified offset for the center of mass of this object, from the calculated location */
