Add scene specific deferred apex commands - this fixes apex actors marked for deletion not being removed until GC.

TTP#333195

diff --git a/Engine/Source/Runtime/Engine/Private/DestructibleComponent.cpp b/Engine/Source/Runtime/Engine/Private/DestructibleComponent.cpp
--- a/Engine/Source/Runtime/Engine/Private/DestructibleComponent.cpp
+++ b/Engine/Source/Runtime/Engine/Private/DestructibleComponent.cpp
@@ -346,7 +346,14 @@ void UDestructibleComponent::DestroyPhysicsState()
 #if WITH_APEX
 	if(ApexDestructibleActor != NULL)
 	{
-		GPhysCommandHandler->DeferredRelease(ApexDestructibleActor);
+		if(UWorld * World = GetWorld())
+		{
+			if(FPhysScene * PhysScene = World->GetPhysicsScene())
+			{
+				PhysScene->DeferredCommandHandler.DeferredRelease(ApexDestructibleActor);
+			}
+		}
+		
 		ApexDestructibleActor = NULL;
 	}
 #endif	// #if WITH_APEX
diff --git a/Engine/Source/Runtime/Engine/Private/LevelTick.cpp b/Engine/Source/Runtime/Engine/Private/LevelTick.cpp
--- a/Engine/Source/Runtime/Engine/Private/LevelTick.cpp
+++ b/Engine/Source/Runtime/Engine/Private/LevelTick.cpp
@@ -1122,6 +1122,7 @@ void UWorld::Tick( ELevelTick TickType, float DeltaSeconds )
 		RunTickGroup(TG_DuringPhysics, false); // No wait here, we should run until idle though. We don't care if all of the async ticks are done before we start running post-phys stuff
 		TickGroup = TG_EndPhysics; // set this here so the current tick group is correct during collision notifies, though I am not sure it matters. 'cause of the false up there^^^
 		RunTickGroup(TG_EndPhysics); 
+		PhysicsScene->DeferredCommandHandler.Flush();
 		RunTickGroup(TG_PreCloth);
 		RunTickGroup(TG_StartCloth);
 		RunTickGroup(TG_EndCloth);
diff --git a/Engine/Source/Runtime/Engine/Public/PhysicsPublic.h b/Engine/Source/Runtime/Engine/Public/PhysicsPublic.h
--- a/Engine/Source/Runtime/Engine/Public/PhysicsPublic.h
+++ b/Engine/Source/Runtime/Engine/Public/PhysicsPublic.h
@@ -392,6 +392,9 @@ public:
 #endif
 #endif
 
+	/** DeferredCommandHandler - this is mainly used for scene specific APEX calls that need to be deferred */
+	FPhysCommandHandler DeferredCommandHandler;
+
 private:
 	/** Initialize a scene of the given type.  Must only be called once for each scene type. */
 	void InitPhysScene(uint32 SceneType);
