Fix assert when attaching destructible mesh to another destructible mesh

diff --git a/Engine/Source/Runtime/Engine/Private/PhysicsEngine/BodyInstance.cpp b/Engine/Source/Runtime/Engine/Private/PhysicsEngine/BodyInstance.cpp
--- a/Engine/Source/Runtime/Engine/Private/PhysicsEngine/BodyInstance.cpp
+++ b/Engine/Source/Runtime/Engine/Private/PhysicsEngine/BodyInstance.cpp
@@ -1030,7 +1030,7 @@ void TermBodyHelper(int32 & SceneIndex, PxRigidActor *& PRigidActor, FBodyInstan
 		{
 			//When DestructibleMesh is used we create fake BodyInstances. In this case the RigidActor is set, but InitBody is never called.
 			//The RigidActor is released by the destructible component, but it's still up to us to NULL out the pointer
-			checkSlow(FPhysxUserData::Get<FDestructibleChunkInfo>(PRigidActor->userData));	//Make sure we are really a destructible
+			checkSlow(!PRigidActor->userData || FPhysxUserData::Get<FDestructibleChunkInfo>(PRigidActor->userData));	//Make sure we are really a destructible. Note you CAN get a case when userData is NULL, for example when trying to attach a destructible to another.
 			PRigidActor = NULL;
 		}
 	}
