#UE4 Change it so AActor::BeginPlay is called right after level initialization, and not when the match starts. Fixes various issues with when begin play is called. Rename functions for clarity and fix a incorrect comments.
Here is the new order:
UWorld::InitializeActorsForPlay calls Pre/PostInitializeComponents and sets up the rest of initialization. UWorld::AreActorsInitialized is true at this step (was HasBegunPlay before)
UWorld::BeginPlay calls BeginPlay on all actors and transitions the game mode to the WaitingToStartMatch state. UWorld::HasBegunPlay is true at this step
GameMode::StartMatch transitions to start match, doesn't call any actor call backs.

TTP #333613: Actor begin play now executes correctly in simulate
TTP #311938: Call actor initialize on all actors in a streamed in sublevel after game start, then call all of the beginplay

diff --git a/Engine/Plugins/Experimental/LiveEditor/Source/LiveEditor/Private/LiveEditorManager.cpp b/Engine/Plugins/Experimental/LiveEditor/Source/LiveEditor/Private/LiveEditorManager.cpp
--- a/Engine/Plugins/Experimental/LiveEditor/Source/LiveEditor/Private/LiveEditorManager.cpp
+++ b/Engine/Plugins/Experimental/LiveEditor/Source/LiveEditor/Private/LiveEditorManager.cpp
@@ -625,7 +625,7 @@ void FLiveEditorManager::CreateLiveEditorWorld()
 										.CreateNavigation(false)
 										.ShouldSimulatePhysics(false)
 										.SetTransactional(bTransactional));
-	LiveEditorWorld->BeginPlay(FURL());
+	LiveEditorWorld->InitializeActorsForPlay(FURL());
 }
 
 
diff --git a/Engine/Source/Editor/UnrealEd/Private/PlayLevel.cpp b/Engine/Source/Editor/UnrealEd/Private/PlayLevel.cpp
--- a/Engine/Source/Editor/UnrealEd/Private/PlayLevel.cpp
+++ b/Engine/Source/Editor/UnrealEd/Private/PlayLevel.cpp
@@ -2423,7 +2423,8 @@ UWorld* UEditorEngine::CreatePlayInEditorWorld(FWorldContext &PieWorldContext, b
 
 	EditorWorld->TransferBlueprintDebugReferences(PlayWorld);
 
-	PlayWorld->BeginPlay(URL);
+	PlayWorld->InitializeActorsForPlay(URL);
+	PlayWorld->BeginPlay();
 
 	// This must have already been set with a call to DisableRealtimeViewports() outside of this method.
 	check(!IsAnyViewportRealtime());
diff --git a/Engine/Source/Runtime/Engine/Classes/AI/Navigation/NavigationSystem.h b/Engine/Source/Runtime/Engine/Classes/AI/Navigation/NavigationSystem.h
--- a/Engine/Source/Runtime/Engine/Classes/AI/Navigation/NavigationSystem.h
+++ b/Engine/Source/Runtime/Engine/Classes/AI/Navigation/NavigationSystem.h
@@ -893,9 +893,9 @@ public:
 	void CleanUp();
 
 	/** 
-	 *	Called when owner-UWorld begins play. 
+	 *	Called when owner-UWorld initializes actors
 	 */
-	virtual void OnBeginPlay();
+	virtual void OnInitializeActors();
 
 	/** */
 	virtual void OnWorldInitDone(NavigationSystem::EMode Mode);
diff --git a/Engine/Source/Runtime/Engine/Classes/Engine/Level.h b/Engine/Source/Runtime/Engine/Classes/Engine/Level.h
--- a/Engine/Source/Runtime/Engine/Classes/Engine/Level.h
+++ b/Engine/Source/Runtime/Engine/Classes/Engine/Level.h
@@ -386,8 +386,8 @@ class ULevel : public ULevelBase, public IInterface_AssetUserData
 	uint32										bAlreadyUpdatedComponents:1;
 	/** Whether we already associated streamable resources.									*/
 	uint32										bAlreadyAssociatedStreamableResources:1;
-	/** Whether we already initialized actors.												*/
-	uint32										bAlreadyInitializedActors:1;
+	/** Whether we already initialized network actors.											*/
+	uint32										bAlreadyInitializedNetworkActors:1;
 	/** Whether we already routed initialize on actors.										*/
 	uint32										bAlreadyRoutedActorInitialize:1;
 	/** Whether we already sorted the actor list.											*/
@@ -548,10 +548,8 @@ public:
 	 */
 	ENGINE_API void SortActorList();
 
-	/**
-	 * Initializes all actors after loading completed.
-	 */
-	void InitializeActors();
+	/** Handles network initialization for actors in this level */
+	void InitializeNetworkActors();
 
 	/** Initializes rendering resources for this level. */
 	void InitializeRenderingResources();
diff --git a/Engine/Source/Runtime/Engine/Classes/Engine/World.h b/Engine/Source/Runtime/Engine/Classes/Engine/World.h
--- a/Engine/Source/Runtime/Engine/Classes/Engine/World.h
+++ b/Engine/Source/Runtime/Engine/Classes/Engine/World.h
@@ -869,8 +869,11 @@ public:
 	/** Whether it was requested that the engine bring up a loading screen and block on async loading. */   
 	uint32 bRequestedBlockOnAsyncLoading:1;
 
-	/** Whether gameplay has begun. */
-	uint32 bBegunPlay:1;		
+	/** Whether actors have been initialized for play */
+	uint32 bActorsInitialized:1;
+
+	/** Whether BeginPlay has been called on actors */
+	uint32 bBegunPlay:1;
 
 	/** Whether the match has been started */
 	uint32 bMatchStarted:1;
@@ -1411,7 +1414,10 @@ public:
 	/** Returns the default brush. */
 	ABrush* GetBrush() const;
 
-	/** Returns true if game has already started, false otherwise. */
+	/** Returns true if the actors have been initialized and are ready to start play */
+	bool AreActorsInitialized() const;
+
+	/** Returns true if gameplay has already started, false otherwise. */
 	bool HasBegunPlay() const;
 
 	/**
@@ -2107,13 +2113,20 @@ public:
 	/** Spawns GameMode for the level. */
 	bool SetGameMode(const FURL& InURL);
 
-	/** BeginPlay - Begins gameplay in the level.
+	/** 
+	 * Initializes all actors and prepares them to start gameplay
 	 * @param InURL commandline URL
 	 * @param bResetTime (optional) whether the WorldSettings's TimeSeconds should be reset to zero
 	 */
-	void BeginPlay(const FURL& InURL, bool bResetTime = true);
+	void InitializeActorsForPlay(const FURL& InURL, bool bResetTime = true);
 
-	/** Looks for a PlayerController that was being swapped by the given NetConnection and, if found, destroys it
+	/**
+	 * Start gameplay. This will cause the game mode to transition to the correct state and call BeginPlay on all actors
+	 */
+	void BeginPlay();
+
+	/** 
+	 * Looks for a PlayerController that was being swapped by the given NetConnection and, if found, destroys it
 	 * (because the swap is complete or the connection was closed)
 	 * @param Connection - the connection that performed the swap
 	 * @return whether a PC waiting for a swap was found
diff --git a/Engine/Source/Runtime/Engine/Classes/GameFramework/GameMode.h b/Engine/Source/Runtime/Engine/Classes/GameFramework/GameMode.h
--- a/Engine/Source/Runtime/Engine/Classes/GameFramework/GameMode.h
+++ b/Engine/Source/Runtime/Engine/Classes/GameFramework/GameMode.h
@@ -69,6 +69,14 @@ class ENGINE_API AGameMode : public AInfo
 	/** Returns true if the match state is WaitingPostMatch or later */
 	virtual bool HasMatchEnded() const;
 
+	/** Transitions to WaitingToStart and calls BeginPlay on actors. */
+	virtual void StartPlay();
+
+	/**
+	 * Called after a seamless level transition has been completed on the *new* GameMode.
+	 * Used to reinitialize players already in the game as they won't have *Login() called on them
+	 */
+	virtual void PostSeamlessTravel();
 
 	/** Transition from WaitingToStart to InProgress. You can call this manually, will also get called if ReadyToStartMatch returns true */
 	virtual void StartMatch();
@@ -241,8 +249,6 @@ public:
 
 	// Begin AActor interface
 	virtual void PreInitializeComponents() OVERRIDE;
-	/** This sets the match state to WaitingToStart unless we're in seamless travel */
-	virtual void PostInitializeComponents() OVERRIDE; 
 	virtual void DisplayDebug(class UCanvas* Canvas, const FDebugDisplayInfo& DebugDisplay, float& YL, float& YPos) OVERRIDE;
 	virtual void Reset() OVERRIDE;
 	virtual void Tick(float DeltaSeconds) OVERRIDE;
@@ -543,12 +549,6 @@ public:
 	 */
 	virtual void SwapPlayerControllers(APlayerController* OldPC, APlayerController* NewPC);
 
-	/** 
-	 * Called after a seamless level transition has been completed on the *new* GameMode. This transitions to WaitingToStart
-	 * Used to reinitialize players already in the game as they won't have *Login() called on them
-	 */
-	virtual void PostSeamlessTravel();
-
 	/** 
 	 * Handles reinitializing players that remained through a seamless level transition
 	 * called from C++ for players that finished loading after the server
diff --git a/Engine/Source/Runtime/Engine/Classes/GameFramework/GameState.h b/Engine/Source/Runtime/Engine/Classes/GameFramework/GameState.h
--- a/Engine/Source/Runtime/Engine/Classes/GameFramework/GameState.h
+++ b/Engine/Source/Runtime/Engine/Classes/GameFramework/GameState.h
@@ -48,6 +48,10 @@ protected:
 	UPROPERTY(replicatedUsing=OnRep_MatchState)
 	FName MatchState;
 
+	/** Previous map state, used to handle if multiple transitions happen per frame */
+	UPROPERTY()
+	FName PreviousMatchState;
+
 	/** Called when the state transitions to WaitingToStart */
 	virtual void HandleMatchIsWaitingToStart();
 
diff --git a/Engine/Source/Runtime/Engine/Classes/GameFramework/WorldSettings.h b/Engine/Source/Runtime/Engine/Classes/GameFramework/WorldSettings.h
--- a/Engine/Source/Runtime/Engine/Classes/GameFramework/WorldSettings.h
+++ b/Engine/Source/Runtime/Engine/Classes/GameFramework/WorldSettings.h
@@ -423,9 +423,13 @@ public:
 		return TimeDilation * MatineeTimeDilation;
 	}
 
-	/** Called by GameMode.HandleMatchHasStarted, used to notify native classes of match startup (such as level scripting).
-	 * @todo UE4 Need to send K2 event about match starting here
-	 * Activates LevelStartup and/or LevelBeginning events in the sequences in the world
+	/** 
+	 * Called by GameMode.HandleMatchIsWaitingToStart, calls BeginPlay on all actors
+	 */
+	virtual void NotifyBeginPlay();
+
+	/** 
+	 * Called by GameMode.HandleMatchHasStarted, used to notify native classes of match startup (such as level scripting)
 	 */	
 	virtual void NotifyMatchStarted();
 
diff --git a/Engine/Source/Runtime/Engine/Private/AI/Navigation/NavGraphGenerator.cpp b/Engine/Source/Runtime/Engine/Private/AI/Navigation/NavGraphGenerator.cpp
--- a/Engine/Source/Runtime/Engine/Private/AI/Navigation/NavGraphGenerator.cpp
+++ b/Engine/Source/Runtime/Engine/Private/AI/Navigation/NavGraphGenerator.cpp
@@ -34,11 +34,6 @@ void FNavGraphGenerator::OnWorldInitDone(bool bAllowedToRebuild)
 
 }
 
-void FNavGraphGenerator::OnPostBeginPlay()
-{
-
-}
-
 void FNavGraphGenerator::OnNavigationBuildingLocked() 
 {
 
diff --git a/Engine/Source/Runtime/Engine/Private/AI/Navigation/NavGraphGenerator.h b/Engine/Source/Runtime/Engine/Private/AI/Navigation/NavGraphGenerator.h
--- a/Engine/Source/Runtime/Engine/Private/AI/Navigation/NavGraphGenerator.h
+++ b/Engine/Source/Runtime/Engine/Private/AI/Navigation/NavGraphGenerator.h
@@ -28,8 +28,6 @@ public:
 	virtual bool IsBuildInProgress(bool bCheckDirtyToo = false) const OVERRIDE;
 
 	virtual void OnWorldInitDone(bool bAllowedToRebuild) OVERRIDE;
-
-	virtual void OnPostBeginPlay() OVERRIDE;
 		
 	//--- accessors --- //
 
diff --git a/Engine/Source/Runtime/Engine/Private/AI/Navigation/NavigationSystem.cpp b/Engine/Source/Runtime/Engine/Private/AI/Navigation/NavigationSystem.cpp
--- a/Engine/Source/Runtime/Engine/Private/AI/Navigation/NavigationSystem.cpp
+++ b/Engine/Source/Runtime/Engine/Private/AI/Navigation/NavigationSystem.cpp
@@ -391,7 +391,7 @@ void UNavigationSystem::PostInitProperties()
 	}
 }
 
-void UNavigationSystem::OnBeginPlay()
+void UNavigationSystem::OnInitializeActors()
 {
 	
 }
diff --git a/Engine/Source/Runtime/Engine/Private/Actor.cpp b/Engine/Source/Runtime/Engine/Private/Actor.cpp
--- a/Engine/Source/Runtime/Engine/Private/Actor.cpp
+++ b/Engine/Source/Runtime/Engine/Private/Actor.cpp
@@ -323,7 +323,7 @@ bool AActor::TeleportTo( const FVector& DestLocation, const FRotator& DestRotati
 	}
 
 	// Can't move non-movable actors during play
-	if( (RootComponent->Mobility == EComponentMobility::Static) && GetWorld()->HasBegunPlay() )
+	if( (RootComponent->Mobility == EComponentMobility::Static) && GetWorld()->AreActorsInitialized() )
 	{
 		return false;
 	}
@@ -573,7 +573,7 @@ void AActor::PostLoadSubobjects(FObjectInstancingGraph* OuterInstanceGraph)
 
 void AActor::ProcessEvent(UFunction* Function, void* Parameters)
 {
-	if( ((GetWorld() && (GetWorld()->HasBegunPlay() || GAllowActorScriptExecutionInEditor)) || HasAnyFlags(RF_ClassDefaultObject)) && !GIsGarbageCollecting )
+	if( ((GetWorld() && (GetWorld()->AreActorsInitialized() || GAllowActorScriptExecutionInEditor)) || HasAnyFlags(RF_ClassDefaultObject)) && !GIsGarbageCollecting )
 	{
 		Super::ProcessEvent(Function, Parameters);
 	}
@@ -2108,7 +2108,7 @@ void AActor::PostSpawnInitialize(FVector const& SpawnLocation, FRotator const& S
 	// This should be the same sequence for deferred or nondeferred spawning.
 
 	// It's not safe to call UWorld accessor functions till the world info has been spawned.
-	bool const bBegunPlay = GetWorld() && GetWorld()->HasBegunPlay();
+	bool const bActorsInitialized = GetWorld() && GetWorld()->AreActorsInitialized();
 
 	CreationTime = GetWorld()->GetTimeSeconds();
 
@@ -2136,7 +2136,7 @@ void AActor::PostSpawnInitialize(FVector const& SpawnLocation, FRotator const& S
 
 #if WITH_EDITOR
 	// When placing actors in the editor, init any random streams 
-	if (!bBegunPlay)
+	if (!bActorsInitialized)
 	{
 		SeedAllRandomStreams();
 	}
@@ -2162,9 +2162,9 @@ void AActor::PostSpawnInitialize(FVector const& SpawnLocation, FRotator const& S
 
 void AActor::PostActorConstruction()
 {
-	bool const bWorldBegunPlay = GetWorld() && GetWorld()->HasBegunPlay();
+	bool const bActorsInitialized = GetWorld() && GetWorld()->AreActorsInitialized();
 
-	if (bWorldBegunPlay)
+	if (bActorsInitialized)
 	{
 		PreInitializeComponents();
 	}
@@ -2175,7 +2175,7 @@ void AActor::PostActorConstruction()
 	// If this is dynamically spawned replicted actor, defer calls to BeginPlay and UpdateOverlaps until replicated properties are deserialized
 	bool deferBeginPlayAndUpdateOverlaps = (bExchangedRoles && RemoteRole == ROLE_Authority);
 
-	if (bWorldBegunPlay)
+	if (bActorsInitialized)
 	{
 		// Call InitializeComponent on components
 		InitializeComponents();
@@ -2186,7 +2186,7 @@ void AActor::PostActorConstruction()
 			UE_LOG(LogActor, Fatal, TEXT("%s failed to route PostInitializeComponents.  Please call Super::PostInitializeComponents() in your <className>::PostInitializeComponents() function. "), *GetFullName() );
 		}
 
-		if (GetWorld()->bMatchStarted && !deferBeginPlayAndUpdateOverlaps)
+		if (GetWorld()->HasBegunPlay() && !deferBeginPlayAndUpdateOverlaps)
 		{
 			BeginPlay();
 		}
@@ -2263,8 +2263,7 @@ void AActor::PostNetInit()
 	}
 	check(RemoteRole == ROLE_Authority);
 
-	AGameState *GameState = GetWorld()->GameState;
-	if (GameState && GameState->HasMatchStarted())
+	if (GetWorld() && GetWorld()->HasBegunPlay())
 	{
 		BeginPlay();
 	}
diff --git a/Engine/Source/Runtime/Engine/Private/GameMode.cpp b/Engine/Source/Runtime/Engine/Private/GameMode.cpp
--- a/Engine/Source/Runtime/Engine/Private/GameMode.cpp
+++ b/Engine/Source/Runtime/Engine/Private/GameMode.cpp
@@ -403,17 +403,6 @@ void AGameMode::PreInitializeComponents()
 	InitGameState();
 }
 
-void AGameMode::PostInitializeComponents()
-{
-	Super::PostInitializeComponents();
-
-	if (!GetWorld()->IsInSeamlessTravel())
-	{
-		// If this isn't seamless travel, enter waiting state immediately. PostSeamlessTravel handles the state transition otherwise
-		SetMatchState(MatchState::WaitingToStart);
-	}
-}
-
 void AGameMode::RestartPlayer(AController* NewPlayer)
 {
 	if ( NewPlayer == NULL || NewPlayer->IsPendingKillPending() )
@@ -491,12 +480,23 @@ void AGameMode::InitStartSpot(AActor* StartSpot, AController* NewPlayer)
 {
 }
 
+void AGameMode::StartPlay()
+{
+	if (MatchState == MatchState::EnteringMap)
+	{
+		SetMatchState(MatchState::WaitingToStart);
+	}
+}
+
 void AGameMode::HandleMatchIsWaitingToStart()
 {
 	if (GameSession != NULL)
 	{
 		GameSession->HandleMatchIsWaitingToStart();
 	}
+
+	// Calls begin play on actors
+	GetWorldSettings()->NotifyBeginPlay();
 }
 
 bool AGameMode::ReadyToStartMatch()
@@ -508,7 +508,14 @@ bool AGameMode::ReadyToStartMatch()
 	}
 
 	// By default start when we have > 0 players
-	return (GetMatchState() == MatchState::WaitingToStart && NumPlayers + NumBots > 0);
+	if (GetMatchState() == MatchState::WaitingToStart)
+	{
+		if (NumPlayers + NumBots > 0)
+		{
+			return true;
+		}
+	}
+	return false;
 }
 
 void AGameMode::StartMatch()
@@ -1314,7 +1321,7 @@ void AGameMode::StartNewPlayer(APlayerController* NewPlayer)
 	// If players should start as spectators, leave them in the spectator state
 	if (!bStartPlayersAsSpectators)
 	{
-		// If match is in progress, start the player. Otherwise match will start from the game mode tick
+		// If match is in progress, start the player
 		if (IsMatchInProgress())
 		{
 			RestartPlayer(NewPlayer);
@@ -1324,6 +1331,15 @@ void AGameMode::StartNewPlayer(APlayerController* NewPlayer)
 				NewPlayer->GetPawn()->ClientSetRotation(NewPlayer->GetPawn()->GetActorRotation());
 			}
 		}
+		// Check to see if we should start right away, avoids a one frame lag in single player games
+		else if (GetMatchState() == MatchState::WaitingToStart)
+		{
+			// Check to see if we should start the match
+			if (ReadyToStartMatch())
+			{
+				StartMatch();
+			}
+		}
 	}
 }
 
@@ -1609,8 +1625,6 @@ void AGameMode::PostSeamlessTravel()
 			}
 		}
 	}
-
-	SetMatchState(MatchState::WaitingToStart);
 }
 
 void AGameMode::MatineeCancelled()
diff --git a/Engine/Source/Runtime/Engine/Private/GameState.cpp b/Engine/Source/Runtime/Engine/Private/GameState.cpp
--- a/Engine/Source/Runtime/Engine/Private/GameState.cpp
+++ b/Engine/Source/Runtime/Engine/Private/GameState.cpp
@@ -19,11 +19,12 @@ AGameState::AGameState(const class FPostConstructInitializeProperties& PCIP)
 	bReplicateMovement = false;
 
 	// Note: this is very important to set to false. Though all replication infos are spawned at run time, during seamless travel
-	// they are held on to and brought over into the new world. In ULevel::InitializeActors, these PlayerStates may be treated as map/startup actors
+	// they are held on to and brought over into the new world. In ULevel::InitializeNetworkActors, these PlayerStates may be treated as map/startup actors
 	// and given static NetGUIDs. This also causes their deletions to be recorded and sent to new clients, which if unlucky due to name conflicts,
 	// may end up deleting the new PlayerStates they had just spaned.
 	bNetLoadOnClient = false;
 	MatchState = MatchState::EnteringMap;
+	PreviousMatchState = MatchState::EnteringMap;
 }
 
 /** Helper to return the default object of the GameMode class corresponding to this GameState. */
@@ -151,7 +152,11 @@ void AGameState::RemovePlayerState(APlayerState* PlayerState)
 
 void AGameState::HandleMatchIsWaitingToStart()
 {
-
+	if (Role != ROLE_Authority)
+	{
+		// Server handles this in AGameMode::HandleMatchIsWaitingToStart
+		GetWorldSettings()->NotifyBeginPlay();
+	}
 }
 
 void AGameState::HandleMatchHasStarted()
@@ -216,11 +221,13 @@ void AGameState::SetMatchState(FName NewState)
 
 void AGameState::OnRep_MatchState()
 {
-	if (MatchState == MatchState::WaitingToStart)
+	if (MatchState == MatchState::WaitingToStart || PreviousMatchState == MatchState::EnteringMap)
 	{
+		// Call MatchIsWaiting to start even if you join in progress at a later state
 		HandleMatchIsWaitingToStart();
 	}
-	else if (MatchState == MatchState::InProgress)
+	
+	if (MatchState == MatchState::InProgress)
 	{
 		HandleMatchHasStarted();
 	}
@@ -232,6 +239,8 @@ void AGameState::OnRep_MatchState()
 	{
 		HandleLeavingMap();
 	}
+
+	PreviousMatchState = MatchState;
 }
 
 bool AGameState::ShouldShowGore() const
diff --git a/Engine/Source/Runtime/Engine/Private/Level.cpp b/Engine/Source/Runtime/Engine/Private/Level.cpp
--- a/Engine/Source/Runtime/Engine/Private/Level.cpp
+++ b/Engine/Source/Runtime/Engine/Private/Level.cpp
@@ -1501,13 +1501,12 @@ ALevelScriptActor* ULevel::GetLevelScriptActor() const
 }
 
 
-void ULevel::InitializeActors()
+void ULevel::InitializeNetworkActors()
 {
 	check( OwningWorld );
 	bool			bIsServer				= OwningWorld->IsServer();
-	APhysicsVolume*	DefaultPhysicsVolume	= OwningWorld->GetDefaultPhysicsVolume();
 
-	// Kill non relevant client actors, initialize render time, set initial physic volume, initialize script execution and rigid body physics.
+	// Kill non relevant client actors and set net roles correctly
 	for( int32 ActorIndex=0; ActorIndex<Actors.Num(); ActorIndex++ )
 	{
 		AActor* Actor = Actors[ActorIndex];
@@ -1582,7 +1581,8 @@ void ULevel::RouteActorInitialize()
 		}
 	}
 
-	const bool bCallBeginPlay = OwningWorld->bMatchStarted;
+	const bool bCallBeginPlay = OwningWorld->HasBegunPlay();
+	TArray<AActor *> ActorsToBeginPlay;
 
 	// Send InitializeComponents on components and PostInitializeComponents.
 	for( int32 ActorIndex=0; ActorIndex<Actors.Num(); ActorIndex++ )
@@ -1592,7 +1592,7 @@ void ULevel::RouteActorInitialize()
 		{
 			if( !Actor->bActorInitialized )
 			{
-				// Call BeginPlay on Components.
+				// Call Initialize on Components.
 				Actor->InitializeComponents();
 
 				if(Actor->bWantsInitialize)
@@ -1610,7 +1610,7 @@ void ULevel::RouteActorInitialize()
 
 				if (bCallBeginPlay)
 				{
-					Actor->BeginPlay();
+					ActorsToBeginPlay.Add(Actor);
 				}
 			}
 
@@ -1620,6 +1620,13 @@ void ULevel::RouteActorInitialize()
 			Actor->UpdateOverlaps(false);
 		}
 	}
+
+	// Do this in a second pass to make sure they're all initialized before begin play starts
+	for (int32 ActorIndex = 0; ActorIndex < ActorsToBeginPlay.Num(); ActorIndex++)
+	{
+		AActor* Actor = ActorsToBeginPlay[ActorIndex];
+		Actor->BeginPlay();			
+	}
 }
 
 bool ULevel::HasAnyActorsOfType(UClass *SearchType)
diff --git a/Engine/Source/Runtime/Engine/Private/LevelActor.cpp b/Engine/Source/Runtime/Engine/Private/LevelActor.cpp
--- a/Engine/Source/Runtime/Engine/Private/LevelActor.cpp
+++ b/Engine/Source/Runtime/Engine/Private/LevelActor.cpp
@@ -390,7 +390,7 @@ bool UWorld::DestroyActor( AActor* ThisActor, bool bNetForce, bool bShouldModify
 	}
 
 	// In-game deletion rules.
-	if( HasBegunPlay() )
+	if( AreActorsInitialized() )
 	{
 		// Never destroy the world settings actor. This used to be enforced by bNoDelete and is actually needed for 
 		// seamless travel and network games.
diff --git a/Engine/Source/Runtime/Engine/Private/PreviewScene.cpp b/Engine/Source/Runtime/Engine/Private/PreviewScene.cpp
--- a/Engine/Source/Runtime/Engine/Private/PreviewScene.cpp
+++ b/Engine/Source/Runtime/Engine/Private/PreviewScene.cpp
@@ -28,7 +28,7 @@ FPreviewScene::FPreviewScene(FPreviewScene::ConstructionValues CVS)
 										.RequiresHitProxies(false)
 										.ShouldSimulatePhysics(CVS.bShouldSimulatePhysics)
 										.SetTransactional(CVS.bTransactional));
-	PreviewWorld->BeginPlay(FURL());
+	PreviewWorld->InitializeActorsForPlay(FURL());
 
 	GetScene()->UpdateDynamicSkyLight(FLinearColor::White * CVS.SkyBrightness, FLinearColor::Black);
 
diff --git a/Engine/Source/Runtime/Engine/Private/PrimitiveComponent.cpp b/Engine/Source/Runtime/Engine/Private/PrimitiveComponent.cpp
--- a/Engine/Source/Runtime/Engine/Private/PrimitiveComponent.cpp
+++ b/Engine/Source/Runtime/Engine/Private/PrimitiveComponent.cpp
@@ -1207,7 +1207,7 @@ bool UPrimitiveComponent::MoveComponent( const FVector& Delta, const FRotator& N
 	if (Mobility == EComponentMobility::Static && Actor && Actor->bActorInitialized)
 	{
 		// TODO: Static components without an owner can move, should they be able to?
-		UE_LOG(LogPrimitiveComponent, Warning, TEXT("Trying to move static component '%s' after begin play"), *GetFullName());
+		UE_LOG(LogPrimitiveComponent, Warning, TEXT("Trying to move static component '%s' after initialization"), *GetFullName());
 		return false;
 	}
 
diff --git a/Engine/Source/Runtime/Engine/Private/SkeletalMeshComponent.cpp b/Engine/Source/Runtime/Engine/Private/SkeletalMeshComponent.cpp
--- a/Engine/Source/Runtime/Engine/Private/SkeletalMeshComponent.cpp
+++ b/Engine/Source/Runtime/Engine/Private/SkeletalMeshComponent.cpp
@@ -476,7 +476,7 @@ bool USkeletalMeshComponent::ShouldTickPose() const
 	// Remote Clients on the Server will always tick animations as updates from the client comes in. To use Client's delta time.
 	const bool bRemoteClientOnServer = CharacterOwner && (CharacterOwner->Role == ROLE_Authority) && CharacterOwner->Controller && !CharacterOwner->Controller->IsLocalController();
 
-	return (Super::ShouldTickPose() && IsRegistered() && AnimScriptInstance && !bPauseAnims && GetWorld()->HasBegunPlay() && !bNoSkeletonUpdate && !bSkipBecauseOfRootMotion && !bRemoteClientOnServer && !bAlreadyTickedThisFrame);
+	return (Super::ShouldTickPose() && IsRegistered() && AnimScriptInstance && !bPauseAnims && GetWorld()->AreActorsInitialized() && !bNoSkeletonUpdate && !bSkipBecauseOfRootMotion && !bRemoteClientOnServer && !bAlreadyTickedThisFrame);
 }
 
 void USkeletalMeshComponent::TickPose(float DeltaTime)
diff --git a/Engine/Source/Runtime/Engine/Private/UnrealEngine.cpp b/Engine/Source/Runtime/Engine/Private/UnrealEngine.cpp
--- a/Engine/Source/Runtime/Engine/Private/UnrealEngine.cpp
+++ b/Engine/Source/Runtime/Engine/Private/UnrealEngine.cpp
@@ -8362,7 +8362,14 @@ bool UEngine::LoadMap( FWorldContext& WorldContext, FURL URL, class UPendingNetG
 	UNavigationSystem::InitializeForWorld(WorldContext.World(), NavigationSystem::GameMode);
 
 	// Initialize gameplay for the level.
-	WorldContext.World()->BeginPlay(URL);
+	WorldContext.World()->InitializeActorsForPlay(URL);
+	WorldContext.World()->BeginPlay();
+
+	AGameMode* const GameMode = WorldContext.World()->GetAuthGameMode();
+	if (GameMode)
+	{
+		GameMode->StartPlay();
+	}
 
 	// Remember the URL. Put this before spawning player controllers so that
 	// a player controller can get the map name during initialization and
diff --git a/Engine/Source/Runtime/Engine/Private/World.cpp b/Engine/Source/Runtime/Engine/Private/World.cpp
--- a/Engine/Source/Runtime/Engine/Private/World.cpp
+++ b/Engine/Source/Runtime/Engine/Private/World.cpp
@@ -906,7 +906,7 @@ void UWorld::RemoveActor(AActor* Actor, bool bShouldModifyLevel)
 			ModifyLevel( CheckLevel );
 		}
 		
-		if (!HasBegunPlay())
+		if (!AreActorsInitialized())
 		{
 			CheckLevel->Actors.ModifyItem(ActorListIndex);
 		}
@@ -1332,7 +1332,6 @@ double InitActorPhysTime = 0.0;
 double InitActorTime = 0.0;
 double RouteActorInitializeTime = 0.0;
 double CrossLevelRefsTime = 0.0;
-double SequenceBeginPlayTime = 0.0;
 double SortActorListTime = 0.0;
 
 /** Helper class, to add the time between this objects creating and destruction to passed in variable. */
@@ -1406,7 +1405,6 @@ void UWorld::AddToWorld( ULevel* Level, const FTransform& LevelTransform )
 		InitActorTime = 0.0;
 		RouteActorInitializeTime = 0.0;
 		CrossLevelRefsTime = 0.0;
-		SequenceBeginPlayTime = 0.0;
 		SortActorListTime = 0.0;
 #endif
 	}
@@ -1463,19 +1461,19 @@ void UWorld::AddToWorld( ULevel* Level, const FTransform& LevelTransform )
 		bExecuteNextStep					= Level->bAreComponentsCurrentlyRegistered;
 	}
 
-	if( IsGameWorld() && HasBegunPlay() )
+	if( IsGameWorld() && AreActorsInitialized() )
 	{
 		// Initialize all actors and start execution.
-		if( bExecuteNextStep && !Level->bAlreadyInitializedActors )
+		if (bExecuteNextStep && !Level->bAlreadyInitializedNetworkActors)
 		{
 			SCOPE_TIME_TO_VAR(&InitActorTime);
 
-			Level->InitializeActors();
-			Level->bAlreadyInitializedActors = true;
-			bExecuteNextStep = (!bConsiderTimeLimit || !IsTimeLimitExceeded( TEXT("initializing actors"), StartTime, Level ));
+			Level->InitializeNetworkActors();
+			Level->bAlreadyInitializedNetworkActors = true;
+			bExecuteNextStep = (!bConsiderTimeLimit || !IsTimeLimitExceeded( TEXT("initializing network actors"), StartTime, Level ));
 		}
 
-		// Route various begin play functions and set volumes.
+		// Route various initialization functions and set volumes.
 		if( bExecuteNextStep && !Level->bAlreadyRoutedActorInitialize )
 		{
 			SCOPE_TIME_TO_VAR(&RouteActorInitializeTime);
@@ -1484,7 +1482,7 @@ void UWorld::AddToWorld( ULevel* Level, const FTransform& LevelTransform )
 			Level->bAlreadyRoutedActorInitialize = true;
 			bStartup = 0;
 
-			bExecuteNextStep = (!bConsiderTimeLimit || !IsTimeLimitExceeded( TEXT("routing BeginPlay on actors"), StartTime, Level ));
+			bExecuteNextStep = (!bConsiderTimeLimit || !IsTimeLimitExceeded( TEXT("routing Initialize on actors"), StartTime, Level ));
 		}
 
 		// Sort the actor list; can't do this on save as the relevant properties for sorting might have been changed by code
@@ -1510,7 +1508,7 @@ void UWorld::AddToWorld( ULevel* Level, const FTransform& LevelTransform )
 	{
 #if PERF_TRACK_DETAILED_ASYNC_STATS
 		// Log out all of the timing information
-		double TotalTime = MoveActorTime + ShiftActorsTime + UpdateComponentsTime + InitBSPPhysTime + InitActorPhysTime + InitActorTime + RouteActorInitializeTime + CrossLevelRefsTime + SequenceBeginPlayTime + SortActorListTime;
+		double TotalTime = MoveActorTime + ShiftActorsTime + UpdateComponentsTime + InitBSPPhysTime + InitActorPhysTime + InitActorTime + RouteActorInitializeTime + CrossLevelRefsTime + SortActorListTime;
 		UE_LOG(LogStreaming, Log, TEXT("Detailed AddToWorld stats for '%s' - Total %6.2fms"), *Level->GetOutermost()->GetName(), TotalTime * 1000 );
 		UE_LOG(LogStreaming, Log, TEXT("Move Actors             : %6.2f ms"), MoveActorTime * 1000 );
 		UE_LOG(LogStreaming, Log, TEXT("Shift Actors            : %6.2f ms"), ShiftActorsTime * 1000 );
@@ -1522,9 +1520,8 @@ void UWorld::AddToWorld( ULevel* Level, const FTransform& LevelTransform )
 		UE_LOG(LogStreaming, Log, TEXT("Init BSP Phys           : %6.2f ms"), InitBSPPhysTime * 1000 );
 		UE_LOG(LogStreaming, Log, TEXT("Init Actor Phys         : %6.2f ms"), InitActorPhysTime * 1000 );
 		UE_LOG(LogStreaming, Log, TEXT("Init Actors             : %6.2f ms"), InitActorTime * 1000 );
-		UE_LOG(LogStreaming, Log, TEXT("BeginPlay               : %6.2f ms"), RouteActorInitializeTime * 1000 );
+		UE_LOG(LogStreaming, Log, TEXT("Initialize               : %6.2f ms"), RouteActorInitializeTime * 1000 );
 		UE_LOG(LogStreaming, Log, TEXT("Cross Level Refs        : %6.2f ms"), CrossLevelRefsTime * 1000 );
-		UE_LOG(LogStreaming, Log, TEXT("Sequence BeginPlay      : %6.2f ms"), SequenceBeginPlayTime * 1000 );
 		UE_LOG(LogStreaming, Log, TEXT("Sort Actor List         : %6.2f ms"), SortActorListTime * 1000 );	
 #endif // PERF_TRACK_DETAILED_ASYNC_STATS
 
@@ -1534,7 +1531,7 @@ void UWorld::AddToWorld( ULevel* Level, const FTransform& LevelTransform )
 		}
 		Level->bAlreadyShiftedActors					= false;
 		Level->bAlreadyUpdatedComponents				= false;
-		Level->bAlreadyInitializedActors				= false;
+		Level->bAlreadyInitializedNetworkActors			= false;
 		Level->bAlreadyRoutedActorInitialize			= false;
 		Level->bAlreadySortedActorList					= false;
 
@@ -2510,8 +2507,7 @@ bool UWorld::SetGameMode(const FURL& InURL)
 	return false;
 }
 
-
-void UWorld::BeginPlay(const FURL& InURL, bool bResetTime)
+void UWorld::InitializeActorsForPlay(const FURL& InURL, bool bResetTime)
 {
 	check(bIsWorldInitialized);
 	double StartTime = FPlatformTime::Seconds();
@@ -2553,7 +2549,7 @@ void UWorld::BeginPlay(const FURL& InURL, bool bResetTime)
 	}
 
 	// Init level gameplay info.
-	if( !HasBegunPlay() )
+	if( !AreActorsInitialized() )
 	{
 		// Check that paths are valid
 		if ( !IsNavigationRebuilt() )
@@ -2571,18 +2567,18 @@ void UWorld::BeginPlay(const FURL& InURL, bool bResetTime)
 			UE_LOG(LogWorld, Log,  TEXT("Bringing %s up for play (max tick rate %i) at %s"), *GetFullName(), FMath::Round(GEngine->GetMaxTickRate(0,false)), *FDateTime::Now().ToString() );
 		}
 
-		// Initialize all actors and start execution.
+		// Initialize network actors and start execution.
 		for( int32 LevelIndex=0; LevelIndex<Levels.Num(); LevelIndex++ )
 		{
 			ULevel*	const Level = Levels[LevelIndex];
-			Level->InitializeActors();
+			Level->InitializeNetworkActors();
 		}
 
 		// Enable actor script calls.
-		bStartup	= 1;
-		bBegunPlay	= 1;
+		bStartup = 1;
+		bActorsInitialized = 1;
 
-		// Spawn serveractors
+		// Spawn server actors
 		ENetMode CurNetMode = GEngine != NULL ? GEngine->GetNetMode(this) : NM_Standalone;
 
 		if (CurNetMode == NM_ListenServer || CurNetMode == NM_DedicatedServer)
@@ -2596,7 +2592,7 @@ void UWorld::BeginPlay(const FURL& InURL, bool bResetTime)
 			AuthorityGameMode->InitGame( FPaths::GetBaseFilename(InURL.Map), Options, Error );
 		}
 
-		// Route various begin play functions and set volumes.
+		// Route various initialization functions and set volumes.
 		for( int32 LevelIndex=0; LevelIndex<Levels.Num(); LevelIndex++ )
 		{
 			ULevel*	const Level = Levels[LevelIndex];
@@ -2627,7 +2623,7 @@ void UWorld::BeginPlay(const FURL& InURL, bool bResetTime)
 	// @note if UWorld starts to host more of these it might a be a good idea to create a generic IManagerInterface of some sort
 	if (NavigationSystem != NULL)
 	{
-		NavigationSystem->OnBeginPlay();
+		NavigationSystem->OnInitializeActors();
 	}
 
 	for( int32 LevelIndex=0; LevelIndex<Levels.Num(); LevelIndex++ )
@@ -2646,7 +2642,14 @@ void UWorld::BeginPlay(const FURL& InURL, bool bResetTime)
 	}
 }
 
-
+void UWorld::BeginPlay()
+{
+	AGameMode* const GameMode = GetAuthGameMode();
+	if (GameMode)
+	{
+		GameMode->StartPlay();
+	}
+}
 
 bool UWorld::IsNavigationRebuilt() const
 {
@@ -2871,6 +2874,10 @@ bool UWorld::HasBegunPlay() const
 	return PersistentLevel && PersistentLevel->Actors.Num() && bBegunPlay;
 }
 
+bool UWorld::AreActorsInitialized() const
+{
+	return PersistentLevel && PersistentLevel->Actors.Num() && bActorsInitialized;
+}
 
 float UWorld::GetTimeSeconds() const
 {
@@ -4284,8 +4291,8 @@ UWorld* FSeamlessTravelHandler::Tick()
 				}
 			}
 
-			// call begin play functions on everything that wasn't carried over from the old world
-			LoadedWorld->BeginPlay(PendingTravelURL, false);
+			// call initialize functions on everything that wasn't carried over from the old world
+			LoadedWorld->InitializeActorsForPlay(PendingTravelURL, false);
 
 			// send loading complete notifications for all local players
 			for (FLocalPlayerIterator It(GEngine, LoadedWorld); It; ++It)
@@ -4333,9 +4340,12 @@ UWorld* FSeamlessTravelHandler::Tick()
 					UNavigationSystem::InitializeForWorld(LoadedWorld, NavigationSystem::GameMode);
 
 					// inform the new GameMode so it can handle players that persisted
-					GameMode->PostSeamlessTravel();
+					GameMode->PostSeamlessTravel();					
 				}
 
+				// Called after post seamless travel to make sure players are setup correctly first
+				LoadedWorld->BeginPlay();
+
 				FCoreDelegates::PostLoadMap.Broadcast();
 			}
 			else
diff --git a/Engine/Source/Runtime/Engine/Private/WorldSettings.cpp b/Engine/Source/Runtime/Engine/Private/WorldSettings.cpp
--- a/Engine/Source/Runtime/Engine/Private/WorldSettings.cpp
+++ b/Engine/Source/Runtime/Engine/Private/WorldSettings.cpp
@@ -118,13 +118,22 @@ float AWorldSettings::GetGravityZ() const
 	return WorldGravityZ;
 }
 
-void AWorldSettings::NotifyMatchStarted()
+void AWorldSettings::NotifyBeginPlay()
 {
 	UWorld* World = GetWorld();
-	for (FActorIterator It(World); It; ++It)
+	if (!World->bBegunPlay)
 	{
-		It->BeginPlay();
+		for (FActorIterator It(World); It; ++It)
+		{
+			It->BeginPlay();
+		}
+		World->bBegunPlay = true;
 	}
+}
+
+void AWorldSettings::NotifyMatchStarted()
+{
+	UWorld* World = GetWorld();
 	World->bMatchStarted = true;
 }
 
diff --git a/Engine/Source/Runtime/Engine/Public/AI/NavDataGenerator.h b/Engine/Source/Runtime/Engine/Public/AI/NavDataGenerator.h
--- a/Engine/Source/Runtime/Engine/Public/AI/NavDataGenerator.h
+++ b/Engine/Source/Runtime/Engine/Public/AI/NavDataGenerator.h
@@ -29,9 +29,6 @@ public:
 	/** Called when world is loaded and all actors are initialized - to be used for post-load steps like data validating */
 	virtual void OnWorldInitDone(bool bAllowedToRebuild) {}	
 
-	/** called from owner's BeginPlay function */
-	virtual void OnPostBeginPlay() {}
-
 	virtual void OnNavigationBuildingLocked() {}
 
 	/** Lets Generator know that building navigation is now allowed. 
diff --git a/Engine/Source/Runtime/SkillSystem/Private/GameplayEffectTests.cpp b/Engine/Source/Runtime/SkillSystem/Private/GameplayEffectTests.cpp
--- a/Engine/Source/Runtime/SkillSystem/Private/GameplayEffectTests.cpp
+++ b/Engine/Source/Runtime/SkillSystem/Private/GameplayEffectTests.cpp
@@ -3998,8 +3998,9 @@ bool FGameplayEffectsTest::RunTest( const FString& Parameters )
 	WorldContext.SetCurrentWorld(World);
 	
 	FURL URL;
-	World->BeginPlay(URL);
-
+	World->InitializeActorsForPlay(URL);
+	World->BeginPlay();
+	
 	GameplayEffectsTest_InstantDamage(World, this);
 	GameplayEffectsTest_InstantDamageRemap(World, this);
 	GameplayEffectsTest_InstantDamage_Buffed(World, this);
