Set the cached mouse position in FSceneViewport::SetMouse
cbuttner
#github https://github.com/EpicGames/UnrealEngine/pull/70
#ttp 333387 [GitHub] 70 : Set the cached mouse position in FSceneViewport::SetMouse

diff --git a/Engine/Source/Runtime/Engine/Private/Slate/SceneViewport.cpp b/Engine/Source/Runtime/Engine/Private/Slate/SceneViewport.cpp
--- a/Engine/Source/Runtime/Engine/Private/Slate/SceneViewport.cpp
+++ b/Engine/Source/Runtime/Engine/Private/Slate/SceneViewport.cpp
@@ -158,6 +158,7 @@ void FSceneViewport::SetMouse( int32 X, int32 Y )
 {
 	FVector2D AbsolutePos = CachedGeometry.LocalToAbsolute(FVector2D(X, Y));
 	FSlateApplication::Get().SetCursorPos( AbsolutePos );
+	CachedMousePos = FIntPoint(X, Y);
 }
 
 void FSceneViewport::UpdateCachedMousePos( const FGeometry& InGeometry, const FPointerEvent& InMouseEvent )
