Add more comments to SpringArmComponent

diff --git a/Engine/Source/Runtime/Engine/Private/GameFramework/SpringArmComponent.cpp b/Engine/Source/Runtime/Engine/Private/GameFramework/SpringArmComponent.cpp
--- a/Engine/Source/Runtime/Engine/Private/GameFramework/SpringArmComponent.cpp
+++ b/Engine/Source/Runtime/Engine/Private/GameFramework/SpringArmComponent.cpp
@@ -28,6 +28,7 @@ USpringArmComponent::USpringArmComponent(const class FPostConstructInitializePro
 	ProbeChannel = ECC_Camera;
 
 	CameraLagSpeed = 10.f;
+	CameraRotationLagSpeed = 10.f;
 
 	RelativeSocketRotation = FQuat::Identity;
 }
@@ -55,14 +56,16 @@ void USpringArmComponent::UpdateDesiredArmLocation(bool bDoTrace, bool bDoLocati
 		}
 	}
 
+	// Apply 'lag' to rotation if desired
 	if(bDoRotationLag)
 	{
 		DesiredRot = FMath::RInterpTo(PreviousDesiredRot, DesiredRot, GAverageMS/1000.f, CameraRotationLagSpeed);
 	}
 	PreviousDesiredRot = DesiredRot;
 
-
+	// Get the spring arm 'origin', the target we want to look at
 	FVector ArmOrigin = GetComponentLocation() + TargetOffset;
+	// We lag the target, not the actuall camera position, so rotating the camera around does not hav lag
 	FVector DesiredLoc = ArmOrigin;
 	if (bDoLocationLag)
 	{
@@ -70,9 +73,12 @@ void USpringArmComponent::UpdateDesiredArmLocation(bool bDoTrace, bool bDoLocati
 	}
 	PreviousDesiredLoc = DesiredLoc;
 
+	// Now offset camera position back along our rotation
 	DesiredLoc -= DesiredRot.Vector() * TargetArmLength;
+	// Add socket offset in local space
 	DesiredLoc += FRotationMatrix(DesiredRot).TransformVector(SocketOffset);
 
+	// Do a sweep to ensure we are not penetrating the world
 	FVector ResultLoc;
 	if (bDoTrace && (TargetArmLength != 0.0f))
 	{
