Move SteerAngle, MaxBrakeTorque, MaxHandBrakeTorque, DampingRate and Mass from FWheelSetup to UVehicleWheel
Remove unused Inertia from UVehicleWheel
Improve comment formatting in UVehicleWheel
Fix TP_VehiclePawn DampingRateFullThrottle (just use default)

diff --git a/Engine/Source/Runtime/Engine/Classes/Vehicles/VehicleWheel.h b/Engine/Source/Runtime/Engine/Classes/Vehicles/VehicleWheel.h
--- a/Engine/Source/Runtime/Engine/Classes/Vehicles/VehicleWheel.h
+++ b/Engine/Source/Runtime/Engine/Classes/Vehicles/VehicleWheel.h
@@ -12,90 +12,113 @@ class ENGINE_API UVehicleWheel : public UObject
 {
 	GENERATED_UCLASS_BODY()
 
-	// Static mesh with collision setup for wheel, will be used to create wheel shape
-	// (if empty, sphere will be added as wheel shape, check bDontCreateShape flag)
+	/** 
+	 * Static mesh with collision setup for wheel, will be used to create wheel shape
+	 * (if empty, sphere will be added as wheel shape, check bDontCreateShape flag)
+	 */
 	UPROPERTY(EditDefaultsOnly, Category=Shape)
 	class UStaticMesh*								CollisionMesh;
 
-	// If set, shape won't be created, but mapped from chassis mesh
+	/** If set, shape won't be created, but mapped from chassis mesh */
 	UPROPERTY(EditDefaultsOnly, Category=Shape)
 	bool											bDontCreateShape;
 
-	// If BoneName is specified, offset the wheel from the bone's location.
-	// Otherwise this offsets the wheel from the vehicle's origin.
-	// (X=Forward, Y=Right, Z=Up)
+	/** 
+	 *	If true, ShapeRadius and ShapeWidth will be used to automatically scale collision taken from CollisionMesh to match wheel size.
+	 *	If false, size of CollisionMesh won't be changed. Use if you want to scale wheels manually.
+	 */
+	UPROPERTY(EditAnywhere, Category=Shape)
+	bool											bAutoAdjustCollisionSize;
+
+	/** 
+	 * If BoneName is specified, offset the wheel from the bone's location.
+	 * Otherwise this offsets the wheel from the vehicle's origin.
+	 */
 	UPROPERTY(EditAnywhere, Category=Wheel)
 	FVector											Offset;
 
-	// Radius of the wheel
+	/** Radius of the wheel */
 	UPROPERTY(EditAnywhere, Category=Wheel)
 	float											ShapeRadius;
 
-	// Width of the wheel
+	/** Width of the wheel */
 	UPROPERTY(EditAnywhere, Category=Wheel)
 	float											ShapeWidth;
 
-	// If true, ShapeRadius and ShapeWidth will be used to automatically scale collision taken from CollisionMesh to match wheel size.
-	// If false, size of CollisionMesh won't be changed. Use if you want to scale wheels manually.
-	UPROPERTY(EditAnywhere, Category=Shape)
-	bool											bAutoAdjustCollisionSize;
+	/** Mass of this wheel */
+	UPROPERTY(EditAnywhere, Category=Wheel)
+	float											Mass;
 
-	// How strongly this wheel opposes acceleration
+	/** Damping rate for this wheel (Kgm^2/s) */
 	UPROPERTY(EditAnywhere, Category=Wheel)
-	float											Inertia;
+	float											DampingRate;
 
-	// Damping rate for how quickly the wheel slows down on its own
-	//UPROPERTY(EditAnywhere, Category=Wheel)
-	//float											DampingRate;
+	// steer angle in degrees for this wheel
+	UPROPERTY(EditAnywhere, Category = WheelsSetup)
+	float SteerAngle;
 
-	// Whether handbrake should affect this wheel
+	/** Whether handbrake should affect this wheel */
 	UPROPERTY(EditAnywhere, Category=Wheel)
 	bool											bAffectedByHandbrake;
 
-	// Tire type for the wheel. Determines friction
+	/** Tire type for the wheel. Determines friction */
 	UPROPERTY(EditAnywhere, Category=Tire)
 	class UTireType*								TireType;
 
-	// Max normalized tire load at which the tire can deliver no more lateral stiffness no matter how much extra load is applied to the tire.
+	/** Max normalized tire load at which the tire can deliver no more lateral stiffness no matter how much extra load is applied to the tire. */
 	UPROPERTY(EditAnywhere, Category=Tire, meta=(ClampMin = "0.01", UIMin = "0.01"))
 	float											LatStiffMaxLoad;
 
-	// How much lateral stiffness to have given lateral slip
+	/** How much lateral stiffness to have given lateral slip */
 	UPROPERTY(EditAnywhere, Category=Tire, meta=(ClampMin = "0.01", UIMin = "0.01"))
 	float											LatStiffValue;
 
-	// Multiplier for the output lateral tire forces
+	/** Multiplier for the output lateral tire forces */
 	UPROPERTY(EditAnywhere, Category=Tire)
 	float											LatStiffFactor;
 
-	// How much longitudinal stiffness to have given longitudinal slip
+	/** How much longitudinal stiffness to have given longitudinal slip */
 	UPROPERTY(EditAnywhere, Category=Tire)
 	float											LongStiffValue;
 
-	// Vertical offset from vehicle center of mass where suspension forces are applied
+	/** Vertical offset from vehicle center of mass where suspension forces are applied */
 	UPROPERTY(EditAnywhere, Category=Suspension)
 	float											SuspensionForceOffset;
 
-	// How far the wheel can drop below the resting position
+	/** How far the wheel can drop below the resting position */
 	UPROPERTY(EditAnywhere, Category=Suspension)
 	float											SuspensionMaxDrop;
 
-	// How far the wheel can go above the resting position
+	/** How far the wheel can go above the resting position */
 	UPROPERTY(EditAnywhere, Category=Suspension)
 	float											SuspensionMaxRaise;
 
-	// Oscillation frequency of suspension. Standard cars have values between 5 and 10
+	/** Oscillation frequency of suspension. Standard cars have values between 5 and 10 */
 	UPROPERTY(EditAnywhere, Category=Suspension)
 	float											SuspensionNaturalFrequency;
 
-	// The rate at which energy is dissipated from the spring. Standard cars have values between 0.8 and 1.2.
-	// values < 1 are more sluggish, values > 1 or more twitchy
+	/**
+	 *	The rate at which energy is dissipated from the spring. Standard cars have values between 0.8 and 1.2.
+	 *	values < 1 are more sluggish, values > 1 or more twitchy
+	 */
 	UPROPERTY(EditAnywhere, Category=Suspension)
 	float											SuspensionDampingRatio;
 
-	// The vehicle that owns us
+	/** max brake torque for this wheel (Nm) */
+	UPROPERTY(EditAnywhere, Category=Brakes)
+	float											MaxBrakeTorque;
+
+	/** 
+	 *	Max handbrake brake torque for this wheel (Nm). A handbrake should have a stronger brake torque
+	 *	than the brake. This will be ignored for wheels that are not affected by the handbrake. 
+	 */
+	UPROPERTY(EditAnywhere, Category=Brakes)
+	float											MaxHandBrakeTorque;
+
+
+	/** The vehicle that owns us */
 	UPROPERTY(transient)
-	class UWheeledVehicleMovementComponent*								VehicleSim;
+	class UWheeledVehicleMovementComponent*			VehicleSim;
 
 	// Our index in the vehicle's (and setup's) wheels array
 	UPROPERTY(transient)
diff --git a/Engine/Source/Runtime/Engine/Classes/Vehicles/WheeledVehicleMovementComponent.h b/Engine/Source/Runtime/Engine/Classes/Vehicles/WheeledVehicleMovementComponent.h
--- a/Engine/Source/Runtime/Engine/Classes/Vehicles/WheeledVehicleMovementComponent.h
+++ b/Engine/Source/Runtime/Engine/Classes/Vehicles/WheeledVehicleMovementComponent.h
@@ -90,43 +90,17 @@ struct FWheelSetup
 	TSubclassOf<UVehicleWheel> WheelClass;
 
 	// Bone name on mesh to create wheel at
-	UPROPERTY(EditAnywhere, Category=WheelsSetup)
+	UPROPERTY(EditAnywhere, Category=WheelSetup)
 	FName BoneName;
 
 	// Additional offset to give the wheels for this axle.
-	UPROPERTY(EditAnywhere, Category=WheelsSetup)
+	UPROPERTY(EditAnywhere, Category=WheelSetup)
 	FVector AdditionalOffset;
 
-	// steer angle in degrees for this wheel
-	UPROPERTY(EditAnywhere, Category=WheelsSetup)
-	float SteerAngle;
-
-	// max brake torque for this wheel (Nm)
-	UPROPERTY(EditAnywhere, Category=WheelsSetup)
-	float MaxBrakeTorque;
-
-	/** Max handbrake brake torque for this wheel (Nm). A handbrake should have a stronger brake torque
-		than the brake. This will be ignored for wheels that are not affected by the handbrake. */
-	UPROPERTY(EditAnywhere, Category=WheelsSetup)
-	float MaxHandBrakeTorque;
-
-	// damping rate for this wheel (Kgm^2/s)
-	UPROPERTY(EditAnywhere, Category=WheelsSetup)
-	float DampingRate;
-
-	// mass of this wheel
-	UPROPERTY(EditAnywhere, Category=WheelsSetup)
-	float Mass;
-
 	FWheelSetup()
 		: WheelClass(UVehicleWheel::StaticClass())
 		, BoneName(NAME_None)
 		, AdditionalOffset(0.0f)
-		, SteerAngle(70.0f)
-		, MaxBrakeTorque(1500.f)
-		, MaxHandBrakeTorque(3000.f)
-		, DampingRate(0.25f)
-		, Mass(20.f)
 	{
 	}
 };
@@ -368,8 +342,6 @@ class ENGINE_API UWheeledVehicleMovementComponent : public UPawnMovementComponen
 	UFUNCTION(BlueprintCallable, Category="Game|Components|WheeledVehicleMovement")
 	bool GetUseAutoGears() const;
 
-	virtual void Serialize(FArchive & Ar) OVERRIDE;
-
 protected:
 
 	// replicated state of vehicle 
diff --git a/Engine/Source/Runtime/Engine/Private/Vehicles/VehicleWheel.cpp b/Engine/Source/Runtime/Engine/Private/Vehicles/VehicleWheel.cpp
--- a/Engine/Source/Runtime/Engine/Private/Vehicles/VehicleWheel.cpp
+++ b/Engine/Source/Runtime/Engine/Private/Vehicles/VehicleWheel.cpp
@@ -20,10 +20,12 @@ UVehicleWheel::UVehicleWheel(const class FPostConstructInitializeProperties& PCI
 	ShapeRadius = 30.0f;
 	ShapeWidth = 10.0f;
 	bAutoAdjustCollisionSize = true;
-	//Mass = 10.0f;
-	Inertia = 1.0f;
+	Mass = 20.0f;
 	bAffectedByHandbrake = true;
-	//DampingRate = 100.0f;
+	SteerAngle = 70.0f;
+	MaxBrakeTorque = 1500.f;
+	MaxHandBrakeTorque = 3000.f;
+	DampingRate = 0.25f;
 	LatStiffMaxLoad = 2.0f;
 	LatStiffValue = 17.0f;
 	LatStiffFactor = 1.0f;
diff --git a/Engine/Source/Runtime/Engine/Private/Vehicles/WheeledVehicleMovementComponent.cpp b/Engine/Source/Runtime/Engine/Private/Vehicles/WheeledVehicleMovementComponent.cpp
--- a/Engine/Source/Runtime/Engine/Private/Vehicles/WheeledVehicleMovementComponent.cpp
+++ b/Engine/Source/Runtime/Engine/Private/Vehicles/WheeledVehicleMovementComponent.cpp
@@ -284,12 +284,12 @@ void UWheeledVehicleMovementComponent::SetupWheels( PxVehicleWheelsSimData* PWhe
 		PxVehicleWheelData PWheelData;
 		PWheelData.mRadius = Wheel->ShapeRadius;
 		PWheelData.mWidth = Wheel->ShapeWidth;	
-		PWheelData.mMaxSteer = FMath::DegreesToRadians(WheelSetups[WheelIdx].SteerAngle);
-		PWheelData.mMaxBrakeTorque = M2ToCm2(WheelSetups[WheelIdx].MaxBrakeTorque);
-		PWheelData.mMaxHandBrakeTorque = Wheel->bAffectedByHandbrake ? M2ToCm2(WheelSetups[WheelIdx].MaxHandBrakeTorque) : 0.0f;
+		PWheelData.mMaxSteer = FMath::DegreesToRadians(Wheel->SteerAngle);
+		PWheelData.mMaxBrakeTorque = M2ToCm2(Wheel->MaxBrakeTorque);
+		PWheelData.mMaxHandBrakeTorque = Wheel->bAffectedByHandbrake ? M2ToCm2(Wheel->MaxHandBrakeTorque) : 0.0f;
 
-		PWheelData.mDampingRate = M2ToCm2(WheelSetups[WheelIdx].DampingRate);
-		PWheelData.mMass = WheelSetups[WheelIdx].Mass;
+		PWheelData.mDampingRate = M2ToCm2(Wheel->DampingRate);
+		PWheelData.mMass = Wheel->Mass;
 		PWheelData.mMOI = 0.5f * PWheelData.mMass * FMath::Square(PWheelData.mRadius);
 		
 		// init tire data
@@ -1295,29 +1295,3 @@ void UWheeledVehicleMovementComponent::ComputeConstants()
 	MaxEngineRPM = 5000.f;
 }
 
-void UWheeledVehicleMovementComponent::Serialize(FArchive & Ar)
-{
-	Super::Serialize(Ar);
-	if (Ar.IsLoading() && Ar.UE4Ver() < VER_UE4_VEHICLES_UNIT_CHANGE)
-	{
-		for (int32 WheelIdx = 0; WheelIdx < WheelSetups.Num(); ++WheelIdx)
-		{
-			//this backwards compatable code breaks if they were using some strange value that is exactly the new default in the new units
-			FWheelSetup & WheelSetup = WheelSetups[WheelIdx];
-			WheelSetup.MaxHandBrakeTorque = WheelSetup.MaxHandBrakeTorque != 1500.f ? Cm2ToM2(WheelSetup.MaxHandBrakeTorque) : 1500.f;	//need to convert from cm^2 to m^2
-			WheelSetup.MaxBrakeTorque = WheelSetup.MaxHandBrakeTorque != 3000.f ? Cm2ToM2(WheelSetup.MaxBrakeTorque) : 3000.f;			//need to convert from cm^2 to m^2
-		}
-
-	}
-
-	if (Ar.IsLoading() && Ar.UE4Ver() < VER_UE4_VEHICLES_UNIT_CHANGE2)
-	{
-		for (int32 WheelIdx = 0; WheelIdx < WheelSetups.Num(); ++WheelIdx)
-		{
-			//this backwards compatable code breaks if they were using some strange value that is exactly the new default in the new units
-			FWheelSetup & WheelSetup = WheelSetups[WheelIdx];
-			WheelSetup.DampingRate = WheelSetup.DampingRate != 0.25f ? Cm2ToM2(WheelSetup.DampingRate) : 0.25;
-		}
-
-	}
-}
