Change UVehicleWheel default LongStiffValue to be 1000
Use default braking force in C++ vehicle template

diff --git a/Engine/Source/Runtime/Engine/Private/Vehicles/VehicleWheel.cpp b/Engine/Source/Runtime/Engine/Private/Vehicles/VehicleWheel.cpp
--- a/Engine/Source/Runtime/Engine/Private/Vehicles/VehicleWheel.cpp
+++ b/Engine/Source/Runtime/Engine/Private/Vehicles/VehicleWheel.cpp
@@ -29,7 +29,7 @@ UVehicleWheel::UVehicleWheel(const class FPostConstructInitializeProperties& PCI
 	LatStiffMaxLoad = 2.0f;
 	LatStiffValue = 17.0f;
 	LatStiffFactor = 1.0f;
-	LongStiffValue = 10.0f;
+	LongStiffValue = 1000.0f;
 	SuspensionForceOffset = 0.0f;
 	SuspensionMaxRaise = 20.0f;
 	SuspensionMaxDrop = 20.0f;
