Uncomment the lines which register the screenshot handler callback.
#ttp 333848 Editor: Matinee: JPG video output incorrectly outputs BMPs

diff --git a/Engine/Source/Runtime/Engine/Private/UnrealEngine.cpp b/Engine/Source/Runtime/Engine/Private/UnrealEngine.cpp
--- a/Engine/Source/Runtime/Engine/Private/UnrealEngine.cpp
+++ b/Engine/Source/Runtime/Engine/Private/UnrealEngine.cpp
@@ -731,7 +731,7 @@ void UEngine::Init(IEngineLoop* InEngineLoop)
 	// register screenshot capture if we are dumping a movie
 	if(GIsDumpingMovie)
 	{
-		//UGameViewportClient::OnScreenshotCaptured().AddUObject(this, &UEngine::HandleScreenshotCaptured);
+		UGameViewportClient::OnScreenshotCaptured().AddUObject(this, &UEngine::HandleScreenshotCaptured);
 	}
 #endif
 
@@ -809,7 +809,7 @@ void UEngine::PreExit()
 	FEngineAnalytics::Shutdown();
 
 #if WITH_EDITOR
-	//UGameViewportClient::OnScreenshotCaptured().RemoveUObject(this, &UEngine::HandleScreenshotCaptured);
+	UGameViewportClient::OnScreenshotCaptured().RemoveUObject(this, &UEngine::HandleScreenshotCaptured);
 #endif
 }
 
