NaN safety checks in particle emitter trail

#UE4

diff --git a/Engine/Source/Runtime/Engine/Private/Particles/ParticleTrail2EmitterInstance.cpp b/Engine/Source/Runtime/Engine/Private/Particles/ParticleTrail2EmitterInstance.cpp
--- a/Engine/Source/Runtime/Engine/Private/Particles/ParticleTrail2EmitterInstance.cpp
+++ b/Engine/Source/Runtime/Engine/Private/Particles/ParticleTrail2EmitterInstance.cpp
@@ -973,9 +973,22 @@ bool FParticleRibbonEmitterInstance::GetSpawnPerUnitAmount(float DeltaTime, int3
 				float ElapsedTime = RunningTime;//SecondsSinceCreation;
 				if (ActiveParticles == 0)
 				{
+					if (ElapsedTime == 0)
+					{
+						ElapsedTime = KINDA_SMALL_NUMBER;
+					}
+					CurrentSourcePosition[InTrailIdx].DiagnosticCheckNaN();
+					LastSourcePosition[InTrailIdx].DiagnosticCheckNaN();
+
 					LastSourceTangent[InTrailIdx] = (CurrentSourcePosition[InTrailIdx] - LastSourcePosition[InTrailIdx]) / ElapsedTime;
 				}
-				FVector CurrTangent = TravelDirection / (ElapsedTime - TrailSpawnTimes[InTrailIdx]);
+
+				float CurrTangentDivisor = (ElapsedTime - TrailSpawnTimes[InTrailIdx]);
+				if (CurrTangentDivisor == 0)
+				{
+					CurrTangentDivisor = KINDA_SMALL_NUMBER;
+				}
+				FVector CurrTangent = TravelDirection / CurrTangentDivisor;
 				CurrTangent.Normalize();
 				FVector PrevTangent = LastSourceTangent[InTrailIdx];
 				PrevTangent.Normalize();
