FAsyncParticleFill - Fixed linker error, removed not needed SCOPE_CYCLE_COUNTER in the ::DoTask(...)

diff --git a/Engine/Source/Runtime/Engine/Private/Particles/ParticleSystemRender.cpp b/Engine/Source/Runtime/Engine/Private/Particles/ParticleSystemRender.cpp
--- a/Engine/Source/Runtime/Engine/Private/Particles/ParticleSystemRender.cpp
+++ b/Engine/Source/Runtime/Engine/Private/Particles/ParticleSystemRender.cpp
@@ -553,7 +553,6 @@ FORCEINLINE FVector GetCameraOffsetFromPayload(
 void FAsyncParticleFill::DoTask(ENamedThreads::Type CurrentThread, const FGraphEventRef& MyCompletionGraphEvent)
 {
 	// Test integrity of async particle fill: FPlatformProcess::Sleep(.01);
-	SCOPE_CYCLE_COUNTER(STAT_ParticleAsyncTime);
 	Parent->DoBufferFill();
 }
 
diff --git a/Engine/Source/Runtime/Engine/Public/ParticleHelper.h b/Engine/Source/Runtime/Engine/Public/ParticleHelper.h
--- a/Engine/Source/Runtime/Engine/Public/ParticleHelper.h
+++ b/Engine/Source/Runtime/Engine/Public/ParticleHelper.h
@@ -1080,7 +1080,7 @@ struct FAsyncParticleFill
 
 	FORCEINLINE TStatId GetStatId() const
 	{
-		RETURN_QUICK_DECLARE_CYCLE_STAT(FAsyncParticleFill, STATGROUP_TaskGraphTasks);
+		return GET_STATID( STAT_ParticleAsyncTime );
 	}
 
 	static ENamedThreads::Type GetDesiredThread()
