Billboards needed MarkRenderStateDirty

diff --git a/Engine/Source/Runtime/Engine/Private/MaterialBillboardComponent.cpp b/Engine/Source/Runtime/Engine/Private/MaterialBillboardComponent.cpp
--- a/Engine/Source/Runtime/Engine/Private/MaterialBillboardComponent.cpp
+++ b/Engine/Source/Runtime/Engine/Private/MaterialBillboardComponent.cpp
@@ -274,4 +274,6 @@ void UMaterialBillboardComponent::AddElement(
 	Element->BaseSizeX = BaseSizeX;
 	Element->BaseSizeY = BaseSizeY;
 	Element->DistanceToSizeCurve = DistanceToSizeCurve;
+
+	MarkRenderStateDirty();
 }
