#ue4 - Fix various spectator spawning issues

- Fix extra spectators being spawned on ListenServers
- Respect SpectatorOnly flag correctly on clients.

diff --git a/Engine/Source/Runtime/Engine/Private/GameMode.cpp b/Engine/Source/Runtime/Engine/Private/GameMode.cpp
--- a/Engine/Source/Runtime/Engine/Private/GameMode.cpp
+++ b/Engine/Source/Runtime/Engine/Private/GameMode.cpp
@@ -1319,7 +1319,7 @@ void AGameMode::StartNewPlayer(APlayerController* NewPlayer)
 	NewPlayer->ClientSetHUD(HUDClass);
 
 	// If players should start as spectators, leave them in the spectator state
-	if (!bStartPlayersAsSpectators)
+	if (!bStartPlayersAsSpectators && !NewPlayer->PlayerState->bOnlySpectator)
 	{
 		// If match is in progress, start the player
 		if (IsMatchInProgress())
diff --git a/Engine/Source/Runtime/Engine/Private/PlayerController.cpp b/Engine/Source/Runtime/Engine/Private/PlayerController.cpp
--- a/Engine/Source/Runtime/Engine/Private/PlayerController.cpp
+++ b/Engine/Source/Runtime/Engine/Private/PlayerController.cpp
@@ -680,7 +680,7 @@ void APlayerController::ReceivedPlayer()
 	{
 		if (GetSpectatorPawn() == NULL)
 		{
-			SetSpectatorPawn(SpawnSpectatorPawn());
+			BeginSpectatingState();
 		}
 	}
 }
@@ -802,7 +802,7 @@ void APlayerController::PostInitializeComponents()
 	AddCheats();
 
 	bPlayerIsWaiting = true;
-	ChangeState(NAME_Spectating);
+	StateName = NAME_Spectating; // Don't use ChangeState, because we want to defer spawning the SpectatorPawn until the Player is received
 
 #if !(UE_BUILD_SHIPPING || UE_BUILD_TEST)
 	if (!IsPendingKill() && DebuggingController == NULL && GetNetMode() != NM_DedicatedServer)
@@ -3261,7 +3261,7 @@ void APlayerController::ReceivedSpectatorClass(TSubclassOf<AGameMode> SpectatorC
 	{
 		if (GetSpectatorPawn() == NULL)
 		{
-			SetSpectatorPawn(SpawnSpectatorPawn());
+			BeginSpectatingState();
 		}
 	}
 }
@@ -3563,6 +3563,8 @@ ASpectatorPawn* APlayerController::SpawnSpectatorPawn()
 				SpawnedSpectator->PossessedBy(this);
 				SpawnedSpectator->PawnClientRestart();
 				SpawnedSpectator->SetActorTickEnabled(true);
+
+				UE_LOG(LogPlayerController, Log, TEXT("Spawned spectator %s [server:%d]"), *GetNameSafe(SpawnedSpectator), GetNetMode() < NM_Client);
 			}
 			else
 			{
@@ -3571,7 +3573,8 @@ ASpectatorPawn* APlayerController::SpawnSpectatorPawn()
 		}
 		else
 		{
-			UE_LOG(LogPlayerController, Warning, TEXT("NULL GameState when trying to spawn spectator!"));
+			// This normally happens on clients if the Player is replicated but the GameState has not yet.
+			UE_LOG(LogPlayerController, Log, TEXT("NULL GameState when trying to spawn spectator!"));
 		}
 	}
 
