Fix material nodes sometimes being flipped in the editor.

diff --git a/Engine/Source/Runtime/Engine/Classes/Materials/MaterialFunction.h b/Engine/Source/Runtime/Engine/Classes/Materials/MaterialFunction.h
--- a/Engine/Source/Runtime/Engine/Classes/Materials/MaterialFunction.h
+++ b/Engine/Source/Runtime/Engine/Classes/Materials/MaterialFunction.h
@@ -59,6 +59,7 @@ public:
 #if WITH_EDITOR
 	virtual void PostEditChangeProperty(FPropertyChangedEvent& PropertyChangedEvent) OVERRIDE;
 #endif // WITH_EDITOR
+	virtual void Serialize(FArchive& Ar) OVERRIDE;
 	virtual void PostLoad() OVERRIDE;
 	virtual void GetAssetRegistryTags(TArray<FAssetRegistryTag>& OutTags) const OVERRIDE;
 	// End UObject interface.
diff --git a/Engine/Source/Runtime/Engine/Private/Materials/Material.cpp b/Engine/Source/Runtime/Engine/Private/Materials/Material.cpp
--- a/Engine/Source/Runtime/Engine/Private/Materials/Material.cpp
+++ b/Engine/Source/Runtime/Engine/Private/Materials/Material.cpp
@@ -31,6 +31,7 @@ bool FMaterialsWithDirtyUsageFlags::IsUsageFlagDirty(EMaterialUsage UsageFlag)
 FUObjectAnnotationSparseBool GMaterialsThatNeedSamplerFixup;
 FUObjectAnnotationSparseBool GMaterialsThatNeedPhysicalConversion;
 FUObjectAnnotationSparse<FMaterialsWithDirtyUsageFlags,true> GMaterialsWithDirtyUsageFlags;
+FUObjectAnnotationSparseBool GMaterialsThatNeedExpressionsFlipped;
 
 #endif // #if WITH_EDITOR
 
@@ -1612,6 +1613,11 @@ void UMaterial::Serialize(FArchive& Ar)
 	{
 		GMaterialsThatNeedPhysicalConversion.Set( this );
 	}
+
+	if (Ar.UE4Ver() < VER_UE4_FLIP_MATERIAL_COORDS)
+	{
+		GMaterialsThatNeedExpressionsFlipped.Set(this);
+	}
 #endif // #if WITH_EDITOR
 
 	if( Ar.UE4Ver() < VER_UE4_MATERIAL_ATTRIBUTES_REORDERING )
@@ -1986,12 +1992,13 @@ void UMaterial::PostLoad()
 		}
 	}
 
-#if WITH_EDITORONLY_DATA
-	if (GetLinkerUE4Version() < VER_UE4_FLIP_MATERIAL_COORDS)
+#if WITH_EDITOR
+	if (GMaterialsThatNeedExpressionsFlipped.Get(this))
 	{
+		GMaterialsThatNeedExpressionsFlipped.Clear(this);
 		FlipExpressionPositions(Expressions, EditorComments, this);
 	}
-#endif //WITH_EDITORONLY_DATA
+#endif // #if WITH_EDITOR
 }
 
 void UMaterial::BeginCacheForCookedPlatformData( const ITargetPlatform *TargetPlatform )
diff --git a/Engine/Source/Runtime/Engine/Private/Materials/MaterialExpressions.cpp b/Engine/Source/Runtime/Engine/Private/Materials/MaterialExpressions.cpp
--- a/Engine/Source/Runtime/Engine/Private/Materials/MaterialExpressions.cpp
+++ b/Engine/Source/Runtime/Engine/Private/Materials/MaterialExpressions.cpp
@@ -9,6 +9,7 @@
 #if WITH_EDITOR
 #include "Slate.h"
 #include "UnrealEd.h"
+#include "UObjectAnnotation.h"
 #endif //WITH_EDITOR
 
 #define LOCTEXT_NAMESPACE "MaterialExpression"
@@ -19,6 +20,10 @@ if( ExpressionInput.Expression == ToBeRemovedExpression )										\
 	ExpressionInput.Expression = ToReplaceWithExpression;										\
 }
 
+#if WITH_EDITOR
+FUObjectAnnotationSparseBool GMaterialFunctionsThatNeedExpressionsFlipped;
+#endif // #if WITH_EDITOR
+
 /** Returns whether the given expression class is allowed. */
 bool IsAllowedExpressionType(UClass* Class, bool bMaterialFunction)
 {
@@ -6002,6 +6007,17 @@ void UMaterialFunction::PostEditChangeProperty(FPropertyChangedEvent& PropertyCh
 }
 #endif // WITH_EDITOR
 
+void UMaterialFunction::Serialize(FArchive& Ar)
+{
+	Super::Serialize(Ar);
+#if WITH_EDITOR
+	if (Ar.UE4Ver() < VER_UE4_FLIP_MATERIAL_COORDS)
+	{
+		GMaterialFunctionsThatNeedExpressionsFlipped.Set(this);
+	}
+#endif // #if WITH_EDITOR
+}
+
 void UMaterialFunction::PostLoad()
 {
 	Super::PostLoad();
@@ -6038,13 +6054,13 @@ void UMaterialFunction::PostLoad()
 			StateId = FGuid::NewGuid();
 		}
 	}
-#endif
-#if WITH_EDITORONLY_DATA
-	if (GetLinkerUE4Version() < VER_UE4_FLIP_MATERIAL_COORDS)
+
+	if (GMaterialFunctionsThatNeedExpressionsFlipped.Get(this))
 	{
+		GMaterialFunctionsThatNeedExpressionsFlipped.Clear(this);
 		UMaterial::FlipExpressionPositions(FunctionExpressions, FunctionEditorComments);
 	}
-#endif //WITH_EDITORONLY_DATA
+#endif // #if WITH_EDITOR
 }
 
 void UMaterialFunction::GetAssetRegistryTags(TArray<FAssetRegistryTag>& OutTags) const
