Build break fix

#UE4

diff --git a/Engine/Source/Runtime/Engine/Private/Texture.cpp b/Engine/Source/Runtime/Engine/Private/Texture.cpp
--- a/Engine/Source/Runtime/Engine/Private/Texture.cpp
+++ b/Engine/Source/Runtime/Engine/Private/Texture.cpp
@@ -286,8 +286,10 @@ void UTexture::PostLoad()
 
 	if( !IsTemplate() )
 	{
+#if WITH_EDITORONLY_DATA
 		// Update cached LOD bias.
 		UpdateCachedLODBias();
+#endif // #if WITH_EDITORONLY_DATA
 
 		// The texture will be cached by the cubemap it is contained within on consoles.
 		UTextureCube* CubeMap = Cast<UTextureCube>(GetOuter());
@@ -441,6 +443,7 @@ bool UTexture::ForceUpdateTextureStreaming()
 		const bool bOldOnlyStreamInTextures = CVarOnlyStreamInTextures->GetInt() != 0;
 		CVarOnlyStreamInTextures->Set(false);
 
+#if WITH_EDITORONLY_DATA
 		for( TObjectIterator<UTexture2D> It; It; ++It )
 		{
 			UTexture* Texture = *It;
@@ -448,6 +451,7 @@ bool UTexture::ForceUpdateTextureStreaming()
 			// Update cached LOD bias.
 			Texture->UpdateCachedLODBias();
 		}
+#endif // #if WITH_EDITORONLY_DATA
 
 		// Make sure we iterate over all textures by setting it to high value.
 		GStreamingManager->SetNumIterationsForNextFrame( 100 );
@@ -587,21 +591,27 @@ void FTextureLODSettings::ReadEntry( int32 GroupId, const TCHAR* GroupName, cons
 
 int32 FTextureLODSettings::CalculateLODBias( UTexture* Texture, bool bIncTextureMips ) const
 {	
-	check( Texture );
-	return CalculateLODBias(Texture->GetSurfaceWidth(), Texture->GetSurfaceHeight(), Texture->LODGroup, (bIncTextureMips ? Texture->LODBias : 0), (bIncTextureMips ? Texture->NumCinematicMipLevels : 0), Texture->MipGenSettings);
-}
+	check(Texture);
 
-int32 FTextureLODSettings::CalculateLODBias( int32 Width, int32 Height, int32 LODGroup, int32 LODBias, int32 NumCinematicMipLevels, TextureMipGenSettings InMipGenSetting ) const
-{	
+#if WITH_EDITORONLY_DATA
 	// Find LOD group.
-	const FTextureLODGroup& LODGroupInfo = TextureLODGroups[LODGroup];
+	const FTextureLODGroup& LODGroupInfo = TextureLODGroups[Texture->LODGroup];
 
 	// Test to see if we have no mip generation as in which case the LOD bias will be ignored
-	const TextureMipGenSettings FinalMipGenSetting = (InMipGenSetting == TMGS_FromTextureGroup) ? LODGroupInfo.MipGenSettings : InMipGenSetting;
-	if ( FinalMipGenSetting == TMGS_NoMipmaps )
+	const TextureMipGenSettings FinalMipGenSetting = (Texture->MipGenSettings == TMGS_FromTextureGroup) ? LODGroupInfo.MipGenSettings : Texture->MipGenSettings;
+	if (FinalMipGenSetting == TMGS_NoMipmaps)
 	{
 		return 0;
 	}
+#endif // #if WITH_EDITORONLY_DATA
+
+	return CalculateLODBias(Texture->GetSurfaceWidth(), Texture->GetSurfaceHeight(), Texture->LODGroup, (bIncTextureMips ? Texture->LODBias : 0), (bIncTextureMips ? Texture->NumCinematicMipLevels : 0));
+}
+
+int32 FTextureLODSettings::CalculateLODBias( int32 Width, int32 Height, int32 LODGroup, int32 LODBias, int32 NumCinematicMipLevels ) const
+{	
+	// Find LOD group.
+	const FTextureLODGroup& LODGroupInfo = TextureLODGroups[LODGroup];
 
 	// Calculate maximum number of miplevels.
 	int32 TextureMaxLOD	= FMath::CeilLogTwo( FMath::Trunc( FMath::Max( Width, Height ) ) );
diff --git a/Engine/Source/Runtime/Engine/Private/Texture2D.cpp b/Engine/Source/Runtime/Engine/Private/Texture2D.cpp
--- a/Engine/Source/Runtime/Engine/Private/Texture2D.cpp
+++ b/Engine/Source/Runtime/Engine/Private/Texture2D.cpp
@@ -522,7 +522,9 @@ FString UTexture2D::GetDesc()
 	uint32 EffectiveSizeX;
 	uint32 EffectiveSizeY;
 
+#if WITH_EDITORONLY_DATA
 	UpdateCachedLODBias();
+#endif //#if WITH_EDITORONLY_DATA
 
 	GSystemSettings.TextureLODSettings.ComputeInGameMaxResolution(GetCachedLODBias(), *this, EffectiveSizeX, EffectiveSizeY);
 
diff --git a/Engine/Source/Runtime/Engine/Private/TextureDerivedData.cpp b/Engine/Source/Runtime/Engine/Private/TextureDerivedData.cpp
--- a/Engine/Source/Runtime/Engine/Private/TextureDerivedData.cpp
+++ b/Engine/Source/Runtime/Engine/Private/TextureDerivedData.cpp
@@ -329,7 +329,7 @@ static void GetTextureBuildSettings(
 	OutBuildSettings.bFlipGreenChannel = Texture.bFlipGreenChannel;
 	OutBuildSettings.CompositeTextureMode = Texture.CompositeTextureMode;
 	OutBuildSettings.CompositePower = Texture.CompositePower;
-	OutBuildSettings.LODBias = GSystemSettings.TextureLODSettings.CalculateLODBias(Texture.Source.GetSizeX(), Texture.Source.GetSizeY(), Texture.LODGroup, Texture.LODBias, Texture.NumCinematicMipLevels, Texture.MipGenSettings);
+	OutBuildSettings.LODBias = GSystemSettings.TextureLODSettings.CalculateLODBias(Texture);
 	OutBuildSettings.bStreamable = !Texture.NeverStream && (Texture.LODGroup != TEXTUREGROUP_UI) && (Cast<const UTexture2D>(&Texture) != NULL);
 }
 
@@ -1197,7 +1197,7 @@ static void SerializePlatformData(
 			const int32 NumCinematicMipLevels = Texture->NumCinematicMipLevels;
 			const TextureMipGenSettings MipGenSetting = Texture->MipGenSettings;
 
-			FirstMipToSerialize = Ar.CookingTarget()->GetTextureLODSettings().CalculateLODBias(Width, Height, LODGroup, LODBias, NumCinematicMipLevels, MipGenSetting);
+			FirstMipToSerialize = Ar.CookingTarget()->GetTextureLODSettings().CalculateLODBias(this);
 			FirstMipToSerialize = FMath::Clamp(FirstMipToSerialize, 0, FMath::Max(NumMips-1,0));
 			NumMips -= FirstMipToSerialize;
 		}
diff --git a/Engine/Source/Runtime/Engine/Public/Texture.h b/Engine/Source/Runtime/Engine/Public/Texture.h
--- a/Engine/Source/Runtime/Engine/Public/Texture.h
+++ b/Engine/Source/Runtime/Engine/Public/Texture.h
@@ -913,10 +913,9 @@ struct FTextureLODSettings
 	 * @param	LODGroup					Which LOD group the texture belongs to
 	 * @param	LODBias						LOD bias to include in the calculation
 	 * @param	NumCinematicMipLevels		The texture cinematic mip levels to include in the calculation
-	 * @param	MipGenSetting				Mip generation setting
 	 * @return	LOD bias
 	 */
-	ENGINE_API int32 CalculateLODBias( int32 Width, int32 Height, int32 LODGroup, int32 LODBias, int32 NumCinematicMipLevels, TextureMipGenSettings MipGenSetting ) const;
+	ENGINE_API int32 CalculateLODBias( int32 Width, int32 Height, int32 LODGroup, int32 LODBias, int32 NumCinematicMipLevels ) const;
 
 	/** 
 	* Useful for stats in the editor.
