#UE4 Move it so BeginPlay now happens after the initial player controllers have spawned, and after pawns have spawned if the match state is ready. This restores it so BeginPlay happens at the same point it used to in modes where there's no delayed start.

diff --git a/Engine/Source/Runtime/Engine/Private/GameMode.cpp b/Engine/Source/Runtime/Engine/Private/GameMode.cpp
--- a/Engine/Source/Runtime/Engine/Private/GameMode.cpp
+++ b/Engine/Source/Runtime/Engine/Private/GameMode.cpp
@@ -486,6 +486,12 @@ void AGameMode::StartPlay()
 	{
 		SetMatchState(MatchState::WaitingToStart);
 	}
+
+	// Check to see if we should immediately transfer to match start
+	if (MatchState == MatchState::WaitingToStart && ReadyToStartMatch())
+	{
+		StartMatch();
+	}
 }
 
 void AGameMode::HandleMatchIsWaitingToStart()
@@ -495,8 +501,12 @@ void AGameMode::HandleMatchIsWaitingToStart()
 		GameSession->HandleMatchIsWaitingToStart();
 	}
 
-	// Calls begin play on actors
-	GetWorldSettings()->NotifyBeginPlay();
+	// Calls begin play on actors, unless we're about to transition to match start
+
+	if (!ReadyToStartMatch())
+	{
+		GetWorldSettings()->NotifyBeginPlay();
+	}
 }
 
 bool AGameMode::ReadyToStartMatch()
@@ -553,7 +563,10 @@ void AGameMode::HandleMatchHasStarted()
 	// Make sure level streaming is up to date before triggering NotifyMatchStarted
 	GetWorld()->FlushLevelStreaming();
 
-	// fire off any level startup events
+	// First fire BeginPlay, if we haven't already in waiting to start match
+	GetWorldSettings()->NotifyBeginPlay();
+
+	// Then fire off match started
 	GetWorldSettings()->NotifyMatchStarted();
 
 	// if passed in bug info, send player to right location
diff --git a/Engine/Source/Runtime/Engine/Private/UnrealEngine.cpp b/Engine/Source/Runtime/Engine/Private/UnrealEngine.cpp
--- a/Engine/Source/Runtime/Engine/Private/UnrealEngine.cpp
+++ b/Engine/Source/Runtime/Engine/Private/UnrealEngine.cpp
@@ -8363,13 +8363,6 @@ bool UEngine::LoadMap( FWorldContext& WorldContext, FURL URL, class UPendingNetG
 
 	// Initialize gameplay for the level.
 	WorldContext.World()->InitializeActorsForPlay(URL);
-	WorldContext.World()->BeginPlay();
-
-	AGameMode* const GameMode = WorldContext.World()->GetAuthGameMode();
-	if (GameMode)
-	{
-		GameMode->StartPlay();
-	}
 
 	// Remember the URL. Put this before spawning player controllers so that
 	// a player controller can get the map name during initialization and
@@ -8400,6 +8393,8 @@ bool UEngine::LoadMap( FWorldContext& WorldContext, FURL URL, class UPendingNetG
 
 	// Prime texture streaming.
 	GStreamingManager->NotifyLevelChange();
+
+	WorldContext.World()->BeginPlay();
 	}
 
 	// send a callback message
