Ensure that the OwnedComponents array is correct after duplication

diff --git a/Engine/Source/Runtime/Engine/Classes/GameFramework/Actor.h b/Engine/Source/Runtime/Engine/Classes/GameFramework/Actor.h
--- a/Engine/Source/Runtime/Engine/Classes/GameFramework/Actor.h
+++ b/Engine/Source/Runtime/Engine/Classes/GameFramework/Actor.h
@@ -1954,6 +1954,8 @@ private:
 	UPROPERTY(transient, duplicatetransient)
 	TArray<UActorComponent*> OwnedComponents;
 
+	void ResetOwnedComponents();
+
 public:
 	//=============================================================================
 	// Navigation related functions
diff --git a/Engine/Source/Runtime/Engine/Private/Actor.cpp b/Engine/Source/Runtime/Engine/Private/Actor.cpp
--- a/Engine/Source/Runtime/Engine/Private/Actor.cpp
+++ b/Engine/Source/Runtime/Engine/Private/Actor.cpp
@@ -252,6 +252,22 @@ void AActor::CheckActorComponents()
 	}
 }
 
+void AActor::ResetOwnedComponents()
+{
+	TArray<UObject*> Children;
+	OwnedComponents.Empty();
+	GetObjectsWithOuter(this, Children, true, RF_PendingKill);
+
+	for (UObject* Child : Children)
+	{
+		UActorComponent* Component = Cast<UActorComponent>(Child);
+		if (Component)
+		{
+			OwnedComponents.Add(Component);
+		}
+	}
+}
+
 void AActor::PostInitProperties()
 {
 	Super::PostInitProperties();
@@ -263,18 +279,7 @@ void AActor::PostInitProperties()
 	// This is not necessary for CDOs
 	if (!HasAnyFlags(RF_ClassDefaultObject))
 	{
-		TArray<UObject*> Children;
-		OwnedComponents.Empty();
-		GetObjectsWithOuter(this, Children, true, RF_PendingKill);
-
-		for (UObject* Child : Children)
-		{
-			UActorComponent* Component = Cast<UActorComponent>(Child);
-			if (Component)
-			{
-				OwnedComponents.Add(Component);
-			}
-		}
+		ResetOwnedComponents();
 	}
 }
 
@@ -550,6 +555,8 @@ void AActor::PostLoadSubobjects(FObjectInstancingGraph* OuterInstanceGraph)
 
 	Super::PostLoadSubobjects(OuterInstanceGraph);
 
+	ResetOwnedComponents();
+
 	if (RootComponent && bHadRoot && OldRoot != RootComponent)
 	{
 		UE_LOG(LogActor, Log, TEXT("Root component has changed, relocating new root component to old position %s->%s"), *OldRoot->GetFullName(), *GetRootComponent()->GetFullName());
