Fix build break that I caused

#UE4

diff --git a/Engine/Source/Runtime/Engine/Private/Texture.cpp b/Engine/Source/Runtime/Engine/Private/Texture.cpp
--- a/Engine/Source/Runtime/Engine/Private/Texture.cpp
+++ b/Engine/Source/Runtime/Engine/Private/Texture.cpp
@@ -589,7 +589,7 @@ void FTextureLODSettings::ReadEntry( int32 GroupId, const TCHAR* GroupName, cons
 	}
 }
 
-int32 FTextureLODSettings::CalculateLODBias( UTexture* Texture, bool bIncTextureMips ) const
+int32 FTextureLODSettings::CalculateLODBias( const UTexture* Texture, bool bIncTextureMips ) const
 {	
 	check(Texture);
 
diff --git a/Engine/Source/Runtime/Engine/Private/TextureDerivedData.cpp b/Engine/Source/Runtime/Engine/Private/TextureDerivedData.cpp
--- a/Engine/Source/Runtime/Engine/Private/TextureDerivedData.cpp
+++ b/Engine/Source/Runtime/Engine/Private/TextureDerivedData.cpp
@@ -329,7 +329,7 @@ static void GetTextureBuildSettings(
 	OutBuildSettings.bFlipGreenChannel = Texture.bFlipGreenChannel;
 	OutBuildSettings.CompositeTextureMode = Texture.CompositeTextureMode;
 	OutBuildSettings.CompositePower = Texture.CompositePower;
-	OutBuildSettings.LODBias = GSystemSettings.TextureLODSettings.CalculateLODBias(Texture);
+	OutBuildSettings.LODBias = GSystemSettings.TextureLODSettings.CalculateLODBias(&Texture);
 	OutBuildSettings.bStreamable = !Texture.NeverStream && (Texture.LODGroup != TEXTUREGROUP_UI) && (Cast<const UTexture2D>(&Texture) != NULL);
 }
 
@@ -1197,7 +1197,7 @@ static void SerializePlatformData(
 			const int32 NumCinematicMipLevels = Texture->NumCinematicMipLevels;
 			const TextureMipGenSettings MipGenSetting = Texture->MipGenSettings;
 
-			FirstMipToSerialize = Ar.CookingTarget()->GetTextureLODSettings().CalculateLODBias(this);
+			FirstMipToSerialize = Ar.CookingTarget()->GetTextureLODSettings().CalculateLODBias(Texture);
 			FirstMipToSerialize = FMath::Clamp(FirstMipToSerialize, 0, FMath::Max(NumMips-1,0));
 			NumMips -= FirstMipToSerialize;
 		}
diff --git a/Engine/Source/Runtime/Engine/Public/Texture.h b/Engine/Source/Runtime/Engine/Public/Texture.h
--- a/Engine/Source/Runtime/Engine/Public/Texture.h
+++ b/Engine/Source/Runtime/Engine/Public/Texture.h
@@ -903,7 +903,7 @@ struct FTextureLODSettings
 	 * @param	bIncTextureMips	If true, takes the textures LOD & Cinematic mips into consideration
 	 * @return	LOD bias
 	 */
-	ENGINE_API int32 CalculateLODBias(UTexture* Texture, bool bIncTextureMips = true) const;
+	ENGINE_API int32 CalculateLODBias(const UTexture* Texture, bool bIncTextureMips = true) const;
 
 	/**
 	 * Calculates and returns the LOD bias based on the information provided.
