fix warning about enum text not resolving when starting the game

diff --git a/Engine/Source/Runtime/Engine/Private/RenderingSettings.cpp b/Engine/Source/Runtime/Engine/Private/RenderingSettings.cpp
--- a/Engine/Source/Runtime/Engine/Private/RenderingSettings.cpp
+++ b/Engine/Source/Runtime/Engine/Private/RenderingSettings.cpp
@@ -29,7 +29,7 @@ void URendererSettings::PostInitProperties()
 				IConsoleVariable* CVar = IConsoleManager::Get().FindConsoleVariable(*CVarName);
 				if (CVar)
 				{
-					if (Property->ImportText(*CVar->GetString(), Property->ContainerPtrToValuePtr<uint8>(this, 0), 0, this) == NULL)
+					if (Property->ImportText(*CVar->GetString(), Property->ContainerPtrToValuePtr<uint8>(this, 0), PPF_ConsoleVariable, this) == NULL)
 					{
 						UE_LOG(LogTemp, Error, TEXT("URendererSettings import failed for %s on console variable %s (=%s)"), *Property->GetName(), *CVarName, *CVar->GetString());
 					}
