Allow vehicles to use multi body setup

diff --git a/Engine/Source/Runtime/Engine/Private/Vehicles/WheeledVehicleMovementComponent.cpp b/Engine/Source/Runtime/Engine/Private/Vehicles/WheeledVehicleMovementComponent.cpp
--- a/Engine/Source/Runtime/Engine/Private/Vehicles/WheeledVehicleMovementComponent.cpp
+++ b/Engine/Source/Runtime/Engine/Private/Vehicles/WheeledVehicleMovementComponent.cpp
@@ -121,7 +121,7 @@ bool UWheeledVehicleMovementComponent::CanCreateVehicle() const
 		return false;
 	}
 
-	if ( UpdatedComponent->BodyInstance.GetPxRigidDynamic() == NULL )
+	if ( UpdatedComponent->GetBodyInstance() == NULL )
 	{
 		UE_LOG( LogVehicles, Warning, TEXT("Cannot create vehicle for %s. UpdatedComponent has not initialized its rigid body actor."),
 			*GetPathName() );
@@ -150,7 +150,7 @@ void UWheeledVehicleMovementComponent::CreateVehicle()
 		if ( CanCreateVehicle() )
 		{
 			check(UpdatedComponent);
-			check(UpdatedComponent->BodyInstance.GetPxRigidDynamic());
+			check(UpdatedComponent->GetBodyInstance()->GetPxRigidDynamic());
 
 			SetupVehicle();
 
@@ -164,7 +164,7 @@ void UWheeledVehicleMovementComponent::CreateVehicle()
 
 void UWheeledVehicleMovementComponent::SetupVehicleShapes()
 {
-	PxRigidDynamic* PVehicleActor = UpdatedComponent->BodyInstance.GetPxRigidDynamic();
+	PxRigidDynamic* PVehicleActor = UpdatedComponent->GetBodyInstance()->GetPxRigidDynamic();
 
 	static PxMaterial* WheelMaterial = GPhysXSDK->createMaterial(0.0f, 0.0f, 0.0f);
 
@@ -230,7 +230,7 @@ void UWheeledVehicleMovementComponent::SetupVehicleShapes()
 
 void UWheeledVehicleMovementComponent::SetupVehicleMass()
 {
-	const FBodyInstance& BodyInst = UpdatedComponent->BodyInstance;
+	const FBodyInstance& BodyInst = *UpdatedComponent->GetBodyInstance();
 	PxRigidDynamic* PVehicleActor = BodyInst.GetPxRigidDynamic();
 
 	// Override mass
@@ -252,7 +252,7 @@ void UWheeledVehicleMovementComponent::SetupVehicleMass()
 
 void UWheeledVehicleMovementComponent::SetupWheels( PxVehicleWheelsSimData* PWheelsSimData )
 {
-	const FBodyInstance& VehicleBodyInstance = UpdatedComponent->BodyInstance;
+	const FBodyInstance& VehicleBodyInstance = *UpdatedComponent->GetBodyInstance();
 	PxRigidDynamic* PVehicleActor = VehicleBodyInstance.GetPxRigidDynamic();
 
 	const PxReal LengthScale = 100.f; // Convert default from m to cm
@@ -476,7 +476,10 @@ FVector UWheeledVehicleMovementComponent::GetWheelRestingPosition( const FWheelS
 		if ( Mesh && Mesh->SkeletalMesh )
 		{
 			const FVector BonePosition = Mesh->SkeletalMesh->GetComposedRefPoseMatrix( WheelSetup.BoneName ).GetOrigin() * Mesh->RelativeScale3D;
-			Offset += BonePosition;
+			//BonePosition is local for the root BONE of the skeletal mesh - however, we are using the Root BODY which may have its own transform, so we need to return the position local to the root BODY
+			const FMatrix RootBodyMTX = Mesh->SkeletalMesh->GetComposedRefPoseMatrix(Mesh->GetBodyInstance()->BodySetup->BoneName);
+			const FVector LocalBonePosition = RootBodyMTX.InverseTransformPosition(BonePosition);
+			Offset += LocalBonePosition;
 		}
 	}
 
diff --git a/Engine/Source/Runtime/Engine/Private/Vehicles/WheeledVehicleMovementComponent4W.cpp b/Engine/Source/Runtime/Engine/Private/Vehicles/WheeledVehicleMovementComponent4W.cpp
--- a/Engine/Source/Runtime/Engine/Private/Vehicles/WheeledVehicleMovementComponent4W.cpp
+++ b/Engine/Source/Runtime/Engine/Private/Vehicles/WheeledVehicleMovementComponent4W.cpp
@@ -227,7 +227,7 @@ void UWheeledVehicleMovementComponent4W::SetupVehicle()
 	PxVehicleDrive4W* PVehicleDrive4W = PxVehicleDrive4W::allocate(4);
 	check(PVehicleDrive4W);
 
-	PVehicleDrive4W->setup( GPhysXSDK, UpdatedComponent->BodyInstance.GetPxRigidDynamic(), *PWheelsSimData, DriveData, 0);
+	PVehicleDrive4W->setup( GPhysXSDK, UpdatedComponent->GetBodyInstance()->GetPxRigidDynamic(), *PWheelsSimData, DriveData, 0);
 	PVehicleDrive4W->setToRestState();
 
 	// cleanup
diff --git a/Engine/Source/Runtime/Engine/Private/WheeledVehicle.cpp b/Engine/Source/Runtime/Engine/Private/WheeledVehicle.cpp
--- a/Engine/Source/Runtime/Engine/Private/WheeledVehicle.cpp
+++ b/Engine/Source/Runtime/Engine/Private/WheeledVehicle.cpp
@@ -20,7 +20,7 @@ AWheeledVehicle::AWheeledVehicle(const class FPostConstructInitializeProperties&
 	Mesh->BodyInstance.bNotifyRigidBodyCollision = true;
 	Mesh->BodyInstance.bUseCCD = true;
 	Mesh->bBlendPhysics = true;
-	Mesh->bUseSingleBodyPhysics = true;
+	Mesh->bUseSingleBodyPhysics = false;
 	Mesh->bGenerateOverlapEvents = true;
 	RootComponent = Mesh;
 
