Engine: Added comments to EComponentMobility values

diff --git a/Engine/Source/Runtime/Engine/Classes/Engine/EngineTypes.h b/Engine/Source/Runtime/Engine/Classes/Engine/EngineTypes.h
--- a/Engine/Source/Runtime/Engine/Classes/Engine/EngineTypes.h
+++ b/Engine/Source/Runtime/Engine/Classes/Engine/EngineTypes.h
@@ -2322,9 +2322,24 @@ namespace EComponentMobility
 {
 	enum Type
 	{
+		// A static object can't be changed in game.
+		// - Allows Baked Lighting
+		// - Fastest Rendering
 		Static,
+
+		// A stationary light will only have its shadowing and bounced lighting from static geometry baked by Lightmass, all other lighting will be dynamic.
+		// - Stationary only makes sense for light components
+		// - It can change color and intensity in game.
+		// - Can't Move
+		// - Allows Partial Baked Lighting
+		// - Dynamic Shadows
 		Stationary,
-		Movable		
+
+		// Movable objects can be moved and changed in game
+		// - Totally Dynamic
+		// - Allows Dynamic Shadows
+		// - Slowest Rendering
+		Movable
 	};
 }
 
