Unhiding various sub-categories of functions that were accidently hidden in CL 2020194.
(TTP #333797/#333799)

diff --git a/Engine/Plugins/Runtime/CableComponent/Source/CableComponent/Classes/CableActor.h b/Engine/Plugins/Runtime/CableComponent/Source/CableComponent/Classes/CableActor.h
--- a/Engine/Plugins/Runtime/CableComponent/Source/CableComponent/Classes/CableActor.h
+++ b/Engine/Plugins/Runtime/CableComponent/Source/CableComponent/Classes/CableActor.h
@@ -4,7 +4,7 @@
 
 #include "CableActor.generated.h"
 
-UCLASS(hidecategories=(Input,Collision,Replication))
+UCLASS(hidecategories=(Input,Collision,Replication), showcategories=("Input|MouseInput", "Input|TouchInput"))
 class ACableActor : public AActor
 {
 	GENERATED_UCLASS_BODY()
diff --git a/Engine/Source/Editor/KismetCompiler/Private/KismetCompiler.cpp b/Engine/Source/Editor/KismetCompiler/Private/KismetCompiler.cpp
--- a/Engine/Source/Editor/KismetCompiler/Private/KismetCompiler.cpp
+++ b/Engine/Source/Editor/KismetCompiler/Private/KismetCompiler.cpp
@@ -1526,6 +1526,10 @@ void FKismetCompilerContext::FinishCompilingClass(UClass* Class)
 		{
 			Class->SetMetaData(TEXT("HideCategories"), *ParentClass->GetMetaData("HideCategories"));
 		}
+		if (ParentClass->HasMetaData(TEXT("ShowCategories")))
+		{
+			Class->SetMetaData(TEXT("ShowCategories"), *ParentClass->GetMetaData("ShowCategories"));
+		}
 		if (ParentClass->HasMetaData(TEXT("HideFunctions")))
 		{
 			Class->SetMetaData(TEXT("HideFunctions"), *ParentClass->GetMetaData("HideFunctions"));
diff --git a/Engine/Source/Runtime/Engine/Classes/AI/Navigation/NavigationTestingActor.h b/Engine/Source/Runtime/Engine/Classes/AI/Navigation/NavigationTestingActor.h
--- a/Engine/Source/Runtime/Engine/Classes/AI/Navigation/NavigationTestingActor.h
+++ b/Engine/Source/Runtime/Engine/Classes/AI/Navigation/NavigationTestingActor.h
@@ -26,7 +26,7 @@ namespace ENavCostDisplay
 	};
 }
 
-UCLASS(hidecategories=(Object, Actor, Input, Rendering), Blueprintable)
+UCLASS(hidecategories=(Object, Actor, Input, Rendering), showcategories=("Input|MouseInput", "Input|TouchInput"), Blueprintable)
 class ENGINE_API ANavigationTestingActor : public AActor, public INavAgentInterface, public INavPathObserverInterface
 {
 	GENERATED_UCLASS_BODY()
diff --git a/Engine/Source/Runtime/Engine/Classes/Atmosphere/AtmosphericFog.h b/Engine/Source/Runtime/Engine/Classes/Atmosphere/AtmosphericFog.h
--- a/Engine/Source/Runtime/Engine/Classes/Atmosphere/AtmosphericFog.h
+++ b/Engine/Source/Runtime/Engine/Classes/Atmosphere/AtmosphericFog.h
@@ -4,7 +4,7 @@
 #pragma once
 #include "AtmosphericFog.generated.h"
 
-UCLASS(showcategories=(Movement, Rendering, "Utilities|Orientation"), ClassGroup=Fog, hidecategories=(Info,Object,Input), MinimalAPI)
+UCLASS(showcategories=(Movement, Rendering, "Utilities|Orientation", "Input|MouseInput", "Input|TouchInput"), ClassGroup=Fog, hidecategories=(Info,Object,Input), MinimalAPI)
 class AAtmosphericFog : public AInfo
 {
 	GENERATED_UCLASS_BODY()
diff --git a/Engine/Source/Runtime/Engine/Classes/Camera/CameraActor.h b/Engine/Source/Runtime/Engine/Classes/Camera/CameraActor.h
--- a/Engine/Source/Runtime/Engine/Classes/Camera/CameraActor.h
+++ b/Engine/Source/Runtime/Engine/Classes/Camera/CameraActor.h
@@ -7,7 +7,7 @@
 #pragma once
 #include "CameraActor.generated.h"
 
-UCLASS(ClassGroup=Common, hideCategories=(Input, Rendering), MinimalAPI, Blueprintable)
+UCLASS(ClassGroup=Common, hideCategories=(Input, Rendering), showcategories=("Input|MouseInput", "Input|TouchInput"), MinimalAPI, Blueprintable)
 class ACameraActor : public AActor
 {
 	GENERATED_UCLASS_BODY()
diff --git a/Engine/Source/Runtime/Engine/Classes/Components/BrushComponent.h b/Engine/Source/Runtime/Engine/Classes/Components/BrushComponent.h
--- a/Engine/Source/Runtime/Engine/Classes/Components/BrushComponent.h
+++ b/Engine/Source/Runtime/Engine/Classes/Components/BrushComponent.h
@@ -4,7 +4,7 @@
 #pragma once
 #include "BrushComponent.generated.h"
 
-UCLASS(editinlinenew, MinimalAPI, showcategories=(Mobility), hidecategories=(Physics, Lighting, LOD, Rendering, TextureStreaming, Transform, Activation, "Components|Activation"))
+UCLASS(editinlinenew, MinimalAPI, hidecategories=(Physics, Lighting, LOD, Rendering, TextureStreaming, Transform, Activation, "Components|Activation"), showcategories=(Mobility, "Rendering|Material"))
 class UBrushComponent : public UPrimitiveComponent
 {
 	GENERATED_UCLASS_BODY()
diff --git a/Engine/Source/Runtime/Engine/Classes/Components/DecalComponent.h b/Engine/Source/Runtime/Engine/Classes/Components/DecalComponent.h
--- a/Engine/Source/Runtime/Engine/Classes/Components/DecalComponent.h
+++ b/Engine/Source/Runtime/Engine/Classes/Components/DecalComponent.h
@@ -9,7 +9,7 @@
 /** 
  * A material that is rendered onto the surface of a mesh. A kind of 'bumper sticker' for a model.
  */
-UCLASS(hidecategories=(Collision, Object, Physics, SceneComponent, Activation, "Components|Activation", Physics, Mobility), ClassGroup=Rendering, MinimalAPI, meta=(BlueprintSpawnableComponent))
+UCLASS(hidecategories=(Collision, Object, Physics, SceneComponent, Activation, "Components|Activation", Mobility), ClassGroup=Rendering, MinimalAPI, meta=(BlueprintSpawnableComponent))
 class UDecalComponent : public USceneComponent
 {
 	GENERATED_UCLASS_BODY()
diff --git a/Engine/Source/Runtime/Engine/Classes/Engine/Brush.h b/Engine/Source/Runtime/Engine/Classes/Engine/Brush.h
--- a/Engine/Source/Runtime/Engine/Classes/Engine/Brush.h
+++ b/Engine/Source/Runtime/Engine/Classes/Engine/Brush.h
@@ -69,7 +69,7 @@ struct FGeomSelection
 	}
 };
 
-UCLASS(hidecategories=(Object, Collision, Display, Rendering, Physics, Input), MinimalAPI, NotBlueprintable, ConversionRoot)
+UCLASS(hidecategories=(Object, Collision, Display, Rendering, Physics, Input), showcategories=("Input|MouseInput", "Input|TouchInput"), MinimalAPI, NotBlueprintable, ConversionRoot)
 class ABrush : public AActor
 {
 	GENERATED_UCLASS_BODY()
diff --git a/Engine/Source/Runtime/Engine/Classes/Engine/DecalActor.h b/Engine/Source/Runtime/Engine/Classes/Engine/DecalActor.h
--- a/Engine/Source/Runtime/Engine/Classes/Engine/DecalActor.h
+++ b/Engine/Source/Runtime/Engine/Classes/Engine/DecalActor.h
@@ -13,7 +13,7 @@
 //
 class UBoxComponent;
 
-UCLASS(hidecategories=(Collision, Attachment, Actor, Input, Replication), MinimalAPI)
+UCLASS(hidecategories=(Collision, Attachment, Actor, Input, Replication), showcategories=("Input|MouseInput", "Input|TouchInput"), MinimalAPI)
 class ADecalActor : public AActor
 {
 	GENERATED_UCLASS_BODY()
diff --git a/Engine/Source/Runtime/Engine/Classes/Engine/LevelBounds.h b/Engine/Source/Runtime/Engine/Classes/Engine/LevelBounds.h
--- a/Engine/Source/Runtime/Engine/Classes/Engine/LevelBounds.h
+++ b/Engine/Source/Runtime/Engine/Classes/Engine/LevelBounds.h
@@ -8,7 +8,7 @@
  * Updates bounding box automatically based on actors transformation changes or holds fixed user defined bounding box
  * 
  */
-UCLASS(hidecategories=(Advanced, Collision, Display, Rendering, Physics, Input), MinimalAPI)
+UCLASS(hidecategories=(Advanced, Collision, Display, Rendering, Physics, Input), showcategories=("Input|MouseInput", "Input|TouchInput"), MinimalAPI)
 class ALevelBounds : public AActor
 {
 	GENERATED_UCLASS_BODY()
diff --git a/Engine/Source/Runtime/Engine/Classes/Engine/Light.h b/Engine/Source/Runtime/Engine/Classes/Engine/Light.h
--- a/Engine/Source/Runtime/Engine/Classes/Engine/Light.h
+++ b/Engine/Source/Runtime/Engine/Classes/Engine/Light.h
@@ -4,7 +4,7 @@
 
 #include "Light.generated.h"
 
-UCLASS(Abstract, ClassGroup=Lights, hidecategories=(Input,Collision,Replication), ConversionRoot, meta=(ChildCanTick))
+UCLASS(Abstract, ClassGroup=Lights, hidecategories=(Input,Collision,Replication), showcategories=("Input|MouseInput", "Input|TouchInput"), ConversionRoot, meta=(ChildCanTick))
 class ENGINE_API ALight : public AActor
 {
 	GENERATED_UCLASS_BODY()
diff --git a/Engine/Source/Runtime/Engine/Classes/Engine/Note.h b/Engine/Source/Runtime/Engine/Classes/Engine/Note.h
--- a/Engine/Source/Runtime/Engine/Classes/Engine/Note.h
+++ b/Engine/Source/Runtime/Engine/Classes/Engine/Note.h
@@ -8,7 +8,7 @@
 #pragma once
 #include "Note.generated.h"
 
-UCLASS(MinimalAPI, hidecategories = (Input))
+UCLASS(MinimalAPI, hidecategories = (Input), showcategories=("Input|MouseInput", "Input|TouchInput"))
 class ANote : public AActor
 {
 	GENERATED_UCLASS_BODY()
diff --git a/Engine/Source/Runtime/Engine/Classes/Engine/SkyLight.h b/Engine/Source/Runtime/Engine/Classes/Engine/SkyLight.h
--- a/Engine/Source/Runtime/Engine/Classes/Engine/SkyLight.h
+++ b/Engine/Source/Runtime/Engine/Classes/Engine/SkyLight.h
@@ -4,7 +4,7 @@
 
 #include "SkyLight.generated.h"
 
-UCLASS(ClassGroup=Lights, hidecategories=(Input,Collision,Replication,Info), ConversionRoot)
+UCLASS(ClassGroup=Lights, hidecategories=(Input,Collision,Replication,Info), showcategories=("Input|MouseInput", "Input|TouchInput"), ConversionRoot)
 class ENGINE_API ASkyLight : public AInfo
 {
 	GENERATED_UCLASS_BODY()
diff --git a/Engine/Source/Runtime/Engine/Classes/Engine/StaticMeshActor.h b/Engine/Source/Runtime/Engine/Classes/Engine/StaticMeshActor.h
--- a/Engine/Source/Runtime/Engine/Classes/Engine/StaticMeshActor.h
+++ b/Engine/Source/Runtime/Engine/Classes/Engine/StaticMeshActor.h
@@ -7,7 +7,7 @@
  * An instance of a StaticMesh in a level
  * Note that PostInitializeComponents() is not called for StaticMeshActors
  */
-UCLASS(hidecategories=(Input), ConversionRoot, meta=(ChildCanTick))
+UCLASS(hidecategories=(Input), showcategories=("Input|MouseInput", "Input|TouchInput"), ConversionRoot, meta=(ChildCanTick))
 class ENGINE_API AStaticMeshActor : public AActor, public INavRelevantActorInterface
 {
 	GENERATED_UCLASS_BODY()
diff --git a/Engine/Source/Runtime/Engine/Classes/GameFramework/HUD.h b/Engine/Source/Runtime/Engine/Classes/GameFramework/HUD.h
--- a/Engine/Source/Runtime/Engine/Classes/GameFramework/HUD.h
+++ b/Engine/Source/Runtime/Engine/Classes/GameFramework/HUD.h
@@ -151,7 +151,7 @@ private:
 // Base class of the heads-up display.
 //
 //=============================================================================
-UCLASS(config=Game, hidecategories=(Rendering,Actor,Input,Replication), notplaceable, transient, dependson=UCanvas, BlueprintType, Blueprintable)
+UCLASS(config=Game, hidecategories=(Rendering,Actor,Input,Replication), showcategories=("Input|MouseInput", "Input|TouchInput"), notplaceable, transient, dependson=UCanvas, BlueprintType, Blueprintable)
 class ENGINE_API AHUD : public AActor
 {
 	GENERATED_UCLASS_BODY()
diff --git a/Engine/Source/Runtime/Engine/Classes/GameFramework/Info.h b/Engine/Source/Runtime/Engine/Classes/GameFramework/Info.h
--- a/Engine/Source/Runtime/Engine/Classes/GameFramework/Info.h
+++ b/Engine/Source/Runtime/Engine/Classes/GameFramework/Info.h
@@ -8,7 +8,7 @@
 #pragma once
 #include "Info.generated.h"
 
-UCLASS(abstract, hidecategories=(Input, Movement, Collision, Rendering, "Utilities|Orientation"),MinimalAPI, NotBlueprintable)
+UCLASS(abstract, hidecategories=(Input, Movement, Collision, Rendering, "Utilities|Orientation"), showcategories=("Input|MouseInput", "Input|TouchInput"), MinimalAPI, NotBlueprintable)
 class AInfo : public AActor
 {
 	GENERATED_UCLASS_BODY()
diff --git a/Engine/Source/Runtime/Engine/Classes/GameFramework/WorldSettings.h b/Engine/Source/Runtime/Engine/Classes/GameFramework/WorldSettings.h
--- a/Engine/Source/Runtime/Engine/Classes/GameFramework/WorldSettings.h
+++ b/Engine/Source/Runtime/Engine/Classes/GameFramework/WorldSettings.h
@@ -173,7 +173,7 @@ struct ENGINE_API FNetViewer
 /**
  * Actor containing all script accessible world properties.
  */
-UCLASS(config=game, showcategories=(Rendering), hidecategories=(Actor, Advanced, Display, Events, Object, Attachment, Info, Input, Blueprint, Layers), notplaceable, dependson=(APostProcessVolume, UMusicTrackDataStructures))
+UCLASS(config=game, showcategories=(Rendering), hidecategories=(Actor, Advanced, Display, Events, Object, Attachment, Info, Input, Blueprint, Layers), showcategories=("Input|MouseInput", "Input|TouchInput"), notplaceable, dependson=(APostProcessVolume, UMusicTrackDataStructures))
 class ENGINE_API AWorldSettings : public AInfo
 {
 	GENERATED_UCLASS_BODY()
diff --git a/Engine/Source/Runtime/Engine/Classes/Landscape/LandscapeComponent.h b/Engine/Source/Runtime/Engine/Classes/Landscape/LandscapeComponent.h
--- a/Engine/Source/Runtime/Engine/Classes/Landscape/LandscapeComponent.h
+++ b/Engine/Source/Runtime/Engine/Classes/Landscape/LandscapeComponent.h
@@ -101,7 +101,7 @@ struct FWeightmapLayerAllocationInfo
 	}
 };
 
-UCLASS(hidecategories=(Display, Attachment, Physics, Debug, Collision, Movement, Rendering, PrimitiveComponent, Object, Transform), MinimalAPI)
+UCLASS(hidecategories=(Display, Attachment, Physics, Debug, Collision, Movement, Rendering, PrimitiveComponent, Object, Transform), showcategories=("Rendering|Material"), MinimalAPI)
 class ULandscapeComponent : public UPrimitiveComponent
 {
 	GENERATED_UCLASS_BODY()
diff --git a/Engine/Source/Runtime/Engine/Classes/Matinee/MatineeActor.h b/Engine/Source/Runtime/Engine/Classes/Matinee/MatineeActor.h
--- a/Engine/Source/Runtime/Engine/Classes/Matinee/MatineeActor.h
+++ b/Engine/Source/Runtime/Engine/Classes/Matinee/MatineeActor.h
@@ -49,7 +49,7 @@ struct FInterpGroupActorInfo
 
 };
 
-UCLASS(dependson=UEngineBaseTypes, MinimalAPI, NotBlueprintable, hidecategories=(Collision, Game, Input))
+UCLASS(dependson=UEngineBaseTypes, MinimalAPI, NotBlueprintable, hidecategories=(Collision, Game, Input), showcategories=("Input|MouseInput", "Input|TouchInput", "Game|Damage"))
 class AMatineeActor : public AActor
 {
 	GENERATED_UCLASS_BODY()
diff --git a/Engine/Source/Runtime/Engine/Classes/Particles/Emitter.h b/Engine/Source/Runtime/Engine/Classes/Particles/Emitter.h
--- a/Engine/Source/Runtime/Engine/Classes/Particles/Emitter.h
+++ b/Engine/Source/Runtime/Engine/Classes/Particles/Emitter.h
@@ -44,7 +44,7 @@ DECLARE_DYNAMIC_MULTICAST_DELEGATE_SixParams( FParticleDeathSignature, FName, Ev
 	*/
 DECLARE_DYNAMIC_MULTICAST_DELEGATE_EightParams( FParticleCollisionSignature, FName, EventName, float, EmitterTime, int32, ParticleTime, FVector, Location, FVector, Velocity, FVector, Direction, FVector, Normal, FName, BoneName);
 
-UCLASS(hidecategories=(Activation,"Components|Activation",Input,Collision))
+UCLASS(hidecategories=(Activation,"Components|Activation",Input,Collision), showcategories=("Input|MouseInput", "Input|TouchInput"))
 class ENGINE_API AEmitter : public AActor
 {
 	GENERATED_UCLASS_BODY()
diff --git a/Engine/Source/Runtime/Engine/Classes/PhysicsEngine/DestructibleActor.h b/Engine/Source/Runtime/Engine/Classes/PhysicsEngine/DestructibleActor.h
--- a/Engine/Source/Runtime/Engine/Classes/PhysicsEngine/DestructibleActor.h
+++ b/Engine/Source/Runtime/Engine/Classes/PhysicsEngine/DestructibleActor.h
@@ -8,7 +8,7 @@
 
 class UDestructibleComponent;
 
-UCLASS(MinimalAPI, hidecategories=(Input))
+UCLASS(MinimalAPI, hidecategories=(Input), showcategories=("Input|MouseInput", "Input|TouchInput"))
 class ADestructibleActor : public AActor
 {
 	GENERATED_UCLASS_BODY()
diff --git a/Engine/Source/Runtime/Engine/Classes/PhysicsEngine/PhysicsConstraintComponent.h b/Engine/Source/Runtime/Engine/Classes/PhysicsEngine/PhysicsConstraintComponent.h
--- a/Engine/Source/Runtime/Engine/Classes/PhysicsEngine/PhysicsConstraintComponent.h
+++ b/Engine/Source/Runtime/Engine/Classes/PhysicsEngine/PhysicsConstraintComponent.h
@@ -9,7 +9,7 @@ DECLARE_DYNAMIC_MULTICAST_DELEGATE_OneParam(FConstraintBrokenSignature, int32, C
 /**
  *	This is effectively a joint that allows you to connect 2 rigid bodies together. You can create different types of joints using the various parameters of this component.
  */
-UCLASS(ClassGroup=Physics, dependson(ConstraintInstance), MinimalAPI, meta=(BlueprintSpawnableComponent), HideCategories=(Activation,"Components|Activation",Physics, Mobility))
+UCLASS(ClassGroup=Physics, dependson(ConstraintInstance), MinimalAPI, meta=(BlueprintSpawnableComponent), HideCategories=(Activation,"Components|Activation", Physics, Mobility), ShowCategories=("Physics|Components|PhysicsConstraint"))
 class UPhysicsConstraintComponent : public USceneComponent
 {
 	GENERATED_UCLASS_BODY()
diff --git a/Engine/Source/Runtime/Engine/Classes/PhysicsEngine/PhysicsThruster.h b/Engine/Source/Runtime/Engine/Classes/PhysicsEngine/PhysicsThruster.h
--- a/Engine/Source/Runtime/Engine/Classes/PhysicsEngine/PhysicsThruster.h
+++ b/Engine/Source/Runtime/Engine/Classes/PhysicsEngine/PhysicsThruster.h
@@ -9,7 +9,7 @@
  *	Attach one of these on an object using physics simulation and it will apply a force down the negative-X direction
  *	ie. point X in the direction you want the thrust in.
  */
-UCLASS(hidecategories=(Input,Collision,Replication))
+UCLASS(hidecategories=(Input,Collision,Replication), showcategories=("Input|MouseInput", "Input|TouchInput"))
 class APhysicsThruster : public ARigidBodyBase
 {
 	GENERATED_UCLASS_BODY()
diff --git a/Engine/Source/Runtime/Engine/Classes/PhysicsEngine/RadialForceActor.h b/Engine/Source/Runtime/Engine/Classes/PhysicsEngine/RadialForceActor.h
--- a/Engine/Source/Runtime/Engine/Classes/PhysicsEngine/RadialForceActor.h
+++ b/Engine/Source/Runtime/Engine/Classes/PhysicsEngine/RadialForceActor.h
@@ -4,7 +4,7 @@
 #pragma once
 #include "RadialForceActor.generated.h"
 
-UCLASS(MinimalAPI, hidecategories=(Collision, Input))
+UCLASS(MinimalAPI, hidecategories=(Collision, Input), showcategories=("Input|MouseInput", "Input|TouchInput"))
 class ARadialForceActor : public ARigidBodyBase
 {
 	GENERATED_UCLASS_BODY()
diff --git a/Engine/Source/Runtime/Engine/Classes/Sound/AmbientSound.h b/Engine/Source/Runtime/Engine/Classes/Sound/AmbientSound.h
--- a/Engine/Source/Runtime/Engine/Classes/Sound/AmbientSound.h
+++ b/Engine/Source/Runtime/Engine/Classes/Sound/AmbientSound.h
@@ -6,7 +6,7 @@
 #pragma once
 #include "AmbientSound.generated.h"
 
-UCLASS(AutoExpandCategories=Audio, ClassGroup=Sounds, MinimalAPI, hidecategories(Collision, Input, Game))
+UCLASS(AutoExpandCategories=Audio, ClassGroup=Sounds, MinimalAPI, hidecategories(Collision, Input, Game), showcategories=("Input|MouseInput", "Input|TouchInput", "Game|Damage"))
 class AAmbientSound : public AActor
 {
 	GENERATED_UCLASS_BODY()
