UE4: Added some more temporary logging to track down a crash in cook.

diff --git a/Engine/Source/Runtime/Engine/Private/GlobalShader.cpp b/Engine/Source/Runtime/Engine/Private/GlobalShader.cpp
--- a/Engine/Source/Runtime/Engine/Private/GlobalShader.cpp
+++ b/Engine/Source/Runtime/Engine/Private/GlobalShader.cpp
@@ -24,6 +24,8 @@ FGlobalShaderMapId::FGlobalShaderMapId(EShaderPlatform Platform)
 {
 	TArray<FShaderType*> ShaderTypes;
 
+	UE_LOG(LogInit, Log, TEXT("FGlobalShaderMapId constructor started"));
+
 	for (TLinkedList<FShaderType*>::TIterator ShaderTypeIt(FShaderType::GetTypeList()); ShaderTypeIt; ShaderTypeIt.Next())
 	{
 		FGlobalShaderType* GlobalShaderType = ShaderTypeIt->GetGlobalShaderType();
@@ -34,8 +36,12 @@ FGlobalShaderMapId::FGlobalShaderMapId(EShaderPlatform Platform)
 		}
 	}
 
+	UE_LOG(LogInit, Log, TEXT("FGlobalShaderMapId shader types added"));
+
 	ShaderTypes.Sort(FCompareShaderTypes());
 
+	UE_LOG(LogInit, Log, TEXT("FGlobalShaderMapId shader types sorted"));
+
 	for (int32 TypeIndex = 0; TypeIndex < ShaderTypes.Num(); TypeIndex++)
 	{
 		FShaderTypeDependency Dependency;
@@ -43,6 +49,8 @@ FGlobalShaderMapId::FGlobalShaderMapId(EShaderPlatform Platform)
 		Dependency.SourceHash = ShaderTypes[TypeIndex]->GetSourceHash();
 		ShaderTypeDependencies.Add(Dependency);
 	}
+
+	UE_LOG(LogInit, Log, TEXT("FGlobalShaderMapId ShaderTypeDependencies added"));
 }
 
 void FGlobalShaderMapId::AppendKeyString(FString& KeyString) const
@@ -290,10 +298,15 @@ FString SaveGlobalShaderFile(EShaderPlatform Platform, FString SavePath)
 /** Creates a string key for the derived data cache entry for the global shader map. */
 FString GetGlobalShaderMapKeyString(const FGlobalShaderMapId& ShaderMapId, EShaderPlatform Platform)
 {
+	UE_LOG(LogInit, Log, TEXT("Begin GetGlobalShaderMapKeyString. Platform: %d"), (int32)Platform);
 	FName Format = LegacyShaderPlatformToShaderFormat(Platform);
+	UE_LOG(LogInit, Log, TEXT("Legacy shader platform converted to shader format. Format: %s"), *Format.ToString());
 	FString ShaderMapKeyString = Format.ToString() + TEXT("_") + FString(FString::FromInt(GetTargetPlatformManagerRef().ShaderFormatVersion(Format))) + TEXT("_");
+	UE_LOG(LogInit, Log, TEXT("Initial shader map key string formed. String: %s"), *ShaderMapKeyString);
 	ShaderMapAppendKeyString(ShaderMapKeyString);
+	UE_LOG(LogInit, Log, TEXT("Shader map key string appended. New string: %s"), *ShaderMapKeyString);
 	ShaderMapId.AppendKeyString(ShaderMapKeyString);
+	UE_LOG(LogInit, Log, TEXT("Shader map id appended."));
 	return FDerivedDataCacheInterface::BuildCacheKey(TEXT("GSM"), GLOBALSHADERMAP_DERIVEDDATA_VER, *ShaderMapKeyString);
 }
 
@@ -389,6 +402,8 @@ TShaderMap<FGlobalShaderType>* GetGlobalShaderMap(EShaderPlatform Platform, bool
 			UE_LOG(LogInit, Log, TEXT("Preparing global shader map on uncooked platform"));
 			FGlobalShaderMapId ShaderMapId(Platform);
 
+			UE_LOG(LogInit, Log, TEXT("FGlobalShaderMapId constructed"));
+
 			TArray<uint8> CachedData;
 			const FString DataKey = GetGlobalShaderMapKeyString(ShaderMapId, Platform);
 
