Fixed ES2

diff --git a/Engine/Shaders/BasePassPixelShader.usf b/Engine/Shaders/BasePassPixelShader.usf
--- a/Engine/Shaders/BasePassPixelShader.usf
+++ b/Engine/Shaders/BasePassPixelShader.usf
@@ -81,7 +81,7 @@ half3 GetImageBasedReflectionLighting(FMaterialPixelParameters MaterialParameter
 #endif
 	
 	BRANCH
-	if( View.UseLightmaps != 0 )
+	if( View.UseLightmaps > 0 )
 	{
 		// Note: make sure this matches the lightmap mixing done on opaque (ReflectionEnvironmentTiledDeferredMain)
 		SpecularIBL.rgb *= Luminance( LowGlossIndirectSpecular );
@@ -414,7 +414,7 @@ float3 GetPrecomputedIndirectLighting(
 	float3 SkyLighting = GetSkyLighting(MaterialParameters.WorldNormal, SkyOcclusionUV);
 
 	BRANCH
-	if( View.UseLightmaps )
+	if( View.UseLightmaps > 0 )
 	{
 		Lighting += SkyLighting;
 	}
diff --git a/Engine/Shaders/ReflectionEnvironmentShaders.usf b/Engine/Shaders/ReflectionEnvironmentShaders.usf
--- a/Engine/Shaders/ReflectionEnvironmentShaders.usf
+++ b/Engine/Shaders/ReflectionEnvironmentShaders.usf
@@ -519,7 +519,7 @@ void ReflectionApplyPS(
 #endif
 		
 		BRANCH
-		if( View.UseLightmaps )
+		if( View.UseLightmaps > 0 )
 		{
 			// We have high frequency directional data but low frequency spatial data in the envmap.
 			// We have high frequency spatial data but low frequency directional data in the lightmap.
diff --git a/Engine/Source/Runtime/Engine/Public/SceneManagement.h b/Engine/Source/Runtime/Engine/Public/SceneManagement.h
--- a/Engine/Source/Runtime/Engine/Public/SceneManagement.h
+++ b/Engine/Source/Runtime/Engine/Public/SceneManagement.h
@@ -2083,7 +2083,7 @@ BEGIN_UNIFORM_BUFFER_STRUCT(FViewUniformShaderParameters,ENGINE_API)
 	DECLARE_UNIFORM_BUFFER_STRUCT_MEMBER(float,RealTime)
 	DECLARE_UNIFORM_BUFFER_STRUCT_MEMBER(uint32,Random)
 	DECLARE_UNIFORM_BUFFER_STRUCT_MEMBER(uint32,FrameNumber)
-	DECLARE_UNIFORM_BUFFER_STRUCT_MEMBER(uint32,UseLightmaps)
+	DECLARE_UNIFORM_BUFFER_STRUCT_MEMBER_EX(float,UseLightmaps, EShaderPrecisionModifier::Half)
 	DECLARE_UNIFORM_BUFFER_STRUCT_MEMBER_EX(float,UnlitViewmodeMask, EShaderPrecisionModifier::Half)
 	DECLARE_UNIFORM_BUFFER_STRUCT_MEMBER_EX(float,ReflectionLightmapMixingMask, EShaderPrecisionModifier::Half)
 	DECLARE_UNIFORM_BUFFER_STRUCT_MEMBER_EX(FLinearColor,DirectionalLightColor, EShaderPrecisionModifier::Half)
diff --git a/Engine/Source/Runtime/Renderer/Private/SceneRendering.cpp b/Engine/Source/Runtime/Renderer/Private/SceneRendering.cpp
--- a/Engine/Source/Runtime/Renderer/Private/SceneRendering.cpp
+++ b/Engine/Source/Runtime/Renderer/Private/SceneRendering.cpp
@@ -368,7 +368,7 @@ TUniformBufferRef<FViewUniformShaderParameters> FViewInfo::CreateUniformBuffer(
 	static const auto CVar = IConsoleManager::Get().FindConsoleVariable(TEXT("r.DiffuseFromCaptures"));
 	const bool bUseLightmaps = CVar->GetInt() == 0;
 
-	ViewUniformShaderParameters.UseLightmaps = bUseLightmaps;
+	ViewUniformShaderParameters.UseLightmaps = bUseLightmaps ? 1 : 0;
 
 	if(State)
 	{
