UE4: A little more temp logging to track down a crash in cook.

diff --git a/Engine/Source/Runtime/Engine/Private/GlobalShader.cpp b/Engine/Source/Runtime/Engine/Private/GlobalShader.cpp
--- a/Engine/Source/Runtime/Engine/Private/GlobalShader.cpp
+++ b/Engine/Source/Runtime/Engine/Private/GlobalShader.cpp
@@ -301,7 +301,15 @@ FString GetGlobalShaderMapKeyString(const FGlobalShaderMapId& ShaderMapId, EShad
 	UE_LOG(LogInit, Log, TEXT("Begin GetGlobalShaderMapKeyString. Platform: %d"), (int32)Platform);
 	FName Format = LegacyShaderPlatformToShaderFormat(Platform);
 	UE_LOG(LogInit, Log, TEXT("Legacy shader platform converted to shader format. Format: %s"), *Format.ToString());
-	FString ShaderMapKeyString = Format.ToString() + TEXT("_") + FString(FString::FromInt(GetTargetPlatformManagerRef().ShaderFormatVersion(Format))) + TEXT("_");
+	ITargetPlatformManagerModule* ManagerModule = GetTargetPlatformManager();
+	UE_LOG(LogInit, Log, TEXT("Target platform manager aquired"));
+	check(ManagerModule);
+	UE_LOG(LogInit, Log, TEXT("Target platform manager verified"));
+	uint16 FortmatVersion = ManagerModule->ShaderFormatVersion(Format);
+	UE_LOG(LogInit, Log, TEXT("Shader format version aquired: %d"), FortmatVersion);
+	FString StringFormatVersion = FString::FromInt(FortmatVersion);
+	UE_LOG(LogInit, Log, TEXT("Shader format version stringified: %s"), *StringFormatVersion);
+	FString ShaderMapKeyString = Format.ToString() + TEXT("_") + StringFormatVersion + TEXT("_");
 	UE_LOG(LogInit, Log, TEXT("Initial shader map key string formed. String: %s"), *ShaderMapKeyString);
 	ShaderMapAppendKeyString(ShaderMapKeyString);
 	UE_LOG(LogInit, Log, TEXT("Shader map key string appended. New string: %s"), *ShaderMapKeyString);
