Rendering: Removed dead CullDistance variable from FParticleSystemSceneProxy

diff --git a/Engine/Source/Runtime/Engine/Private/Particles/ParticleSystemRender.cpp b/Engine/Source/Runtime/Engine/Private/Particles/ParticleSystemRender.cpp
--- a/Engine/Source/Runtime/Engine/Private/Particles/ParticleSystemRender.cpp
+++ b/Engine/Source/Runtime/Engine/Private/Particles/ParticleSystemRender.cpp
@@ -6776,7 +6776,6 @@ int32 FDynamicAnimTrailEmitterData::FillVertexData(struct FAsyncBufferFillData&
 FParticleSystemSceneProxy::FParticleSystemSceneProxy(const UParticleSystemComponent* Component, FParticleDynamicData* InDynamicData)
 	: FPrimitiveSceneProxy(Component, Component->Template ? Component->Template->GetFName() : NAME_None)
 	, Owner(Component->GetOwner())
-	, CullDistance(Component->CachedMaxDrawDistance > 0 ? Component->CachedMaxDrawDistance : WORLD_MAX)
 	, bCastShadow(Component->CastShadow)
 	, MaterialRelevance(
 		((Component->GetCurrentLODIndex() >= 0) && (Component->GetCurrentLODIndex() < Component->CachedViewRelevanceFlags.Num())) ?
diff --git a/Engine/Source/Runtime/Engine/Public/ParticleHelper.h b/Engine/Source/Runtime/Engine/Public/ParticleHelper.h
--- a/Engine/Source/Runtime/Engine/Public/ParticleHelper.h
+++ b/Engine/Source/Runtime/Engine/Public/ParticleHelper.h
@@ -2664,8 +2664,7 @@ public:
 	// While this isn't good OO design, access to everything is made public.
 	// This is to allow custom emitter instances to easily be written when extending the engine.
 	FMatrix GetWorldToLocal() const		{	return GetLocalToWorld().Inverse();	}
-	float GetCullDistance()				{	return CullDistance;			}
-	bool GetCastShadow()				{	return bCastShadow;				}
+	bool GetCastShadow() const			{	return bCastShadow;				}
 	const FMaterialRelevance& GetMaterialRelevance() const
 	{
 		return MaterialRelevance;
@@ -2690,7 +2689,6 @@ protected:
 
 	AActor* Owner;
 
-	float CullDistance;
 #if STATS
 	double LastStatCaptureTime;
 	bool bCountedThisFrame;
