Reverted calls to CalculateLODBias which were changed in 2060296

#ttp 333394 - EDITOR: texture editor bugs (was RE: Flow map action for Photoshop CS5 and above)

#branch UE4

#change Modified CalculateLODBias so that it still performs the TMGS_NoMipmaps test when passing multiple params, and prevent failure compiling monolithic builds.

diff --git a/Engine/Source/Runtime/Engine/Private/Texture.cpp b/Engine/Source/Runtime/Engine/Private/Texture.cpp
--- a/Engine/Source/Runtime/Engine/Private/Texture.cpp
+++ b/Engine/Source/Runtime/Engine/Private/Texture.cpp
@@ -591,28 +591,26 @@ void FTextureLODSettings::ReadEntry( int32 GroupId, const TCHAR* GroupName, cons
 
 int32 FTextureLODSettings::CalculateLODBias( const UTexture* Texture, bool bIncTextureMips ) const
 {	
-	check(Texture);
-
+	check( Texture );
+	TextureMipGenSettings MipGenSetting = TMGS_MAX;
 #if WITH_EDITORONLY_DATA
-	// Find LOD group.
-	const FTextureLODGroup& LODGroupInfo = TextureLODGroups[Texture->LODGroup];
-
-	// Test to see if we have no mip generation as in which case the LOD bias will be ignored
-	const TextureMipGenSettings FinalMipGenSetting = (Texture->MipGenSettings == TMGS_FromTextureGroup) ? LODGroupInfo.MipGenSettings : Texture->MipGenSettings;
-	if (FinalMipGenSetting == TMGS_NoMipmaps)
-	{
-		return 0;
-	}
+	MipGenSetting = Texture->MipGenSettings;
 #endif // #if WITH_EDITORONLY_DATA
-
-	return CalculateLODBias(Texture->GetSurfaceWidth(), Texture->GetSurfaceHeight(), Texture->LODGroup, (bIncTextureMips ? Texture->LODBias : 0), (bIncTextureMips ? Texture->NumCinematicMipLevels : 0));
+	return CalculateLODBias(Texture->GetSurfaceWidth(), Texture->GetSurfaceHeight(), Texture->LODGroup, (bIncTextureMips ? Texture->LODBias : 0), (bIncTextureMips ? Texture->NumCinematicMipLevels : 0), MipGenSetting);
 }
 
-int32 FTextureLODSettings::CalculateLODBias( int32 Width, int32 Height, int32 LODGroup, int32 LODBias, int32 NumCinematicMipLevels ) const
+int32 FTextureLODSettings::CalculateLODBias( int32 Width, int32 Height, int32 LODGroup, int32 LODBias, int32 NumCinematicMipLevels, TextureMipGenSettings InMipGenSetting ) const
 {	
 	// Find LOD group.
 	const FTextureLODGroup& LODGroupInfo = TextureLODGroups[LODGroup];
 
+	// Test to see if we have no mip generation as in which case the LOD bias will be ignored
+	const TextureMipGenSettings FinalMipGenSetting = (InMipGenSetting == TMGS_FromTextureGroup) ? LODGroupInfo.MipGenSettings : InMipGenSetting;
+	if ( FinalMipGenSetting == TMGS_NoMipmaps )
+	{
+		return 0;
+	}
+
 	// Calculate maximum number of miplevels.
 	int32 TextureMaxLOD	= FMath::CeilLogTwo( FMath::Trunc( FMath::Max( Width, Height ) ) );
 
diff --git a/Engine/Source/Runtime/Engine/Private/TextureDerivedData.cpp b/Engine/Source/Runtime/Engine/Private/TextureDerivedData.cpp
--- a/Engine/Source/Runtime/Engine/Private/TextureDerivedData.cpp
+++ b/Engine/Source/Runtime/Engine/Private/TextureDerivedData.cpp
@@ -329,7 +329,7 @@ static void GetTextureBuildSettings(
 	OutBuildSettings.bFlipGreenChannel = Texture.bFlipGreenChannel;
 	OutBuildSettings.CompositeTextureMode = Texture.CompositeTextureMode;
 	OutBuildSettings.CompositePower = Texture.CompositePower;
-	OutBuildSettings.LODBias = GSystemSettings.TextureLODSettings.CalculateLODBias(&Texture);
+	OutBuildSettings.LODBias = GSystemSettings.TextureLODSettings.CalculateLODBias(Texture.Source.GetSizeX(), Texture.Source.GetSizeY(), Texture.LODGroup, Texture.LODBias, Texture.NumCinematicMipLevels, Texture.MipGenSettings);
 	OutBuildSettings.bStreamable = !Texture.NeverStream && (Texture.LODGroup != TEXTUREGROUP_UI) && (Cast<const UTexture2D>(&Texture) != NULL);
 }
 
@@ -1197,7 +1197,7 @@ static void SerializePlatformData(
 			const int32 NumCinematicMipLevels = Texture->NumCinematicMipLevels;
 			const TextureMipGenSettings MipGenSetting = Texture->MipGenSettings;
 
-			FirstMipToSerialize = Ar.CookingTarget()->GetTextureLODSettings().CalculateLODBias(Texture);
+			FirstMipToSerialize = Ar.CookingTarget()->GetTextureLODSettings().CalculateLODBias(Width, Height, LODGroup, LODBias, NumCinematicMipLevels, MipGenSetting);
 			FirstMipToSerialize = FMath::Clamp(FirstMipToSerialize, 0, FMath::Max(NumMips-1,0));
 			NumMips -= FirstMipToSerialize;
 		}
diff --git a/Engine/Source/Runtime/Engine/Public/Texture.h b/Engine/Source/Runtime/Engine/Public/Texture.h
--- a/Engine/Source/Runtime/Engine/Public/Texture.h
+++ b/Engine/Source/Runtime/Engine/Public/Texture.h
@@ -913,9 +913,10 @@ struct FTextureLODSettings
 	 * @param	LODGroup					Which LOD group the texture belongs to
 	 * @param	LODBias						LOD bias to include in the calculation
 	 * @param	NumCinematicMipLevels		The texture cinematic mip levels to include in the calculation
+	 * @param	MipGenSetting				Mip generation setting
 	 * @return	LOD bias
 	 */
-	ENGINE_API int32 CalculateLODBias( int32 Width, int32 Height, int32 LODGroup, int32 LODBias, int32 NumCinematicMipLevels ) const;
+	ENGINE_API int32 CalculateLODBias( int32 Width, int32 Height, int32 LODGroup, int32 LODBias, int32 NumCinematicMipLevels, TextureMipGenSettings MipGenSetting ) const;
 
 	/** 
 	* Useful for stats in the editor.
