Added some comments to the SpawnParameters fields.

diff --git a/Engine/Source/Runtime/Engine/Classes/Engine/World.h b/Engine/Source/Runtime/Engine/Classes/Engine/World.h
--- a/Engine/Source/Runtime/Engine/Classes/Engine/World.h
+++ b/Engine/Source/Runtime/Engine/Classes/Engine/World.h
@@ -406,17 +406,37 @@ struct ENGINE_API FActorSpawnParameters
 	{
 	}
 
-
+	/* A name to assign as the Name of the Actor being spawned. If no value is specified, the name of the spawned Actor will be automatically generated using the form [Class]_[Number]. */
 	FName Name;
+
+	/* An Actor to use as a template when spawning the new Actor. The spawned Actor will be initialized using the property values of the template Actor. If left NULL the class default object (CDO) will be used to initialize the spawned Actor. */
 	AActor* Template;
-	AActor* Owner; 
+
+	/* The Actor that spawned this Actor. (Can be left as NULL). */
+	AActor* Owner;
+
+	/* The APawn that is responsible for damage done by the spawned Actor. (Can be left as NULL). */
 	APawn*	Instigator;
+
+	/* The ULevel to spawn the Actor in, i.e. the Outer of the Actor. If left as NULL the Outer of the Owner is used. If the Owner is NULL the persistent level is used. */
 	class	ULevel* OverrideLevel;
+
+	/* Determines whether a collision test will be performed when spawning the Actor. If true, no collision test will be performed when spawning the Actor regardless of the collision settings of the root component or template Actor. */
 	uint32	bNoCollisionFail:1;
+
+	/* Is the actor remotely owned. */
 	uint32	bRemoteOwned:1;
+
+	/* Determines whether spawning will not fail if certain conditions are not met. If true, spawning will not fail because the class being spawned is `bStatic=true` or because the class of the template Actor is not the same as the class of the Actor being spawned. */
 	uint32	bNoFail:1;
+
+	/* Determines whether the construction script will be run. If true, the construction script will not be run on the spawned Actor. Only applicable if the Actor is being spawned from a Blueprint. */
 	uint32	bDeferConstruction:1;
+	
+	/* Determines whether or not the actor may be spawned when running a construction script. If true spawning will fail if a construction script is being run. */
 	uint32	bAllowDuringConstructionScript:1;
+	
+	/* Flags used to describe the spawned actor/object instance. */
 	EObjectFlags ObjectFlags;		
 };
 
