Remove bUseSingleBodyPhysics from skeletal mesh

diff --git a/Engine/Source/Runtime/Engine/Classes/Components/SkeletalMeshComponent.h b/Engine/Source/Runtime/Engine/Classes/Components/SkeletalMeshComponent.h
--- a/Engine/Source/Runtime/Engine/Classes/Components/SkeletalMeshComponent.h
+++ b/Engine/Source/Runtime/Engine/Classes/Components/SkeletalMeshComponent.h
@@ -269,13 +269,6 @@ public:
 	UPROPERTY(EditAnywhere, AdvancedDisplay, BlueprintReadWrite, Category=Animation)
 	float GlobalAnimRateScale;
 
-	/**
-	 * Only instance Root Bone rigid body for physics. Mostly used by Vehicles.
-	 * Other Rigid Bodies are ignored for physics, but still considered for traces.
-	 */
-	UPROPERTY()
-	uint32 bUseSingleBodyPhysics:1;
-
 	/** If true, there is at least one body in the current PhysicsAsset with a valid bone in the current SkeletalMesh */
 	UPROPERTY(transient)
 	uint32 bHasValidBodies:1;
diff --git a/Engine/Source/Runtime/Engine/Private/PhysicsEngine/BodyInstance.cpp b/Engine/Source/Runtime/Engine/Private/PhysicsEngine/BodyInstance.cpp
--- a/Engine/Source/Runtime/Engine/Private/PhysicsEngine/BodyInstance.cpp
+++ b/Engine/Source/Runtime/Engine/Private/PhysicsEngine/BodyInstance.cpp
@@ -501,7 +501,7 @@ void FBodyInstance::UpdatePhysicsFilterData()
 	// Get skelmeshcomp ID
 	uint32 SkelMeshCompID = 0;
 	USkeletalMeshComponent* SkelMeshComp = Cast<USkeletalMeshComponent>(OwnerComponent.Get());
-	if(SkelMeshComp != NULL && !SkelMeshComp->bUseSingleBodyPhysics)
+	if(SkelMeshComp != NULL)
 	{
 		SkelMeshCompID = SkelMeshComp->GetUniqueID();
 
diff --git a/Engine/Source/Runtime/Engine/Private/PhysicsEngine/PhysAnim.cpp b/Engine/Source/Runtime/Engine/Private/PhysicsEngine/PhysAnim.cpp
--- a/Engine/Source/Runtime/Engine/Private/PhysicsEngine/PhysAnim.cpp
+++ b/Engine/Source/Runtime/Engine/Private/PhysicsEngine/PhysAnim.cpp
@@ -197,11 +197,6 @@ void USkeletalMeshComponent::BlendPhysicsBones( TArray<FBoneIndexType>& InRequir
 
 bool USkeletalMeshComponent::ShouldBlendPhysicsBones()
 {
-	if (bUseSingleBodyPhysics)
-	{
-		return false;
-	}
-
 	for (int32 BodyIndex = 0; BodyIndex < Bodies.Num(); ++BodyIndex)
 	{
 		if (Bodies[BodyIndex]->PhysicsBlendWeight > 0.f)
@@ -250,7 +245,7 @@ void USkeletalMeshComponent::UpdateKinematicBonesToPhysics(bool bTeleport)
 	// - but if we just update physics bone without update current pose, it will have stale data
 	// If desired, pass the animation data to the physics joints so they can be used by motors.
 	// See if we are going to need to update kinematics
-	const bool bUpdateKinematics = (!bUseSingleBodyPhysics && KinematicBonesUpdateType != EKinematicBonesUpdateToPhysics::SkipAllBones);
+	const bool bUpdateKinematics = (KinematicBonesUpdateType != EKinematicBonesUpdateToPhysics::SkipAllBones);
 
 	// If desired, update physics bodies associated with skeletal mesh component to match.
 	if( !bUpdateKinematics )
diff --git a/Engine/Source/Runtime/Engine/Private/PhysicsEngine/PhysicsConstraintComponent.cpp b/Engine/Source/Runtime/Engine/Private/PhysicsEngine/PhysicsConstraintComponent.cpp
--- a/Engine/Source/Runtime/Engine/Private/PhysicsEngine/PhysicsConstraintComponent.cpp
+++ b/Engine/Source/Runtime/Engine/Private/PhysicsEngine/PhysicsConstraintComponent.cpp
@@ -88,7 +88,7 @@ FTransform UPhysicsConstraintComponent::GetBodyTransformInternal(EConstraintFram
 
 	// Skeletal case
 	USkeletalMeshComponent* SkelComp = Cast<USkeletalMeshComponent>(PrimComp);
-	if(SkelComp != NULL && !SkelComp->bUseSingleBodyPhysics)
+	if(SkelComp != NULL)
 	{
 		int32 BoneIndex = SkelComp->GetBoneIndex(InBoneName);
 		if(BoneIndex != INDEX_NONE)
@@ -113,7 +113,7 @@ FBox UPhysicsConstraintComponent::GetBodyBoxInternal(EConstraintFrame::Type Fram
 
 	// Skeletal case
 	USkeletalMeshComponent* SkelComp = Cast<USkeletalMeshComponent>(PrimComp);
-	if(SkelComp != NULL && !SkelComp->bUseSingleBodyPhysics)
+	if(SkelComp != NULL)
 	{
 		UPhysicsAsset * const PhysicsAsset = SkelComp->GetPhysicsAsset();
 		if (PhysicsAsset)
diff --git a/Engine/Source/Runtime/Engine/Private/PrimitiveComponent.cpp b/Engine/Source/Runtime/Engine/Private/PrimitiveComponent.cpp
--- a/Engine/Source/Runtime/Engine/Private/PrimitiveComponent.cpp
+++ b/Engine/Source/Runtime/Engine/Private/PrimitiveComponent.cpp
@@ -1561,7 +1561,7 @@ bool UPrimitiveComponent::ComponentOverlapComponent(class UPrimitiveComponent* P
 {
 	// if target is skeletalmeshcomponent and do not support singlebody physics
 	USkeletalMeshComponent * OtherComp = Cast<USkeletalMeshComponent>(PrimComp);
-	if (OtherComp && !OtherComp->bUseSingleBodyPhysics)
+	if (OtherComp)
 	{
 		UE_LOG(LogCollision, Log, TEXT("ComponentOverlapMulti : (%s) Does not support skeletalmesh with Physics Asset"), *PrimComp->GetPathName());
 		return false;
diff --git a/Engine/Source/Runtime/Engine/Private/SkeletalMeshComponent.cpp b/Engine/Source/Runtime/Engine/Private/SkeletalMeshComponent.cpp
--- a/Engine/Source/Runtime/Engine/Private/SkeletalMeshComponent.cpp
+++ b/Engine/Source/Runtime/Engine/Private/SkeletalMeshComponent.cpp
@@ -1200,11 +1200,6 @@ void USkeletalMeshComponent::UnHideBone( int32 BoneIndex )
 
 bool USkeletalMeshComponent::IsAnySimulatingPhysics() const
 {
-	if( bUseSingleBodyPhysics )
-	{
-		return IsSimulatingPhysics();
-	}
-
 	for ( int32 BodyIndex=0; BodyIndex<Bodies.Num(); ++BodyIndex )
 	{
 		if (Bodies[BodyIndex]->IsInstanceSimulatingPhysics())
@@ -1592,51 +1587,44 @@ FTransform USkeletalMeshComponent::ConvertLocalRootMotionToWorld(const FTransfor
 
 float USkeletalMeshComponent::CalculateMass(FName BoneName)
 {
-	if (bUseSingleBodyPhysics)
-	{
-		return UPrimitiveComponent::CalculateMass();
-	}
-	else
-	{
-		float Mass = 0.0f;
+	float Mass = 0.0f;
 
-		if (Bodies.Num())
+	if (Bodies.Num())
+	{
+		for (int32 i = 0; i < Bodies.Num(); ++i)
 		{
-			for (int32 i = 0; i < Bodies.Num(); ++i)
+			//if bone name is not provided calculate entire mass - otherwise get mass for just the bone
+			if (Bodies[i]->BodySetup.IsValid() && (BoneName == NAME_None || BoneName == Bodies[i]->BodySetup->BoneName))
 			{
-				//if bone name is not provided calculate entire mass - otherwise get mass for just the bone
-				if (Bodies[i]->BodySetup.IsValid() && (BoneName == NAME_None || BoneName == Bodies[i]->BodySetup->BoneName))
-				{
-					Mass += Bodies[i]->BodySetup->CalculateMass(this);
-				}
+				Mass += Bodies[i]->BodySetup->CalculateMass(this);
 			}
 		}
-		else	//We want to calculate mass before we've initialized body instances - in this case use physics asset setup
+	}
+	else	//We want to calculate mass before we've initialized body instances - in this case use physics asset setup
+	{
+		TArray<class UBodySetup*> * BodySetups = NULL;
+		if (PhysicsAssetOverride)
+		{
+			BodySetups = &PhysicsAssetOverride->BodySetup;
+		}
+		else if (UPhysicsAsset * PhysicsAsset = GetPhysicsAsset())
 		{
-			TArray<class UBodySetup*> * BodySetups = NULL;
-			if (PhysicsAssetOverride)
-			{
-				BodySetups = &PhysicsAssetOverride->BodySetup;
-			}
-			else if (UPhysicsAsset * PhysicsAsset = GetPhysicsAsset())
-			{
-				BodySetups = &PhysicsAsset->BodySetup;
-			}
+			BodySetups = &PhysicsAsset->BodySetup;
+		}
 
-			if (BodySetups)
+		if (BodySetups)
+		{
+			for (int32 i = 0; i < BodySetups->Num(); ++i)
 			{
-				for (int32 i = 0; i < BodySetups->Num(); ++i)
+				if ((*BodySetups)[i] && (BoneName == NAME_None || BoneName == (*BodySetups)[i]->BoneName))
 				{
-					if ((*BodySetups)[i] && (BoneName == NAME_None || BoneName == (*BodySetups)[i]->BoneName))
-					{
-						Mass += (*BodySetups)[i]->CalculateMass(this);
-					}
+					Mass += (*BodySetups)[i]->CalculateMass(this);
 				}
 			}
 		}
-
-		return Mass;
 	}
+
+	return Mass;
 }
 
 #if WITH_EDITOR
diff --git a/Engine/Source/Runtime/Engine/Private/SkeletalMeshComponentPhysics.cpp b/Engine/Source/Runtime/Engine/Private/SkeletalMeshComponentPhysics.cpp
--- a/Engine/Source/Runtime/Engine/Private/SkeletalMeshComponentPhysics.cpp
+++ b/Engine/Source/Runtime/Engine/Private/SkeletalMeshComponentPhysics.cpp
@@ -379,12 +379,6 @@ void USkeletalMeshComponent::AddRadialImpulse(FVector Origin, float Radius, floa
 		return;
 	}
 
-	if(bUseSingleBodyPhysics)
-	{
-		Super::AddRadialImpulse(Origin, Radius, Strength, Falloff, bVelChange);
-		return;
-	}
-
 	for(int32 i=0; i<Bodies.Num(); i++)
 	{
 		Bodies[i]->AddRadialImpulseToBody(Origin, Radius, Strength, Falloff, bVelChange);
@@ -400,12 +394,6 @@ void USkeletalMeshComponent::AddRadialForce(FVector Origin, float Radius, float
 		return;
 	}
 
-	if(bUseSingleBodyPhysics)
-	{
-		Super::AddRadialForce(Origin, Radius, Strength, Falloff);
-		return;
-	}
-
 	for(int32 i=0; i<Bodies.Num(); i++)
 	{
 		Bodies[i]->AddRadialForceToBody(Origin, Radius, Strength, Falloff);
@@ -415,35 +403,21 @@ void USkeletalMeshComponent::AddRadialForce(FVector Origin, float Radius, float
 
 void USkeletalMeshComponent::WakeAllRigidBodies()
 {
-	if(bUseSingleBodyPhysics)
-	{
-		Super::WakeAllRigidBodies();
-	}
-	else
+	for (int32 i=0; i < Bodies.Num(); i++)
 	{
-		for (int32 i=0; i < Bodies.Num(); i++)
-		{
-			FBodyInstance* BI = Bodies[i];
-			check(BI);
-			BI->WakeInstance();
-		}
+		FBodyInstance* BI = Bodies[i];
+		check(BI);
+		BI->WakeInstance();
 	}
 }
 
 void USkeletalMeshComponent::PutAllRigidBodiesToSleep()
 {
-	if(bUseSingleBodyPhysics)
+	for (int32 i=0; i < Bodies.Num(); i++)
 	{
-		Super::PutAllRigidBodiesToSleep();
-	}
-	else
-	{
-		for (int32 i=0; i < Bodies.Num(); i++)
-		{
-			FBodyInstance* BI = Bodies[i];
-			check(BI);
-			BI->PutInstanceToSleep();
-		}
+		FBodyInstance* BI = Bodies[i];
+		check(BI);
+		BI->PutInstanceToSleep();
 	}
 }
 
@@ -452,23 +426,16 @@ bool USkeletalMeshComponent::IsAnyRigidBodyAwake()
 {
 	bool bAwake = false;
 
-	if(bUseSingleBodyPhysics)
-	{
-		bAwake = Super::IsAnyRigidBodyAwake();
-	}
-	else
+	// ..iterate over each body to find any that are awak
+	for(int32 i=0; i<Bodies.Num(); i++)
 	{
-		// ..iterate over each body to find any that are awak
-		for(int32 i=0; i<Bodies.Num(); i++)
+		FBodyInstance* BI = Bodies[i];
+		check(BI);
+		if(BI->IsInstanceAwake())
 		{
-			FBodyInstance* BI = Bodies[i];
-			check(BI);
-			if(BI->IsInstanceAwake())
-			{
-				// Found an awake one - so mesh is considered 'awake'
-				bAwake = true;
-				continue;
-			}
+			// Found an awake one - so mesh is considered 'awake'
+			bAwake = true;
+			continue;
 		}
 	}
 
@@ -478,28 +445,17 @@ bool USkeletalMeshComponent::IsAnyRigidBodyAwake()
 
 void USkeletalMeshComponent::SetAllPhysicsLinearVelocity(FVector NewVel, bool bAddToCurrent)
 {
-	if(bUseSingleBodyPhysics)
+	for (int32 i=0; i < Bodies.Num(); i++)
 	{
-		Super::SetAllPhysicsLinearVelocity(NewVel, bAddToCurrent);
-	}
-	else
-	{
-		for (int32 i=0; i < Bodies.Num(); i++)
-		{
-			FBodyInstance* BodyInstance = Bodies[i];
-			check(BodyInstance);
-			BodyInstance->SetLinearVelocity(NewVel, bAddToCurrent);
-		}
+		FBodyInstance* BodyInstance = Bodies[i];
+		check(BodyInstance);
+		BodyInstance->SetLinearVelocity(NewVel, bAddToCurrent);
 	}
 }
 
 void USkeletalMeshComponent::SetAllPhysicsAngularVelocity(FVector const& NewAngVel, bool bAddToCurrent)
 {
-	if(bUseSingleBodyPhysics)
-	{
-		Super::SetAllPhysicsAngularVelocity(NewAngVel, bAddToCurrent);
-	}
-	else if(RootBodyIndex < Bodies.Num())
+	if(RootBodyIndex < Bodies.Num())
 	{
 		// Find the root actor. We use its location as the center of the rotation.
 		FBodyInstance* RootBodyInst = Bodies[ RootBodyIndex ];
@@ -521,11 +477,7 @@ void USkeletalMeshComponent::SetAllPhysicsAngularVelocity(FVector const& NewAngV
 
 void USkeletalMeshComponent::SetAllPhysicsPosition(FVector NewPos)
 {
-	if(bUseSingleBodyPhysics)
-	{
-		Super::SetAllPhysicsPosition(NewPos);
-	}
-	else if(RootBodyIndex < Bodies.Num())
+	if(RootBodyIndex < Bodies.Num())
 	{
 		// calculate the deltas to get the root body to NewPos
 		FBodyInstance* RootBI = Bodies[RootBodyIndex];
@@ -565,11 +517,7 @@ void USkeletalMeshComponent::SetAllPhysicsPosition(FVector NewPos)
 
 void USkeletalMeshComponent::SetAllPhysicsRotation(FRotator NewRot)
 {
-	if(bUseSingleBodyPhysics)
-	{
-		Super::SetAllPhysicsRotation(NewRot);
-	}
-	else if(RootBodyIndex < Bodies.Num())
+	if(RootBodyIndex < Bodies.Num())
 	{
 		// calculate the deltas to get the root body to NewRot
 		FBodyInstance* RootBI = Bodies[RootBodyIndex];
@@ -1174,16 +1122,7 @@ void USkeletalMeshComponent::OnUpdateTransform(bool bSkipPhysicsMove)
 	// Always send new transform to physics
 	if(bPhysicsStateCreated && !bSkipPhysicsMove )
 	{
-		// In the single body case - treat it like a regular one-body component
-		if (bUseSingleBodyPhysics)
-		{
-			// @todo UE4 rather than always pass false, this function should know if it is being teleported or not!
-			SendPhysicsTransform(false);
-		}
-		else
-		{
-			UpdateKinematicBonesToPhysics(false);
-		}
+		UpdateKinematicBonesToPhysics(false);
 	}
 
 #if WITH_APEX_CLOTHING
@@ -1203,18 +1142,11 @@ void USkeletalMeshComponent::UpdateOverlaps(TArray<FOverlapInfo> const* PendingO
 
 void USkeletalMeshComponent::CreatePhysicsState()
 {
-	if(bUseSingleBodyPhysics)
-	{
-		Super::CreatePhysicsState(); // use PrimitiveComponent implementation
-	}
-	else 
-	{
-		//	UE_LOG(LogSkeletalMesh, Warning, TEXT("Creating Physics State (%s : %s)"), *GetNameSafe(GetOuter()),  *GetName());
-		// Init physics
-		InitArticulated(World->GetPhysicsScene());
+	//	UE_LOG(LogSkeletalMesh, Warning, TEXT("Creating Physics State (%s : %s)"), *GetNameSafe(GetOuter()),  *GetName());
+	// Init physics
+	InitArticulated(World->GetPhysicsScene());
 
-		USceneComponent::CreatePhysicsState(); // Need to route CreatePhysicsState, skip PrimitiveComponent
-	}
+	USceneComponent::CreatePhysicsState(); // Need to route CreatePhysicsState, skip PrimitiveComponent
 }
 
 
@@ -1330,11 +1262,6 @@ FName USkeletalMeshComponent::FindConstraintBoneName( int32 ConstraintIndex )
 FBodyInstance* USkeletalMeshComponent::GetBodyInstance(FName BoneName) const
 {
 	UPhysicsAsset * const PhysicsAsset = GetPhysicsAsset();
-	if( bUseSingleBodyPhysics )
-	{
-		return Super::GetBodyInstance(BoneName);
-	}
-
 	FBodyInstance* BodyInst = NULL;
 
 	if(PhysicsAsset != NULL)
@@ -1476,19 +1403,12 @@ void USkeletalMeshComponent::UpdateHasValidBodies()
 
 void USkeletalMeshComponent::UpdatePhysicsToRBChannels()
 {
-	if(bUseSingleBodyPhysics)
-	{
-		Super::UpdatePhysicsToRBChannels();
-	}
-	else
+	// Iterate over each bone/body.
+	for (int32 i = 0; i < Bodies.Num(); i++)
 	{
-		// Iterate over each bone/body.
-		for (int32 i = 0; i < Bodies.Num(); i++)
-		{
-			FBodyInstance* BI = Bodies[i];
-			check(BI);
-			BI->UpdatePhysicsFilterData();
-		}
+		FBodyInstance* BI = Bodies[i];
+		check(BI);
+		BI->UpdatePhysicsFilterData();
 	}
 }
 
@@ -1500,11 +1420,6 @@ extern float DebugLineLifetime;
 bool USkeletalMeshComponent::LineTraceComponent(struct FHitResult& OutHit, const FVector Start, const FVector End, const struct FCollisionQueryParams& Params)
 {
 	UPhysicsAsset * const PhysicsAsset = GetPhysicsAsset();
-	if( bUseSingleBodyPhysics )
-	{
-		return Super::LineTraceComponent(OutHit, Start, End, Params);
-	}
-
 	bool bHaveHit = false;
 
 	float MinTime = MAX_FLT;
@@ -1539,11 +1454,6 @@ bool USkeletalMeshComponent::LineTraceComponent(struct FHitResult& OutHit, const
 
 bool USkeletalMeshComponent::SweepComponent( FHitResult& OutHit, const FVector Start, const FVector End, const FCollisionShape& CollisionShape, bool bTraceComplex)
 {
-	if( bUseSingleBodyPhysics )
-	{
-		return Super::SweepComponent(OutHit,Start,End, CollisionShape, bTraceComplex);
-	}
-
 	bool bHaveHit = false;
 
 	for (int32 BodyIdx=0; BodyIdx < Bodies.Num(); ++BodyIdx)
@@ -1569,17 +1479,12 @@ bool USkeletalMeshComponent::ComponentOverlapComponent(class UPrimitiveComponent
 
 	// if target is skeletalmeshcomponent and do not support singlebody physics
 	USkeletalMeshComponent * OtherComp = Cast<USkeletalMeshComponent>(PrimComp);
-	if (OtherComp && !OtherComp->bUseSingleBodyPhysics)
+	if (OtherComp)
 	{
 		UE_LOG(LogCollision, Log, TEXT("ComponentOverlapComponent : (%s) Does not support skeletalmesh with Physics Asset"), *PrimComp->GetPathName());
 		return false;
 	}
 
-	if( bUseSingleBodyPhysics )
-	{
-		return Super::ComponentOverlapComponent(PrimComp,Pos,Rot, Params);
-	}
-
 	// calculate the test global pose of the actor
 	PxTransform PTestGlobalPose = U2PTransform(FTransform(Rot, Pos));
 
@@ -1635,10 +1540,6 @@ bool USkeletalMeshComponent::ComponentOverlapComponent(class UPrimitiveComponent
 
 bool USkeletalMeshComponent::OverlapComponent(const FVector& Pos, const FQuat& Rot, const FCollisionShape& CollisionShape)
 {
-	if( bUseSingleBodyPhysics )
-	{
-		return Super::OverlapComponent(Pos, Rot, CollisionShape);
-	}
 
 #if WITH_PHYSX
 	PxSphereGeometry PSphereGeom;
@@ -1695,11 +1596,6 @@ bool USkeletalMeshComponent::ComponentOverlapMulti(TArray<struct FOverlapResult>
 {
 #if WITH_PHYSX
 
-	if(bUseSingleBodyPhysics)
-	{
-		return UPrimitiveComponent::ComponentOverlapMulti(OutOverlaps, World, Pos, Rot, TestChannel, Params, ObjectQueryParams);
-	}
-
 	OutOverlaps.Reset();
 
 	if(!Bodies.IsValidIndex(RootBodyIndex))
@@ -4254,25 +4150,17 @@ void USkeletalMeshComponent::DrawClothingPhysicalMeshWire(FPrimitiveDrawInterfac
 
 float USkeletalMeshComponent::GetMass() const
 {
-	if (bUseSingleBodyPhysics)
+	float Mass = 0.0f;
+	for (int32 i=0; i < Bodies.Num(); ++i)
 	{
-		return Super::GetMass();
-	}
-	else
-	{
-		float Mass = 0.0f;
-		for (int32 i=0; i < Bodies.Num(); ++i)
-		{
-			FBodyInstance* BI = Bodies[i];
+		FBodyInstance* BI = Bodies[i];
 
-			if (BI->IsValidBodyInstance())
-			{
-				Mass += BI->GetBodyMass();
-			}
+		if (BI->IsValidBodyInstance())
+		{
+			Mass += BI->GetBodyMass();
 		}
-		return Mass;
 	}
-
+	return Mass;
 }
 
 // blueprint callable methods 
diff --git a/Engine/Source/Runtime/Engine/Private/WheeledVehicle.cpp b/Engine/Source/Runtime/Engine/Private/WheeledVehicle.cpp
--- a/Engine/Source/Runtime/Engine/Private/WheeledVehicle.cpp
+++ b/Engine/Source/Runtime/Engine/Private/WheeledVehicle.cpp
@@ -20,7 +20,6 @@ AWheeledVehicle::AWheeledVehicle(const class FPostConstructInitializeProperties&
 	Mesh->BodyInstance.bNotifyRigidBodyCollision = true;
 	Mesh->BodyInstance.bUseCCD = true;
 	Mesh->bBlendPhysics = true;
-	Mesh->bUseSingleBodyPhysics = false;
 	Mesh->bGenerateOverlapEvents = true;
 	RootComponent = Mesh;
 
