More logging for actors causing nav rebuilds

#UE4

diff --git a/Engine/Source/Runtime/Engine/Private/Actor.cpp b/Engine/Source/Runtime/Engine/Private/Actor.cpp
--- a/Engine/Source/Runtime/Engine/Private/Actor.cpp
+++ b/Engine/Source/Runtime/Engine/Private/Actor.cpp
@@ -2076,6 +2076,8 @@ void AActor::SetNavigationRelevancy(const bool bNewRelevancy)
 		{
 			if (bNewRelevancy == true)
 			{
+				UE_LOG(LogNavigation, Verbose, TEXT("AActor::SetNavigationRelevancy registering new relevant actor %s"), *GetFullName());
+
 				MyWorld->GetNavigationSystem()->RegisterNavigationRelevantActor(this);
 			}
 			else
