Add sub-stepping support on non APEX build
Fix sub-stepping for dedicated server

diff --git a/Engine/Source/Runtime/Engine/Private/PhysicsEngine/PhysScene.cpp b/Engine/Source/Runtime/Engine/Private/PhysicsEngine/PhysScene.cpp
--- a/Engine/Source/Runtime/Engine/Private/PhysicsEngine/PhysScene.cpp
+++ b/Engine/Source/Runtime/Engine/Private/PhysicsEngine/PhysScene.cpp
@@ -28,18 +28,6 @@ FORCEINLINE EPhysicsSceneType SceneType(const FBodyInstance * BodyInstance)
 	return PST_Sync;
 }
 
-/**
- * Return true if we should be running in single threaded mode, ala dedicated server
-**/
-FORCEINLINE static bool PhysSingleThreadedMode()
-{
-	if (IsRunningDedicatedServer() || FPlatformMisc::NumberOfCores() < 3 || !FPlatformProcess::SupportsMultithreading())
-	{
-		return true;
-	}
-	return false;
-}
-
 /**
  * Return true if we should lag the async scene a frame
 **/
@@ -299,20 +287,20 @@ void FPhysScene::DeferredDestructibleDamageNotify(const NxApexDamageEventReportD
 
 bool FPhysScene::SubstepSimulation(uint32 SceneType, FGraphEventRef &InOutCompletionEvent)
 {
+#if WITH_PHYSX
 	check(SceneType != PST_Cloth); //we don't bother sub-stepping cloth
 	float UseDelta = UseSyncTime(SceneType)? SyncDeltaSeconds : DeltaSeconds;
 	float SubTime = PhysSubSteppers[SceneType]->UpdateTime(UseDelta);
 	PxScene* PScene = GetPhysXScene(SceneType);
-#if WITH_APEX
-	NxApexScene* ApexScene = GetApexScene(SceneType);
-	if(!ApexScene || SubTime <= 0.f)
+	if(SubTime <= 0.f)
 	{
 		return false;
 	}else
 	{
 		//we have valid scene and subtime so enqueue task
 		PhysXCompletionTask* Task = new PhysXCompletionTask(InOutCompletionEvent, PScene->getTaskManager());
-		FSimpleDelegateGraphTask::CreateAndDispatchWhenReady(FSimpleDelegateGraphTask::FDelegate::CreateRaw(PhysSubSteppers[SceneType], &FPhysSubstepTask::StepSimulation, ApexScene, Task), TEXT("SubstepSimulationImp"));
+		ENamedThreads::Type NamedThread = PhysSingleThreadedMode() ? ENamedThreads::GameThread : ENamedThreads::AnyThread;
+		FSimpleDelegateGraphTask::CreateAndDispatchWhenReady(FSimpleDelegateGraphTask::FDelegate::CreateRaw(PhysSubSteppers[SceneType], &FPhysSubstepTask::StepSimulation, Task), TEXT("SubstepSimulationImp"), NULL, NamedThread);
 		return true;
 	}
 #endif
@@ -1119,7 +1107,13 @@ void FPhysScene::InitPhysScene(uint32 SceneType)
 #if WITH_SUBSTEPPING
 	//Initialize substeppers
 	//we don't bother sub-stepping cloth
+#if WITH_PHYSX
+#if WITH_APEX
+	PhysSubSteppers[SceneType] = SceneType == PST_Cloth ? NULL : new FPhysSubstepTask(ApexScene);
+#else
 	PhysSubSteppers[SceneType] = SceneType == PST_Cloth ? NULL : new FPhysSubstepTask(PScene);
+#endif
+#endif
 	if (SceneType == PST_Sync)
 	{
 		PhysSubSteppers[SceneType]->SetVehicleManager(VehicleManager);
diff --git a/Engine/Source/Runtime/Engine/Private/PhysicsEngine/PhysSubstepTasks.cpp b/Engine/Source/Runtime/Engine/Private/PhysicsEngine/PhysSubstepTasks.cpp
--- a/Engine/Source/Runtime/Engine/Private/PhysicsEngine/PhysSubstepTasks.cpp
+++ b/Engine/Source/Runtime/Engine/Private/PhysicsEngine/PhysSubstepTasks.cpp
@@ -12,7 +12,7 @@
 #include "PhysSubstepTasks.h"
 
 #if WITH_PHYSX
-FPhysSubstepTask::FPhysSubstepTask(PxScene * GivenPScene) :
+FPhysSubstepTask::FPhysSubstepTask(PxApexScene * GivenScene) :
 	NumSubsteps(0),
 	SubTime(0.f),
 	DeltaSeconds(0.f),
@@ -23,14 +23,10 @@ FPhysSubstepTask::FPhysSubstepTask(PxScene * GivenPScene) :
 	TotalSubTime(0.f),
 	CurrentSubStep(0),
 	VehicleManager(NULL),
-#if WITH_APEX
-	ApexScene(NULL),
-#endif
-	PScene(GivenPScene)
+	PAScene(GivenScene)
 {
-	check(PScene);
+	check(PAScene);
 }
-
 #endif
 
 void FPhysSubstepTask::SwapBuffers()
@@ -52,6 +48,7 @@ void FPhysSubstepTask::RemoveBodyInstance(FBodyInstance * BodyInstance)
 
 void FPhysSubstepTask::SetKinematicTarget(FBodyInstance * Body, const FTransform & TM)
 {
+#if WITH_PHYSX
 	check(Body);
 	TM.DiagnosticCheckNaN_All();
 
@@ -65,11 +62,12 @@ void FPhysSubstepTask::SetKinematicTarget(FBodyInstance * Body, const FTransform
 		TargetState.bKinematicTarget = true;
 		TargetState.KinematicTarget = KinmaticTarget;
 	}
-
+#endif
 }
 
 void FPhysSubstepTask::AddForce(FBodyInstance * Body, const FVector & Force)
 {
+#if WITH_PHYSX
 	check(Body);
 	PxRigidDynamic * PRigidDynamic = Body->GetPxRigidDynamic();
 	SCOPED_SCENE_READ_LOCK(PRigidDynamic->getScene());
@@ -84,10 +82,12 @@ void FPhysSubstepTask::AddForce(FBodyInstance * Body, const FVector & Force)
 		FPhysTarget & TargetState = PhysTargetBuffers[External].FindOrAdd(Body);
 		TargetState.Forces.Add(ForceTarget);
 	}
+#endif
 }
 
 void FPhysSubstepTask::AddForceAtPosition(FBodyInstance * Body, const FVector & Force, const FVector & Position)
 {
+#if WITH_PHYSX
 	check(Body);
 
 	PxRigidDynamic * PRigidDynamic = Body->GetPxRigidDynamic();
@@ -103,10 +103,11 @@ void FPhysSubstepTask::AddForceAtPosition(FBodyInstance * Body, const FVector &
 		FPhysTarget & TargetState = PhysTargetBuffers[External].FindOrAdd(Body);
 		TargetState.Forces.Add(ForceTarget);
 	}
-
+#endif
 }
 void FPhysSubstepTask::AddTorque(FBodyInstance * Body, const FVector & Torque)
 {
+#if WITH_PHYSX
 	check(Body);
 
 	PxRigidDynamic * PRigidDynamic = Body->GetPxRigidDynamic();
@@ -120,6 +121,7 @@ void FPhysSubstepTask::AddTorque(FBodyInstance * Body, const FVector & Torque)
 		FPhysTarget & TargetState = PhysTargetBuffers[External].FindOrAdd(Body);
 		TargetState.Torques.Add(TorqueTarget);
 	}
+#endif
 }
 
 /** Applies forces - Assumes caller has obtained writer lock */
@@ -192,6 +194,12 @@ void FPhysSubstepTask::InterpolateKinematicActor(const FPhysTarget & PhysTarget,
 void FPhysSubstepTask::SubstepInterpolation(float Alpha)
 {
 #if WITH_PHYSX
+#if WITH_APEX
+	PxScene * PScene = PAScene->getPhysXScene();
+#else
+	PxScene * PScene = PAScene;
+#endif
+
 	PhysTargetMap & Targets = PhysTargetBuffers[!External];
 
 	/** Note: We lock the entire scene before iterating. The assumption is that removing an FBodyInstance from the map will also be wrapped by this lock */
@@ -246,8 +254,8 @@ float FPhysSubstepTask::UpdateTime(float UseDelta)
 	return SubTime;
 }
 
-#if WITH_APEX
-void FPhysSubstepTask::StepSimulation(NxApexScene * GivenApexScene, PhysXCompletionTask * Task)
+#if WITH_PHYSX
+void FPhysSubstepTask::StepSimulation(PhysXCompletionTask * Task)
 {
 	check(SubTime > 0.f);
 	check(DeltaSeconds > 0.f);
@@ -257,22 +265,24 @@ void FPhysSubstepTask::StepSimulation(NxApexScene * GivenApexScene, PhysXComplet
 	StepScale = SubTime / DeltaSeconds;
 	TotalSubTime = 0.f;
 	CurrentSubStep = 0;
-	ApexScene = GivenApexScene;
 
 	SubstepSimulationStart();
 }
+#endif
 
 void FPhysSubstepTask::SubstepSimulationStart()
 {
+#if WITH_PHYSX
 	check(SubTime > 0.f);
 	check(DeltaSeconds > 0.f);
 	
 	check(!CompletionEvent.GetReference());	//should be done
 	CompletionEvent = FGraphEvent::CreateGraphEvent();
-	PhysXCompletionTask* SubstepTask = new PhysXCompletionTask(CompletionEvent, ApexScene->getTaskManager());
-	FDelegateGraphTask::CreateAndDispatchWhenReady(FDelegateGraphTask::FDelegate::CreateRaw(this, &FPhysSubstepTask::SubstepSimulationEnd), TEXT("ProcessPhysSubstepSimulation"), CompletionEvent);
+	PhysXCompletionTask* SubstepTask = new PhysXCompletionTask(CompletionEvent, PAScene->getTaskManager());
+	ENamedThreads::Type NamedThread = PhysSingleThreadedMode() ? ENamedThreads::GameThread : ENamedThreads::AnyThread;
+	FDelegateGraphTask::CreateAndDispatchWhenReady(FDelegateGraphTask::FDelegate::CreateRaw(this, &FPhysSubstepTask::SubstepSimulationEnd), TEXT("ProcessPhysSubstepSimulation"), CompletionEvent, NamedThread, NamedThread);
 
-	++CurrentSubStep;
+	++CurrentSubStep;	
 
 	bool bLastSubstep = CurrentSubStep >= NumSubsteps;
 
@@ -292,17 +302,32 @@ void FPhysSubstepTask::SubstepSimulationStart()
 	}
 
 	SubstepInterpolation(Interpolation);
-	ApexScene->simulate(DeltaTime, bLastSubstep, SubstepTask);
+
+#if WITH_APEX
+	PAScene->simulate(DeltaTime, bLastSubstep, SubstepTask);
+#else
+	PAScene->lockWrite();
+	PAScene->simulate(DeltaTime, SubstepTask);
+	PAScene->unlockWrite();
+#endif
 	SubstepTask->removeReference();
+#endif
 }
 
 void FPhysSubstepTask::SubstepSimulationEnd(ENamedThreads::Type CurrentThread, const FGraphEventRef& MyCompletionGraphEvent)
 {
+#if WITH_PHYSX
 	CompletionEvent = NULL;
 	if (CurrentSubStep < NumSubsteps)
 	{
 		uint32 OutErrorCode = 0;
-		ApexScene->fetchResults(true, &OutErrorCode);
+#if WITH_APEX
+		PAScene->fetchResults(true, &OutErrorCode);
+#else
+		PAScene->lockWrite();
+		PAScene->fetchResults(true, &OutErrorCode);
+		PAScene->unlockWrite();
+#endif
 		if (OutErrorCode != 0)
 		{
 			UE_LOG(LogPhysics, Log, TEXT("PHYSX FETCHRESULTS ERROR: %d"), OutErrorCode);
@@ -315,9 +340,8 @@ void FPhysSubstepTask::SubstepSimulationEnd(ENamedThreads::Type CurrentThread, c
 		//final step we call fetch on in game thread
 		FullSimulationTask->removeReference();
 	}
-
-}
 #endif
+}
 
 
 void FPhysSubstepTask::SetVehicleManager(FPhysXVehicleManager * InVehicleManager)
diff --git a/Engine/Source/Runtime/Engine/Private/PhysicsEngine/PhysSubstepTasks.h b/Engine/Source/Runtime/Engine/Private/PhysicsEngine/PhysSubstepTasks.h
--- a/Engine/Source/Runtime/Engine/Private/PhysicsEngine/PhysSubstepTasks.h
+++ b/Engine/Source/Runtime/Engine/Private/PhysicsEngine/PhysSubstepTasks.h
@@ -88,13 +88,22 @@ struct FPhysTarget
 	
 };
 
+#if WITH_PHYSX
+#if WITH_APEX
+typedef NxApexScene PxApexScene;	//helper typedef so we don't have to use as many ifdefs
+#else
+typedef PxScene PxApexScene;
+#endif
+#endif
+
 /** Holds information used for substepping a scene */
 class FPhysSubstepTask
 {
 public:
 #if WITH_PHYSX
-	FPhysSubstepTask(PxScene * GivenPScene);
+	FPhysSubstepTask(PxApexScene * GivenScene);
 #endif
+
 	void SetKinematicTarget(FBodyInstance * Body, const FTransform & TM);
 	void AddForce(FBodyInstance * Body, const FVector & Force);
 	void AddForceAtPosition(FBodyInstance * Body, const FVector & Force, const FVector & Position);
@@ -108,11 +117,10 @@ public:
 	float UpdateTime(float UseDelta);
 
 	void SetVehicleManager(class FPhysXVehicleManager *	InVehicleManager);
-
-#if WITH_APEX
-	void StepSimulation(NxApexScene * ApexScene, PhysXCompletionTask * Task);
 	void SubstepSimulationStart();
 	void SubstepSimulationEnd(ENamedThreads::Type CurrentThread, const FGraphEventRef& MyCompletionGraphEvent);
+#if WITH_PHYSX
+	void StepSimulation(PhysXCompletionTask * Task);
 #endif
 		
 private:
@@ -139,12 +147,8 @@ private:
 	/** Vehicle scene */
 	class FPhysXVehicleManager*			VehicleManager;
 
-#if WITH_APEX
-	class NxApexScene * ApexScene;
-#endif
-
 #if WITH_PHYSX
-	PxScene * PScene;
+	PxApexScene * PAScene;
 #endif
 
 };
diff --git a/Engine/Source/Runtime/Engine/Public/Engine.h b/Engine/Source/Runtime/Engine/Public/Engine.h
--- a/Engine/Source/Runtime/Engine/Public/Engine.h
+++ b/Engine/Source/Runtime/Engine/Public/Engine.h
@@ -30,9 +30,8 @@
 #endif
 
 #if WITH_APEX
-
 #ifndef WITH_SUBSTEPPING
-	#define WITH_SUBSTEPPING (1 && WITH_APEX)
+	#define WITH_SUBSTEPPING 1
 #endif
 
 #ifndef WITH_APEX_CLOTHING
diff --git a/Engine/Source/Runtime/Engine/Public/PhysicsPublic.h b/Engine/Source/Runtime/Engine/Public/PhysicsPublic.h
--- a/Engine/Source/Runtime/Engine/Public/PhysicsPublic.h
+++ b/Engine/Source/Runtime/Engine/Public/PhysicsPublic.h
@@ -448,6 +448,18 @@ private:
 	TMap< uint32, TMap<struct FRigidBodyIndexPair, bool>* >		CollisionDisableTableLookup;
 };
 
+/**
+* Return true if we should be running in single threaded mode, ala dedicated server
+**/
+FORCEINLINE bool PhysSingleThreadedMode()
+{
+	if (IsRunningDedicatedServer() || FPlatformMisc::NumberOfCores() < 3 || !FPlatformProcess::SupportsMultithreading())
+	{
+		return true;
+	}
+	return false;
+}
+
 #if WITH_PHYSX
 /** Struct used for passing info to the PhysX shader */
 
