Add support for taking orbit offset in to account when generating particle spawn and death events

diff --git a/Engine/Source/Runtime/Engine/Classes/Particles/Event/ParticleModuleEventGenerator.h b/Engine/Source/Runtime/Engine/Classes/Particles/Event/ParticleModuleEventGenerator.h
--- a/Engine/Source/Runtime/Engine/Classes/Particles/Event/ParticleModuleEventGenerator.h
+++ b/Engine/Source/Runtime/Engine/Classes/Particles/Event/ParticleModuleEventGenerator.h
@@ -35,6 +35,10 @@ struct FParticleEvent_GenerateInfo
 	UPROPERTY(EditAnywhere, Category=ParticleEvent_GenerateInfo)
 	uint32 UseReflectedImpactVector:1;
 
+	/** Use the orbit offset when computing the position at which the event occurred. */
+	UPROPERTY(EditAnywhere, Category=ParticleEvent_GenerateInfo)
+	uint32 bUseOrbitOffset:1;
+
 	/** Should the event tag with a custom name? Leave blank for the default. */
 	UPROPERTY(EditAnywhere, Category=ParticleEvent_GenerateInfo)
 	FName CustomName;
diff --git a/Engine/Source/Runtime/Engine/Private/Particles/ParticleEmitterInstances.cpp b/Engine/Source/Runtime/Engine/Private/Particles/ParticleEmitterInstances.cpp
--- a/Engine/Source/Runtime/Engine/Private/Particles/ParticleEmitterInstances.cpp
+++ b/Engine/Source/Runtime/Engine/Private/Particles/ParticleEmitterInstances.cpp
@@ -962,6 +962,45 @@ FBox FParticleEmitterInstance::GetBoundingBox()
 	return ParticleBoundingBox;
 }
 
+int32 FParticleEmitterInstance::GetOrbitPayloadOffset()
+{
+	UParticleLODLevel* LODLevel = GetCurrentLODLevelChecked();
+
+	int32 OrbitOffsetValue = -1;
+	if (LODLevel->OrbitModules.Num() > 0)
+	{
+		UParticleLODLevel* HighestLODLevel = SpriteTemplate->LODLevels[0];
+		check(HighestLODLevel);
+
+		UParticleModuleOrbit* OrbitModule = HighestLODLevel->OrbitModules[LODLevel->OrbitModules.Num() - 1];
+		if (OrbitModule)
+		{
+			uint32* OrbitOffsetIndex = ModuleOffsetMap.Find(OrbitModule);
+			if (OrbitOffsetIndex)
+			{
+				OrbitOffsetValue = *OrbitOffsetIndex;
+			}
+		}
+	}
+	return OrbitOffsetValue;
+}
+
+FVector FParticleEmitterInstance::GetParticleLocationWithOrbitOffset(FBaseParticle* Particle)
+{
+	int32 OrbitOffsetValue = GetOrbitPayloadOffset();
+	if (OrbitOffsetValue == -1)
+	{
+		return Particle->Location;
+	}
+	else
+	{
+		int32 CurrentOffset = OrbitOffsetValue;
+		const uint8* ParticleBase = (const uint8*)Particle;
+		PARTICLE_ELEMENT(FOrbitChainModuleInstancePayload, OrbitPayload);
+		return Particle->Location + OrbitPayload.Offset;
+	}
+}
+
 /**
  *	Update the bounding box for the emitter
  *
@@ -1005,19 +1044,7 @@ void FParticleEmitterInstance::UpdateBoundingBox(float DeltaTime)
 		}
 
 		// Store off the orbit offset, if there is one
-		int32 OrbitOffsetValue = -1;
-		if (LODLevel->OrbitModules.Num() > 0)
-		{
-			UParticleModuleOrbit* OrbitModule = HighestLODLevel->OrbitModules[LODLevel->OrbitModules.Num() - 1];
-			if (OrbitModule)
-			{
-				uint32* OrbitOffsetIndex = ModuleOffsetMap.Find(OrbitModule);
-				if (OrbitOffsetIndex)
-				{
-					OrbitOffsetValue = *OrbitOffsetIndex;
-				}
-			}
-		}
+		int32 OrbitOffsetValue = GetOrbitPayloadOffset();
 
 		// For each particle, offset the box appropriately 
 		FVector MinVal(HALF_WORLD_MAX);
@@ -1131,19 +1158,7 @@ void FParticleEmitterInstance::ForceUpdateBoundingBox()
 		ParticleBoundingBox.Init();
 
 		// Store off the orbit offset, if there is one
-		int32 OrbitOffsetValue = -1;
-		if (LODLevel->OrbitModules.Num() > 0)
-		{
-			UParticleModuleOrbit* OrbitModule = HighestLODLevel->OrbitModules[LODLevel->OrbitModules.Num() - 1];
-			if (OrbitModule)
-			{
-				uint32* OrbitOffsetIndex = ModuleOffsetMap.Find(OrbitModule);
-				if (OrbitOffsetIndex)
-				{
-					OrbitOffsetValue = *OrbitOffsetIndex;
-				}
-			}
-		}
+		int32 OrbitOffsetValue = GetOrbitPayloadOffset();
 
 		const bool bUseLocalSpace = LODLevel->RequiredModule->bUseLocalSpace;
 
diff --git a/Engine/Source/Runtime/Engine/Private/Particles/ParticleModules_Event.cpp b/Engine/Source/Runtime/Engine/Private/Particles/ParticleModules_Event.cpp
--- a/Engine/Source/Runtime/Engine/Private/Particles/ParticleModules_Event.cpp
+++ b/Engine/Source/Runtime/Engine/Private/Particles/ParticleModules_Event.cpp
@@ -114,8 +114,10 @@ bool UParticleModuleEventGenerator::HandleParticleSpawned(FParticleEmitterInstan
 		{
 			if (EventGenInfo.Frequency == 0 || (EventPayload->SpawnTrackingCount % EventGenInfo.Frequency) == 0)
 			{
+				FVector ParticleLocation = EventGenInfo.bUseOrbitOffset ? Owner->GetParticleLocationWithOrbitOffset(NewParticle) : NewParticle->Location;
+
 				Owner->Component->ReportEventSpawn(EventGenInfo.CustomName, Owner->EmitterTime, 
-					NewParticle->Location, NewParticle->Velocity, EventGenInfo.ParticleModuleEventsToSendToGame);
+					ParticleLocation, NewParticle->Velocity, EventGenInfo.ParticleModuleEventsToSendToGame);
 				bProcessed = true;
 #if !(UE_BUILD_SHIPPING || UE_BUILD_TEST)
 				Owner->EventCount++;
@@ -142,8 +144,10 @@ bool UParticleModuleEventGenerator::HandleParticleKilled(FParticleEmitterInstanc
 		{
 			if (EventGenInfo.Frequency == 0 || (EventPayload->DeathTrackingCount % EventGenInfo.Frequency) == 0)
 			{
+				FVector ParticleLocation = EventGenInfo.bUseOrbitOffset ? Owner->GetParticleLocationWithOrbitOffset(DeadParticle) : DeadParticle->Location;
+
 				Owner->Component->ReportEventDeath(EventGenInfo.CustomName, 
-					Owner->EmitterTime, DeadParticle->Location, DeadParticle->Velocity, 
+					Owner->EmitterTime, ParticleLocation, DeadParticle->Velocity, 
 					EventGenInfo.ParticleModuleEventsToSendToGame, DeadParticle->RelativeTime);
 				bProcessed = true;
 #if !(UE_BUILD_SHIPPING || UE_BUILD_TEST)
diff --git a/Engine/Source/Runtime/Engine/Public/ParticleEmitterInstances.h b/Engine/Source/Runtime/Engine/Public/ParticleEmitterInstances.h
--- a/Engine/Source/Runtime/Engine/Public/ParticleEmitterInstances.h
+++ b/Engine/Source/Runtime/Engine/Public/ParticleEmitterInstances.h
@@ -474,6 +474,12 @@ public:
 		bHaltSpawning = bInHaltSpawning;
 	}
 
+	/** Get the offset of the orbit payload. */
+	int32 GetOrbitPayloadOffset();
+
+	/** Get the position of the particle taking orbit in to account. */
+	FVector GetParticleLocationWithOrbitOffset(FBaseParticle* Particle);
+
 	virtual FBaseParticle* GetParticle(int32 Index);
 	/**
 	 *	Get the physical index of the particle at the given index
