#ue4 - Fix CheatManager::Walk() / Fly() / Ghost() functions to warp network clients around.

- This was caused by the Possess() call in Walk(). There is no need to repossess the pawn.

diff --git a/Engine/Source/Runtime/Engine/Private/CheatManager.cpp b/Engine/Source/Runtime/Engine/Private/CheatManager.cpp
--- a/Engine/Source/Runtime/Engine/Private/CheatManager.cpp
+++ b/Engine/Source/Runtime/Engine/Private/CheatManager.cpp
@@ -194,8 +194,6 @@ void UCheatManager::Walk()
 				Character->ClientCheatWalk_Implementation();
 			}
 		}
-		
-		GetOuterAPlayerController()->Possess( Pawn );
 	}
 }
 
