Fix tab-complete to pick the previewed option

diff --git a/Engine/Source/Runtime/Engine/Private/UserInterface/Console.cpp b/Engine/Source/Runtime/Engine/Private/UserInterface/Console.cpp
--- a/Engine/Source/Runtime/Engine/Private/UserInterface/Console.cpp
+++ b/Engine/Source/Runtime/Engine/Private/UserInterface/Console.cpp
@@ -589,6 +589,12 @@ bool UConsole::InputKey_InputLine( int32 ControllerId, FKey Key, EInputEvent Eve
 				SetInputText("");
 				SetCursorPos(0);
 				HistoryCur = HistoryTop;
+
+				AutoCompleteIndex = 0;
+				AutoCompleteCursor = -1;
+				AutoCompleteIndices.Empty();
+				bAutoCompleteLocked = false;
+
 				return true;
 			}
 			else
@@ -600,9 +606,10 @@ bool UConsole::InputKey_InputLine( int32 ControllerId, FKey Key, EInputEvent Eve
 		}
 		else if (Key == EKeys::Tab && Event == IE_Pressed)
 		{
-			if (AutoCompleteIndices.Num() > 0 && !bAutoCompleteLocked && AutoCompleteCursor >= 0)
+			if (AutoCompleteIndices.Num() > 0 && !bAutoCompleteLocked)
 			{
-				const FAutoCompleteCommand& Cmd = AutoCompleteList[AutoCompleteIndices[AutoCompleteIndex + AutoCompleteCursor]];
+				const int32 Index = AutoCompleteIndex + (AutoCompleteCursor >= 0 ? AutoCompleteCursor : 0);
+				const FAutoCompleteCommand& Cmd = AutoCompleteList[AutoCompleteIndices[Index]];
 
 				TypedStr = Cmd.Command;
 				SetCursorPos(TypedStr.Len());
