Fix issue where some landscape components didn't compile the material. Also fixes crash when clearing the LandscapeMaterial property on an existing landscape.
TTP#330619

diff --git a/Engine/Source/Runtime/Engine/Private/Landscape/LandscapeEdit.cpp b/Engine/Source/Runtime/Engine/Private/Landscape/LandscapeEdit.cpp
--- a/Engine/Source/Runtime/Engine/Private/Landscape/LandscapeEdit.cpp
+++ b/Engine/Source/Runtime/Engine/Private/Landscape/LandscapeEdit.cpp
@@ -3733,9 +3733,13 @@ void ALandscape::PostEditChangeProperty(FPropertyChangedEvent& PropertyChangedEv
 			LandscapeEdit.RecalculateNormals();
 		}
 
-		for (auto It = Info->XYtoComponentMap.CreateIterator(); It; ++It )
+		TArray<ULandscapeComponent*> AllComponents;
+		Info->XYtoComponentMap.GenerateValueArray(AllComponents);
+
+		// We cannot iterate the XYtoComponentMap directly because reregistering components modifies the array.
+		for (auto It = AllComponents.CreateIterator(); It; ++It)
 		{
-			ULandscapeComponent* Comp = It.Value();
+			ULandscapeComponent* Comp = *It;
 			if( Comp )
 			{
 				if( bNeedsRecalcBoundingBox )
diff --git a/Engine/Source/Runtime/Engine/Private/Materials/HLSLMaterialTranslator.h b/Engine/Source/Runtime/Engine/Private/Materials/HLSLMaterialTranslator.h
--- a/Engine/Source/Runtime/Engine/Private/Materials/HLSLMaterialTranslator.h
+++ b/Engine/Source/Runtime/Engine/Private/Materials/HLSLMaterialTranslator.h
@@ -349,29 +349,36 @@ public:
 
 			check(FunctionStack.Num() == 1);
 
-		
-			int32 WPOFreq = (int32)GetMaterialPropertyShaderFrequency(MP_WorldPositionOffset);
-			FShaderCodeChunk& WPOChunk = CodeChunks[MP_WorldPositionOffset][WPOFreq][Chunk[MP_WorldPositionOffset]];
-			FMaterialRenderContext DummyContext( NULL, *Material, NULL );
-			
-			// Determine whether the world position offset is used. 
-			// If the output chunk has a uniform expression, it is constant, and GetNumberValue returns the default property value then WPO isn't used.
-			MaterialCompilationOutput.bModifiesMeshPosition = true;
-			if( WPOChunk.UniformExpression && WPOChunk.UniformExpression->IsConstant() )
-			{
-				float DummyFloat;
-				FColor DummyColor;
-				FVector DefaultVector;
-				GetDefaultForMaterialProperty(MP_WorldPositionOffset, DummyFloat,DummyColor,DefaultVector);
-
-				FLinearColor WPOValue;
-				WPOChunk.UniformExpression->GetNumberValue(DummyContext, WPOValue);
-				if( FVector(WPOValue) == DefaultVector)
+			if (Chunk[MP_WorldPositionOffset] == -1)
+			{
+				MaterialCompilationOutput.bModifiesMeshPosition = false;
+			}
+			else
+			{
+				int32 WPOFreq = (int32)GetMaterialPropertyShaderFrequency(MP_WorldPositionOffset);
+				FShaderCodeChunk& WPOChunk = CodeChunks[MP_WorldPositionOffset][WPOFreq][Chunk[MP_WorldPositionOffset]];
+				FMaterialRenderContext DummyContext( NULL, *Material, NULL );
+
+				// Determine whether the world position offset is used. 
+				// If the output chunk has a uniform expression, it is constant, and GetNumberValue returns the default property value then WPO isn't used.
+				MaterialCompilationOutput.bModifiesMeshPosition = true;
+				if( WPOChunk.UniformExpression && WPOChunk.UniformExpression->IsConstant() )
 				{
-					MaterialCompilationOutput.bModifiesMeshPosition = false;
+					float DummyFloat;
+					FColor DummyColor;
+					FVector DefaultVector;
+					GetDefaultForMaterialProperty(MP_WorldPositionOffset, DummyFloat,DummyColor,DefaultVector);
+
+					FLinearColor WPOValue;
+					WPOChunk.UniformExpression->GetNumberValue(DummyContext, WPOValue);
+					if( FVector(WPOValue) == DefaultVector)
+					{
+						MaterialCompilationOutput.bModifiesMeshPosition = false;
+					}
 				}
 			}
 
+
 			if (Material->GetBlendMode() == BLEND_Modulate && Material->GetLightingModel() != MLM_Unlit && !Material->IsUsedWithDeferredDecal())
 			{
 				Errorf(TEXT("Dynamically lit translucency is not supported for BLEND_Modulate materials."));
