Fix for 360 Controller losing input after pause in Editor
#ttp 313999 UDN: EDITOR: 360 Controller does not have input after pause
#proj Editor
#branch UE4
#summary ULocalPlayer will re-enable capture on joypad and mouse on unpause command. Initially it was only resetting input settings on pause & unpause.

diff --git a/Engine/Source/Runtime/Engine/Private/LocalPlayer.cpp b/Engine/Source/Runtime/Engine/Private/LocalPlayer.cpp
--- a/Engine/Source/Runtime/Engine/Private/LocalPlayer.cpp
+++ b/Engine/Source/Runtime/Engine/Private/LocalPlayer.cpp
@@ -795,7 +795,19 @@ bool ULocalPlayer::HandlePauseCommand( const TCHAR* Cmd, FOutputDevice& Ar, UWor
 {
 	Super::Exec(InWorld, TEXT("Pause"),Ar);
 
-	FSlateApplication::Get().ResetToDefaultInputSettings();
+	if (!InWorld->IsPaused())
+	{
+		if (ViewportClient && ViewportClient->Viewport)
+		{
+			ViewportClient->Viewport->CaptureJoystickInput(true);
+			ViewportClient->Viewport->CaptureMouse(true);
+		}
+	}
+	else
+	{
+		FSlateApplication::Get().ResetToDefaultInputSettings();
+	}
+	
 
 	return true;
 }
