Update other dynamic expression parameter nodes when another is renamed

#ttp 333844 - material editor:  the dynamic parameter node does not properly update it's ui when another dynamic parameter node changes one of the output pins parameter names

#branch UE4

#change It was renaming the other node, but it wasn't reconstructing it so it would appear with the same name.

diff --git a/Engine/Source/Runtime/Engine/Private/Materials/Material.cpp b/Engine/Source/Runtime/Engine/Private/Materials/Material.cpp
--- a/Engine/Source/Runtime/Engine/Private/Materials/Material.cpp
+++ b/Engine/Source/Runtime/Engine/Private/Materials/Material.cpp
@@ -2398,6 +2398,9 @@ void UMaterial::UpdateExpressionDynamicParameterNames(UMaterialExpression* Expre
 				{
 					CheckParam->ParamNames[NameIndex] = DynParam->ParamNames[NameIndex];
 				}
+#if WITH_EDITORONLY_DATA
+				CheckParam->GraphNode->ReconstructNode();
+#endif // WITH_EDITORONLY_DATA
 			}
 		}
 	}
