Fixed world origin offset for GPU particles

diff --git a/Engine/Source/Runtime/Engine/Private/Particles/ParticleGpuSimulation.cpp b/Engine/Source/Runtime/Engine/Private/Particles/ParticleGpuSimulation.cpp
--- a/Engine/Source/Runtime/Engine/Private/Particles/ParticleGpuSimulation.cpp
+++ b/Engine/Source/Runtime/Engine/Private/Particles/ParticleGpuSimulation.cpp
@@ -3566,7 +3566,7 @@ private:
 			float InterpFraction = (float)i / (float)SpawnInfo.Count;
 
 			NewParticle->Velocity = TempParticle->BaseVelocity;
-			NewParticle->Position = TempParticle->Location + InterpFraction * EmitterDelta + SpawnInfo.StartTime * NewParticle->Velocity + EmitterInfo.OrbitOffsetBase + EmitterInfo.OrbitOffsetRange * RandomOrbit - PositionOffsetThisTick;
+			NewParticle->Position = TempParticle->Location + InterpFraction * EmitterDelta + SpawnInfo.StartTime * NewParticle->Velocity + EmitterInfo.OrbitOffsetBase + EmitterInfo.OrbitOffsetRange * RandomOrbit;
 			NewParticle->RelativeTime = TempParticle->RelativeTime;
 			NewParticle->TimeScale = FMath::Max<float>(TempParticle->OneOverMaxLifetime, 0.001f);
 			NewParticle->Size.X = TempParticle->BaseSize.X * EmitterInfo.InvMaxSize.X;
