Fixed world origin offset for components spawned directly in the world, with no actor

diff --git a/Engine/Source/Runtime/Engine/Private/World.cpp b/Engine/Source/Runtime/Engine/Private/World.cpp
--- a/Engine/Source/Runtime/Engine/Private/World.cpp
+++ b/Engine/Source/Runtime/Engine/Private/World.cpp
@@ -3769,6 +3769,21 @@ bool UWorld::SetNewWorldOrigin(const FIntPoint& InNewOrigin)
 			LevelToShift->ApplyWorldOffset(Offset, true);
 		}
 	}
+
+	// Apply offset to components with no actor (like UGameplayStatics::SpawnEmitterAtLocation) 
+	{
+		TArray <UObject*> WorldChildren; 
+		GetObjectsWithOuter(this, WorldChildren, false);
+
+		for (UObject* ChildObject : WorldChildren)
+		{
+		   UActorComponent* Component = Cast<UActorComponent>(ChildObject);
+		   if (Component && Component->GetOwner() == nullptr)
+		   {
+				Component->ApplyWorldOffset(Offset, true);
+		   }
+		}
+	}
 			
 	// Shift navigation meshes
 	if (NavigationSystem)
