Based movement now uses the character base

To better align characters with moving base objects that have pitch and roll, use the base of the character (the bottom of the capsule) as the 'pivot' for based movement.

diff --git a/Engine/Source/Runtime/Engine/Private/CharacterMovementComponent.cpp b/Engine/Source/Runtime/Engine/Private/CharacterMovementComponent.cpp
--- a/Engine/Source/Runtime/Engine/Private/CharacterMovementComponent.cpp
+++ b/Engine/Source/Runtime/Engine/Private/CharacterMovementComponent.cpp
@@ -1081,8 +1081,14 @@ void UCharacterMovementComponent::UpdateBasedMovement(float DeltaSeconds)
 				FRotator PawnReducedRotation = finalRotation - PawnOldRotation;
 				UpdateBasedRotation(finalRotation, PawnReducedRotation);
 			}
-			FVector const LocalPos = OldLocalToWorld.InverseTransformPosition(CharacterOwner->GetActorLocation());
-			FVector const NewWorldPos = ConstrainLocationToPlane(NewLocalToWorld.TransformPosition(LocalPos));
+
+			// We need to offset the base of the character here, not its origin, so offset by half height
+			float HalfHeight, Radius;
+			CharacterOwner->CapsuleComponent->GetScaledCapsuleSize(Radius, HalfHeight);
+
+			FVector const BaseOffset(0.0f, 0.0f, HalfHeight);
+			FVector const LocalBasePos = OldLocalToWorld.InverseTransformPosition(CharacterOwner->GetActorLocation() - BaseOffset);
+			FVector const NewWorldPos = ConstrainLocationToPlane(NewLocalToWorld.TransformPosition(LocalBasePos) + BaseOffset);
 			FVector const Delta = ConstrainDirectionToPlane(NewWorldPos - CharacterOwner->GetActorLocation());
 
 			// move attached actor
