Based-movement for characters now uses quaternions

Based movement was calculating deltas from euler angles and applying them to the character. This caused problems when rotating in more than one cardinal plane (such as being based on the side of a rotating cylinder, as was the case in shootergame) because the euler representation of the underlying quaternion could suddently flip axes, breaking the delta calculations.
Changed based-movement calculations to use quaternions rather than euler angles.

TTP# 320677 - SHOOTERGAME: PlayerÆs camera gets turned around when walking into a trashcan that has been knocked down.

diff --git a/Engine/Source/Runtime/Engine/Classes/GameFramework/CharacterMovementComponent.h b/Engine/Source/Runtime/Engine/Classes/GameFramework/CharacterMovementComponent.h
--- a/Engine/Source/Runtime/Engine/Classes/GameFramework/CharacterMovementComponent.h
+++ b/Engine/Source/Runtime/Engine/Classes/GameFramework/CharacterMovementComponent.h
@@ -202,12 +202,10 @@ public:
 	uint8 CustomMovementMode;
 
 	/** Save Base location to check in UpdateBasedMovement() whether base moved in the last frame, and therefore pawn needs an update. */
-	UPROPERTY()
 	FVector OldBaseLocation;
 
 	/** Save Base rotation to check in UpdateBasedMovement() whether base moved in the last frame, and therefore pawn needs an update. */
-	UPROPERTY()
-	FRotator OldBaseRotation;
+	FQuat OldBaseQuat;
 
 	/** Custom gravity scale. Gravity is multiplied by this amount for the character. */
 	UPROPERTY(Category="Character Movement", EditAnywhere, BlueprintReadWrite)
diff --git a/Engine/Source/Runtime/Engine/Private/Character.cpp b/Engine/Source/Runtime/Engine/Private/Character.cpp
--- a/Engine/Source/Runtime/Engine/Private/Character.cpp
+++ b/Engine/Source/Runtime/Engine/Private/Character.cpp
@@ -479,7 +479,7 @@ void ACharacter::SetBase( UPrimitiveComponent* NewBaseComponent, bool bNotifyPaw
 			{
 				// Update OldBaseLocation/Rotation as those were referring to a different base
 				CharacterMovement->OldBaseLocation = MovementBase->GetComponentLocation();
-				CharacterMovement->OldBaseRotation = MovementBase->GetComponentRotation();
+				CharacterMovement->OldBaseQuat = MovementBase->GetComponentQuat();
 			}
 			else
 			{
diff --git a/Engine/Source/Runtime/Engine/Private/CharacterMovementComponent.cpp b/Engine/Source/Runtime/Engine/Private/CharacterMovementComponent.cpp
--- a/Engine/Source/Runtime/Engine/Private/CharacterMovementComponent.cpp
+++ b/Engine/Source/Runtime/Engine/Private/CharacterMovementComponent.cpp
@@ -170,6 +170,9 @@ UCharacterMovementComponent::UCharacterMovementComponent(const class FPostConstr
 	AvoidanceGroup.bGroup0 = true;
 	GroupsToAvoid.Packed = 0xFFFFFFFF;
 	GroupsToIgnore.Packed = 0;
+
+	OldBaseQuat = FQuat::Identity;
+	OldBaseLocation = FVector::ZeroVector;
 }
 
 void UCharacterMovementComponent::PostLoad()
@@ -951,7 +954,7 @@ void UCharacterMovementComponent::SimulateMovement(float DeltaSeconds)
 	{
 		return;
 	}
-	
+
 	Acceleration = Velocity.SafeNormal();	// Not currently used for simulated movement
 	AnalogInputModifier = 1.0f;				// Not currently used for simulated movement
 
@@ -1045,20 +1048,19 @@ void UCharacterMovementComponent::UpdateBasedMovement(float DeltaSeconds)
 	// Ignore collision with bases during these movements.
 	TGuardValue<EMoveComponentFlags> ScopedFlagRestore(MoveComponentFlags, MoveComponentFlags | MOVECOMP_IgnoreBases);
 
-	FRotator ReducedRotation(0,0,0);
-	const bool bRotationChanged = (OldBaseRotation != MovementBase->GetComponentRotation());
+	FQuat DeltaQuat = FQuat::Identity;
+	const bool bRotationChanged = !OldBaseQuat.Equals(MovementBase->GetComponentQuat());
 	if( bRotationChanged )
 	{
-		//@hack - Angles used to need to be reduced for sine and cosine table lookup, maybe remove this
-		ReducedRotation = MovementBase->GetComponentRotation() - OldBaseRotation;
+		DeltaQuat = OldBaseQuat * MovementBase->GetComponentQuat().Inverse();
 	}
 
 	// only if base moved
 	if ( bRotationChanged || (OldBaseLocation != MovementBase->GetComponentLocation()) )
 	{
 		// Calculate new transform matrix of base actor (ignoring scale).
-		const FRotationTranslationMatrix OldLocalToWorld(OldBaseRotation, OldBaseLocation);
-		const FRotationTranslationMatrix NewLocalToWorld(MovementBase->GetComponentRotation(), MovementBase->GetComponentLocation());
+		const FQuatRotationTranslationMatrix OldLocalToWorld(OldBaseQuat, OldBaseLocation);
+		const FQuatRotationTranslationMatrix NewLocalToWorld(MovementBase->GetComponentQuat(), MovementBase->GetComponentLocation());
 
 		FVector RotMotion(0.f);
 		if( CharacterOwner->IsMatineeControlled() )
@@ -1070,16 +1072,17 @@ void UCharacterMovementComponent::UpdateBasedMovement(float DeltaSeconds)
 		}
 		else
 		{
-			FRotator finalRotation = CharacterOwner->GetActorRotation() + ReducedRotation;
-			FRotator PawnOldRotation = CharacterOwner->GetActorRotation();
-			CharacterOwner->FaceRotation(finalRotation, 0.f);
-			finalRotation = CharacterOwner->GetActorRotation();
+			FQuat FinalQuat = DeltaQuat * CharacterOwner->GetActorQuat();
+			FQuat PawnOldQuat = CharacterOwner->GetActorQuat();
+			CharacterOwner->FaceRotation(FinalQuat.Rotator(), 0.f);
+			FinalQuat = CharacterOwner->GetActorQuat();
 
 			if( bRotationChanged )
 			{
-				//@hack - Angles used to need to be reduced for sine and cosine table lookup, maybe remove this
-				FRotator PawnReducedRotation = finalRotation - PawnOldRotation;
-				UpdateBasedRotation(finalRotation, PawnReducedRotation);
+				const FQuat PawnDeltaRotation = PawnOldQuat * FinalQuat.Inverse();
+				FRotator FinalRotation = FinalQuat.Rotator();
+				UpdateBasedRotation(FinalRotation, PawnDeltaRotation.Rotator());
+				FinalQuat = FinalRotation.Quaternion();
 			}
 
 			// We need to offset the base of the character here, not its origin, so offset by half height
@@ -1095,11 +1098,11 @@ void UCharacterMovementComponent::UpdateBasedMovement(float DeltaSeconds)
 			if (bFastAttachedMove)
 			{
 				// we're trusting no other obstacle can prevent the move here
-				UpdatedComponent->SetWorldLocationAndRotation(NewWorldPos, finalRotation, false);
+				UpdatedComponent->SetWorldLocationAndRotation(NewWorldPos, FinalQuat.Rotator(), false);
 			}
 			else
 			{
-				MoveUpdatedComponent( Delta, finalRotation, true );
+				MoveUpdatedComponent( Delta, FinalQuat.Rotator(), true );
 			}
 		}
 	}
@@ -1308,8 +1311,8 @@ void UCharacterMovementComponent::SaveBaseLocation()
 		OldBaseLocation = MovementBase->GetComponentLocation();
 		CharacterOwner->RelativeMovement.Location = CharacterOwner->GetActorLocation() - OldBaseLocation;
 
-		OldBaseRotation = MovementBase->GetComponentRotation();
-		CharacterOwner->RelativeMovement.Rotation = (FRotationMatrix(CharacterOwner->GetActorRotation()) * FRotationMatrix(OldBaseRotation).GetTransposed()).Rotator();
+		OldBaseQuat = MovementBase->GetComponentQuat();
+		CharacterOwner->RelativeMovement.Rotation = (FRotationMatrix(CharacterOwner->GetActorRotation()) * FQuatRotationTranslationMatrix(OldBaseQuat, FVector::ZeroVector).GetTransposed()).Rotator();
 	}
 }
 
