[GitHub] Integrate PR 107: Added Blueprint support for creating / finding, updating and getting size of a UCanvasRenderTarget2D

https://github.com/EpicGames/UnrealEngine/pull/107

From James Tan:

Changed delegate into a dynamic multi cast delegate to allow Blueprint to assign an event to it.
Added a static function which will find a UCanvasRenderTarget2D from the static global TMap<>. It can also optionally create a new UCanvasRenderTarget2D if an existing one wasn't found.
Added a static function which will ask a UCanvasRenderTarget2D to update itself.
Added a static function which will get a UCanvasRenderTarget2D's size.
Added a static global TMap<> which stores all statically created UCanvasRenderTarget2Ds. FName is used as the key. All UCanvasRenderTarget2D is stored as an auto weak pointer.
Creating the FCanvas now just uses GWorld instead of trying to get the time from FApp.

diff --git a/Engine/Source/Runtime/Engine/Classes/Engine/CanvasRenderTarget2D.h b/Engine/Source/Runtime/Engine/Classes/Engine/CanvasRenderTarget2D.h
--- a/Engine/Source/Runtime/Engine/Classes/Engine/CanvasRenderTarget2D.h
+++ b/Engine/Source/Runtime/Engine/Classes/Engine/CanvasRenderTarget2D.h
@@ -5,15 +5,19 @@
 #include "Canvas.h"
 #include "CanvasRenderTarget2D.generated.h"
 
-/** This delegate allows developers to bind a function to be called every time the texture is updated, so that they can use the supplied canvas to draw to the texture. */
-DECLARE_DELEGATE_OneParam( FOnCanvasRenderTargetUpdate, class UCanvas* );
+/** This delegate is assignable through Blueprint and has similar functionality to the above. */
+DECLARE_DYNAMIC_MULTICAST_DELEGATE_ThreeParams(FOnCanvasRenderTargetUpdate, UCanvas*, Canvas, int32, Width, int32, Height);
 
 /**
  * CanvasRenderTarget2D is 2D render target which exposes a Canvas interface to allow you to draw elements onto 
- * it directly.
+ * it directly.  Use FindCanvasRenderTarget2D() to find or create a render target texture by unique name, then
+ * bind a function to the OnCanvasRenderTargetUpdate delegate which will be called when the render target is
+ * updated.  If you need to repaint your canvas every single frame, simply call UpdateResource() on it from a Tick
+ * function.  Also, remember to hold onto your new canvas render target with a reference so that it doesn't get
+ * garbage collected.
  */
-UCLASS(MinimalAPI)
-class UCanvasRenderTarget2D : public UTextureRenderTarget2D
+UCLASS()
+class ENGINE_API UCanvasRenderTarget2D : public UTextureRenderTarget2D
 {
 	GENERATED_UCLASS_BODY()
 
@@ -30,18 +34,42 @@ public:
 	virtual void UpdateResource();
 
 	/**
-	 * Assigns a user delegate to be called every time the texture needs to be refreshed for rendering.  This delegate
-	 * passes along a Canvas object that you can use to draw to the texture directly!
+	 * Statically finds a canvas render target from the global map of canvas render targets. It can optionally create a new render target if not found as well.
 	 *
-	 * @param	InDelegate	The canvas update delegate to use.  You can change this at any time.
+	 * @param	InRenderTargetName	Name of the render target. This should be unique.
+	 * @param	Width				Width of the render target.
+	 * @param	Height				Height of the render target.
+	 * @param	bCreateIfNotFound	True to create the render target if we couldn't find one with this name
+	 * @return						Returns the instanced render target.
 	 */
-	void SetOnCanvasRenderTargetUpdate( const FOnCanvasRenderTargetUpdate& InDelegate )
-	{
-		OnCanvasRenderTargetUpdate = InDelegate;
-	}
+	UFUNCTION(BlueprintCallable, Category="Canvas Render Target 2D")
+	static UCanvasRenderTarget2D* FindCanvasRenderTarget2D(FName InRenderTargetName, int32 Width, int32 Height, bool bCreateIfNotFound=false);
 
-private:
+	/**
+	 * Updates a specific render target from the global map of canvas render targets.  Does not trigger an error if the specified render target can't be found.
+	 *
+	 * @param	InRenderTargetName		Name of the render target to update.
+	 */
+	UFUNCTION(BlueprintCallable, Category="Canvas Render Target 2D")
+	static void UpdateCanvasRenderTarget(FName InRenderTargetName);
 
-	/** Texture update delegate */
+	/**
+	 * Gets a specific render target's size from the global map of canvas render targets.  Does not trigger an error if the specified render target can't be found.
+	 *
+	 * @param	Width	Out: The width of this render target
+	 * @param	Height	Out: The height of this render target
+	 */
+	UFUNCTION(BlueprintCallable, Category="Canvas Render Target 2D")
+	static void GetCanvasRenderTargetSize(FName InRenderTargetName, int32& Width, int32& Height);
+
+
+	/** Called when this Canvas Render Target is asked to update its texture resource. */
+	UPROPERTY(BlueprintAssignable, Category="Canvas Render Target 2D")
 	FOnCanvasRenderTargetUpdate OnCanvasRenderTargetUpdate;
-};
+
+
+protected:
+
+	/** Global mapped array of all canvas render targets */
+	static TMap<FName, TAutoWeakObjectPtr<UCanvasRenderTarget2D>> GCanvasRenderTarget2Ds;
+};
\ No newline at end of file
diff --git a/Engine/Source/Runtime/Engine/Private/CanvasRenderTarget2D.cpp b/Engine/Source/Runtime/Engine/Private/CanvasRenderTarget2D.cpp
--- a/Engine/Source/Runtime/Engine/Private/CanvasRenderTarget2D.cpp
+++ b/Engine/Source/Runtime/Engine/Private/CanvasRenderTarget2D.cpp
@@ -3,13 +3,14 @@
 #include "EnginePrivate.h"
 #include "Engine/CanvasRenderTarget2D.h"
 
+TMap<FName, TAutoWeakObjectPtr<UCanvasRenderTarget2D>> UCanvasRenderTarget2D::GCanvasRenderTarget2Ds;
+
 UCanvasRenderTarget2D::UCanvasRenderTarget2D( const class FPostConstructInitializeProperties& PCIP )
 	: Super(PCIP)
 {
 	bNeedsTwoCopies = false;
 }
 
-
 void UCanvasRenderTarget2D::UpdateResource()
 {
 	// Call parent implementation
@@ -48,10 +49,9 @@ void UCanvasRenderTarget2D::UpdateResource()
 	FCanvas RenderCanvas(GameThread_GetRenderTargetResource(), nullptr, FApp::GetCurrentTime() - GStartTime, FApp::GetDeltaTime(), FApp::GetCurrentTime() - GStartTime);
 	Canvas->Canvas = &RenderCanvas;
 
-	// Check if the render target update delegate is bound and if it is, execute the delegate with the current Canvas being used.
-	if (OnCanvasRenderTargetUpdate.IsBound())
+	if (!IsPendingKill() && OnCanvasRenderTargetUpdate.IsBound())
 	{
-		OnCanvasRenderTargetUpdate.Execute(Canvas);
+		OnCanvasRenderTargetUpdate.Broadcast(Canvas, GetSurfaceWidth(), GetSurfaceHeight());
 	}
 
 	// Clean up and flush the rendering canvas.
@@ -71,3 +71,71 @@ void UCanvasRenderTarget2D::UpdateResource()
 		}
 	);
 }
+
+UCanvasRenderTarget2D* UCanvasRenderTarget2D::FindCanvasRenderTarget2D(FName InRenderTargetName, int32 Width, int32 Height, bool bCreateIfNotFound)
+{
+	if (Width > 0 && Height > 0 && InRenderTargetName != NAME_None)
+	{
+		// Check if the render target with the name has already been created.
+		if (GCanvasRenderTarget2Ds.Num() > 0)
+		{
+			for (TMap<FName, TAutoWeakObjectPtr<UCanvasRenderTarget2D>>::TConstIterator It(GCanvasRenderTarget2Ds); It; ++It)
+			{
+				if (It->Key == InRenderTargetName && (It->Value).IsValid())
+				{
+					return (It->Value).Get();
+				}
+			}
+		}
+
+		// Create a new one since it doesn't exist.
+		if (bCreateIfNotFound)
+		{
+			UCanvasRenderTarget2D* NewCanvasRenderTarget = ConstructObject<UCanvasRenderTarget2D>(UCanvasRenderTarget2D::StaticClass(), GetTransientPackage());
+			if (NewCanvasRenderTarget)
+			{
+				NewCanvasRenderTarget->InitAutoFormat(Width, Height);
+
+				// Add it to the global static map.
+				GCanvasRenderTarget2Ds.Add(InRenderTargetName, NewCanvasRenderTarget);
+
+				return NewCanvasRenderTarget;
+			}
+		}
+	}
+
+	return nullptr;
+}
+
+void UCanvasRenderTarget2D::UpdateCanvasRenderTarget(FName InRenderTargetName)
+{
+	if (GCanvasRenderTarget2Ds.Num() > 0)
+	{
+		for (TMap<FName, TAutoWeakObjectPtr<UCanvasRenderTarget2D>>::TConstIterator It(GCanvasRenderTarget2Ds); It; ++It)
+		{
+			if (It->Key == InRenderTargetName && (It->Value).IsValid())
+			{
+				(It->Value)->UpdateResource();
+				
+				break;
+			}
+		}
+	}
+}
+
+void UCanvasRenderTarget2D::GetCanvasRenderTargetSize(FName InRenderTargetName, int32& Width, int32& Height)
+{
+	if (GCanvasRenderTarget2Ds.Num() > 0)
+	{
+		for (TMap<FName, TAutoWeakObjectPtr<UCanvasRenderTarget2D>>::TConstIterator It(GCanvasRenderTarget2Ds); It; ++It)
+		{
+			if (It->Key == InRenderTargetName && (It->Value).IsValid())
+			{
+				Width = (It->Value)->GetSurfaceWidth();
+				Height = (It->Value)->GetSurfaceHeight();
+
+				break;
+			}
+		}
+	}
+}
\ No newline at end of file
