Point attractors can now filter motion along axes

diff --git a/Engine/Source/Runtime/Core/Public/Math/Vector.h b/Engine/Source/Runtime/Core/Public/Math/Vector.h
--- a/Engine/Source/Runtime/Core/Public/Math/Vector.h
+++ b/Engine/Source/Runtime/Core/Public/Math/Vector.h
@@ -1772,4 +1772,14 @@ FORCEINLINE float FVector::BoxPushOut( const FVector & Normal, const FVector & S
 	return FMath::Abs(Normal.X*Size.X) + FMath::Abs(Normal.Y*Size.Y) + FMath::Abs(Normal.Z*Size.Z);
 }
 
+/** Component-wise clamp for FVector */
+FORCEINLINE FVector ClampVector( const FVector& V, const FVector& Min, const FVector& Max )
+{
+	return FVector(
+		FMath::Clamp(V.X,Min.X,Max.X),
+		FMath::Clamp(V.Y,Min.Y,Max.Y),
+		FMath::Clamp(V.Z,Min.Z,Max.Z)
+		);
+}
+
 template <> struct TIsPODType<FVector> { enum { Value = true }; };
diff --git a/Engine/Source/Runtime/Engine/Classes/Particles/Attractor/ParticleModuleAttractorPoint.h b/Engine/Source/Runtime/Engine/Classes/Particles/Attractor/ParticleModuleAttractorPoint.h
--- a/Engine/Source/Runtime/Engine/Classes/Particles/Attractor/ParticleModuleAttractorPoint.h
+++ b/Engine/Source/Runtime/Engine/Classes/Particles/Attractor/ParticleModuleAttractorPoint.h
@@ -37,6 +37,30 @@ class UParticleModuleAttractorPoint : public UParticleModuleAttractorBase
 	UPROPERTY(EditAnywhere, Category=Attractor)
 	uint32 bUseWorldSpacePosition:1;
 
+	/** Whether particles can move along the positive X axis. */
+	UPROPERTY(EditAnywhere, Category=Attractor)
+	uint32 Positive_X:1;
+
+	/** Whether particles can move along the positive Y axis. */
+	UPROPERTY(EditAnywhere, Category=Attractor)
+	uint32 Positive_Y:1;
+
+	/** Whether particles can move along the positive Z axis. */
+	UPROPERTY(EditAnywhere, Category=Attractor)
+	uint32 Positive_Z:1;
+
+	/** Whether particles can move along the negative X axis. */
+	UPROPERTY(EditAnywhere, Category=Attractor)
+	uint32 Negative_X:1;
+
+	/** Whether particles can move along the negative Y axis. */
+	UPROPERTY(EditAnywhere, Category=Attractor)
+	uint32 Negative_Y:1;
+
+	/** Whether particles can move along the negative Z axis. */
+	UPROPERTY(EditAnywhere, Category=Attractor)
+	uint32 Negative_Z:1;
+
 	/** Initializes the default values for this property */
 	void InitializeDefaults();
 
diff --git a/Engine/Source/Runtime/Engine/Private/Particles/ParticleModules.cpp b/Engine/Source/Runtime/Engine/Private/Particles/ParticleModules.cpp
--- a/Engine/Source/Runtime/Engine/Private/Particles/ParticleModules.cpp
+++ b/Engine/Source/Runtime/Engine/Private/Particles/ParticleModules.cpp
@@ -3918,6 +3918,14 @@ UParticleModuleAttractorPoint::UParticleModuleAttractorPoint(const class FPostCo
 	bAffectBaseVelocity = false;
 	bOverrideVelocity = false;
 	bSupported3DDrawMode = true;
+
+	Positive_X = true;
+	Positive_Y = true;
+	Positive_Z = true;
+
+	Negative_X = true;
+	Negative_Y = true;
+	Negative_Z = true;
 }
 
 void UParticleModuleAttractorPoint::InitializeDefaults()
@@ -3991,6 +3999,9 @@ void UParticleModuleAttractorPoint::Update(FParticleEmitterInstance* Owner, int3
 
 	AttractorRange *= ScaleSize;
 
+	FVector MinNormalizedDir(Negative_X ? -1.0f : 0.0f, Negative_Y ? -1.0f : 0.0f, Negative_Z ? -1.0f : 0.0f);
+	FVector MaxNormalizedDir(Positive_X ? +1.0f : 0.0f, Positive_Y ? +1.0f : 0.0f, Positive_Z ? +1.0f : 0.0f);
+
 	BEGIN_UPDATE_LOOP;
 		// If the particle is within range...
 		FVector Dir = AttractorPosition - Particle.Location;
@@ -4022,8 +4033,17 @@ void UParticleModuleAttractorPoint::Update(FParticleEmitterInstance* Owner, int3
 				AttractorStrength *= ScaleSize;
 			}
 
-			// Adjust the VELOCITY of the particle based on the attractor... 
 			Dir.Normalize();
+
+			// If the strength is negative, flip direction before clamping.
+			if (AttractorStrength < 0.0f)
+			{
+				Dir = -Dir;
+				AttractorStrength = -AttractorStrength;
+			}
+
+			// Adjust the VELOCITY of the particle based on the attractor...
+			Dir = ClampVector(Dir,MinNormalizedDir,MaxNormalizedDir);
     		Particle.Velocity	+= Dir * AttractorStrength * DeltaTime;
 			if (bAffectBaseVelocity)
 			{
