* hide alpha channel in postprocess colors

diff --git a/Engine/Source/Runtime/Engine/Classes/Engine/Scene.h b/Engine/Source/Runtime/Engine/Classes/Engine/Scene.h
--- a/Engine/Source/Runtime/Engine/Classes/Engine/Scene.h
+++ b/Engine/Source/Runtime/Engine/Classes/Engine/Scene.h
@@ -370,7 +370,7 @@ struct FPostProcessSettings
 
 
 	/** Scene tint color */
-	UPROPERTY(interp, Category=SceneColor, AdvancedDisplay, meta=(editcondition = "bOverride_SceneColorTint"))
+	UPROPERTY(interp, Category=SceneColor, AdvancedDisplay, meta=(editcondition = "bOverride_SceneColorTint", HideAlphaChannel))
 	FLinearColor SceneColorTint;
 	
 	/** in percent, Scene chromatic aberration / color fringe (camera imperfection) to simulate an artifact that happens in real-world lens, mostly visible in the image corners. */
@@ -429,19 +429,19 @@ struct FPostProcessSettings
 	float Bloom5Size;
 
 	/** Bloom1 tint color */
-	UPROPERTY(interp, Category=Bloom, AdvancedDisplay, meta=(editcondition = "bOverride_Bloom1Tint", DisplayName = "#1 Tint"))
+	UPROPERTY(interp, Category=Bloom, AdvancedDisplay, meta=(editcondition = "bOverride_Bloom1Tint", DisplayName = "#1 Tint", HideAlphaChannel))
 	FLinearColor Bloom1Tint;
 	/** Bloom2 tint color */
-	UPROPERTY(interp, Category=Bloom, AdvancedDisplay, meta=(editcondition = "bOverride_Bloom2Tint", DisplayName = "#2 Tint"))
+	UPROPERTY(interp, Category=Bloom, AdvancedDisplay, meta=(editcondition = "bOverride_Bloom2Tint", DisplayName = "#2 Tint", HideAlphaChannel))
 	FLinearColor Bloom2Tint;
 	/** Bloom3 tint color */
-	UPROPERTY(interp, Category=Bloom, AdvancedDisplay, meta=(editcondition = "bOverride_Bloom3Tint", DisplayName = "#3 Tint"))
+	UPROPERTY(interp, Category=Bloom, AdvancedDisplay, meta=(editcondition = "bOverride_Bloom3Tint", DisplayName = "#3 Tint", HideAlphaChannel))
 	FLinearColor Bloom3Tint;
 	/** Bloom4 tint color */
-	UPROPERTY(interp, Category=Bloom, AdvancedDisplay, meta=(editcondition = "bOverride_Bloom4Tint", DisplayName = "#4 Tint"))
+	UPROPERTY(interp, Category=Bloom, AdvancedDisplay, meta=(editcondition = "bOverride_Bloom4Tint", DisplayName = "#4 Tint", HideAlphaChannel))
 	FLinearColor Bloom4Tint;
 	/** Bloom5 tint color */
-	UPROPERTY(interp, Category=Bloom, AdvancedDisplay, meta=(editcondition = "bOverride_Bloom5Tint", DisplayName = "#5 Tint"))
+	UPROPERTY(interp, Category=Bloom, AdvancedDisplay, meta=(editcondition = "bOverride_Bloom5Tint", DisplayName = "#5 Tint", HideAlphaChannel))
 	FLinearColor Bloom5Tint;
 
 	/** BloomDirtMask intensity */
