Remove latStiffFactor and update some default tuning

diff --git a/Engine/Source/Runtime/Engine/Classes/Vehicles/VehicleWheel.h b/Engine/Source/Runtime/Engine/Classes/Vehicles/VehicleWheel.h
--- a/Engine/Source/Runtime/Engine/Classes/Vehicles/VehicleWheel.h
+++ b/Engine/Source/Runtime/Engine/Classes/Vehicles/VehicleWheel.h
@@ -73,10 +73,6 @@ class ENGINE_API UVehicleWheel : public UObject
 	UPROPERTY(EditAnywhere, Category=Tire, meta=(ClampMin = "0.01", UIMin = "0.01"))
 	float											LatStiffValue;
 
-	/** Multiplier for the output lateral tire forces */
-	UPROPERTY(EditAnywhere, Category=Tire)
-	float											LatStiffFactor;
-
 	/** How much longitudinal stiffness to have given longitudinal slip */
 	UPROPERTY(EditAnywhere, Category=Tire)
 	float											LongStiffValue;
diff --git a/Engine/Source/Runtime/Engine/Private/Vehicles/VehicleWheel.cpp b/Engine/Source/Runtime/Engine/Private/Vehicles/VehicleWheel.cpp
--- a/Engine/Source/Runtime/Engine/Private/Vehicles/VehicleWheel.cpp
+++ b/Engine/Source/Runtime/Engine/Private/Vehicles/VehicleWheel.cpp
@@ -28,12 +28,11 @@ UVehicleWheel::UVehicleWheel(const class FPostConstructInitializeProperties& PCI
 	DampingRate = 0.25f;
 	LatStiffMaxLoad = 2.0f;
 	LatStiffValue = 17.0f;
-	LatStiffFactor = 1.0f;
 	LongStiffValue = 1000.0f;
 	SuspensionForceOffset = 0.0f;
-	SuspensionMaxRaise = 20.0f;
-	SuspensionMaxDrop = 20.0f;
-	SuspensionNaturalFrequency = 5.0f;
+	SuspensionMaxRaise = 10.0f;
+	SuspensionMaxDrop = 10.0f;
+	SuspensionNaturalFrequency = 7.0f;
 	SuspensionDampingRatio = 1.0f;
 }
 
