Fixed so that the RemoveNan function doesn't always compress unless it has to.

diff --git a/Engine/Source/Runtime/Engine/Private/Animation/AnimSequence.cpp b/Engine/Source/Runtime/Engine/Private/Animation/AnimSequence.cpp
--- a/Engine/Source/Runtime/Engine/Private/Animation/AnimSequence.cpp
+++ b/Engine/Source/Runtime/Engine/Private/Animation/AnimSequence.cpp
@@ -1621,6 +1621,8 @@ void UAnimSequence::PostProcessSequence()
 #if WITH_EDITOR
 void UAnimSequence::RemoveNaNTracks()
 {
+	bool bRecompress = false;
+
 	for( int32 TrackIndex=0; TrackIndex<RawAnimationData.Num(); ++TrackIndex )
 	{
 		const FRawAnimSequenceTrack & RawTrack = RawAnimationData[TrackIndex];
@@ -1653,10 +1655,15 @@ void UAnimSequence::RemoveNaNTracks()
 			// remove this track
 			RemoveTrack(TrackIndex);
 			--TrackIndex;
+
+			bRecompress = true;
 		}
 	}
 
-	FAnimationUtils::CompressAnimSequence(this, false, false);
+	if (bRecompress)
+	{
+		FAnimationUtils::CompressAnimSequence(this, false, false);
+	}
 }
 #endif // WITH_EDITOR
 
