#ue4 - UMovementComponent::SlideAlongSurface() does not attempt a second move if TwoWallAdjust() returns a very small vector.

diff --git a/Engine/Source/Runtime/Engine/Classes/GameFramework/MovementComponent.h b/Engine/Source/Runtime/Engine/Classes/GameFramework/MovementComponent.h
--- a/Engine/Source/Runtime/Engine/Classes/GameFramework/MovementComponent.h
+++ b/Engine/Source/Runtime/Engine/Classes/GameFramework/MovementComponent.h
@@ -194,7 +194,7 @@ public:
 	virtual FVector ComputeSlideVector(const FVector& Delta, const float Time, const FVector& Normal, const FHitResult& Hit) const;
 
 	/**
-	 * Slide smoothly along a surface, and slide away from multiple impacts using TwoWallAdjust if necessary. Calls HandleImpact for the first surface hit, if requested.
+	 * Slide smoothly along a surface, and slide away from multiple impacts using TwoWallAdjust if necessary. Calls HandleImpact for each surface hit, if requested.
 	 * Uses SafeMoveUpdatedComponent for movement, and ComputeSlideVector to determine the slide direction.
 	 * @param bHandleImpact - Whether to call HandleImpact on each hit.
 	 * @return The percent of the Time value actually applied to movement (between 0 and 1). 0 if no movement occurred, non-zero if movement occurred (in which case Hit.Time is the result of the last move).
diff --git a/Engine/Source/Runtime/Engine/Private/MovementComponent.cpp b/Engine/Source/Runtime/Engine/Private/MovementComponent.cpp
--- a/Engine/Source/Runtime/Engine/Private/MovementComponent.cpp
+++ b/Engine/Source/Runtime/Engine/Private/MovementComponent.cpp
@@ -367,6 +367,11 @@ FVector UMovementComponent::ComputeSlideVector(const FVector& Delta, const float
 
 float UMovementComponent::SlideAlongSurface(const FVector& Delta, float Time, const FVector& Normal, FHitResult& Hit, bool bHandleImpact)
 {
+	if (!Hit.bBlockingHit)
+	{
+		return 0.f;
+	}
+
 	float PercentTimeApplied = 0.f;
 	const FVector OldHitNormal = Normal;
 
@@ -390,15 +395,19 @@ float UMovementComponent::SlideAlongSurface(const FVector& Delta, float Time, co
 			// Compute new slide normal when hitting multiple surfaces.
 			TwoWallAdjust(SlideDelta, Hit, OldHitNormal);
 
-			// Perform second move
-			SafeMoveUpdatedComponent(SlideDelta, Rotation, true, Hit);
-			const float SecondHitPercent = Hit.Time * (1.f - FirstHitPercent);
-			PercentTimeApplied += SecondHitPercent;
-
-			// Notify second impact
-			if (bHandleImpact && Hit.bBlockingHit)
+			// Only proceed if the new direction is of significant length.
+			if (!SlideDelta.IsNearlyZero(1e-3f))
 			{
-				HandleImpact(Hit, SecondHitPercent * Time, SlideDelta);
+				// Perform second move
+				SafeMoveUpdatedComponent(SlideDelta, Rotation, true, Hit);
+				const float SecondHitPercent = Hit.Time * (1.f - FirstHitPercent);
+				PercentTimeApplied += SecondHitPercent;
+
+				// Notify second impact
+				if (bHandleImpact && Hit.bBlockingHit)
+				{
+					HandleImpact(Hit, SecondHitPercent * Time, SlideDelta);
+				}
 			}
 		}
 
