Physics: Added ensure(Volume >= 0) to mass calculation code

diff --git a/Engine/Source/Runtime/Engine/Classes/Components/ActorComponent.h b/Engine/Source/Runtime/Engine/Classes/Components/ActorComponent.h
--- a/Engine/Source/Runtime/Engine/Classes/Components/ActorComponent.h
+++ b/Engine/Source/Runtime/Engine/Classes/Components/ActorComponent.h
@@ -392,10 +392,10 @@ public:
 	virtual void OnActorEnableCollisionChanged() {}
 
 	/** 
-		Give a readable name for this component, including asset name if applicable 
-		By default this appends a space plus AdditionalStatObject()
-	**/
-	virtual FString GetReadableName();
+	 * Returns a readable name for this component, including the asset name if applicable 
+	 * By default this appends a space plus AdditionalStatObject()
+	 */
+	virtual FString GetReadableName() const;
 
 	/** Give a readable name for this component, including asset name if applicable */
 	virtual UObject const* AdditionalStatObject() const
diff --git a/Engine/Source/Runtime/Engine/Private/ActorComponent.cpp b/Engine/Source/Runtime/Engine/Private/ActorComponent.cpp
--- a/Engine/Source/Runtime/Engine/Private/ActorComponent.cpp
+++ b/Engine/Source/Runtime/Engine/Private/ActorComponent.cpp
@@ -273,7 +273,7 @@ bool UActorComponent::ComponentIsInPersistentLevel(bool bIncludeLevelStreamingPe
 														MyWorld->StreamingLevels[0]->GetLoadedLevel() == MyLevel ) );
 }
 
-FString UActorComponent::GetReadableName()
+FString UActorComponent::GetReadableName() const
 {
 	FString Result = GetNameSafe(GetOwner()) + TEXT(".") + GetName();
 	UObject const *Add = AdditionalStatObject();
diff --git a/Engine/Source/Runtime/Engine/Private/PhysicsEngine/BodySetup.cpp b/Engine/Source/Runtime/Engine/Private/PhysicsEngine/BodySetup.cpp
--- a/Engine/Source/Runtime/Engine/Private/PhysicsEngine/BodySetup.cpp
+++ b/Engine/Source/Runtime/Engine/Private/PhysicsEngine/BodySetup.cpp
@@ -1018,10 +1018,11 @@ float UBodySetup::CalculateMass(const UPrimitiveComponent* Component)
 	}
 
 	// Then scale mass to avoid big differences between big and small objects.
-	float BasicMass = AggGeom.GetVolume(ComponentScale) * DensityKGPerCubicUU;
+	const float BasicMass = AggGeom.GetVolume(ComponentScale) * DensityKGPerCubicUU;
+	ensureMsgf(BasicMass >= 0.0f, TEXT("UBodySetup::CalculateMass(%s) - The volume of the aggregate geometry is negative"), *Component->GetReadableName());
 
-	float UsePow = FMath::Clamp<float>(RaiseMassToPower, KINDA_SMALL_NUMBER, 1.f);
-	float RealMass = FMath::Pow(BasicMass, UsePow);
+	const float UsePow = FMath::Clamp<float>(RaiseMassToPower, KINDA_SMALL_NUMBER, 1.f);
+	const float RealMass = FMath::Pow(BasicMass, UsePow);
 
 	return RealMass * MassScale;
 }
