rearrange drop raise order

diff --git a/Engine/Source/Runtime/Engine/Classes/Vehicles/VehicleWheel.h b/Engine/Source/Runtime/Engine/Classes/Vehicles/VehicleWheel.h
--- a/Engine/Source/Runtime/Engine/Classes/Vehicles/VehicleWheel.h
+++ b/Engine/Source/Runtime/Engine/Classes/Vehicles/VehicleWheel.h
@@ -81,14 +81,14 @@ class ENGINE_API UVehicleWheel : public UObject
 	UPROPERTY(EditAnywhere, Category=Suspension)
 	float											SuspensionForceOffset;
 
-	/** How far the wheel can drop below the resting position */
-	UPROPERTY(EditAnywhere, Category=Suspension)
-	float											SuspensionMaxDrop;
-
 	/** How far the wheel can go above the resting position */
-	UPROPERTY(EditAnywhere, Category=Suspension)
+	UPROPERTY(EditAnywhere, Category = Suspension)
 	float											SuspensionMaxRaise;
 
+	/** How far the wheel can drop below the resting position */
+	UPROPERTY(EditAnywhere, Category=Suspension)
+	float											SuspensionMaxDrop;
+	
 	/** Oscillation frequency of suspension. Standard cars have values between 5 and 10 */
 	UPROPERTY(EditAnywhere, Category=Suspension)
 	float											SuspensionNaturalFrequency;
