Engine: Temporarily disable negative-volume ensure to unblock cooks

diff --git a/Engine/Source/Runtime/Engine/Private/PhysicsEngine/BodySetup.cpp b/Engine/Source/Runtime/Engine/Private/PhysicsEngine/BodySetup.cpp
--- a/Engine/Source/Runtime/Engine/Private/PhysicsEngine/BodySetup.cpp
+++ b/Engine/Source/Runtime/Engine/Private/PhysicsEngine/BodySetup.cpp
@@ -1019,7 +1019,7 @@ float UBodySetup::CalculateMass(const UPrimitiveComponent* Component)
 
 	// Then scale mass to avoid big differences between big and small objects.
 	const float BasicMass = AggGeom.GetVolume(ComponentScale) * DensityKGPerCubicUU;
-	ensureMsgf(BasicMass >= 0.0f, TEXT("UBodySetup::CalculateMass(%s) - The volume of the aggregate geometry is negative"), *Component->GetReadableName());
+	//@TODO: Some static meshes are triggering this - disabling until content can be analyzed - ensureMsgf(BasicMass >= 0.0f, TEXT("UBodySetup::CalculateMass(%s) - The volume of the aggregate geometry is negative"), *Component->GetReadableName());
 
 	const float UsePow = FMath::Clamp<float>(RaiseMassToPower, KINDA_SMALL_NUMBER, 1.f);
 	const float RealMass = FMath::Pow(BasicMass, UsePow);
