Expose vehicle tire load filter data

diff --git a/Engine/Source/Runtime/Engine/Classes/Vehicles/WheeledVehicleMovementComponent.h b/Engine/Source/Runtime/Engine/Classes/Vehicles/WheeledVehicleMovementComponent.h
--- a/Engine/Source/Runtime/Engine/Classes/Vehicles/WheeledVehicleMovementComponent.h
+++ b/Engine/Source/Runtime/Engine/Classes/Vehicles/WheeledVehicleMovementComponent.h
@@ -232,10 +232,18 @@ class ENGINE_API UWheeledVehicleMovementComponent : public UPawnMovementComponen
 	UPROPERTY(EditAnywhere, Category=VehicleSetup, AdvancedDisplay)
 	float MinNormalizedTireLoad;
 
+	/** Clamp normalized tire load to this value */
+	UPROPERTY(EditAnywhere, Category = VehicleSetup, AdvancedDisplay)
+	float MinNormalizedTireLoadFiltered;
+
 	/** Clamp normalized tire load to this value */
 	UPROPERTY(EditAnywhere, Category=VehicleSetup, AdvancedDisplay)
 	float MaxNormalizedTireLoad;
 
+	/** Clamp normalized tire load to this value */
+	UPROPERTY(EditAnywhere, Category = VehicleSetup, AdvancedDisplay)
+	float MaxNormalizedTireLoadFiltered;
+
 	// Our instanced wheels
 	UPROPERTY(transient, duplicatetransient, BlueprintReadOnly, Category=Vehicle)
 	TArray<class UVehicleWheel*> Wheels;
diff --git a/Engine/Source/Runtime/Engine/Private/Vehicles/WheeledVehicleMovementComponent.cpp b/Engine/Source/Runtime/Engine/Private/Vehicles/WheeledVehicleMovementComponent.cpp
--- a/Engine/Source/Runtime/Engine/Private/Vehicles/WheeledVehicleMovementComponent.cpp
+++ b/Engine/Source/Runtime/Engine/Private/Vehicles/WheeledVehicleMovementComponent.cpp
@@ -58,9 +58,9 @@ void PTireShader(const void* shaderData, const PxF32 tireFriction,
 	Input.WheelOmega = wheelOmega;
 	Input.WheelRadius = wheelRadius;
 	Input.RecipWheelRadius = recipWheelRadius;
-	Input.NormalizedTireLoad = FMath::Clamp( normalisedTireLoad, Wheel->VehicleSim->MinNormalizedTireLoad, Wheel->VehicleSim->MaxNormalizedTireLoad );
+	Input.NormalizedTireLoad = normalisedTireLoad;
 	Input.RestTireLoad = restTireLoad;
-	Input.TireLoad = restTireLoad * Input.NormalizedTireLoad;
+	Input.TireLoad = tireLoad;
 	Input.Gravity = gravity;
 	Input.RecipGravity = recipGravity;
 
@@ -108,6 +108,16 @@ UWheeledVehicleMovementComponent::UWheeledVehicleMovementComponent(const class F
 	HandbrakeInputRate.FallRate = 12.0f;
 	SteeringInputRate.RiseRate = 2.5f;
 	SteeringInputRate.FallRate = 5.0f;
+
+
+#if WITH_PHYSX
+	// tire load filtering
+	PxVehicleTireLoadFilterData PTireLoadFilterDef;
+	MinNormalizedTireLoad = PTireLoadFilterDef.mMinNormalisedLoad;
+	MinNormalizedTireLoadFiltered = PTireLoadFilterDef.mMinFilteredNormalisedLoad;
+	MaxNormalizedTireLoad = PTireLoadFilterDef.mMaxNormalisedLoad;
+	MaxNormalizedTireLoadFiltered = PTireLoadFilterDef.mMaxFilteredNormalisedLoad;
+#endif
 }
 
 #if WITH_PHYSX
@@ -342,6 +352,16 @@ void UWheeledVehicleMovementComponent::SetupWheels( PxVehicleWheelsSimData* PWhe
 		PWheelsSimData->setWheelShapeMapping( WheelIdx, WheelShapeIndex );
 		PWheelsSimData->setSceneQueryFilterData( WheelIdx, Shapes[WheelShapeIndex]->getQueryFilterData());
 	}
+
+	// tire load filtering
+	PxVehicleTireLoadFilterData PTireLoadFilter;
+	PTireLoadFilter.mMinNormalisedLoad = MinNormalizedTireLoad;
+	PTireLoadFilter.mMinFilteredNormalisedLoad = MinNormalizedTireLoadFiltered;
+	PTireLoadFilter.mMaxNormalisedLoad = MaxNormalizedTireLoad;
+	PTireLoadFilter.mMaxFilteredNormalisedLoad = MaxNormalizedTireLoadFiltered;
+	PWheelsSimData->setTireLoadFilterData(PTireLoadFilter);
+
+	
 }
  ////////////////////////////////////////////////////////////////////////////
 //Default tire force shader function.
@@ -456,9 +476,6 @@ void UWheeledVehicleMovementComponent::GenerateTireForces( UVehicleWheel* Wheel,
 	//UE_LOG( LogVehicles, Warning, TEXT("Friction = %f	LongSlip = %f	LatSlip = %f"), Input.TireFriction, Input.LongSlip, Input.LatSlip );	
 	//UE_LOG( LogVehicles, Warning, TEXT("WheelTorque= %f	LongForce = %f	LatForce = %f"), Output.WheelTorque, Output.LongForce, Output.LatForce );
 	//UE_LOG( LogVehicles, Warning, TEXT("RestLoad= %f	NormLoad = %f	TireLoad = %f"),Input.RestTireLoad, Input.NormalizedTireLoad, Input.TireLoad );
-	
-	//Output.LatForce *= FMath::Lerp( Wheel->WheelSetup->LatStiffFactor, Wheel->WheelSetup->HandbrakeLatStiffFactor, OutputHandbrake );
-	//Output.LatForce *= Wheel->LatStiffFactor;
 }
 
 void UWheeledVehicleMovementComponent::PostSetupVehicle()
