Replace PeakTorque with UI-editable TorqueCurve
Replace MaxSteeringSpeed & SteeringCurve with UI-editable SteeringCurve
Add GetRichCurveConst to FRuntimeFloatCurve and remove Eval/GetTimeRange

Submitting for James

diff --git a/Engine/Source/Editor/DetailCustomizations/Private/CurveStructCustomization.cpp b/Engine/Source/Editor/DetailCustomizations/Private/CurveStructCustomization.cpp
--- a/Engine/Source/Editor/DetailCustomizations/Private/CurveStructCustomization.cpp
+++ b/Engine/Source/Editor/DetailCustomizations/Private/CurveStructCustomization.cpp
@@ -55,6 +55,8 @@ void FCurveStructCustomization::CustomizeStructHeader( TSharedRef<class IPropert
 				InStructPropertyHandle->CreatePropertyNameWidget( TEXT( "" ), false )
 			]
 			.ValueContent()
+			.MinDesiredWidth(0.f)
+			.MaxDesiredWidth(0.f)
 			[
 				SNew(SVerticalBox)
 				+SVerticalBox::Slot()
diff --git a/Engine/Source/Runtime/Engine/Classes/Curves/CurveFloat.h b/Engine/Source/Runtime/Engine/Classes/Curves/CurveFloat.h
--- a/Engine/Source/Runtime/Engine/Classes/Curves/CurveFloat.h
+++ b/Engine/Source/Runtime/Engine/Classes/Curves/CurveFloat.h
@@ -18,11 +18,9 @@ struct FRuntimeFloatCurve
 
 	FRuntimeFloatCurve();
 
-	/** Evaluate this curve at the specified time */
-	float Eval(float InTime) const;
-
-	/** Get range of input time values. Outside this region curve continues constantly the start/end values. */
-	void GetTimeRange(float& MinTime, float& MaxTime) const;
+	/** Get the current curve struct */
+	FRichCurve* GetRichCurve();
+	const FRichCurve* GetRichCurveConst() const;
 };
 
 UCLASS(BlueprintType, MinimalAPI)
diff --git a/Engine/Source/Runtime/Engine/Classes/Vehicles/WheeledVehicleMovementComponent.h b/Engine/Source/Runtime/Engine/Classes/Vehicles/WheeledVehicleMovementComponent.h
--- a/Engine/Source/Runtime/Engine/Classes/Vehicles/WheeledVehicleMovementComponent.h
+++ b/Engine/Source/Runtime/Engine/Classes/Vehicles/WheeledVehicleMovementComponent.h
@@ -157,21 +157,6 @@ struct FVehicleInputRate
 	}
 };
 
-// Temporary until curve editing in
-USTRUCT()
-struct FFloatPair
-{
-	GENERATED_USTRUCT_BODY()
-
-	UPROPERTY(EditAnywhere, Category=FloatPair)
-	float InVal;
-
-	UPROPERTY(EditAnywhere, Category=FloatPair)
-	float OutVal;
-};
-
-void ConvertCurve( const TArray<FFloatPair>& Pairs, struct FRichCurve& OutCurve );
-
 /**
  * Component to handle the vehicle simulation for an actor.
  */
diff --git a/Engine/Source/Runtime/Engine/Classes/Vehicles/WheeledVehicleMovementComponent4W.h b/Engine/Source/Runtime/Engine/Classes/Vehicles/WheeledVehicleMovementComponent4W.h
--- a/Engine/Source/Runtime/Engine/Classes/Vehicles/WheeledVehicleMovementComponent4W.h
+++ b/Engine/Source/Runtime/Engine/Classes/Vehicles/WheeledVehicleMovementComponent4W.h
@@ -59,9 +59,9 @@ struct FVehicleEngineData
 {
 	GENERATED_USTRUCT_BODY()
 
-	/** Maximum torque available to apply to the engine (Nm). */ 
-	UPROPERTY(EditAnywhere, Category=Setup)
-	float PeakTorque;
+	/** Torque (Nm) at a given normalized RPM (0-1, internally scaled by MaxRPM). */
+	UPROPERTY(EditAnywhere, Category = Setup)
+	FRuntimeFloatCurve TorqueCurve;
 
 	/** Maximum revolutions per minute of the engine */
 	UPROPERTY(EditAnywhere, Category=Setup)
@@ -82,6 +82,9 @@ struct FVehicleEngineData
 	/** Damping rate of engine in at zero throttle when the clutch is disengaged (in neutral gear) (Kgm^2/s)*/
 	UPROPERTY(EditAnywhere, Category=Setup, AdvancedDisplay)
 	float DampingRateZeroThrottleClutchDisengaged;
+
+	/** Find the peak torque produced by the TorqueCurve */
+	float FindPeakTorque() const;
 };
 
 
@@ -157,17 +160,14 @@ class ENGINE_API UWheeledVehicleMovementComponent4W : public UWheeledVehicleMove
 	UPROPERTY(EditAnywhere, Category = MechanicalSetup)
 	FVehicleTransmissionData TransmissionSetup;
 
-	/** Max speed value in steering curve (km/h)*/
-	UPROPERTY(EditAnywhere, Category=SteeringSetup)
-	float MaxSteeringSpeed;
+	/** Maximum steering angle versus forward speed (km/h) */
+	UPROPERTY(EditAnywhere, Category = SteeringSetup)
+	FRuntimeFloatCurve SteeringCurve;
 
 	/** Accuracy of Ackermann steer calculation (range: 0..1) */
 	UPROPERTY(EditAnywhere, Category = SteeringSetup, AdvancedDisplay)
 	float AckermannAccuracy;
 
-	/** Steering strength defined for speed in normalized range [0 .. 1] */
-	UPROPERTY(EditAnywhere, Category=SteeringSetup)
-	TArray<FFloatPair> SteeringCurve;
 
 	virtual void Serialize(FArchive & Ar) OVERRIDE;
 	virtual void ComputeConstants() OVERRIDE;
@@ -179,8 +179,6 @@ protected:
 
 #if WITH_PHYSX
 
-	float SteeringMap[8*2];
-
 	/** Allocate and setup the PhysX vehicle */
 	virtual void SetupVehicle() OVERRIDE;
 
diff --git a/Engine/Source/Runtime/Engine/Private/CurveFloat.cpp b/Engine/Source/Runtime/Engine/Private/CurveFloat.cpp
--- a/Engine/Source/Runtime/Engine/Private/CurveFloat.cpp
+++ b/Engine/Source/Runtime/Engine/Private/CurveFloat.cpp
@@ -11,27 +11,27 @@ FRuntimeFloatCurve::FRuntimeFloatCurve()
 {
 }
 
-float FRuntimeFloatCurve::Eval(float InTime) const
+FRichCurve* FRuntimeFloatCurve::GetRichCurve()
 {
-	if (ExternalCurve)
+	if(ExternalCurve != NULL)
 	{
-		return ExternalCurve->GetFloatValue(InTime);
+		return &(ExternalCurve->FloatCurve);
 	}
 	else
 	{
-		return EditorCurveData.Eval(InTime);
+		return &EditorCurveData;
 	}
 }
 
-void FRuntimeFloatCurve::GetTimeRange(float& MinTime, float& MaxTime) const
+const FRichCurve* FRuntimeFloatCurve::GetRichCurveConst() const
 {
-	if (ExternalCurve)
+	if (ExternalCurve != NULL)
 	{
-		return ExternalCurve->GetTimeRange(MinTime, MaxTime);
+		return &(ExternalCurve->FloatCurve);
 	}
 	else
 	{
-		return EditorCurveData.GetTimeRange(MinTime, MaxTime);
+		return &EditorCurveData;
 	}
 }
 
diff --git a/Engine/Source/Runtime/Engine/Private/SoundNodeEnveloper.cpp b/Engine/Source/Runtime/Engine/Private/SoundNodeEnveloper.cpp
--- a/Engine/Source/Runtime/Engine/Private/SoundNodeEnveloper.cpp
+++ b/Engine/Source/Runtime/Engine/Private/SoundNodeEnveloper.cpp
@@ -117,9 +117,9 @@ void USoundNodeEnveloper::ParseNodes( FAudioDevice* AudioDevice, const UPTRINT N
 
 	FSoundParseParameters UpdatedParams = ParseParams;
 
-	UpdatedParams.Volume *= VolumeCurve.Eval(PlayTime) * UsedVolumeModulation;
+	UpdatedParams.Volume *= VolumeCurve.GetRichCurve()->Eval(PlayTime) * UsedVolumeModulation;
 
-	UpdatedParams.Pitch *= PitchCurve.Eval(PlayTime) * UsedPitchModulation;
+	UpdatedParams.Pitch *= PitchCurve.GetRichCurve()->Eval(PlayTime) * UsedPitchModulation;
 
 	Super::ParseNodes(AudioDevice, NodeWaveInstanceHash, ActiveSound, UpdatedParams, WaveInstances);
 }
diff --git a/Engine/Source/Runtime/Engine/Private/UserInterface/ForceFeedbackEffect.cpp b/Engine/Source/Runtime/Engine/Private/UserInterface/ForceFeedbackEffect.cpp
--- a/Engine/Source/Runtime/Engine/Private/UserInterface/ForceFeedbackEffect.cpp
+++ b/Engine/Source/Runtime/Engine/Private/UserInterface/ForceFeedbackEffect.cpp
@@ -30,7 +30,7 @@ float UForceFeedbackEffect::GetDuration()
 		float MinTime, MaxTime;
 		for (int32 Index = 0; Index < ChannelDetails.Num(); ++Index)
 		{
-			ChannelDetails[Index].Curve.GetTimeRange(MinTime, MaxTime);
+			ChannelDetails[Index].Curve.GetRichCurve()->GetTimeRange(MinTime, MaxTime);
 
 			if (MaxTime > Duration)
 			{
@@ -47,7 +47,7 @@ void UForceFeedbackEffect::GetValues(const float EvalTime, FForceFeedbackValues&
 	for (int32 Index = 0; Index < ChannelDetails.Num(); ++Index)
 	{
 		const FForceFeedbackChannelDetails& Details = ChannelDetails[Index];
-		const float Value = Details.Curve.Eval(EvalTime);
+		const float Value = Details.Curve.GetRichCurveConst()->Eval(EvalTime);
 
 		if (Details.bAffectsLeftLarge)
 		{
diff --git a/Engine/Source/Runtime/Engine/Private/Vehicles/WheeledVehicleMovementComponent.cpp b/Engine/Source/Runtime/Engine/Private/Vehicles/WheeledVehicleMovementComponent.cpp
--- a/Engine/Source/Runtime/Engine/Private/Vehicles/WheeledVehicleMovementComponent.cpp
+++ b/Engine/Source/Runtime/Engine/Private/Vehicles/WheeledVehicleMovementComponent.cpp
@@ -9,17 +9,6 @@
 #include "PhysXVehicleManager.h"
 #endif //WITH_PHYSX
 
-// Temporary until curve editing in
-void ConvertCurve( const TArray<FFloatPair>& Pairs, FRichCurve& OutCurve )
-{
-	OutCurve.Reset();
-
-	for ( auto It = Pairs.CreateConstIterator(); It; ++It )
-	{
-		OutCurve.AddKey( It->InVal, It->OutVal );
-	}
-}
-
 #if WITH_PHYSX
 
 /**
diff --git a/Engine/Source/Runtime/Engine/Private/Vehicles/WheeledVehicleMovementComponent4W.cpp b/Engine/Source/Runtime/Engine/Private/Vehicles/WheeledVehicleMovementComponent4W.cpp
--- a/Engine/Source/Runtime/Engine/Private/Vehicles/WheeledVehicleMovementComponent4W.cpp
+++ b/Engine/Source/Runtime/Engine/Private/Vehicles/WheeledVehicleMovementComponent4W.cpp
@@ -22,12 +22,20 @@ UWheeledVehicleMovementComponent4W::UWheeledVehicleMovementComponent4W(const cla
 
 	PxVehicleEngineData DefEngineData;
 	EngineSetup.MOI = DefEngineData.mMOI;
-	EngineSetup.PeakTorque = DefEngineData.mPeakTorque;
 	EngineSetup.MaxRPM = OmegaToRPM(DefEngineData.mMaxOmega);
 	EngineSetup.DampingRateFullThrottle = DefEngineData.mDampingRateFullThrottle;
 	EngineSetup.DampingRateZeroThrottleClutchEngaged = DefEngineData.mDampingRateZeroThrottleClutchEngaged;
 	EngineSetup.DampingRateZeroThrottleClutchDisengaged = DefEngineData.mDampingRateZeroThrottleClutchDisengaged;
 
+	// Convert from PhysX curve to ours
+	FRichCurve* TorqueCurveData = EngineSetup.TorqueCurve.GetRichCurve();
+	for(PxU32 KeyIdx=0; KeyIdx<DefEngineData.mTorqueCurve.getNbDataPairs(); KeyIdx++)
+	{
+		float Input = DefEngineData.mTorqueCurve.getX(KeyIdx);
+		float Output = DefEngineData.mTorqueCurve.getY(KeyIdx) * DefEngineData.mPeakTorque;
+		TorqueCurveData->AddKey(Input, Output);
+	}
+
 	PxVehicleClutchData DefClutchData;
 	TransmissionSetup.ClutchStrength = DefClutchData.mStrength;
 
@@ -53,17 +61,12 @@ UWheeledVehicleMovementComponent4W::UWheeledVehicleMovementComponent4W(const cla
 		TransmissionSetup.ForwardGears.Add(GearData);
 	}
 	
-	MaxSteeringSpeed = 100.f; // editable in vehicle blueprint
-	SteeringCurve.AddZeroed(4);
-	SteeringCurve[0].InVal = 0.0f;
-	SteeringCurve[0].OutVal = 0.75f;
-	SteeringCurve[1].InVal = 0.05f;
-	SteeringCurve[1].OutVal = 0.75f;
-	SteeringCurve[2].InVal = 0.25f;
-	SteeringCurve[2].OutVal = 0.125f;
-	SteeringCurve[3].InVal = 1.0f;
-	SteeringCurve[3].OutVal = 0.1f;
+	// Init steering speed curve
+	FRichCurve* SteeringCurveData = SteeringCurve.GetRichCurve();
+	SteeringCurveData->AddKey(0.f, 1.f);
+	SteeringCurveData->AddKey(50.f, 1.f);
 
+	// Initialize WheelSetups array with 4 wheels
 	WheelSetups.SetNum(4);
 #endif // WITH_PHYSX
 }
@@ -90,6 +93,16 @@ void UWheeledVehicleMovementComponent4W::PostEditChangeProperty(struct FProperty
 			GearData.UpRatio = FMath::Max(GearData.DownRatio, GearData.UpRatio);
 		}
 	}
+	else if (PropertyName == TEXT("SteeringCurve"))
+	{
+		//make sure values are capped between 0 and 1
+		TArray<FRichCurveKey> SteerKeys = SteeringCurve.GetRichCurve()->GetCopyOfKeys();
+		for (int32 KeyIdx = 0; KeyIdx < SteerKeys.Num(); ++KeyIdx)
+		{
+			float NewValue = FMath::Clamp(SteerKeys[KeyIdx].Value, 0.f, 1.f);
+			SteeringCurve.GetRichCurve()->UpdateOrAddKey(SteerKeys[KeyIdx].Time, NewValue);
+		}
+	}
 }
 #endif
 
@@ -129,14 +142,39 @@ static void GetVehicleDifferential4WSetup(const FVehicleDifferential4WData& Setu
 	PxSetup.mRearBias = Setup.RearBias;
 }
 
+float FVehicleEngineData::FindPeakTorque() const
+{
+	// Find max torque
+	float PeakTorque = 0.f;
+	TArray<FRichCurveKey> TorqueKeys = TorqueCurve.GetRichCurveConst()->GetCopyOfKeys();
+	for (int32 KeyIdx = 0; KeyIdx < TorqueKeys.Num(); KeyIdx++)
+	{
+		FRichCurveKey& Key = TorqueKeys[KeyIdx];
+		PeakTorque = FMath::Max(PeakTorque, Key.Value);
+	}
+	return PeakTorque;
+}
+
 static void GetVehicleEngineSetup(const FVehicleEngineData& Setup, PxVehicleEngineData& PxSetup)
 {
 	PxSetup.mMOI = M2ToCm2(Setup.MOI);
-	PxSetup.mPeakTorque = M2ToCm2(Setup.PeakTorque);	//convert Nm to (kg cm^2/s^2)
 	PxSetup.mMaxOmega = RPMToOmega(Setup.MaxRPM);
 	PxSetup.mDampingRateFullThrottle = M2ToCm2(Setup.DampingRateFullThrottle);
 	PxSetup.mDampingRateZeroThrottleClutchEngaged = M2ToCm2(Setup.DampingRateZeroThrottleClutchEngaged);
 	PxSetup.mDampingRateZeroThrottleClutchDisengaged = M2ToCm2(Setup.DampingRateZeroThrottleClutchDisengaged);
+
+	float PeakTorque = Setup.FindPeakTorque(); // In Nm
+	PxSetup.mPeakTorque = M2ToCm2(PeakTorque);	// convert Nm to (kg cm^2/s^2)
+
+	// Convert from our curve to PhysX
+	PxSetup.mTorqueCurve.clear();
+	TArray<FRichCurveKey> TorqueKeys = Setup.TorqueCurve.GetRichCurveConst()->GetCopyOfKeys();
+	int32 NumTorqueCurveKeys = FMath::Min<int32>(TorqueKeys.Num(), PxVehicleEngineData::eMAX_NB_ENGINE_TORQUE_CURVE_ENTRIES);
+	for (int32 KeyIdx = 0; KeyIdx < NumTorqueCurveKeys; KeyIdx++)
+	{
+		FRichCurveKey& Key = TorqueKeys[KeyIdx];
+		PxSetup.mTorqueCurve.addPair(FMath::Clamp(Key.Time, 0.f, 1.f), Key.Value/PeakTorque); // Normalize torque to 0-1 range
+	}
 }
 
 static void GetVehicleGearSetup(const FVehicleTransmissionData& Setup, PxVehicleGearsData& PxSetup)
@@ -239,16 +277,9 @@ void UWheeledVehicleMovementComponent4W::SetupVehicle()
 	PVehicleDrive = PVehicleDrive4W;
 
 	// MSS Using PxVehicle built-in automatic transmission for now. No reverse yet, have to fix that.
- //SetUseAutoGears(bUseGearAutoBox);
+	//SetUseAutoGears(bUseGearAutoBox);
 	SetUseAutoGears(true);
 	//SetTargetGear(1, true);
-
-	const int MaxSteeringSamples = FMath::Min(8, SteeringCurve.Num());
-	for (int i = 0; i < 8; i++)
-	{
-		SteeringMap[i*2 + 0] = (i < MaxSteeringSamples) ? (SteeringCurve[i].InVal * KmHToCmS(MaxSteeringSpeed)) : PX_MAX_F32;
-		SteeringMap[i*2 + 1] = (i < MaxSteeringSamples) ? FMath::Clamp<float>(SteeringCurve[i].OutVal, 0.0f, 1.0f) : PX_MAX_F32;
-	}
 }
 
 void UWheeledVehicleMovementComponent4W::UpdateSimulation(float DeltaTime)
@@ -268,7 +299,16 @@ void UWheeledVehicleMovementComponent4W::UpdateSimulation(float DeltaTime)
 		RawInputData.setGearDown(bRawGearDownInput);
 	}
 
-	PxFixedSizeLookupTable<8> SpeedSteerLookup(SteeringMap,4);
+	// Convert from our curve to PxFixedSizeLookupTable
+	PxFixedSizeLookupTable<8> SpeedSteerLookup;
+	TArray<FRichCurveKey> SteerKeys = SteeringCurve.GetRichCurve()->GetCopyOfKeys();
+	const int32 MaxSteeringSamples = FMath::Min(8, SteerKeys.Num());
+	for (int32 KeyIdx = 0; KeyIdx < MaxSteeringSamples; KeyIdx++)
+	{
+		FRichCurveKey& Key = SteerKeys[KeyIdx];
+		SpeedSteerLookup.addPair(KmHToCmS(Key.Time), FMath::Clamp(Key.Value, 0.f, 1.f));
+	}
+
 	PxVehiclePadSmoothingData SmoothData = {
 		{ ThrottleInputRate.RiseRate, BrakeInputRate.RiseRate, HandbrakeInputRate.RiseRate, SteeringInputRate.RiseRate, SteeringInputRate.RiseRate },
 		{ ThrottleInputRate.FallRate, BrakeInputRate.FallRate, HandbrakeInputRate.FallRate, SteeringInputRate.FallRate, SteeringInputRate.FallRate }
@@ -344,16 +384,13 @@ void UWheeledVehicleMovementComponent4W::Serialize(FArchive & Ar)
 		float DefaultRPM = OmegaToRPM(DefEngineData.mMaxOmega);
 		
 		//we need to convert from old units to new. This backwards compatable code fails in the rare case that they were using very strange values that are the new defaults in the correct units.
-		BackwardsConvertCm2ToM2(EngineSetup.PeakTorque, DefEngineData.mPeakTorque);
 		EngineSetup.MaxRPM = EngineSetup.MaxRPM != DefaultRPM ? OmegaToRPM(EngineSetup.MaxRPM) : DefaultRPM;	//need to convert from rad/s to RPM
-		MaxSteeringSpeed = MaxSteeringSpeed != 100.f ? CmSToKmH(MaxSteeringSpeed) : 100.f;	//we now store as km/h instead of cm/s
 	}
 
 	if (Ar.IsLoading() && Ar.UE4Ver() < VER_UE4_VEHICLES_UNIT_CHANGE2)
 	{
 		PxVehicleEngineData DefEngineData;
 		PxVehicleClutchData DefClutchData;
-		float DefaultRPM = OmegaToRPM(DefEngineData.mMaxOmega);
 
 		//we need to convert from old units to new. This backwards compatable code fails in the rare case that they were using very strange values that are the new defaults in the correct units.
 		BackwardsConvertCm2ToM2(EngineSetup.DampingRateFullThrottle, DefEngineData.mDampingRateFullThrottle);
