#ue4 - Add Character::OnWalkingOffLedge() event

- Makes it easier to detect transition from walking -> falling, exclusive of jumping or other things that may cause it.

diff --git a/Engine/Source/Runtime/Engine/Classes/GameFramework/Character.h b/Engine/Source/Runtime/Engine/Classes/GameFramework/Character.h
--- a/Engine/Source/Runtime/Engine/Classes/GameFramework/Character.h
+++ b/Engine/Source/Runtime/Engine/Classes/GameFramework/Character.h
@@ -375,6 +375,13 @@ public:
 	UFUNCTION(BlueprintImplementableEvent)
 	virtual void OnLanded(const FHitResult& Hit);
 
+	/**
+	 * Event fired when the Character is walking off a surface and is about to fall because CharacterMovement->CurrentFloor became unwalkable.
+	 * If CharacterMovement->MovementMode does not change (from Walking) during this event then the character will start falling.
+	 */
+	UFUNCTION(BlueprintNativeEvent, Category="Pawn|Character")
+	void OnWalkingOffLedge();
+
 	/** Called when pawn's movement is blocked
 		@PARAM Impact describes the blocking hit. */
 	virtual void MoveBlockedBy(const FHitResult& Impact) {};
diff --git a/Engine/Source/Runtime/Engine/Private/Character.cpp b/Engine/Source/Runtime/Engine/Private/Character.cpp
--- a/Engine/Source/Runtime/Engine/Private/Character.cpp
+++ b/Engine/Source/Runtime/Engine/Private/Character.cpp
@@ -184,6 +184,10 @@ UActorComponent* ACharacter::FindComponentByClass(const TSubclassOf<UActorCompon
 	return Super::FindComponentByClass(ComponentClass);
 }
 
+void ACharacter::OnWalkingOffLedge_Implementation()
+{
+}
+
 
 void ACharacter::Landed(const FHitResult& Hit)
 {
diff --git a/Engine/Source/Runtime/Engine/Private/CharacterMovementComponent.cpp b/Engine/Source/Runtime/Engine/Private/CharacterMovementComponent.cpp
--- a/Engine/Source/Runtime/Engine/Private/CharacterMovementComponent.cpp
+++ b/Engine/Source/Runtime/Engine/Private/CharacterMovementComponent.cpp
@@ -2691,7 +2691,12 @@ bool UCharacterMovementComponent::CheckFall(FHitResult &Hit, FVector Delta, FVec
 
 	if (bMustJump || CanWalkOffLedges())
 	{
-		StartFalling(Iterations, remainingTime, timeTick, Delta, subLoc);
+		CharacterOwner->OnWalkingOffLedge();
+		if (IsMovingOnGround())
+		{
+			// If still walking, then fall. If not, assume the user set a different mode they want to keep.
+			StartFalling(Iterations, remainingTime, timeTick, Delta, subLoc);
+		}
 		return true;
 	}
 	return false;
@@ -2999,7 +3004,12 @@ void UCharacterMovementComponent::PhysWalking(float deltaTime, int32 Iterations)
 			{
 				if (ShouldCatchAir(OldFloor, CurrentFloor))
 				{
-					StartFalling(Iterations, remainingTime, timeTick, Delta, subLoc);
+					CharacterOwner->OnWalkingOffLedge();
+					if (IsMovingOnGround())
+					{
+						// If still walking, then fall. If not, assume the user set a different mode they want to keep.
+						StartFalling(Iterations, remainingTime, timeTick, Delta, subLoc);
+					}
 					return;
 				}
 
