[GitHub] CanvasRenderTarget2D improvements

From DiCoJamesTan on GitHub https://github.com/EpicGames/UnrealEngine/pull/107
- Greatly simplified API, makes C++ and Blueprints more consistent
- New 'ReceiveUpdate()' event that can be implemented in a Blueprint to allow easy painting to the texture

MikeF: I cleaned up various minor bits during the integration (ENGINE_API, some naming, missing comments, etc.)

diff --git a/Engine/Source/Runtime/Engine/Classes/Engine/CanvasRenderTarget2D.h b/Engine/Source/Runtime/Engine/Classes/Engine/CanvasRenderTarget2D.h
--- a/Engine/Source/Runtime/Engine/Classes/Engine/CanvasRenderTarget2D.h
+++ b/Engine/Source/Runtime/Engine/Classes/Engine/CanvasRenderTarget2D.h
@@ -16,7 +16,7 @@ DECLARE_DYNAMIC_MULTICAST_DELEGATE_ThreeParams(FOnCanvasRenderTargetUpdate, UCan
  * function.  Also, remember to hold onto your new canvas render target with a reference so that it doesn't get
  * garbage collected.
  */
-UCLASS()
+UCLASS(BlueprintType, Blueprintable)
 class ENGINE_API UCanvasRenderTarget2D : public UTextureRenderTarget2D
 {
 	GENERATED_UCLASS_BODY()
@@ -31,45 +31,42 @@ public:
 	 * then the user's update delegate is called.  The canvas is then flushed and the RHI resolves the 
 	 * texture to make it available for rendering.
 	 */
+	UFUNCTION(BlueprintCallable, Category="Canvas Render Target 2D")
 	virtual void UpdateResource();
 
 	/**
-	 * Statically finds a canvas render target from the global map of canvas render targets. It can optionally create a new render target if not found as well.
+	 * Creates a new canvas render target and initializes it to the specified dimensions
 	 *
-	 * @param	InRenderTargetName	Name of the render target. This should be unique.
+	 * @param	Class				Class of the render target.  Unless you want to use a special sub-class, you can simply pass UCanvasRenderTarget2D::StaticClass() here.
 	 * @param	Width				Width of the render target.
 	 * @param	Height				Height of the render target.
-	 * @param	bCreateIfNotFound	True to create the render target if we couldn't find one with this name
+	 *
 	 * @return						Returns the instanced render target.
 	 */
 	UFUNCTION(BlueprintCallable, Category="Canvas Render Target 2D")
-	static UCanvasRenderTarget2D* FindCanvasRenderTarget2D(FName InRenderTargetName, int32 Width, int32 Height, bool bCreateIfNotFound=false);
+	static UCanvasRenderTarget2D* CreateCanvasRenderTarget2D(TSubclassOf<UCanvasRenderTarget2D> CanvasRenderTarget2DClass, int32 Width, int32 Height);
 
 	/**
-	 * Updates a specific render target from the global map of canvas render targets.  Does not trigger an error if the specified render target can't be found.
+	 * Allows a Blueprint to implement how this Canvas Render Target 2D should be updated.
 	 *
-	 * @param	InRenderTargetName		Name of the render target to update.
+	 * @param	Canvas				Canvas object that can be used to paint to the render target
+	 * @param	Width				Width of the render target.
+	 * @param	Height				Height of the render target.
 	 */
-	UFUNCTION(BlueprintCallable, Category="Canvas Render Target 2D")
-	static void UpdateCanvasRenderTarget(FName InRenderTargetName);
+	UFUNCTION(BlueprintImplementableEvent, Category="Canvas Render Target 2D")
+	virtual void ReceiveUpdate(class UCanvas* Canvas, int32 Width, int32 Height);
 
 	/**
-	 * Gets a specific render target's size from the global map of canvas render targets.  Does not trigger an error if the specified render target can't be found.
+	 * Gets a specific render target's size from the global map of canvas render targets.
 	 *
-	 * @param	Width	Out: The width of this render target
-	 * @param	Height	Out: The height of this render target
+	 * @param	Width	Output variable for the render target's width
+	 * @param	Height	Output variable for the render target's height
 	 */
 	UFUNCTION(BlueprintCallable, Category="Canvas Render Target 2D")
-	static void GetCanvasRenderTargetSize(FName InRenderTargetName, int32& Width, int32& Height);
+	void GetSize(int32& Width, int32& Height);
 
 
 	/** Called when this Canvas Render Target is asked to update its texture resource. */
 	UPROPERTY(BlueprintAssignable, Category="Canvas Render Target 2D")
 	FOnCanvasRenderTargetUpdate OnCanvasRenderTargetUpdate;
-
-
-protected:
-
-	/** Global mapped array of all canvas render targets */
-	static TMap<FName, TAutoWeakObjectPtr<UCanvasRenderTarget2D>> GCanvasRenderTarget2Ds;
 };
\ No newline at end of file
diff --git a/Engine/Source/Runtime/Engine/Private/CanvasRenderTarget2D.cpp b/Engine/Source/Runtime/Engine/Private/CanvasRenderTarget2D.cpp
--- a/Engine/Source/Runtime/Engine/Private/CanvasRenderTarget2D.cpp
+++ b/Engine/Source/Runtime/Engine/Private/CanvasRenderTarget2D.cpp
@@ -3,8 +3,6 @@
 #include "EnginePrivate.h"
 #include "Engine/CanvasRenderTarget2D.h"
 
-TMap<FName, TAutoWeakObjectPtr<UCanvasRenderTarget2D>> UCanvasRenderTarget2D::GCanvasRenderTarget2Ds;
-
 UCanvasRenderTarget2D::UCanvasRenderTarget2D( const class FPostConstructInitializeProperties& PCIP )
 	: Super(PCIP)
 {
@@ -54,6 +52,8 @@ void UCanvasRenderTarget2D::UpdateResource()
 		OnCanvasRenderTargetUpdate.Broadcast(Canvas, GetSurfaceWidth(), GetSurfaceHeight());
 	}
 
+	ReceiveUpdate(Canvas, GetSurfaceWidth(), GetSurfaceHeight());
+
 	// Clean up and flush the rendering canvas.
 	Canvas->Canvas = nullptr;
 	RenderCanvas.Flush();
@@ -72,70 +72,23 @@ void UCanvasRenderTarget2D::UpdateResource()
 	);
 }
 
-UCanvasRenderTarget2D* UCanvasRenderTarget2D::FindCanvasRenderTarget2D(FName InRenderTargetName, int32 Width, int32 Height, bool bCreateIfNotFound)
+UCanvasRenderTarget2D* UCanvasRenderTarget2D::CreateCanvasRenderTarget2D(TSubclassOf<UCanvasRenderTarget2D> CanvasRenderTarget2DClass, int32 Width, int32 Height)
 {
-	if (Width > 0 && Height > 0 && InRenderTargetName != NAME_None)
+	if (Width > 0 && Height > 0)
 	{
-		// Check if the render target with the name has already been created.
-		if (GCanvasRenderTarget2Ds.Num() > 0)
-		{
-			for (TMap<FName, TAutoWeakObjectPtr<UCanvasRenderTarget2D>>::TConstIterator It(GCanvasRenderTarget2Ds); It; ++It)
-			{
-				if (It->Key == InRenderTargetName && (It->Value).IsValid())
-				{
-					return (It->Value).Get();
-				}
-			}
-		}
-
-		// Create a new one since it doesn't exist.
-		if (bCreateIfNotFound)
+		UCanvasRenderTarget2D* NewCanvasRenderTarget = ConstructObject<UCanvasRenderTarget2D>(CanvasRenderTarget2DClass, GetTransientPackage());
+		if (NewCanvasRenderTarget)
 		{
-			UCanvasRenderTarget2D* NewCanvasRenderTarget = ConstructObject<UCanvasRenderTarget2D>(UCanvasRenderTarget2D::StaticClass(), GetTransientPackage());
-			if (NewCanvasRenderTarget)
-			{
-				NewCanvasRenderTarget->InitAutoFormat(Width, Height);
-
-				// Add it to the global static map.
-				GCanvasRenderTarget2Ds.Add(InRenderTargetName, NewCanvasRenderTarget);
-
-				return NewCanvasRenderTarget;
-			}
+			NewCanvasRenderTarget->InitAutoFormat(Width, Height);
+			return NewCanvasRenderTarget;
 		}
 	}
 
 	return nullptr;
 }
 
-void UCanvasRenderTarget2D::UpdateCanvasRenderTarget(FName InRenderTargetName)
+void UCanvasRenderTarget2D::GetSize(int32& Width, int32& Height)
 {
-	if (GCanvasRenderTarget2Ds.Num() > 0)
-	{
-		for (TMap<FName, TAutoWeakObjectPtr<UCanvasRenderTarget2D>>::TConstIterator It(GCanvasRenderTarget2Ds); It; ++It)
-		{
-			if (It->Key == InRenderTargetName && (It->Value).IsValid())
-			{
-				(It->Value)->UpdateResource();
-				
-				break;
-			}
-		}
-	}
-}
-
-void UCanvasRenderTarget2D::GetCanvasRenderTargetSize(FName InRenderTargetName, int32& Width, int32& Height)
-{
-	if (GCanvasRenderTarget2Ds.Num() > 0)
-	{
-		for (TMap<FName, TAutoWeakObjectPtr<UCanvasRenderTarget2D>>::TConstIterator It(GCanvasRenderTarget2Ds); It; ++It)
-		{
-			if (It->Key == InRenderTargetName && (It->Value).IsValid())
-			{
-				Width = (It->Value)->GetSurfaceWidth();
-				Height = (It->Value)->GetSurfaceHeight();
-
-				break;
-			}
-		}
-	}
+	Width = GetSurfaceWidth();
+	Height = GetSurfaceHeight();
 }
\ No newline at end of file
