default steering curve

diff --git a/Engine/Source/Editor/DetailCustomizations/Private/WheeledVehicleMovementComponent4WDetails.cpp b/Engine/Source/Editor/DetailCustomizations/Private/WheeledVehicleMovementComponent4WDetails.cpp
--- a/Engine/Source/Editor/DetailCustomizations/Private/WheeledVehicleMovementComponent4WDetails.cpp
+++ b/Engine/Source/Editor/DetailCustomizations/Private/WheeledVehicleMovementComponent4WDetails.cpp
@@ -58,6 +58,7 @@ void FWheeledVehicleMovementComponent4WDetails::CustomizeDetails(IDetailLayoutBu
 		.ViewMaxOutput(1.f)
 		.TimelineLength(150)
 		.HideUI(false)
+		.ZoomToFit(false)
 		.DesiredSize(FVector2D(512, 128))
 	];
 
diff --git a/Engine/Source/Runtime/Engine/Classes/Vehicles/WheeledVehicleMovementComponent4W.h b/Engine/Source/Runtime/Engine/Classes/Vehicles/WheeledVehicleMovementComponent4W.h
--- a/Engine/Source/Runtime/Engine/Classes/Vehicles/WheeledVehicleMovementComponent4W.h
+++ b/Engine/Source/Runtime/Engine/Classes/Vehicles/WheeledVehicleMovementComponent4W.h
@@ -160,7 +160,7 @@ class ENGINE_API UWheeledVehicleMovementComponent4W : public UWheeledVehicleMove
 	UPROPERTY(EditAnywhere, Category = MechanicalSetup)
 	FVehicleTransmissionData TransmissionSetup;
 
-	/** Maximum steering angle versus forward speed (km/h) */
+	/** Maximum steering versus forward speed (km/h) */
 	UPROPERTY(EditAnywhere, Category = SteeringSetup)
 	FRuntimeFloatCurve SteeringCurve;
 
diff --git a/Engine/Source/Runtime/Engine/Private/Vehicles/WheeledVehicleMovementComponent4W.cpp b/Engine/Source/Runtime/Engine/Private/Vehicles/WheeledVehicleMovementComponent4W.cpp
--- a/Engine/Source/Runtime/Engine/Private/Vehicles/WheeledVehicleMovementComponent4W.cpp
+++ b/Engine/Source/Runtime/Engine/Private/Vehicles/WheeledVehicleMovementComponent4W.cpp
@@ -64,7 +64,9 @@ UWheeledVehicleMovementComponent4W::UWheeledVehicleMovementComponent4W(const cla
 	// Init steering speed curve
 	FRichCurve* SteeringCurveData = SteeringCurve.GetRichCurve();
 	SteeringCurveData->AddKey(0.f, 1.f);
-	SteeringCurveData->AddKey(50.f, 1.f);
+	SteeringCurveData->AddKey(20.f, 0.9f);
+	SteeringCurveData->AddKey(60.f, 0.8f);
+	SteeringCurveData->AddKey(120.f, 0.7f);
 
 	// Initialize WheelSetups array with 4 wheels
 	WheelSetups.SetNum(4);
