Clean up torque curve

diff --git a/Engine/Source/Editor/DetailCustomizations/Private/WheeledVehicleMovementComponent4WDetails.cpp b/Engine/Source/Editor/DetailCustomizations/Private/WheeledVehicleMovementComponent4WDetails.cpp
--- a/Engine/Source/Editor/DetailCustomizations/Private/WheeledVehicleMovementComponent4WDetails.cpp
+++ b/Engine/Source/Editor/DetailCustomizations/Private/WheeledVehicleMovementComponent4WDetails.cpp
@@ -31,38 +31,64 @@ void FWheeledVehicleMovementComponent4WDetails::CustomizeDetails(IDetailLayoutBu
 	else if (UWheeledVehicleMovementComponent4W * VehicleComponent = Cast<UWheeledVehicleMovementComponent4W>(SelectedObjects[0].Get()))
 	{
 		SteeringCurveEditor = FSteeringCurveEditor(VehicleComponent);
+		TorqueCurveEditor = FTorqueCurveEditor(VehicleComponent);
 	}
 	else
 	{
 		return;
 	}
 
-	// See which categories are hidden
-	TArray<FString> HideCategories;
-	DetailBuilder.GetDetailsView().GetBaseClass()->GetHideCategories(HideCategories);
+	//Torque curve
+	{
+		IDetailCategoryBuilder& MechanicalCategory = DetailBuilder.EditCategory("MechanicalSetup");
+		TSharedRef<IPropertyHandle> TorqueCurveHandle = DetailBuilder.GetProperty("EngineSetup.TorqueCurve");
+
+		MechanicalCategory.AddProperty(TorqueCurveHandle).CustomWidget()
+		.NameContent()
+			[
+				TorqueCurveHandle->CreatePropertyNameWidget()
+			]
+		.ValueContent()
+			.MinDesiredWidth(125.f * 3.f)
+			[
+				SAssignNew(TorqueCurveWidget, SCurveEditor)
+				.ViewMinInput(0.f)
+				.ViewMaxInput(70000.f)
+				.ViewMinOutput(0.f)
+				.ViewMaxOutput(1.f)
+				.TimelineLength(7000.f)
+				.HideUI(false)
+				.DesiredSize(FVector2D(512, 128))
+			];
+
+		TorqueCurveWidget->SetCurveOwner(&TorqueCurveEditor);
+	}
 	
-	IDetailCategoryBuilder& SteeringCategory = DetailBuilder.EditCategory("SteeringSetup");
-	TSharedRef<IPropertyHandle> SteeringCurveHandle = DetailBuilder.GetProperty("SteeringCurve");
-	SteeringCategory.AddProperty(SteeringCurveHandle).CustomWidget()
-	.NameContent()
-	[
-		SteeringCurveHandle->CreatePropertyNameWidget()
-	]
-	.ValueContent()
-	.MinDesiredWidth(125.f * 3.f)
-	[
-		SAssignNew(SteeringCurveWidget, SCurveEditor)
-		.ViewMinInput(0.f)
-		.ViewMaxInput(150.f)
-		.ViewMinOutput(0.f)
-		.ViewMaxOutput(1.f)
-		.TimelineLength(150)
-		.HideUI(false)
-		.ZoomToFit(false)
-		.DesiredSize(FVector2D(512, 128))
-	];
-
-	SteeringCurveWidget->SetCurveOwner(&SteeringCurveEditor);
+	//Steering curve
+	{
+		IDetailCategoryBuilder& SteeringCategory = DetailBuilder.EditCategory("SteeringSetup");
+		TSharedRef<IPropertyHandle> SteeringCurveHandle = DetailBuilder.GetProperty("SteeringCurve");
+		SteeringCategory.AddProperty(SteeringCurveHandle).CustomWidget()
+			.NameContent()
+			[
+				SteeringCurveHandle->CreatePropertyNameWidget()
+			]
+		.ValueContent()
+			.MinDesiredWidth(125.f * 3.f)
+			[
+				SAssignNew(SteeringCurveWidget, SCurveEditor)
+				.ViewMinInput(0.f)
+				.ViewMaxInput(150.f)
+				.ViewMinOutput(0.f)
+				.ViewMaxOutput(1.f)
+				.TimelineLength(150)
+				.HideUI(false)
+				.ZoomToFit(false)
+				.DesiredSize(FVector2D(512, 128))
+			];
+
+		SteeringCurveWidget->SetCurveOwner(&SteeringCurveEditor);
+	}
 }
 
 TArray<FRichCurveEditInfoConst> FWheeledVehicleMovementComponent4WDetails::FSteeringCurveEditor::GetCurves() const
@@ -108,5 +134,49 @@ FWheeledVehicleMovementComponent4WDetails::FSteeringCurveEditor::FSteeringCurveE
 	Owner = VehicleComponent;
 }
 
+
+TArray<FRichCurveEditInfoConst> FWheeledVehicleMovementComponent4WDetails::FTorqueCurveEditor::GetCurves() const
+{
+	TArray<FRichCurveEditInfoConst> Curves;
+	Curves.Add(&VehicleComponent->EngineSetup.TorqueCurve.EditorCurveData);
+
+	return Curves;
+}
+
+TArray<FRichCurveEditInfo> FWheeledVehicleMovementComponent4WDetails::FTorqueCurveEditor::GetCurves()
+{
+	TArray<FRichCurveEditInfo> Curves;
+	Curves.Add(&VehicleComponent->EngineSetup.TorqueCurve.EditorCurveData);
+
+	return Curves;
+}
+
+UObject * FWheeledVehicleMovementComponent4WDetails::FTorqueCurveEditor::GetOwner()
+{
+	return Owner;
+}
+
+void FWheeledVehicleMovementComponent4WDetails::FTorqueCurveEditor::ModifyOwner()
+{
+	if (Owner)
+	{
+		Owner->Modify();
+	}
+}
+
+void FWheeledVehicleMovementComponent4WDetails::FTorqueCurveEditor::MakeTransactional()
+{
+	if (Owner)
+	{
+		Owner->SetFlags(Owner->GetFlags() | RF_Transactional);
+	}
+}
+
+FWheeledVehicleMovementComponent4WDetails::FTorqueCurveEditor::FTorqueCurveEditor(UWheeledVehicleMovementComponent4W * InVehicle)
+{
+	VehicleComponent = InVehicle;
+	Owner = VehicleComponent;
+}
+
 #undef LOCTEXT_NAMESPACE
 
diff --git a/Engine/Source/Editor/DetailCustomizations/Private/WheeledVehicleMovementComponent4WDetails.h b/Engine/Source/Editor/DetailCustomizations/Private/WheeledVehicleMovementComponent4WDetails.h
--- a/Engine/Source/Editor/DetailCustomizations/Private/WheeledVehicleMovementComponent4WDetails.h
+++ b/Engine/Source/Editor/DetailCustomizations/Private/WheeledVehicleMovementComponent4WDetails.h
@@ -31,9 +31,28 @@ private:
 
 	} SteeringCurveEditor;
 
+	struct FTorqueCurveEditor : public FCurveOwnerInterface
+	{
+		FTorqueCurveEditor(UWheeledVehicleMovementComponent4W * InVehicle = NULL);
+		/** FCurveOwnerInterface interface */
+		virtual TArray<FRichCurveEditInfoConst> GetCurves() const OVERRIDE;
+		virtual TArray<FRichCurveEditInfo> GetCurves() OVERRIDE;
+		virtual UObject* GetOwner() OVERRIDE;
+		virtual void ModifyOwner() OVERRIDE;
+		virtual void MakeTransactional() OVERRIDE;
+
+	private:
+		UWheeledVehicleMovementComponent4W * VehicleComponent;
+		UObject * Owner;
+
+	} TorqueCurveEditor;
+
 	TArray<TWeakObjectPtr<UObject> > SelectedObjects;	//the objects we're showing details for
 	
 	/** Steering curve widget */
 	TSharedPtr<class SCurveEditor> SteeringCurveWidget;
+
+	/** Torque curve widget */
+	TSharedPtr<class SCurveEditor> TorqueCurveWidget;
 };
 
diff --git a/Engine/Source/Runtime/Engine/Classes/Vehicles/WheeledVehicleMovementComponent4W.h b/Engine/Source/Runtime/Engine/Classes/Vehicles/WheeledVehicleMovementComponent4W.h
--- a/Engine/Source/Runtime/Engine/Classes/Vehicles/WheeledVehicleMovementComponent4W.h
+++ b/Engine/Source/Runtime/Engine/Classes/Vehicles/WheeledVehicleMovementComponent4W.h
@@ -30,27 +30,27 @@ struct FVehicleDifferential4WData
 	TEnumAsByte<EVehicleDifferential4W::Type> DifferentialType;
 	
 	/** Ratio of torque split between front and rear (>0.5 means more to front, <0.5 means more to rear, works only with 4W type) */
-	UPROPERTY(EditAnywhere, Category=Setup)
+	UPROPERTY(EditAnywhere, Category = Setup, meta = (ClampMin = "0.0", UIMin = "0.0", ClampMax = "1.0", UIMax = "1.0"))
 	float FrontRearSplit;
 
 	/** Ratio of torque split between front-left and front-right (>0.5 means more to front-left, <0.5 means more to front-right, works only with 4W and LimitedSlip_FrontDrive) */
-	UPROPERTY(EditAnywhere, Category=Setup)
+	UPROPERTY(EditAnywhere, Category = Setup, meta = (ClampMin = "0.0", UIMin = "0.0", ClampMax = "1.0", UIMax = "1.0"))
 	float FrontLeftRightSplit;
 
 	/** Ratio of torque split between rear-left and rear-right (>0.5 means more to rear-left, <0.5 means more to rear-right, works only with 4W and LimitedSlip_RearDrive) */
-	UPROPERTY(EditAnywhere, Category=Setup)
+	UPROPERTY(EditAnywhere, Category = Setup, meta = (ClampMin = "0.0", UIMin = "0.0", ClampMax = "1.0", UIMax = "1.0"))
 	float RearLeftRightSplit;
 	
 	/** Maximum allowed ratio of average front wheel rotation speed and rear wheel rotation speeds (range: 1..inf, works only with LimitedSlip_4W) */
-	UPROPERTY(EditAnywhere, Category=Setup)
+	UPROPERTY(EditAnywhere, Category = Setup, meta = (ClampMin = "1.0", UIMin = "1.0"))
 	float CentreBias;
 
 	/** Maximum allowed ratio of front-left and front-right wheel rotation speeds (range: 1..inf, works only with LimitedSlip_4W, LimitedSlip_FrontDrive) */
-	UPROPERTY(EditAnywhere, Category=Setup)
+	UPROPERTY(EditAnywhere, Category = Setup, meta = (ClampMin = "1.0", UIMin = "1.0"))
 	float FrontBias;
 
 	/** Maximum allowed ratio of rear-left and rear-right wheel rotation speeds (range: 1..inf, works only with LimitedSlip_4W, LimitedSlip_FrontDrive) */
-	UPROPERTY(EditAnywhere, Category=Setup)
+	UPROPERTY(EditAnywhere, Category = Setup, meta = (ClampMin = "1.0", UIMin = "1.0"))
 	float RearBias;
 };
 
@@ -59,28 +59,28 @@ struct FVehicleEngineData
 {
 	GENERATED_USTRUCT_BODY()
 
-	/** Torque (Nm) at a given normalized RPM (0-1, internally scaled by MaxRPM). */
+	/** Torque (Nm) at a given RPM*/
 	UPROPERTY(EditAnywhere, Category = Setup)
 	FRuntimeFloatCurve TorqueCurve;
 
 	/** Maximum revolutions per minute of the engine */
-	UPROPERTY(EditAnywhere, Category=Setup)
+	UPROPERTY(EditAnywhere, Category = Setup, meta = (ClampMin = "0.01", UIMin = "0.01"))
 	float MaxRPM;
 
 	/** Moment of inertia of the engine around the axis of rotation (Kgm^2). */
-	UPROPERTY(EditAnywhere, Category = Setup)
+	UPROPERTY(EditAnywhere, Category = Setup, meta = (ClampMin = "0.01", UIMin = "0.01"))
 	float MOI;
 
 	/** Damping rate of engine when full throttle is applied (Kgm^2/s) */
-	UPROPERTY(EditAnywhere, Category=Setup, AdvancedDisplay)
+	UPROPERTY(EditAnywhere, Category = Setup, AdvancedDisplay, meta = (ClampMin = "0.0", UIMin = "0.0"))
 	float DampingRateFullThrottle;
 
 	/** Damping rate of engine in at zero throttle when the clutch is engaged (Kgm^2/s)*/
-	UPROPERTY(EditAnywhere, Category=Setup, AdvancedDisplay)
+	UPROPERTY(EditAnywhere, Category = Setup, AdvancedDisplay, meta = (ClampMin = "0.0", UIMin = "0.0"))
 	float DampingRateZeroThrottleClutchEngaged;
 
 	/** Damping rate of engine in at zero throttle when the clutch is disengaged (in neutral gear) (Kgm^2/s)*/
-	UPROPERTY(EditAnywhere, Category=Setup, AdvancedDisplay)
+	UPROPERTY(EditAnywhere, Category = Setup, AdvancedDisplay, meta = (ClampMin = "0.0", UIMin = "0.0"))
 	float DampingRateZeroThrottleClutchDisengaged;
 
 	/** Find the peak torque produced by the TorqueCurve */
@@ -115,11 +115,11 @@ struct FVehicleTransmissionData
 	bool bUseGearAutoBox;
 
 	/** Time it takes to switch gears (seconds) */
-	UPROPERTY(EditAnywhere, Category = Setup)
+	UPROPERTY(EditAnywhere, Category = Setup, meta = (ClampMin = "0.0", UIMin = "0.0"))
 	float GearSwitchTime;
 	
 	/** Minimum time it takes the automatic transmission to initiate a gear change (seconds)*/
-	UPROPERTY(EditAnywhere, Category = Setup, meta = (editcondition = "bUseGearAutoBox"))
+	UPROPERTY(EditAnywhere, Category = Setup, meta = (editcondition = "bUseGearAutoBox", ClampMin = "0.0", UIMin="0.0"))
 	float GearAutoBoxLatency;
 
 	/** The final gear ratio multiplies the transmission gear ratios.*/
@@ -139,7 +139,7 @@ struct FVehicleTransmissionData
 	float NeutralGearUpRatio;
 
 	/** Strength of clutch (Kgm^2/s)*/
-	UPROPERTY(EditAnywhere, Category = Setup, AdvancedDisplay)
+	UPROPERTY(EditAnywhere, Category = Setup, AdvancedDisplay, meta = (ClampMin = "0.0", UIMin = "0.0"))
 	float ClutchStrength;
 };
 
@@ -165,7 +165,7 @@ class ENGINE_API UWheeledVehicleMovementComponent4W : public UWheeledVehicleMove
 	FRuntimeFloatCurve SteeringCurve;
 
 	/** Accuracy of Ackermann steer calculation (range: 0..1) */
-	UPROPERTY(EditAnywhere, Category = SteeringSetup, AdvancedDisplay)
+	UPROPERTY(EditAnywhere, Category = SteeringSetup, AdvancedDisplay, meta = (ClampMin = "0.0", UIMin = "0.0", ClampMax = "1.0", UIMax = "1.0"))
 	float AckermannAccuracy;
 
 
diff --git a/Engine/Source/Runtime/Engine/Private/Vehicles/WheeledVehicleMovementComponent4W.cpp b/Engine/Source/Runtime/Engine/Private/Vehicles/WheeledVehicleMovementComponent4W.cpp
--- a/Engine/Source/Runtime/Engine/Private/Vehicles/WheeledVehicleMovementComponent4W.cpp
+++ b/Engine/Source/Runtime/Engine/Private/Vehicles/WheeledVehicleMovementComponent4W.cpp
@@ -31,7 +31,7 @@ UWheeledVehicleMovementComponent4W::UWheeledVehicleMovementComponent4W(const cla
 	FRichCurve* TorqueCurveData = EngineSetup.TorqueCurve.GetRichCurve();
 	for(PxU32 KeyIdx=0; KeyIdx<DefEngineData.mTorqueCurve.getNbDataPairs(); KeyIdx++)
 	{
-		float Input = DefEngineData.mTorqueCurve.getX(KeyIdx);
+		float Input = DefEngineData.mTorqueCurve.getX(KeyIdx) * EngineSetup.MaxRPM;
 		float Output = DefEngineData.mTorqueCurve.getY(KeyIdx) * DefEngineData.mPeakTorque;
 		TorqueCurveData->AddKey(Input, Output);
 	}
@@ -175,7 +175,7 @@ static void GetVehicleEngineSetup(const FVehicleEngineData& Setup, PxVehicleEngi
 	for (int32 KeyIdx = 0; KeyIdx < NumTorqueCurveKeys; KeyIdx++)
 	{
 		FRichCurveKey& Key = TorqueKeys[KeyIdx];
-		PxSetup.mTorqueCurve.addPair(FMath::Clamp(Key.Time, 0.f, 1.f), Key.Value/PeakTorque); // Normalize torque to 0-1 range
+		PxSetup.mTorqueCurve.addPair(FMath::Clamp(Key.Time / Setup.MaxRPM, 0.f, 1.f), Key.Value/PeakTorque); // Normalize torque to 0-1 range
 	}
 }
 
