Moved StopMovement function from AAIController to AController as means to stop movement requested with SimpleMoveTo[Actor/Location] #UE4

diff --git a/Engine/Source/Runtime/Engine/Classes/GameFramework/AIController.h b/Engine/Source/Runtime/Engine/Classes/GameFramework/AIController.h
--- a/Engine/Source/Runtime/Engine/Classes/GameFramework/AIController.h
+++ b/Engine/Source/Runtime/Engine/Classes/GameFramework/AIController.h
@@ -154,9 +154,8 @@ public:
 	/** resumes last AI-performed, paused request provided it's ID was equivalent to RequestToResume */
 	bool ResumeMove(FAIRequestID RequestToResume);
 
-	/** Aborts the move AI is currently performing */
-	UFUNCTION(BlueprintCallable, Category="AI|Navigation")
-	virtual void StopMovement();
+	/** Aborts the move the controller is currently performing */
+	virtual void StopMovement() OVERRIDE;
 
 	/** Called on completing current movement request */
 	virtual void OnMoveCompleted(FAIRequestID RequestID, EPathFollowingResult::Type Result);
diff --git a/Engine/Source/Runtime/Engine/Classes/GameFramework/Controller.h b/Engine/Source/Runtime/Engine/Classes/GameFramework/Controller.h
--- a/Engine/Source/Runtime/Engine/Classes/GameFramework/Controller.h
+++ b/Engine/Source/Runtime/Engine/Classes/GameFramework/Controller.h
@@ -266,7 +266,7 @@ public:
 	virtual FVector GetNavAgentLocation() const OVERRIDE;
 	virtual void GetMoveGoalReachTest(class AActor* MovingActor, const FVector& MoveOffset, FVector& GoalOffset, float& GoalRadius, float& GoalHalfHeight) const OVERRIDE;
 	// End INavAgentInterface Interface
-
+	
 	/** prepares path finding and path following components */
 	virtual void InitNavigationControl(UNavigationComponent*& PathFindingComp, UPathFollowingComponent*& PathFollowingComp);
 
@@ -274,6 +274,10 @@ public:
 	 *	makes them update their cached data */
 	virtual void UpdateNavigationComponents();
 
+	/** Aborts the move the controller is currently performing */
+	UFUNCTION(BlueprintCallable, Category = "AI|Navigation")
+	virtual void StopMovement();
+
 protected:
 	/** State entered when inactive (no possessed pawn, not spectating, etc). */
 	virtual void BeginInactiveState();
diff --git a/Engine/Source/Runtime/Engine/Private/AIController.cpp b/Engine/Source/Runtime/Engine/Private/AIController.cpp
--- a/Engine/Source/Runtime/Engine/Private/AIController.cpp
+++ b/Engine/Source/Runtime/Engine/Private/AIController.cpp
@@ -546,7 +546,7 @@ bool AAIController::ResumeMove(FAIRequestID RequestToResume)
 
 void AAIController::StopMovement()
 {
-	UE_VLOG(this, LogNavigation, Log, TEXT("%s STOP MOVEMENT"), *GetNameSafe(GetPawn()) );
+	UE_VLOG(this, LogNavigation, Log, TEXT("AAIController::StopMovement: %s STOP MOVEMENT"), *GetNameSafe(GetPawn()) );
 	PathFollowingComponent->AbortMove(TEXT("StopMovement"));
 }
 
diff --git a/Engine/Source/Runtime/Engine/Private/Controller.cpp b/Engine/Source/Runtime/Engine/Private/Controller.cpp
--- a/Engine/Source/Runtime/Engine/Private/Controller.cpp
+++ b/Engine/Source/Runtime/Engine/Private/Controller.cpp
@@ -592,6 +592,17 @@ void AController::InitNavigationControl(UNavigationComponent*& PathFindingComp,
 	}
 }
 
+void AController::StopMovement()
+{
+	UE_VLOG(this, LogNavigation, Log, TEXT("AController::StopMovement: %s STOP MOVEMENT"), *GetNameSafe(GetPawn()));
+
+	UPathFollowingComponent* PathFollowingComp = FindComponentByClass<UPathFollowingComponent>();
+	if (PathFollowingComp != NULL)
+	{
+		PathFollowingComp->AbortMove(TEXT("StopMovement"));
+	}
+}
+
 void AController::GetLifetimeReplicatedProps( TArray< FLifetimeProperty > & OutLifetimeProps ) const
 {
 	Super::GetLifetimeReplicatedProps( OutLifetimeProps );
