Don't clear the CachedMousePos when the mouse leaves as this can be triggered by a right-click.
#ttp 333924 CRITICAL: Regression: EDITOR: Place Actor location incorrect after deleting SkySphere

diff --git a/Engine/Source/Runtime/Engine/Private/Slate/SceneViewport.cpp b/Engine/Source/Runtime/Engine/Private/Slate/SceneViewport.cpp
--- a/Engine/Source/Runtime/Engine/Private/Slate/SceneViewport.cpp
+++ b/Engine/Source/Runtime/Engine/Private/Slate/SceneViewport.cpp
@@ -385,7 +385,6 @@ void FSceneViewport::OnMouseEnter( const FGeometry& MyGeometry, const FPointerEv
 
 void FSceneViewport::OnMouseLeave( const FPointerEvent& MouseEvent )
 {
-	CachedMousePos = FIntPoint(-1, -1);
 	ViewportClient->MouseLeave( this );
 }
 
