#ttp 331438 UE4: PlatformMath refactor

diff --git a/Engine/Source/Runtime/Engine/Private/InterpolationDraw.cpp b/Engine/Source/Runtime/Engine/Private/InterpolationDraw.cpp
--- a/Engine/Source/Runtime/Engine/Private/InterpolationDraw.cpp
+++ b/Engine/Source/Runtime/Engine/Private/InterpolationDraw.cpp
@@ -50,7 +50,7 @@ void UInterpTrack::DrawTrack( FCanvas* Canvas, UInterpGroup* Group, const FInter
 	{
 		float KeyTime = GetKeyframeTime(i);
 
-		int32 PixelPos = FMath::Trunc((KeyTime - Params.StartTime) * Params.PixelsPerSec);
+		int32 PixelPos = FMath::TruncToInt((KeyTime - Params.StartTime) * Params.PixelsPerSec);
 
 		FIntPoint A(PixelPos - KeyHalfTriSize,	Params.TrackHeight - KeyVertOffset);
 		FIntPoint B(PixelPos + KeyHalfTriSize,	Params.TrackHeight - KeyVertOffset);
@@ -139,7 +139,7 @@ void UInterpTrackMove::DrawTrack( FCanvas* Canvas, UInterpGroup* Group, const FI
 	{
 		float KeyTime = GetKeyframeTime(KeyIndex);
 
-		int32 PixelPos = FMath::Trunc((KeyTime - Params.StartTime) * Params.PixelsPerSec);
+		int32 PixelPos = FMath::TruncToInt((KeyTime - Params.StartTime) * Params.PixelsPerSec);
 
 		FIntPoint A(PixelPos - KeyHalfTriSize,	Params.TrackHeight - KeyVertOffset);
 		FIntPoint B(PixelPos + KeyHalfTriSize,	Params.TrackHeight - KeyVertOffset);
@@ -244,7 +244,7 @@ void UInterpTrackMove::Render3DTrack(UInterpTrackInst* TrInst,
 		if( TotalTime > 0.0f )
 		{
 			// Determine the number of steps to draw.  More steps means a smoother curve
-			int32 NumSteps = FMath::Ceil( TotalTime/DrawTrackTimeRes );
+			int32 NumSteps = FMath::CeilToInt( TotalTime/DrawTrackTimeRes );
 			// Ensure the number of steps to draw wont cause a rendering perf hit.
 			NumSteps = FMath::Min( 100, NumSteps );
 			float DrawSubstep = TotalTime/NumSteps;
@@ -402,7 +402,7 @@ void UInterpTrackMove::Render3DTrack(UInterpTrackInst* TrInst,
 			// If not the first keypoint, draw a line to the last keypoint.
 			if(i>0)
 			{
-				int32 NumSteps = FMath::Ceil( (NewKeyTime - OldKeyTime)/DrawTrackTimeRes );
+				int32 NumSteps = FMath::CeilToInt( (NewKeyTime - OldKeyTime)/DrawTrackTimeRes );
 								// Limit the number of steps to prevent a rendering performance hit
 				NumSteps = FMath::Min( 100, NumSteps );
 				float DrawSubstep = (NewKeyTime - OldKeyTime)/NumSteps;
@@ -596,7 +596,7 @@ void UInterpTrackToggle::DrawTrack( FCanvas* Canvas, UInterpGroup* Group, const
 	for (int32 i=0; i<NumKeys; i++)
 	{
 		float KeyTime = GetKeyframeTime(i);
-		int32 PixelPos = FMath::Trunc((KeyTime - Params.StartTime) * Params.PixelsPerSec);
+		int32 PixelPos = FMath::TruncToInt((KeyTime - Params.StartTime) * Params.PixelsPerSec);
 
 		FToggleTrackKey& Key = ToggleTrack[i];
 		if ((Key.ToggleAction == ETTA_Off) && bLastPosWasOn)
@@ -623,7 +623,7 @@ void UInterpTrackToggle::DrawTrack( FCanvas* Canvas, UInterpGroup* Group, const
 	{
 		float KeyTime = GetKeyframeTime(i);
 
-		int32 PixelPos = FMath::Trunc((KeyTime - Params.StartTime) * Params.PixelsPerSec);
+		int32 PixelPos = FMath::TruncToInt((KeyTime - Params.StartTime) * Params.PixelsPerSec);
 
 		FIntPoint A, A_Offset;
 		FIntPoint B, B_Offset;
@@ -813,7 +813,7 @@ void UInterpTrackEvent::DrawTrack( FCanvas* Canvas, UInterpGroup* Group, const F
 	for(int32 i=0; i<EventTrack.Num(); i++)
 	{
 		float KeyTime = EventTrack[i].Time;
-		int32 PixelPos = FMath::Trunc((KeyTime - Params.StartTime) * Params.PixelsPerSec);
+		int32 PixelPos = FMath::TruncToInt((KeyTime - Params.StartTime) * Params.PixelsPerSec);
 
 		int32 XL, YL;
 		FString Str = EventTrack[i].EventName.ToString();
@@ -872,14 +872,14 @@ void UInterpTrackDirector::DrawTrack( FCanvas* Canvas, UInterpGroup* Group, cons
 		// If its valid, and its not this track, draw a box over duration of shot.
 		if((CutGroupIndex != INDEX_NONE) && (CutTrack[i].TargetCamGroup != Group->GroupName))
 		{
-			int32 StartPixelPos = FMath::Trunc((KeyTime - Params.StartTime) * Params.PixelsPerSec);
-			int32 EndPixelPos = FMath::Trunc((NextKeyTime - Params.StartTime) * Params.PixelsPerSec);
+			int32 StartPixelPos = FMath::TruncToInt((KeyTime - Params.StartTime) * Params.PixelsPerSec);
+			int32 EndPixelPos = FMath::TruncToInt((NextKeyTime - Params.StartTime) * Params.PixelsPerSec);
 
 			if ( bAllowBarSelection )
 			{
 				Canvas->SetHitProxy( new HInterpTrackKeypointProxy( Group, this, i ) );
 			}
-			Canvas->DrawTile( StartPixelPos, KeyVertOffset, EndPixelPos - StartPixelPos, FMath::Trunc(Params.TrackHeight - 2.f*KeyVertOffset), 0.f, 0.f, 1.f, 1.f, Data->InterpGroups[CutGroupIndex]->GroupColor );
+			Canvas->DrawTile( StartPixelPos, KeyVertOffset, EndPixelPos - StartPixelPos, FMath::TruncToFloat(Params.TrackHeight - 2.f*KeyVertOffset), 0.f, 0.f, 1.f, 1.f, Data->InterpGroups[CutGroupIndex]->GroupColor );
 			if ( bAllowBarSelection )
 			{
 				Canvas->SetHitProxy( NULL );
@@ -894,7 +894,7 @@ void UInterpTrackDirector::DrawTrack( FCanvas* Canvas, UInterpGroup* Group, cons
 	for(int32 i=0; i<CutTrack.Num(); i++)
 	{
 		float KeyTime = CutTrack[i].Time;
-		int32 PixelPos = FMath::Trunc((KeyTime - Params.StartTime) * Params.PixelsPerSec);
+		int32 PixelPos = FMath::TruncToInt((KeyTime - Params.StartTime) * Params.PixelsPerSec);
 
 		int32 XL, YL;
 		//Append the shot name to the target group name
@@ -1080,8 +1080,8 @@ void UInterpTrackAnimControl::DrawTrack( FCanvas* Canvas, UInterpGroup* Group, c
 			LoopEndTime = Data->InterpLength;
 		}
 
-		int32 StartPixelPos = FMath::Trunc((SeqStartTime - Params.StartTime) * Params.PixelsPerSec);
-		int32 EndPixelPos = FMath::Trunc((SeqEndTime - Params.StartTime) * Params.PixelsPerSec);
+		int32 StartPixelPos = FMath::TruncToInt((SeqStartTime - Params.StartTime) * Params.PixelsPerSec);
+		int32 EndPixelPos = FMath::TruncToInt((SeqEndTime - Params.StartTime) * Params.PixelsPerSec);
 
 		// Find if this key is one of the selected ones.
 		bool bKeySelected = false;
@@ -1098,14 +1098,14 @@ void UInterpTrackAnimControl::DrawTrack( FCanvas* Canvas, UInterpGroup* Group, c
 
 		if( Seq && CurKey.bLooping )
 		{
-			int32 LoopEndPixelPos = FMath::Ceil((LoopEndTime - Params.StartTime) * Params.PixelsPerSec);
+			int32 LoopEndPixelPos = FMath::CeilToInt((LoopEndTime - Params.StartTime) * Params.PixelsPerSec);
 
 			if ( bAllowBarSelection )
 			{
 				Canvas->SetHitProxy( new HInterpTrackKeypointProxy( Group, this, i ) );
 			}
-			Canvas->DrawTile( StartPixelPos, KeyVertOffset, LoopEndPixelPos - StartPixelPos, FMath::Trunc(Params.TrackHeight - 2.f*KeyVertOffset), 0.f, 0.f, 1.f, 1.f, FColor(0,0,0) );
-			Canvas->DrawTile( StartPixelPos+1, KeyVertOffset+1, LoopEndPixelPos - StartPixelPos - 1, FMath::Trunc(Params.TrackHeight - 2.f*KeyVertOffset) - 2, 0.f, 0.f, 1.f, 1.f, FColor(0,75,150) );
+			Canvas->DrawTile( StartPixelPos, KeyVertOffset, LoopEndPixelPos - StartPixelPos, FMath::TruncToFloat(Params.TrackHeight - 2.f*KeyVertOffset), 0.f, 0.f, 1.f, 1.f, FColor(0,0,0) );
+			Canvas->DrawTile( StartPixelPos+1, KeyVertOffset+1, LoopEndPixelPos - StartPixelPos - 1, FMath::TruncToFloat(Params.TrackHeight - 2.f*KeyVertOffset) - 2, 0.f, 0.f, 1.f, 1.f, FColor(0,75,150) );
 			if ( bAllowBarSelection )
 			{
 				Canvas->SetHitProxy( NULL );
@@ -1116,7 +1116,7 @@ void UInterpTrackAnimControl::DrawTrack( FCanvas* Canvas, UInterpGroup* Group, c
 
 			while(LoopTime < LoopEndTime)
 			{
-				int32 DashPixelPos = FMath::Trunc((LoopTime - Params.StartTime) * Params.PixelsPerSec);
+				int32 DashPixelPos = FMath::TruncToInt((LoopTime - Params.StartTime) * Params.PixelsPerSec);
 				LineItem.Draw( Canvas, FVector2D(DashPixelPos, KeyVertOffset + 2), FVector2D(DashPixelPos, Params.TrackHeight - KeyVertOffset - 2) );
 				LoopTime += SeqLength;
 			}
@@ -1128,7 +1128,7 @@ void UInterpTrackAnimControl::DrawTrack( FCanvas* Canvas, UInterpGroup* Group, c
 		}
 		
 		// Draw background blocks
-		Canvas->DrawTile( StartPixelPos, KeyVertOffset, EndPixelPos - StartPixelPos + 1, FMath::Trunc(Params.TrackHeight - 2.f*KeyVertOffset), 0.f, 0.f, 1.f, 1.f, BorderColor );
+		Canvas->DrawTile( StartPixelPos, KeyVertOffset, EndPixelPos - StartPixelPos + 1, FMath::TruncToFloat(Params.TrackHeight - 2.f*KeyVertOffset), 0.f, 0.f, 1.f, 1.f, BorderColor );
 
 		// If the current key is reversed then change the color of the block.
 		FColor BlockColor;
@@ -1142,7 +1142,7 @@ void UInterpTrackAnimControl::DrawTrack( FCanvas* Canvas, UInterpGroup* Group, c
 			BlockColor = NormalBlockColor;
 		}
 
-		Canvas->DrawTile( StartPixelPos+1, KeyVertOffset+1, EndPixelPos - StartPixelPos - 1, FMath::Trunc(Params.TrackHeight - 2.f*KeyVertOffset) - 2, 0.f, 0.f, 1.f, 1.f,  BlockColor);
+		Canvas->DrawTile( StartPixelPos+1, KeyVertOffset+1, EndPixelPos - StartPixelPos - 1, FMath::TruncToFloat(Params.TrackHeight - 2.f*KeyVertOffset) - 2, 0.f, 0.f, 1.f, 1.f,  BlockColor);
 
 		if ( bAllowBarSelection )
 		{
@@ -1154,11 +1154,11 @@ void UInterpTrackAnimControl::DrawTrack( FCanvas* Canvas, UInterpGroup* Group, c
 		{
 			// Left Edge
 			Canvas->SetHitProxy(new HInterpEdInputInterface(this, FInterpEdInputData(ACDT_AnimBlockLeftEdge, i)));
-			Canvas->DrawTile( StartPixelPos-2, KeyVertOffset, 4, FMath::Trunc(Params.TrackHeight - 2.f*KeyVertOffset), 0.f, 0.f, 1.f, 1.f, FLinearColor::Black );
+			Canvas->DrawTile( StartPixelPos-2, KeyVertOffset, 4, FMath::TruncToFloat(Params.TrackHeight - 2.f*KeyVertOffset), 0.f, 0.f, 1.f, 1.f, FLinearColor::Black );
 
 			// Right Edge
 			Canvas->SetHitProxy(new HInterpEdInputInterface(this, FInterpEdInputData(ACDT_AnimBlockRightEdge, i)));
-			Canvas->DrawTile( EndPixelPos-1, KeyVertOffset, 4, FMath::Trunc(Params.TrackHeight - 2.f*KeyVertOffset), 0.f, 0.f, 1.f, 1.f, FLinearColor::Black );
+			Canvas->DrawTile( EndPixelPos-1, KeyVertOffset, 4, FMath::TruncToFloat(Params.TrackHeight - 2.f*KeyVertOffset), 0.f, 0.f, 1.f, 1.f, FLinearColor::Black );
 
 			Canvas->SetHitProxy(NULL);
 		}
@@ -1177,7 +1177,7 @@ void UInterpTrackAnimControl::DrawTrack( FCanvas* Canvas, UInterpGroup* Group, c
 				if( Params.bPreferFrameNumbers && Params.SnapAmount > KINDA_SMALL_NUMBER )
 				{
 					// Convert to the animation time values to frame numbers
-					const int32 CursorFrameWithinAnim = FMath::Trunc( CursorPosWithinAnim / Params.SnapAmount );
+					const int32 CursorFrameWithinAnim = FMath::TruncToInt( CursorPosWithinAnim / Params.SnapAmount );
 					TimeCursorString = FString::Printf( TEXT( "%i" ), CursorFrameWithinAnim );
 				}
 				else
@@ -1206,7 +1206,7 @@ void UInterpTrackAnimControl::DrawTrack( FCanvas* Canvas, UInterpGroup* Group, c
 		}
 
 		float SeqStartTime = CurKey.StartTime;
-		int32 PixelPos = FMath::Trunc((SeqStartTime - Params.StartTime) * Params.PixelsPerSec);
+		int32 PixelPos = FMath::TruncToInt((SeqStartTime - Params.StartTime) * Params.PixelsPerSec);
 
 		UAnimSequence* Seq = CurKey.AnimSeq;
 		FString SeqString = Seq != NULL ? Seq->GetName() : FString::Printf(TEXT("NULL"));
@@ -1219,8 +1219,8 @@ void UInterpTrackAnimControl::DrawTrack( FCanvas* Canvas, UInterpGroup* Group, c
 				if( Params.bPreferFrameNumbers && Params.SnapAmount > KINDA_SMALL_NUMBER )
 				{
 					// Convert to the animation time values to frame numbers
-					const int32 AnimFrameOffsetFromStart = FMath::Round( CurKey.AnimStartOffset / Params.SnapAmount );
-					const int32 AnimFrameOffsetFromEnd = FMath::Round( ( Seq->SequenceLength - CurKey.AnimEndOffset ) / Params.SnapAmount );
+					const int32 AnimFrameOffsetFromStart = FMath::RoundToInt( CurKey.AnimStartOffset / Params.SnapAmount );
+					const int32 AnimFrameOffsetFromEnd = FMath::RoundToInt( ( Seq->SequenceLength - CurKey.AnimEndOffset ) / Params.SnapAmount );
 					SeqString += FString::Printf( TEXT(" (%i->%i)"), AnimFrameOffsetFromStart, AnimFrameOffsetFromEnd );
 				}
 				else
@@ -1269,7 +1269,7 @@ void UInterpTrackAnimControl::DrawTrack( FCanvas* Canvas, UInterpGroup* Group, c
 
 
 		const int32 TimeCursorPixelPos =
-			FMath::Trunc( ( Params.TimeCursorPosition - Params.StartTime ) * Params.PixelsPerSec );
+			FMath::TruncToInt( ( Params.TimeCursorPosition - Params.StartTime ) * Params.PixelsPerSec );
 
 		int32 XL, YL;
 		StringSize( GEngine->GetTinyFont(), XL, YL, *TimeCursorString );
@@ -1356,8 +1356,8 @@ void UInterpTrackSound::DrawTrack( FCanvas* Canvas, UInterpGroup* Group, const F
 			}
 		}
 
-		int32 StartPixelPos = FMath::Trunc((SoundStartTime - Params.StartTime) * Params.PixelsPerSec);
-		int32 EndPixelPos = FMath::Trunc((SoundEndTime - Params.StartTime) * Params.PixelsPerSec);
+		int32 StartPixelPos = FMath::TruncToInt((SoundStartTime - Params.StartTime) * Params.PixelsPerSec);
+		int32 EndPixelPos = FMath::TruncToInt((SoundEndTime - Params.StartTime) * Params.PixelsPerSec);
 
 		// Find if this sound is one of the selected ones.
 		bool bKeySelected = false;
@@ -1376,8 +1376,8 @@ void UInterpTrackSound::DrawTrack( FCanvas* Canvas, UInterpGroup* Group, const F
 		{
 			Canvas->SetHitProxy( new HInterpTrackKeypointProxy( Group, this, i ) );
 		}
-		Canvas->DrawTile( StartPixelPos, KeyVertOffset, EndPixelPos - StartPixelPos + 1, FMath::Trunc(Params.TrackHeight - 2.f*KeyVertOffset), 0.f, 0.f, 1.f, 1.f, BorderColor );
-		Canvas->DrawTile( StartPixelPos+1, KeyVertOffset+1, EndPixelPos - StartPixelPos - 1, FMath::Trunc(Params.TrackHeight - 2.f*KeyVertOffset) - 2, 0.f, 0.f, 1.f, 1.f, FColor(0,200,100) );
+		Canvas->DrawTile( StartPixelPos, KeyVertOffset, EndPixelPos - StartPixelPos + 1, FMath::TruncToFloat(Params.TrackHeight - 2.f*KeyVertOffset), 0.f, 0.f, 1.f, 1.f, BorderColor );
+		Canvas->DrawTile( StartPixelPos+1, KeyVertOffset+1, EndPixelPos - StartPixelPos - 1, FMath::TruncToFloat(Params.TrackHeight - 2.f*KeyVertOffset) - 2, 0.f, 0.f, 1.f, 1.f, FColor(0,200,100) );
 		if( bAllowBarSelection )
 		{
 			Canvas->SetHitProxy( NULL );
@@ -1391,7 +1391,7 @@ void UInterpTrackSound::DrawTrack( FCanvas* Canvas, UInterpGroup* Group, const F
 	for (int32 i = 0; i < Sounds.Num(); i++)
 	{
 		float SoundStartTime = Sounds[i].Time;
-		int32 PixelPos = FMath::Trunc((SoundStartTime - Params.StartTime) * Params.PixelsPerSec);
+		int32 PixelPos = FMath::TruncToInt((SoundStartTime - Params.StartTime) * Params.PixelsPerSec);
 
 		USoundBase* Sound = Sounds[i].Sound;
 	
@@ -1492,7 +1492,7 @@ void UInterpTrackVisibility::DrawTrack( FCanvas* Canvas, UInterpGroup* Group, co
 	for (int32 i=0; i<NumKeys; i++)
 	{
 		float KeyTime = GetKeyframeTime(i);
-		int32 PixelPos = FMath::Trunc((KeyTime - Params.StartTime) * Params.PixelsPerSec);
+		int32 PixelPos = FMath::TruncToInt((KeyTime - Params.StartTime) * Params.PixelsPerSec);
 
 		FVisibilityTrackKey& Key = VisibilityTrack[i];
 		if ((Key.Action == EVTA_Hide) && bLastPosWasOn)
@@ -1519,7 +1519,7 @@ void UInterpTrackVisibility::DrawTrack( FCanvas* Canvas, UInterpGroup* Group, co
 	{
 		float KeyTime = GetKeyframeTime(i);
 
-		int32 PixelPos = FMath::Trunc((KeyTime - Params.StartTime) * Params.PixelsPerSec);
+		int32 PixelPos = FMath::TruncToInt((KeyTime - Params.StartTime) * Params.PixelsPerSec);
 
 		FIntPoint A, A_Offset;
 		FIntPoint B, B_Offset;
@@ -1674,7 +1674,7 @@ void UInterpTrackParticleReplay::ObjectDragged(FInterpEdInputData& InputData)
 						// Snap the new time position
 						if( FixedTimeStep > SMALL_NUMBER )
 						{
-							const int32 InterpPositionInFrames = FMath::Round( SelectedKey.Time / FixedTimeStep );
+							const int32 InterpPositionInFrames = FMath::RoundToInt( SelectedKey.Time / FixedTimeStep );
 							SelectedKey.Time = InterpPositionInFrames * FixedTimeStep;
 						}
 
@@ -1705,7 +1705,7 @@ void UInterpTrackParticleReplay::ObjectDragged(FInterpEdInputData& InputData)
 						{
 							float EndTime = SelectedKey.Time + SelectedKey.Duration;
 
-							const int32 InterpPositionInFrames = FMath::Round( EndTime / FixedTimeStep );
+							const int32 InterpPositionInFrames = FMath::RoundToInt( EndTime / FixedTimeStep );
 							EndTime = InterpPositionInFrames * FixedTimeStep;
 
 							if( EndTime < SelectedKey.Time )
@@ -1740,8 +1740,8 @@ void UInterpTrackParticleReplay::DrawTrack( FCanvas* Canvas, UInterpGroup* Group
 
 		float KeyTime = GetKeyframeTime( CurKeyIndex );
 
-		const int32 StartPixelPos = FMath::Trunc( ( KeyTime - Params.StartTime ) * Params.PixelsPerSec );
-		const int32 EndPixelPos = FMath::Trunc( ( KeyTime - Params.StartTime + CurKey.Duration ) * Params.PixelsPerSec );
+		const int32 StartPixelPos = FMath::TruncToInt( ( KeyTime - Params.StartTime ) * Params.PixelsPerSec );
+		const int32 EndPixelPos = FMath::TruncToInt( ( KeyTime - Params.StartTime + CurKey.Duration ) * Params.PixelsPerSec );
 
 
 		// Is this key selected?
@@ -1775,7 +1775,7 @@ void UInterpTrackParticleReplay::DrawTrack( FCanvas* Canvas, UInterpGroup* Group
 				StartPixelPos,
 				KeyVertOffset,
 				EndPixelPos - StartPixelPos + 1,
-				FMath::Trunc( Params.TrackHeight - 2.0f * KeyVertOffset ),
+				FMath::TruncToFloat( Params.TrackHeight - 2.0f * KeyVertOffset ),
 				0.f, 0.f, 1.f, 1.f,
 				BackgroundTileColor );
 
@@ -1791,11 +1791,11 @@ void UInterpTrackParticleReplay::DrawTrack( FCanvas* Canvas, UInterpGroup* Group
 		{
 			// Left Edge
 			Canvas->SetHitProxy(new HInterpEdInputInterface(this, FInterpEdInputData(PRDT_LeftEdge, CurKeyIndex)));
-			Canvas->DrawTile( StartPixelPos-2, KeyVertOffset, 4, FMath::Trunc(Params.TrackHeight - 2.f*KeyVertOffset), 0.f, 0.f, 1.f, 1.f, FLinearColor::Black );
+			Canvas->DrawTile( StartPixelPos-2, KeyVertOffset, 4, FMath::TruncToFloat(Params.TrackHeight - 2.f*KeyVertOffset), 0.f, 0.f, 1.f, 1.f, FLinearColor::Black );
 
 			// Right Edge
 			Canvas->SetHitProxy(new HInterpEdInputInterface(this, FInterpEdInputData(PRDT_RightEdge, CurKeyIndex)));
-			Canvas->DrawTile( EndPixelPos-1, KeyVertOffset, 4, FMath::Trunc(Params.TrackHeight - 2.f*KeyVertOffset), 0.f, 0.f, 1.f, 1.f, FLinearColor::Black );
+			Canvas->DrawTile( EndPixelPos-1, KeyVertOffset, 4, FMath::TruncToFloat(Params.TrackHeight - 2.f*KeyVertOffset), 0.f, 0.f, 1.f, 1.f, FLinearColor::Black );
 
 			Canvas->SetHitProxy(NULL);
 		}
@@ -1803,7 +1803,7 @@ void UInterpTrackParticleReplay::DrawTrack( FCanvas* Canvas, UInterpGroup* Group
 
 		// Draw key frame information text
 		{
-			const int32 DurationInFrames = FMath::Round( CurKey.Duration / FixedTimeStep );
+			const int32 DurationInFrames = FMath::RoundToInt( CurKey.Duration / FixedTimeStep );
 			FString InfoText = FString::Printf( TEXT( "[Clip %i] %i frames (%.2f s)" ), CurKey.ClipIDNumber, DurationInFrames, CurKey.Duration );
 
 			int32 XL, YL;
diff --git a/Engine/Source/Runtime/Engine/Private/KismetMathLibrary.cpp b/Engine/Source/Runtime/Engine/Private/KismetMathLibrary.cpp
--- a/Engine/Source/Runtime/Engine/Private/KismetMathLibrary.cpp
+++ b/Engine/Source/Runtime/Engine/Private/KismetMathLibrary.cpp
@@ -421,17 +421,17 @@ float UKismetMathLibrary::Square(float A)
 
 int32 UKismetMathLibrary::Round(float A)
 {
-	return FMath::Round(A);
+	return FMath::RoundToInt(A);
 }	
 
 int32 UKismetMathLibrary::FFloor(float A)
 {
-	return FMath::Floor(A);
+	return FMath::FloorToInt(A);
 }	
 
 int32 UKismetMathLibrary::FCeil(float A)
 {
-	return FMath::Ceil(A);
+	return FMath::CeilToInt(A);
 }	
 
 int32 UKismetMathLibrary::FMod(float Dividend, float Divisor, float& Remainder)
@@ -440,7 +440,7 @@ int32 UKismetMathLibrary::FMod(float Dividend, float Divisor, float& Remainder)
 	if( Divisor != 0.f )
 	{
 		const float Quotient = Dividend / Divisor;
-		Result = (Quotient < 0.f ? -1 : 1) * FMath::Floor( FMath::Abs(Quotient) );
+		Result = (Quotient < 0.f ? -1 : 1) * FMath::FloorToInt( FMath::Abs(Quotient) );
 		Remainder = FMath::Fmod(Dividend, Divisor);
 	}
 	else
diff --git a/Engine/Source/Runtime/Engine/Private/Landscape/Landscape.cpp b/Engine/Source/Runtime/Engine/Private/Landscape/Landscape.cpp
--- a/Engine/Source/Runtime/Engine/Private/Landscape/Landscape.cpp
+++ b/Engine/Source/Runtime/Engine/Private/Landscape/Landscape.cpp
@@ -1744,8 +1744,8 @@ void ULandscapeInfo::RecreateLandscapeInfo(UWorld* InWorld, bool bMapCheck)
 				FVector	DrawScale = FirstVisibleLandscape->GetRootComponent()->RelativeScale3D;
 				
 				FIntPoint QuadsSpaceOffset; 
-				QuadsSpaceOffset.X = FMath::Round(Offset.X/DrawScale.X);
-				QuadsSpaceOffset.Y = FMath::Round(Offset.Y/DrawScale.Y);
+				QuadsSpaceOffset.X = FMath::RoundToInt(Offset.X/DrawScale.X);
+				QuadsSpaceOffset.Y = FMath::RoundToInt(Offset.Y/DrawScale.Y);
 				Proxy->SetAbsoluteSectionBase(QuadsSpaceOffset);
 			}
 
diff --git a/Engine/Source/Runtime/Engine/Private/Landscape/LandscapeEdit.cpp b/Engine/Source/Runtime/Engine/Private/Landscape/LandscapeEdit.cpp
--- a/Engine/Source/Runtime/Engine/Private/Landscape/LandscapeEdit.cpp
+++ b/Engine/Source/Runtime/Engine/Private/Landscape/LandscapeEdit.cpp
@@ -580,11 +580,11 @@ void ULandscapeComponent::UpdateCollisionHeightData(const FColor* HeightmapTextu
 	float CollisionQuadRatio = (float)CollisionSubsectionSizeQuads / (float)SubsectionSizeQuads;
 
 	// XY offset into heightmap mip data
-	int32 HeightmapOffsetX = FMath::Round(HeightmapScaleBias.Z * (float)HeightmapSizeU) >> CollisionMipLevel;
-	int32 HeightmapOffsetY = FMath::Round(HeightmapScaleBias.W * (float)HeightmapSizeV) >> CollisionMipLevel;
+	int32 HeightmapOffsetX = FMath::RoundToInt(HeightmapScaleBias.Z * (float)HeightmapSizeU) >> CollisionMipLevel;
+	int32 HeightmapOffsetY = FMath::RoundToInt(HeightmapScaleBias.W * (float)HeightmapSizeV) >> CollisionMipLevel;
 
-	//int32 WeightmapOffsetX = FMath::Round(WeightmapScaleBias.Z * (float)XYMipSizeU) >> CollisionMipLevel;
-	//int32 WeightmapOffsetY = FMath::Round(WeightmapScaleBias.W * (float)XYMipSizeV) >> CollisionMipLevel;
+	//int32 WeightmapOffsetX = FMath::RoundToInt(WeightmapScaleBias.Z * (float)XYMipSizeU) >> CollisionMipLevel;
+	//int32 WeightmapOffsetY = FMath::RoundToInt(WeightmapScaleBias.W * (float)XYMipSizeV) >> CollisionMipLevel;
 
 	for( int32 SubsectionY = 0;SubsectionY < NumSubsections;SubsectionY++ )
 	{
@@ -611,10 +611,10 @@ void ULandscapeComponent::UpdateCollisionHeightData(const FColor* HeightmapTextu
 			int32 SubY2 = ComponentY2 - SubsectionSizeQuads*SubsectionY;
 
 			// Area to update in collision mip level coords
-			int32 CollisionSubX1 = FMath::Floor( (float)SubX1 * CollisionQuadRatio );
-			int32 CollisionSubY1 = FMath::Floor( (float)SubY1 * CollisionQuadRatio );
-			int32 CollisionSubX2 = FMath::Ceil(  (float)SubX2 * CollisionQuadRatio );
-			int32 CollisionSubY2 = FMath::Ceil(  (float)SubY2 * CollisionQuadRatio );
+			int32 CollisionSubX1 = FMath::FloorToInt( (float)SubX1 * CollisionQuadRatio );
+			int32 CollisionSubY1 = FMath::FloorToInt( (float)SubY1 * CollisionQuadRatio );
+			int32 CollisionSubX2 = FMath::CeilToInt(  (float)SubX2 * CollisionQuadRatio );
+			int32 CollisionSubY2 = FMath::CeilToInt(  (float)SubY2 * CollisionQuadRatio );
 
 			// Clamp area to update
 			int32 VertX1 = FMath::Clamp<int32>(CollisionSubX1, 0, CollisionSubsectionSizeQuads);
@@ -688,10 +688,10 @@ void ULandscapeComponent::UpdateCollisionHeightData(const FColor* HeightmapTextu
 		}
 		else
 		{
-			int32 CollisionCompX1 = FMath::Floor( (float)ComponentX1 * CollisionQuadRatio );
-			int32 CollisionCompY1 = FMath::Floor( (float)ComponentY1 * CollisionQuadRatio );
-			int32 CollisionCompX2 = FMath::Ceil(  (float)ComponentX2 * CollisionQuadRatio );
-			int32 CollisionCompY2 = FMath::Ceil(  (float)ComponentY2 * CollisionQuadRatio );
+			int32 CollisionCompX1 = FMath::FloorToInt( (float)ComponentX1 * CollisionQuadRatio );
+			int32 CollisionCompY1 = FMath::FloorToInt( (float)ComponentY1 * CollisionQuadRatio );
+			int32 CollisionCompX2 = FMath::CeilToInt(  (float)ComponentX2 * CollisionQuadRatio );
+			int32 CollisionCompY2 = FMath::CeilToInt(  (float)ComponentY2 * CollisionQuadRatio );
 			CollisionComp->UpdateHeightfieldRegion(CollisionCompX1, CollisionCompY1, CollisionCompX2, CollisionCompY2);
 		}
 	}
@@ -823,10 +823,10 @@ void ULandscapeComponent::UpdateCollisionLayerData(TArray<FColor*>& WeightmapTex
 					int32 SubY2 = ComponentY2 - SubsectionSizeQuads*SubsectionY;
 
 					// Area to update in collision mip level coords
-					int32 CollisionSubX1 = FMath::Floor( (float)SubX1 * CollisionQuadRatio );
-					int32 CollisionSubY1 = FMath::Floor( (float)SubY1 * CollisionQuadRatio );
-					int32 CollisionSubX2 = FMath::Ceil(  (float)SubX2 * CollisionQuadRatio );
-					int32 CollisionSubY2 = FMath::Ceil(  (float)SubY2 * CollisionQuadRatio );
+					int32 CollisionSubX1 = FMath::FloorToInt( (float)SubX1 * CollisionQuadRatio );
+					int32 CollisionSubY1 = FMath::FloorToInt( (float)SubY1 * CollisionQuadRatio );
+					int32 CollisionSubX2 = FMath::CeilToInt(  (float)SubX2 * CollisionQuadRatio );
+					int32 CollisionSubY2 = FMath::CeilToInt(  (float)SubY2 * CollisionQuadRatio );
 
 					// Clamp area to update
 					int32 VertX1 = FMath::Clamp<int32>(CollisionSubX1, 0, CollisionSubsectionSizeQuads);
@@ -933,8 +933,8 @@ void ULandscapeComponent::GenerateHeightmapMips( TArray<FColor*>& HeightmapTextu
 	int32 HeightmapSizeU = HeightmapTexture->Source.GetSizeX();
 	int32 HeightmapSizeV = HeightmapTexture->Source.GetSizeY();
 
-	int32 HeightmapOffsetX = FMath::Round(HeightmapScaleBias.Z * (float)HeightmapSizeU);
-	int32 HeightmapOffsetY = FMath::Round(HeightmapScaleBias.W * (float)HeightmapSizeV);
+	int32 HeightmapOffsetX = FMath::RoundToInt(HeightmapScaleBias.Z * (float)HeightmapSizeU);
+	int32 HeightmapOffsetY = FMath::RoundToInt(HeightmapScaleBias.W * (float)HeightmapSizeV);
 
 	for( int32 SubsectionY = 0;SubsectionY < NumSubsections;SubsectionY++ )
 	{
@@ -981,10 +981,10 @@ void ULandscapeComponent::GenerateHeightmapMips( TArray<FColor*>& HeightmapTextu
 				int32 MipHeightmapOffsetY = HeightmapOffsetY>>Mip;
 
 				// Area to update in current mip level coords
-				int32 MipSubX1 = FMath::Floor( (float)MipSubsectionSizeQuads * (float)PrevMipSubX1 * InvPrevMipSubsectionSizeQuads );
-				int32 MipSubY1 = FMath::Floor( (float)MipSubsectionSizeQuads * (float)PrevMipSubY1 * InvPrevMipSubsectionSizeQuads );
-				int32 MipSubX2 = FMath::Ceil(  (float)MipSubsectionSizeQuads * (float)PrevMipSubX2 * InvPrevMipSubsectionSizeQuads );
-				int32 MipSubY2 = FMath::Ceil(  (float)MipSubsectionSizeQuads * (float)PrevMipSubY2 * InvPrevMipSubsectionSizeQuads );
+				int32 MipSubX1 = FMath::FloorToInt( (float)MipSubsectionSizeQuads * (float)PrevMipSubX1 * InvPrevMipSubsectionSizeQuads );
+				int32 MipSubY1 = FMath::FloorToInt( (float)MipSubsectionSizeQuads * (float)PrevMipSubY1 * InvPrevMipSubsectionSizeQuads );
+				int32 MipSubX2 = FMath::CeilToInt(  (float)MipSubsectionSizeQuads * (float)PrevMipSubX2 * InvPrevMipSubsectionSizeQuads );
+				int32 MipSubY2 = FMath::CeilToInt(  (float)MipSubsectionSizeQuads * (float)PrevMipSubY2 * InvPrevMipSubsectionSizeQuads );
 
 				// Clamp area to update
 				int32 VertX1 = FMath::Clamp<int32>(MipSubX1, 0, MipSubsectionSizeQuads);
@@ -1006,13 +1006,13 @@ void ULandscapeComponent::GenerateHeightmapMips( TArray<FColor*>& HeightmapTextu
 						// above to process all verts from 0 .. MipSubsectionSizeQuads.
 						if( VertX < VertX1 || VertX > VertX2 )
 						{
-							check( FMath::Ceil(PrevMipVertX) < PrevMipSubX1 || FMath::Floor(PrevMipVertX) > PrevMipSubX2 );
+							check( FMath::CeilToInt(PrevMipVertX) < PrevMipSubX1 || FMath::FloorToInt(PrevMipVertX) > PrevMipSubX2 );
 							continue;
 						}
 
 						if( VertY < VertY1 || VertY > VertY2 )
 						{
-							check( FMath::Ceil(PrevMipVertY) < PrevMipSubY1 || FMath::Floor(PrevMipVertY) > PrevMipSubY2 );
+							check( FMath::CeilToInt(PrevMipVertY) < PrevMipSubY1 || FMath::FloorToInt(PrevMipVertY) > PrevMipSubY2 );
 							continue;
 						}
 #endif
@@ -1024,9 +1024,9 @@ void ULandscapeComponent::GenerateHeightmapMips( TArray<FColor*>& HeightmapTextu
 						float fPrevMipTexX = (float)(PrevMipHeightmapOffsetX) + (float)((PrevMipSubsectionSizeQuads+1) * SubsectionX) + PrevMipVertX;
 						float fPrevMipTexY = (float)(PrevMipHeightmapOffsetY) + (float)((PrevMipSubsectionSizeQuads+1) * SubsectionY) + PrevMipVertY;
 
-						int32 PrevMipTexX = FMath::Floor(fPrevMipTexX);
+						int32 PrevMipTexX = FMath::FloorToInt(fPrevMipTexX);
 						float fPrevMipTexFracX = FMath::Fractional(fPrevMipTexX);
-						int32 PrevMipTexY = FMath::Floor(fPrevMipTexY);
+						int32 PrevMipTexY = FMath::FloorToInt(fPrevMipTexY);
 						float fPrevMipTexFracY = FMath::Fractional(fPrevMipTexY);
 
 						checkSlow( TexX >= 0 && TexX < MipSizeU );
@@ -1071,7 +1071,7 @@ void ULandscapeComponent::GenerateHeightmapMips( TArray<FColor*>& HeightmapTextu
 						uint16 PrevMipHeightValue01 = PreMipTexData01->R << 8 | PreMipTexData01->G;
 						uint16 PrevMipHeightValue10 = PreMipTexData10->R << 8 | PreMipTexData10->G;
 						uint16 PrevMipHeightValue11 = PreMipTexData11->R << 8 | PreMipTexData11->G;
-						uint16 HeightValue = FMath::Round(
+						uint16 HeightValue = FMath::RoundToInt(
 							FMath::Lerp(
 							FMath::Lerp( (float)PrevMipHeightValue00, (float)PrevMipHeightValue10, fPrevMipTexFracX),
 							FMath::Lerp( (float)PrevMipHeightValue01, (float)PrevMipHeightValue11, fPrevMipTexFracX),
@@ -1081,13 +1081,13 @@ void ULandscapeComponent::GenerateHeightmapMips( TArray<FColor*>& HeightmapTextu
 						TexData->G = HeightValue & 255;
 
 						// Lerp tangents
-						TexData->B = FMath::Round(
+						TexData->B = FMath::RoundToInt(
 							FMath::Lerp(
 							FMath::Lerp( (float)PreMipTexData00->B, (float)PreMipTexData10->B, fPrevMipTexFracX),
 							FMath::Lerp( (float)PreMipTexData01->B, (float)PreMipTexData11->B, fPrevMipTexFracX),
 							fPrevMipTexFracY) );
 
-						TexData->A = FMath::Round(
+						TexData->A = FMath::RoundToInt(
 							FMath::Lerp(
 							FMath::Lerp( (float)PreMipTexData00->A, (float)PreMipTexData10->A, fPrevMipTexFracX),
 							FMath::Lerp( (float)PreMipTexData01->A, (float)PreMipTexData11->A, fPrevMipTexFracX),
@@ -1206,7 +1206,7 @@ namespace
 	template<typename DataType>
 	void BiLerpTextureData(DataType* Output, const DataType* Data00, const DataType* Data10, const DataType* Data01, const DataType* Data11, float FracX, float FracY)
 	{
-		*Output = FMath::Round(
+		*Output = FMath::RoundToInt(
 			FMath::Lerp(
 			FMath::Lerp( (float)*Data00, (float)*Data10, FracX),
 			FMath::Lerp( (float)*Data01, (float)*Data11, FracX),
@@ -1216,22 +1216,22 @@ namespace
 	template<>
 	void BiLerpTextureData(FColor* Output, const FColor* Data00, const FColor* Data10, const FColor* Data01, const FColor* Data11, float FracX, float FracY)
 	{
-		Output->R = FMath::Round(
+		Output->R = FMath::RoundToInt(
 			FMath::Lerp(
 			FMath::Lerp( (float)Data00->R, (float)Data10->R, FracX),
 			FMath::Lerp( (float)Data01->R, (float)Data11->R, FracX),
 			FracY) );
-		Output->G = FMath::Round(
+		Output->G = FMath::RoundToInt(
 			FMath::Lerp(
 			FMath::Lerp( (float)Data00->G, (float)Data10->G, FracX),
 			FMath::Lerp( (float)Data01->G, (float)Data11->G, FracX),
 			FracY) );
-		Output->B = FMath::Round(
+		Output->B = FMath::RoundToInt(
 			FMath::Lerp(
 			FMath::Lerp( (float)Data00->B, (float)Data10->B, FracX),
 			FMath::Lerp( (float)Data01->B, (float)Data11->B, FracX),
 			FracY) );
-		Output->A = FMath::Round(
+		Output->A = FMath::RoundToInt(
 			FMath::Lerp(
 			FMath::Lerp( (float)Data00->A, (float)Data10->A, FracX),
 			FMath::Lerp( (float)Data01->A, (float)Data11->A, FracX),
@@ -1321,10 +1321,10 @@ void ULandscapeComponent::UpdateMipsTempl(int32 InNumSubsections, int32 InSubsec
 				float InvMipSubsectionSizeQuads = 1.f / (float)MipSubsectionSizeQuads;
 
 				// Area to update in current mip level coords
-				int32 MipSubX1 = FMath::Floor( (float)MipSubsectionSizeQuads * (float)PrevMipSubX1 * InvPrevMipSubsectionSizeQuads );
-				int32 MipSubY1 = FMath::Floor( (float)MipSubsectionSizeQuads * (float)PrevMipSubY1 * InvPrevMipSubsectionSizeQuads );
-				int32 MipSubX2 = FMath::Ceil(  (float)MipSubsectionSizeQuads * (float)PrevMipSubX2 * InvPrevMipSubsectionSizeQuads );
-				int32 MipSubY2 = FMath::Ceil(  (float)MipSubsectionSizeQuads * (float)PrevMipSubY2 * InvPrevMipSubsectionSizeQuads );
+				int32 MipSubX1 = FMath::FloorToInt( (float)MipSubsectionSizeQuads * (float)PrevMipSubX1 * InvPrevMipSubsectionSizeQuads );
+				int32 MipSubY1 = FMath::FloorToInt( (float)MipSubsectionSizeQuads * (float)PrevMipSubY1 * InvPrevMipSubsectionSizeQuads );
+				int32 MipSubX2 = FMath::CeilToInt(  (float)MipSubsectionSizeQuads * (float)PrevMipSubX2 * InvPrevMipSubsectionSizeQuads );
+				int32 MipSubY2 = FMath::CeilToInt(  (float)MipSubsectionSizeQuads * (float)PrevMipSubY2 * InvPrevMipSubsectionSizeQuads );
 
 				// Clamp area to update
 				int32 VertX1 = FMath::Clamp<int32>(MipSubX1, 0, MipSubsectionSizeQuads);
@@ -1347,9 +1347,9 @@ void ULandscapeComponent::UpdateMipsTempl(int32 InNumSubsections, int32 InSubsec
 						float fPrevMipTexX = (float)((PrevMipSubsectionSizeQuads+1) * SubsectionX) + PrevMipVertX;
 						float fPrevMipTexY = (float)((PrevMipSubsectionSizeQuads+1) * SubsectionY) + PrevMipVertY;
 
-						int32 PrevMipTexX = FMath::Floor(fPrevMipTexX);
+						int32 PrevMipTexX = FMath::FloorToInt(fPrevMipTexX);
 						float fPrevMipTexFracX = FMath::Fractional(fPrevMipTexX);
-						int32 PrevMipTexY = FMath::Floor(fPrevMipTexY);
+						int32 PrevMipTexY = FMath::FloorToInt(fPrevMipTexY);
 						float fPrevMipTexFracY = FMath::Fractional(fPrevMipTexY);
 
 						check( TexX >= 0 && TexX < MipSizeU );
@@ -1479,10 +1479,10 @@ float ULandscapeComponent::GetLayerWeightAtLocation( const FVector& InLocation,
 	float TestY = (InLocation.Y-Landscape->GetActorLocation().Y) / DrawScale.Y - (float)GetSectionBase().Y;
 
 	// Find data
-	int32 X1 = FMath::Floor(TestX);
-	int32 Y1 = FMath::Floor(TestY);
-	int32 X2 = FMath::Ceil(TestX);
-	int32 Y2 = FMath::Ceil(TestY);
+	int32 X1 = FMath::FloorToInt(TestX);
+	int32 Y1 = FMath::FloorToInt(TestY);
+	int32 X2 = FMath::CeilToInt(TestX);
+	int32 Y2 = FMath::CeilToInt(TestY);
 
 	int32 Stride = (SubsectionSizeQuads+1) * NumSubsections;
 
@@ -2151,8 +2151,8 @@ void ALandscapeProxy::Import(FGuid Guid, int32 VertsX, int32 VertsY,
 
 							HeightmapInfo.HeightmapTextureMipData[0][HeightTexDataIdx].R = HeightValue >> 8;
 							HeightmapInfo.HeightmapTextureMipData[0][HeightTexDataIdx].G = HeightValue & 255;
-							HeightmapInfo.HeightmapTextureMipData[0][HeightTexDataIdx].B = FMath::Round( 127.5f * (Normal.X + 1.f) );
-							HeightmapInfo.HeightmapTextureMipData[0][HeightTexDataIdx].A = FMath::Round( 127.5f * (Normal.Y + 1.f) );
+							HeightmapInfo.HeightmapTextureMipData[0][HeightTexDataIdx].B = FMath::RoundToInt( 127.5f * (Normal.X + 1.f) );
+							HeightmapInfo.HeightmapTextureMipData[0][HeightTexDataIdx].A = FMath::RoundToInt( 127.5f * (Normal.Y + 1.f) );
 
 							for( int32 WeightmapIndex=0;WeightmapIndex<WeightValues.Num(); WeightmapIndex++ )
 							{
@@ -2482,10 +2482,10 @@ FVector ULandscapeInfo::GetLandscapeCenterPos(float& LengthZ, int32 MinX /*= MAX
 
 		const int32 Dist = (ComponentSizeQuads+1) >> 1; // Should be same in ALandscapeGizmoActiveActor::SetTargetLandscape
 		FVector2D MidPoint(((float)(MinX+MaxX))/2.f, ((float)(MinY+MaxY))/2.f);
-		MinX = FMath::Floor(MidPoint.X) - Dist;
-		MaxX = FMath::Ceil(MidPoint.X) + Dist;
-		MinY = FMath::Floor(MidPoint.Y) - Dist;
-		MaxY = FMath::Ceil(MidPoint.Y) + Dist;
+		MinX = FMath::FloorToInt(MidPoint.X) - Dist;
+		MaxX = FMath::CeilToInt(MidPoint.X) + Dist;
+		MinY = FMath::FloorToInt(MidPoint.Y) - Dist;
+		MaxY = FMath::CeilToInt(MidPoint.Y) + Dist;
 		check(MidPoint.X == ((float)(MinX+MaxX))/2.f && MidPoint.Y == ((float)(MinY+MaxY))/2.f );
 	}
 
@@ -2838,10 +2838,10 @@ bool ULandscapeInfo::ApplySplinesInternal(bool bOnlySelected, ALandscapeProxy* L
 		FBox LandscapeBounds = ControlPoint->GetBounds();
 		LandscapeBounds = LandscapeBounds.TransformBy(SplineToLandscape.ToMatrixWithScale());
 
-		int32 MinX = FMath::Ceil(LandscapeBounds.Min.X);
-		int32 MinY = FMath::Ceil(LandscapeBounds.Min.Y);
-		int32 MaxX = FMath::Floor(LandscapeBounds.Max.X);
-		int32 MaxY = FMath::Floor(LandscapeBounds.Max.Y);
+		int32 MinX = FMath::CeilToInt(LandscapeBounds.Min.X);
+		int32 MinY = FMath::CeilToInt(LandscapeBounds.Min.Y);
+		int32 MaxX = FMath::FloorToInt(LandscapeBounds.Max.X);
+		int32 MaxY = FMath::FloorToInt(LandscapeBounds.Max.Y);
 
 		FLandscapeEditDataInterface LandscapeEdit(this);
 		TArray<FLandscapeSplineInterpPoint> Points = ControlPoint->GetPoints();
@@ -2991,10 +2991,10 @@ bool ULandscapeInfo::ApplySplinesInternal(bool bOnlySelected, ALandscapeProxy* L
 		FBox LandscapeBounds = Segment->GetBounds();
 		LandscapeBounds = LandscapeBounds.TransformBy(SplineToLandscape.ToMatrixWithScale());
 
-		int32 MinX = FMath::Ceil(LandscapeBounds.Min.X);
-		int32 MinY = FMath::Ceil(LandscapeBounds.Min.Y);
-		int32 MaxX = FMath::Floor(LandscapeBounds.Max.X);
-		int32 MaxY = FMath::Floor(LandscapeBounds.Max.Y);
+		int32 MinX = FMath::CeilToInt(LandscapeBounds.Min.X);
+		int32 MinY = FMath::CeilToInt(LandscapeBounds.Min.Y);
+		int32 MaxX = FMath::FloorToInt(LandscapeBounds.Max.X);
+		int32 MaxY = FMath::FloorToInt(LandscapeBounds.Max.Y);
 
 		FLandscapeEditDataInterface LandscapeEdit(this);
 		TArray<FLandscapeSplineInterpPoint> Points = Segment->GetPoints();
@@ -4903,10 +4903,10 @@ void ULandscapeComponent::GeneratePlatformVertexData()
 				{
 					for (int32 x = 0; x < LodSubsectionSizeQuads; x++)
 					{
-						int32 x0 = FMath::Round((float)x * MipRatio);
-						int32 y0 = FMath::Round((float)y * MipRatio);
-						int32 x1 = FMath::Round((float)(x + 1) * MipRatio);
-						int32 y1 = FMath::Round((float)(y + 1) * MipRatio);
+						int32 x0 = FMath::RoundToInt((float)x * MipRatio);
+						int32 y0 = FMath::RoundToInt((float)y * MipRatio);
+						int32 x1 = FMath::RoundToInt((float)(x + 1) * MipRatio);
+						int32 y1 = FMath::RoundToInt((float)(y + 1) * MipRatio);
 
 						FLandscapeVertexRef V1(x0, y0, SubX, SubY);
 						FLandscapeVertexRef V2(x1, y0, SubX, SubY);
@@ -4955,8 +4955,8 @@ void ULandscapeComponent::GeneratePlatformVertexData()
 
 		float HeightmapScaleBiasZ = HeightmapScaleBias.Z + HeightmapSubsectionOffsetU * (float)SubX;
 		float HeightmapScaleBiasW = HeightmapScaleBias.W + HeightmapSubsectionOffsetV * (float)SubY;
-		int32 BaseMipOfsX = FMath::Round(HeightmapScaleBiasZ * (float)HeightmapTexture->Source.GetSizeX());
-		int32 BaseMipOfsY = FMath::Round(HeightmapScaleBiasW * (float)HeightmapTexture->Source.GetSizeY());
+		int32 BaseMipOfsX = FMath::RoundToInt(HeightmapScaleBiasZ * (float)HeightmapTexture->Source.GetSizeX());
+		int32 BaseMipOfsY = FMath::RoundToInt(HeightmapScaleBiasW * (float)HeightmapTexture->Source.GetSizeY());
 
 		DstVert->Position[0] = X;
 		DstVert->Position[1] = Y;
@@ -4979,8 +4979,8 @@ void ULandscapeComponent::GeneratePlatformVertexData()
 			int32 CurrentMipOfsY = BaseMipOfsY >> Mip;
 
 			float MipRatio = (float)MipSubsectionSizeQuads / (float)SubsectionSizeQuads; // Morph Base to current MIP
-			int32 MipX = FMath::Round((float)X * MipRatio);
-			int32 MipY = FMath::Round((float)Y * MipRatio);
+			int32 MipX = FMath::RoundToInt((float)X * MipRatio);
+			int32 MipY = FMath::RoundToInt((float)Y * MipRatio);
 
 			FColor* CurrentMipSrcRow = HeightmapMipData[Mip] + (CurrentMipOfsY + MipY) * MipSizeX + CurrentMipOfsX;
 			uint16 Height = CurrentMipSrcRow[MipX].R << 8 | CurrentMipSrcRow[MipX].G;
@@ -4999,11 +4999,11 @@ void ULandscapeComponent::GeneratePlatformVertexData()
 		{
 			if (Mip < 4)
 			{
-				DstVert->LODHeights[4 + Mip] = FMath::Round(float(MipHeights[Mip] - MinHeight) / (MaxHeight - MinHeight) * 255);
+				DstVert->LODHeights[4 + Mip] = FMath::RoundToInt(float(MipHeights[Mip] - MinHeight) / (MaxHeight - MinHeight) * 255);
 			}
 			else // Mip 4 5 packed into SubX, SubY
 			{
-				DstVert->Position[Mip - 2] += (FMath::Round(float(MipHeights[Mip] - MinHeight) / (MaxHeight - MinHeight) * 255)) & (0xfffe);
+				DstVert->Position[Mip - 2] += (FMath::RoundToInt(float(MipHeights[Mip] - MinHeight) / (MaxHeight - MinHeight) * 255)) & (0xfffe);
 			}
 		}
 
diff --git a/Engine/Source/Runtime/Engine/Private/Landscape/LandscapeEditInterface.cpp b/Engine/Source/Runtime/Engine/Private/Landscape/LandscapeEditInterface.cpp
--- a/Engine/Source/Runtime/Engine/Private/Landscape/LandscapeEditInterface.cpp
+++ b/Engine/Source/Runtime/Engine/Private/Landscape/LandscapeEditInterface.cpp
@@ -345,8 +345,8 @@ void FLandscapeEditDataInterface::SetHeightData(int32 X1, int32 Y1, int32 X2, in
 							if( VertexNormals && LandscapeX > X1 && LandscapeX < X2 && LandscapeY > Y1 && LandscapeY < Y2 )
 							{
 								FVector Normal = VertexNormals[DataIndex].SafeNormal();
-								TexData.B = FMath::Round( 127.5f * (Normal.X + 1.f) );
-								TexData.A = FMath::Round( 127.5f * (Normal.Y + 1.f) );
+								TexData.B = FMath::RoundToInt( 127.5f * (Normal.X + 1.f) );
+								TexData.A = FMath::RoundToInt( 127.5f * (Normal.Y + 1.f) );
 							}
 							else if (NormalData)
 							{
@@ -495,8 +495,8 @@ void FLandscapeEditDataInterface::RecalculateNormals()
 
 						// Update the texture
 						FVector Normal = VertexNormals[DataIndex].SafeNormal();
-						TexData.B = FMath::Round( 127.5f * (Normal.X + 1.f) );
-						TexData.A = FMath::Round( 127.5f * (Normal.Y + 1.f) );
+						TexData.B = FMath::RoundToInt( 127.5f * (Normal.X + 1.f) );
+						TexData.A = FMath::RoundToInt( 127.5f * (Normal.Y + 1.f) );
 					}
 				}
 			}
@@ -1650,7 +1650,7 @@ void ULandscapeComponent::DeleteLayer( ULandscapeLayerInfoObject* LayerInfo, str
 								if( LayerIdx != DeleteLayerIdx && LayerNoWeightBlends[LayerIdx]==false )
 								{
 									uint8& Weight = LayerDataPtrs[LayerIdx][TexDataIndex];
-									Weight = FMath::Clamp<int32>( FMath::Round(255.f * (float)Weight/(float)OtherLayerWeightSum), 0, 255 );
+									Weight = FMath::Clamp<int32>( FMath::RoundToInt(255.f * (float)Weight/(float)OtherLayerWeightSum), 0, 255 );
 								}
 							}
 						}
@@ -2349,7 +2349,7 @@ void FLandscapeEditDataInterface::SetAlphaData(ULandscapeLayerInfoObject* const
 											// Exclude bNoWeightBlend layers
 											if (LayerNoWeightBlends[LayerIdx] == false)
 											{
-												Weight = FMath::Clamp<uint8>( FMath::Round((float)(255 - NewWeight) * (float)Weight/(float)OtherLayerWeightSum), 0, 255 );
+												Weight = FMath::Clamp<uint8>( FMath::RoundToInt((float)(255 - NewWeight) * (float)Weight/(float)OtherLayerWeightSum), 0, 255 );
 											}
 										}
 
diff --git a/Engine/Source/Runtime/Engine/Private/Landscape/LandscapeGizmoActor.cpp b/Engine/Source/Runtime/Engine/Private/Landscape/LandscapeGizmoActor.cpp
--- a/Engine/Source/Runtime/Engine/Private/Landscape/LandscapeGizmoActor.cpp
+++ b/Engine/Source/Runtime/Engine/Private/Landscape/LandscapeGizmoActor.cpp
@@ -771,8 +771,8 @@ void ALandscapeGizmoActiveActor::SampleData(int32 SizeX, int32 SizeY)
 			{
 				float TexX = X * SizeX / TexSizeX;
 				float TexY = Y * SizeY / TexSizeY;
-				int32 LX = FMath::Floor(TexX);
-				int32 LY = FMath::Floor(TexY);
+				int32 LX = FMath::FloorToInt(TexX);
+				int32 LY = FMath::FloorToInt(TexY);
 
 				float FracX = TexX - LX;
 				float FracY = TexY - LY;
diff --git a/Engine/Source/Runtime/Engine/Private/Landscape/LandscapeLight.cpp b/Engine/Source/Runtime/Engine/Private/Landscape/LandscapeLight.cpp
--- a/Engine/Source/Runtime/Engine/Private/Landscape/LandscapeLight.cpp
+++ b/Engine/Source/Runtime/Engine/Private/Landscape/LandscapeLight.cpp
@@ -140,8 +140,8 @@ namespace
 		float Ratio = (float)(LODHeightmapSize) / (HeightmapStride);
 		float Offset = 0.5f * Ratio;
 
-		int32 CurrentHeightmapOffsetX = FMath::Round((float)(LODHeightmapSize)* LandscapeComponent->HeightmapScaleBias.Z);
-		int32 CurrentHeightmapOffsetY = FMath::Round((float)(LODHeightmapSize)* LandscapeComponent->HeightmapScaleBias.W);
+		int32 CurrentHeightmapOffsetX = FMath::RoundToInt((float)(LODHeightmapSize)* LandscapeComponent->HeightmapScaleBias.Z);
+		int32 CurrentHeightmapOffsetY = FMath::RoundToInt((float)(LODHeightmapSize)* LandscapeComponent->HeightmapScaleBias.W);
 
 		// Need to match for component edge cases, otherwise it causes a little 
 		float XX = FMath::Clamp<float>((X - HeightmapOffsetX) * Ratio - Offset, 0.f, ComponentSize - 1.f) + CurrentHeightmapOffsetX;
@@ -161,11 +161,11 @@ namespace
 		FColor H3 = HeightMipData[XI + FMath::Min(YI + 1, LODHeightmapSize - 1) * LODHeightmapSize];
 		FColor H4 = HeightMipData[FMath::Min(XI + 1, LODHeightmapSize - 1) + FMath::Min(YI + 1, LODHeightmapSize - 1) * LODHeightmapSize];
 
-		uint16 Height = FMath::Round(FMath::Lerp(FMath::Lerp<float>(((H1.R << 8) + H1.G), ((H2.R << 8) + H2.G), XF),
+		uint16 Height = FMath::RoundToInt(FMath::Lerp(FMath::Lerp<float>(((H1.R << 8) + H1.G), ((H2.R << 8) + H2.G), XF),
 			FMath::Lerp<float>(((H3.R << 8) + H3.G), ((H4.R << 8) + H4.G), XF), YF));
-		uint8 B = FMath::Round(FMath::Lerp(FMath::Lerp<float>((H1.B), (H2.B), XF),
+		uint8 B = FMath::RoundToInt(FMath::Lerp(FMath::Lerp<float>((H1.B), (H2.B), XF),
 			FMath::Lerp<float>((H3.B), (H4.B), XF), YF));
-		uint8 A = FMath::Round(FMath::Lerp<float>(FMath::Lerp((H1.A), (H2.A), XF),
+		uint8 A = FMath::RoundToInt(FMath::Lerp<float>(FMath::Lerp((H1.A), (H2.A), XF),
 			FMath::Lerp<float>((H3.A), (H4.A), XF), YF));
 
 		OutHeight = FColor((Height >> 8), Height & 255, B, A);
@@ -177,9 +177,9 @@ namespace
 			FColor X3 = XYOffsetMipData[XI + FMath::Min(YI + 1, LODHeightmapSize - 1) * LODHeightmapSize];
 			FColor X4 = XYOffsetMipData[FMath::Min(XI + 1, LODHeightmapSize - 1) + FMath::Min(YI + 1, LODHeightmapSize - 1) * LODHeightmapSize];
 
-			uint16 XComp = FMath::Round(FMath::Lerp(FMath::Lerp<float>(((X1.R << 8) + X1.G), ((X2.R << 8) + X2.G), XF),
+			uint16 XComp = FMath::RoundToInt(FMath::Lerp(FMath::Lerp<float>(((X1.R << 8) + X1.G), ((X2.R << 8) + X2.G), XF),
 				FMath::Lerp<float>(((X3.R << 8) + X3.G), ((X4.R << 8) + X4.G), XF), YF));
-			uint16 YComp = FMath::Round(FMath::Lerp(FMath::Lerp<float>(((X1.B << 8) + X1.A), ((X2.B << 8) + X2.A), XF),
+			uint16 YComp = FMath::RoundToInt(FMath::Lerp(FMath::Lerp<float>(((X1.B << 8) + X1.A), ((X2.B << 8) + X2.A), XF),
 				FMath::Lerp<float>(((X3.B << 8) + X3.A), ((X4.B << 8) + X4.A), XF), YF));
 
 			OutXYOffset = FColor((XComp >> 8), XComp & 255, YComp >> 8, YComp & 255);
@@ -319,21 +319,21 @@ namespace
 							// Need interpolation
 							uint16 Height0 = (Height[0].R << 8) + Height[0].G;
 							uint16 Height1 = (Height[1].R << 8) + Height[1].G;
-							uint16 LerpHeight = FMath::Round(FMath::Lerp<float>(Height0, Height1, MorphAlpha));
+							uint16 LerpHeight = FMath::RoundToInt(FMath::Lerp<float>(Height0, Height1, MorphAlpha));
 
 							CompHeightData[X + Y * HeightmapStride] =
 								FColor((LerpHeight >> 8), LerpHeight & 255,
-								FMath::Round(FMath::Lerp<float>(Height[0].B, Height[1].B, MorphAlpha)),
-								FMath::Round(FMath::Lerp<float>(Height[0].A, Height[1].A, MorphAlpha)));
+								FMath::RoundToInt(FMath::Lerp<float>(Height[0].B, Height[1].B, MorphAlpha)),
+								FMath::RoundToInt(FMath::Lerp<float>(Height[0].A, Height[1].A, MorphAlpha)));
 							if (LandscapeComponent->XYOffsetmapTexture)
 							{
 								uint16 XComp0 = (XYOffset[0].R << 8) + XYOffset[0].G;
 								uint16 XComp1 = (XYOffset[1].R << 8) + XYOffset[1].G;
-								uint16 LerpXComp = FMath::Round(FMath::Lerp<float>(XComp0, XComp1, MorphAlpha));
+								uint16 LerpXComp = FMath::RoundToInt(FMath::Lerp<float>(XComp0, XComp1, MorphAlpha));
 
 								uint16 YComp0 = (XYOffset[0].B << 8) + XYOffset[0].A;
 								uint16 YComp1 = (XYOffset[1].B << 8) + XYOffset[1].A;
-								uint16 LerpYComp = FMath::Round(FMath::Lerp<float>(YComp0, YComp1, MorphAlpha));
+								uint16 LerpYComp = FMath::RoundToInt(FMath::Lerp<float>(YComp0, YComp1, MorphAlpha));
 
 								CompXYOffsetData[X + Y * HeightmapStride] =
 									FColor(LerpXComp >> 8, LerpXComp & 255, LerpYComp >> 8, LerpYComp & 255);
diff --git a/Engine/Source/Runtime/Engine/Private/Landscape/LandscapeRender.cpp b/Engine/Source/Runtime/Engine/Private/Landscape/LandscapeRender.cpp
--- a/Engine/Source/Runtime/Engine/Private/Landscape/LandscapeRender.cpp
+++ b/Engine/Source/Runtime/Engine/Private/Landscape/LandscapeRender.cpp
@@ -1195,7 +1195,7 @@ int32 FLandscapeComponentSceneProxy::CalcLODForSubsectionNoForced(const class FS
 	if (View.Family->LandscapeLODOverride >= 0)
 	{
 		float fLOD = FMath::Clamp<float>(View.Family->LandscapeLODOverride, FMath::Max<int32>(LODBias, MinLOD), FMath::Min<int32>(MaxLOD, MaxLOD + LODBias));
-		return FMath::Floor(fLOD);
+		return FMath::FloorToInt(fLOD);
 	}
 #endif
 
@@ -1206,7 +1206,7 @@ int32 FLandscapeComponentSceneProxy::CalcLODForSubsectionNoForced(const class FS
 		float ComponentDistance = FVector2D(CurrentCameraLocalPos-ComponentPosition).Size() + DistDiff;
 		// Clamp calculated distance based LOD with LODBiased values
 		float fLOD = FMath::Clamp<float>( ComponentDistance / LODDistance, FMath::Max<int32>(LODBias, MinLOD), FMath::Min<int32>(MaxLOD, MaxLOD+LODBias) );
-		return FMath::Floor( fLOD );
+		return FMath::FloorToInt( fLOD );
 	}
 	else
 	{
@@ -1214,7 +1214,7 @@ int32 FLandscapeComponentSceneProxy::CalcLODForSubsectionNoForced(const class FS
 
 		// The "/ 5.0f" is totally arbitrary
 		float fLOD = FMath::Clamp<float>(FMath::Sqrt(Scale / 5.0f), FMath::Max<int32>(LODBias, MinLOD), FMath::Min<int32>(MaxLOD, MaxLOD + LODBias));
-		return FMath::Floor(fLOD);
+		return FMath::FloorToInt(fLOD);
 	}
 }
 
@@ -1586,10 +1586,10 @@ void FLandscapeSharedBuffers::CreateIndexBuffers()
 					{
 						for (int32 x = 0; x < LodSubsectionSizeQuads; x++)
 						{
-							int32 x0 = FMath::Round((float)x * MipRatio);
-							int32 y0 = FMath::Round((float)y * MipRatio);
-							int32 x1 = FMath::Round((float)(x + 1) * MipRatio);
-							int32 y1 = FMath::Round((float)(y + 1) * MipRatio);
+							int32 x0 = FMath::RoundToInt((float)x * MipRatio);
+							int32 y0 = FMath::RoundToInt((float)y * MipRatio);
+							int32 x1 = FMath::RoundToInt((float)(x + 1) * MipRatio);
+							int32 y1 = FMath::RoundToInt((float)(y + 1) * MipRatio);
 
 							FLandscapeVertexRef V00(x0, y0, SubX, SubY);
 							FLandscapeVertexRef V10(x1, y0, SubX, SubY);
diff --git a/Engine/Source/Runtime/Engine/Private/Landscape/LandscapeSplines.cpp b/Engine/Source/Runtime/Engine/Private/Landscape/LandscapeSplines.cpp
--- a/Engine/Source/Runtime/Engine/Private/Landscape/LandscapeSplines.cpp
+++ b/Engine/Source/Runtime/Engine/Private/Landscape/LandscapeSplines.cpp
@@ -1204,7 +1204,7 @@ void ULandscapeSplineSegment::UpdateSplinePoints(bool bUpdateCollision)
 	const float StartRoll = FMath::DegreesToRadians(StartRollDegrees);
 	const float EndRoll = FMath::DegreesToRadians(EndRollDegrees);
 
-	int32 NumPoints = FMath::Ceil(SplineLength / OuterSplines->SplineResolution);
+	int32 NumPoints = FMath::CeilToInt(SplineLength / OuterSplines->SplineResolution);
 	NumPoints = FMath::Clamp(NumPoints, 1, 1000);
 
 	float OldKeyTime = 0;
@@ -1227,7 +1227,7 @@ void ULandscapeSplineSegment::UpdateSplinePoints(bool bUpdateCollision)
 		// If not the first keypoint, interp from the last keypoint.
 		if(i > 0)
 		{
-			int32 NumSteps = FMath::Ceil( (NewKeyTime - OldKeyTime) * NumPoints );
+			int32 NumSteps = FMath::CeilToInt( (NewKeyTime - OldKeyTime) * NumPoints );
 			float DrawSubstep = (NewKeyTime - OldKeyTime) / NumSteps;
 
 			// Add a point for each substep, except the ends because that's the point added outside the interp'ing.
diff --git a/Engine/Source/Runtime/Engine/Private/Level.cpp b/Engine/Source/Runtime/Engine/Private/Level.cpp
--- a/Engine/Source/Runtime/Engine/Private/Level.cpp
+++ b/Engine/Source/Runtime/Engine/Private/Level.cpp
@@ -827,9 +827,9 @@ void ULevel::CreateModelComponents()
 				}
 				else
 				{
-					Key.X				= FMath::Floor(NodeBounds.GetCenter().X / MODEL_GRID_SIZE_XY);
-					Key.Y				= FMath::Floor(NodeBounds.GetCenter().Y / MODEL_GRID_SIZE_XY);
-					Key.Z				= FMath::Floor(NodeBounds.GetCenter().Z / MODEL_GRID_SIZE_Z);
+					Key.X				= FMath::FloorToInt(NodeBounds.GetCenter().X / MODEL_GRID_SIZE_XY);
+					Key.Y				= FMath::FloorToInt(NodeBounds.GetCenter().Y / MODEL_GRID_SIZE_XY);
+					Key.Z				= FMath::FloorToInt(NodeBounds.GetCenter().Z / MODEL_GRID_SIZE_Z);
 				}
 
 				Key.MaskedPolyFlags = Surf.PolyFlags & PF_ModelComponentMask;
diff --git a/Engine/Source/Runtime/Engine/Private/LightMap.cpp b/Engine/Source/Runtime/Engine/Private/LightMap.cpp
--- a/Engine/Source/Runtime/Engine/Private/LightMap.cpp
+++ b/Engine/Source/Runtime/Engine/Private/LightMap.cpp
@@ -150,15 +150,15 @@ static int32 AdjustTextureMappingSize(int32 InSize)
 	int32 NewSize = InSize;
 	if (InSize > TextureMappingMinDownsampleSize0 && InSize <= TextureMappingMinDownsampleSize1)
 	{
-		NewSize = FMath::Trunc(InSize * TextureMappingDownsampleFactor0);
+		NewSize = FMath::TruncToInt(InSize * TextureMappingDownsampleFactor0);
 	}
 	else if (InSize > TextureMappingMinDownsampleSize1 && InSize <= TextureMappingMinDownsampleSize2)
 	{
-		NewSize = FMath::Trunc(InSize * TextureMappingDownsampleFactor1);
+		NewSize = FMath::TruncToInt(InSize * TextureMappingDownsampleFactor1);
 	}
 	else if (InSize > TextureMappingMinDownsampleSize2)
 	{
-		NewSize = FMath::Trunc(InSize * TextureMappingDownsampleFactor2);
+		NewSize = FMath::TruncToInt(InSize * TextureMappingDownsampleFactor2);
 	}
 	return NewSize;
 }
diff --git a/Engine/Source/Runtime/Engine/Private/LocalPlayer.cpp b/Engine/Source/Runtime/Engine/Private/LocalPlayer.cpp
--- a/Engine/Source/Runtime/Engine/Private/LocalPlayer.cpp
+++ b/Engine/Source/Runtime/Engine/Private/LocalPlayer.cpp
@@ -580,10 +580,10 @@ bool ULocalPlayer::GetProjectionData(FViewport* Viewport, EStereoscopicPass Ster
 		return false;
 	}
 
-	int32 X = FMath::Trunc(Origin.X * Viewport->GetSizeXY().X);
-	int32 Y = FMath::Trunc(Origin.Y * Viewport->GetSizeXY().Y);
-	uint32 SizeX = FMath::Trunc(Size.X * Viewport->GetSizeXY().X);
-	uint32 SizeY = FMath::Trunc(Size.Y * Viewport->GetSizeXY().Y);
+	int32 X = FMath::TruncToInt(Origin.X * Viewport->GetSizeXY().X);
+	int32 Y = FMath::TruncToInt(Origin.Y * Viewport->GetSizeXY().Y);
+	uint32 SizeX = FMath::TruncToInt(Size.X * Viewport->GetSizeXY().X);
+	uint32 SizeY = FMath::TruncToInt(Size.Y * Viewport->GetSizeXY().Y);
 
 #if !(UE_BUILD_SHIPPING || UE_BUILD_TEST)
 
diff --git a/Engine/Source/Runtime/Engine/Private/Materials/MaterialUniformExpressions.h b/Engine/Source/Runtime/Engine/Private/Materials/MaterialUniformExpressions.h
--- a/Engine/Source/Runtime/Engine/Private/Materials/MaterialUniformExpressions.h
+++ b/Engine/Source/Runtime/Engine/Private/Materials/MaterialUniformExpressions.h
@@ -804,10 +804,10 @@ public:
 	{
 		X->GetNumberValue(Context, OutValue);
 
-		OutValue.R = FMath::Floor(OutValue.R);
-		OutValue.G = FMath::Floor(OutValue.G);
-		OutValue.B = FMath::Floor(OutValue.B);
-		OutValue.A = FMath::Floor(OutValue.A);
+		OutValue.R = FMath::FloorToInt(OutValue.R);
+		OutValue.G = FMath::FloorToInt(OutValue.G);
+		OutValue.B = FMath::FloorToInt(OutValue.B);
+		OutValue.A = FMath::FloorToInt(OutValue.A);
 	}
 	virtual bool IsConstant() const
 	{
@@ -848,10 +848,10 @@ public:
 	{
 		X->GetNumberValue(Context, OutValue);
 
-		OutValue.R = FMath::Ceil(OutValue.R);
-		OutValue.G = FMath::Ceil(OutValue.G);
-		OutValue.B = FMath::Ceil(OutValue.B);
-		OutValue.A = FMath::Ceil(OutValue.A);
+		OutValue.R = FMath::CeilToInt(OutValue.R);
+		OutValue.G = FMath::CeilToInt(OutValue.G);
+		OutValue.B = FMath::CeilToInt(OutValue.B);
+		OutValue.A = FMath::CeilToInt(OutValue.A);
 	}
 	virtual bool IsConstant() const
 	{
@@ -892,10 +892,10 @@ public:
 	{
 		X->GetNumberValue(Context, OutValue);
 
-		OutValue.R = OutValue.R - FMath::Floor(OutValue.R);
-		OutValue.G = OutValue.G - FMath::Floor(OutValue.G);
-		OutValue.B = OutValue.B - FMath::Floor(OutValue.B);
-		OutValue.A = OutValue.A - FMath::Floor(OutValue.A);
+		OutValue.R = OutValue.R - FMath::FloorToInt(OutValue.R);
+		OutValue.G = OutValue.G - FMath::FloorToInt(OutValue.G);
+		OutValue.B = OutValue.B - FMath::FloorToInt(OutValue.B);
+		OutValue.A = OutValue.A - FMath::FloorToInt(OutValue.A);
 	}
 	virtual bool IsConstant() const
 	{
diff --git a/Engine/Source/Runtime/Engine/Private/ModelComponent.cpp b/Engine/Source/Runtime/Engine/Private/ModelComponent.cpp
--- a/Engine/Source/Runtime/Engine/Private/ModelComponent.cpp
+++ b/Engine/Source/Runtime/Engine/Private/ModelComponent.cpp
@@ -430,7 +430,7 @@ bool UModelComponent::GenerateElements(bool bBuildRenderData)
 			}
 
 			// Find an element with the same material as this node.
-			uint32 ResValue = FMath::Trunc(Surf.LightMapScale);
+			uint32 ResValue = FMath::TruncToInt(Surf.LightMapScale);
 			TMap<UMaterialInterface*,FModelElement*>* MaterialElementMap = ResToMaterialElementMap->FindRef(ResValue);
 			if (!MaterialElementMap)
 			{
diff --git a/Engine/Source/Runtime/Engine/Private/ModelLight.cpp b/Engine/Source/Runtime/Engine/Private/ModelLight.cpp
--- a/Engine/Source/Runtime/Engine/Private/ModelLight.cpp
+++ b/Engine/Source/Runtime/Engine/Private/ModelLight.cpp
@@ -419,13 +419,13 @@ void UModelComponent::GetSurfaceLightMapResolution( int32 SurfaceIndex, int32 Qu
 	if (bFoundNode)
 	{
 		float Scale = Surf.LightMapScale * QualityScale;
-		MinUV.X = (float)FMath::Floor(MinUV.X / Scale) * Scale;
-		MinUV.Y = (float)FMath::Floor(MinUV.Y / Scale) * Scale;
-		MaxUV.X = (float)FMath::Ceil(MaxUV.X / Scale) * Scale;
-		MaxUV.Y = (float)FMath::Ceil(MaxUV.Y / Scale) * Scale;
+		MinUV.X = FMath::FloorToFloat(MinUV.X / Scale) * Scale;
+		MinUV.Y = FMath::FloorToFloat(MinUV.Y / Scale) * Scale;
+		MaxUV.X = FMath::CeilToFloat(MaxUV.X / Scale) * Scale;
+		MaxUV.Y = FMath::CeilToFloat(MaxUV.Y / Scale) * Scale;
 
-		Width = FMath::Clamp(FMath::Ceil((MaxUV.X - MinUV.X) / (Surf.LightMapScale * QualityScale)),4,SHADOWMAP_MAX_WIDTH);
-		Height = FMath::Clamp(FMath::Ceil((MaxUV.Y - MinUV.Y) / (Surf.LightMapScale * QualityScale)),4,SHADOWMAP_MAX_HEIGHT);
+		Width = FMath::Clamp(FMath::CeilToInt((MaxUV.X - MinUV.X) / (Surf.LightMapScale * QualityScale)),4,SHADOWMAP_MAX_WIDTH);
+		Height = FMath::Clamp(FMath::CeilToInt((MaxUV.Y - MinUV.Y) / (Surf.LightMapScale * QualityScale)),4,SHADOWMAP_MAX_HEIGHT);
 		WorldToMap = FMatrix(
 			FPlane(MapX.X / (MaxUV.X - MinUV.X),	MapY.X / (MaxUV.Y - MinUV.Y),	Surf.Plane.X,	0),
 			FPlane(MapX.Y / (MaxUV.X - MinUV.X),	MapY.Y / (MaxUV.Y - MinUV.Y),	Surf.Plane.Y,	0),
@@ -454,7 +454,7 @@ bool UModelComponent::GetLightMapResolution( int32& Width, int32& Height ) const
 		LightMapArea += SizeX * SizeY;
 	}
 
-	Width = FMath::Trunc( FMath::Sqrt( LightMapArea ) );
+	Width = FMath::TruncToInt( FMath::Sqrt( LightMapArea ) );
 	Height = Width;
 	return false;
 }
@@ -478,14 +478,14 @@ void UModelComponent::GetLightAndShadowMapMemoryUsage( int32& LightMapMemoryUsag
 	
 	// Stored in texture.
 	const float MIP_FACTOR = 1.33f;
-	ShadowMapMemoryUsage	= FMath::Trunc( MIP_FACTOR * LightMapWidth * LightMapHeight ); // G8
+	ShadowMapMemoryUsage	= FMath::TruncToInt( MIP_FACTOR * LightMapWidth * LightMapHeight ); // G8
 	if( AllowHighQualityLightmaps() )
 	{ 
-		LightMapMemoryUsage = FMath::Trunc( NUM_HQ_LIGHTMAP_COEF * MIP_FACTOR * LightMapWidth * LightMapHeight ); // DXT5
+		LightMapMemoryUsage = FMath::TruncToInt( NUM_HQ_LIGHTMAP_COEF * MIP_FACTOR * LightMapWidth * LightMapHeight ); // DXT5
 	}
 	else
 	{
-		LightMapMemoryUsage = FMath::Trunc( NUM_LQ_LIGHTMAP_COEF * MIP_FACTOR * LightMapWidth * LightMapHeight / 2 ); // DXT1
+		LightMapMemoryUsage = FMath::TruncToInt( NUM_LQ_LIGHTMAP_COEF * MIP_FACTOR * LightMapWidth * LightMapHeight / 2 ); // DXT1
 	}
 	return;
 }
@@ -911,10 +911,10 @@ void UModel::ApplyStaticLighting()
 
 						// Encode LogL
 						LogL = ( LogL - GroupQuantizedData->Add[0][3] ) / GroupQuantizedData->Scale[0][3];
-						Residual = LogL * 255.0f - FMath::Round( LogL * 255.0f ) + 0.5f;
+						Residual = LogL * 255.0f - FMath::RoundToFloat( LogL * 255.0f ) + 0.5f;
 
-						DestSample.Coefficients[0][3] = (uint8)FMath::Clamp<int32>( FMath::Round( LogL * 255.0f ), 0, 255 );
-						DestSample.Coefficients[1][3] = (uint8)FMath::Clamp<int32>( FMath::Round( Residual * 255.0f ), 0, 255 );
+						DestSample.Coefficients[0][3] = (uint8)FMath::Clamp<int32>( FMath::RoundToInt( LogL * 255.0f ), 0, 255 );
+						DestSample.Coefficients[1][3] = (uint8)FMath::Clamp<int32>( FMath::RoundToInt( Residual * 255.0f ), 0, 255 );
 					}
 	
 					// go over each color coefficient and dequantize and requantize with new Scale/Add
@@ -932,7 +932,7 @@ void UModel::ApplyStaticLighting()
 							const float Repacked = ( Unpacked - GroupQuantizedData->Add[CoefficientIndex][ColorIndex] ) / GroupQuantizedData->Scale[CoefficientIndex][ColorIndex];
 
 							// requantize it
-							DestSample.Coefficients[CoefficientIndex][ColorIndex] = (uint8)FMath::Clamp<int32>( FMath::Round( FMath::Pow( Repacked, 1.0f / Exponent ) * 255.0f ), 0, 255 );
+							DestSample.Coefficients[CoefficientIndex][ColorIndex] = (uint8)FMath::Clamp<int32>( FMath::RoundToInt( FMath::Pow( Repacked, 1.0f / Exponent ) * 255.0f ), 0, 255 );
 						}
 					}
 
diff --git a/Engine/Source/Runtime/Engine/Private/NetConnection.cpp b/Engine/Source/Runtime/Engine/Private/NetConnection.cpp
--- a/Engine/Source/Runtime/Engine/Private/NetConnection.cpp
+++ b/Engine/Source/Runtime/Engine/Private/NetConnection.cpp
@@ -1440,11 +1440,11 @@ void UNetConnection::Tick()
 	// Update queued byte count.
 	// this should be at the end so that the cap is applied *after* sending (and adjusting QueuedBytes for) any remaining data for this tick
 	float DeltaBytes = CurrentNetSpeed * DeltaTime;
-	QueuedBytes -= FMath::Trunc(DeltaBytes);
+	QueuedBytes -= FMath::TruncToInt(DeltaBytes);
 	float AllowedLag = 2.f * DeltaBytes;
 	if (QueuedBytes < -AllowedLag)
 	{
-		QueuedBytes = FMath::Trunc(-AllowedLag);
+		QueuedBytes = FMath::TruncToInt(-AllowedLag);
 	}
 }
 
diff --git a/Engine/Source/Runtime/Engine/Private/NetworkDriver.cpp b/Engine/Source/Runtime/Engine/Private/NetworkDriver.cpp
--- a/Engine/Source/Runtime/Engine/Private/NetworkDriver.cpp
+++ b/Engine/Source/Runtime/Engine/Private/NetworkDriver.cpp
@@ -142,23 +142,23 @@ void UNetDriver::TickFlush(float DeltaSeconds)
 		float RealTime = Time - StatUpdateTime;
 
 		// Use the elapsed time to keep things scaled to one measured unit
-		InBytes = FMath::Trunc(InBytes / RealTime);
-		OutBytes = FMath::Trunc(OutBytes / RealTime);
+		InBytes = FMath::TruncToInt(InBytes / RealTime);
+		OutBytes = FMath::TruncToInt(OutBytes / RealTime);
 
-		NetGUIDOutBytes = FMath::Trunc(NetGUIDOutBytes / RealTime);
-		NetGUIDInBytes = FMath::Trunc(NetGUIDInBytes / RealTime);
+		NetGUIDOutBytes = FMath::TruncToInt(NetGUIDOutBytes / RealTime);
+		NetGUIDInBytes = FMath::TruncToInt(NetGUIDInBytes / RealTime);
 
 		// Save off for stats later
 
 		InBytesPerSecond = InBytes;
 		OutBytesPerSecond = OutBytes;
 
-		InPackets = FMath::Trunc(InPackets / RealTime);
-		OutPackets = FMath::Trunc(OutPackets / RealTime);
-		InBunches = FMath::Trunc(InBunches / RealTime);
-		OutBunches = FMath::Trunc(OutBunches / RealTime);
-		OutPacketsLost = FMath::Trunc(100.f * OutPacketsLost / FMath::Max((float)OutPackets,1.f));
-		InPacketsLost = FMath::Trunc(100.f * InPacketsLost / FMath::Max((float)InPackets + InPacketsLost,1.f));
+		InPackets = FMath::TruncToInt(InPackets / RealTime);
+		OutPackets = FMath::TruncToInt(OutPackets / RealTime);
+		InBunches = FMath::TruncToInt(InBunches / RealTime);
+		OutBunches = FMath::TruncToInt(OutBunches / RealTime);
+		OutPacketsLost = FMath::TruncToInt(100.f * OutPacketsLost / FMath::Max((float)OutPackets,1.f));
+		InPacketsLost = FMath::TruncToInt(100.f * InPacketsLost / FMath::Max((float)InPackets + InPacketsLost,1.f));
 		
 #if STATS
 
@@ -167,7 +167,7 @@ void UNetDriver::TickFlush(float DeltaSeconds)
 			// Copy the net status values over
 			if (ServerConnection != NULL && ServerConnection->PlayerController != NULL && ServerConnection->PlayerController->PlayerState != NULL)
 			{
-				SET_DWORD_STAT(STAT_Ping, FMath::Trunc(1000.0f * ServerConnection->PlayerController->PlayerState->ExactPing));
+				SET_DWORD_STAT(STAT_Ping, FMath::TruncToInt(1000.0f * ServerConnection->PlayerController->PlayerState->ExactPing));
 			}
 			else
 			{
@@ -195,17 +195,17 @@ void UNetDriver::TickFlush(float DeltaSeconds)
 			SET_DWORD_STAT(STAT_NetGUIDOutRate,NetGUIDOutBytes);
 
 			// Use the elapsed time to keep things scaled to one measured unit
-			VoicePacketsSent = FMath::Trunc(VoicePacketsSent / RealTime);
+			VoicePacketsSent = FMath::TruncToInt(VoicePacketsSent / RealTime);
 			SET_DWORD_STAT(STAT_VoicePacketsSent,VoicePacketsSent);
-			VoicePacketsRecv = FMath::Trunc(VoicePacketsRecv / RealTime);
+			VoicePacketsRecv = FMath::TruncToInt(VoicePacketsRecv / RealTime);
 			SET_DWORD_STAT(STAT_VoicePacketsRecv,VoicePacketsRecv);
-			VoiceBytesSent = FMath::Trunc(VoiceBytesSent / RealTime);
+			VoiceBytesSent = FMath::TruncToInt(VoiceBytesSent / RealTime);
 			SET_DWORD_STAT(STAT_VoiceBytesSent,VoiceBytesSent);
-			VoiceBytesRecv = FMath::Trunc(VoiceBytesRecv / RealTime);
+			VoiceBytesRecv = FMath::TruncToInt(VoiceBytesRecv / RealTime);
 			SET_DWORD_STAT(STAT_VoiceBytesRecv,VoiceBytesRecv);
 			// Determine voice percentages
-			VoiceInPercent = (InBytes > 0) ? FMath::Trunc(100.f * (float)VoiceBytesRecv / (float)InBytes) : 0;
-			VoiceOutPercent = (OutBytes > 0) ? FMath::Trunc(100.f * (float)VoiceBytesSent / (float)OutBytes) : 0;
+			VoiceInPercent = (InBytes > 0) ? FMath::TruncToInt(100.f * (float)VoiceBytesRecv / (float)InBytes) : 0;
+			VoiceOutPercent = (OutBytes > 0) ? FMath::TruncToInt(100.f * (float)VoiceBytesSent / (float)OutBytes) : 0;
 
 			SET_DWORD_STAT(STAT_PercentInVoice,VoiceInPercent);
 			SET_DWORD_STAT(STAT_PercentOutVoice,VoiceOutPercent);
@@ -1662,7 +1662,7 @@ FActorPriority::FActorPriority(UNetConnection* InConnection, UActorChannel* InCh
 	Priority = 0;
 	for (int32 i = 0; i < Viewers.Num(); i++)
 	{
-		Priority = FMath::Max<int32>(Priority, FMath::Round(65536.0f * Actor->GetNetPriority(Viewers[i].ViewLocation, Viewers[i].ViewDir, Viewers[i].InViewer, InChannel, Time, bLowBandwidth)));
+		Priority = FMath::Max<int32>(Priority, FMath::RoundToInt(65536.0f * Actor->GetNetPriority(Viewers[i].ViewLocation, Viewers[i].ViewDir, Viewers[i].InViewer, InChannel, Time, bLowBandwidth)));
 	}
 }
 
@@ -1734,7 +1734,7 @@ int32 UNetDriver::ServerReplicateActors(float DeltaSeconds)
 			//@todo - ideally we wouldn't want to tick more clients with a higher deltatime as that's not going to be good for performance and probably saturate bandwidth in hitchy situations, maybe 
 			// come up with a solution that is greedier with higher framerates, but still won't risk saturating server upstream bandwidth
 			float ClientUpdatesThisFrame = GEngine->NetClientTicksPerSecond * (DeltaSeconds + DeltaTimeOverflow) * (LanPlay ? 2.f : 1.f);
-			NumClientsToTick = FMath::Min<int32>(NumClientsToTick,FMath::Trunc(ClientUpdatesThisFrame));
+			NumClientsToTick = FMath::Min<int32>(NumClientsToTick,FMath::TruncToInt(ClientUpdatesThisFrame));
 			//UE_LOG(LogNet, Log, TEXT("%2.3f: Ticking %d clients this frame, %2.3f/%2.4f"),GetWorld()->GetTimeSeconds(),NumClientsToTick,DeltaSeconds,ClientUpdatesThisFrame);
 			if (NumClientsToTick == 0)
 			{
diff --git a/Engine/Source/Runtime/Engine/Private/NiagaraComponent.cpp b/Engine/Source/Runtime/Engine/Private/NiagaraComponent.cpp
--- a/Engine/Source/Runtime/Engine/Private/NiagaraComponent.cpp
+++ b/Engine/Source/Runtime/Engine/Private/NiagaraComponent.cpp
@@ -395,7 +395,7 @@ private:
 	int32 CalcNumToSpawn(float DeltaSeconds)
 	{
 		float FloatNumToSpawn = SpawnRemainder + (DeltaSeconds * SpawnRate);
-		int32 NumToSpawn = FMath::Floor(FloatNumToSpawn);
+		int32 NumToSpawn = FMath::FloorToInt(FloatNumToSpawn);
 		SpawnRemainder = FloatNumToSpawn - NumToSpawn;
 		return NumToSpawn;
 	}
diff --git a/Engine/Source/Runtime/Engine/Private/Particles/ParticleBeam2EmitterInstance.cpp b/Engine/Source/Runtime/Engine/Private/Particles/ParticleBeam2EmitterInstance.cpp
--- a/Engine/Source/Runtime/Engine/Private/Particles/ParticleBeam2EmitterInstance.cpp
+++ b/Engine/Source/Runtime/Engine/Private/Particles/ParticleBeam2EmitterInstance.cpp
@@ -641,7 +641,7 @@ float FParticleBeam2EmitterInstance::SpawnBeamParticles(float OldLeftover, float
 	float	NewLeftover = OldLeftover + DeltaTime * Rate;
 
 	// Ensure continous spawning... lots of fiddling.
-	int32		Number		= FMath::Floor(NewLeftover);
+	int32		Number		= FMath::FloorToInt(NewLeftover);
 	float	Increment	= 1.f / Rate;
 	float	StartTime	= DeltaTime + OldLeftover * Increment - Increment;
 	NewLeftover			= NewLeftover - Number;
@@ -681,11 +681,11 @@ float FParticleBeam2EmitterInstance::SpawnBeamParticles(float OldLeftover, float
 	{
 		if (DeltaTime < 0.25f)
 		{
-			bProcessSpawn = Resize(NewCount + FMath::Trunc(FMath::Sqrt((float)NewCount)) + 1);
+			bProcessSpawn = Resize(NewCount + FMath::TruncToInt(FMath::Sqrt((float)NewCount)) + 1);
 		}
 		else
 		{
-			bProcessSpawn = Resize((NewCount + FMath::Trunc(FMath::Sqrt((float)NewCount)) + 1), false);
+			bProcessSpawn = Resize((NewCount + FMath::TruncToInt(FMath::Sqrt((float)NewCount)) + 1), false);
 		}
 	}
 
diff --git a/Engine/Source/Runtime/Engine/Private/Particles/ParticleBeamModules.cpp b/Engine/Source/Runtime/Engine/Private/Particles/ParticleBeamModules.cpp
--- a/Engine/Source/Runtime/Engine/Private/Particles/ParticleBeamModules.cpp
+++ b/Engine/Source/Runtime/Engine/Private/Particles/ParticleBeamModules.cpp
@@ -356,7 +356,7 @@ void UParticleModuleTypeDataBeam2::Update(FParticleEmitterInstance* Owner, int32
 					// Determine the steps and step size
 					BeamData->StepSize		= FullMagnitude / InterpolationCount;
 					BeamData->TravelRatio	= TrueMagnitude / FullMagnitude;
-					BeamData->Steps			= FMath::Floor(BeamData->TravelRatio * InterpolationCount);
+					BeamData->Steps			= FMath::FloorToInt(BeamData->TravelRatio * InterpolationCount);
 					// Readjust the travel ratio
 					BeamData->TravelRatio	= (TrueMagnitude - (BeamData->StepSize * BeamData->Steps)) / BeamData->StepSize;
 				}
@@ -378,7 +378,7 @@ void UParticleModuleTypeDataBeam2::Update(FParticleEmitterInstance* Owner, int32
 				{
 					// The noise points are based on distance...
 					// Determine the number of points to drop.
-					int32 Count = FMath::Trunc(FullMagnitude / BeamNoise->FrequencyDistance);
+					int32 Count = FMath::TruncToInt(FullMagnitude / BeamNoise->FrequencyDistance);
 					Count = FMath::Min<int32>(Count, Freq);
 					BeamData->StepSize		= FullMagnitude / (Count + 1);
 					BeamData->Steps			= Count;
@@ -409,12 +409,12 @@ void UParticleModuleTypeDataBeam2::Update(FParticleEmitterInstance* Owner, int32
 
 				if (BeamNoise->FrequencyDistance > 0.0f)
 				{
-					int32 Count = FMath::Trunc(FullMagnitude / BeamNoise->FrequencyDistance);
+					int32 Count = FMath::TruncToInt(FullMagnitude / BeamNoise->FrequencyDistance);
 					Count = FMath::Min<int32>(Count, Freq);
 					BeamData->StepSize		= FullMagnitude / (Count + 1);
 					// Determine the partial trail amount and the steps taken
 					BeamData->TravelRatio	= TrueMagnitude / FullMagnitude;
-					BeamData->Steps			= FMath::Floor(BeamData->TravelRatio * (Count + 1));
+					BeamData->Steps			= FMath::FloorToInt(BeamData->TravelRatio * (Count + 1));
 					// Lock the steps to the frequency
 					if (BeamData->Steps > Count)
 					{
@@ -445,7 +445,7 @@ void UParticleModuleTypeDataBeam2::Update(FParticleEmitterInstance* Owner, int32
 					BeamData->StepSize		= FullMagnitude / (Freq + 1);
 					// Determine the partial trail amount and the steps taken
 					BeamData->TravelRatio	= TrueMagnitude / FullMagnitude;
-					BeamData->Steps			= FMath::Floor(BeamData->TravelRatio * (Freq + 1));
+					BeamData->Steps			= FMath::FloorToInt(BeamData->TravelRatio * (Freq + 1));
 					// Lock the steps to the frequency
 					if (BeamData->Steps > Freq)
 					{
@@ -552,7 +552,7 @@ void UParticleModuleTypeDataBeam2::Update(FParticleEmitterInstance* Owner, int32
 				{
 					[REDACTED_EMAIL]. Fix this!
 					// When the data fills in (vertices), it is incorrect.
-					BeamData->TriangleCount	+= FMath::Floor(BeamData->TravelRatio * NoiseTess) * 2;
+					BeamData->TriangleCount	+= FMath::FloorToInt(BeamData->TravelRatio * NoiseTess) * 2;
 				}
 			}
 		}
@@ -581,7 +581,7 @@ void UParticleModuleTypeDataBeam2::Update(FParticleEmitterInstance* Owner, int32
 				{
 					[REDACTED_EMAIL]. Fix this!
 					// When the data fills in (vertices), it is incorrect.
-//					BeamData->TriangleCount	+= FMath::Floor(BeamData->TravelRatio * NoiseTess) * 2;
+//					BeamData->TriangleCount	+= FMath::FloorToInt(BeamData->TravelRatio * NoiseTess) * 2;
 				}
 			}
 		}
@@ -1305,7 +1305,7 @@ void UParticleModuleBeamNoise::Spawn(FParticleEmitterInstance* Owner, int32 Offs
 	int32 CalcFreq = Frequency;
 	if (Frequency_LowRange > 0)
 	{
-		CalcFreq = FMath::Trunc((FMath::SRand() * (Frequency - Frequency_LowRange)) + Frequency_LowRange);
+		CalcFreq = FMath::TruncToInt((FMath::SRand() * (Frequency - Frequency_LowRange)) + Frequency_LowRange);
 	}
 	BEAM2_TYPEDATA_SETFREQUENCY(BeamData->Lock_Max_NumNoisePoints, CalcFreq);
 	
diff --git a/Engine/Source/Runtime/Engine/Private/Particles/ParticleComponents.cpp b/Engine/Source/Runtime/Engine/Private/Particles/ParticleComponents.cpp
--- a/Engine/Source/Runtime/Engine/Private/Particles/ParticleComponents.cpp
+++ b/Engine/Source/Runtime/Engine/Private/Particles/ParticleComponents.cpp
@@ -721,11 +721,11 @@ int32	UParticleLODLevel::CalculateMaxActiveParticleCount()
 			// Special case for one loop... 
 			if (ParticleLifetime < MaxDuration)
 			{
-				MaxAPC += FMath::Ceil(ParticleLifetime * MaxSpawnRate);
+				MaxAPC += FMath::CeilToInt(ParticleLifetime * MaxSpawnRate);
 			}
 			else
 			{
-				MaxAPC += FMath::Ceil(MaxDuration * MaxSpawnRate);
+				MaxAPC += FMath::CeilToInt(MaxDuration * MaxSpawnRate);
 			}
 			// Safety zone...
 			MaxAPC += 1;
@@ -736,11 +736,11 @@ int32	UParticleLODLevel::CalculateMaxActiveParticleCount()
 		{
 			if (ParticleLifetime < MaxDuration)
 			{
-				MaxAPC += FMath::Ceil(ParticleLifetime * MaxSpawnRate);
+				MaxAPC += FMath::CeilToInt(ParticleLifetime * MaxSpawnRate);
 			}
 			else
 			{
-				MaxAPC += (FMath::Ceil(FMath::Ceil(MaxDuration * MaxSpawnRate) * ParticleLifetime));
+				MaxAPC += (FMath::CeilToInt(FMath::CeilToInt(MaxDuration * MaxSpawnRate) * ParticleLifetime));
 			}
 			// Safety zone...
 			MaxAPC += 1;
@@ -748,7 +748,7 @@ int32	UParticleLODLevel::CalculateMaxActiveParticleCount()
 			MaxAPC += MaxBurstCount;
 			if (ParticleLifetime > MaxDuration)
 			{
-				MaxAPC += MaxBurstCount * FMath::Ceil(ParticleLifetime - MaxDuration);
+				MaxAPC += MaxBurstCount * FMath::CeilToInt(ParticleLifetime - MaxDuration);
 			}
 		}
 	}
@@ -758,7 +758,7 @@ int32	UParticleLODLevel::CalculateMaxActiveParticleCount()
 		// Single loop case is all we will worry about. Safer base estimate - but not ideal.
 		if (ParticleLifetime < MaxDuration)
 		{
-			MaxAPC += FMath::Ceil(ParticleLifetime * FMath::Ceil(MaxSpawnRate));
+			MaxAPC += FMath::CeilToInt(ParticleLifetime * FMath::CeilToInt(MaxSpawnRate));
 		}
 		else
 		{
@@ -766,21 +766,21 @@ int32	UParticleLODLevel::CalculateMaxActiveParticleCount()
 			{
 				if (ParticleLifetime <= MaxDuration)
 				{
-					MaxAPC += FMath::Ceil(MaxDuration * MaxSpawnRate);
+					MaxAPC += FMath::CeilToInt(MaxDuration * MaxSpawnRate);
 				}
 				else //if (ParticleLifetime > MaxDuration)
 				{
-					MaxAPC += FMath::Ceil(MaxDuration * MaxSpawnRate) * ParticleLifetime;
+					MaxAPC += FMath::CeilToInt(MaxDuration * MaxSpawnRate) * ParticleLifetime;
 				}
 			}
 			else
 			{
 				// No lifetime, no duration...
-				MaxAPC += FMath::Ceil(MaxSpawnRate);
+				MaxAPC += FMath::CeilToInt(MaxSpawnRate);
 			}
 		}
 		// Safety zone...
-		MaxAPC += FMath::Max<int32>(FMath::Ceil(MaxSpawnRate * 0.032f), 2);
+		MaxAPC += FMath::Max<int32>(FMath::CeilToInt(MaxSpawnRate * 0.032f), 2);
 		// Burst
 		MaxAPC += MaxBurstCount;
 	}
diff --git a/Engine/Source/Runtime/Engine/Private/Particles/ParticleEmitterInstances.cpp b/Engine/Source/Runtime/Engine/Private/Particles/ParticleEmitterInstances.cpp
--- a/Engine/Source/Runtime/Engine/Private/Particles/ParticleEmitterInstances.cpp
+++ b/Engine/Source/Runtime/Engine/Private/Particles/ParticleEmitterInstances.cpp
@@ -1371,11 +1371,11 @@ float FParticleEmitterInstance::GetCurrentBurstRateOffset(float& DeltaTime, int3
 							int32 Count = BurstEntry->Count;
 							if (BurstEntry->CountLow > -1)
 							{
-								Count = BurstEntry->CountLow + FMath::Round(FMath::SRand() * (float)(BurstEntry->Count - BurstEntry->CountLow));
+								Count = BurstEntry->CountLow + FMath::RoundToInt(FMath::SRand() * (float)(BurstEntry->Count - BurstEntry->CountLow));
 							}
 							// Take in to account scale.
 							float Scale = LODLevel->SpawnModule->BurstScale.GetValue(EmitterTime, Component);
-							Count = FMath::Ceil(Count * Scale);
+							Count = FMath::CeilToInt(Count * Scale);
 							SpawnRateInc += Count / DeltaTime;
 							Burst += Count;
 							LocalBurstFired.Fired[BurstIdx] = true;
@@ -1755,7 +1755,7 @@ float FParticleEmitterInstance::Spawn(float DeltaTime)
 		if (DetailMode != DM_High)
 		{
 			SpawnRate = FMath::Max<float>(0.0f, SpawnRate * SpriteTemplate->MediumDetailSpawnRateScale);
-			BurstCount = FMath::Ceil(BurstCount * SpriteTemplate->MediumDetailSpawnRateScale);
+			BurstCount = FMath::CeilToInt(BurstCount * SpriteTemplate->MediumDetailSpawnRateScale);
 		}
 	}
 	else
@@ -1772,7 +1772,7 @@ float FParticleEmitterInstance::Spawn(float DeltaTime)
 		float SafetyLeftover = OldLeftover;
 		// Ensure continuous spawning... lots of fiddling.
 		float	NewLeftover = OldLeftover + DeltaTime * SpawnRate;
-		int32		Number		= FMath::Floor(NewLeftover);
+		int32		Number		= FMath::FloorToInt(NewLeftover);
 		float	Increment	= (SpawnRate > 0.0f) ? (1.f / SpawnRate) : 0.0f;
 		float	StartTime	= DeltaTime + OldLeftover * Increment - Increment;
 		NewLeftover			= NewLeftover - Number;
@@ -1794,11 +1794,11 @@ float FParticleEmitterInstance::Spawn(float DeltaTime)
 		{
 			if (DeltaTime < PeakActiveParticleUpdateDelta)
 			{
-				bProcessSpawn = Resize(NewCount + FMath::Trunc(FMath::Sqrt(FMath::Sqrt((float)NewCount)) + 1));
+				bProcessSpawn = Resize(NewCount + FMath::TruncToInt(FMath::Sqrt(FMath::Sqrt((float)NewCount)) + 1));
 			}
 			else
 			{
-				bProcessSpawn = Resize((NewCount + FMath::Trunc(FMath::Sqrt(FMath::Sqrt((float)NewCount)) + 1)), false);
+				bProcessSpawn = Resize((NewCount + FMath::TruncToInt(FMath::Sqrt(FMath::Sqrt((float)NewCount)) + 1)), false);
 			}
 		}
 
@@ -1949,11 +1949,11 @@ void FParticleEmitterInstance::ForceSpawn(float DeltaTime, int32 InSpawnCount, i
 		{
 			if (DeltaTime < PeakActiveParticleUpdateDelta)
 			{
-				bProcessSpawn = Resize(NewCount + FMath::Trunc(FMath::Sqrt(FMath::Sqrt((float)NewCount)) + 1));
+				bProcessSpawn = Resize(NewCount + FMath::TruncToInt(FMath::Sqrt(FMath::Sqrt((float)NewCount)) + 1));
 			}
 			else
 			{
-				bProcessSpawn = Resize((NewCount + FMath::Trunc(FMath::Sqrt(FMath::Sqrt((float)NewCount)) + 1)), false);
+				bProcessSpawn = Resize((NewCount + FMath::TruncToInt(FMath::Sqrt(FMath::Sqrt((float)NewCount)) + 1)), false);
 			}
 		}
 
diff --git a/Engine/Source/Runtime/Engine/Private/Particles/ParticleGpuSimulation.cpp b/Engine/Source/Runtime/Engine/Private/Particles/ParticleGpuSimulation.cpp
--- a/Engine/Source/Runtime/Engine/Private/Particles/ParticleGpuSimulation.cpp
+++ b/Engine/Source/Runtime/Engine/Private/Particles/ParticleGpuSimulation.cpp
@@ -1104,7 +1104,7 @@ static void BuildTileVertexBuffer( FParticleBufferParamRef TileOffsetsRef, const
 		const uint32 TileIndex = Tiles[Index];
 		TileOffset[Index] = FVector2D(
 			FMath::Fractional( (float)TileIndex / (float)GParticleSimulationTileCountX ),
-			FMath::Fractional( FMath::TruncFloat( (float)TileIndex / (float)GParticleSimulationTileCountX ) / (float)GParticleSimulationTileCountY )
+			FMath::Fractional( FMath::TruncToFloat( (float)TileIndex / (float)GParticleSimulationTileCountX ) / (float)GParticleSimulationTileCountY )
 									  );
 	}
 	for ( ; Index < AlignedTileCount; ++Index )
@@ -1777,7 +1777,7 @@ static void BuildParticleVertexBuffer( FVertexBufferRHIParamRef VertexBufferRHI,
 		const uint32 TileIndex = InTiles[Index];
 		const FVector2D TileOffset(
 			FMath::Fractional( (float)TileIndex / (float)GParticleSimulationTileCountX ),
-			FMath::Fractional( FMath::TruncFloat( (float)TileIndex / (float)GParticleSimulationTileCountX ) / (float)GParticleSimulationTileCountY )
+			FMath::Fractional( FMath::TruncToFloat( (float)TileIndex / (float)GParticleSimulationTileCountX ) / (float)GParticleSimulationTileCountY )
 			);
 		for ( int32 ParticleY = 0; ParticleY < GParticleSimulationTileSize; ++ParticleY )
 		{
@@ -3294,7 +3294,7 @@ private:
 		if (AllowedLoopCount == 0)
 		{
 			const int32 EstMaxTiles = (EmitterInfo.MaxParticleCount + GParticlesPerTile - 1) / GParticlesPerTile;
-			const int32 SlackTiles = FMath::Ceil(FXConsoleVariables::ParticleSlackGPU * (float)EstMaxTiles);
+			const int32 SlackTiles = FMath::CeilToInt(FXConsoleVariables::ParticleSlackGPU * (float)EstMaxTiles);
 			return FMath::Min<int32>(EstMaxTiles + SlackTiles, FXConsoleVariables::MaxParticleTilePreAllocation);
 		}
 		return 0;
@@ -3401,7 +3401,7 @@ private:
 		uint32 TileIndex = (AllocatedTiles.IsValidIndex(TileToAllocateFrom)) ? AllocatedTiles[TileToAllocateFrom] : INDEX_NONE;
 		FVector2D TileOffset(
 			FMath::Fractional((float)TileIndex / (float)GParticleSimulationTileCountX),
-			FMath::Fractional(FMath::TruncFloat((float)TileIndex / (float)GParticleSimulationTileCountX) / (float)GParticleSimulationTileCountY)
+			FMath::Fractional(FMath::TruncToFloat((float)TileIndex / (float)GParticleSimulationTileCountX) / (float)GParticleSimulationTileCountY)
 			);
 
 		for (int32 ParticleIndex = 0; ParticleIndex < NumNewParticles; ++ParticleIndex)
@@ -3435,7 +3435,7 @@ private:
 				ActiveTileCount++;
 				TileIndex = AllocatedTiles[TileToAllocateFrom];
 				TileOffset.X = FMath::Fractional((float)TileIndex / (float)GParticleSimulationTileCountX);
-				TileOffset.Y = FMath::Fractional(FMath::TruncFloat((float)TileIndex / (float)GParticleSimulationTileCountX) / (float)GParticleSimulationTileCountY);
+				TileOffset.Y = FMath::Fractional(FMath::TruncToFloat((float)TileIndex / (float)GParticleSimulationTileCountX) / (float)GParticleSimulationTileCountY);
 				FreeParticlesInTile = GParticlesPerTile;
 			}
 			FNewParticle& Particle = *new(InNewParticles) FNewParticle();
@@ -3483,7 +3483,7 @@ private:
 		// Determine how many to spawn.
 		FSpawnInfo Info;
 		float AccumSpawnCount = SpawnFraction + SpawnRate * DeltaSeconds;
-		Info.Count = FMath::Min(FMath::Trunc(AccumSpawnCount), FXConsoleVariables::MaxGPUParticlesSpawnedPerFrame);
+		Info.Count = FMath::Min(FMath::TruncToInt(AccumSpawnCount), FXConsoleVariables::MaxGPUParticlesSpawnedPerFrame);
 		Info.Increment = (SpawnRate > 0.0f) ? (1.f / SpawnRate) : 0.0f;
 		Info.StartTime = DeltaSeconds + SpawnFraction * Info.Increment - Info.Increment;
 		SpawnFraction = AccumSpawnCount - Info.Count;
diff --git a/Engine/Source/Runtime/Engine/Private/Particles/ParticleModules.cpp b/Engine/Source/Runtime/Engine/Private/Particles/ParticleModules.cpp
--- a/Engine/Source/Runtime/Engine/Private/Particles/ParticleModules.cpp
+++ b/Engine/Source/Runtime/Engine/Private/Particles/ParticleModules.cpp
@@ -638,13 +638,13 @@ uint32 UParticleModule::PrepRandomSeedInstancePayload(FParticleEmitterInstance*
 			{
 				if (InRandSeedInfo.bInstanceSeedIsIndex == false)
 				{
-					InRandSeedPayload->RandomStream.Initialize(FMath::Round(SeedValue));
+					InRandSeedPayload->RandomStream.Initialize(FMath::RoundToInt(SeedValue));
 				}
 				else
 				{
 					if (InRandSeedInfo.RandomSeeds.Num() > 0)
 					{
-						int32 Index = FMath::Min<int32>((InRandSeedInfo.RandomSeeds.Num() - 1), FMath::Trunc(SeedValue));
+						int32 Index = FMath::Min<int32>((InRandSeedInfo.RandomSeeds.Num() - 1), FMath::TruncToInt(SeedValue));
 						InRandSeedPayload->RandomStream.Initialize(InRandSeedInfo.RandomSeeds[Index]);
 						return 0;
 					}
@@ -1970,7 +1970,7 @@ float UParticleModuleSubUV::DetermineImageIndex(FParticleEmitterInstance* Owner,
 
 		if (InterpMethod == PSUVIM_Linear)
 		{
-			ImageIndex = FMath::TruncFloat(ImageIndex);
+			ImageIndex = FMath::TruncToFloat(ImageIndex);
 		}
 	}
 	else if ((InterpMethod == PSUVIM_Random) || (InterpMethod == PSUVIM_Random_Blend))
@@ -1980,13 +1980,13 @@ float UParticleModuleSubUV::DetermineImageIndex(FParticleEmitterInstance* Owner,
 			(SubUVPayload.RandomImageTime == 0.0f))
 		{
 			const float RandomNumber = FMath::SRand();
-			ImageIndex = FMath::Trunc(RandomNumber * TotalSubImages);
+			ImageIndex = FMath::TruncToInt(RandomNumber * TotalSubImages);
 			SubUVPayload.RandomImageTime	= Particle->RelativeTime;
 		}
 
 		if (InterpMethod == PSUVIM_Random)
 		{
-			ImageIndex = FMath::TruncFloat(ImageIndex);
+			ImageIndex = FMath::TruncToFloat(ImageIndex);
 		}
 	}
 	else
@@ -2118,7 +2118,7 @@ void UParticleModuleSubUVMovie::Spawn(FParticleEmitterInstance* Owner, int32 Off
 			}
 			else if (StartingFrame == 0)
 			{
-				MoviePayload.Time = FMath::Trunc(FMath::SRand() * (iTotalSubImages-1));
+				MoviePayload.Time = FMath::TruncToFloat(FMath::SRand() * (iTotalSubImages-1));
 			}
 
 			// Update the payload
@@ -2159,7 +2159,7 @@ float UParticleModuleSubUVMovie::DetermineImageIndex(FParticleEmitterInstance* O
 	float ImageIndex = MoviePayload.Time * UserSetFrameRate;
 	if (InterpMethod != PSUVIM_Linear_Blend)
 	{
-		ImageIndex = FMath::TruncFloat(ImageIndex);
+		ImageIndex = FMath::TruncToFloat(ImageIndex);
 	}
 	return ImageIndex;
 }
@@ -3750,7 +3750,7 @@ void UParticleModuleAttractorParticle::Spawn(FParticleEmitterInstance* Owner, in
 			switch (SelectionMethod)
 			{
 			case EAPSM_Random:
-				LastSelIndex		= FMath::Trunc(FMath::SRand() * AttractorEmitterInst->ActiveParticles);
+				LastSelIndex		= FMath::TruncToInt(FMath::SRand() * AttractorEmitterInst->ActiveParticles);
 				Data.SourceIndex	= LastSelIndex;
 				break;
 			case EAPSM_Sequential:
diff --git a/Engine/Source/Runtime/Engine/Private/Particles/ParticleModules_Collision.cpp b/Engine/Source/Runtime/Engine/Private/Particles/ParticleModules_Collision.cpp
--- a/Engine/Source/Runtime/Engine/Private/Particles/ParticleModules_Collision.cpp
+++ b/Engine/Source/Runtime/Engine/Private/Particles/ParticleModules_Collision.cpp
@@ -133,7 +133,7 @@ void UParticleModuleCollision::Spawn(FParticleEmitterInstance* Owner, int32 Offs
 		PARTICLE_ELEMENT(FParticleCollisionPayload, CollisionPayload);
 		CollisionPayload.UsedDampingFactor = DampingFactor.GetValue(Owner->EmitterTime, Owner->Component);
 		CollisionPayload.UsedDampingFactorRotation = DampingFactorRotation.GetValue(Owner->EmitterTime, Owner->Component);
-		CollisionPayload.UsedCollisions = FMath::Round(MaxCollisions.GetValue(Owner->EmitterTime, Owner->Component));
+		CollisionPayload.UsedCollisions = FMath::RoundToInt(MaxCollisions.GetValue(Owner->EmitterTime, Owner->Component));
 		CollisionPayload.Delay = DelayAmount.GetValue(Owner->EmitterTime, Owner->Component);
 		if (CollisionPayload.Delay > SpawnTime)
 		{
diff --git a/Engine/Source/Runtime/Engine/Private/Particles/ParticleModules_Event.cpp b/Engine/Source/Runtime/Engine/Private/Particles/ParticleModules_Event.cpp
--- a/Engine/Source/Runtime/Engine/Private/Particles/ParticleModules_Event.cpp
+++ b/Engine/Source/Runtime/Engine/Private/Particles/ParticleModules_Event.cpp
@@ -340,24 +340,24 @@ bool UParticleModuleEventReceiverSpawn::ProcessParticleEvent(FParticleEmitterIns
 		{
 		case EPET_Spawn:
 		case EPET_Burst:
-			Count = FMath::Round(SpawnCount.GetValue(InEvent.EmitterTime));
+			Count = FMath::RoundToInt(SpawnCount.GetValue(InEvent.EmitterTime));
 			break;
 		case EPET_Death:
 			{
 				FParticleEventDeathData* DeathData = (FParticleEventDeathData*)(&InEvent);
-				Count = FMath::Round(SpawnCount.GetValue(bUseParticleTime ? DeathData->ParticleTime : InEvent.EmitterTime));
+				Count = FMath::RoundToInt(SpawnCount.GetValue(bUseParticleTime ? DeathData->ParticleTime : InEvent.EmitterTime));
 			}
 			break;
 		case EPET_Collision:
 			{
 				FParticleEventCollideData* CollideData = (FParticleEventCollideData*)(&InEvent);
-				Count = FMath::Round(SpawnCount.GetValue(bUseParticleTime ? CollideData->ParticleTime : InEvent.EmitterTime));
+				Count = FMath::RoundToInt(SpawnCount.GetValue(bUseParticleTime ? CollideData->ParticleTime : InEvent.EmitterTime));
 			}
 			break;
 		case EPET_Blueprint:
 			{
 				FParticleEventKismetData* KismetData = (FParticleEventKismetData*)(&InEvent);
-				Count = FMath::Round(SpawnCount.GetValue(InEvent.EmitterTime));
+				Count = FMath::RoundToInt(SpawnCount.GetValue(InEvent.EmitterTime));
 			}
 			break;
 		}
diff --git a/Engine/Source/Runtime/Engine/Private/Particles/ParticleModules_Location.cpp b/Engine/Source/Runtime/Engine/Private/Particles/ParticleModules_Location.cpp
--- a/Engine/Source/Runtime/Engine/Private/Particles/ParticleModules_Location.cpp
+++ b/Engine/Source/Runtime/Engine/Private/Particles/ParticleModules_Location.cpp
@@ -90,7 +90,7 @@ void UParticleModuleLocation::SpawnEx(FParticleEmitterInstance* Owner, int32 Off
 		{
 			FVector Min, Max;
 			StartLocation.Distribution->GetRange(Min, Max);
-			FVector Lerped = FMath::Lerp(Min, Max, FMath::TruncFloat((FMath::SRand() * (DistributeOverNPoints - 1.0f)) + 0.5f)/(DistributeOverNPoints - 1.0f));
+			FVector Lerped = FMath::Lerp(Min, Max, FMath::TruncToFloat((FMath::SRand() * (DistributeOverNPoints - 1.0f)) + 0.5f)/(DistributeOverNPoints - 1.0f));
 			LocationOffset.Set(Lerped.X, Lerped.Y, Lerped.Z);
 		}
 	}
@@ -442,7 +442,7 @@ void UParticleModuleLocationEmitter::Spawn(FParticleEmitterInstance* Owner, int3
 			{
 			case ELESM_Random:
 				{
-					Index = FMath::Trunc(FMath::SRand() * LocationEmitterInst->ActiveParticles);
+					Index = FMath::TruncToInt(FMath::SRand() * LocationEmitterInst->ActiveParticles);
 					if (Index >= LocationEmitterInst->ActiveParticles)
 					{
 						Index = LocationEmitterInst->ActiveParticles - 1;
@@ -1545,7 +1545,7 @@ void UParticleModuleLocationBoneSocket::Spawn(FParticleEmitterInstance* Owner, i
 	else if (SelectionMethod == BONESOCKETSEL_Random)
 	{
 		// Note: This can select the same socket over and over...
-		SourceIndex = FMath::Trunc(FMath::SRand() * (SourceLocations.Num() - 1));
+		SourceIndex = FMath::TruncToInt(FMath::SRand() * (SourceLocations.Num() - 1));
 		InstancePayload->LastSelectedIndex = SourceIndex;
 	}
 
@@ -2079,7 +2079,7 @@ void UParticleModuleLocationSkelVertSurface::Spawn(FParticleEmitterInstance* Own
 	{
 		int32 SourceLocationsCount(SkelMeshResource->LODModels[0].VertexBufferGPUSkin.GetNumVertices());
 
-		SourceIndex = FMath::Trunc(FMath::SRand() * ((float)SourceLocationsCount) - 1);
+		SourceIndex = FMath::TruncToInt(FMath::SRand() * ((float)SourceLocationsCount) - 1);
 		InstancePayload->VertIndex = SourceIndex;
 
 		if(SourceIndex != -1)
@@ -2099,10 +2099,10 @@ void UParticleModuleLocationSkelVertSurface::Spawn(FParticleEmitterInstance* Own
 	{
 		FStaticLODModel& LODModel = SkelMeshResource->LODModels[0];
 		int32 SectionCount = LODModel.Sections.Num();
-		int32 RandomSection = FMath::Round(FMath::SRand() * ((float)SectionCount-1));
+		int32 RandomSection = FMath::RoundToInt(FMath::SRand() * ((float)SectionCount-1));
 
 		SourceIndex = LODModel.Sections[RandomSection].BaseIndex +
-			(FMath::Trunc(FMath::SRand() * ((float)LODModel.Sections[RandomSection].NumTriangles))*3);
+			(FMath::TruncToInt(FMath::SRand() * ((float)LODModel.Sections[RandomSection].NumTriangles))*3);
 
 		InstancePayload->VertIndex = SourceIndex;
 
diff --git a/Engine/Source/Runtime/Engine/Private/Particles/ParticleModules_Spawn.cpp b/Engine/Source/Runtime/Engine/Private/Particles/ParticleModules_Spawn.cpp
--- a/Engine/Source/Runtime/Engine/Private/Particles/ParticleModules_Spawn.cpp
+++ b/Engine/Source/Runtime/Engine/Private/Particles/ParticleModules_Spawn.cpp
@@ -137,7 +137,7 @@ bool UParticleModuleSpawn::GenerateLODModuleValues(UParticleModule* SourceModule
 			// Don't drop below 1...
 			if (Burst->Count > 0)
 			{
-				Burst->Count = FMath::Trunc(SourceBurst->Count * (Percentage / 100.0f));
+				Burst->Count = FMath::TruncToInt(SourceBurst->Count * (Percentage / 100.0f));
 				if (Burst->Count == 0)
 				{
 					Burst->Count = 1;
@@ -357,7 +357,7 @@ bool UParticleModuleSpawnPerUnit::GetSpawnAmount(FParticleEmitterInstance* Owner
 
 			// Calculate number of particles to emit
 			float NewLeftover = (TravelDistance + LeftoverTravel) * ParticlesPerUnit;
-			Number = FMath::Floor(NewLeftover);
+			Number = FMath::FloorToInt(NewLeftover);
 			float InvDeltaTime = (DeltaTime > 0.0f) ? 1.0f / DeltaTime : 0.0f;
 			Rate = Number * InvDeltaTime;
 			NewTravelLeftover = (TravelDistance + LeftoverTravel) - (Number * UnitScalar);
diff --git a/Engine/Source/Runtime/Engine/Private/Particles/ParticleSystemRender.cpp b/Engine/Source/Runtime/Engine/Private/Particles/ParticleSystemRender.cpp
--- a/Engine/Source/Runtime/Engine/Private/Particles/ParticleSystemRender.cpp
+++ b/Engine/Source/Runtime/Engine/Private/Particles/ParticleSystemRender.cpp
@@ -2414,7 +2414,7 @@ void FDynamicMeshEmitterData::GetInstanceData(void* InstanceData, void* DynamicP
 
 			float SubImageIndex = SubUVPayload->ImageIndex;
 			float SubImageLerp = FMath::Fractional(SubImageIndex);
-			int32 SubImageA = FMath::Floor(SubImageIndex);
+			int32 SubImageA = FMath::FloorToInt(SubImageIndex);
 			int32 SubImageB = SubImageA + 1;
 			int32 SubImageAH = SubImageA % SubImagesX;
 			int32 SubImageBH = SubImageB % SubImagesX;
@@ -4887,7 +4887,7 @@ int32 FDynamicBeam2EmitterData::FillData_InterpolatedNoise(FAsyncBufferFillData&
 		float InterpFraction = FMath::Fractional(InterpStepSize);
 		//bool bInterpFractionIsZero = (FMath::Abs(InterpFraction) < KINDA_SMALL_NUMBER) ? true : false;
 		bool bInterpFractionIsZero = false;
-		int32 InterpIndex = FMath::Trunc(InterpStepSize);
+		int32 InterpIndex = FMath::TruncToInt(InterpStepSize);
 
 		FVector* NoisePoints	= TargetNoisePoints;
 		FVector* NextNoise		= NextNoisePoints;
@@ -5817,12 +5817,12 @@ void FDynamicRibbonEmitterData::RenderDebug(FParticleSystemSceneProxy* Proxy, FP
 
 					DrawPosition = DebugParticle->Location;
 					DrawSize = DebugParticle->Size.X * Source.Scale.X;
-					int32 Red   = FMath::Trunc(255.0f * (1.0f - ColorScale));
-					int32 Green = FMath::Trunc(255.0f * ColorScale);
+					int32 Red   = FMath::TruncToInt(255.0f * (1.0f - ColorScale));
+					int32 Green = FMath::TruncToInt(255.0f * ColorScale);
 					ColorScale += Increment;
 					DrawColor = FColor(Red,Green,0);
-					Red   = FMath::Trunc(255.0f * (1.0f - ColorScale));
-					Green = FMath::Trunc(255.0f * ColorScale);
+					Red   = FMath::TruncToInt(255.0f * (1.0f - ColorScale));
+					Green = FMath::TruncToInt(255.0f * ColorScale);
 					PrevDrawColor = FColor(Red,Green,0);
 
 					if (bRenderParticles == true)
@@ -6495,12 +6495,12 @@ void FDynamicAnimTrailEmitterData::RenderDebug(FParticleSystemSceneProxy* Proxy,
 			{
 				RenderData.CalcVertexData(0.0f, DrawPosition, DrawFirstEdgePosition, DrawSecondEdgePosition,TiledU, DrawSize, DummyColor, NULL);
 
-				int32 Red   = FMath::Trunc(255.0f * (1.0f - ColorScale));
-				int32 Green = FMath::Trunc(255.0f * ColorScale);
+				int32 Red   = FMath::TruncToInt(255.0f * (1.0f - ColorScale));
+				int32 Green = FMath::TruncToInt(255.0f * ColorScale);
 				ColorScale += Increment;
 				DrawColor = FColor(Red,Green,0);
-				Red   = FMath::Trunc(255.0f * (1.0f - ColorScale));
-				Green = FMath::Trunc(255.0f * ColorScale);
+				Red   = FMath::TruncToInt(255.0f * (1.0f - ColorScale));
+				Green = FMath::TruncToInt(255.0f * ColorScale);
 				PrevDrawColor = FColor(Red,Green,0);
 
 				if (bRenderParticles == true)
