#ttp 331438 UE4: PlatformMath refactor

diff --git a/Engine/Source/Runtime/Engine/Private/Particles/ParticleTrail2EmitterInstance.cpp b/Engine/Source/Runtime/Engine/Private/Particles/ParticleTrail2EmitterInstance.cpp
--- a/Engine/Source/Runtime/Engine/Private/Particles/ParticleTrail2EmitterInstance.cpp
+++ b/Engine/Source/Runtime/Engine/Private/Particles/ParticleTrail2EmitterInstance.cpp
@@ -1016,7 +1016,7 @@ bool FParticleRibbonEmitterInstance::GetSpawnPerUnitAmount(float DeltaTime, int3
 				NewLeftover += CheckTangent * TrailTypeData->TangentSpawningScalar;
 
 				const float InvDeltaTime = DeltaTime > 0.0f ? 1.0f / DeltaTime : 0.0f;
-				OutCount = (TrailTypeData->bSpawnInitialParticle && !ActiveParticles && (NewLeftover < 1.0f))? 1: FMath::Floor(NewLeftover);
+				OutCount = (TrailTypeData->bSpawnInitialParticle && !ActiveParticles && (NewLeftover < 1.0f))? 1: FMath::FloorToInt(NewLeftover);
 				OutRate = OutCount * InvDeltaTime;
 				NewTravelLeftover = (TravelDistance + LeftoverTravel) - (OutCount * SpawnPerUnitModule->UnitScalar);
 				SourceDistanceTraveled[InTrailIdx] = FMath::Max<float>(0.0f, NewTravelLeftover);
@@ -1223,7 +1223,7 @@ float FParticleRibbonEmitterInstance::Spawn(float DeltaTime)
 	const int32 LocalMaxParticleInTrailCount = TrailTypeData->MaxParticleInTrailCount;
 	float SafetyLeftover = OldLeftover;
 	float NewLeftover = OldLeftover + DeltaTime * SpawnRate;
-	int32 SpawnNumber	= FMath::Floor(NewLeftover);
+	int32 SpawnNumber	= FMath::FloorToInt(NewLeftover);
 	float SliceIncrement = (SpawnRate > 0.0f) ? (1.f / SpawnRate) : 0.0f;
 	float SpawnStartTime = DeltaTime + OldLeftover * SliceIncrement - SliceIncrement;
 	SpawnFraction = NewLeftover - SpawnNumber;
@@ -1251,11 +1251,11 @@ float FParticleRibbonEmitterInstance::Spawn(float DeltaTime)
 	{
 		if (DeltaTime < 0.25f)
 		{
-			bProcessSpawn = Resize(NewCount + FMath::Trunc(FMath::Sqrt((float)NewCount)) + 1);
+			bProcessSpawn = Resize(NewCount + FMath::TruncToInt(FMath::Sqrt((float)NewCount)) + 1);
 		}
 		else
 		{
-			bProcessSpawn = Resize((NewCount + FMath::Trunc(FMath::Sqrt((float)NewCount)) + 1), false);
+			bProcessSpawn = Resize((NewCount + FMath::TruncToInt(FMath::Sqrt((float)NewCount)) + 1), false);
 		}
 	}
 
@@ -1585,11 +1585,11 @@ bool FParticleRibbonEmitterInstance::Spawn_Source(float DeltaTime)
 		{
 			if (DeltaTime < 0.25f)
 			{
-				bProcessSpawn = Resize(NewCount + FMath::Trunc(FMath::Sqrt((float)NewCount)) + 1);
+				bProcessSpawn = Resize(NewCount + FMath::TruncToInt(FMath::Sqrt((float)NewCount)) + 1);
 			}
 			else
 			{
-				bProcessSpawn = Resize((NewCount + FMath::Trunc(FMath::Sqrt((float)NewCount)) + 1), false);
+				bProcessSpawn = Resize((NewCount + FMath::TruncToInt(FMath::Sqrt((float)NewCount)) + 1), false);
 			}
 		}
 
@@ -2188,7 +2188,7 @@ bool FParticleRibbonEmitterInstance::ResolveSourcePoint(int32 InTrailIdx,
 						{
 						case EPSSM_Random:
 							{
-								Index = FMath::Trunc(FMath::FRand() * SourceEmitter->ActiveParticles);
+								Index = FMath::TruncToInt(FMath::FRand() * SourceEmitter->ActiveParticles);
 							}
 							break;
 						case EPSSM_Sequential:
@@ -2431,7 +2431,7 @@ void FParticleRibbonEmitterInstance::DetermineVertexAndTriangleCount()
 
 					//@todo. Need to adjust the tangent diff count when the distance is REALLY small...
 					float TangDiff = CheckTangent * TrailTypeData->TangentTessellationScalar;
-					int32 InterpCount = FMath::Trunc(DistDiff) + FMath::Trunc(TangDiff);
+					int32 InterpCount = FMath::TruncToInt(DistDiff) + FMath::TruncToInt(TangDiff);
 
 					// There always is at least 1 point (the source particle itself)
 					InterpCount = (InterpCount > 0) ? InterpCount : 1;
@@ -3306,7 +3306,7 @@ float FParticleAnimTrailEmitterInstance::Spawn(float DeltaTime)
 
 	float SafetyLeftover = OldLeftover;
 	float NewLeftover = OldLeftover + DeltaTime * SpawnRate;
-	int32 SpawnNumber	= FMath::Floor(NewLeftover);
+	int32 SpawnNumber	= FMath::FloorToInt(NewLeftover);
 	float SliceIncrement = (SpawnRate > 0.0f) ? (1.f / SpawnRate) : 0.0f;
 	float SpawnStartTime = DeltaTime + OldLeftover * SliceIncrement - SliceIncrement;
 	SpawnFraction = NewLeftover - SpawnNumber;
@@ -3332,11 +3332,11 @@ float FParticleAnimTrailEmitterInstance::Spawn(float DeltaTime)
 	{
 		if (DeltaTime < 0.25f)
 		{
-			bProcessSpawn = Resize(TotalCount + FMath::Trunc(FMath::Sqrt((float)TotalCount)) + 1);
+			bProcessSpawn = Resize(TotalCount + FMath::TruncToInt(FMath::Sqrt((float)TotalCount)) + 1);
 		}
 		else
 		{
-			bProcessSpawn = Resize((TotalCount + FMath::Trunc(FMath::Sqrt((float)TotalCount)) + 1), false);
+			bProcessSpawn = Resize((TotalCount + FMath::TruncToInt(FMath::Sqrt((float)TotalCount)) + 1), false);
 		}
 	}
 
@@ -3747,7 +3747,7 @@ void FParticleAnimTrailEmitterInstance::DetermineVertexAndTriangleCount()
 						// Determine the number of rendered interpolated points between these two particles
 						float CheckDistance = (CurrParticle->Location - PrevParticle->Location).Size();
 						float DistDiff = CheckDistance / TrailTypeData->DistanceTessellationStepSize;
-						InterpCount += FMath::Trunc(DistDiff);
+						InterpCount += FMath::TruncToInt(DistDiff);
 					}
 					
 					if (bApplyTangentTessellation )
@@ -3766,7 +3766,7 @@ void FParticleAnimTrailEmitterInstance::DetermineVertexAndTriangleCount()
 							CheckTangent = FMath::Max( CheckTangent, NextCheckTangent );
 						}		
 						float TangDiff = CheckTangent / TangentTessellationStepSize;
-						InterpCount += FMath::Trunc(TangDiff);
+						InterpCount += FMath::TruncToInt(TangDiff);
 					}
 
 					if( bApplyWidthTessellation )
@@ -3778,7 +3778,7 @@ void FParticleAnimTrailEmitterInstance::DetermineVertexAndTriangleCount()
 							CheckWidth = FMath::Max( CheckWidth, FMath::Abs(NextTrailData->Length - CurrTrailData->Length) );
 						}
 						float WidthDiff = CheckWidth / TrailTypeData->WidthTessellationStepSize;
-						InterpCount += FMath::Trunc(WidthDiff);
+						InterpCount += FMath::TruncToInt(WidthDiff);
 					}
 
 					// Store off the rendering interp count for this particle
diff --git a/Engine/Source/Runtime/Engine/Private/PhysicsEngine/PhysSubstepTasks.cpp b/Engine/Source/Runtime/Engine/Private/PhysicsEngine/PhysSubstepTasks.cpp
--- a/Engine/Source/Runtime/Engine/Private/PhysicsEngine/PhysSubstepTasks.cpp
+++ b/Engine/Source/Runtime/Engine/Private/PhysicsEngine/PhysSubstepTasks.cpp
@@ -247,7 +247,7 @@ float FPhysSubstepTask::UpdateTime(float UseDelta)
 
 	//Figure out how big dt to make for desired framerate
 	DeltaSeconds = UseDelta;
-	NumSubsteps = FMath::Ceil(UseDelta * FrameRateInv);
+	NumSubsteps = FMath::CeilToInt(UseDelta * FrameRateInv);
 	NumSubsteps = FMath::Max(NumSubsteps > MaxSubSteps ? MaxSubSteps : NumSubsteps, (uint32) 1);
 	SubTime = UseDelta / NumSubsteps;
 
diff --git a/Engine/Source/Runtime/Engine/Private/PlayerState.cpp b/Engine/Source/Runtime/Engine/Private/PlayerState.cpp
--- a/Engine/Source/Runtime/Engine/Private/PlayerState.cpp
+++ b/Engine/Source/Runtime/Engine/Private/PlayerState.cpp
@@ -45,13 +45,13 @@ void APlayerState::UpdatePing(float InPing)
 		CurPingBucketTimestamp = CurTime;
 
 
-		PingBucket[CurPingBucket].PingSum = FMath::Floor(InPing * 1000.f);
+		PingBucket[CurPingBucket].PingSum = FMath::FloorToInt(InPing * 1000.f);
 		PingBucket[CurPingBucket].PingCount = 1;
 	}
 	// Limit the number of pings we accept per-bucket, to avoid overflowing PingBucket values
 	else if (PingBucket[CurPingBucket].PingCount < 7)
 	{
-		PingBucket[CurPingBucket].PingSum += FMath::Floor(InPing * 1000.f);
+		PingBucket[CurPingBucket].PingSum += FMath::FloorToInt(InPing * 1000.f);
 		PingBucket[CurPingBucket].PingCount++;
 	}
 }
diff --git a/Engine/Source/Runtime/Engine/Private/PointLightSceneProxy.h b/Engine/Source/Runtime/Engine/Private/PointLightSceneProxy.h
--- a/Engine/Source/Runtime/Engine/Private/PointLightSceneProxy.h
+++ b/Engine/Source/Runtime/Engine/Private/PointLightSceneProxy.h
@@ -47,11 +47,11 @@ static bool GetPointLightBounds(
 			float X = (Dot3(P,Axis) / P.W + 1.0f * AxisSign) / 2.0f * AxisSign;
 			if(FMath::IsNegativeFloat(Nxa) ^ FMath::IsNegativeFloat(AxisSign))
 			{
-				OutMaxX = FMath::Min<int64>(FMath::Ceil(ViewSizeX * X + ViewX),OutMaxX);
+				OutMaxX = FMath::Min<int64>(FMath::CeilToInt(ViewSizeX * X + ViewX),OutMaxX);
 			}
 			else
 			{
-				OutMinX = FMath::Max<int64>(FMath::Floor(ViewSizeX * X + ViewX),OutMinX);
+				OutMinX = FMath::Max<int64>(FMath::FloorToInt(ViewSizeX * X + ViewX),OutMinX);
 			}
 		}
 
@@ -63,11 +63,11 @@ static bool GetPointLightBounds(
 			float X = (Dot3(P,Axis) / P.W + 1.0f * AxisSign) / 2.0f * AxisSign;
 			if(FMath::IsNegativeFloat(Nxb) ^ FMath::IsNegativeFloat(AxisSign))
 			{
-				OutMaxX = FMath::Min<int64>(FMath::Ceil(ViewSizeX * X + ViewX),OutMaxX);
+				OutMaxX = FMath::Min<int64>(FMath::CeilToInt(ViewSizeX * X + ViewX),OutMaxX);
 			}
 			else
 			{
-				OutMinX = FMath::Max<int64>(FMath::Floor(ViewSizeX * X + ViewX),OutMinX);
+				OutMinX = FMath::Max<int64>(FMath::FloorToInt(ViewSizeX * X + ViewX),OutMinX);
 			}
 		}
 	}
diff --git a/Engine/Source/Runtime/Engine/Private/PrimitiveDrawingUtils.cpp b/Engine/Source/Runtime/Engine/Private/PrimitiveDrawingUtils.cpp
--- a/Engine/Source/Runtime/Engine/Private/PrimitiveDrawingUtils.cpp
+++ b/Engine/Source/Runtime/Engine/Private/PrimitiveDrawingUtils.cpp
@@ -1062,7 +1062,7 @@ void DrawDashedLine(FPrimitiveDrawInterface* PDI, const FVector& Start, const FV
 		LineDir /= LineLeft;
 	}
 
-	const int32 nLines = FMath::Ceil(LineLeft / (DashSize*2));
+	const int32 nLines = FMath::CeilToInt(LineLeft / (DashSize*2));
 	PDI->AddReserveLines(DepthPriority, nLines, DepthBias != 0);
 
 	const FVector Dash = (DashSize * LineDir);
@@ -1447,11 +1447,11 @@ void ClampUVs(FVector2D* UVs, int32 NumUVs)
 
 	if (MinU < -FudgeFactor || MinU > (1.0f+FudgeFactor))
 	{
-		Bias.X = FMath::Floor(MinU);
+		Bias.X = FMath::FloorToFloat(MinU);
 	}
 	if (MinV < -FudgeFactor || MinV > (1.0f+FudgeFactor))
 	{
-		Bias.Y = FMath::Floor(MinV);
+		Bias.Y = FMath::FloorToFloat(MinV);
 	}
 
 	for (int32 i = 0; i < NumUVs; i++)
diff --git a/Engine/Source/Runtime/Engine/Private/ReflectionCaptureComponent.cpp b/Engine/Source/Runtime/Engine/Private/ReflectionCaptureComponent.cpp
--- a/Engine/Source/Runtime/Engine/Private/ReflectionCaptureComponent.cpp
+++ b/Engine/Source/Runtime/Engine/Private/ReflectionCaptureComponent.cpp
@@ -288,10 +288,10 @@ FColor RGBMEncode( FLinearColor Color )
 		MaxValue = Tonemapped;
 	}
 
-	Encoded.A = FMath::Min( FMath::Ceil( MaxValue * 255.0f ), 255 );
-	Encoded.R = FMath::Round( ( Color.R * 255.0f / Encoded.A ) * 255.0f );
-	Encoded.G = FMath::Round( ( Color.G * 255.0f / Encoded.A ) * 255.0f );
-	Encoded.B = FMath::Round( ( Color.B * 255.0f / Encoded.A ) * 255.0f );
+	Encoded.A = FMath::Min( FMath::CeilToInt( MaxValue * 255.0f ), 255 );
+	Encoded.R = FMath::RoundToInt( ( Color.R * 255.0f / Encoded.A ) * 255.0f );
+	Encoded.G = FMath::RoundToInt( ( Color.G * 255.0f / Encoded.A ) * 255.0f );
+	Encoded.B = FMath::RoundToInt( ( Color.B * 255.0f / Encoded.A ) * 255.0f );
 
 	return Encoded;
 }
diff --git a/Engine/Source/Runtime/Engine/Private/SceneView.cpp b/Engine/Source/Runtime/Engine/Private/SceneView.cpp
--- a/Engine/Source/Runtime/Engine/Private/SceneView.cpp
+++ b/Engine/Source/Runtime/Engine/Private/SceneView.cpp
@@ -467,8 +467,8 @@ void FSceneView::DeprojectFVector2D(const FVector2D& ScreenPos, FVector& out_Wor
 
 void FSceneView::DeprojectScreenToWorld(const FVector2D& ScreenPos, const FIntRect& ViewRect, const FMatrix& InvViewMatrix, const FMatrix& InvProjectionMatrix, FVector& out_WorldOrigin, FVector& out_WorldDirection)
 {
-	int32 PixelX = FMath::Trunc(ScreenPos.X);
-	int32 PixelY = FMath::Trunc(ScreenPos.Y);
+	int32 PixelX = FMath::TruncToInt(ScreenPos.X);
+	int32 PixelY = FMath::TruncToInt(ScreenPos.Y);
 
 	// Get the eye position and direction of the mouse cursor in two stages (inverse transform projection, then inverse transform view).
 	// This avoids the numerical instability that occurs when a view matrix with large translation is composed with a projection matrix
@@ -1197,8 +1197,8 @@ void FSceneViewFamily::ComputeFamilySize()
 
 	// We render to the actual position of the viewports so with black borders we need the max.
 	// We could change it by rendering all to left top but that has implications for splitscreen. 
-	FamilySizeX = FMath::Trunc(MaxFamilyX);
-	FamilySizeY = FMath::Trunc(MaxFamilyY);	
+	FamilySizeX = FMath::TruncToInt(MaxFamilyX);
+	FamilySizeY = FMath::TruncToInt(MaxFamilyY);	
 
 	check(bInitializedExtents);
 }
diff --git a/Engine/Source/Runtime/Engine/Private/SkeletalMesh.cpp b/Engine/Source/Runtime/Engine/Private/SkeletalMesh.cpp
--- a/Engine/Source/Runtime/Engine/Private/SkeletalMesh.cpp
+++ b/Engine/Source/Runtime/Engine/Private/SkeletalMesh.cpp
@@ -2021,7 +2021,7 @@ static void GetStreamingTextureFactorForLOD(FStaticLODModel& LODModel, TArray<fl
 					// Disregard upper 75% of texel ratios.
 					// This is to ignore backfacing surfaces or other non-visible surfaces that tend to map a small number of texels to a large surface.
 					TexelRatios[UVIndex].Sort( TGreater<float>() );
-					float TexelRatio = TexelRatios[UVIndex][ FMath::Trunc(TexelRatios[UVIndex].Num() * 0.75f) ];
+					float TexelRatio = TexelRatios[UVIndex][ FMath::TruncToInt(TexelRatios[UVIndex].Num() * 0.75f) ];
 					if ( UVIndex == 0 )
 					{
 						TexelRatio *= StreamingDistanceMultiplier;
@@ -4264,7 +4264,7 @@ void FSkeletalMeshSceneProxy::DrawDynamicElementsSection(FPrimitiveDrawInterface
 	{
 		if( Mesh.MaterialRenderProxy->GetMaterial(GRHIFeatureLevel)->GetBlendMode() == BLEND_Translucent )
 		{
-			BatchElement.NumPrimitives = FMath::Round(((float)Section.NumTriangles)*FMath::Clamp<float>(MeshObject->ProgressiveDrawingFraction,0.f,1.f));
+			BatchElement.NumPrimitives = FMath::RoundToInt(((float)Section.NumTriangles)*FMath::Clamp<float>(MeshObject->ProgressiveDrawingFraction,0.f,1.f));
 			if( BatchElement.NumPrimitives == 0 )
 			{
 				return;
diff --git a/Engine/Source/Runtime/Engine/Private/Slate/SceneViewport.cpp b/Engine/Source/Runtime/Engine/Private/Slate/SceneViewport.cpp
--- a/Engine/Source/Runtime/Engine/Private/Slate/SceneViewport.cpp
+++ b/Engine/Source/Runtime/Engine/Private/Slate/SceneViewport.cpp
@@ -245,7 +245,7 @@ void FSceneViewport::OnDrawViewport( const FGeometry& AllottedGeometry, const FS
 	FScopedConditionalWorldSwitcher WorldSwitcher( ViewportClient );
 
 	/** Check to see if the viewport should be resized */
-	FIntPoint DrawSize = FIntPoint( FMath::Trunc( AllottedGeometry.GetDrawSize().X ), FMath::Trunc( AllottedGeometry.GetDrawSize().Y ) );
+	FIntPoint DrawSize = FIntPoint( FMath::TruncToInt( AllottedGeometry.GetDrawSize().X ), FMath::TruncToInt( AllottedGeometry.GetDrawSize().Y ) );
 	if( GetSizeXY() != DrawSize )
 	{
 		TSharedPtr<SWindow> Window = FSlateApplication::Get().FindWidgetWindow( ViewportWidget.Pin().ToSharedRef() );
@@ -258,10 +258,10 @@ void FSceneViewport::OnDrawViewport( const FGeometry& AllottedGeometry, const FS
 	float CanvasMinX = FMath::Max(0.0f, AllottedGeometry.AbsolutePosition.X);
 	float CanvasMinY = FMath::Max(0.0f, AllottedGeometry.AbsolutePosition.Y);
 	FIntRect CanvasRect(
-		FMath::Trunc( CanvasMinX ),
-		FMath::Trunc( CanvasMinY ),
-		FMath::Trunc( CanvasMinX + AllottedGeometry.Size.X * AllottedGeometry.Scale ), 
-		FMath::Trunc( CanvasMinY + AllottedGeometry.Size.Y * AllottedGeometry.Scale ) );
+		FMath::TruncToInt( CanvasMinX ),
+		FMath::TruncToInt( CanvasMinY ),
+		FMath::TruncToInt( CanvasMinX + AllottedGeometry.Size.X * AllottedGeometry.Scale ), 
+		FMath::TruncToInt( CanvasMinY + AllottedGeometry.Size.Y * AllottedGeometry.Scale ) );
 
 
 	DebugCanvasDrawer->BeginRenderingCanvas( CanvasRect );
diff --git a/Engine/Source/Runtime/Engine/Private/Spline.cpp b/Engine/Source/Runtime/Engine/Private/Spline.cpp
--- a/Engine/Source/Runtime/Engine/Private/Spline.cpp
+++ b/Engine/Source/Runtime/Engine/Private/Spline.cpp
@@ -75,7 +75,7 @@ public:
 				}
 				else
 				{
-					int32 NumSteps = FMath::Ceil( (NewKeyTime - OldKeyTime)/SplineDrawRes );
+					int32 NumSteps = FMath::CeilToInt( (NewKeyTime - OldKeyTime)/SplineDrawRes );
 					float DrawSubstep = (NewKeyTime - OldKeyTime)/NumSteps;
 
 					// Find position on first keyframe.
diff --git a/Engine/Source/Runtime/Engine/Private/StaticMesh.cpp b/Engine/Source/Runtime/Engine/Private/StaticMesh.cpp
--- a/Engine/Source/Runtime/Engine/Private/StaticMesh.cpp
+++ b/Engine/Source/Runtime/Engine/Private/StaticMesh.cpp
@@ -1344,7 +1344,7 @@ void UStaticMesh::GetAssetRegistryTags(TArray<FAssetRegistryTag>& OutTags) const
 	{
 		Bounds = RenderData->Bounds;
 	}
-	const FString ApproxSizeStr = FString::Printf(TEXT("%dx%dx%d"), FMath::Round(Bounds.BoxExtent.X * 2.0f), FMath::Round(Bounds.BoxExtent.Y * 2.0f), FMath::Round(Bounds.BoxExtent.Z * 2.0f));
+	const FString ApproxSizeStr = FString::Printf(TEXT("%dx%dx%d"), FMath::RoundToInt(Bounds.BoxExtent.X * 2.0f), FMath::RoundToInt(Bounds.BoxExtent.Y * 2.0f), FMath::RoundToInt(Bounds.BoxExtent.Z * 2.0f));
 
 	OutTags.Add( FAssetRegistryTag("Triangles", FString::FromInt(NumTriangles), FAssetRegistryTag::TT_Numerical) );
 	OutTags.Add( FAssetRegistryTag("Vertices", FString::FromInt(NumVertices), FAssetRegistryTag::TT_Numerical) );
diff --git a/Engine/Source/Runtime/Engine/Private/StaticMeshComponent.cpp b/Engine/Source/Runtime/Engine/Private/StaticMeshComponent.cpp
--- a/Engine/Source/Runtime/Engine/Private/StaticMeshComponent.cpp
+++ b/Engine/Source/Runtime/Engine/Private/StaticMeshComponent.cpp
@@ -1307,14 +1307,14 @@ void UStaticMeshComponent::GetTextureLightAndShadowMapMemoryUsage(int32 InWidth,
 {
 	// Stored in texture.
 	const float MIP_FACTOR = 1.33f;
-	OutShadowMapMemoryUsage = FMath::Trunc(MIP_FACTOR * InWidth * InHeight); // G8
+	OutShadowMapMemoryUsage = FMath::TruncToInt(MIP_FACTOR * InWidth * InHeight); // G8
 	if( AllowHighQualityLightmaps() )
 	{
-		OutLightMapMemoryUsage = FMath::Trunc(NUM_HQ_LIGHTMAP_COEF * MIP_FACTOR * InWidth * InHeight); // DXT5
+		OutLightMapMemoryUsage = FMath::TruncToInt(NUM_HQ_LIGHTMAP_COEF * MIP_FACTOR * InWidth * InHeight); // DXT5
 	}
 	else
 	{
-		OutLightMapMemoryUsage = FMath::Trunc(NUM_LQ_LIGHTMAP_COEF * MIP_FACTOR * InWidth * InHeight / 2); // DXT1
+		OutLightMapMemoryUsage = FMath::TruncToInt(NUM_LQ_LIGHTMAP_COEF * MIP_FACTOR * InWidth * InHeight / 2); // DXT1
 	}
 }
 
diff --git a/Engine/Source/Runtime/Engine/Private/SubtitleManager.cpp b/Engine/Source/Runtime/Engine/Private/SubtitleManager.cpp
--- a/Engine/Source/Runtime/Engine/Private/SubtitleManager.cpp
+++ b/Engine/Source/Runtime/Engine/Private/SubtitleManager.cpp
@@ -113,7 +113,7 @@ void FSubtitleManager::DisplaySubtitle( FCanvas* Canvas, FActiveSubtitle* Subtit
 
 	float FontHeight = GEngine->GetSubtitleFont()->GetMaxCharHeight();
 	float HeightTest = Canvas->GetRenderTarget()->GetSizeXY().Y;
-	int32 SubtitleHeight = FMath::Trunc( ( FontHeight * MULTILINE_SPACING_SCALING ) );
+	int32 SubtitleHeight = FMath::TruncToInt( ( FontHeight * MULTILINE_SPACING_SCALING ) );
 	FIntRect BackgroundBoxOffset = DrawStringOutlineBoxOffset;
 
 	// Needed to add a drop shadow and doing all 4 sides was the only way to make them look correct.  If this shows up as a framerate hit we'll have to think of a different way of dealing with this.
diff --git a/Engine/Source/Runtime/Engine/Private/Texture.cpp b/Engine/Source/Runtime/Engine/Private/Texture.cpp
--- a/Engine/Source/Runtime/Engine/Private/Texture.cpp
+++ b/Engine/Source/Runtime/Engine/Private/Texture.cpp
@@ -612,7 +612,7 @@ int32 FTextureLODSettings::CalculateLODBias( int32 Width, int32 Height, int32 LO
 	}
 
 	// Calculate maximum number of miplevels.
-	int32 TextureMaxLOD	= FMath::CeilLogTwo( FMath::Trunc( FMath::Max( Width, Height ) ) );
+	int32 TextureMaxLOD	= FMath::CeilLogTwo( FMath::TruncToInt( FMath::Max( Width, Height ) ) );
 
 	// Calculate LOD bias.
 	int32 UsedLODBias	= LODGroupInfo.LODBias + LODBias + NumCinematicMipLevels;
@@ -630,8 +630,8 @@ int32 FTextureLODSettings::CalculateLODBias( int32 Width, int32 Height, int32 LO
 */
 void FTextureLODSettings::ComputeInGameMaxResolution(int32 LODBias, UTexture &Texture, uint32 &OutSizeX, uint32 &OutSizeY) const
 {
-	uint32 ImportedSizeX = FMath::Trunc(Texture.GetSurfaceWidth());
-	uint32 ImportedSizeY = FMath::Trunc(Texture.GetSurfaceHeight());
+	uint32 ImportedSizeX = FMath::TruncToInt(Texture.GetSurfaceWidth());
+	uint32 ImportedSizeY = FMath::TruncToInt(Texture.GetSurfaceHeight());
 	
 	const FTextureLODGroup& LODGroup = GetTextureLODGroup((TextureGroup)Texture.LODGroup);
 
diff --git a/Engine/Source/Runtime/Engine/Private/TimerManager.cpp b/Engine/Source/Runtime/Engine/Private/TimerManager.cpp
--- a/Engine/Source/Runtime/Engine/Private/TimerManager.cpp
+++ b/Engine/Source/Runtime/Engine/Private/TimerManager.cpp
@@ -308,7 +308,7 @@ void FTimerManager::Tick(float DeltaTime)
 
 			// Determine how many times the timer may have elapsed (e.g. for large DeltaTime on a short looping timer)
 			int32 const CallCount = CurrentlyExecutingTimer.bLoop ? 
-				FMath::Trunc( (InternalTime - CurrentlyExecutingTimer.ExpireTime) / CurrentlyExecutingTimer.Rate ) + 1
+				FMath::TruncToInt( (InternalTime - CurrentlyExecutingTimer.ExpireTime) / CurrentlyExecutingTimer.Rate ) + 1
 				: 1;
 
 			// Now call the function
diff --git a/Engine/Source/Runtime/Engine/Private/UnrealClient.cpp b/Engine/Source/Runtime/Engine/Private/UnrealClient.cpp
--- a/Engine/Source/Runtime/Engine/Private/UnrealClient.cpp
+++ b/Engine/Source/Runtime/Engine/Private/UnrealClient.cpp
@@ -349,7 +349,7 @@ int32 FStatUnitData::DrawStat(FViewport* InViewport, FCanvas* InCanvas, int32 In
 
 	// Render CPU thread and GPU frame times.
 	UFont* Font = (!FPlatformProperties::SupportsWindowedMode() && GEngine->GetMediumFont()) ? GEngine->GetMediumFont() : GEngine->GetSmallFont();
-	const int32 SafeZone = FPlatformProperties::SupportsWindowedMode() ? 0 : FMath::Trunc(InViewport->GetSizeXY().X * 0.05f);
+	const int32 SafeZone = FPlatformProperties::SupportsWindowedMode() ? 0 : FMath::TruncToInt(InViewport->GetSizeXY().X * 0.05f);
 
 	FColor Color;
 	int32 X3 = InViewport->GetSizeXY().X - SafeZone;
@@ -359,7 +359,7 @@ int32 FStatUnitData::DrawStat(FViewport* InViewport, FCanvas* InCanvas, int32 In
 	}
 	const int32 X2 = X3 - Font->GetStringSize(TEXT(" 000.00 ms "));
 	const int32 X1 = X2 - Font->GetStringSize(TEXT("Frame: "));
-	const int32 RowHeight = FMath::Trunc(Font->GetMaxCharHeight() * 1.1f);
+	const int32 RowHeight = FMath::TruncToInt(Font->GetMaxCharHeight() * 1.1f);
 	const bool bShowUnitTimeGraph = InViewport->GetClient() ? InViewport->GetClient()->IsStatEnabled(TEXT("UnitGraph")) : false;
 
 	// 0-34 ms: Green, 34-50 ms: Yellow, 50+ ms: Red
@@ -560,7 +560,7 @@ int32 FStatUnitData::DrawStat(FViewport* InViewport, FCanvas* InCanvas, int32 In
 				{
 					const int32 AlertPadding = 1;
 					float MaxValue = TimeValues[CurUnitIndex];
-					int32 MinCheckFrames = FMath::Min<int32>(FPlatformMath::Ceil((float)AlertPrintWidth / GraphHorizPixelsPerFrame) + 10, NumberOfSamples);
+					int32 MinCheckFrames = FMath::Min<int32>(FPlatformMath::CeilToInt((float)AlertPrintWidth / GraphHorizPixelsPerFrame) + 10, NumberOfSamples);
 					int32 CheckIndex = CurUnitIndex + 1;
 					for (; CheckIndex < MinCheckFrames; ++CheckIndex)
 					{
@@ -605,9 +605,9 @@ int32 FStatHitchesData::DrawStat(FViewport* InViewport, FCanvas* InCanvas, int32
 			if (GEngine->ActiveMatinee.IsValid())
 			{
 				float MatineeTime = GEngine->ActiveMatinee.Get()->InterpPosition;
-				float MatineeMM = FPlatformMath::TruncFloat(MatineeTime / 60.0f);
-				float MatineeSS = FPlatformMath::TruncFloat(MatineeTime - MatineeMM * 60.0f);
-				float MatineeMS = FPlatformMath::TruncFloat((MatineeTime - MatineeMM * 60.0f - MatineeSS) * 1000.0f);
+				float MatineeMM = FPlatformMath::TruncToFloat(MatineeTime / 60.0f);
+				float MatineeSS = FPlatformMath::TruncToFloat(MatineeTime - MatineeMM * 60.0f);
+				float MatineeMS = FPlatformMath::TruncToFloat((MatineeTime - MatineeMM * 60.0f - MatineeSS) * 1000.0f);
 				UE_LOG(LogClient, Warning, TEXT("HITCH @ %02dm:%02d.%03ds,%d,%d,%d"),
 					(int32)MatineeMM, (int32)MatineeSS, (int32)MatineeMS, int32(MatineeTime * 1000), int32(DeltaSeconds * 1000), Count++);
 			}
@@ -1377,15 +1377,15 @@ FIntRect FViewport::CalculateViewExtents(float AspectRatio, const FIntRect& View
 		if( AspectRatioDifference > 0.0f )
 		{
 			// Calculate desired Y size.
-			const int32 NewSizeY = FMath::Max(1, FMath::Round( CurrentSizeX / AdjustedAspectRatio ) );
-			Result.Min.Y = FMath::Round( 0.5f * (CurrentSizeY - NewSizeY) );
+			const int32 NewSizeY = FMath::Max(1, FMath::RoundToInt( CurrentSizeX / AdjustedAspectRatio ) );
+			Result.Min.Y = FMath::RoundToInt( 0.5f * (CurrentSizeY - NewSizeY) );
 			Result.Max.Y = Result.Min.Y + NewSizeY;
 		}
 		// Otherwise - will place bars on the sides.
 		else
 		{
-			const int32 NewSizeX = FMath::Max(1, FMath::Round( CurrentSizeY * AdjustedAspectRatio ) );
-			Result.Min.X = FMath::Round( 0.5f * (CurrentSizeX - NewSizeX) );
+			const int32 NewSizeX = FMath::Max(1, FMath::RoundToInt( CurrentSizeY * AdjustedAspectRatio ) );
+			Result.Min.X = FMath::RoundToInt( 0.5f * (CurrentSizeX - NewSizeX) );
 			Result.Max.X = Result.Min.X + NewSizeX;
 		}
 	}
diff --git a/Engine/Source/Runtime/Engine/Private/UnrealEngine.cpp b/Engine/Source/Runtime/Engine/Private/UnrealEngine.cpp
--- a/Engine/Source/Runtime/Engine/Private/UnrealEngine.cpp
+++ b/Engine/Source/Runtime/Engine/Private/UnrealEngine.cpp
@@ -5906,7 +5906,7 @@ bool UEngine::HandleLogoutStatLevelsCommand( const TCHAR* Cmd, FOutputDevice& Ar
 		}
 		else if( GetAsyncLoadPercentage( *LevelName ) >= 0 )
 		{
-			const int32 Percentage = FMath::Trunc( GetAsyncLoadPercentage( *LevelName ) );
+			const int32 Percentage = FMath::TruncToInt( GetAsyncLoadPercentage( *LevelName ) );
 			LevelName += FString::Printf(TEXT(" - %3i %%"), Percentage ); 
 		}
 
@@ -6087,10 +6087,10 @@ static void DrawProperty(UCanvas* CanvasObject, UObject* Obj, const FDebugDispla
 		int32 XL2 = XL;
 		if (TextLines.Num() > 0)
 		{
-			XL2 += FMath::Trunc(TextLines[0].LineExtent.X);
+			XL2 += FMath::TruncToInt(TextLines[0].LineExtent.X);
 			for (int32 i = 1; i < TextLines.Num(); i++)
 			{
-				XL2 = FMath::Max<int32>(XL2, FMath::Trunc(TextLines[i].LineExtent.X));
+				XL2 = FMath::Max<int32>(XL2, FMath::TruncToInt(TextLines[i].LineExtent.X));
 			}
 		}
 		Canvas->DrawTile(  X, Y, XL2 + 1, YL * FMath::Max<int>(TextLines.Num(), 1), 0, 0, CanvasObject->DefaultTexture->GetSizeX(), CanvasObject->DefaultTexture->GetSizeY(),
@@ -6367,7 +6367,7 @@ void DrawStatsHUD( UWorld* World, FViewport* Viewport, FCanvas* Canvas, UCanvas*
 
 	{
 		int32 X = (CanvasObject) ? CanvasObject->SizeX - FPSXOffset : Viewport->GetSizeXY().X - FPSXOffset; //??
-		int32 Y = (GEngine->IsStereoscopic3D()) ? FMath::Trunc(Viewport->GetSizeXY().Y * 0.40f) : FMath::Trunc(Viewport->GetSizeXY().Y * 0.20f);
+		int32 Y = (GEngine->IsStereoscopic3D()) ? FMath::TruncToInt(Viewport->GetSizeXY().Y * 0.40f) : FMath::TruncToInt(Viewport->GetSizeXY().Y * 0.20f);
 
 		//give the viewport first shot at drawing stats
 		Y = Viewport->DrawStatsHUD(Canvas, X, Y);
@@ -9850,7 +9850,7 @@ int32 UEngine::RenderStatFPS(UWorld* World, FViewport* Viewport, FCanvas* Canvas
 	FColor FPSColor = GAverageFPS < 20.0f ? FColor(255, 0, 0) : (GAverageFPS < 29.5f ? FColor(255, 255, 0) : FColor(0, 255, 0));
 
 	// Start drawing the various counters.
-	const int32 RowHeight = FMath::Trunc(Font->GetMaxCharHeight() * 1.1f);
+	const int32 RowHeight = FMath::TruncToInt(Font->GetMaxCharHeight() * 1.1f);
 	// Draw the FPS counter.
 	Canvas->DrawShadowedString(
 		X,
@@ -9910,7 +9910,7 @@ int32 UEngine::RenderStatSummary(UWorld* World, FViewport* Viewport, FCanvas* Ca
 		FColor(30, 144, 255)
 		);
 
-	const int32 RowHeight = FMath::Trunc(Font->GetMaxCharHeight() * 1.1f);
+	const int32 RowHeight = FMath::TruncToInt(Font->GetMaxCharHeight() * 1.1f);
 	Y += RowHeight;
 	return Y;
 }
@@ -9946,10 +9946,10 @@ int32 UEngine::RenderStatColorList(UWorld* World, FViewport* Viewport, FCanvas*
 {
 	UFont* Font = GetTinyFont();
 
-	const int32 LineHeight = FMath::Trunc(Font->GetMaxCharHeight());
+	const int32 LineHeight = FMath::TruncToInt(Font->GetMaxCharHeight());
 	const int32 ColorsNum = GColorList.GetColorsNum();
 	const int32 MaxLinesInColumn = 35;
-	const int32 ColumnsNum = FMath::Ceil((float)ColorsNum / (float)MaxLinesInColumn);
+	const int32 ColumnsNum = FMath::CeilToInt((float)ColorsNum / (float)MaxLinesInColumn);
 
 	Y += 16;
 	const int32 SavedY = Y;
@@ -10041,7 +10041,7 @@ int32 UEngine::RenderStatLevels(UWorld* World, FViewport* Viewport, FCanvas* Can
 		}
 		else if (GetAsyncLoadPercentage(*LevelName) >= 0)
 		{
-			const int32 Percentage = FMath::Trunc(GetAsyncLoadPercentage(*LevelName));
+			const int32 Percentage = FMath::TruncToInt(GetAsyncLoadPercentage(*LevelName));
 			LevelName += FString::Printf(TEXT(" - %3i %%"), Percentage);
 		}
 
@@ -10399,7 +10399,7 @@ int32 UEngine::RenderStatSoundWaves(UWorld* World, FViewport* Viewport, FCanvas*
 		R = G = B = 0;
 		int32 Max = AudioDevice->MaxChannels / 2;
 		float f = FMath::Clamp<float>((float)(ActiveInstances - Max) / (float)Max, 0.f, 1.f);
-		R = FMath::Trunc(f * 255);
+		R = FMath::TruncToInt(f * 255);
 		if (ActiveInstances > Max)
 		{
 			f = FMath::Clamp<float>((float)(Max - ActiveInstances) / (float)Max, 0.5f, 1.f);
@@ -10408,7 +10408,7 @@ int32 UEngine::RenderStatSoundWaves(UWorld* World, FViewport* Viewport, FCanvas*
 		{
 			f = 1.0f;
 		}
-		G = FMath::Trunc(f * 255);
+		G = FMath::TruncToInt(f * 255);
 
 		Canvas->DrawShadowedString(X, Y, *FString::Printf(TEXT(" Total: %i"), ActiveInstances), GetSmallFont(), FColor(R, G, B));
 		Y += 12;
@@ -10797,7 +10797,7 @@ int32 UEngine::RenderStatAI(UWorld* World, FViewport* Viewport, FCanvas* Canvas,
 
 	}
 
-	const int32 RowHeight = FMath::Trunc(Font->GetMaxCharHeight() * 1.1f);
+	const int32 RowHeight = FMath::TruncToInt(Font->GetMaxCharHeight() * 1.1f);
 	Canvas->DrawShadowedString(
 		X,
 		Y,
diff --git a/Engine/Source/Runtime/Engine/Private/UserInterface/Canvas.cpp b/Engine/Source/Runtime/Engine/Private/UserInterface/Canvas.cpp
--- a/Engine/Source/Runtime/Engine/Private/UserInterface/Canvas.cpp
+++ b/Engine/Source/Runtime/Engine/Private/UserInterface/Canvas.cpp
@@ -756,8 +756,8 @@ ENGINE_API void StringSize(UFont* Font,int32& XL,int32& YL,const TCHAR* Text)
 	// this functionality has been moved to a static function in UIString
 	FTextSizingParameters Parameters(Font,1.f,1.f);
 	UCanvas::CanvasStringSize(Parameters, Text);
-	XL = FMath::Trunc(Parameters.DrawXL);
-	YL = FMath::Trunc(Parameters.DrawYL);
+	XL = FMath::TruncToInt(Parameters.DrawXL);
+	YL = FMath::TruncToInt(Parameters.DrawYL);
 }
 
 /**
@@ -1104,8 +1104,8 @@ void UCanvas::UpdateSafeZoneData()
 		FDisplayMetrics DisplayMetrics;
 		FSlateApplication::Get().GetDisplayMetrics(DisplayMetrics);
 	
-		SafeZonePadX = FMath::Ceil(DisplayMetrics.TitleSafePaddingSize.X);
-		SafeZonePadY = FMath::Ceil(DisplayMetrics.TitleSafePaddingSize.Y);
+		SafeZonePadX = FMath::CeilToInt(DisplayMetrics.TitleSafePaddingSize.X);
+		SafeZonePadY = FMath::CeilToInt(DisplayMetrics.TitleSafePaddingSize.Y);
 
 		CachedDisplayWidth = DisplayMetrics.PrimaryDisplayWidth;
 		CachedDisplayHeight = DisplayMetrics.PrimaryDisplayHeight;
@@ -1141,7 +1141,7 @@ void UCanvas::SetLinearDrawColor(FLinearColor InColor, float OpacityOverride)
 
 	if( OpacityOverride != -1 )
 	{
-		DrawColor.A = FMath::Clamp(FMath::Trunc(OpacityOverride * 255.0f),0,255);
+		DrawColor.A = FMath::Clamp(FMath::TruncToInt(OpacityOverride * 255.0f),0,255);
 	}
 }
 
@@ -1224,8 +1224,8 @@ void UCanvas::ClippedStrLen( UFont* Font, float ScaleX, float ScaleY, int32& XL,
 		FTextSizingParameters Parameters(Font,ScaleX,ScaleY);
 		CanvasStringSize(Parameters, Text);
 
-		XL = FMath::Trunc(Parameters.DrawXL);
-		YL = FMath::Trunc(Parameters.DrawYL);
+		XL = FMath::TruncToInt(Parameters.DrawXL);
+		YL = FMath::TruncToInt(Parameters.DrawYL);
 	}
 }
 
@@ -1248,7 +1248,7 @@ float UCanvas::DrawText(UFont* InFont, const FText& InText, float X, float Y, fl
 
 	if (RenderInfo.bClipText)
 	{
-		FCanvasTextItem TextItem(FVector2D(FMath::Trunc(OrgX + X), FMath::Trunc(OrgY + Y)), InText, InFont, DrawColor);
+		FCanvasTextItem TextItem(FVector2D(FMath::TruncToFloat(OrgX + X), FMath::TruncToFloat(OrgY + Y)), InText, InFont, DrawColor);
 		TextItem.Scale = FVector2D(XScale, YScale),
 			TextItem.BlendMode = SE_BLEND_Translucent;
 		TextItem.FontRenderInfo = RenderInfo;
@@ -1327,8 +1327,8 @@ int32 UCanvas::WrappedPrint(bool Draw, float X, float Y, int32& out_XL, int32& o
 		YL += Font->GetMaxCharHeight() * ScaleY;
 	}
 
-	out_XL = FMath::Trunc(XL);
-	out_YL = FMath::Trunc(YL);
+	out_XL = FMath::TruncToInt(XL);
+	out_YL = FMath::TruncToInt(YL);
 	return WrappedStrings.Num();
 }
 
diff --git a/Engine/Source/Runtime/Engine/Private/UserInterface/Console.cpp b/Engine/Source/Runtime/Engine/Private/UserInterface/Console.cpp
--- a/Engine/Source/Runtime/Engine/Private/UserInterface/Console.cpp
+++ b/Engine/Source/Runtime/Engine/Private/UserInterface/Console.cpp
@@ -938,7 +938,7 @@ bool UConsole::InputKey_Open( int32 ControllerId, FKey Key, EInputEvent Event, f
 void UConsole::PostRender_Console_Open(UCanvas* Canvas)
 {
 	// the height of the buffer will be 75% of the height of the screen
-	float Height = FMath::Floor(Canvas->ClipY * 0.75f);
+	float Height = FMath::FloorToFloat(Canvas->ClipY * 0.75f);
 
 	// shrink for TVs
 	float ClipX = Canvas->ClipX;
diff --git a/Engine/Source/Runtime/Engine/Private/UserInterface/PlayerInput.cpp b/Engine/Source/Runtime/Engine/Private/UserInterface/PlayerInput.cpp
--- a/Engine/Source/Runtime/Engine/Private/UserInterface/PlayerInput.cpp
+++ b/Engine/Source/Runtime/Engine/Private/UserInterface/PlayerInput.cpp
@@ -1084,7 +1084,7 @@ float UPlayerInput::SmoothMouse(float aMouse, float DeltaTime, uint8& SampleCoun
 	float DetectedMouseSampleHz = MouseSamples / MouseSamplingTotal;
 
 	// actual num samples should be either this or this+1
-	int32 ExpectedNumSamplesThisFrame = FMath::Trunc(DeltaTime * DetectedMouseSampleHz);
+	int32 ExpectedNumSamplesThisFrame = FMath::TruncToInt(DeltaTime * DetectedMouseSampleHz);
 
 	if (DeltaTime < 0.25f)
 	{
diff --git a/Engine/Source/Runtime/Engine/Private/VectorField.cpp b/Engine/Source/Runtime/Engine/Private/VectorField.cpp
--- a/Engine/Source/Runtime/Engine/Private/VectorField.cpp
+++ b/Engine/Source/Runtime/Engine/Private/VectorField.cpp
@@ -762,7 +762,7 @@ public:
 			FrameTime += AnimatedVectorField->FramesPerSecond * DeltaSeconds;
 
 			// Compute the two frames to lerp.
-			const int32 FrameA_Unclamped = FMath::Trunc(FrameTime);
+			const int32 FrameA_Unclamped = FMath::TruncToInt(FrameTime);
 			const int32 FrameA = AnimatedVectorField->bLoop ?
 				(FrameA_Unclamped % AnimatedVectorField->FrameCount) :
 				FMath::Min<int32>(FrameA_Unclamped, AnimatedVectorField->FrameCount - 1);
diff --git a/Engine/Source/Runtime/Engine/Private/VisualLog.cpp b/Engine/Source/Runtime/Engine/Private/VisualLog.cpp
--- a/Engine/Source/Runtime/Engine/Private/VisualLog.cpp
+++ b/Engine/Source/Runtime/Engine/Private/VisualLog.cpp
@@ -100,7 +100,7 @@ FVisLogEntry::FVisLogEntry(TSharedPtr<FJsonValue> FromJson)
 			if (JsonLogLine.IsValid())
 			{
 				LogLines[LogLineIndex].Category = FName(*(JsonLogLine->GetStringField(VisualLogJson::TAG_CATEGORY)));
-				LogLines[LogLineIndex].Verbosity = TEnumAsByte<ELogVerbosity::Type>((uint8)FMath::Trunc(JsonLogLine->GetNumberField(VisualLogJson::TAG_VERBOSITY)));
+				LogLines[LogLineIndex].Verbosity = TEnumAsByte<ELogVerbosity::Type>((uint8)FMath::TruncToInt(JsonLogLine->GetNumberField(VisualLogJson::TAG_VERBOSITY)));
 				LogLines[LogLineIndex].Line = JsonLogLine->GetStringField(VisualLogJson::TAG_LINE);
 			}
 		}
diff --git a/Engine/Source/Runtime/Engine/Private/World.cpp b/Engine/Source/Runtime/Engine/Private/World.cpp
--- a/Engine/Source/Runtime/Engine/Private/World.cpp
+++ b/Engine/Source/Runtime/Engine/Private/World.cpp
@@ -2560,11 +2560,11 @@ void UWorld::InitializeActorsForPlay(const FURL& InURL, bool bResetTime)
 		// Lock the level.
 		if(WorldType == EWorldType::Preview)
 		{
-			UE_LOG(LogWorld, Verbose,  TEXT("Bringing preview %s up for play (max tick rate %i) at %s"), *GetFullName(), FMath::Round(GEngine->GetMaxTickRate(0,false)), *FDateTime::Now().ToString() );
+			UE_LOG(LogWorld, Verbose,  TEXT("Bringing preview %s up for play (max tick rate %i) at %s"), *GetFullName(), FMath::RoundToInt(GEngine->GetMaxTickRate(0,false)), *FDateTime::Now().ToString() );
 		}
 		else
 		{
-			UE_LOG(LogWorld, Log,  TEXT("Bringing %s up for play (max tick rate %i) at %s"), *GetFullName(), FMath::Round(GEngine->GetMaxTickRate(0,false)), *FDateTime::Now().ToString() );
+			UE_LOG(LogWorld, Log,  TEXT("Bringing %s up for play (max tick rate %i) at %s"), *GetFullName(), FMath::RoundToInt(GEngine->GetMaxTickRate(0,false)), *FDateTime::Now().ToString() );
 		}
 
 		// Initialize network actors and start execution.
diff --git a/Engine/Source/Runtime/Engine/Public/AnimEncoding.h b/Engine/Source/Runtime/Engine/Public/AnimEncoding.h
--- a/Engine/Source/Runtime/Engine/Public/AnimEncoding.h
+++ b/Engine/Source/Runtime/Engine/Public/AnimEncoding.h
@@ -575,7 +575,7 @@ FORCEINLINE float AnimEncoding::TimeToIndex(
 					const float KeyPos = RelativePos * float(NumKeys);
 					checkSlow(KeyPos >= 0.0f);
 					const float KeyPosFloor = floorf(KeyPos);
-					PosIndex0CacheOut = FMath::Min( FMath::Trunc(KeyPosFloor), NumKeys );
+					PosIndex0CacheOut = FMath::Min( FMath::TruncToInt(KeyPosFloor), NumKeys );
 					AlphaCacheOut = KeyPos - KeyPosFloor;
 					PosIndex1CacheOut = FMath::Min( PosIndex0CacheOut + 1, NumKeys );
 				}
@@ -601,7 +601,7 @@ FORCEINLINE float AnimEncoding::TimeToIndex(
 						const float KeyPos = RelativePos * float(NumFrames);
 						checkSlow(KeyPos >= 0.0f);
 						const float KeyPosFloor = floorf(KeyPos);
-						PosIndex0CacheOut = FMath::Min( FMath::Trunc(KeyPosFloor), NumFrames - 1 );
+						PosIndex0CacheOut = FMath::Min( FMath::TruncToInt(KeyPosFloor), NumFrames - 1 );
 						AlphaCacheOut = KeyPos - KeyPosFloor;
 						PosIndex1CacheOut = PosIndex0CacheOut + 1;
 					}
@@ -621,7 +621,7 @@ FORCEINLINE float AnimEncoding::TimeToIndex(
 						float const KeyPos = AdjustedPosition * float(NumKeys-1);
 						checkSlow(KeyPos >= 0.0f);
 						float const KeyPosFloor = floorf(KeyPos);
-						PosIndex0CacheOut = FMath::Min( FMath::Trunc(KeyPosFloor), NumKeys - 1 );
+						PosIndex0CacheOut = FMath::Min( FMath::TruncToInt(KeyPosFloor), NumKeys - 1 );
 						AlphaCacheOut = KeyPos - KeyPosFloor;
 						PosIndex1CacheOut = FMath::Min( PosIndex0CacheOut + 1, NumKeys - 1 );
 					}
@@ -744,8 +744,8 @@ FORCEINLINE float AnimEncoding::TimeToIndex(
 		const int32 LastFrame= TotalFrames-1;
 		const float KeyPos = RelativePos * (float)LastKey;
 		const float FramePos = RelativePos * (float)TotalFrames;
-		const int32 FramePosFloor = FMath::Clamp(FMath::Trunc(FramePos), 0, LastFrame);
-		const int32 KeyEstimate = FMath::Clamp(FMath::Trunc(KeyPos), 0, LastKey);
+		const int32 FramePosFloor = FMath::Clamp(FMath::TruncToInt(FramePos), 0, LastFrame);
+		const int32 KeyEstimate = FMath::Clamp(FMath::TruncToInt(KeyPos), 0, LastKey);
 
 		int32 LowFrame = 0;
 		int32 HighFrame = 0;
diff --git a/Engine/Source/Runtime/Engine/Public/Distributions.h b/Engine/Source/Runtime/Engine/Public/Distributions.h
--- a/Engine/Source/Runtime/Engine/Public/Distributions.h
+++ b/Engine/Source/Runtime/Engine/Public/Distributions.h
@@ -110,7 +110,7 @@ struct FDistributionLookupTable
 		LerpAlpha = FMath::Fractional(Time);
 
 		// get the entries to lerp between
-		const uint32 Index = FMath::Trunc(Time);
+		const uint32 Index = FMath::TruncToInt(Time);
 		const uint32 Index1 = FMath::Min<uint32>( Index + 0, EntryCount - 1 ) * EntryStride;
 		const uint32 Index2 = FMath::Min<uint32>( Index + 1, EntryCount - 1 ) * EntryStride;
 		Entry1 = &Values[Index1];
diff --git a/Engine/Source/Runtime/Engine/Public/InstancedFoliage.h b/Engine/Source/Runtime/Engine/Public/InstancedFoliage.h
--- a/Engine/Source/Runtime/Engine/Public/InstancedFoliage.h
+++ b/Engine/Source/Runtime/Engine/Public/InstancedFoliage.h
@@ -244,7 +244,7 @@ private:
 
 	uint64 MakeKey( const FVector& Location )
 	{
-		return  MakeKey( FMath::Floor(Location.X) >> HashCellBits, FMath::Floor(Location.Y) >> HashCellBits );
+		return  MakeKey( FMath::FloorToInt(Location.X) >> HashCellBits, FMath::FloorToInt(Location.Y) >> HashCellBits );
 	}
 
 	// Locality map
@@ -274,10 +274,10 @@ public:
 
 	void GetInstancesOverlappingBox( const FBox& InBox, TSet<int32>& OutInstanceIndices )
 	{
-		int32 MinX = FMath::Floor(InBox.Min.X) >> HashCellBits;
-		int32 MinY = FMath::Floor(InBox.Min.Y) >> HashCellBits;
-		int32 MaxX = FMath::Floor(InBox.Max.X) >> HashCellBits;
-		int32 MaxY = FMath::Floor(InBox.Max.Y) >> HashCellBits;
+		int32 MinX = FMath::FloorToInt(InBox.Min.X) >> HashCellBits;
+		int32 MinY = FMath::FloorToInt(InBox.Min.Y) >> HashCellBits;
+		int32 MaxX = FMath::FloorToInt(InBox.Max.X) >> HashCellBits;
+		int32 MaxY = FMath::FloorToInt(InBox.Max.Y) >> HashCellBits;
 
 		for( int32 y=MinY;y<=MaxY;y++ )
 		{
diff --git a/Engine/Source/Runtime/Engine/Public/Landscape/LandscapeDataAccess.h b/Engine/Source/Runtime/Engine/Public/Landscape/LandscapeDataAccess.h
--- a/Engine/Source/Runtime/Engine/Public/Landscape/LandscapeDataAccess.h
+++ b/Engine/Source/Runtime/Engine/Public/Landscape/LandscapeDataAccess.h
@@ -139,8 +139,8 @@ struct FLandscapeComponentDataInterface
 
 		// Offset and stride for this component's data in heightmap texture
 		HeightmapStride = Component->HeightmapTexture->Source.GetSizeX() >> MipLevel;
-		HeightmapComponentOffsetX = FMath::Round( (float)(Component->HeightmapTexture->Source.GetSizeX() >> MipLevel) * Component->HeightmapScaleBias.Z );
-		HeightmapComponentOffsetY = FMath::Round( (float)(Component->HeightmapTexture->Source.GetSizeY() >> MipLevel) * Component->HeightmapScaleBias.W );
+		HeightmapComponentOffsetX = FMath::RoundToInt( (float)(Component->HeightmapTexture->Source.GetSizeX() >> MipLevel) * Component->HeightmapScaleBias.Z );
+		HeightmapComponentOffsetY = FMath::RoundToInt( (float)(Component->HeightmapTexture->Source.GetSizeY() >> MipLevel) * Component->HeightmapScaleBias.W );
 		HeightmapSubsectionOffset = (Component->SubsectionSizeQuads + 1) >> MipLevel;
 
 		ComponentSizeVerts = (Component->ComponentSizeQuads + 1) >> MipLevel;
diff --git a/Engine/Source/Runtime/Engine/Public/PrecomputedLightVolume.h b/Engine/Source/Runtime/Engine/Public/PrecomputedLightVolume.h
--- a/Engine/Source/Runtime/Engine/Public/PrecomputedLightVolume.h
+++ b/Engine/Source/Runtime/Engine/Public/PrecomputedLightVolume.h
@@ -42,7 +42,7 @@ public:
 
 	void SetPackedSkyBentNormal(FVector InSkyBentNormal)
 	{
-		PackedSkyBentNormal = FColor(FMath::Trunc((InSkyBentNormal.X * .5f + .5f) * 255.0f), FMath::Trunc((InSkyBentNormal.Y * .5f + .5f) * 255.0f), FMath::Trunc((InSkyBentNormal.Z * .5f + .5f) * 255.0f));
+		PackedSkyBentNormal = FColor(FMath::TruncToInt((InSkyBentNormal.X * .5f + .5f) * 255.0f), FMath::TruncToInt((InSkyBentNormal.Y * .5f + .5f) * 255.0f), FMath::TruncToInt((InSkyBentNormal.Z * .5f + .5f) * 255.0f));
 	}
 
 	inline FVector GetSkyBentNormalUnpacked() const
diff --git a/Engine/Source/Runtime/Engine/Public/Raster.h b/Engine/Source/Runtime/Engine/Public/Raster.h
--- a/Engine/Source/Runtime/Engine/Public/Raster.h
+++ b/Engine/Source/Runtime/Engine/Public/Raster.h
@@ -101,8 +101,8 @@ private:
 		bool BackFacing
 		)
 	{
-		int32	IntMinY = FMath::Clamp(FMath::Ceil(MinY),RasterPolicyType::GetMinY(),RasterPolicyType::GetMaxY() + 1),
-			IntMaxY = FMath::Clamp(FMath::Ceil(MaxY),RasterPolicyType::GetMinY(),RasterPolicyType::GetMaxY() + 1);
+		int32	IntMinY = FMath::Clamp(FMath::CeilToInt(MinY),RasterPolicyType::GetMinY(),RasterPolicyType::GetMaxY() + 1),
+			IntMaxY = FMath::Clamp(FMath::CeilToInt(MaxY),RasterPolicyType::GetMinY(),RasterPolicyType::GetMaxY() + 1);
 
 		for(int32 IntY = IntMinY;IntY < IntMaxY;IntY++)
 		{
@@ -120,8 +120,8 @@ private:
 
 			if(MaxX > MinX)
 			{
-				int32				IntMinX = FMath::Clamp(FMath::Ceil(MinX),RasterPolicyType::GetMinX(),RasterPolicyType::GetMaxX() + 1),
-								IntMaxX = FMath::Clamp(FMath::Ceil(MaxX),RasterPolicyType::GetMinX(),RasterPolicyType::GetMaxX() + 1);
+				int32				IntMinX = FMath::Clamp(FMath::CeilToInt(MinX),RasterPolicyType::GetMinX(),RasterPolicyType::GetMaxX() + 1),
+								IntMaxX = FMath::Clamp(FMath::CeilToInt(MaxX),RasterPolicyType::GetMinX(),RasterPolicyType::GetMaxX() + 1);
 				InterpolantType	DeltaInterpolant = (MaxInterpolant - MinInterpolant) / (MaxX - MinX);
 
 				for(int32 X = IntMinX;X < IntMaxX;X++)
diff --git a/Engine/Source/Runtime/Engine/Public/ShadowMap.h b/Engine/Source/Runtime/Engine/Public/ShadowMap.h
--- a/Engine/Source/Runtime/Engine/Public/ShadowMap.h
+++ b/Engine/Source/Runtime/Engine/Public/ShadowMap.h
@@ -37,12 +37,12 @@ struct FQuantizedSignedDistanceFieldShadowSample
 	{
 		if (Index == 0)
 		{
-			Distance = (uint8)FMath::Clamp<int32>(FMath::Trunc(InComponent * 255.0f),0,255);
+			Distance = (uint8)FMath::Clamp<int32>(FMath::TruncToInt(InComponent * 255.0f),0,255);
 		}
 		else
 		{
 			checkSlow(Index == 1);
-			PenumbraSize = (uint8)FMath::Clamp<int32>(FMath::Trunc(InComponent * 255.0f),0,255);
+			PenumbraSize = (uint8)FMath::Clamp<int32>(FMath::TruncToInt(InComponent * 255.0f),0,255);
 		}
 	}
 
@@ -57,8 +57,8 @@ struct FQuantizedSignedDistanceFieldShadowSample
 	FQuantizedSignedDistanceFieldShadowSample() {}
 	FQuantizedSignedDistanceFieldShadowSample(const FSignedDistanceFieldShadowSample& InSample)
 	{
-		Distance = (uint8)FMath::Clamp<int32>(FMath::Trunc(InSample.Distance * 255.0f),0,255);
-		PenumbraSize = (uint8)FMath::Clamp<int32>(FMath::Trunc(InSample.PenumbraSize * 255.0f),0,255);
+		Distance = (uint8)FMath::Clamp<int32>(FMath::TruncToInt(InSample.Distance * 255.0f),0,255);
+		PenumbraSize = (uint8)FMath::Clamp<int32>(FMath::TruncToInt(InSample.PenumbraSize * 255.0f),0,255);
 		Coverage = InSample.IsMapped ? 255 : 0;
 	}
 };
diff --git a/Engine/Source/Runtime/Launch/Private/LaunchEngineLoop.cpp b/Engine/Source/Runtime/Launch/Private/LaunchEngineLoop.cpp
--- a/Engine/Source/Runtime/Launch/Private/LaunchEngineLoop.cpp
+++ b/Engine/Source/Runtime/Launch/Private/LaunchEngineLoop.cpp
@@ -1709,7 +1709,7 @@ void FEngineLoop::InitTime()
 #endif // !UE_BUILD_SHIPPING
 
 	// convert FloatMaxTickTime into number of frames (using 1 / FApp::GetFixedDeltaTime() to convert fps to seconds )
-	MaxFrameCounter = FMath::Trunc(MaxTickTime / FApp::GetFixedDeltaTime());
+	MaxFrameCounter = FMath::TruncToInt(MaxTickTime / FApp::GetFixedDeltaTime());
 }
 
 void FlushStatsFrame(bool bDiscardCallstack)
diff --git a/Engine/Source/Runtime/Online/OnlineSubsystem/Public/OnlineJsonSerializer.h b/Engine/Source/Runtime/Online/OnlineSubsystem/Public/OnlineJsonSerializer.h
--- a/Engine/Source/Runtime/Online/OnlineSubsystem/Public/OnlineJsonSerializer.h
+++ b/Engine/Source/Runtime/Online/OnlineSubsystem/Public/OnlineJsonSerializer.h
@@ -270,7 +270,7 @@ public:
 	{
 		if (JsonObject->HasField(Name))
 		{
-			Value = FMath::Trunc(JsonObject->GetNumberField(Name));
+			Value = FMath::TruncToInt(JsonObject->GetNumberField(Name));
 		}
 	}
 	/**
@@ -283,7 +283,7 @@ public:
 	{
 		if (JsonObject->HasField(Name))
 		{
-			Value = FMath::Trunc(JsonObject->GetNumberField(Name));
+			Value = FMath::TruncToInt(JsonObject->GetNumberField(Name));
 		}
 	}
 	/**
diff --git a/Engine/Source/Runtime/Online/OnlineSubsystemUtils/Private/Tests/TestCloudInterface.cpp b/Engine/Source/Runtime/Online/OnlineSubsystemUtils/Private/Tests/TestCloudInterface.cpp
--- a/Engine/Source/Runtime/Online/OnlineSubsystemUtils/Private/Tests/TestCloudInterface.cpp
+++ b/Engine/Source/Runtime/Online/OnlineSubsystemUtils/Private/Tests/TestCloudInterface.cpp
@@ -188,8 +188,8 @@ void FTestCloudInterface::WriteNUserCloudFiles(const FUniqueNetId& InUserId, con
 	for (int32 FileIdx=0; FileIdx<FileCount; FileIdx++)
 	{
 		// @todo: make the dummy data size configurable on a per-platform basis, since different platforms (e.g. PS4) have different size limits
-		//WriteRandomFile(DummyData, FMath::Trunc(FMath::FRandRange(1024, 100*1024)));
-		WriteRandomFile(DummyData, FMath::Trunc(FMath::FRandRange(256, 1024)));
+		//WriteRandomFile(DummyData, FMath::TruncToInt(FMath::FRandRange(1024, 100*1024)));
+		WriteRandomFile(DummyData, FMath::TruncToInt(FMath::FRandRange(256, 1024)));
 		UserCloud->WriteUserFile(InUserId, FString::Printf(TEXT("%s%d.%s"), *FPaths::GetBaseFilename(FileName), FileIdx, *FPaths::GetExtension(FileName)), DummyData);
 	}
 }
@@ -207,7 +207,7 @@ void FTestCloudInterface::WriteNSharedCloudFiles(const FUniqueNetId& InUserId, c
 	TArray<uint8> DummyData;
 	for (int32 FileIdx=0; FileIdx<FileCount; FileIdx++)
 	{
-		WriteRandomFile(DummyData, FMath::Trunc(FMath::FRandRange(1024, 100*1024)));
+		WriteRandomFile(DummyData, FMath::TruncToInt(FMath::FRandRange(1024, 100*1024)));
 		SharedCloud->WriteSharedFile(InUserId, FString::Printf(TEXT("%s%d.%s"), *FPaths::GetBaseFilename(FileName), FileIdx, *FPaths::GetExtension(FileName)), DummyData);
 	}
 }
diff --git a/Engine/Source/Runtime/OpenGLDrv/Private/OpenGLDevice.cpp b/Engine/Source/Runtime/OpenGLDrv/Private/OpenGLDevice.cpp
--- a/Engine/Source/Runtime/OpenGLDrv/Private/OpenGLDevice.cpp
+++ b/Engine/Source/Runtime/OpenGLDrv/Private/OpenGLDevice.cpp
@@ -848,7 +848,7 @@ void FOpenGLDynamicRHI::Init()
 			float PoolSize = float(GPoolSizeVRAMPercentage) * 0.01f * float(GOpenGLTotalGraphicsMemory);
 
 			// Truncate GTexturePoolSize to MB (but still counted in bytes)
-			GTexturePoolSize = int64(FGenericPlatformMath::TruncFloat(PoolSize / 1024.0f / 1024.0f)) * 1024 * 1024;
+			GTexturePoolSize = int64(FGenericPlatformMath::TruncToInt(PoolSize / 1024.0f / 1024.0f)) * 1024 * 1024;
 
 			UE_LOG(LogRHI, Log, TEXT("Texture pool is %llu MB (%d%% of %llu MB)"),
 				GTexturePoolSize / 1024 / 1024,
diff --git a/Engine/Source/Runtime/OpenGLDrv/Private/OpenGLDrv.cpp b/Engine/Source/Runtime/OpenGLDrv/Private/OpenGLDrv.cpp
--- a/Engine/Source/Runtime/OpenGLDrv/Private/OpenGLDrv.cpp
+++ b/Engine/Source/Runtime/OpenGLDrv/Private/OpenGLDrv.cpp
@@ -142,7 +142,7 @@ void FOpenGLGPUProfiler::EndFrame()
 	{
 		uint64 GPUTiming = FrameTiming.GetTiming();
 		uint64 GPUFreq = FrameTiming.GetTimingFrequency();
-		GGPUFrameTime = FMath::Trunc( double(GPUTiming) / double(GPUFreq) / FPlatformTime::GetSecondsPerCycle() );
+		GGPUFrameTime = FMath::TruncToInt( double(GPUTiming) / double(GPUFreq) / FPlatformTime::GetSecondsPerCycle() );
 	}
 	else if (FOpenGLDisjointTimeStampQuery::IsSupported() && GNumActiveGPUsForRendering == 1)
 	{
@@ -152,7 +152,7 @@ void FOpenGLGPUProfiler::EndFrame()
 		int OldestQueryIndex = (CurrentGPUFrameQueryIndex + 1) % MAX_GPUFRAMEQUERIES;
 		if ( DisjointGPUFrameTimeQuery[OldestQueryIndex].IsResultValid() && DisjointGPUFrameTimeQuery[OldestQueryIndex].GetResult(&GPUTiming) )
 		{
-			GGPUFrameTime = FMath::Trunc( double(GPUTiming) / double(GPUFreq) / FPlatformTime::GetSecondsPerCycle() );
+			GGPUFrameTime = FMath::TruncToInt( double(GPUTiming) / double(GPUFreq) / FPlatformTime::GetSecondsPerCycle() );
 			GLastGPUFrameTime = GGPUFrameTime;
 		}
 		else
diff --git a/Engine/Source/Runtime/RenderCore/Private/RenderUtils.cpp b/Engine/Source/Runtime/RenderCore/Private/RenderUtils.cpp
--- a/Engine/Source/Runtime/RenderCore/Private/RenderUtils.cpp
+++ b/Engine/Source/Runtime/RenderCore/Private/RenderUtils.cpp
@@ -537,9 +537,9 @@ void FPackedPosition::Set( const FVector& InVector )
 #else
 	// Too confusing to use .5f - wanted to use the last bit!
 	// Change to int for easier read
-	Vector.X = FMath::Clamp<int32>(FMath::Trunc(InVector.X * 1023.0f),-1023,1023);
-	Vector.Y = FMath::Clamp<int32>(FMath::Trunc(InVector.Y * 1023.0f),-1023,1023);
-	Vector.Z = FMath::Clamp<int32>(FMath::Trunc(InVector.Z * 511.0f),-511,511);
+	Vector.X = FMath::Clamp<int32>(FMath::TruncToInt(InVector.X * 1023.0f),-1023,1023);
+	Vector.Y = FMath::Clamp<int32>(FMath::TruncToInt(InVector.Y * 1023.0f),-1023,1023);
+	Vector.Z = FMath::Clamp<int32>(FMath::TruncToInt(InVector.Z * 511.0f),-511,511);
 #endif
 }
 
diff --git a/Engine/Source/Runtime/RenderCore/Public/PackedNormal.inl b/Engine/Source/Runtime/RenderCore/Public/PackedNormal.inl
--- a/Engine/Source/Runtime/RenderCore/Public/PackedNormal.inl
+++ b/Engine/Source/Runtime/RenderCore/Public/PackedNormal.inl
@@ -5,18 +5,18 @@
 
 FORCEINLINE void FPackedNormal::operator=(const FVector& InVector)
 {
-	Vector.X = FMath::Clamp(FMath::Trunc(InVector.X * 127.5f + 127.5f),0,255);
-	Vector.Y = FMath::Clamp(FMath::Trunc(InVector.Y * 127.5f + 127.5f),0,255);
-	Vector.Z = FMath::Clamp(FMath::Trunc(InVector.Z * 127.5f + 127.5f),0,255);
+	Vector.X = FMath::Clamp(FMath::TruncToInt(InVector.X * 127.5f + 127.5f),0,255);
+	Vector.Y = FMath::Clamp(FMath::TruncToInt(InVector.Y * 127.5f + 127.5f),0,255);
+	Vector.Z = FMath::Clamp(FMath::TruncToInt(InVector.Z * 127.5f + 127.5f),0,255);
 	Vector.W = 128;
 }
 
 FORCEINLINE void FPackedNormal::operator=(const FVector4& InVector)
 {
-	Vector.X = FMath::Clamp(FMath::Trunc(InVector.X * 127.5f + 127.5f),0,255);
-	Vector.Y = FMath::Clamp(FMath::Trunc(InVector.Y * 127.5f + 127.5f),0,255);
-	Vector.Z = FMath::Clamp(FMath::Trunc(InVector.Z * 127.5f + 127.5f),0,255);
-	Vector.W = FMath::Clamp(FMath::Trunc(InVector.W * 127.5f + 127.5f),0,255);
+	Vector.X = FMath::Clamp(FMath::TruncToInt(InVector.X * 127.5f + 127.5f),0,255);
+	Vector.Y = FMath::Clamp(FMath::TruncToInt(InVector.Y * 127.5f + 127.5f),0,255);
+	Vector.Z = FMath::Clamp(FMath::TruncToInt(InVector.Z * 127.5f + 127.5f),0,255);
+	Vector.W = FMath::Clamp(FMath::TruncToInt(InVector.W * 127.5f + 127.5f),0,255);
 }
 
 FORCEINLINE bool FPackedNormal::operator==(const FPackedNormal& B) const
diff --git a/Engine/Source/Runtime/RenderCore/Public/RenderUtils.h b/Engine/Source/Runtime/RenderCore/Public/RenderUtils.h
--- a/Engine/Source/Runtime/RenderCore/Public/RenderUtils.h
+++ b/Engine/Source/Runtime/RenderCore/Public/RenderUtils.h
@@ -98,7 +98,7 @@ FORCEINLINE float GetBasisDeterminantSign( const FVector& XAxis, const FVector&
 */
 FORCEINLINE uint8 GetBasisDeterminantSignByte( const FPackedNormal& XAxis, const FPackedNormal& YAxis, const FPackedNormal& ZAxis )
 {
-	return FMath::Trunc(GetBasisDeterminantSign(XAxis,YAxis,ZAxis) * 127.5f + 127.5f);
+	return FMath::TruncToInt(GetBasisDeterminantSign(XAxis,YAxis,ZAxis) * 127.5f + 127.5f);
 }
 
 /** Information about a pixel format. */
diff --git a/Engine/Source/Runtime/Renderer/Private/DeferredShadingRenderer.cpp b/Engine/Source/Runtime/Renderer/Private/DeferredShadingRenderer.cpp
--- a/Engine/Source/Runtime/Renderer/Private/DeferredShadingRenderer.cpp
+++ b/Engine/Source/Runtime/Renderer/Private/DeferredShadingRenderer.cpp
@@ -1031,10 +1031,10 @@ void FDeferredShadingSceneRenderer::UpdateDownsampledDepthSurface()
 
 			PixelShader->SetParameters(View);
 
-			const uint32 DownsampledX = FMath::Trunc(View.ViewRect.Min.X / GSceneRenderTargets.GetSmallColorDepthDownsampleFactor());
-			const uint32 DownsampledY = FMath::Trunc(View.ViewRect.Min.Y / GSceneRenderTargets.GetSmallColorDepthDownsampleFactor());
-			const uint32 DownsampledSizeX = FMath::Trunc(View.ViewRect.Width() / GSceneRenderTargets.GetSmallColorDepthDownsampleFactor());
-			const uint32 DownsampledSizeY = FMath::Trunc(View.ViewRect.Height() / GSceneRenderTargets.GetSmallColorDepthDownsampleFactor());
+			const uint32 DownsampledX = FMath::TruncToInt(View.ViewRect.Min.X / GSceneRenderTargets.GetSmallColorDepthDownsampleFactor());
+			const uint32 DownsampledY = FMath::TruncToInt(View.ViewRect.Min.Y / GSceneRenderTargets.GetSmallColorDepthDownsampleFactor());
+			const uint32 DownsampledSizeX = FMath::TruncToInt(View.ViewRect.Width() / GSceneRenderTargets.GetSmallColorDepthDownsampleFactor());
+			const uint32 DownsampledSizeY = FMath::TruncToInt(View.ViewRect.Height() / GSceneRenderTargets.GetSmallColorDepthDownsampleFactor());
 
 			RHISetViewport(DownsampledX, DownsampledY, 0.0f, DownsampledX + DownsampledSizeX, DownsampledY + DownsampledSizeY, 1.0f);
 
diff --git a/Engine/Source/Runtime/Renderer/Private/IndirectLightingCache.cpp b/Engine/Source/Runtime/Renderer/Private/IndirectLightingCache.cpp
--- a/Engine/Source/Runtime/Renderer/Private/IndirectLightingCache.cpp
+++ b/Engine/Source/Runtime/Renderer/Private/IndirectLightingCache.cpp
@@ -217,9 +217,9 @@ void FIndirectLightingCache::CalculateBlockPositionAndSize(const FBoxSphereBound
 	RoundedBoundsSize.Z = FMath::LogX(BoundSizeRoundUpBase, Bounds.BoxExtent.Z * 2);
 
 	// Round up to the next integer exponent to provide stability even when Bounds.BoxExtent is changing
-	RoundedBoundsSize.X = FMath::Pow(BoundSizeRoundUpBase, FMath::Trunc(RoundedBoundsSize.X) + 1);
-	RoundedBoundsSize.Y = FMath::Pow(BoundSizeRoundUpBase, FMath::Trunc(RoundedBoundsSize.Y) + 1);
-	RoundedBoundsSize.Z = FMath::Pow(BoundSizeRoundUpBase, FMath::Trunc(RoundedBoundsSize.Z) + 1);
+	RoundedBoundsSize.X = FMath::Pow(BoundSizeRoundUpBase, FMath::TruncToInt(RoundedBoundsSize.X) + 1);
+	RoundedBoundsSize.Y = FMath::Pow(BoundSizeRoundUpBase, FMath::TruncToInt(RoundedBoundsSize.Y) + 1);
+	RoundedBoundsSize.Z = FMath::Pow(BoundSizeRoundUpBase, FMath::TruncToInt(RoundedBoundsSize.Z) + 1);
 
 	// For single sample allocations, use an effective texel size of 5 for snapping
 	const int32 EffectiveTexelSize = TexelSize > 2 ? TexelSize : 5;
@@ -233,9 +233,9 @@ void FIndirectLightingCache::CalculateBlockPositionAndSize(const FBoxSphereBound
 	const FVector BoundsMin = Bounds.Origin - Bounds.BoxExtent;
 
 	FVector SnappedMin;
-	SnappedMin.X = CellSize.X * FMath::Floor(BoundsMin.X / CellSize.X);
-	SnappedMin.Y = CellSize.Y * FMath::Floor(BoundsMin.Y / CellSize.Y);
-	SnappedMin.Z = CellSize.Z * FMath::Floor(BoundsMin.Z / CellSize.Z);
+	SnappedMin.X = CellSize.X * FMath::FloorToFloat(BoundsMin.X / CellSize.X);
+	SnappedMin.Y = CellSize.Y * FMath::FloorToFloat(BoundsMin.Y / CellSize.Y);
+	SnappedMin.Z = CellSize.Z * FMath::FloorToFloat(BoundsMin.Z / CellSize.Z);
 
 	if (TexelSize > 2)
 	{
@@ -778,9 +778,9 @@ void FIndirectLightingCache::InterpolateBlock(
 
 			// Number of cells to increment by for query blocks
 			FIntVector NumStepCells;
-			NumStepCells.X = FMath::Max(1, FMath::Floor(WorldTargetSize / WorldCellSize.X));
-			NumStepCells.Y = FMath::Max(1, FMath::Floor(WorldTargetSize / WorldCellSize.Y));
-			NumStepCells.Z = FMath::Max(1, FMath::Floor(WorldTargetSize / WorldCellSize.Z));
+			NumStepCells.X = FMath::Max(1, FMath::FloorToInt(WorldTargetSize / WorldCellSize.X));
+			NumStepCells.Y = FMath::Max(1, FMath::FloorToInt(WorldTargetSize / WorldCellSize.Y));
+			NumStepCells.Z = FMath::Max(1, FMath::FloorToInt(WorldTargetSize / WorldCellSize.Z));
 			FIntVector NumQueryStepCells(0, 0, 0);
 
 			// World space size to increment by for query blocks
diff --git a/Engine/Source/Runtime/Renderer/Private/LightFunctionRendering.cpp b/Engine/Source/Runtime/Renderer/Private/LightFunctionRendering.cpp
--- a/Engine/Source/Runtime/Renderer/Private/LightFunctionRendering.cpp
+++ b/Engine/Source/Runtime/Renderer/Private/LightFunctionRendering.cpp
@@ -168,7 +168,7 @@ static float GetLightFunctionFadeFraction(const FViewInfo& View, FSphere LightBo
 		FMath::Max(ScreenPosition.W, 1.0f);
 
 	static auto CVarShadowTexelsPerPixel = IConsoleManager::Get().FindTConsoleVariableDataFloat(TEXT("r.Shadow.TexelsPerPixel"));
-	const uint32 UnclampedResolution = FMath::Trunc(ScreenRadius * CVarShadowTexelsPerPixel->GetValueOnRenderThread());
+	const uint32 UnclampedResolution = FMath::TruncToInt(ScreenRadius * CVarShadowTexelsPerPixel->GetValueOnRenderThread());
 	extern float CalculateShadowFadeAlpha(int32 MaxUnclampedResolution, int32 ShadowFadeResolution, int32 MinShadowResolution);
 	const float ResolutionFadeAlpha = CalculateShadowFadeAlpha(UnclampedResolution, ShadowFadeResolution, MinShadowResolution);
 	return ResolutionFadeAlpha;
diff --git a/Engine/Source/Runtime/Renderer/Private/LightShaftRendering.cpp b/Engine/Source/Runtime/Renderer/Private/LightShaftRendering.cpp
--- a/Engine/Source/Runtime/Renderer/Private/LightShaftRendering.cpp
+++ b/Engine/Source/Runtime/Renderer/Private/LightShaftRendering.cpp
@@ -108,7 +108,7 @@ public:
 	void SetParameters(const ShaderRHIParamRef Shader, const FLightSceneInfo* LightSceneInfo, const FSceneView& View, TRefCountPtr<IPooledRenderTarget>& PassSource)
 	{
 		const uint32 DownsampleFactor = GetLightShaftDownsampleFactor();
-		FIntPoint DownSampledViewSize(FMath::Floor(View.ViewRect.Width() / DownsampleFactor), FMath::Floor(View.ViewRect.Height() / DownsampleFactor));
+		FIntPoint DownSampledViewSize(FMath::FloorToInt(View.ViewRect.Width() / DownsampleFactor), FMath::FloorToInt(View.ViewRect.Height() / DownsampleFactor));
 		const FIntPoint FilterBufferSize = GSceneRenderTargets.GetBufferSizeXY() / DownsampleFactor;
 
 		const FVector2D ViewRatioOfBuffer((float)DownSampledViewSize.X / FilterBufferSize.X, (float)DownSampledViewSize.Y / FilterBufferSize.Y);
diff --git a/Engine/Source/Runtime/Renderer/Private/PostProcess/PostProcessWeightedSampleSum.cpp b/Engine/Source/Runtime/Renderer/Private/PostProcess/PostProcessWeightedSampleSum.cpp
--- a/Engine/Source/Runtime/Renderer/Private/PostProcess/PostProcessWeightedSampleSum.cpp
+++ b/Engine/Source/Runtime/Renderer/Private/PostProcess/PostProcessWeightedSampleSum.cpp
@@ -617,5 +617,5 @@ float FRCPassPostProcessWeightedSampleSum::GetClampedKernelRadius( float KernelR
 
 int FRCPassPostProcessWeightedSampleSum::GetIntegerKernelRadius( float KernelRadius )
 {
-	return FMath::Min<int32>(FMath::Ceil(GetClampedKernelRadius(KernelRadius)), GetMaxNumSamples() - 1);
+	return FMath::Min<int32>(FMath::CeilToInt(GetClampedKernelRadius(KernelRadius)), GetMaxNumSamples() - 1);
 }
diff --git a/Engine/Source/Runtime/Renderer/Private/PostProcess/SceneRenderTargets.cpp b/Engine/Source/Runtime/Renderer/Private/PostProcess/SceneRenderTargets.cpp
--- a/Engine/Source/Runtime/Renderer/Private/PostProcess/SceneRenderTargets.cpp
+++ b/Engine/Source/Runtime/Renderer/Private/PostProcess/SceneRenderTargets.cpp
@@ -1101,7 +1101,7 @@ FIntPoint FSceneRenderTargets::GetPreShadowCacheTextureResolution() const
 	const int32 ExpandFactor = 2;
 
 	auto CVarPreShadowResolutionFactor = IConsoleManager::Get().FindTConsoleVariableDataFloat(TEXT("r.Shadow.PreShadowResolutionFactor"));
-	return FIntPoint(FMath::Trunc(ShadowDepthResolution.X * CVarPreShadowResolutionFactor->GetValueOnRenderThread()), FMath::Trunc(ShadowDepthResolution.Y * CVarPreShadowResolutionFactor->GetValueOnRenderThread())) * ExpandFactor;
+	return FIntPoint(FMath::TruncToInt(ShadowDepthResolution.X * CVarPreShadowResolutionFactor->GetValueOnRenderThread()), FMath::TruncToInt(ShadowDepthResolution.Y * CVarPreShadowResolutionFactor->GetValueOnRenderThread())) * ExpandFactor;
 }
 
 FIntPoint FSceneRenderTargets::GetTranslucentShadowDepthTextureResolution() const
diff --git a/Engine/Source/Runtime/Renderer/Private/ReflectionEnvironmentCapture.cpp b/Engine/Source/Runtime/Renderer/Private/ReflectionEnvironmentCapture.cpp
--- a/Engine/Source/Runtime/Renderer/Private/ReflectionEnvironmentCapture.cpp
+++ b/Engine/Source/Runtime/Renderer/Private/ReflectionEnvironmentCapture.cpp
@@ -977,7 +977,7 @@ void FScene::AllocateReflectionCaptures(const TArray<UReflectionCaptureComponent
 				float Exponent = FMath::LogX(MaxCubemapsRoundUpBase, ReflectionSceneData.AllocatedReflectionCapturesGameThread.Num());
 
 				// Round up to the next integer exponent to provide stability and reduce reallocations
-				DesiredMaxCubemaps = FMath::Pow(MaxCubemapsRoundUpBase, FMath::Trunc(Exponent) + 1);
+				DesiredMaxCubemaps = FMath::Pow(MaxCubemapsRoundUpBase, FMath::TruncToInt(Exponent) + 1);
 			}
 
 			DesiredMaxCubemaps = FMath::Min(DesiredMaxCubemaps, GMaxNumReflectionCaptures);
diff --git a/Engine/Source/Runtime/Renderer/Private/SceneOcclusion.cpp b/Engine/Source/Runtime/Renderer/Private/SceneOcclusion.cpp
--- a/Engine/Source/Runtime/Renderer/Private/SceneOcclusion.cpp
+++ b/Engine/Source/Runtime/Renderer/Private/SceneOcclusion.cpp
@@ -124,10 +124,10 @@ const uint8* FSceneViewState::GetPrecomputedVisibilityData(FViewInfo& View, cons
 		// Calculate the bucket that ViewOrigin falls into
 		// Cells are hashed into buckets to reduce search time
 		const float FloatOffsetX = (View.ViewMatrices.ViewOrigin.X - Handler.PrecomputedVisibilityCellBucketOriginXY.X) / Handler.PrecomputedVisibilityCellSizeXY;
-		// FMath::Trunc rounds toward 0, we want to always round down
-		const int32 BucketIndexX = FMath::Abs((FMath::Trunc(FloatOffsetX) - (FloatOffsetX < 0.0f ? 1 : 0)) / Handler.PrecomputedVisibilityCellBucketSizeXY % Handler.PrecomputedVisibilityNumCellBuckets);
+		// FMath::TruncToInt rounds toward 0, we want to always round down
+		const int32 BucketIndexX = FMath::Abs((FMath::TruncToInt(FloatOffsetX) - (FloatOffsetX < 0.0f ? 1 : 0)) / Handler.PrecomputedVisibilityCellBucketSizeXY % Handler.PrecomputedVisibilityNumCellBuckets);
 		const float FloatOffsetY = (View.ViewMatrices.ViewOrigin.Y -Handler.PrecomputedVisibilityCellBucketOriginXY.Y) / Handler.PrecomputedVisibilityCellSizeXY;
-		const int32 BucketIndexY = FMath::Abs((FMath::Trunc(FloatOffsetY) - (FloatOffsetY < 0.0f ? 1 : 0)) / Handler.PrecomputedVisibilityCellBucketSizeXY % Handler.PrecomputedVisibilityNumCellBuckets);
+		const int32 BucketIndexY = FMath::Abs((FMath::TruncToInt(FloatOffsetY) - (FloatOffsetY < 0.0f ? 1 : 0)) / Handler.PrecomputedVisibilityCellBucketSizeXY % Handler.PrecomputedVisibilityNumCellBuckets);
 		const int32 PrecomputedVisibilityBucketIndex = BucketIndexY * Handler.PrecomputedVisibilityCellBucketSizeXY + BucketIndexX;
 
 		check(PrecomputedVisibilityBucketIndex < Handler.PrecomputedVisibilityCellBuckets.Num());
@@ -970,10 +970,10 @@ void FDeferredShadingSceneRenderer::BeginOcclusionTests()
 
 		if (bUseDownsampledDepth)
 		{
-			const uint32 DownsampledX = FMath::Trunc(View.ViewRect.Min.X / GSceneRenderTargets.GetSmallColorDepthDownsampleFactor());
-			const uint32 DownsampledY = FMath::Trunc(View.ViewRect.Min.Y / GSceneRenderTargets.GetSmallColorDepthDownsampleFactor());
-			const uint32 DownsampledSizeX = FMath::Trunc(View.ViewRect.Width() / GSceneRenderTargets.GetSmallColorDepthDownsampleFactor());
-			const uint32 DownsampledSizeY = FMath::Trunc(View.ViewRect.Height() / GSceneRenderTargets.GetSmallColorDepthDownsampleFactor());
+			const uint32 DownsampledX = FMath::TruncToInt(View.ViewRect.Min.X / GSceneRenderTargets.GetSmallColorDepthDownsampleFactor());
+			const uint32 DownsampledY = FMath::TruncToInt(View.ViewRect.Min.Y / GSceneRenderTargets.GetSmallColorDepthDownsampleFactor());
+			const uint32 DownsampledSizeX = FMath::TruncToInt(View.ViewRect.Width() / GSceneRenderTargets.GetSmallColorDepthDownsampleFactor());
+			const uint32 DownsampledSizeY = FMath::TruncToInt(View.ViewRect.Height() / GSceneRenderTargets.GetSmallColorDepthDownsampleFactor());
 
 			// Setup the viewport for rendering to the downsampled depth buffer
 			RHISetViewport(DownsampledX, DownsampledY, 0.0f, DownsampledX + DownsampledSizeX, DownsampledY + DownsampledSizeY, 1.0f);
diff --git a/Engine/Source/Runtime/Renderer/Private/ShadowRendering.cpp b/Engine/Source/Runtime/Renderer/Private/ShadowRendering.cpp
--- a/Engine/Source/Runtime/Renderer/Private/ShadowRendering.cpp
+++ b/Engine/Source/Runtime/Renderer/Private/ShadowRendering.cpp
@@ -1819,7 +1819,7 @@ void FProjectedShadowInfo::RenderProjection(int32 ViewIndex, const FViewInfo* Vi
 			else if (LocalQuality > 2 && !bWholeSceneShadow)
 			{
 				auto CVarPreShadowResolutionFactor = IConsoleManager::Get().FindTConsoleVariableDataFloat(TEXT("r.Shadow.PreShadowResolutionFactor"));
-				const int32 TargetResolution = bPreShadow ? FMath::Trunc(512 * CVarPreShadowResolutionFactor->GetValueOnRenderThread()) : 512;
+				const int32 TargetResolution = bPreShadow ? FMath::TruncToInt(512 * CVarPreShadowResolutionFactor->GetValueOnRenderThread()) : 512;
 
 				int32 Reduce = 0;
 
diff --git a/Engine/Source/Runtime/Renderer/Private/ShadowSetup.cpp b/Engine/Source/Runtime/Renderer/Private/ShadowSetup.cpp
--- a/Engine/Source/Runtime/Renderer/Private/ShadowSetup.cpp
+++ b/Engine/Source/Runtime/Renderer/Private/ShadowSetup.cpp
@@ -446,7 +446,7 @@ FProjectedShadowInfo::FProjectedShadowInfo(
 
 		MinPreSubjectZ = Initializer.MinLightW;
 
-		ResolutionY = FMath::Min<uint32>(FMath::Trunc(InResolutionX / AspectRatio), MaxShadowResolutionY);
+		ResolutionY = FMath::Min<uint32>(FMath::TruncToInt(InResolutionX / AspectRatio), MaxShadowResolutionY);
 
 		// Store the view matrix
 		// Reorder the vectors to match the main view, since ShadowViewMatrix will be used to override the main view's view matrix during shadow depth rendering
@@ -1268,7 +1268,7 @@ void FDeferredShadingSceneRenderer::CreatePerObjectProjectedShadow(
 		MaxScreenPercent = FMath::Max(MaxScreenPercent, ScreenPercent);
 
 		// Determine the amount of shadow buffer resolution needed for this view.
-		const uint32 UnclampedResolution = FMath::Trunc(ScreenRadius * CVarShadowTexelsPerPixel.GetValueOnRenderThread());
+		const uint32 UnclampedResolution = FMath::TruncToInt(ScreenRadius * CVarShadowTexelsPerPixel.GetValueOnRenderThread());
 		MaxUnclampedResolution = FMath::Max( MaxUnclampedResolution, UnclampedResolution );
 		MaxDesiredResolution = FMath::Max(
 			MaxDesiredResolution,
@@ -1284,7 +1284,7 @@ void FDeferredShadingSceneRenderer::CreatePerObjectProjectedShadow(
 		MaxResolutionFadeAlpha = FMath::Max(MaxResolutionFadeAlpha, ViewSpecificAlpha);
 		ResolutionFadeAlphas.Add(ViewSpecificAlpha);
 
-		const float ViewSpecificPreShadowAlpha = CalculateShadowFadeAlpha( FMath::Trunc(UnclampedResolution * CVarPreShadowResolutionFactor.GetValueOnRenderThread()), CVarPreShadowFadeResolution.GetValueOnRenderThread(), CVarMinPreShadowResolution.GetValueOnRenderThread() );
+		const float ViewSpecificPreShadowAlpha = CalculateShadowFadeAlpha( FMath::TruncToInt(UnclampedResolution * CVarPreShadowResolutionFactor.GetValueOnRenderThread()), CVarPreShadowFadeResolution.GetValueOnRenderThread(), CVarMinPreShadowResolution.GetValueOnRenderThread() );
 		MaxResolutionPreShadowFadeAlpha = FMath::Max(MaxResolutionPreShadowFadeAlpha, ViewSpecificPreShadowAlpha);
 		ResolutionPreShadowFadeAlphas.Add(ViewSpecificPreShadowAlpha);
 	}
@@ -1406,7 +1406,7 @@ void FDeferredShadingSceneRenderer::CreatePerObjectProjectedShadow(
 			&& bOpaqueRelevance)
 		{
 			// Round down to the nearest power of two so that resolution changes are always doubling or halving the resolution, which increases filtering stability.
-			int32 PreshadowSizeX = 1 << (FMath::CeilLogTwo(FMath::Trunc(MaxDesiredResolution * CVarPreShadowResolutionFactor.GetValueOnRenderThread())) - 1);
+			int32 PreshadowSizeX = 1 << (FMath::CeilLogTwo(FMath::TruncToInt(MaxDesiredResolution * CVarPreShadowResolutionFactor.GetValueOnRenderThread())) - 1);
 
 			const FIntPoint PreshadowCacheResolution = GSceneRenderTargets.GetPreShadowCacheTextureResolution();
 			checkSlow(PreshadowSizeX <= PreshadowCacheResolution.X);
@@ -1451,7 +1451,7 @@ void FDeferredShadingSceneRenderer::CreatePerObjectProjectedShadow(
 						ShadowInitializer,
 						true,
 						PreshadowSizeX,
-						FMath::Trunc(MaxShadowResolutionY * CVarPreShadowResolutionFactor.GetValueOnRenderThread()),
+						FMath::TruncToInt(MaxShadowResolutionY * CVarPreShadowResolutionFactor.GetValueOnRenderThread()),
 						MaxScreenPercent,
 						ResolutionPreShadowFadeAlphas,
 						false
@@ -1526,7 +1526,7 @@ void FDeferredShadingSceneRenderer::CreateWholeSceneProjectedShadow(FLightSceneI
 				FMath::Max(ScreenPosition.W,1.0f);
 
 			// Determine the amount of shadow buffer resolution needed for this view.
-			const uint32 UnclampedResolution = FMath::Trunc(ScreenRadius * CVarShadowTexelsPerPixel.GetValueOnRenderThread());
+			const uint32 UnclampedResolution = FMath::TruncToInt(ScreenRadius * CVarShadowTexelsPerPixel.GetValueOnRenderThread());
 			MaxUnclampedResolution = FMath::Max( MaxUnclampedResolution, UnclampedResolution );
 			MaxDesiredResolution = FMath::Max(
 				MaxDesiredResolution,
@@ -1553,7 +1553,7 @@ void FDeferredShadingSceneRenderer::CreateWholeSceneProjectedShadow(FLightSceneI
 				// Round down to the nearest power of two so that resolution changes are always doubling or halving the resolution, which increases filtering stability
 				// Use the max resolution if the desired resolution is larger than that
 				int32 SizeX = MaxDesiredResolution >= MaxShadowResolution ? MaxShadowResolution : (1 << (FMath::CeilLogTwo(MaxDesiredResolution) - 1));
-				const uint32 DesiredSizeY = FMath::Trunc(MaxDesiredResolution);
+				const uint32 DesiredSizeY = FMath::TruncToInt(MaxDesiredResolution);
 				int32 SizeY = DesiredSizeY >= MaxShadowResolutionY ? MaxShadowResolutionY : (1 << (FMath::CeilLogTwo(DesiredSizeY) - 1));
 
 				if (ProjectedShadowInitializer.bOnePassPointLightShadow)
diff --git a/Engine/Source/Runtime/Renderer/Private/TranslucentLighting.cpp b/Engine/Source/Runtime/Renderer/Private/TranslucentLighting.cpp
--- a/Engine/Source/Runtime/Renderer/Private/TranslucentLighting.cpp
+++ b/Engine/Source/Runtime/Renderer/Private/TranslucentLighting.cpp
@@ -1362,13 +1362,13 @@ FVolumeBounds CalculateLightVolumeBounds(const FSphere& LightBounds, const FView
 		const FVector MinPosition = (LightBounds.Center - LightBounds.W - View.TranslucencyLightingVolumeMin[VolumeCascadeIndex]) / View.TranslucencyVolumeVoxelSize[VolumeCascadeIndex];
 		const FVector MaxPosition = (LightBounds.Center + LightBounds.W - View.TranslucencyLightingVolumeMin[VolumeCascadeIndex]) / View.TranslucencyVolumeVoxelSize[VolumeCascadeIndex];
 
-		VolumeBounds.MinX = FMath::Max(FMath::Trunc(MinPosition.X), 0);
-		VolumeBounds.MinY = FMath::Max(FMath::Trunc(MinPosition.Y), 0);
-		VolumeBounds.MinZ = FMath::Max(FMath::Trunc(MinPosition.Z), 0);
+		VolumeBounds.MinX = FMath::Max(FMath::TruncToInt(MinPosition.X), 0);
+		VolumeBounds.MinY = FMath::Max(FMath::TruncToInt(MinPosition.Y), 0);
+		VolumeBounds.MinZ = FMath::Max(FMath::TruncToInt(MinPosition.Z), 0);
 
-		VolumeBounds.MaxX = FMath::Min(FMath::Trunc(MaxPosition.X) + 1, GTranslucencyLightingVolumeDim);
-		VolumeBounds.MaxY = FMath::Min(FMath::Trunc(MaxPosition.Y) + 1, GTranslucencyLightingVolumeDim);
-		VolumeBounds.MaxZ = FMath::Min(FMath::Trunc(MaxPosition.Z) + 1, GTranslucencyLightingVolumeDim);
+		VolumeBounds.MaxX = FMath::Min(FMath::TruncToInt(MaxPosition.X) + 1, GTranslucencyLightingVolumeDim);
+		VolumeBounds.MaxY = FMath::Min(FMath::TruncToInt(MaxPosition.Y) + 1, GTranslucencyLightingVolumeDim);
+		VolumeBounds.MaxZ = FMath::Min(FMath::TruncToInt(MaxPosition.Z) + 1, GTranslucencyLightingVolumeDim);
 	}
 
 	return VolumeBounds;
diff --git a/Engine/Source/Runtime/Slate/Private/Framework/Application/SlateApplication.cpp b/Engine/Source/Runtime/Slate/Private/Framework/Application/SlateApplication.cpp
--- a/Engine/Source/Runtime/Slate/Private/Framework/Application/SlateApplication.cpp
+++ b/Engine/Source/Runtime/Slate/Private/Framework/Application/SlateApplication.cpp
@@ -2036,10 +2036,10 @@ void FSlateApplication::LockCursor( const TSharedPtr<SWidget>& Widget )
 				// Note: We round the upper left coordinate of the clip rect so we guarantee the rect is inside the geometry of the widget.  If we truncated when there is a half pixel we would cause the clip
 				// rect to be half a pixel larger than the geometry and cause the mouse to go outside of the geometry.
 				RECT ClipRect;
-				ClipRect.left = FMath::Round( SlateClipRect.Left );
-				ClipRect.top = FMath::Round( SlateClipRect.Top );
-				ClipRect.right = FMath::Trunc( SlateClipRect.Right );
-				ClipRect.bottom = FMath::Trunc( SlateClipRect.Bottom );
+				ClipRect.left = FMath::RoundToInt( SlateClipRect.Left );
+				ClipRect.top = FMath::RoundToInt( SlateClipRect.Top );
+				ClipRect.right = FMath::TruncToInt( SlateClipRect.Right );
+				ClipRect.bottom = FMath::TruncToInt( SlateClipRect.Bottom );
 
 				// Lock the mouse to the widget
 				PlatformApplication->Cursor->Lock( &ClipRect );
@@ -2864,10 +2864,10 @@ FSlateRect FSlateApplication::GetPreferredWorkArea() const
 FSlateRect FSlateApplication::GetWorkArea( const FSlateRect& InRect ) const
 {
 	FPlatformRect InPlatformRect;
-	InPlatformRect.Left = FMath::Trunc(InRect.Left);
-	InPlatformRect.Top = FMath::Trunc(InRect.Top);
-	InPlatformRect.Right = FMath::Trunc(InRect.Right);
-	InPlatformRect.Bottom = FMath::Trunc(InRect.Bottom);
+	InPlatformRect.Left = FMath::TruncToInt(InRect.Left);
+	InPlatformRect.Top = FMath::TruncToInt(InRect.Top);
+	InPlatformRect.Right = FMath::TruncToInt(InRect.Right);
+	InPlatformRect.Bottom = FMath::TruncToInt(InRect.Bottom);
 
 	const FPlatformRect OutPlatformRect = PlatformApplication->GetWorkArea( InPlatformRect );
 	return FSlateRect( OutPlatformRect.Left, OutPlatformRect.Top, OutPlatformRect.Right, OutPlatformRect.Bottom );
diff --git a/Engine/Source/Runtime/Slate/Private/Framework/Docking/TabManager.cpp b/Engine/Source/Runtime/Slate/Private/Framework/Docking/TabManager.cpp
--- a/Engine/Source/Runtime/Slate/Private/Framework/Docking/TabManager.cpp
+++ b/Engine/Source/Runtime/Slate/Private/Framework/Docking/TabManager.cpp
@@ -195,7 +195,7 @@ TSharedPtr<FTabManager::FLayout> FTabManager::FLayout::NewFromString( const FStr
 	{
 		const FString LayoutName = JsonObject->GetStringField(TEXT("Name"));
 		TSharedRef<FTabManager::FLayout> NewLayout = FTabManager::NewLayout( *LayoutName );
-		int32 PrimaryAreaIndex = FMath::Trunc( JsonObject->GetNumberField(TEXT("PrimaryAreaIndex")) );
+		int32 PrimaryAreaIndex = FMath::TruncToInt( JsonObject->GetNumberField(TEXT("PrimaryAreaIndex")) );
 
 		TArray< TSharedPtr<FJsonValue> > Areas = JsonObject->GetArrayField(TEXT("Areas"));
 		for(int32 AreaIndex=0; AreaIndex < Areas.Num(); ++AreaIndex)
diff --git a/Engine/Source/Runtime/Slate/Private/Framework/Layout/ScrollyZoomy.cpp b/Engine/Source/Runtime/Slate/Private/Framework/Layout/ScrollyZoomy.cpp
--- a/Engine/Source/Runtime/Slate/Private/Framework/Layout/ScrollyZoomy.cpp
+++ b/Engine/Source/Runtime/Slate/Private/Framework/Layout/ScrollyZoomy.cpp
@@ -89,8 +89,8 @@ FReply FScrollyZoomy::OnMouseButtonUp( const TSharedRef<SWidget> MyWidget, const
 			FVector2D CursorPosition = MyGeometry.LocalToAbsolute( SoftwareCursorPosition );
 
 			FIntPoint BestPositionInPanel(
-				FMath::Round( FMath::Clamp( CursorPosition.X, PanelScreenSpaceRect.Left, PanelScreenSpaceRect.Right ) ),
-				FMath::Round( FMath::Clamp( CursorPosition.Y, PanelScreenSpaceRect.Top, PanelScreenSpaceRect.Bottom ) )
+				FMath::RoundToInt( FMath::Clamp( CursorPosition.X, PanelScreenSpaceRect.Left, PanelScreenSpaceRect.Right ) ),
+				FMath::RoundToInt( FMath::Clamp( CursorPosition.Y, PanelScreenSpaceRect.Top, PanelScreenSpaceRect.Bottom ) )
 				);
 
 			Reply.SetMousePos( BestPositionInPanel );
diff --git a/Engine/Source/Runtime/Slate/Private/Framework/MultiBox/MultiBoxCustomization.cpp b/Engine/Source/Runtime/Slate/Private/Framework/MultiBox/MultiBoxCustomization.cpp
--- a/Engine/Source/Runtime/Slate/Private/Framework/MultiBox/MultiBoxCustomization.cpp
+++ b/Engine/Source/Runtime/Slate/Private/Framework/MultiBox/MultiBoxCustomization.cpp
@@ -123,8 +123,8 @@ void FMultiBoxCustomizationData::LoadCustomizedBlocks()
 			{
 				FString CommandName = SavedCommand->GetStringField( TEXT("CommandName") );
 				FString Context = SavedCommand->GetStringField( TEXT("Context") );
-				int32 Index = FMath::Trunc( SavedCommand->GetNumberField( TEXT("Index") ) );
-				FCustomBlockTransaction::ETransactionType TransType = (FCustomBlockTransaction::ETransactionType)FMath::Trunc( SavedCommand->GetNumberField("TransactionType") );
+				int32 Index = FMath::TruncToInt( SavedCommand->GetNumberField( TEXT("Index") ) );
+				FCustomBlockTransaction::ETransactionType TransType = (FCustomBlockTransaction::ETransactionType)FMath::TruncToInt( SavedCommand->GetNumberField("TransactionType") );
 
 				if( !CommandName.IsEmpty() && !Context.IsEmpty() && Index >= 0 )
 				{
diff --git a/Engine/Source/Runtime/Slate/Private/Framework/MultiBox/SClippingHorizontalBox.cpp b/Engine/Source/Runtime/Slate/Private/Framework/MultiBox/SClippingHorizontalBox.cpp
--- a/Engine/Source/Runtime/Slate/Private/Framework/MultiBox/SClippingHorizontalBox.cpp
+++ b/Engine/Source/Runtime/Slate/Private/Framework/MultiBox/SClippingHorizontalBox.cpp
@@ -25,7 +25,7 @@ void SClippingHorizontalBox::ArrangeChildren( const FGeometry& AllottedGeometry,
 	for (int32 ChildIdx = NumChildren - 2; ChildIdx >= 0; --ChildIdx)
 	{
 		const FArrangedWidget& CurWidget = ArrangedChildren(ChildIdx);
-		if (FMath::Trunc(CurWidget.Geometry.AbsolutePosition.X + CurWidget.Geometry.Size.X) > FMath::Trunc(AllottedGeometry.AbsolutePosition.X + AllottedGeometry.Size.X))
+		if (FMath::TruncToInt(CurWidget.Geometry.AbsolutePosition.X + CurWidget.Geometry.Size.X) > FMath::TruncToInt(AllottedGeometry.AbsolutePosition.X + AllottedGeometry.Size.X))
 		{
 			ArrangedChildren.Remove(ChildIdx);
 			IndexClippedAt = ChildIdx;
@@ -47,7 +47,7 @@ void SClippingHorizontalBox::ArrangeChildren( const FGeometry& AllottedGeometry,
 		for (int32 ChildIdx = IndexClippedAt - 1; ChildIdx >= 0; --ChildIdx)
 		{
 			const FArrangedWidget& CurWidget = ArrangedChildren(ChildIdx);
-			if (FMath::Trunc(CurWidget.Geometry.AbsolutePosition.X + CurWidget.Geometry.Size.X) > FMath::Trunc(ArrangedButton.Geometry.AbsolutePosition.X))
+			if (FMath::TruncToInt(CurWidget.Geometry.AbsolutePosition.X + CurWidget.Geometry.Size.X) > FMath::TruncToInt(ArrangedButton.Geometry.AbsolutePosition.X))
 			{
 				ArrangedChildren.Remove(ChildIdx);
 			}
@@ -77,7 +77,7 @@ int32 SClippingHorizontalBox::OnPaint( const FGeometry& AllottedGeometry, const
 		if (IndexClippedAt == ArrangedChildren.Num() - 2)
 		{
 			// Only recalculate the alloted geometry size if said size is fitted to the toolbar/menubar
-			if (FMath::Trunc(AllottedGeometry.AbsolutePosition.X + AllottedGeometry.Size.X) <= FMath::Trunc(LastChild.Geometry.AbsolutePosition.X + LastChild.Geometry.Size.X))
+			if (FMath::TruncToInt(AllottedGeometry.AbsolutePosition.X + AllottedGeometry.Size.X) <= FMath::TruncToInt(LastChild.Geometry.AbsolutePosition.X + LastChild.Geometry.Size.X))
 			{
 				// Calculate the size of the custom border
 				ResizedGeometry.Size.X = LastCippedChild.Geometry.AbsolutePosition.X + LastCippedChild.Geometry.Size.X - FirstChild.Geometry.AbsolutePosition.X;
diff --git a/Engine/Source/Runtime/Slate/Private/Widgets/Input/SVirtualJoystick.cpp b/Engine/Source/Runtime/Slate/Private/Widgets/Input/SVirtualJoystick.cpp
--- a/Engine/Source/Runtime/Slate/Private/Widgets/Input/SVirtualJoystick.cpp
+++ b/Engine/Source/Runtime/Slate/Private/Widgets/Input/SVirtualJoystick.cpp
@@ -106,7 +106,7 @@ int32 SVirtualJoystick::OnPaint( const FGeometry& AllottedGeometry, const FSlate
 	if (bVisible)
 	{
 		FColor ColorAndOpacitySRGB = InWidgetStyle.GetColorAndOpacityTint();
-		ColorAndOpacitySRGB.A = FMath::Round(255.f * CurrentOpacity);
+		ColorAndOpacitySRGB.A = FMath::RoundToInt(255.f * CurrentOpacity);
 
 		for (int32 ControlIndex = 0; ControlIndex < Controls.Num(); ControlIndex++)
 		{
diff --git a/Engine/Source/Runtime/Slate/Private/Widgets/Layout/SScrollBox.cpp b/Engine/Source/Runtime/Slate/Private/Widgets/Layout/SScrollBox.cpp
--- a/Engine/Source/Runtime/Slate/Private/Widgets/Layout/SScrollBox.cpp
+++ b/Engine/Source/Runtime/Slate/Private/Widgets/Layout/SScrollBox.cpp
@@ -264,8 +264,8 @@ FReply SScrollBox::OnMouseButtonUp( const FGeometry& MyGeometry, const FPointerE
 			FVector2D CursorPosition = MyGeometry.LocalToAbsolute( SoftwareCursorPosition );
 
 			FIntPoint BestPositionInPanel(
-				FMath::Round( FMath::Clamp( CursorPosition.X, PanelScreenSpaceRect.Left, PanelScreenSpaceRect.Right ) ),
-				FMath::Round( FMath::Clamp( CursorPosition.Y, PanelScreenSpaceRect.Top, PanelScreenSpaceRect.Bottom ) )
+				FMath::RoundToInt( FMath::Clamp( CursorPosition.X, PanelScreenSpaceRect.Left, PanelScreenSpaceRect.Right ) ),
+				FMath::RoundToInt( FMath::Clamp( CursorPosition.Y, PanelScreenSpaceRect.Top, PanelScreenSpaceRect.Bottom ) )
 				);
 
 			Reply.SetMousePos(BestPositionInPanel);
diff --git a/Engine/Source/Runtime/Slate/Private/Widgets/Text/SMultiLineEditableText.cpp b/Engine/Source/Runtime/Slate/Private/Widgets/Text/SMultiLineEditableText.cpp
--- a/Engine/Source/Runtime/Slate/Private/Widgets/Text/SMultiLineEditableText.cpp
+++ b/Engine/Source/Runtime/Slate/Private/Widgets/Text/SMultiLineEditableText.cpp
@@ -109,9 +109,9 @@ FReply SMultiLineEditableText::MoveCursor( ECursorMoveMethod::Type Method, int8
 			NumLinesScrollOffset = NewWrappedCursorPosition;
 			UpdateScrollbar( CachedLastFrameGeometry );
 		}
-		else if ( FMath::Floor(NumLinesScrollOffset + NumLinesOnScreen - 1) < NewWrappedCursorPosition )
+		else if ( FMath::FloorToInt(NumLinesScrollOffset + NumLinesOnScreen - 1) < NewWrappedCursorPosition )
 		{
-			NumLinesScrollOffset = FMath::Clamp( FMath::Floor(NewWrappedCursorPosition - NumLinesOnScreen + 1 ), 0, WrappedText.Num());
+			NumLinesScrollOffset = FMath::Clamp( FMath::FloorToInt(NewWrappedCursorPosition - NumLinesOnScreen + 1 ), 0, WrappedText.Num());
 			UpdateScrollbar( CachedLastFrameGeometry );
 		}
 	}
@@ -212,7 +212,7 @@ int32 SMultiLineEditableText::OnPaint( const FGeometry& AllottedGeometry, const
 	const bool bEnabled = ShouldBeEnabled( bParentEnabled );
 	const ESlateDrawEffect::Type DrawEffects = (bEnabled) ? ESlateDrawEffect::None : ESlateDrawEffect::DisabledEffect;
 
-	const int32 WholeLinesOffset = FMath::Floor( NumLinesScrollOffset );
+	const int32 WholeLinesOffset = FMath::FloorToInt( NumLinesScrollOffset );
 	const float PartialLineOffset = FMath::Fractional( NumLinesScrollOffset );
 
 
@@ -347,7 +347,7 @@ int32 SMultiLineEditableText::OnPaint( const FGeometry& AllottedGeometry, const
 					const bool bForceCursorVisible = (CurrentTime - LastCursorInteractionTime) < EditableTextDefs::CaretBlinkPauseTime;
 					const float CursorOpacity = (bForceCursorVisible)
 						? 1.0f
-						: FMath::Round( FMath::MakePulsatingValue( CurrentTime, EditableTextDefs::BlinksPerSecond ));
+						: FMath::RoundToInt( FMath::MakePulsatingValue( CurrentTime, EditableTextDefs::BlinksPerSecond ));
 
 					const int32 CursorOffsetInWrappedLine = CursorPosition.GetOffset() - ThisWrappedLine.FirstCharIndex;
 					const FVector2D SlateUnitsCursorOffsetInLine = FontMeasureService->Measure( StringToPaint.Left(CursorOffsetInWrappedLine), FontInfo, AllottedGeometry.Scale );
diff --git a/Engine/Source/Runtime/Slate/Private/Widgets/Text/STextBlock.cpp b/Engine/Source/Runtime/Slate/Private/Widgets/Text/STextBlock.cpp
--- a/Engine/Source/Runtime/Slate/Private/Widgets/Text/STextBlock.cpp
+++ b/Engine/Source/Runtime/Slate/Private/Widgets/Text/STextBlock.cpp
@@ -242,7 +242,7 @@ void STextBlock::CacheDesiredSize()
 		if(bWrapText)
 		{
 			// WrapText takes care of clearing CachedWrappedLineData
-			CachedWrappedString = SlateWordWrapper::WrapText(CachedOriginalString, FontInfo, FMath::Trunc(WrappingWidth), 1.0f, &CachedWrappedLineData);
+			CachedWrappedString = SlateWordWrapper::WrapText(CachedOriginalString, FontInfo, FMath::TruncToInt(WrappingWidth), 1.0f, &CachedWrappedLineData);
 		}
 		else
 		{
diff --git a/Engine/Source/Runtime/Slate/Private/Widgets/Views/SListPanel.cpp b/Engine/Source/Runtime/Slate/Private/Widgets/Views/SListPanel.cpp
--- a/Engine/Source/Runtime/Slate/Private/Widgets/Views/SListPanel.cpp
+++ b/Engine/Source/Runtime/Slate/Private/Widgets/Views/SListPanel.cpp
@@ -57,7 +57,7 @@ void SListPanel::ArrangeChildren( const FGeometry& AllottedGeometry, FArrangedCh
 		float WidthSoFar = 0;
 		float LocalItemWidth = ItemWidth.Get();
 		float LocalItemHeight = ItemHeight.Get();
-		float HeightSoFar = -FMath::Floor(SmoothScrollOffsetInItems * LocalItemHeight);
+		float HeightSoFar = -FMath::FloorToInt(SmoothScrollOffsetInItems * LocalItemHeight);
 
 		for( int32 ItemIndex = 0; ItemIndex < Children.Num(); ++ItemIndex )
 		{
@@ -79,7 +79,7 @@ void SListPanel::ArrangeChildren( const FGeometry& AllottedGeometry, FArrangedCh
 		if (Children.Num() > 0)
 		{
 			// This is a normal list, arrange items vertically
-			float HeightSoFar = -FMath::Floor(SmoothScrollOffsetInItems * Children[0].Widget->GetDesiredSize().Y);
+			float HeightSoFar = -FMath::FloorToInt(SmoothScrollOffsetInItems * Children[0].Widget->GetDesiredSize().Y);
 			for( int32 ItemIndex = 0; ItemIndex < Children.Num(); ++ItemIndex )
 			{
 					const FVector2D ItemDesiredSize = Children[ItemIndex].Widget->GetDesiredSize();
@@ -109,8 +109,8 @@ void SListPanel::Tick( const FGeometry& AllottedGeometry, const double InCurrent
 
 		if (TotalItemSize > 0.0f && NumChildren > 0)
 		{
-			int32 NumColumns = FMath::Clamp(FMath::Ceil(AllottedWidth / TotalItemSize) - 1, 1, NumChildren);
-			PreferredRowNum = FMath::Ceil(NumChildren / (float)NumColumns);
+			int32 NumColumns = FMath::Clamp(FMath::CeilToInt(AllottedWidth / TotalItemSize) - 1, 1, NumChildren);
+			PreferredRowNum = FMath::CeilToInt(NumChildren / (float)NumColumns);
 		}
 		else
 		{
@@ -192,7 +192,7 @@ float SListPanel::GetItemWidth() const
 float SListPanel::GetItemPadding(const FGeometry& AllottedGeometry) const
 {
 	float LocalItemWidth = ItemWidth.Get();
-	const int32 NumItemsWide = LocalItemWidth > 0 ? FMath::Floor(AllottedGeometry.Size.X / LocalItemWidth) : 0;
+	const int32 NumItemsWide = LocalItemWidth > 0 ? FMath::FloorToInt(AllottedGeometry.Size.X / LocalItemWidth) : 0;
 	
 	// Only add padding between items if we have more total items that we can fit on a single row.  Otherwise,
 	// the padding around items would continue to change proportionately with the (ample) free horizontal space
diff --git a/Engine/Source/Runtime/Slate/Private/Widgets/Views/STableViewBase.cpp b/Engine/Source/Runtime/Slate/Private/Widgets/Views/STableViewBase.cpp
--- a/Engine/Source/Runtime/Slate/Private/Widgets/Views/STableViewBase.cpp
+++ b/Engine/Source/Runtime/Slate/Private/Widgets/Views/STableViewBase.cpp
@@ -278,8 +278,8 @@ FReply STableViewBase::OnMouseButtonUp( const FGeometry& MyGeometry, const FPoin
 			FVector2D CursorPosition = MyGeometry.LocalToAbsolute( SoftwareCursorPosition );
 
 			FIntPoint BestPositionInList(
-				FMath::Round( FMath::Clamp( CursorPosition.X, ListScreenSpaceRect.Left, ListScreenSpaceRect.Right ) ),
-				FMath::Round( FMath::Clamp( CursorPosition.Y, ListScreenSpaceRect.Top, ListScreenSpaceRect.Bottom ) )
+				FMath::RoundToInt( FMath::Clamp( CursorPosition.X, ListScreenSpaceRect.Left, ListScreenSpaceRect.Right ) ),
+				FMath::RoundToInt( FMath::Clamp( CursorPosition.Y, ListScreenSpaceRect.Top, ListScreenSpaceRect.Bottom ) )
 				);
 
 			Reply.SetMousePos(BestPositionInList);
diff --git a/Engine/Source/Runtime/Slate/Public/Widgets/Views/SListView.h b/Engine/Source/Runtime/Slate/Public/Widgets/Views/SListView.h
--- a/Engine/Source/Runtime/Slate/Public/Widgets/Views/SListView.h
+++ b/Engine/Source/Runtime/Slate/Public/Widgets/Views/SListView.h
@@ -755,7 +755,7 @@ public:
 
 			// Index of the item at which we start generating based on how far scrolled down we are
 			// Note that we must generate at LEAST one item.
-			int32 StartIndex = FMath::Clamp( FMath::Floor(ScrollOffset), 0, SourceItems->Num()-1 );
+			int32 StartIndex = FMath::Clamp( FMath::FloorToInt(ScrollOffset), 0, SourceItems->Num()-1 );
 
 			// Height of the first item that is generated. This item is at the location where the user requested we scroll
 			float FirstItemHeight = 0.0f;
diff --git a/Engine/Source/Runtime/Slate/Public/Widgets/Views/STileView.h b/Engine/Source/Runtime/Slate/Public/Widgets/Views/STileView.h
--- a/Engine/Source/Runtime/Slate/Public/Widgets/Views/STileView.h
+++ b/Engine/Source/Runtime/Slate/Public/Widgets/Views/STileView.h
@@ -178,7 +178,7 @@ public:
 			float WidthUsedSoFar = 0.0f;
 			float HeightUsedSoFar = 0.0f;
 			// Index of the item at which we start generating based on how far scrolled down we are
-			int32 StartIndex = FMath::Max( 0, FMath::Floor(ClampedScrollOffset / NumItemsWide) * NumItemsWide );
+			int32 StartIndex = FMath::Max( 0, FMath::FloorToInt(ClampedScrollOffset / NumItemsWide) * NumItemsWide );
 
 			// Let the WidgetGenerator that we are starting a pass so that it can keep track of data items and widgets.
 			this->WidgetGenerator.OnBeginGenerationPass();
@@ -251,7 +251,7 @@ protected:
 	virtual int32 GetNumItemsWide() const OVERRIDE
 	{
 		const float ItemWidth = this->GetItemWidth();
-		const int32 NumItemsWide = ItemWidth > 0 ? FMath::Floor(this->PanelGeometryLastTick.Size.X / ItemWidth) : 1;
+		const int32 NumItemsWide = ItemWidth > 0 ? FMath::FloorToInt(this->PanelGeometryLastTick.Size.X / ItemWidth) : 1;
 		return FMath::Max(1, NumItemsWide);
 	}
 };
diff --git a/Engine/Source/Runtime/SlateCore/Private/Rendering/ElementBatcher.cpp b/Engine/Source/Runtime/SlateCore/Private/Rendering/ElementBatcher.cpp
--- a/Engine/Source/Runtime/SlateCore/Private/Rendering/ElementBatcher.cpp
+++ b/Engine/Source/Runtime/SlateCore/Private/Rendering/ElementBatcher.cpp
@@ -168,7 +168,7 @@ void FSlateElementBatcher::AddBoxElement( const FVector2D& InPosition, const FVe
 
 	if(InPayload.BrushResource->DrawAs != ESlateBrushDrawType::NoDrawType)
 	{
-		FVector2D Position( FMath::Trunc(InPosition.X), FMath::Trunc(InPosition.Y) );
+		FVector2D Position( FMath::TruncToInt(InPosition.X), FMath::TruncToInt(InPosition.Y) );
 
 		const float Angle = InPayload.Angle;
 		const FVector2D RotationPoint = InPayload.RotationPoint;
@@ -297,8 +297,8 @@ void FSlateElementBatcher::AddBoxElement( const FVector2D& InPosition, const FVe
 
 			FVector2D EndPos = FVector2D(Position.X, Position.Y) + Size;
 
-			EndPos.X = FMath::Trunc( EndPos.X );
-			EndPos.Y = FMath::Trunc( EndPos.Y );
+			EndPos.X = FMath::TruncToInt( EndPos.X );
+			EndPos.Y = FMath::TruncToInt( EndPos.Y );
 
 			BatchVertices.Add( FSlateVertex( Position,									StartUV,									Tiling,	Tint, InClippingRect ) ); //0
 			BatchVertices.Add( FSlateVertex( FVector2D( Position.X, TopMarginY ),		FVector2D( StartUV.X, TopMarginV ),			Tiling,	Tint, InClippingRect ) ); //1
@@ -452,8 +452,8 @@ void FSlateElementBatcher::AddTextElement( const FVector2D& Position, const FVec
 
 	int32 Kerning = 0;
 
-	int32 PosX = FMath::Trunc(Position.X);
-	int32 PosY = FMath::Trunc(Position.Y);
+	int32 PosX = FMath::TruncToInt(Position.X);
+	int32 PosY = FMath::TruncToInt(Position.Y);
 
 	LineX = PosX;
 	
@@ -687,7 +687,7 @@ void FSlateElementBatcher::AddSplineElement( const FVector2D& Position, float Sc
 
 	const float DirectLength = (InPayload.EndPt - InPayload.StartPt).Size();
 	const float HandleLength = ((InPayload.EndPt - InPayload.EndDir) - (InPayload.StartPt + InPayload.StartDir)).Size();
-	float NumSteps = FMath::Clamp<float>(FMath::Ceil(FMath::Max(DirectLength,HandleLength)/15.0f), 1, 256);
+	float NumSteps = FMath::Clamp<float>(FMath::CeilToInt(FMath::Max(DirectLength,HandleLength)/15.0f), 1, 256);
 
 	// 1 is the minimum thickness we support
 	float InThickness = FMath::Max( 1.0f, InPayload.Thickness );
@@ -697,7 +697,7 @@ void FSlateElementBatcher::AddSplineElement( const FVector2D& Position, float Sc
 
 	// Compute the actual size of the line we need based on thickness.  Need to ensure pixels that are at least Thickness/2 + Sample radius are generated so that we have enough pixels to blend.
 	// The idea for the spline anti-alising technique is based on the fast prefiltered lines technique published in GPU Gems 2 
-	const float LineThickness = FMath::Ceil( (2.0f * Radius + InThickness ) * FMath::Sqrt(2.0f) );
+	const float LineThickness = FMath::CeilToInt( (2.0f * Radius + InThickness ) * FMath::Sqrt(2.0f) );
 
 	// The amount we increase each side of the line to generate enough pixels
 	const float HalfThickness = LineThickness * .5f + Radius;
@@ -832,7 +832,7 @@ void FSlateElementBatcher::AddLineElement( const FVector2D& Position, const FVec
 
 		// Compute the actual size of the line we need based on thickness.  Need to ensure pixels that are at least Thickness/2 + Sample radius are generated so that we have enough pixels to blend.
 		// The idea for the anti-aliasing technique is based on the fast prefiltered lines technique published in GPU Gems 2 
-		const float LineThickness = FMath::Ceil( (2.0f * Radius + RequestedThickness ) * FMath::Sqrt(2.0f) );
+		const float LineThickness = FMath::CeilToInt( (2.0f * Radius + RequestedThickness ) * FMath::Sqrt(2.0f) );
 
 		// The amount we increase each side of the line to generate enough pixels
 		const float HalfThickness = LineThickness * .5f + Radius;
@@ -1005,8 +1005,8 @@ void FSlateElementBatcher::AddViewportElement( const FVector2D& InPosition, cons
 		bRequiresVsync |= ViewportPin->RequiresVsync();
 	}
 
-	FVector2D Position = FVector2D( FMath::Trunc( InPosition.X ), FMath::Trunc( InPosition.Y ) );
-	FVector2D Size = FVector2D( FMath::Trunc( InSize.X ), FMath::Trunc( InSize.Y ) );
+	FVector2D Position = FVector2D( FMath::TruncToInt( InPosition.X ), FMath::TruncToInt( InPosition.Y ) );
+	FVector2D Size = FVector2D( FMath::TruncToInt( InSize.X ), FMath::TruncToInt( InSize.Y ) );
 
 	// Determine the four corners of the quad
 	FVector2D TopLeft = Position;
@@ -1051,7 +1051,7 @@ void FSlateElementBatcher::AddBorderElement( const FVector2D& InPosition, const
 
 	check( InPayload.BrushResource );
 
-	const FVector2D Position( FMath::Trunc(InPosition.X), FMath::Trunc(InPosition.Y) );
+	const FVector2D Position( FMath::TruncToInt(InPosition.X), FMath::TruncToInt(InPosition.Y) );
 
 	uint32 TextureWidth = 1;
 	uint32 TextureHeight = 1;
@@ -1136,8 +1136,8 @@ void FSlateElementBatcher::AddBorderElement( const FVector2D& InPosition, const
 
 	FVector2D EndPos = FVector2D(Position.X, Position.Y) + Size;
 
-	EndPos.X = FMath::Trunc(EndPos.X);
-	EndPos.Y = FMath::Trunc(EndPos.Y);
+	EndPos.X = FMath::TruncToInt(EndPos.X);
+	EndPos.Y = FMath::TruncToInt(EndPos.Y);
 
 	// The start index of these vertices in the index buffer
 	uint32 IndexStart = BatchVertices.Num();
