add some NULL checks in spline component - waiting for EGK on proper fix

diff --git a/Engine/Source/Runtime/Engine/Private/Landscape/LandscapeSplines.cpp b/Engine/Source/Runtime/Engine/Private/Landscape/LandscapeSplines.cpp
--- a/Engine/Source/Runtime/Engine/Private/Landscape/LandscapeSplines.cpp
+++ b/Engine/Source/Runtime/Engine/Private/Landscape/LandscapeSplines.cpp
@@ -60,29 +60,33 @@ public:
 		Segments.Reserve(Component->Segments.Num());
 		for (int32 i = 0; i < Component->Segments.Num(); i++)
 		{
-			ULandscapeSplineSegment* Segment = Component->Segments[i];
-
-			FSegmentProxy SegmentProxy;
-			SegmentProxy.Owner = Segment;
-			SegmentProxy.HitProxy = NULL;
-			SegmentProxy.Points = Segment->GetPoints();
-			SegmentProxy.bSelected = Segment->IsSplineSelected();
-			Segments.Add(SegmentProxy);
+			if (ULandscapeSplineSegment* Segment = Component->Segments[i])
+			{
+
+				FSegmentProxy SegmentProxy;
+				SegmentProxy.Owner = Segment;
+				SegmentProxy.HitProxy = NULL;
+				SegmentProxy.Points = Segment->GetPoints();
+				SegmentProxy.bSelected = Segment->IsSplineSelected();
+				Segments.Add(SegmentProxy);
+			}
 		}
 
 		ControlPoints.Reserve(Component->ControlPoints.Num());
 		for (int32 i = 0; i < Component->ControlPoints.Num(); i++)
 		{
-			ULandscapeSplineControlPoint* ControlPoint = Component->ControlPoints[i];
+			if (ULandscapeSplineControlPoint* ControlPoint = Component->ControlPoints[i])
+			{
 
-			FControlPointProxy ControlPointProxy;
-			ControlPointProxy.Owner = ControlPoint;
-			ControlPointProxy.HitProxy = NULL;
-			ControlPointProxy.Location = ControlPoint->Location;
-			ControlPointProxy.Points = ControlPoint->GetPoints();
-			ControlPointProxy.SpriteScale = FMath::Clamp<float>(ControlPoint->Width != 0 ? ControlPoint->Width / 2 : ControlPoint->SideFalloff / 4, 10, 1000);
-			ControlPointProxy.bSelected = ControlPoint->IsSplineSelected();
-			ControlPoints.Add(ControlPointProxy);
+				FControlPointProxy ControlPointProxy;
+				ControlPointProxy.Owner = ControlPoint;
+				ControlPointProxy.HitProxy = NULL;
+				ControlPointProxy.Location = ControlPoint->Location;
+				ControlPointProxy.Points = ControlPoint->GetPoints();
+				ControlPointProxy.SpriteScale = FMath::Clamp<float>(ControlPoint->Width != 0 ? ControlPoint->Width / 2 : ControlPoint->SideFalloff / 4, 10, 1000);
+				ControlPointProxy.bSelected = ControlPoint->IsSplineSelected();
+				ControlPoints.Add(ControlPointProxy);
+			}
 		}
 	}
 
@@ -324,13 +328,17 @@ void ULandscapeSplinesComponent::OnRegister()
 
 	for (int32 iControlPoint = 0; iControlPoint < ControlPoints.Num(); iControlPoint++)
 	{
-		ULandscapeSplineControlPoint* ControlPoint = ControlPoints[iControlPoint];
-		ControlPoint->RegisterComponents();
+		if (ULandscapeSplineControlPoint* ControlPoint = ControlPoints[iControlPoint])
+		{
+			ControlPoint->RegisterComponents();
+		}
 	}
 	for (int32 iSegment = 0; iSegment < Segments.Num(); iSegment++)
 	{
-		ULandscapeSplineSegment* Segment = Segments[iSegment];
-		Segment->RegisterComponents();
+		if (ULandscapeSplineSegment* Segment = Segments[iSegment])
+		{
+			Segment->RegisterComponents();
+		}
 	}
 }
 
@@ -338,13 +346,17 @@ void ULandscapeSplinesComponent::OnUnregister()
 {
 	for (int32 iControlPoint = 0; iControlPoint < ControlPoints.Num(); iControlPoint++)
 	{
-		ULandscapeSplineControlPoint* ControlPoint = ControlPoints[iControlPoint];
-		ControlPoint->UnregisterComponents();
+		if (ULandscapeSplineControlPoint* ControlPoint = ControlPoints[iControlPoint])
+		{
+			ControlPoint->UnregisterComponents();
+		}
 	}
 	for (int32 iSegment = 0; iSegment < Segments.Num(); iSegment++)
 	{
-		ULandscapeSplineSegment* Segment = Segments[iSegment];
-		Segment->UnregisterComponents();
+		if (ULandscapeSplineSegment* Segment = Segments[iSegment])
+		{
+			Segment->UnregisterComponents();
+		}
 	}
 
 	Super::OnUnregister();
@@ -363,14 +375,18 @@ FBoxSphereBounds ULandscapeSplinesComponent::CalcBounds(const FTransform & Local
 
 	for (int32 iControlPoint = 0; iControlPoint < ControlPoints.Num(); iControlPoint++)
 	{
-		ULandscapeSplineControlPoint* ControlPoint = ControlPoints[iControlPoint];
-		NewBoundsCalc += ControlPoint->GetBounds();
+		if (ULandscapeSplineControlPoint* ControlPoint = ControlPoints[iControlPoint])
+		{
+			NewBoundsCalc += ControlPoint->GetBounds();
+		}
 	}
 
 	for (int32 iSegment = 0; iSegment < Segments.Num(); iSegment++)
 	{
-		ULandscapeSplineSegment* Segment = Segments[iSegment];
-		NewBoundsCalc += Segment->GetBounds();
+		if (ULandscapeSplineSegment* Segment = Segments[iSegment])
+		{
+			NewBoundsCalc += Segment->GetBounds();
+		}
 	}
 
 	NewBoundsCalc = NewBoundsCalc.TransformBy(LocalToWorld);
@@ -410,13 +426,17 @@ void ULandscapeSplinesComponent::PostEditChangeProperty(FPropertyChangedEvent& P
 
 	for (int32 iControlPoint = 0; iControlPoint < ControlPoints.Num(); iControlPoint++)
 	{
-		ULandscapeSplineControlPoint* ControlPoint = ControlPoints[iControlPoint];
-		ControlPoint->UpdateSplinePoints(bUpdateCollision);
+		if (ULandscapeSplineControlPoint* ControlPoint = ControlPoints[iControlPoint])
+		{
+			ControlPoint->UpdateSplinePoints(bUpdateCollision);
+		}
 	}
 	for (int32 iSegment = 0; iSegment < Segments.Num(); iSegment++)
 	{
-		ULandscapeSplineSegment* Segment = Segments[iSegment];
-		Segment->UpdateSplinePoints(bUpdateCollision);
+		if (ULandscapeSplineSegment* Segment = Segments[iSegment])
+		{
+			Segment->UpdateSplinePoints(bUpdateCollision);
+		}
 	}
 }
 
