Fix precloth tick not being correctly registered in builds that don't use WITH_APEX_CLOTHING

diff --git a/Engine/Source/Runtime/Engine/Private/SkeletalMeshComponent.cpp b/Engine/Source/Runtime/Engine/Private/SkeletalMeshComponent.cpp
--- a/Engine/Source/Runtime/Engine/Private/SkeletalMeshComponent.cpp
+++ b/Engine/Source/Runtime/Engine/Private/SkeletalMeshComponent.cpp
@@ -33,10 +33,12 @@ USkeletalMeshComponent::USkeletalMeshComponent(const class FPostConstructInitial
 	KinematicBonesUpdateType = EKinematicBonesUpdateToPhysics::SkipSimulatingBones;
 	bGenerateOverlapEvents = false;
 	LineCheckBoundsScale = FVector(1.0f, 1.0f, 1.0f);
-#if WITH_APEX_CLOTHING
+
 	PreClothTickFunction.TickGroup = TG_PreCloth;
 	PreClothTickFunction.bCanEverTick = true;
 	PreClothTickFunction.bStartWithTickEnabled = true;
+
+#if WITH_APEX_CLOTHING
 	ClothMaxDistanceScale = 1.0f;
 	ClothTeleportMode = FClothingActor::Continuous;
 	PrevRootBoneMatrix = GetBoneMatrix(0); // save the root bone transform
@@ -71,8 +73,6 @@ void USkeletalMeshComponent::RegisterComponentTickFunctions(bool bRegister)
 {
 	Super::RegisterComponentTickFunctions(bRegister);
 
-#if WITH_APEX_CLOTHING
-
 	if (bRegister)
 	{
 		if (SetupActorComponentTickFunction(&PreClothTickFunction))
@@ -92,7 +92,6 @@ void USkeletalMeshComponent::RegisterComponentTickFunctions(bool bRegister)
 			PreClothTickFunction.UnRegisterTickFunction();
 		}
 	}
-#endif
 
 }
 
