Fix recompile remote shaders not working for cookonthefly.

diff --git a/Engine/Source/Runtime/Engine/Private/Materials/MaterialShader.cpp b/Engine/Source/Runtime/Engine/Private/Materials/MaterialShader.cpp
--- a/Engine/Source/Runtime/Engine/Private/Materials/MaterialShader.cpp
+++ b/Engine/Source/Runtime/Engine/Private/Materials/MaterialShader.cpp
@@ -856,16 +856,21 @@ void FMaterialShaderMap::SaveForRemoteRecompile(FArchive& Ar, const TMap<FString
 				ShaderMap->GetShaderList(Shaders);
 
 				// get the resources from the shaders
-				for (TMap<FShaderId, FShader*>::TIterator ShaderIt(Shaders); It; ++It)
+				for (TMap<FShaderId, FShader*>::TIterator ShaderIt(Shaders); ShaderIt; ++ShaderIt)
 				{
-					FShaderResourceId ShaderId = ShaderIt.Value()->GetResourceId();
+					FShader *Shader = ShaderIt.Value();
+					FShaderResourceId ShaderId = Shader->GetResourceId();
+
+					// UE_LOG(LogMaterial, Display, TEXT("Shader Outputhash %s vs %s vs %s"), *ShaderId.OutputHash.ToString(), *Shader->GetOutputHash().ToString(), *Shader->GetResourceOutputHash().ToString() );
 					// skip this shader if the Id was already on the client (ie, it didn't change)
-					if (ClientResourceIds.Contains(ShaderId) == false)
+					if (ClientResourceIds.Contains(ShaderId) == false ) // || true)
 					{
 						// lookup the resource by ID
 						FShaderResource* Resource = FShaderResource::FindShaderResourceById(ShaderId);
 						// add it if it's unique
 						UniqueResources.AddUnique(Resource);
+
+						// UE_LOG(LogMaterial, Display, TEXT("Saving material resource id %d %s"), GetTypeHash(ShaderId), *Resource->GetOutputHash().ToString() );
 					}
 					else
 					{
@@ -932,9 +937,12 @@ void FMaterialShaderMap::LoadForRemoteRecompile(FArchive& Ar, EShaderPlatform Sh
 		FShaderResource* Resource = new FShaderResource();
 		Resource->Serialize(Ar);
 
+		// UE_LOG(LogMaterial, Display, TEXT("Loaded material resource id %d"), GetTypeHash(Resource->GetId()));
+
 		// if this Id is already in memory, that means that this is a repeated resource and so we skip it
 		if (FShaderResource::FindShaderResourceById(Resource->GetId()) != NULL)
 		{
+			// UE_LOG(LogMaterial, Display, TEXT("Found ShaderResourceById %d"), GetTypeHash(Resource->GetId()));
 			delete Resource;
 		}
 		// otherwise, it's a new resource, so we register it for the maps to find below
