CameraShakes no longer unnecessarily instantiate the source shake data.

diff --git a/Engine/Source/Runtime/Engine/Classes/Camera/CameraModifier_CameraShake.h b/Engine/Source/Runtime/Engine/Classes/Camera/CameraModifier_CameraShake.h
--- a/Engine/Source/Runtime/Engine/Classes/Camera/CameraModifier_CameraShake.h
+++ b/Engine/Source/Runtime/Engine/Classes/Camera/CameraModifier_CameraShake.h
@@ -16,7 +16,7 @@ struct ENGINE_API FCameraShakeInstance
 
 	/** source shake */
 	UPROPERTY()
-	class UCameraShake* SourceShake;
+	const class UCameraShake* SourceShake;
 
 	/** Used to identify shakes when single instances are desired */
 	UPROPERTY()
diff --git a/Engine/Source/Runtime/Engine/Private/Camera/CameraModifier_CameraShake.cpp b/Engine/Source/Runtime/Engine/Private/Camera/CameraModifier_CameraShake.cpp
--- a/Engine/Source/Runtime/Engine/Private/Camera/CameraModifier_CameraShake.cpp
+++ b/Engine/Source/Runtime/Engine/Private/Camera/CameraModifier_CameraShake.cpp
@@ -234,7 +234,7 @@ void UCameraModifier_CameraShake::ReinitShake(FCameraShakeInstance& ShakeInst, f
 	}
 	ShakeInst.Scale = NewScale;
 
-	UCameraShake* const SourceShake = ShakeInst.SourceShake;
+	UCameraShake const* const SourceShake = ShakeInst.SourceShake;
 
 	if (SourceShake->OscillationDuration != 0.f)
 	{
@@ -279,7 +279,7 @@ FCameraShakeInstance UCameraModifier_CameraShake::InitializeShake(TSubclassOf<cl
 	Inst.SourceShakeName = NewShake->GetFName();
 
 	// Create a camera shake object of class NewShake
-	UCameraShake* const SourceShake = Cast<UCameraShake>(StaticConstructObject(NewShake, this));
+	UCameraShake const* const SourceShake = CastChecked<UCameraShake>(NewShake->GetDefaultObject());
 	Inst.SourceShake = SourceShake;
 
 	Inst.Scale = Scale;
