#ue4 - Improve comments for UMovementComponent::SlideAlongSurface().

- Document previously undocumented params.

diff --git a/Engine/Source/Runtime/Engine/Classes/GameFramework/MovementComponent.h b/Engine/Source/Runtime/Engine/Classes/GameFramework/MovementComponent.h
--- a/Engine/Source/Runtime/Engine/Classes/GameFramework/MovementComponent.h
+++ b/Engine/Source/Runtime/Engine/Classes/GameFramework/MovementComponent.h
@@ -187,17 +187,21 @@ public:
 	/**
 	 * Compute a vector to slide along a surface, given an attempted move, time, and normal.
 	 * @param Delta:	Attempted move.
-	 * @param Time:		Amount of move to apply, usually between 0 and 1.
-	 * @param Normal:	Normal opposed to movement. May not equal Hit.Normal.
-	 * @param Hit:		Hit result of the move that resulted in the slide.
+	 * @param Time:		Amount of move to apply (between 0 and 1).
+	 * @param Normal:	Normal opposed to movement. Not necessarily equal to Hit.Normal.
+	 * @param Hit:		HitResult of the move that resulted in the slide.
 	 */
 	virtual FVector ComputeSlideVector(const FVector& Delta, const float Time, const FVector& Normal, const FHitResult& Hit) const;
 
 	/**
 	 * Slide smoothly along a surface, and slide away from multiple impacts using TwoWallAdjust if necessary. Calls HandleImpact for each surface hit, if requested.
-	 * Uses SafeMoveUpdatedComponent for movement, and ComputeSlideVector to determine the slide direction.
-	 * @param bHandleImpact - Whether to call HandleImpact on each hit.
-	 * @return The percent of the Time value actually applied to movement (between 0 and 1). 0 if no movement occurred, non-zero if movement occurred (in which case Hit.Time is the result of the last move).
+	 * Uses SafeMoveUpdatedComponent() for movement, and ComputeSlideVector() to determine the slide direction.
+	 * @param Delta:	Attempted movement vector.
+	 * @param Time:		Percent of Delta to apply (between 0 and 1). Usually equal to the remaining time after a collision: (1.0 - Hit.Time).
+	 * @param Normal:	Normal opposing movement, along which we will slide.
+	 * @param Hit:		[In] HitResult of the attempted move that resulted in the impact triggering the slide. [Out] HitResult of last attempted move.
+	 * @param bHandleImpact:	Whether to call HandleImpact on each hit.
+	 * @return The percentage of requested distance (Delta * Percent) actually applied (between 0 and 1). 0 if no movement occurred, non-zero if movement occurred.
 	 */
 	virtual float SlideAlongSurface(const FVector& Delta, float Time, const FVector& Normal, FHitResult &Hit, bool bHandleImpact = false);
 
