Check for any wrongly saved material coordinates on load

TTP#333911 - Material Editor: Material graph has been reversed

Added another version number to make sure materials aren't continuously checked for this problem in future. Detect whether materials/functions were saved with their coordinates the wrong way around by testing whether most of the inputs to the main outputs are on the right. This will ensure that we don't reverse materials in situations where a strange layout has been setup by the user and we can't be sure whether there is a problem with the coordinates.

diff --git a/Engine/Source/Runtime/Core/Public/UObject/ObjectVersion.h b/Engine/Source/Runtime/Core/Public/UObject/ObjectVersion.h
--- a/Engine/Source/Runtime/Core/Public/UObject/ObjectVersion.h
+++ b/Engine/Source/Runtime/Core/Public/UObject/ObjectVersion.h
@@ -538,6 +538,8 @@ enum EUnrealEngineObjectUE4Version
 	VER_UE4_UCLASS_SERIALIZE_INTERFACES_AFTER_LINKING,
 	// Change from LOD distances to display factors
 	VER_UE4_STATIC_MESH_SCREEN_SIZE_LODS,
+	// Requires test of material coords to ensure they're saved correctly
+	VER_UE4_FIX_MATERIAL_COORDS,
  
 	// -----<new versions can be added before this line>-------------------------------------------------
 	// - this needs to be the last line (see note below)
diff --git a/Engine/Source/Runtime/Engine/Classes/Materials/Material.h b/Engine/Source/Runtime/Engine/Classes/Materials/Material.h
--- a/Engine/Source/Runtime/Engine/Classes/Materials/Material.h
+++ b/Engine/Source/Runtime/Engine/Classes/Materials/Material.h
@@ -1149,11 +1149,17 @@ public:
 	/**
 	 * Flip the X coordinates of a material's expressions and space them out more
 	 *
-	 * @param	Expressions	Array of material expressions
-	 * @param	Comments	Array of material expression comments
-	 * @param	Material	The Material to flip its home coords (optional)
+	 * @param	Expressions		Array of material expressions
+	 * @param	Comments		Array of material expression comments
+	 * @param	bScaleCoords	Whether to scale the coordinates to space out nodes
+	 * @param	Material		The Material to flip its home coords (optional)
 	 */
-	static void FlipExpressionPositions(const TArray<UMaterialExpression*>& Expressions, const TArray<UMaterialExpressionComment*>& Comments, UMaterial* Material = NULL);
+	static void FlipExpressionPositions(const TArray<UMaterialExpression*>& Expressions, const TArray<UMaterialExpressionComment*>& Comments, bool bScaleCoords, UMaterial* Material = NULL);
+
+	/**
+	 * Checks whether a Material is arranged in the old style, with inputs flowing from right to left
+	 */
+	bool HasFlippedCoordinates();
 #endif //WITH_EDITORONLY_DATA
 
 private:
diff --git a/Engine/Source/Runtime/Engine/Classes/Materials/MaterialFunction.h b/Engine/Source/Runtime/Engine/Classes/Materials/MaterialFunction.h
--- a/Engine/Source/Runtime/Engine/Classes/Materials/MaterialFunction.h
+++ b/Engine/Source/Runtime/Engine/Classes/Materials/MaterialFunction.h
@@ -86,6 +86,11 @@ public:
 
 	void UpdateInputOutputTypes();
 
+	/**
+	 * Checks whether a Material Function is arranged in the old style, with inputs flowing from right to left
+	 */
+	bool HasFlippedCoordinates() const;
+
 	UPROPERTY(AssetRegistrySearchable)
 	uint32 CombinedInputTypes;
 
diff --git a/Engine/Source/Runtime/Engine/Private/Materials/Material.cpp b/Engine/Source/Runtime/Engine/Private/Materials/Material.cpp
--- a/Engine/Source/Runtime/Engine/Private/Materials/Material.cpp
+++ b/Engine/Source/Runtime/Engine/Private/Materials/Material.cpp
@@ -32,6 +32,7 @@ FUObjectAnnotationSparseBool GMaterialsThatNeedSamplerFixup;
 FUObjectAnnotationSparseBool GMaterialsThatNeedPhysicalConversion;
 FUObjectAnnotationSparse<FMaterialsWithDirtyUsageFlags,true> GMaterialsWithDirtyUsageFlags;
 FUObjectAnnotationSparseBool GMaterialsThatNeedExpressionsFlipped;
+FUObjectAnnotationSparseBool GMaterialsThatNeedCoordinateCheck;
 
 #endif // #if WITH_EDITOR
 
@@ -1588,6 +1589,10 @@ void UMaterial::Serialize(FArchive& Ar)
 	{
 		GMaterialsThatNeedExpressionsFlipped.Set(this);
 	}
+	else if (Ar.UE4Ver() < VER_UE4_FIX_MATERIAL_COORDS)
+	{
+		GMaterialsThatNeedCoordinateCheck.Set(this);
+	}
 #endif // #if WITH_EDITOR
 
 	if( Ar.UE4Ver() < VER_UE4_MATERIAL_ATTRIBUTES_REORDERING )
@@ -1966,7 +1971,15 @@ void UMaterial::PostLoad()
 	if (GMaterialsThatNeedExpressionsFlipped.Get(this))
 	{
 		GMaterialsThatNeedExpressionsFlipped.Clear(this);
-		FlipExpressionPositions(Expressions, EditorComments, this);
+		FlipExpressionPositions(Expressions, EditorComments, true, this);
+	}
+	else if (GMaterialsThatNeedCoordinateCheck.Get(this))
+	{
+		GMaterialsThatNeedCoordinateCheck.Clear(this);
+		if (HasFlippedCoordinates())
+		{
+			FlipExpressionPositions(Expressions, EditorComments, false, this);
+		}
 	}
 #endif // #if WITH_EDITOR
 }
@@ -3475,10 +3488,10 @@ bool UMaterial::IsPropertyActive(EMaterialProperty InProperty)const
 }
 
 #if WITH_EDITORONLY_DATA
-void UMaterial::FlipExpressionPositions(const TArray<UMaterialExpression*>& Expressions, const TArray<UMaterialExpressionComment*>& Comments, UMaterial* InMaterial)
+void UMaterial::FlipExpressionPositions(const TArray<UMaterialExpression*>& Expressions, const TArray<UMaterialExpressionComment*>& Comments, bool bScaleCoords, UMaterial* InMaterial)
 {
 	// Rough estimate of average increase in node size for the new editor
-	const float PosScaling = 1.25f;
+	const float PosScaling = bScaleCoords ? 1.25f : 1.0f;
 
 	if (InMaterial)
 	{
@@ -3499,4 +3512,30 @@ void UMaterial::FlipExpressionPositions(const TArray<UMaterialExpression*>& Expr
 		Comment->SizeY *= PosScaling;
 	}
 }
+
+bool UMaterial::HasFlippedCoordinates()
+{
+	uint32 ReversedInputCount = 0;
+	uint32 StandardInputCount = 0;
+
+	// Check inputs to see if they are right of the root node
+	for (int32 InputIndex = 0; InputIndex < MP_MAX; InputIndex++)
+	{
+		FExpressionInput* Input = GetExpressionInputForProperty((EMaterialProperty)InputIndex);
+		if (Input->Expression)
+		{
+			if (Input->Expression->MaterialExpressionEditorX > EditorX)
+			{
+				++ReversedInputCount;
+			}
+			else
+			{
+				++StandardInputCount;
+			}
+		}
+	}
+
+	// Can't be sure coords are flipped if most are set out correctly
+	return ReversedInputCount > StandardInputCount;
+}
 #endif //WITH_EDITORONLY_DATA
diff --git a/Engine/Source/Runtime/Engine/Private/Materials/MaterialExpressions.cpp b/Engine/Source/Runtime/Engine/Private/Materials/MaterialExpressions.cpp
--- a/Engine/Source/Runtime/Engine/Private/Materials/MaterialExpressions.cpp
+++ b/Engine/Source/Runtime/Engine/Private/Materials/MaterialExpressions.cpp
@@ -22,6 +22,7 @@ if( ExpressionInput.Expression == ToBeRemovedExpression )										\
 
 #if WITH_EDITOR
 FUObjectAnnotationSparseBool GMaterialFunctionsThatNeedExpressionsFlipped;
+FUObjectAnnotationSparseBool GMaterialFunctionsThatNeedCoordinateCheck;
 #endif // #if WITH_EDITOR
 
 /** Returns whether the given expression class is allowed. */
@@ -6015,6 +6016,10 @@ void UMaterialFunction::Serialize(FArchive& Ar)
 	{
 		GMaterialFunctionsThatNeedExpressionsFlipped.Set(this);
 	}
+	else if (Ar.UE4Ver() < VER_UE4_FIX_MATERIAL_COORDS)
+	{
+		GMaterialFunctionsThatNeedCoordinateCheck.Set(this);
+	}
 #endif // #if WITH_EDITOR
 }
 
@@ -6058,7 +6063,15 @@ void UMaterialFunction::PostLoad()
 	if (GMaterialFunctionsThatNeedExpressionsFlipped.Get(this))
 	{
 		GMaterialFunctionsThatNeedExpressionsFlipped.Clear(this);
-		UMaterial::FlipExpressionPositions(FunctionExpressions, FunctionEditorComments);
+		UMaterial::FlipExpressionPositions(FunctionExpressions, FunctionEditorComments, true);
+	}
+	else if (GMaterialFunctionsThatNeedCoordinateCheck.Get(this))
+	{
+		GMaterialFunctionsThatNeedCoordinateCheck.Clear(this);
+		if (HasFlippedCoordinates())
+		{
+			UMaterial::FlipExpressionPositions(FunctionExpressions, FunctionEditorComments, false);
+		}
 	}
 #endif // #if WITH_EDITOR
 }
@@ -6356,6 +6369,33 @@ void UMaterialFunction::AppendReferencedTextures(TArray<UTexture*>& InOutTexture
 	}
 }
 
+#if WITH_EDITORONLY_DATA
+bool UMaterialFunction::HasFlippedCoordinates() const
+{
+	uint32 ReversedInputCount = 0;
+	uint32 StandardInputCount = 0;
+
+	for (int32 Index = 0; Index < FunctionExpressions.Num(); ++Index)
+	{
+		UMaterialExpressionFunctionOutput* FunctionOutput = Cast<UMaterialExpressionFunctionOutput>(FunctionExpressions[Index]);
+		if (FunctionOutput && FunctionOutput->A.Expression)
+		{
+			if (FunctionOutput->A.Expression->MaterialExpressionEditorX > FunctionOutput->MaterialExpressionEditorX)
+			{
+				++ReversedInputCount;
+			}
+			else
+			{
+				++StandardInputCount;
+			}
+		}
+	}
+
+	// Can't be sure coords are flipped if most are set out correctly
+	return ReversedInputCount > StandardInputCount;
+}
+#endif //WITH_EDITORONLY_DATA
+
 ///////////////////////////////////////////////////////////////////////////////
 // UMaterialExpressionMaterialFunctionCall
 ///////////////////////////////////////////////////////////////////////////////
