Refactoring streaming manager so that it's not so specific to textures.

Renamed FStreamingManagerBase to IStreamingManager and made this into the manager of a Streaming Manager singleton instead of a global variable. Added an interface class to add just the functions specific to texture streaming so that the collection manager doesn't have to deal with textures specifically and prevent this from being part of an Audio Streaming manager in future. This meant a few changes to code using texture functions outside of the streaming managers, with it now having to check that texture streaming is enabled before requesting the texture streaming manager.

diff --git a/Engine/Source/Editor/UnrealEd/Private/Editor.cpp b/Engine/Source/Editor/UnrealEd/Private/Editor.cpp
--- a/Engine/Source/Editor/UnrealEd/Private/Editor.cpp
+++ b/Engine/Source/Editor/UnrealEd/Private/Editor.cpp
@@ -1375,7 +1375,7 @@ void UEditorEngine::Tick( float DeltaSeconds, bool bIdleMode )
 		}
 
 		// Update resource streaming after both regular Editor viewports and PIE had a chance to add viewers.
-		GStreamingManager->Tick( DeltaSeconds );
+		IStreamingManager::Get().Tick( DeltaSeconds );
 	}
 
 	// Update Audio. This needs to occur after rendering as the rendering code updates the listener position.
@@ -1573,7 +1573,7 @@ bool UEditorEngine::UpdateSingleViewportClient(FEditorViewportClient* InViewport
 		// Add view information for perspective viewports.
 		if( InViewportClient->IsPerspective() )
 		{
-			GStreamingManager->AddViewInformation( InViewportClient->GetViewLocation(), InViewportClient->Viewport->GetSizeXY().X, InViewportClient->Viewport->GetSizeXY().X / FMath::Tan(InViewportClient->ViewFOV) );
+			IStreamingManager::Get().AddViewInformation( InViewportClient->GetViewLocation(), InViewportClient->Viewport->GetSizeXY().X, InViewportClient->Viewport->GetSizeXY().X / FMath::Tan(InViewportClient->ViewFOV) );
 			GWorld->ViewLocationsRenderedLastFrame.Add(InViewportClient->GetViewLocation());
 	
 			// If we're currently simulating in editor, then we'll need to make sure that sub-levels are streamed in.
diff --git a/Engine/Source/Editor/UnrealEd/Private/EditorLevelUtils.cpp b/Engine/Source/Editor/UnrealEd/Private/EditorLevelUtils.cpp
--- a/Engine/Source/Editor/UnrealEd/Private/EditorLevelUtils.cpp
+++ b/Engine/Source/Editor/UnrealEd/Private/EditorLevelUtils.cpp
@@ -427,7 +427,7 @@ namespace EditorLevelUtils
 
 		InLevel->ReleaseRenderingResources();
 
-		GStreamingManager->RemoveLevel( InLevel );
+		IStreamingManager::Get().RemoveLevel( InLevel );
 		UWorld* World = InLevel->OwningWorld;
 		World->RemoveLevel(InLevel);
 		InLevel->ClearLevelComponents();
diff --git a/Engine/Source/Editor/UnrealEd/Private/ObjectTools.cpp b/Engine/Source/Editor/UnrealEd/Private/ObjectTools.cpp
--- a/Engine/Source/Editor/UnrealEd/Private/ObjectTools.cpp
+++ b/Engine/Source/Editor/UnrealEd/Private/ObjectTools.cpp
@@ -3304,7 +3304,7 @@ namespace ThumbnailTools
 		{
 			FlushAsyncLoading();
 
-			GStreamingManager->StreamAllResources( 100.0f );
+			IStreamingManager::Get().StreamAllResources( 100.0f );
 		}
 
 		// If this object's thumbnail will be rendered to a texture on the GPU.
diff --git a/Engine/Source/Editor/UnrealEd/Private/ThumbnailHelpers.cpp b/Engine/Source/Editor/UnrealEd/Private/ThumbnailHelpers.cpp
--- a/Engine/Source/Editor/UnrealEd/Private/ThumbnailHelpers.cpp
+++ b/Engine/Source/Editor/UnrealEd/Private/ThumbnailHelpers.cpp
@@ -112,7 +112,7 @@ void FThumbnailPreviewScene::GetView(FSceneViewFamily* ViewFamily, int32 X, int3
 		// NOTE: Sizes may not actually be in screen space depending on how the thumbnail ends up stretched by the UI.  Not a big deal though.
 		// NOTE: Textures still take a little time to stream if the view has not been re-rendered recently, so they may briefly appear blurry while mips are prepared
 		// NOTE: Content Browser only renders thumbnails for loaded assets, and only when the mouse is over the panel. They'll be frozen in their last state while the mouse cursor is not over the panel.  This is for performance reasons
-		GStreamingManager->AddViewInformation( Origin, SizeX, SizeX / FMath::Tan( FOVDegrees ) );
+		IStreamingManager::Get().AddViewInformation( Origin, SizeX, SizeX / FMath::Tan( FOVDegrees ) );
 	}
 }
 
diff --git a/Engine/Source/Runtime/Engine/Private/Actor.cpp b/Engine/Source/Runtime/Engine/Private/Actor.cpp
--- a/Engine/Source/Runtime/Engine/Private/Actor.cpp
+++ b/Engine/Source/Runtime/Engine/Private/Actor.cpp
@@ -2217,7 +2217,7 @@ void AActor::PostActorConstruction()
 	}
 
 	// Notify the texture streaming manager about the new actor.
-	GStreamingManager->NotifyActorSpawned(this);
+	IStreamingManager::Get().NotifyActorSpawned(this);
 }
 
 void AActor::SetReplicates(bool bInReplicates)
diff --git a/Engine/Source/Runtime/Engine/Private/ActorComponent.cpp b/Engine/Source/Runtime/Engine/Private/ActorComponent.cpp
--- a/Engine/Source/Runtime/Engine/Private/ActorComponent.cpp
+++ b/Engine/Source/Runtime/Engine/Private/ActorComponent.cpp
@@ -624,7 +624,7 @@ void UActorComponent::UnregisterComponent()
 	const UPrimitiveComponent* Primitive = Cast<const UPrimitiveComponent>(this);
 	if ( Primitive )
 	{
-		GStreamingManager->NotifyPrimitiveDetached( Primitive );
+		IStreamingManager::Get().NotifyPrimitiveDetached( Primitive );
 	}
 
 	AActor* Owner = GetOwner(); // Get Owner while we are still registered
diff --git a/Engine/Source/Runtime/Engine/Private/ContentStreaming.cpp b/Engine/Source/Runtime/Engine/Private/ContentStreaming.cpp
--- a/Engine/Source/Runtime/Engine/Private/ContentStreaming.cpp
+++ b/Engine/Source/Runtime/Engine/Private/ContentStreaming.cpp
@@ -76,23 +76,20 @@ double GStreamingDynamicPrimitivesTime = 0.0;
 	Globals.
 -----------------------------------------------------------------------------*/
 
-/** Global streaming manager */
-ENGINE_API FStreamingManagerCollection* GStreamingManager;
-
 /** Collection of views that need to be taken into account for streaming. */
-TArray<FStreamingViewInfo> FStreamingManagerBase::CurrentViewInfos;
+TArray<FStreamingViewInfo> IStreamingManager::CurrentViewInfos;
 
 /** Pending views. Emptied every frame. */
-TArray<FStreamingViewInfo> FStreamingManagerBase::PendingViewInfos;
+TArray<FStreamingViewInfo> IStreamingManager::PendingViewInfos;
 
 /** Views that stick around for a while. Override views are ignored if no movie is playing. */
-TArray<FStreamingViewInfo> FStreamingManagerBase::LastingViewInfos;
+TArray<FStreamingViewInfo> IStreamingManager::LastingViewInfos;
 
 /** Collection of view locations that will be added at the next call to AddViewInformation. */
-TArray<FStreamingManagerBase::FSlaveLocation> FStreamingManagerBase::SlaveLocations;
+TArray<IStreamingManager::FSlaveLocation> IStreamingManager::SlaveLocations;
 
 /** Set when Tick() has been called. The first time a new view is added, it will clear out all old views. */
-bool FStreamingManagerBase::bPendingRemoveViews = false;
+bool IStreamingManager::bPendingRemoveViews = false;
 
 /** The lightmap used by the currently selected component (toggledebugcamera), if it's a static mesh component. */
 extern ENGINE_API FLightMap2D* GDebugSelectedLightmap;
@@ -165,7 +162,7 @@ float ENGINE_API GNegativeOneOneTable[2] = {-1.0f,1.0f};
 void FlushResourceStreaming()
 {
 	RETURN_IF_EXIT_REQUESTED;
-	GStreamingManager->BlockTillAllRequestsFinished();
+	IStreamingManager::Get().BlockTillAllRequestsFinished();
 }
 
 /**
@@ -1204,9 +1201,34 @@ private:
 
 
 /*-----------------------------------------------------------------------------
-	FStreamingManagerBase implementation.
+	IStreamingManager implementation.
 -----------------------------------------------------------------------------*/
 
+static FStreamingManagerCollection* StreamingManagerCollection = NULL;
+
+FStreamingManagerCollection& IStreamingManager::Get()
+{
+	if (StreamingManagerCollection == NULL)
+	{
+		StreamingManagerCollection = new FStreamingManagerCollection();
+	}
+	return *StreamingManagerCollection;
+}
+
+void IStreamingManager::Shutdown()
+{
+	if (StreamingManagerCollection != NULL)
+	{
+		delete StreamingManagerCollection;
+		StreamingManagerCollection = (FStreamingManagerCollection*)-1;//Force Error if manager used after shutdown
+	}
+}
+
+bool IStreamingManager::HasShutdown()
+{
+	return StreamingManagerCollection == (FStreamingManagerCollection*)-1;
+}
+
 /**
  * Adds the passed in view information to the static array.
  *
@@ -1219,7 +1241,7 @@ private:
  * @param Duration				How long the streaming system should keep checking this location (in seconds). 0 means just for the next Tick.
  * @param InActorToBoost		Optional pointer to an actor who's textures should have their streaming priority boosted
  */
-void FStreamingManagerBase::AddViewInfoToArray( TArray<FStreamingViewInfo> &ViewInfos, const FVector& ViewOrigin, float ScreenSize, float FOVScreenSize, float BoostFactor, bool bOverrideLocation, float Duration, TWeakObjectPtr<AActor> InActorToBoost )
+void IStreamingManager::AddViewInfoToArray( TArray<FStreamingViewInfo> &ViewInfos, const FVector& ViewOrigin, float ScreenSize, float FOVScreenSize, float BoostFactor, bool bOverrideLocation, float Duration, TWeakObjectPtr<AActor> InActorToBoost )
 {
 	// Check for duplicates and existing overrides.
 	bool bShouldAddView = true;
@@ -1251,7 +1273,7 @@ void FStreamingManagerBase::AddViewInfoToArray( TArray<FStreamingViewInfo> &View
  * @param ViewInfos				[in/out] Array to remove the view from
  * @param ViewOrigin			View origin
  */
-void FStreamingManagerBase::RemoveViewInfoFromArray( TArray<FStreamingViewInfo> &ViewInfos, const FVector& ViewOrigin )
+void IStreamingManager::RemoveViewInfoFromArray( TArray<FStreamingViewInfo> &ViewInfos, const FVector& ViewOrigin )
 {
 	for ( int32 ViewIndex=0; ViewIndex < ViewInfos.Num(); ++ViewIndex )
 	{
@@ -1273,7 +1295,7 @@ TArray<FStreamingViewInfo> GPrevViewLocations;
  *
  * @param DeltaTime		Time since last call in seconds
  */
-void FStreamingManagerBase::SetupViewInfos( float DeltaTime )
+void IStreamingManager::SetupViewInfos( float DeltaTime )
 {
 	// Reset CurrentViewInfos
 	CurrentViewInfos.Empty( PendingViewInfos.Num() + LastingViewInfos.Num() + SlaveLocations.Num() );
@@ -1428,7 +1450,7 @@ void FStreamingManagerBase::SetupViewInfos( float DeltaTime )
  * @param Duration				How long the streaming system should keep checking this location (in seconds). 0 means just for the next Tick.
  * @param InActorToBoost		Optional pointer to an actor who's textures should have their streaming priority boosted
  */
-void FStreamingManagerBase::AddViewInformation( const FVector& ViewOrigin, float ScreenSize, float FOVScreenSize, float BoostFactor/*=1.0f*/, bool bOverrideLocation/*=false*/, float Duration/*=0.0f*/, TWeakObjectPtr<AActor> InActorToBoost /*=NULL*/ )
+void IStreamingManager::AddViewInformation( const FVector& ViewOrigin, float ScreenSize, float FOVScreenSize, float BoostFactor/*=1.0f*/, bool bOverrideLocation/*=false*/, float Duration/*=0.0f*/, TWeakObjectPtr<AActor> InActorToBoost /*=NULL*/ )
 {
 	// Is this a reasonable location?
 	if ( FMath::Abs(ViewOrigin.X) < (1.0e+20f) && FMath::Abs(ViewOrigin.Y) < (1.0e+20f) && FMath::Abs(ViewOrigin.Z) < (1.0e+20f) )
@@ -1474,7 +1496,7 @@ void FStreamingManagerBase::AddViewInformation( const FVector& ViewOrigin, float
  * @param bOverrideLocation		Whether this is an override location, which forces the streaming system to ignore all other locations
  * @param Duration				How long the streaming system should keep checking this location (in seconds). 0 means just for the next Tick.
  */
-void FStreamingManagerBase::AddViewSlaveLocation( const FVector& SlaveLocation, float BoostFactor/*=1.0f*/, bool bOverrideLocation/*=false*/, float Duration/*=0.0f*/ )
+void IStreamingManager::AddViewSlaveLocation( const FVector& SlaveLocation, float BoostFactor/*=1.0f*/, bool bOverrideLocation/*=false*/, float Duration/*=0.0f*/ )
 {
 	BoostFactor *= CVarStreamingBoost.GetValueOnGameThread();
 
@@ -1494,7 +1516,7 @@ void FStreamingManagerBase::AddViewSlaveLocation( const FVector& SlaveLocation,
  *
  * @param RemovalType	What types of views to remove (all or just the normal views)
  */
-void FStreamingManagerBase::RemoveStreamingViews( ERemoveStreamingViews RemovalType )
+void IStreamingManager::RemoveStreamingViews( ERemoveStreamingViews RemovalType )
 {
 	PendingViewInfos.Empty();
 	SlaveLocations.Empty();
@@ -1510,7 +1532,7 @@ void FStreamingManagerBase::RemoveStreamingViews( ERemoveStreamingViews RemovalT
  * @param DeltaTime				Time since last call in seconds
  * @param bProcessEverything	[opt] If true, process all resources with no throttling limits
  */
-void FStreamingManagerBase::Tick( float DeltaTime, bool bProcessEverything/*=false*/ )
+void IStreamingManager::Tick( float DeltaTime, bool bProcessEverything/*=false*/ )
 {
 	UpdateResourceStreaming( DeltaTime, bProcessEverything );
 
@@ -1562,7 +1584,7 @@ void FStreamingManagerCollection::Tick( float DeltaTime, bool bProcessEverything
 {
 	AddOrRemoveTextureStreamingManagerIfNeeded();
 
-	FStreamingManagerBase::Tick(DeltaTime, bProcessEverything);
+	IStreamingManager::Tick(DeltaTime, bProcessEverything);
 }
 
 /**
@@ -1590,7 +1612,7 @@ void FStreamingManagerCollection::UpdateResourceStreaming( float DeltaTime, bool
 			// Route to streaming managers.
 			for( int32 ManagerIndex=0; ManagerIndex<StreamingManagers.Num(); ManagerIndex++ )
 			{
-				FStreamingManagerBase* StreamingManager = StreamingManagers[ManagerIndex];
+				IStreamingManager* StreamingManager = StreamingManagers[ManagerIndex];
 				StreamingManager->UpdateResourceStreaming( DeltaTime, bProcessEverything );
 			}
 		}
@@ -1630,25 +1652,6 @@ int32 FStreamingManagerCollection::StreamAllResources( float TimeLimit/*=0.0f*/
 	return NumPendingRequests;
 }
 
-/**
- * Updates streaming for an individual texture, taking into account all view infos.
- *
- * @param Texture	Texture to update
- */
-void FStreamingManagerCollection::UpdateIndividualResource( UTexture2D* Texture )
-{
-	// only allow this if its not disabled
-	if (DisableResourceStreamingCount == 0)
-	{
-		// Route to streaming managers.
-		for( int32 ManagerIndex=0; ManagerIndex<StreamingManagers.Num(); ManagerIndex++ )
-		{
-			FStreamingManagerBase* StreamingManager = StreamingManagers[ManagerIndex];
-			StreamingManager->UpdateIndividualResource( Texture );
-		}
-	}
-}
-
 /**
  * Blocks till all pending requests are fulfilled.
  *
@@ -1663,7 +1666,7 @@ int32 FStreamingManagerCollection::BlockTillAllRequestsFinished( float TimeLimit
 	// Route to streaming managers.
 	for( int32 ManagerIndex=0; ManagerIndex<StreamingManagers.Num(); ManagerIndex++ )
 	{
-		FStreamingManagerBase* StreamingManager = StreamingManagers[ManagerIndex];
+		IStreamingManager* StreamingManager = StreamingManagers[ManagerIndex];
 		NumPendingRequests += StreamingManager->BlockTillAllRequestsFinished( TimeLimit, bLogResults );
 	}
 
@@ -1678,7 +1681,7 @@ int32 FStreamingManagerCollection::GetNumWantingResources() const
 	// Route to streaming managers.
 	for( int32 ManagerIndex=0; ManagerIndex<StreamingManagers.Num(); ManagerIndex++ )
 	{
-		FStreamingManagerBase* StreamingManager = StreamingManagers[ManagerIndex];
+		IStreamingManager* StreamingManager = StreamingManagers[ManagerIndex];
 		NumWantingResources += StreamingManager->GetNumWantingResources();
 	}
 
@@ -1698,7 +1701,7 @@ int32 FStreamingManagerCollection::GetNumWantingResourcesID() const
 	// Route to streaming managers.
 	for( int32 ManagerIndex=0; ManagerIndex<StreamingManagers.Num(); ManagerIndex++ )
 	{
-		FStreamingManagerBase* StreamingManager = StreamingManagers[ManagerIndex];
+		IStreamingManager* StreamingManager = StreamingManagers[ManagerIndex];
 		NumWantingResourcesCounter = FMath::Min( NumWantingResourcesCounter, StreamingManager->GetNumWantingResourcesID() );
 	}
 
@@ -1713,7 +1716,7 @@ void FStreamingManagerCollection::CancelForcedResources()
 	// Route to streaming managers.
 	for( int32 ManagerIndex=0; ManagerIndex<StreamingManagers.Num(); ManagerIndex++ )
 	{
-		FStreamingManagerBase* StreamingManager = StreamingManagers[ManagerIndex];
+		IStreamingManager* StreamingManager = StreamingManagers[ManagerIndex];
 		StreamingManager->CancelForcedResources();
 	}
 }
@@ -1726,38 +1729,35 @@ void FStreamingManagerCollection::NotifyLevelChange()
 	// Route to streaming managers.
 	for( int32 ManagerIndex=0; ManagerIndex<StreamingManagers.Num(); ManagerIndex++ )
 	{
-		FStreamingManagerBase* StreamingManager = StreamingManagers[ManagerIndex];
+		IStreamingManager* StreamingManager = StreamingManagers[ManagerIndex];
 		StreamingManager->NotifyLevelChange();
 	}
 }
 
+bool FStreamingManagerCollection::IsStreamingEnabled() const
+{
+	return DisableResourceStreamingCount == 0;
+}
+
 bool FStreamingManagerCollection::IsTextureStreamingEnabled() const
 {
 	return TextureStreamingManager != 0;
 }
 
-/** Don't stream world resources for the next NumFrames. */
-void FStreamingManagerCollection::SetDisregardWorldResourcesForFrames(int32 NumFrames )
+ITextureStreamingManager& FStreamingManagerCollection::GetTextureStreamingManager() const
 {
-	// Route to streaming managers.
-	for( int32 ManagerIndex=0; ManagerIndex<StreamingManagers.Num(); ManagerIndex++ )
-	{
-		FStreamingManagerBase* StreamingManager = StreamingManagers[ManagerIndex];
-		StreamingManager->SetDisregardWorldResourcesForFrames(NumFrames);
-	}
+	check(TextureStreamingManager != 0);
+	return *TextureStreamingManager;
 }
 
-/**
- * Temporarily boosts the streaming distance factor by the specified number.
- * This factor is automatically reset to 1.0 after it's been used for mip-calculations.
- */
-void FStreamingManagerCollection::BoostTextures( AActor* Actor, float BoostFactor )
+/** Don't stream world resources for the next NumFrames. */
+void FStreamingManagerCollection::SetDisregardWorldResourcesForFrames(int32 NumFrames )
 {
 	// Route to streaming managers.
 	for( int32 ManagerIndex=0; ManagerIndex<StreamingManagers.Num(); ManagerIndex++ )
 	{
-		FStreamingManagerBase* StreamingManager = StreamingManagers[ManagerIndex];
-		StreamingManager->BoostTextures(Actor, BoostFactor);
+		IStreamingManager* StreamingManager = StreamingManagers[ManagerIndex];
+		StreamingManager->SetDisregardWorldResourcesForFrames(NumFrames);
 	}
 }
 
@@ -1766,7 +1766,7 @@ void FStreamingManagerCollection::BoostTextures( AActor* Actor, float BoostFacto
  *
  * @param StreamingManager	Streaming manager to add
  */
-void FStreamingManagerCollection::AddStreamingManager( FStreamingManagerBase* StreamingManager )
+void FStreamingManagerCollection::AddStreamingManager( IStreamingManager* StreamingManager )
 {
 	StreamingManagers.Add( StreamingManager );
 }
@@ -1776,7 +1776,7 @@ void FStreamingManagerCollection::AddStreamingManager( FStreamingManagerBase* St
  *
  * @param StreamingManager	Streaming manager to remove
  */
-void FStreamingManagerCollection::RemoveStreamingManager( FStreamingManagerBase* StreamingManager )
+void FStreamingManagerCollection::RemoveStreamingManager( IStreamingManager* StreamingManager )
 {
 	StreamingManagers.Remove( StreamingManager );
 }
@@ -1802,25 +1802,6 @@ void FStreamingManagerCollection::EnableResourceStreaming()
 	checkf(DisableResourceStreamingCount >= 0, TEXT("Mismatched number of calls to FStreamingManagerCollection::DisableResourceStreaming/EnableResourceStreaming"));
 }
 
-/**
- *	Try to stream out texture mip-levels to free up more memory.
- *	@param RequiredMemorySize	- Required minimum available texture memory
- *	@return						- Whether it succeeded or not
- **/
-bool FStreamingManagerCollection::StreamOutTextureData( int32 RequiredMemorySize )
-{
-	// Route to streaming managers.
-	for( int32 ManagerIndex=0; ManagerIndex<StreamingManagers.Num(); ManagerIndex++ )
-	{
-		FStreamingManagerBase* StreamingManager = StreamingManagers[ManagerIndex];
-		if ( StreamingManager->StreamOutTextureData( RequiredMemorySize ) )
-		{
-			return true;
-		}
-	}
-	return false;
-}
-
 /**
  * Allows the streaming manager to process exec commands.
  *
@@ -1833,7 +1814,7 @@ bool FStreamingManagerCollection::Exec( UWorld* InWorld, const TCHAR* Cmd, FOutp
 	// Route to streaming managers.
 	for( int32 ManagerIndex=0; ManagerIndex<StreamingManagers.Num(); ManagerIndex++ )
 	{
-		FStreamingManagerBase* StreamingManager = StreamingManagers[ManagerIndex];
+		IStreamingManager* StreamingManager = StreamingManagers[ManagerIndex];
 		if ( StreamingManager->Exec( InWorld, Cmd, Ar ) )
 		{
 			return true;
@@ -1842,32 +1823,6 @@ bool FStreamingManagerCollection::Exec( UWorld* InWorld, const TCHAR* Cmd, FOutp
 	return false;
 }
 
-/**
- * Adds a new texture to the streaming manager.
- */
-void FStreamingManagerCollection::AddStreamingTexture( UTexture2D* Texture )
-{
-	// Route to streaming managers.
-	for( int32 ManagerIndex=0; ManagerIndex<StreamingManagers.Num(); ManagerIndex++ )
-	{
-		FStreamingManagerBase* StreamingManager = StreamingManagers[ManagerIndex];
-		StreamingManager->AddStreamingTexture( Texture );
-	}
-}
-
-/**
- * Removes a texture from the streaming manager.
- */
-void FStreamingManagerCollection::RemoveStreamingTexture( UTexture2D* Texture )
-{
-	// Route to streaming managers.
-	for( int32 ManagerIndex=0; ManagerIndex<StreamingManagers.Num(); ManagerIndex++ )
-	{
-		FStreamingManagerBase* StreamingManager = StreamingManagers[ManagerIndex];
-		StreamingManager->RemoveStreamingTexture( Texture );
-	}
-}
-
 /** Adds a ULevel to the streaming manager. */
 void FStreamingManagerCollection::AddLevel( ULevel* Level )
 {
@@ -1897,7 +1852,7 @@ void FStreamingManagerCollection::RemoveLevel( ULevel* Level )
 	// Route to streaming managers.
 	for( int32 ManagerIndex=0; ManagerIndex<StreamingManagers.Num(); ManagerIndex++ )
 	{
-		FStreamingManagerBase* StreamingManager = StreamingManagers[ManagerIndex];
+		IStreamingManager* StreamingManager = StreamingManagers[ManagerIndex];
 		StreamingManager->RemoveLevel( Level );
 	}
 }
@@ -1914,7 +1869,7 @@ void FStreamingManagerCollection::AddPreparedLevel( class ULevel* Level )
 	// Route to streaming managers.
 	for( int32 ManagerIndex=0; ManagerIndex<StreamingManagers.Num(); ManagerIndex++ )
 	{
-		FStreamingManagerBase* StreamingManager = StreamingManagers[ManagerIndex];
+		IStreamingManager* StreamingManager = StreamingManagers[ManagerIndex];
 		StreamingManager->AddPreparedLevel( Level );
 	}
 }
@@ -1927,7 +1882,7 @@ void FStreamingManagerCollection::NotifyActorSpawned( AActor* Actor )
 	// Route to streaming managers.
 	for( int32 ManagerIndex=0; ManagerIndex<StreamingManagers.Num(); ManagerIndex++ )
 	{
-		FStreamingManagerBase* StreamingManager = StreamingManagers[ManagerIndex];
+		IStreamingManager* StreamingManager = StreamingManagers[ManagerIndex];
 		StreamingManager->NotifyActorSpawned( Actor );
 	}
 	STAT( GStreamingDynamicPrimitivesTime += FPlatformTime::Seconds() - StartTime );
@@ -1941,7 +1896,7 @@ void FStreamingManagerCollection::NotifyActorDestroyed( AActor* Actor )
 	// Route to streaming managers.
 	for( int32 ManagerIndex=0; ManagerIndex<StreamingManagers.Num(); ManagerIndex++ )
 	{
-		FStreamingManagerBase* StreamingManager = StreamingManagers[ManagerIndex];
+		IStreamingManager* StreamingManager = StreamingManagers[ManagerIndex];
 		StreamingManager->NotifyActorDestroyed( Actor );
 	}
 	STAT( GStreamingDynamicPrimitivesTime += FPlatformTime::Seconds() - StartTime );
@@ -1962,7 +1917,7 @@ void FStreamingManagerCollection::NotifyPrimitiveAttached( const UPrimitiveCompo
 	{
 		for( int32 ManagerIndex=0; ManagerIndex<StreamingManagers.Num(); ManagerIndex++ )
 		{
-			FStreamingManagerBase* StreamingManager = StreamingManagers[ManagerIndex];
+			IStreamingManager* StreamingManager = StreamingManagers[ManagerIndex];
 			StreamingManager->NotifyPrimitiveAttached( Primitive, DynamicType );
 		}
 	}
@@ -1980,7 +1935,7 @@ void FStreamingManagerCollection::NotifyPrimitiveDetached( const UPrimitiveCompo
 		// Route to streaming managers.
 		for( int32 ManagerIndex=0; ManagerIndex<StreamingManagers.Num(); ManagerIndex++ )
 		{
-			FStreamingManagerBase* StreamingManager = StreamingManagers[ManagerIndex];
+			IStreamingManager* StreamingManager = StreamingManagers[ManagerIndex];
 			StreamingManager->NotifyPrimitiveDetached( Primitive );
 		}
 	}
@@ -2001,28 +1956,13 @@ void FStreamingManagerCollection::NotifyPrimitiveUpdated( const UPrimitiveCompon
 		// Route to streaming managers.
 		for( int32 ManagerIndex=0; ManagerIndex<StreamingManagers.Num(); ManagerIndex++ )
 		{
-			FStreamingManagerBase* StreamingManager = StreamingManagers[ManagerIndex];
+			IStreamingManager* StreamingManager = StreamingManagers[ManagerIndex];
 			StreamingManager->NotifyPrimitiveUpdated( Primitive );
 		}
 	}
 	STAT( GStreamingDynamicPrimitivesTime += FPlatformTime::Seconds() - StartTime );
 }
 
-/** Returns true if this is a streaming resource that is managed by the streaming manager. */
-bool FStreamingManagerCollection::IsManagedStreamingResource( const UTexture2D* Texture2D )
-{
-	// Route to streaming managers.
-	for( int32 ManagerIndex=0; ManagerIndex<StreamingManagers.Num(); ManagerIndex++ )
-	{
-		FStreamingManagerBase* StreamingManager = StreamingManagers[ManagerIndex];
-		if ( StreamingManager->IsManagedStreamingResource( Texture2D ) )
-		{
-			return true;
-		}
-	}
-	return false;
-}
-
 void FStreamingManagerCollection::AddOrRemoveTextureStreamingManagerIfNeeded(bool bIsInit)
 {
 	bool bUseTextureStreaming = false;
@@ -3191,11 +3131,14 @@ bool FStreamingManagerTexture::IsManagedStreamingTexture( const UTexture2D* Text
  *
  * @param Texture	Texture to update
  */
-void FStreamingManagerTexture::UpdateIndividualResource( UTexture2D* Texture )
+void FStreamingManagerTexture::UpdateIndividualTexture( UTexture2D* Texture )
 {
-	IndividualStreamingTexture = Texture;
-	UpdateResourceStreaming( 0.0f );
-	IndividualStreamingTexture = NULL;
+	if (IStreamingManager::Get().IsStreamingEnabled())
+	{
+		IndividualStreamingTexture = Texture;
+		UpdateResourceStreaming( 0.0f );
+		IndividualStreamingTexture = NULL;
+	}
 }
 
 /**
@@ -4639,7 +4582,7 @@ bool FStreamingManagerTexture::HandleStreamOutCommand( const TCHAR* Cmd, FOutput
 	int32 FreeMB = (Parameter.Len() > 0) ? FCString::Atoi(*Parameter) : 0;
 	if ( FreeMB > 0 )
 	{
-		bool bSucceeded = GStreamingManager->StreamOutTextureData( FreeMB * 1024 * 1024 );
+		bool bSucceeded = StreamOutTextureData( FreeMB * 1024 * 1024 );
 		Ar.Logf( TEXT("Tried to stream out %d MB of texture data: %s"), FreeMB, bSucceeded ? TEXT("Succeeded") : TEXT("Failed") );
 	}
 	else
diff --git a/Engine/Source/Runtime/Engine/Private/GameEngine.cpp b/Engine/Source/Runtime/Engine/Private/GameEngine.cpp
--- a/Engine/Source/Runtime/Engine/Private/GameEngine.cpp
+++ b/Engine/Source/Runtime/Engine/Private/GameEngine.cpp
@@ -1005,11 +1005,11 @@ void UGameEngine::Tick( float DeltaSeconds, bool bIdleMode )
 
 		if( GIsClient )
 		{
-			// GStreamingManager is updated outside of a world context. For now, assuming it needs to tick here, before possibly calling PostLoadMap. 
+			// IStreamingManager is updated outside of a world context. For now, assuming it needs to tick here, before possibly calling PostLoadMap. 
 			// Will need to take another look when trying to support multiple worlds.
 
 			// Update resource streaming after viewports have had a chance to update view information. Normal update.
-			GStreamingManager->Tick( DeltaSeconds );
+			IStreamingManager::Get().Tick( DeltaSeconds );
 
 			if ( Context.World()->bTriggerPostLoadMap )
 			{
diff --git a/Engine/Source/Runtime/Engine/Private/GameViewportClient.cpp b/Engine/Source/Runtime/Engine/Private/GameViewportClient.cpp
--- a/Engine/Source/Runtime/Engine/Private/GameViewportClient.cpp
+++ b/Engine/Source/Runtime/Engine/Private/GameViewportClient.cpp
@@ -777,7 +777,7 @@ void UGameViewportClient::Draw(FViewport* InViewport, FCanvas* SceneCanvas)
 						}
 						
 						// Add view information for resource streaming.
-						GStreamingManager->AddViewInformation( View->ViewMatrices.ViewOrigin, View->ViewRect.Width(), View->ViewRect.Width() * View->ViewMatrices.ProjMatrix.M[0][0] );
+						IStreamingManager::Get().AddViewInformation( View->ViewMatrices.ViewOrigin, View->ViewRect.Width(), View->ViewRect.Width() * View->ViewMatrices.ProjMatrix.M[0][0] );
 						GetWorld()->ViewLocationsRenderedLastFrame.Add(View->ViewMatrices.ViewOrigin);
 					}
 				}
diff --git a/Engine/Source/Runtime/Engine/Private/Interpolation.cpp b/Engine/Source/Runtime/Engine/Private/Interpolation.cpp
--- a/Engine/Source/Runtime/Engine/Private/Interpolation.cpp
+++ b/Engine/Source/Runtime/Engine/Private/Interpolation.cpp
@@ -580,7 +580,7 @@ void AMatineeActor::UpdateStreamingForCameraCuts(float CurrentTime, bool bPrevie
 			float TimeDifference = CutInfo.TimeStamp - CurrentTime;
 			if ( TimeDifference > 0.0f && TimeDifference < GCameraCutLookAhead )
 			{
-				GStreamingManager->AddViewSlaveLocation( CutInfo.Location );
+				IStreamingManager::Get().AddViewSlaveLocation( CutInfo.Location );
 			}
 			else if ( TimeDifference >= GCameraCutLookAhead )
 			{
diff --git a/Engine/Source/Runtime/Engine/Private/Level.cpp b/Engine/Source/Runtime/Engine/Private/Level.cpp
--- a/Engine/Source/Runtime/Engine/Private/Level.cpp
+++ b/Engine/Source/Runtime/Engine/Private/Level.cpp
@@ -649,10 +649,10 @@ void ULevel::ClearLevelComponents()
 
 void ULevel::BeginDestroy()
 {
-	if ( GStreamingManager )
+	if (!IStreamingManager::HasShutdown())
 	{
 		// At this time, referenced UTexture2Ds are still in memory.
-		GStreamingManager->RemoveLevel( this );
+		IStreamingManager::Get().RemoveLevel( this );
 	}
 
 	Super::BeginDestroy();
@@ -1268,7 +1268,7 @@ void ULevel::BuildStreamingData(UTexture2D* UpdateSpecificTextureOnly/*=NULL*/)
 	if ( UpdateSpecificTextureOnly == NULL )
 	{
 		// Reset the streaming manager, when building data for a whole Level
-		GStreamingManager->RemoveLevel( this );
+		IStreamingManager::Get().RemoveLevel( this );
 		TextureToInstancesMap.Empty();
 		DynamicTextureInstances.Empty();
 		ForceStreamTextures.Empty();
@@ -1426,7 +1426,7 @@ void ULevel::BuildStreamingData(UTexture2D* UpdateSpecificTextureOnly/*=NULL*/)
 		bTextureStreamingBuilt = true;
 
 		// Update the streaming manager.
-		GStreamingManager->AddPreparedLevel( this );
+		IStreamingManager::Get().AddPreparedLevel( this );
 	}
 }
 
diff --git a/Engine/Source/Runtime/Engine/Private/LevelActor.cpp b/Engine/Source/Runtime/Engine/Private/LevelActor.cpp
--- a/Engine/Source/Runtime/Engine/Private/LevelActor.cpp
+++ b/Engine/Source/Runtime/Engine/Private/LevelActor.cpp
@@ -434,7 +434,7 @@ bool UWorld::DestroyActor( AActor* ThisActor, bool bNetForce, bool bShouldModify
 	ThisActor->bPendingKillPending = true;
 
 	// Notify the texture streaming manager about the destruction of this actor.
-	GStreamingManager->NotifyActorDestroyed( ThisActor );
+	IStreamingManager::Get().NotifyActorDestroyed( ThisActor );
 
 	// Tell this actor it's about to be destroyed.
 	ThisActor->Destroyed();
diff --git a/Engine/Source/Runtime/Engine/Private/MeshComponent.cpp b/Engine/Source/Runtime/Engine/Private/MeshComponent.cpp
--- a/Engine/Source/Runtime/Engine/Private/MeshComponent.cpp
+++ b/Engine/Source/Runtime/Engine/Private/MeshComponent.cpp
@@ -15,7 +15,7 @@ UMeshComponent::UMeshComponent(const class FPostConstructInitializeProperties& P
 void UMeshComponent::BeginDestroy()
 {
 	// Notify the streaming system.
-	GStreamingManager->NotifyPrimitiveDetached( this );
+	IStreamingManager::Get().NotifyPrimitiveDetached( this );
 
 	Super::BeginDestroy();
 }
@@ -77,7 +77,7 @@ void UMeshComponent::PrestreamTextures( float Seconds, bool bPrioritizeCharacter
 	// If requested, tell the streaming system to only process character textures for 30 frames.
 	if ( bPrioritizeCharacterTextures )
 	{
-		GStreamingManager->SetDisregardWorldResourcesForFrames( 30 );
+		IStreamingManager::Get().SetDisregardWorldResourcesForFrames( 30 );
 	}
 
 	int32 NumElements = GetNumMaterials();
diff --git a/Engine/Source/Runtime/Engine/Private/PlayerController.cpp b/Engine/Source/Runtime/Engine/Private/PlayerController.cpp
--- a/Engine/Source/Runtime/Engine/Private/PlayerController.cpp
+++ b/Engine/Source/Runtime/Engine/Private/PlayerController.cpp
@@ -265,9 +265,9 @@ bool APlayerController::ServerUpdateLevelVisibility_Validate(FName PackageName,
 
 void APlayerController::ClientAddTextureStreamingLoc_Implementation(FVector InLoc, float Duration, bool bOverrideLocation )
 {
-	if (GStreamingManager != NULL)
+	if (!IStreamingManager::HasShutdown())
 	{
-		GStreamingManager->AddViewSlaveLocation(InLoc, 1.0f, bOverrideLocation, Duration);
+		IStreamingManager::Get().AddViewSlaveLocation(InLoc, 1.0f, bOverrideLocation, Duration);
 	}
 }
 
diff --git a/Engine/Source/Runtime/Engine/Private/PrimitiveComponent.cpp b/Engine/Source/Runtime/Engine/Private/PrimitiveComponent.cpp
--- a/Engine/Source/Runtime/Engine/Private/PrimitiveComponent.cpp
+++ b/Engine/Source/Runtime/Engine/Private/PrimitiveComponent.cpp
@@ -266,7 +266,7 @@ void UPrimitiveComponent::OnRegister()
 	Super::OnRegister();
 
 	// Notify the streaming system. Will only update the component data if it's already tracked.
-	GStreamingManager->NotifyPrimitiveUpdated(this);
+	IStreamingManager::Get().NotifyPrimitiveUpdated(this);
 
 	AActor* Owner = GetOwner();
 	if (bCanEverAffectNavigation && Owner != NULL)
diff --git a/Engine/Source/Runtime/Engine/Private/SkinnedMeshComponent.cpp b/Engine/Source/Runtime/Engine/Private/SkinnedMeshComponent.cpp
--- a/Engine/Source/Runtime/Engine/Private/SkinnedMeshComponent.cpp
+++ b/Engine/Source/Runtime/Engine/Private/SkinnedMeshComponent.cpp
@@ -748,7 +748,7 @@ void USkinnedMeshComponent::SetSkeletalMesh(USkeletalMesh* InSkelMesh)
 	
 	// Notify the streaming system. Don't use Update(), because this may be the first time the mesh has been set
 	// and the component may have to be added to the streaming system for the first time.
-	GStreamingManager->NotifyPrimitiveAttached( this, DPT_Spawned );
+	IStreamingManager::Get().NotifyPrimitiveAttached( this, DPT_Spawned );
 }
 
 FSkeletalMeshResource* USkinnedMeshComponent::GetSkeletalMeshResource() const
diff --git a/Engine/Source/Runtime/Engine/Private/StaticMeshComponent.cpp b/Engine/Source/Runtime/Engine/Private/StaticMeshComponent.cpp
--- a/Engine/Source/Runtime/Engine/Private/StaticMeshComponent.cpp
+++ b/Engine/Source/Runtime/Engine/Private/StaticMeshComponent.cpp
@@ -1139,7 +1139,7 @@ bool UStaticMeshComponent::SetStaticMesh(UStaticMesh* NewMesh)
 
 	// Notify the streaming system. Don't use Update(), because this may be the first time the mesh has been set
 	// and the component may have to be added to the streaming system for the first time.
-	GStreamingManager->NotifyPrimitiveAttached( this, DPT_Spawned );
+	IStreamingManager::Get().NotifyPrimitiveAttached( this, DPT_Spawned );
 
 	// Since we have new mesh, we need to update bounds
 	UpdateBounds();
diff --git a/Engine/Source/Runtime/Engine/Private/Texture.cpp b/Engine/Source/Runtime/Engine/Private/Texture.cpp
--- a/Engine/Source/Runtime/Engine/Private/Texture.cpp
+++ b/Engine/Source/Runtime/Engine/Private/Texture.cpp
@@ -436,7 +436,7 @@ UEnum* UTexture::GetPixelFormatEnum()
 
 bool UTexture::ForceUpdateTextureStreaming()
 {
-	if ( GStreamingManager )
+	if (!IStreamingManager::HasShutdown())
 	{
 		// Make sure textures can be streamed out so that we can unload current mips.
 		static auto CVarOnlyStreamInTextures = IConsoleManager::Get().FindConsoleVariable(TEXT("r.OnlyStreamInTextures"));
@@ -454,11 +454,11 @@ bool UTexture::ForceUpdateTextureStreaming()
 #endif // #if WITH_EDITORONLY_DATA
 
 		// Make sure we iterate over all textures by setting it to high value.
-		GStreamingManager->SetNumIterationsForNextFrame( 100 );
+		IStreamingManager::Get().SetNumIterationsForNextFrame( 100 );
 		// Update resource streaming with updated texture LOD bias/ max texture mip count.
-		GStreamingManager->UpdateResourceStreaming( 0 );
+		IStreamingManager::Get().UpdateResourceStreaming( 0 );
 		// Block till requests are finished.
-		GStreamingManager->BlockTillAllRequestsFinished();
+		IStreamingManager::Get().BlockTillAllRequestsFinished();
 
 		// Restore streaming out of textures.
 		CVarOnlyStreamInTextures->Set(bOldOnlyStreamInTextures);
diff --git a/Engine/Source/Runtime/Engine/Private/Texture2D.cpp b/Engine/Source/Runtime/Engine/Private/Texture2D.cpp
--- a/Engine/Source/Runtime/Engine/Private/Texture2D.cpp
+++ b/Engine/Source/Runtime/Engine/Private/Texture2D.cpp
@@ -397,17 +397,17 @@ float UTexture2D::GetAverageBrightness(bool bIgnoreTrueBlack, bool bUseGrayscale
 
 void UTexture2D::LinkStreaming()
 {
-	if (!IsTemplate())
+	if (!IsTemplate() && IStreamingManager::Get().IsTextureStreamingEnabled())
 	{
-		GStreamingManager->AddStreamingTexture(this);
+		IStreamingManager::Get().GetTextureStreamingManager().AddStreamingTexture(this);
 	}
 }
 
 void UTexture2D::UnlinkStreaming()
 {
-	if (!IsTemplate())
+	if (!IsTemplate() && IStreamingManager::Get().IsTextureStreamingEnabled())
 	{
-		GStreamingManager->RemoveStreamingTexture(this);
+		IStreamingManager::Get().GetTextureStreamingManager().RemoveStreamingTexture(this);
 	}
 }
 
@@ -549,7 +549,10 @@ bool UTexture2D::IsReadyForStreaming()
 
 void UTexture2D::WaitForStreaming()
 {
-	GStreamingManager->UpdateIndividualResource( this );
+	if (IStreamingManager::Get().IsTextureStreamingEnabled())
+	{
+		IStreamingManager::Get().GetTextureStreamingManager().UpdateIndividualTexture( this );
+	}
 
 	int32 RequestStatus = PendingMipChangeRequestStatus.GetValue();
 
@@ -831,7 +834,7 @@ FTextureResource* UTexture2D::CreateResource()
 	int32 NumMips = GetNumMips();
 
 	// Determine whether or not this texture can be streamed.
-	bIsStreamable = GStreamingManager->IsTextureStreamingEnabled() &&
+	bIsStreamable = IStreamingManager::Get().IsTextureStreamingEnabled() &&
 					!NeverStream && 
 					(NumMips > 1) && 
 					(LODGroup != TEXTUREGROUP_UI) && 
diff --git a/Engine/Source/Runtime/Engine/Private/UnrealClient.cpp b/Engine/Source/Runtime/Engine/Private/UnrealClient.cpp
--- a/Engine/Source/Runtime/Engine/Private/UnrealClient.cpp
+++ b/Engine/Source/Runtime/Engine/Private/UnrealClient.cpp
@@ -1010,7 +1010,7 @@ void FViewport::Draw( bool bShouldPresent /*= true */)
 						if(View)
 						{
 							// Add view information for resource streaming.
-							GStreamingManager->AddViewInformation( View->ViewMatrices.ViewOrigin, View->ViewRect.Width(), View->ViewRect.Width() * View->ViewMatrices.ProjMatrix.M[0][0] );
+							IStreamingManager::Get().AddViewInformation( View->ViewMatrices.ViewOrigin, View->ViewRect.Width(), View->ViewRect.Width() * View->ViewMatrices.ProjMatrix.M[0][0] );
 							World->ViewLocationsRenderedLastFrame.Add(View->ViewMatrices.ViewOrigin);
 						}
 					}
diff --git a/Engine/Source/Runtime/Engine/Private/UnrealEngine.cpp b/Engine/Source/Runtime/Engine/Private/UnrealEngine.cpp
--- a/Engine/Source/Runtime/Engine/Private/UnrealEngine.cpp
+++ b/Engine/Source/Runtime/Engine/Private/UnrealEngine.cpp
@@ -706,7 +706,7 @@ void UEngine::Init(IEngineLoop* InEngineLoop)
 	OnTravelFailure().AddUObject(this, &UEngine::HandleTravelFailure);
 	OnNetworkFailure().AddUObject(this, &UEngine::HandleNetworkFailure);
 
-	UE_LOG(LogInit, Log, TEXT("Texture streaming: %s"), GStreamingManager->IsTextureStreamingEnabled() ? TEXT("Enabled") : TEXT("Disabled") );
+	UE_LOG(LogInit, Log, TEXT("Texture streaming: %s"), IStreamingManager::Get().IsTextureStreamingEnabled() ? TEXT("Enabled") : TEXT("Disabled") );
 
 	IOnlineSubsystem* SubSystem = IOnlineSubsystem::Get();
 	if(SubSystem)
@@ -2320,7 +2320,7 @@ bool UEngine::Exec( UWorld* InWorld, const TCHAR* Cmd, FOutputDevice& Ar )
 		// console variable interaction (get value, set value or get help)
 		return true;
 	}
-	else if (GStreamingManager && GStreamingManager->Exec( InWorld, Cmd,Ar ))
+	else if (!IStreamingManager::HasShutdown() && IStreamingManager::Get().Exec( InWorld, Cmd,Ar ))
 	{
 		// The streaming manager has handled the exec command.
 	}
@@ -2990,7 +2990,7 @@ bool UEngine::HandleListTexturesCommand( const TCHAR* Cmd, FOutputDevice& Ar )
 		int32				DroppedMips			= Texture->GetNumMips() - Texture->ResidentMips;
 		int32				CurSizeX			= Texture->GetSizeX() >> DroppedMips;
 		int32				CurSizeY			= Texture->GetSizeY() >> DroppedMips;
-		bool			bIsStreamingTexture = GStreamingManager->IsManagedStreamingResource( Texture );
+		bool			bIsStreamingTexture		= IStreamingManager::Get().IsTextureStreamingEnabled() ? IStreamingManager::Get().GetTextureStreamingManager().IsManagedStreamingTexture( Texture ) : false;
 		int32				MaxSize				= Texture->CalcTextureMemorySizeEnum( TMC_AllMips );
 		int32				CurrentSize			= Texture->CalcTextureMemorySizeEnum( TMC_ResidentMips );
 		int32				UsageCount			= TextureToUsageMap.FindRef( Texture );
@@ -5588,7 +5588,7 @@ void UEngine::EnableScreenSaver( bool bEnable )
  */
 void UEngine::AddTextureStreamingSlaveLoc(FVector InLoc, float BoostFactor, bool bOverrideLocation, float OverrideDuration)
 {
-	GStreamingManager->AddViewSlaveLocation(InLoc, BoostFactor, bOverrideLocation, OverrideDuration);
+	IStreamingManager::Get().AddViewSlaveLocation(InLoc, BoostFactor, bOverrideLocation, OverrideDuration);
 }
 
 /** Looks up the GUID of a package on disk. The package must NOT be in the autodownload cache.
@@ -8110,9 +8110,9 @@ bool UEngine::LoadMap( FWorldContext& WorldContext, FURL URL, class UPendingNetG
 	CollectGarbage( GARBAGE_COLLECTION_KEEPFLAGS, true );
 
 	// Cancels the Forced StreamType for textures using a timer.
-	if ( GStreamingManager )
+	if (!IStreamingManager::HasShutdown())
 	{
-		GStreamingManager->CancelForcedResources();
+		IStreamingManager::Get().CancelForcedResources();
 	}
 
 	if (FPlatformProperties::RequiresCookedData())
@@ -8392,7 +8392,7 @@ bool UEngine::LoadMap( FWorldContext& WorldContext, FURL URL, class UPendingNetG
 	}
 
 	// Prime texture streaming.
-	GStreamingManager->NotifyLevelChange();
+	IStreamingManager::Get().NotifyLevelChange();
 
 	WorldContext.World()->BeginPlay();
 	}
@@ -8407,7 +8407,7 @@ bool UEngine::LoadMap( FWorldContext& WorldContext, FURL URL, class UPendingNetG
 	RedrawViewports(false);
 
 	// RedrawViewports() may have added a dummy playerstart location. Remove all views to start from fresh the next Tick().
-	GStreamingManager->RemoveStreamingViews( RemoveStreamingViews_All );
+	IStreamingManager::Get().RemoveStreamingViews( RemoveStreamingViews_All );
 	
 	MALLOC_PROFILER( FMallocProfiler::SnapshotMemoryLoadMapEnd( URL.Map ); )
 
@@ -9233,7 +9233,7 @@ bool UEngine::CommitMapChange( FWorldContext &Context )
 		Context.PendingMapChangeFailureDescription = TEXT("");
 
 		// Prime texture streaming.
-		GStreamingManager->NotifyLevelChange();
+		IStreamingManager::Get().NotifyLevelChange();
 
 		// tell the game we are done switching levels
 		AGameMode* const GameMode = Context.World()->GetAuthGameMode();
diff --git a/Engine/Source/Runtime/Engine/Private/World.cpp b/Engine/Source/Runtime/Engine/Private/World.cpp
--- a/Engine/Source/Runtime/Engine/Private/World.cpp
+++ b/Engine/Source/Runtime/Engine/Private/World.cpp
@@ -1558,7 +1558,7 @@ void UWorld::AddToWorld( ULevel* Level, const FTransform& LevelTransform )
 		Level->InitializeRenderingResources();
 
 		// Notify the texture streaming system now that everything is set up.
-		GStreamingManager->AddLevel( Level );
+		IStreamingManager::Get().AddLevel( Level );
 
 		Level->bIsVisible = true;
 	
@@ -1616,7 +1616,7 @@ void UWorld::RemoveFromWorld( ULevel* Level )
 		Level->ReleaseRenderingResources();
 
 		// Remove from the world's level array and destroy actor components.
-		GStreamingManager->RemoveLevel( Level );
+		IStreamingManager::Get().RemoveLevel( Level );
 		
 		Level->ClearLevelComponents();
 
@@ -2629,7 +2629,7 @@ void UWorld::InitializeActorsForPlay(const FURL& InURL, bool bResetTime)
 	for( int32 LevelIndex=0; LevelIndex<Levels.Num(); LevelIndex++ )
 	{
 		ULevel*	Level = Levels[LevelIndex];
-		GStreamingManager->AddLevel( Level );
+		IStreamingManager::Get().AddLevel( Level );
 	}
 
 	if(WorldType == EWorldType::Preview)
diff --git a/Engine/Source/Runtime/Engine/Public/ComponentReregisterContext.h b/Engine/Source/Runtime/Engine/Public/ComponentReregisterContext.h
--- a/Engine/Source/Runtime/Engine/Public/ComponentReregisterContext.h
+++ b/Engine/Source/Runtime/Engine/Public/ComponentReregisterContext.h
@@ -62,7 +62,7 @@ protected:
 			const UPrimitiveComponent* Primitive = Cast<const UPrimitiveComponent>(InComponent);
 			if ( Primitive )
 			{
-				GStreamingManager->NotifyPrimitiveUpdated( Primitive );
+				IStreamingManager::Get().NotifyPrimitiveUpdated( Primitive );
 			}
 		}
 	}
diff --git a/Engine/Source/Runtime/Engine/Public/ContentStreaming.h b/Engine/Source/Runtime/Engine/Public/ContentStreaming.h
--- a/Engine/Source/Runtime/Engine/Public/ContentStreaming.h
+++ b/Engine/Source/Runtime/Engine/Public/ContentStreaming.h
@@ -156,18 +156,25 @@ struct FStreamingViewInfo
 /**
  * Pure virtual base class of a streaming manager.
  */
-struct FStreamingManagerBase
+struct IStreamingManager
 {
-	FStreamingManagerBase()
+	IStreamingManager()
 	:	NumWantingResources(0)
 	,	NumWantingResourcesCounter(0)
 	{
 	}
 
 	/** Virtual destructor */
-	virtual ~FStreamingManagerBase()
+	virtual ~IStreamingManager()
 	{}
 
+	ENGINE_API static struct FStreamingManagerCollection& Get();
+
+	ENGINE_API static void Shutdown();
+
+	/** Checks if the streaming manager has already been shut down. **/
+	ENGINE_API static bool HasShutdown();
+
 	/**
 	 * Calls UpdateResourceStreaming(), and does per-frame cleaning. Call once per frame.
 	 *
@@ -195,15 +202,6 @@ struct FStreamingManagerBase
 		return 0;
 	}
 
-	/**
-	 * Updates streaming for an individual texture, taking into account all view infos.
-	 *
-	 * @param Texture		Texture to update
-	 */
-	virtual void UpdateIndividualResource( UTexture2D* Texture )
-	{
-	}
-
 	/**
 	 * Blocks till all pending requests are fulfilled.
 	 *
@@ -256,22 +254,6 @@ struct FStreamingManagerBase
 	/** Don't stream world resources for the next NumFrames. */
 	virtual void SetDisregardWorldResourcesForFrames( int32 NumFrames ) = 0;
 
-	/**
-	 * Temporarily boosts the streaming distance factor by the specified number.
-	 * This factor is automatically reset to 1.0 after it's been used for mip-calculations.
-	 */
-	virtual void BoostTextures( AActor* Actor, float BoostFactor ) = 0;
-
-	/**
-	 *	Try to stream out texture mip-levels to free up more memory.
-	 *	@param RequiredMemorySize	- Required minimum available texture memory
-	 *	@return						- Whether it succeeded or not
-	 **/
-	virtual bool StreamOutTextureData( int32 RequiredMemorySize )
-	{
-		return false;
-	}
-
 	/**
 	 * Allows the streaming manager to process exec commands.
 	 *
@@ -291,16 +273,6 @@ struct FStreamingManagerBase
 	/** Removes a ULevel from the streaming manager. */
 	virtual void RemoveLevel( class ULevel* Level ) = 0;
 
-	/** Adds a new texture to the streaming manager. */
-	virtual void AddStreamingTexture( UTexture2D* Texture )
-	{
-	}
-
-	/** Removes a texture from the streaming manager. */
-	virtual void RemoveStreamingTexture( UTexture2D* Texture )
-	{
-	}
-
 	/** Called when an actor is spawned. */
 	virtual void NotifyActorSpawned( AActor* Actor )
 	{
@@ -386,9 +358,6 @@ struct FStreamingManagerBase
 		return NumWantingResourcesCounter;
 	}
 
-	/** Returns true if this is a streaming resource that is managed by the streaming manager. */
-	virtual bool IsManagedStreamingResource( const UTexture2D* Texture2D ) = 0;
-
 protected:
 
 	/**
@@ -467,11 +436,46 @@ protected:
 	int32		NumWantingResourcesCounter;
 };
 
+/**
+ * Interface to add functions specifically related to texture streaming
+ */
+struct ITextureStreamingManager : public IStreamingManager
+{
+	/**
+	* Updates streaming for an individual texture, taking into account all view infos.
+	*
+	* @param Texture		Texture to update
+	*/
+	virtual void UpdateIndividualTexture(UTexture2D* Texture) = 0;
+
+	/**
+	* Temporarily boosts the streaming distance factor by the specified number.
+	* This factor is automatically reset to 1.0 after it's been used for mip-calculations.
+	*/
+	virtual void BoostTextures(AActor* Actor, float BoostFactor) = 0;
+
+	/**
+	*	Try to stream out texture mip-levels to free up more memory.
+	*	@param RequiredMemorySize	- Required minimum available texture memory
+	*	@return						- Whether it succeeded or not
+	**/
+	virtual bool StreamOutTextureData(int32 RequiredMemorySize) = 0;
+
+	/** Adds a new texture to the streaming manager. */
+	virtual void AddStreamingTexture(UTexture2D* Texture) = 0;
+
+	/** Removes a texture from the streaming manager. */
+	virtual void RemoveStreamingTexture(UTexture2D* Texture) = 0;
+
+	/** Returns true if this is a streaming texture that is managed by the streaming manager. */
+	virtual bool IsManagedStreamingTexture(const UTexture2D* Texture2D) = 0;
+};
+
 /**
  * Streaming manager collection, routing function calls to streaming managers that have been added
  * via AddStreamingManager.
  */
-struct FStreamingManagerCollection : public FStreamingManagerBase
+struct FStreamingManagerCollection : public IStreamingManager
 {
 	/** Default constructor, initializing all member variables. */
 	ENGINE_API FStreamingManagerCollection();
@@ -500,13 +504,6 @@ struct FStreamingManagerCollection : public FStreamingManagerBase
 	 */
 	virtual int32 StreamAllResources( float TimeLimit=0.0f );
 
-	/**
-	 * Updates streaming for an individual texture, taking into account all view infos.
-	 *
-	 * @param Texture	Texture to update
-	 */
-	virtual void UpdateIndividualResource( UTexture2D* Texture );
-
 	/**
 	 * Blocks till all pending requests are fulfilled.
 	 *
@@ -537,21 +534,34 @@ struct FStreamingManagerCollection : public FStreamingManagerBase
 	 */
 	virtual void NotifyLevelChange();
 
+	/**
+	 * Checks whether any kind of streaming is active
+	 */
+	bool IsStreamingEnabled() const;
+
+	/**
+	 * Checks whether texture streaming is active
+	 */
 	virtual bool IsTextureStreamingEnabled() const;
 
+	/**
+	 * Gets a reference to the Texture Streaming Manager interface
+	 */
+	ITextureStreamingManager& GetTextureStreamingManager() const;
+
 	/**
 	 * Adds a streaming manager to the array of managers to route function calls to.
 	 *
 	 * @param StreamingManager	Streaming manager to add
 	 */
-	void AddStreamingManager( FStreamingManagerBase* StreamingManager );
+	void AddStreamingManager( IStreamingManager* StreamingManager );
 
 	/**
 	 * Removes a streaming manager from the array of managers to route function calls to.
 	 *
 	 * @param StreamingManager	Streaming manager to remove
 	 */
-	void RemoveStreamingManager( FStreamingManagerBase* StreamingManager );
+	void RemoveStreamingManager( IStreamingManager* StreamingManager );
 
 	/**
 	 * Sets the number of iterations to use for the next time UpdateResourceStreaming is being called. This 
@@ -564,12 +574,6 @@ struct FStreamingManagerCollection : public FStreamingManagerBase
 	/** Don't stream world resources for the next NumFrames. */
 	virtual void SetDisregardWorldResourcesForFrames( int32 NumFrames );
 
-	/**
-	 * Temporarily boosts the streaming distance factor by the specified number.
-	 * This factor is automatically reset to 1.0 after it's been used for mip-calculations.
-	 */
-	virtual void BoostTextures( AActor* Actor, float BoostFactor );
-
 	/**
 	 * Disables resource streaming. Enable with EnableResourceStreaming. Disable/enable can be called multiple times nested
 	 */
@@ -581,13 +585,6 @@ struct FStreamingManagerCollection : public FStreamingManagerBase
 	 */
 	void EnableResourceStreaming();
 
-	/**
-	 *	Try to stream out texture mip-levels to free up more memory.
-	 *	@param RequiredMemorySize	- Required minimum available texture memory
-	 *	@return						- Whether it succeeded or not
-	 **/
-	virtual bool StreamOutTextureData( int32 RequiredMemorySize );
-
 	/**
 	 * Allows the streaming manager to process exec commands.
 	 *
@@ -598,12 +595,6 @@ struct FStreamingManagerCollection : public FStreamingManagerBase
 	 */
 	virtual bool Exec( UWorld* InWorld, const TCHAR* Cmd, FOutputDevice& Ar ) OVERRIDE;
 
-	/** Adds a new texture to the streaming manager. */
-	virtual void AddStreamingTexture( UTexture2D* Texture );
-
-	/** Removes a texture from the streaming manager. */
-	virtual void RemoveStreamingTexture( UTexture2D* Texture );
-
 	/** Adds a ULevel to the streaming manager. */
 	virtual void AddLevel( class ULevel* Level );
 
@@ -637,15 +628,12 @@ struct FStreamingManagerCollection : public FStreamingManagerBase
 	 */
 	virtual void NotifyPrimitiveUpdated( const UPrimitiveComponent* Primitive );
 
-	/** Returns true if this is a streaming resource that is managed by the streaming manager. */
-	virtual bool IsManagedStreamingResource( const UTexture2D* Texture2D );
-
 protected:
 
 	virtual void AddOrRemoveTextureStreamingManagerIfNeeded(bool bIsInit=false);
 
 	/** Array of streaming managers to route function calls to */
-	TArray<FStreamingManagerBase*> StreamingManagers;
+	TArray<IStreamingManager*> StreamingManagers;
 	/** Number of iterations to perform. Gets reset to 1 each frame. */
 	int32 NumIterations;
 
@@ -734,9 +722,6 @@ struct FStreamMemoryTracker
 	Globals.
 -----------------------------------------------------------------------------*/
 
-/** Global streaming manager */
-ENGINE_API extern FStreamingManagerCollection*	GStreamingManager;
-
 /** Texture streaming memory tracker. */
 extern FStreamMemoryTracker GStreamMemoryTracker;
 
@@ -789,7 +774,7 @@ struct FStreamableTextureInstance4
 /**
  * Streaming manager dealing with textures.
  */
-struct FStreamingManagerTexture : public FStreamingManagerBase
+struct FStreamingManagerTexture : public ITextureStreamingManager
 {
 	/** Constructor, initializing all members */
 	FStreamingManagerTexture();
@@ -809,7 +794,7 @@ struct FStreamingManagerTexture : public FStreamingManagerBase
 	 *
 	 * @param Texture	Texture to update
 	 */
-	virtual void UpdateIndividualResource( UTexture2D* Texture );
+	virtual void UpdateIndividualTexture( UTexture2D* Texture );
 
 	/**
 	 * Blocks till all pending requests are fulfilled.
@@ -954,13 +939,7 @@ struct FStreamingManagerTexture : public FStreamingManagerBase
 	FStreamingTexture& GetStreamingTexture( const UTexture2D* Texture2D );
 
 	/** Returns true if this is a streaming texture that is managed by the streaming manager. */
-	bool IsManagedStreamingTexture( const UTexture2D* Texture2D );
-
-	/** Returns true if this is a streaming resource that is managed by the streaming manager. */
-	virtual bool IsManagedStreamingResource( const UTexture2D* Texture2D )
-	{
-		return IsManagedStreamingTexture( Texture2D );
-	}
+	virtual bool IsManagedStreamingTexture( const UTexture2D* Texture2D );
 
 	/** Updates the I/O state of a texture (allowing it to progress to the next stage) and some stats. */
 	void UpdateTextureStatus( FStreamingTexture& StreamingTexture, FStreamingContext& Context );
diff --git a/Engine/Source/Runtime/Launch/Private/LaunchEngineLoop.cpp b/Engine/Source/Runtime/Launch/Private/LaunchEngineLoop.cpp
--- a/Engine/Source/Runtime/Launch/Private/LaunchEngineLoop.cpp
+++ b/Engine/Source/Runtime/Launch/Private/LaunchEngineLoop.cpp
@@ -1108,7 +1108,7 @@ int32 FEngineLoop::PreInit( const TCHAR* CmdLine )
 
 	EndInitTextLocalization();
 
-	GStreamingManager = new FStreamingManagerCollection();
+	IStreamingManager::Get();
 
 	if (FPlatformProcess::SupportsMultithreading() && !IsRunningDedicatedServer() && (bIsRegularClient || bHasEditorToken))
 	{
@@ -1839,10 +1839,10 @@ void FEngineLoop::Exit()
 	FlushAsyncLoading();
 
 	// Block till all outstanding resource streaming requests are fulfilled.
-	if (GStreamingManager != NULL)
+	if (!IStreamingManager::HasShutdown())
 	{
 		UTexture2D::CancelPendingTextureStreaming();
-		GStreamingManager->BlockTillAllRequestsFinished();
+		IStreamingManager::Get().BlockTillAllRequestsFinished();
 	}
 
 #if WITH_ENGINE
@@ -1899,8 +1899,7 @@ void FEngineLoop::Exit()
 
 	FTaskGraphInterface::Shutdown();
 
-	delete GStreamingManager;
-	GStreamingManager	= NULL;
+	IStreamingManager::Shutdown();
 
 	FIOSystem::Shutdown();
 }
