Some slight cleanup on ULevel::ClearLevelComponents()

diff --git a/Engine/Source/Runtime/Engine/Private/Level.cpp b/Engine/Source/Runtime/Engine/Private/Level.cpp
--- a/Engine/Source/Runtime/Engine/Private/Level.cpp
+++ b/Engine/Source/Runtime/Engine/Private/Level.cpp
@@ -614,35 +614,26 @@ void ULevel::ClearLevelComponents()
 	bAreComponentsCurrentlyRegistered = false;
 
 	// Remove the model components from the scene.
-	for(int32 ComponentIndex = 0;ComponentIndex < ModelComponents.Num();ComponentIndex++)
+	for (UModelComponent* ModelComponent : ModelComponents)
 	{
-		if(ModelComponents[ComponentIndex] && ModelComponents[ComponentIndex]->IsRegistered())
+		if (ModelComponent && ModelComponent->IsRegistered())
 		{
-			ModelComponents[ComponentIndex]->UnregisterComponent();
+			ModelComponent->UnregisterComponent();
 		}
 	}
 
-	TArray<UWorld*> Worlds;
 	// Remove the actors' components from the scene and build a list of relevant worlds
-	for( int32 ActorIndex=0; ActorIndex<Actors.Num(); ActorIndex++ )
+	for( AActor* Actor : Actors )
 	{
-		AActor* Actor = Actors[ActorIndex];
-		if( Actor )
-		{
-			Actor->UnregisterAllComponents();
-			if( Actor->GetWorld() )
-			{
-				Worlds.AddUnique( Actor->GetWorld() );
-			}
-		}
+		Actor->UnregisterAllComponents();
 	}
 
-	// clear global motion blur state info in the relevant worlds
-	for( int32 WorldIndex=0; WorldIndex<Worlds.Num(); WorldIndex++ )
+	if (IsPersistentLevel())
 	{
-		if(Worlds[WorldIndex]->PersistentLevel == this && Worlds[WorldIndex]->Scene)
+		FSceneInterface* WorldScene = GetWorld()->Scene;
+		if (WorldScene)
 		{
-			Worlds[WorldIndex]->Scene->SetClearMotionBlurInfoGameThread();
+			WorldScene->SetClearMotionBlurInfoGameThread();
 		}
 	}
 }
