Fix crash when a NULL entry exists in the Actors array

diff --git a/Engine/Source/Runtime/Engine/Private/Level.cpp b/Engine/Source/Runtime/Engine/Private/Level.cpp
--- a/Engine/Source/Runtime/Engine/Private/Level.cpp
+++ b/Engine/Source/Runtime/Engine/Private/Level.cpp
@@ -625,7 +625,10 @@ void ULevel::ClearLevelComponents()
 	// Remove the actors' components from the scene and build a list of relevant worlds
 	for( AActor* Actor : Actors )
 	{
-		Actor->UnregisterAllComponents();
+		if (Actor)
+		{
+			Actor->UnregisterAllComponents();
+		}
 	}
 
 	if (IsPersistentLevel())
