CheckActorComponents is now going to be used by CheckDefaultSubobjects for Actors

diff --git a/Engine/Source/Runtime/CoreUObject/Private/UObject/Obj.cpp b/Engine/Source/Runtime/CoreUObject/Private/UObject/Obj.cpp
--- a/Engine/Source/Runtime/CoreUObject/Private/UObject/Obj.cpp
+++ b/Engine/Source/Runtime/CoreUObject/Private/UObject/Obj.cpp
@@ -964,19 +964,29 @@ DEFINE_LOG_CATEGORY_STATIC(LogCheckSubobjects, Fatal, All);
 		UE_LOG(LogCheckSubobjects, Log, TEXT("CompCheck %s failed."), TEXT(#Pred)); \
 	} 
 
-bool UObject::CheckDefaultSubobjects(bool bForceCheck /*= false*/)
+bool UObject::CanCheckDefaultSubObjects(bool bForceCheck, bool& bResult)
 {
+	bool bCanCheck = true;
+	bResult = true;
 	if (!this)
 	{
-		return false; // these aren't in a suitable spot in their lifetime for testing
+		bResult = false; // these aren't in a suitable spot in their lifetime for testing
+		bCanCheck = false;
 	}
 	if (HasAnyFlags(RF_NeedLoad | RF_NeedPostLoad | RF_NeedPostLoadSubobjects | RF_Unreachable | RF_PendingKill) || GIsDuplicatingClassForReinstancing)
 	{
-		return true; // these aren't in a suitable spot in their lifetime for testing
+		bResult = true; // these aren't in a suitable spot in their lifetime for testing
+		bCanCheck = false;
 	}
-	bool Result = true;
 	// If errors are suppressed, we will not take the time to run this test unless forced to.
-	if (bForceCheck || UE_LOG_ACTIVE(LogCheckSubobjects, Error))
+	bCanCheck = bCanCheck && (bForceCheck || UE_LOG_ACTIVE(LogCheckSubobjects, Error));
+	return bCanCheck;
+}
+
+bool UObject::CheckDefaultSubobjects(bool bForceCheck /*= false*/)
+{
+	bool Result = true;	
+	if (CanCheckDefaultSubObjects(bForceCheck, Result))
 	{
 		CompCheck(this);
 		UClass* Class = GetClass();
diff --git a/Engine/Source/Runtime/CoreUObject/Public/UObject/UObject.h b/Engine/Source/Runtime/CoreUObject/Public/UObject/UObject.h
--- a/Engine/Source/Runtime/CoreUObject/Public/UObject/UObject.h
+++ b/Engine/Source/Runtime/CoreUObject/Public/UObject/UObject.h
@@ -582,14 +582,7 @@ public:
 	 * @param bForceCheck Force checks even if not enabled globally.
 	 * @return true if the assumptions are met, false otherwise.
 	 */
-	bool CheckDefaultSubobjects(bool bForceCheck = false);
-
-	/**
-	 * Checks components for validity, implemented in AActor
-	 */
-	virtual void CheckActorComponents()
-	{
-	}
+	virtual bool CheckDefaultSubobjects(bool bForceCheck = false);
 
 	/**
 	 * Save configuration.
@@ -953,6 +946,11 @@ public:
 		ProcessEvent(FindFunctionChecked(NAME_ExecuteUbergraph),&Parms);
 	}
 
+protected: 
+
+	/** Checks it's ok to perform subobjects check at this time. */
+	bool CanCheckDefaultSubObjects(bool bForceCheck, bool& bResult);
+
 private:
 	void ProcessContextOpcode(FFrame& Stack, RESULT_DECL, bool bCanFailSilent);
 };
diff --git a/Engine/Source/Runtime/Engine/Classes/GameFramework/Actor.h b/Engine/Source/Runtime/Engine/Classes/GameFramework/Actor.h
--- a/Engine/Source/Runtime/Engine/Classes/GameFramework/Actor.h
+++ b/Engine/Source/Runtime/Engine/Classes/GameFramework/Actor.h
@@ -1004,7 +1004,7 @@ public:
 	
 
 	// Begin UObject Interface
-	virtual void CheckActorComponents() OVERRIDE;
+	virtual bool CheckDefaultSubobjects(bool bForceCheck = false) OVERRIDE;
 	virtual void PostInitProperties() OVERRIDE;
 	virtual bool Modify( bool bAlwaysMarkDirty=true ) OVERRIDE;
 	virtual void ProcessEvent( UFunction* Function, void* Parameters ) OVERRIDE;
@@ -1707,6 +1707,11 @@ protected:
 	/** Runs UserConstructionScript, delays component registration until it's complete. */
 	void ProcessUserConstructionScript();
 
+	/**
+	* Checks components for validity, implemented in AActor
+	*/
+	bool CheckActorComponents();
+
 public:
 
 	/** Walk up the attachment chain from RootComponent until we encounter a different actor, and return it. If we are not attached to a component in a different actor, returns NULL */
diff --git a/Engine/Source/Runtime/Engine/Private/Actor.cpp b/Engine/Source/Runtime/Engine/Private/Actor.cpp
--- a/Engine/Source/Runtime/Engine/Private/Actor.cpp
+++ b/Engine/Source/Runtime/Engine/Private/Actor.cpp
@@ -190,11 +190,22 @@ FString FActorTickFunction::DiagnosticMessage()
 	return Target->GetFullName() + TEXT("[TickActor]");
 }
 
+bool AActor::CheckDefaultSubobjects(bool bForceCheck /*= false*/)
+{
+	bool Result = true;
+	if (CanCheckDefaultSubObjects(bForceCheck, Result))
+	{
+		Result = Super::CheckDefaultSubobjects(bForceCheck);
+		Result = Result && CheckActorComponents();
+	}
+	return Result;
+}
 
-void AActor::CheckActorComponents()
+bool AActor::CheckActorComponents()
 {
 	DEFINE_LOG_CATEGORY_STATIC(LogCheckComponents, Warning, All);
 
+	bool bResult = true;
 	TArray<UActorComponent*> Components;
 	GetComponents(Components);
 
@@ -212,6 +223,7 @@ void AActor::CheckActorComponents()
 		if (Inner->IsTemplate() && !IsTemplate())
 		{
 			UE_LOG(LogCheckComponents, Error, TEXT("Component is a template but I am not. Me = %s, Component = %s"), *this->GetFullName(), *Inner->GetFullName());
+			bResult = false;
 		}
 		UObject* Archetype = Inner->GetArchetype();
 		if (Archetype != Inner->GetClass()->GetDefaultObject())
@@ -219,6 +231,7 @@ void AActor::CheckActorComponents()
 			if (Archetype != GetClass()->GetDefaultSubobjectByName(Inner->GetFName()))
 			{
 				UE_LOG(LogCheckComponents, Error, TEXT("Component archetype is not the CDO nor a default subobject of my class. Me = %s, Component = %s, Archetype = %s"), *this->GetFullName(), *Inner->GetFullName(), *Archetype->GetFullName());
+				bResult = false;
 			}
 		}
 	}
@@ -232,6 +245,7 @@ void AActor::CheckActorComponents()
 		if (Inner->GetOuter() != this)
 		{
 			UE_LOG(LogCheckComponents, Error, TEXT("SerializedComponent does not have me as an outer. Me = %s, Component = %s"), *this->GetFullName(), *Inner->GetFullName());
+			bResult = false;
 		}
 		if (Inner->IsPendingKill())
 		{
@@ -240,6 +254,7 @@ void AActor::CheckActorComponents()
 		if (Inner->IsTemplate() && !IsTemplate())
 		{
 			UE_LOG(LogCheckComponents, Error, TEXT("SerializedComponent is a template but I am not. Me = %s, Component = %s"), *this->GetFullName(), *Inner->GetFullName());
+			bResult = false;
 		}
 		UObject* Archetype = Inner->GetArchetype();
 		if (Archetype != Inner->GetClass()->GetDefaultObject())
@@ -247,9 +262,11 @@ void AActor::CheckActorComponents()
 			if (Archetype != GetClass()->GetDefaultSubobjectByName(Inner->GetFName()))
 			{
 				UE_LOG(LogCheckComponents, Error, TEXT("SerializedComponent archetype is not the CDO nor a default subobject of my class. Me = %s, Component = %s, Archetype = %s"), *this->GetFullName(), *Inner->GetFullName(), *Archetype->GetFullName());
+				bResult = false;
 			}
 		}
 	}
+	return bResult;
 }
 
 void AActor::ResetOwnedComponents()
