332362	PUNTABLE: Typo in FAnimationUtils::TallyErrorsFromPerturbation

diff --git a/Engine/Source/Runtime/Engine/Private/Animation/AnimationUtils.cpp b/Engine/Source/Runtime/Engine/Private/Animation/AnimationUtils.cpp
--- a/Engine/Source/Runtime/Engine/Private/Animation/AnimationUtils.cpp
+++ b/Engine/Source/Runtime/Engine/Private/Animation/AnimationUtils.cpp
@@ -1712,7 +1712,7 @@ void FAnimationUtils::TallyErrorsFromPerturbation(
 		TrackError.MaxErrorInRotDueToRot = MaxErrorR_DueToR;
 		TrackError.MaxErrorInScaleDueToRot = MaxErrorS_DueToR;
 		TrackError.MaxErrorInTransDueToScale = MaxErrorT_DueToR;
-		TrackError.MaxErrorInScaleDueToScale = MaxErrorR_DueToR;
+		TrackError.MaxErrorInRotDueToScale = MaxErrorR_DueToR;
 		TrackError.MaxErrorInScaleDueToScale = MaxErrorS_DueToR;
 	}
 }
