Implement GetWorld for ULevel

diff --git a/Engine/Source/Runtime/Engine/Classes/Engine/Level.h b/Engine/Source/Runtime/Engine/Classes/Engine/Level.h
--- a/Engine/Source/Runtime/Engine/Classes/Engine/Level.h
+++ b/Engine/Source/Runtime/Engine/Classes/Engine/Level.h
@@ -465,6 +465,7 @@ public:
 	virtual void BeginDestroy() OVERRIDE;
 	virtual bool IsReadyForFinishDestroy() OVERRIDE;
 	virtual void FinishDestroy() OVERRIDE;
+	virtual UWorld* GetWorld() const OVERRIDE;
 
 #if	WITH_EDITOR
 	virtual void PreEditUndo() OVERRIDE;
diff --git a/Engine/Source/Runtime/Engine/Private/Level.cpp b/Engine/Source/Runtime/Engine/Private/Level.cpp
--- a/Engine/Source/Runtime/Engine/Private/Level.cpp
+++ b/Engine/Source/Runtime/Engine/Private/Level.cpp
@@ -608,6 +608,10 @@ void ULevel::PostLoad()
 	} 
 }
 
+UWorld* ULevel::GetWorld() const
+{
+	return OwningWorld;
+}
 
 void ULevel::ClearLevelComponents()
 {
