New EndPlay function/event
* Replaces Destroyed in Blueprints
* Provides an opportunity to respond to any time an Actor is being removed from a world, be it due to destruction, level transition, sublevel streamed out, or ending a PIE session.

diff --git a/Engine/Config/BaseEngine.ini b/Engine/Config/BaseEngine.ini
--- a/Engine/Config/BaseEngine.ini
+++ b/Engine/Config/BaseEngine.ini
@@ -726,10 +726,14 @@ bUseStreamingPause=false
 +K2FieldRedirects=(OldFieldName="Actor.SetActorRelativeTranslation",NewFieldName="Actor.SetActorRelativeLocation")
 +K2ParamRedirects=(NodeName="/Script/BlueprintGraph.K2Node_CallFunction", OldParamName="Actor.SetActorRelativeLocation.NewRelativeTranslation", NewParamName="NewRelativeLocation")
 
-
 +K2FieldRedirects=(OldFieldName="LevelScriptActor.BeginGame",NewFieldName="Actor.ReceiveBeginPlay")
 +K2FieldRedirects=(OldFieldName="LevelScriptActor.BeginGameSignature__DelegateSignature", NewFieldName="Actor.ReceiveBeginPlaySignature__DelegateSignature")
 
++K2FieldRedirects=(OldFieldName="Actor.ReceiveDestroyed",NewFieldName="Actor.ReceiveEndPlay")
++K2FieldRedirects=(OldFieldName="Actor.ReceiveDestroyed__DelegateSignature",NewFieldName="Actor.ReceiveEndPlay__DelegateSignature")
++K2FieldRedirects=(OldFieldName="Actor.OnDestroyed",NewFieldName="Actor.OnEndPlay")
++K2FieldRedirects=(OldFieldName="Actor.DestroyedSignature__DelegateSignature",NewFieldName="Actor.EndPlaySignature__DelegateSignature")
+
 ; AttachTo change - changed bMaintainWorldPosition to enum
 +K2ParamRedirects=(NodeName="/Script/BlueprintGraph.K2Node_CallFunction", OldParamName="AttachTo.bMaintainWorldPosition", NewParamName="AttachType", OldParamValues="false;true", NewParamValues="KeepRelativeOffset;KeepWorldPosition")
 +K2ParamRedirects=(NodeName="/Script/BlueprintGraph.K2Node_CallFunction", OldParamName="AttachRootComponentTo.bMaintainWorldPosition", NewParamName="AttachLocationType", OldParamValues="false;true", NewParamValues="KeepRelativeOffset;KeepWorldPosition")
diff --git a/Engine/Source/Editor/UnrealEd/Private/PlayLevel.cpp b/Engine/Source/Editor/UnrealEd/Private/PlayLevel.cpp
--- a/Engine/Source/Editor/UnrealEd/Private/PlayLevel.cpp
+++ b/Engine/Source/Editor/UnrealEd/Private/PlayLevel.cpp
@@ -403,8 +403,6 @@ void UEditorEngine::TeardownPlaySession(FWorldContext &PieWorldContext)
 							}	
 						}
 					}
-
-					World->CleanupWorld();
 				}
 			}
 		}
@@ -424,7 +422,13 @@ void UEditorEngine::TeardownPlaySession(FWorldContext &PieWorldContext)
 			}
 		}
 	}
-	
+
+	// Go through and let all the PlayWorld Actor's know they are being destroyed
+	for (FActorIterator ActorIt(PlayWorld); ActorIt; ++ActorIt)
+	{
+		ActorIt->EndPlay(AActor::EEndPlayReason::EndPlayInEditor);
+	}
+
 	// Move blueprint debugging pointers back to the objects in the editor world
 	PlayWorld->TransferBlueprintDebugReferences(EditorWorld);
 
diff --git a/Engine/Source/Runtime/Engine/Classes/GameFramework/Actor.h b/Engine/Source/Runtime/Engine/Classes/GameFramework/Actor.h
--- a/Engine/Source/Runtime/Engine/Classes/GameFramework/Actor.h
+++ b/Engine/Source/Runtime/Engine/Classes/GameFramework/Actor.h
@@ -24,7 +24,7 @@ DECLARE_DYNAMIC_MULTICAST_DELEGATE_OneParam( FActorOnInputTouchEndSignature, ETo
 DECLARE_DYNAMIC_MULTICAST_DELEGATE_OneParam( FActorBeginTouchOverSignature, ETouchIndex::Type, FingerIndex );
 DECLARE_DYNAMIC_MULTICAST_DELEGATE_OneParam( FActorEndTouchOverSignature, ETouchIndex::Type, FingerIndex );
 
-DECLARE_DYNAMIC_MULTICAST_DELEGATE( FActorDestroyedSignature );
+DECLARE_DYNAMIC_MULTICAST_DELEGATE( FActorEndPlaySignature );
 
 DECLARE_DELEGATE_FourParams(FMakeNoiseDelegate, class AActor*, float, class APawn*, const FVector&);
 
@@ -995,14 +995,13 @@ public:
 	bool Destroy(bool bNetForce = false, bool bShouldModifyLevel = true );
 
 	/** Event to notify blueprints this actor is about to be deleted. */
-	UFUNCTION(BlueprintImplementableEvent, meta=(Keywords = "delete", FriendlyName = "Destroyed"))
-	virtual void ReceiveDestroyed();
+	UFUNCTION(BlueprintImplementableEvent, meta=(Keywords = "delete", FriendlyName = "End Play"))
+	virtual void ReceiveEndPlay();
 
 	/** Called when the actor is destroyed. */
 	UPROPERTY(BlueprintAssignable)
-	FActorDestroyedSignature OnDestroyed;
+	FActorEndPlaySignature OnEndPlay;
 	
-
 	// Begin UObject Interface
 	virtual bool CheckDefaultSubobjects(bool bForceCheck = false) OVERRIDE;
 	virtual void PostInitProperties() OVERRIDE;
@@ -1626,8 +1625,16 @@ public:
 	/**	Do anything needed to clear out cross level references; Called from ULevel::PreSave	 */
 	virtual void ClearCrossLevelReferences();
 	
-	/** Called when this actor is in a level which is being removed from the world (e.g. my level is getting UWorld::RemoveFromWorld called on it) */
-	virtual void OnRemoveFromWorld();
+	enum EEndPlayReason
+	{
+		ActorDestroyed,		// When the Actor is explicitly destroyed
+		LevelTransition,	// When the world is being unloaded for a level transition
+		EndPlayInEditor,	// When the world is being unloaded because PIE is ending
+		RemovedFromWorld,	// When the level it is a member of is streamed out
+	};
+
+	/** Called whenever this actor is being removed from a level */
+	virtual void EndPlay(const EEndPlayReason EndPlayReason);
 
 	/** iterates up the Base chain to see whether or not this Actor is based on the given Actor
 	 * @param Other the Actor to test for
diff --git a/Engine/Source/Runtime/Engine/Private/Actor.cpp b/Engine/Source/Runtime/Engine/Private/Actor.cpp
--- a/Engine/Source/Runtime/Engine/Private/Actor.cpp
+++ b/Engine/Source/Runtime/Engine/Private/Actor.cpp
@@ -677,13 +677,14 @@ bool AActor::IsActorTickEnabled() const
 
 bool AActor::Rename( const TCHAR* InName, UObject* NewOuter, ERenameFlags Flags )
 {
-	if (NewOuter)
-	{
-		RegisterAllActorTickFunctions(false, true); // unregister all tick functions
-	}
+	RegisterAllActorTickFunctions(false, true); // unregister all tick functions
+	UnregisterAllComponents();
+
 	bool bSuccess = Super::Rename( InName, NewOuter, Flags );
-	if (NewOuter)
+
+	if (NewOuter && NewOuter->IsA<ULevel>())
 	{
+		RegisterAllComponents();
 		RegisterAllActorTickFunctions(true, true); // register all tick functions
 	}
 	return bSuccess;
@@ -1513,13 +1514,19 @@ void AActor::PrestreamTextures( float Seconds, bool bEnableStreaming, int32 Cine
 void AActor::OnRep_Instigator() {}
 
 
-void AActor::OnRemoveFromWorld()
+void AActor::EndPlay(const EEndPlayReason EndPlayReason)
 {
-	bActorInitialized = false;
-
-	SetNavigationRelevancy(false);
+	// Dispatch the blueprint events
+	ReceiveEndPlay();
+	OnEndPlay.Broadcast();
 
-	GetWorld()->RemoveNetworkActor( this );
+	// Behaviors specific to an actor being unloaded due to a streaming level removal
+	if (EndPlayReason == EEndPlayReason::RemovedFromWorld)
+	{
+		bActorInitialized = false;
+		SetNavigationRelevancy(false);
+		GetWorld()->RemoveNetworkActor(this);
+	}
 }
 
 FVector AActor::GetPlacementExtent() const
@@ -1580,8 +1587,10 @@ FTransform AActor::GetTransform() const
 
 void AActor::Destroyed()
 {
-	ReceiveDestroyed();
-	OnDestroyed.Broadcast();
+	if (bActorInitialized)
+	{
+		EndPlay(EEndPlayReason::ActorDestroyed);
+	}
 	GetWorld()->RemoveNetworkActor( this );
 }
 
diff --git a/Engine/Source/Runtime/Engine/Private/UnrealEngine.cpp b/Engine/Source/Runtime/Engine/Private/UnrealEngine.cpp
--- a/Engine/Source/Runtime/Engine/Private/UnrealEngine.cpp
+++ b/Engine/Source/Runtime/Engine/Private/UnrealEngine.cpp
@@ -8074,6 +8074,14 @@ bool UEngine::LoadMap( FWorldContext& WorldContext, FURL URL, class UPendingNetG
 			Player->bSentSplitJoin = false;
 		}
 
+		for (FActorIterator ActorIt(WorldContext.World()); ActorIt; ++ActorIt)
+		{
+			if (ActorIt->bActorInitialized)
+			{
+				ActorIt->EndPlay(AActor::EEndPlayReason::LevelTransition);
+			}
+		}
+
 		// Do this after destroying pawns/playercontrollers, in case that spawns new things (e.g. dropped weapons)
 		WorldContext.World()->CleanupWorld();
 
diff --git a/Engine/Source/Runtime/Engine/Private/World.cpp b/Engine/Source/Runtime/Engine/Private/World.cpp
--- a/Engine/Source/Runtime/Engine/Private/World.cpp
+++ b/Engine/Source/Runtime/Engine/Private/World.cpp
@@ -1589,7 +1589,7 @@ void UWorld::RemoveFromWorld( ULevel* Level )
 			AActor* Actor = Level->Actors[ActorIdx];
 			if (Actor != NULL)
 			{
-				Actor->OnRemoveFromWorld();
+				Actor->EndPlay(AActor::EEndPlayReason::RemovedFromWorld);
 				if (NetDriver)
 				{
 					NetDriver->NotifyActorLevelUnloaded(Actor);
@@ -4106,15 +4106,6 @@ UWorld* FSeamlessTravelHandler::Tick()
 		{
 			// Make sure there are no pending visibility requests.
 			CurrentWorld->FlushLevelStreaming( NULL, true );
-			// only consider session ended if we're making the final switch so that HUD, etc. UI elements stay around until the end
-			CurrentWorld->CleanupWorld(bSwitchedToDefaultMap); 
-			CurrentWorld->RemoveFromRoot();
-			CurrentWorld->ClearFlags(RF_Standalone);
-			// Stop all audio to remove references to old world
-			if (GEngine != NULL && GEngine->GetAudioDevice() != NULL)
-			{
-				GEngine->GetAudioDevice()->Flush( NULL );
-			}
 
 			if (CurrentWorld->GameState)
 			{
@@ -4123,12 +4114,19 @@ UWorld* FSeamlessTravelHandler::Tick()
 			
 			// mark actors we want to keep
 			FUObjectAnnotationSparseBool KeepAnnotation;
+			TArray<AActor*> KeepActors;
 
-			// keep GameMode if traveling to entry
-			if (!bSwitchedToDefaultMap && CurrentWorld->GetAuthGameMode() != NULL)
+			AGameMode* AuthGameMode = CurrentWorld->GetAuthGameMode();
+			if (AuthGameMode)
 			{
-				KeepAnnotation.Set(CurrentWorld->GetAuthGameMode());
+				// keep GameMode if traveling to transition map
+				if (!bSwitchedToDefaultMap)
+				{
+					KeepAnnotation.Set(CurrentWorld->GetAuthGameMode());
+				}
+				AuthGameMode->GetSeamlessTravelActorList(!bSwitchedToDefaultMap, KeepActors);
 			}
+
 			// always keep Controllers that belong to players
 			if (CurrentWorld->GetNetMode() == NM_Client)
 			{
@@ -4152,12 +4150,6 @@ UWorld* FSeamlessTravelHandler::Tick()
 				}
 			}
 
-			// ask the game class what else we should keep
-			TArray<AActor*> KeepActors;
-			if (CurrentWorld->GetAuthGameMode() != NULL)
-			{
-				CurrentWorld->GetAuthGameMode()->GetSeamlessTravelActorList(!bSwitchedToDefaultMap, KeepActors);
-			}
 			// ask players what else we should keep
 			for (FLocalPlayerIterator It(GEngine, CurrentWorld); It; ++It)
 			{
@@ -4197,11 +4189,11 @@ UWorld* FSeamlessTravelHandler::Tick()
 					// if it's a Controller or a Pawn, add it to the appropriate list in the new world's WorldSettings
 					if (Cast<AController>(TheActor))
 					{
-						LoadedWorld->AddController(Cast<AController>(TheActor) );
+						LoadedWorld->AddController(static_cast<AController*>(TheActor));
 					}
 					else if (Cast<APawn>(TheActor))
 					{
-						LoadedWorld->AddPawn(Cast<APawn>(TheActor));
+						LoadedWorld->AddPawn(static_cast<APawn*>(TheActor));
 					}
 					// add to new world's actor list and remove from old
 					LoadedWorld->PersistentLevel->Actors.Add(TheActor);
@@ -4210,6 +4202,16 @@ UWorld* FSeamlessTravelHandler::Tick()
 				}
 				else
 				{
+					if (bShouldKeep)
+					{
+						UE_LOG(LogWorld, Warning, TEXT("Actor '%s' was indicated to be kept but exists in level '%s', not the persistent level.  Actor will not travel."), *TheActor->GetName(), *TheActor->GetLevel()->GetOutermost()->GetName());
+					}
+
+					if (TheActor->bActorInitialized)
+					{
+						TheActor->EndPlay(AActor::EEndPlayReason::LevelTransition);
+					}
+
 					// otherwise, set to be deleted
 					KeepAnnotation.Clear(TheActor);
 					TheActor->MarkPendingKill();
@@ -4224,6 +4226,16 @@ UWorld* FSeamlessTravelHandler::Tick()
 
 			CopyWorldData(); // This copies the net driver too (LoadedWorld now has whatever NetDriver was previously held by CurrentWorld)
 
+			// only consider session ended if we're making the final switch so that HUD, etc. UI elements stay around until the end
+			CurrentWorld->CleanupWorld(bSwitchedToDefaultMap);
+			CurrentWorld->RemoveFromRoot();
+			CurrentWorld->ClearFlags(RF_Standalone);
+			// Stop all audio to remove references to old world
+			if (GEngine != NULL && GEngine->GetAudioDevice() != NULL)
+			{
+				GEngine->GetAudioDevice()->Flush(NULL);
+			}
+
 			// Copy the standby cheat status
 			bool bHasStandbyCheatTriggered = (CurrentWorld->NetworkManager) ? CurrentWorld->NetworkManager->bHasStandbyCheatTriggered : false;
 
