#ue4 - Character movement iterations during walking movement abort when there is no change in position.

- Previously at low framerates / high delta times if no movement occurred with non-zero velocity, further iterations were still attempted.

diff --git a/Engine/Source/Runtime/Engine/Private/CharacterMovementComponent.cpp b/Engine/Source/Runtime/Engine/Private/CharacterMovementComponent.cpp
--- a/Engine/Source/Runtime/Engine/Private/CharacterMovementComponent.cpp
+++ b/Engine/Source/Runtime/Engine/Private/CharacterMovementComponent.cpp
@@ -3044,6 +3044,12 @@ void UCharacterMovementComponent::PhysWalking(float deltaTime, int32 Iterations)
 				bCheckedFall = true;
 			}
 		}
+
+		// If we didn't move at all this iteration then abort (since future iterations will also be stuck).
+		if (CharacterOwner->GetActorLocation() == subLoc)
+		{
+			break;
+		}
 	}
 
 	// Allow overlap events and such to change physics state and velocity
