Add warning when vehicle wheel body is set in physics asset

diff --git a/Engine/Source/Runtime/Engine/Private/Vehicles/WheeledVehicleMovementComponent.cpp b/Engine/Source/Runtime/Engine/Private/Vehicles/WheeledVehicleMovementComponent.cpp
--- a/Engine/Source/Runtime/Engine/Private/Vehicles/WheeledVehicleMovementComponent.cpp
+++ b/Engine/Source/Runtime/Engine/Private/Vehicles/WheeledVehicleMovementComponent.cpp
@@ -2,6 +2,7 @@
 
 #include "EnginePrivate.h"
 #include "Net/UnrealNetwork.h"
+#include "MessageLog.h"
 
 #if WITH_PHYSX
 #include "../PhysicsEngine/PhysXSupport.h"
@@ -9,6 +10,8 @@
 #include "PhysXVehicleManager.h"
 #endif //WITH_PHYSX
 
+#define LOCTEXT_NAMESPACE "UWheeledVehicleMovementComponent"
+
 #if WITH_PHYSX
 
 /**
@@ -479,6 +482,14 @@ FVector UWheeledVehicleMovementComponent::GetWheelRestingPosition( const FWheelS
 	{
 		USkinnedMeshComponent* Mesh = GetMesh();
 
+		//warn if wheels have a body associated with it
+		FBodyInstance * WheelBody = Mesh->GetBodyInstance(WheelSetup.BoneName);
+		if (WheelBody)
+		{
+			UE_LOG(LogPhysics, Warning, TEXT("Wheel %s has a body instance associated with it. These are created by vehicle"), *WheelSetup.BoneName.ToString());
+			FMessageLog("PIE").Warning(FText::Format(LOCTEXT("WheelHasBodyWarning", "The wheel bone {0} already has a body setup in the physics asset. The vehicle simulation sets this body up automatically, the wheel body should be removed"), FText::FromString(WheelSetup.BoneName.ToString())));
+		}
+
 		if ( Mesh && Mesh->SkeletalMesh )
 		{
 			const FVector BonePosition = Mesh->SkeletalMesh->GetComposedRefPoseMatrix( WheelSetup.BoneName ).GetOrigin() * Mesh->RelativeScale3D;
@@ -486,6 +497,7 @@ FVector UWheeledVehicleMovementComponent::GetWheelRestingPosition( const FWheelS
 			const FMatrix RootBodyMTX = Mesh->SkeletalMesh->GetComposedRefPoseMatrix(Mesh->GetBodyInstance()->BodySetup->BoneName);
 			const FVector LocalBonePosition = RootBodyMTX.InverseTransformPosition(BonePosition);
 			Offset += LocalBonePosition;
+
 		}
 	}
 
@@ -1304,3 +1316,4 @@ void UWheeledVehicleMovementComponent::ComputeConstants()
 	MaxEngineRPM = 5000.f;
 }
 
+#undef LOCTEXT_NAMESPACE
\ No newline at end of file
