Settings: All projects settings are now saved in the Default.ini

Up until this change, project settings caused a lot of confusion amongst users, because they were changed to the user's local INI files. The local INI files, however, would be ignored when packaging a project, which means that local setting changes would not be applied to a packaged project and cause unexpected behavior. With this change the project settings will always save to the default INI file and will therefore be included in packaged projects. We will see how this works out for everyone, and we may change this behavior to something else if we can come up with a better, more intuitive workflow for project settings.

Upgrade Notes:
- If you have changes for Project Settings in your local INI files, you may want to 'Set as Default' in the Editor and then delete the local INI
- The default INI files (i.e. DefaultEngine.ini) of your project need to be writable. If you use source control, make sure you check out the corresponding files before changing project settings

diff --git a/Engine/Source/Developer/HTML5/HTML5TargetPlatform/Classes/HTML5TargetSettings.h b/Engine/Source/Developer/HTML5/HTML5TargetPlatform/Classes/HTML5TargetSettings.h
--- a/Engine/Source/Developer/HTML5/HTML5TargetPlatform/Classes/HTML5TargetSettings.h
+++ b/Engine/Source/Developer/HTML5/HTML5TargetPlatform/Classes/HTML5TargetSettings.h
@@ -12,7 +12,7 @@
 /**
  * Implements the settings for the HTML5 target platform.
  */
-UCLASS()
+UCLASS(config=Engine, defaultconfig)
 class HTML5TARGETPLATFORM_API UHTML5TargetSettings
 	: public UObject
 {
diff --git a/Engine/Source/Developer/Windows/WindowsTargetPlatform/Classes/WindowsTargetSettings.h b/Engine/Source/Developer/Windows/WindowsTargetPlatform/Classes/WindowsTargetSettings.h
--- a/Engine/Source/Developer/Windows/WindowsTargetPlatform/Classes/WindowsTargetSettings.h
+++ b/Engine/Source/Developer/Windows/WindowsTargetPlatform/Classes/WindowsTargetSettings.h
@@ -12,7 +12,7 @@
 /**
  * Implements the settings for the Windows target platform.
  */
-UCLASS()
+UCLASS(config=Engine, defaultconfig)
 class WINDOWSTARGETPLATFORM_API UWindowsTargetSettings
 	: public UObject
 {
diff --git a/Engine/Source/Editor/UnrealEd/Classes/Settings/ProjectPackagingSettings.h b/Engine/Source/Editor/UnrealEd/Classes/Settings/ProjectPackagingSettings.h
--- a/Engine/Source/Editor/UnrealEd/Classes/Settings/ProjectPackagingSettings.h
+++ b/Engine/Source/Editor/UnrealEd/Classes/Settings/ProjectPackagingSettings.h
@@ -51,7 +51,7 @@ struct FDirectoryPath
 /**
  * Implements the Editor's user settings.
  */
-UCLASS(config=Game)
+UCLASS(config=Game, defaultconfig)
 class UNREALED_API UProjectPackagingSettings
 	: public UObject
 {
diff --git a/Engine/Source/Runtime/Android/AndroidRuntimeSettings/Classes/AndroidRuntimeSettings.h b/Engine/Source/Runtime/Android/AndroidRuntimeSettings/Classes/AndroidRuntimeSettings.h
--- a/Engine/Source/Runtime/Android/AndroidRuntimeSettings/Classes/AndroidRuntimeSettings.h
+++ b/Engine/Source/Runtime/Android/AndroidRuntimeSettings/Classes/AndroidRuntimeSettings.h
@@ -72,7 +72,7 @@ struct FGooglePlayLeaderboardMapping
 /**
  * Implements the settings for the Android runtime platform.
  */
-UCLASS(config = Engine)
+UCLASS(config=Engine, defaultconfig)
 class ANDROIDRUNTIMESETTINGS_API UAndroidRuntimeSettings : public UObject
 {
 public:
diff --git a/Engine/Source/Runtime/Engine/Classes/AI/Navigation/NavigationSystem.h b/Engine/Source/Runtime/Engine/Classes/AI/Navigation/NavigationSystem.h
--- a/Engine/Source/Runtime/Engine/Classes/AI/Navigation/NavigationSystem.h
+++ b/Engine/Source/Runtime/Engine/Classes/AI/Navigation/NavigationSystem.h
@@ -516,7 +516,7 @@ struct FNavigationSystemExec: public FSelfRegisteringExec
 	// End FExec Interface
 };
 
-UCLASS(Within=World, config=Engine)
+UCLASS(Within=World, config=Engine, defaultconfig)
 class ENGINE_API UNavigationSystem : public UBlueprintFunctionLibrary
 {
 	GENERATED_UCLASS_BODY()
diff --git a/Engine/Source/Runtime/Engine/Classes/AI/Navigation/RecastNavMesh.h b/Engine/Source/Runtime/Engine/Classes/AI/Navigation/RecastNavMesh.h
--- a/Engine/Source/Runtime/Engine/Classes/AI/Navigation/RecastNavMesh.h
+++ b/Engine/Source/Runtime/Engine/Classes/AI/Navigation/RecastNavMesh.h
@@ -363,7 +363,7 @@ struct FTileSetItem
 	}
 };
 
-UCLASS(config=Engine, hidecategories=(Transform,"Utilities|Orientation",Actor,Layers,Replication))
+UCLASS(config=Engine, defaultconfig, hidecategories=(Transform,"Utilities|Orientation",Actor,Layers,Replication))
 class ENGINE_API ARecastNavMesh : public ANavigationData
 {
 	GENERATED_UCLASS_BODY()
diff --git a/Engine/Source/Runtime/Engine/Classes/Engine/CollisionProfile.h b/Engine/Source/Runtime/Engine/Classes/Engine/CollisionProfile.h
--- a/Engine/Source/Runtime/Engine/Classes/Engine/CollisionProfile.h
+++ b/Engine/Source/Runtime/Engine/Classes/Engine/CollisionProfile.h
@@ -9,6 +9,10 @@
 
 #include "CollisionProfile.generated.h"
 
+
+/**
+ * Structure for collision response templates.
+ */
 USTRUCT()
 struct ENGINE_API FCollisionResponseTemplate
 {
@@ -52,6 +56,10 @@ struct ENGINE_API FCollisionResponseTemplate
 	void CreateCustomResponsesFromResponseContainers();
 };
 
+
+/**
+ * Structure for custom channel setup information.
+ */
 USTRUCT()
 struct FCustomChannelSetup
 {
@@ -89,8 +97,10 @@ struct FCustomChannelSetup
 	}
 };
 
+
 /**
- * This is used only when customize profile
+ * Structure for custom profiles.
+ *
  * if you'd like to just add custom channels, not changing anything else engine defined
  * if you'd like to override all about profile, please use 
  * +Profiles=(Name=NameOfProfileYouLikeToOverwrite,....)
@@ -108,8 +118,13 @@ struct ENGINE_API FCustomProfile
 	TArray<FResponseChannel>	CustomResponses;
 };
 
-UCLASS(abstract, config=Engine, MinimalAPI)
-class UCollisionProfile : public UObject
+
+/**
+ * Implements a collision profile for the collision sub-system.
+ */
+UCLASS(abstract, config=Engine, defaultconfig, MinimalAPI)
+class UCollisionProfile
+	: public UObject
 {
 	GENERATED_UCLASS_BODY()
 
@@ -173,6 +188,7 @@ public:
 
 	/* custom collision profile name that you can modify what you'd like */
 	ENGINE_API static FName CustomCollisionProfileName;
+
 private:
 
 	/** 
@@ -227,4 +243,3 @@ private:
 
 	friend class FCollisionProfileDetails;
 };
-
diff --git a/Engine/Source/Runtime/Engine/Classes/Engine/Console.h b/Engine/Source/Runtime/Engine/Classes/Engine/Console.h
--- a/Engine/Source/Runtime/Engine/Classes/Engine/Console.h
+++ b/Engine/Source/Runtime/Engine/Classes/Engine/Console.h
@@ -7,6 +7,10 @@
 
 #include "Console.generated.h"
 
+
+/**
+ * Structure for auto-complete commands and their descriptions.
+ */
 USTRUCT()
 struct FAutoCompleteCommand
 {
@@ -20,7 +24,10 @@ struct FAutoCompleteCommand
 
 };
 
-/** Node for storing an auto-complete tree based on each char in the command */
+
+/**
+ * Node for storing an auto-complete tree based on each char in the command.
+ */
 USTRUCT()
 struct FAutoCompleteNode
 {
@@ -30,37 +37,39 @@ struct FAutoCompleteNode
 	UPROPERTY()
 	int32 IndexChar;
 
-	/** Indicies into AutoCompleteList for commands that match to this level */
+	/** Indices into AutoCompleteList for commands that match to this level */
 	UPROPERTY()
 	TArray<int32> AutoCompleteListIndices;
 
+	/** Children for further matching */
+	TArray<FAutoCompleteNode*> ChildNodes;
 
+	FAutoCompleteNode()
+	{
+		IndexChar = INDEX_NONE;
+	}
 
-		/** Children for further matching */
-		TArray<FAutoCompleteNode*> ChildNodes;
+	FAutoCompleteNode(int32 NewChar)
+	{
+		IndexChar = NewChar;
+	}
 
-		FAutoCompleteNode()
-		{
-			IndexChar = INDEX_NONE;
-		}
-		FAutoCompleteNode(int32 NewChar)
+	~FAutoCompleteNode()
+	{
+		for (int32 Idx = 0; Idx < ChildNodes.Num(); Idx++)
 		{
-			IndexChar = NewChar;
+			FAutoCompleteNode *Node = ChildNodes[Idx];
+			delete Node;
 		}
-		~FAutoCompleteNode()
-		{
-			for (int32 Idx = 0; Idx < ChildNodes.Num(); Idx++)
-			{
-				FAutoCompleteNode *Node = ChildNodes[Idx];
-				delete Node;
-			}
-			ChildNodes.Empty();
-		}
-	
+		ChildNodes.Empty();
+	}	
 };
 
-UCLASS(Within=GameViewportClient, transient, config=Input)
-class ENGINE_API UConsole : public UObject, public FOutputDevice
+
+UCLASS(Within=GameViewportClient, config=Input, defaultconfig, transient)
+class ENGINE_API UConsole
+	: public UObject
+	, public FOutputDevice
 {
 	GENERATED_UCLASS_BODY()
 
@@ -149,7 +158,6 @@ class ENGINE_API UConsole : public UObject, public FOutputDevice
 	/** Current list of matching commands for auto-complete, @see UpdateCompleteIndices() */
 	TArray<int32> AutoCompleteIndices;
 
-
 	/** Max number of command history entries */
 	static const int32 MAX_HISTORY_ENTRIES = 16;
 
@@ -250,7 +258,6 @@ class ENGINE_API UConsole : public UObject, public FOutputDevice
 	 */
 	virtual void PostRender_Console_Open(class UCanvas* Canvas);
 
-
 	/** Perform actions on transition to the Open state */
 	virtual void BeginState_Open(FName PreviousStateName);
 
@@ -290,6 +297,3 @@ private:
 	*/
 	virtual void PurgeCommandFromHistory(const FString& Command);
 };
-
-
-
diff --git a/Engine/Source/Runtime/Engine/Classes/Engine/Engine.h b/Engine/Source/Runtime/Engine/Classes/Engine/Engine.h
--- a/Engine/Source/Runtime/Engine/Classes/Engine/Engine.h
+++ b/Engine/Source/Runtime/Engine/Classes/Engine/Engine.h
@@ -465,8 +465,10 @@ class IAnalyticsProvider;
 DECLARE_DELEGATE_OneParam(FBeginStreamingPauseDelegate, FViewport*);
 DECLARE_DELEGATE(FEndStreamingPauseDelegate);
 
-UCLASS(abstract, config=Engine, transient)
-class ENGINE_API UEngine : public UObject, public FExec
+UCLASS(abstract, config=Engine, defaultconfig, transient)
+class ENGINE_API UEngine
+	: public UObject
+	, public FExec
 {
 	GENERATED_UCLASS_BODY()
 
diff --git a/Engine/Source/Runtime/Engine/Classes/Engine/RendererSettings.h b/Engine/Source/Runtime/Engine/Classes/Engine/RendererSettings.h
--- a/Engine/Source/Runtime/Engine/Classes/Engine/RendererSettings.h
+++ b/Engine/Source/Runtime/Engine/Classes/Engine/RendererSettings.h
@@ -4,6 +4,10 @@
 
 #include "RendererSettings.generated.h"
 
+
+/**
+ * Enumerates ways to clear a scene.
+ */
 UENUM()
 namespace EClearSceneOptions
 {
@@ -15,6 +19,10 @@ namespace EClearSceneOptions
 	};
 }
 
+
+/**
+ * Enumerates available compositing sample counts.
+ */
 UENUM()
 namespace ECompositingSampleCount
 {
@@ -27,6 +35,10 @@ namespace ECompositingSampleCount
 	};
 }
 
+
+/**
+ * Enumerates available options for custom depth.
+ */
 UENUM()
 namespace ECustomDepth
 {
@@ -38,6 +50,10 @@ namespace ECustomDepth
 	};
 }
 
+
+/**
+ * Enumerates available options for early Z-passes.
+ */
 UENUM()
 namespace EEarlyZPass
 {
@@ -50,8 +66,13 @@ namespace EEarlyZPass
 	};
 }
 
-UCLASS(config=Engine)
-class ENGINE_API URendererSettings : public UObject
+
+/**
+ * Implements project settings for the Rendering sub-system.
+ */
+UCLASS(config=Engine, defaultconfig)
+class ENGINE_API URendererSettings
+	: public UObject
 {
 	GENERATED_UCLASS_BODY()
 
@@ -145,10 +166,14 @@ class ENGINE_API URendererSettings : public UObject
 		ToolTip="Screen radius at which wireframe objects are culled. Larger values can improve performance when viewing a scene in wireframe."))
 	float WireframeCullThreshold;
 
-	// Begin UObject interface.
+public:
+
+	// Begin UObject interface
+
 	virtual void PostInitProperties() OVERRIDE;
 #if WITH_EDITOR
 	virtual void PostEditChangeProperty(FPropertyChangedEvent& PropertyChangedEvent) OVERRIDE;
 #endif
-	// End UObject interface.
+
+	// End UObject interface
 };
diff --git a/Engine/Source/Runtime/Engine/Classes/GameFramework/InputSettings.h b/Engine/Source/Runtime/Engine/Classes/GameFramework/InputSettings.h
--- a/Engine/Source/Runtime/Engine/Classes/GameFramework/InputSettings.h
+++ b/Engine/Source/Runtime/Engine/Classes/GameFramework/InputSettings.h
@@ -4,8 +4,10 @@
 #include "PlayerInput.h"
 #include "InputSettings.generated.h"
 
-UCLASS(config=Input)
-class ENGINE_API UInputSettings : public UObject
+
+UCLASS(config=Input, defaultconfig)
+class ENGINE_API UInputSettings
+	: public UObject
 {
 	GENERATED_UCLASS_BODY()
 
diff --git a/Engine/Source/Runtime/Engine/Classes/PhysicsEngine/PhysicsSettings.h b/Engine/Source/Runtime/Engine/Classes/PhysicsEngine/PhysicsSettings.h
--- a/Engine/Source/Runtime/Engine/Classes/PhysicsEngine/PhysicsSettings.h
+++ b/Engine/Source/Runtime/Engine/Classes/PhysicsEngine/PhysicsSettings.h
@@ -8,6 +8,10 @@
 
 #include "PhysicsSettings.generated.h"
 
+
+/**
+ * Structure that represents the name of physical surfaces.
+ */
 USTRUCT()
 struct FPhysicalSurfaceName
 {
@@ -28,8 +32,13 @@ struct FPhysicalSurfaceName
 	{}
 };
 
-UCLASS(config=Engine)
-class ENGINE_API UPhysicsSettings: public UObject
+
+/**
+ * Implements project settings for the physics sub-system.
+ */
+UCLASS(config=Engine, defaultconfig)
+class ENGINE_API UPhysicsSettings
+	: public UObject
 {
 	GENERATED_UCLASS_BODY()
 
@@ -41,7 +50,6 @@ class ENGINE_API UPhysicsSettings: public UObject
 	UPROPERTY(config, EditAnywhere, AdvancedDisplay, Category=Simulation)
 	bool bEnableAsyncScene;
 
-
 	/**Max Physics Delta Time to be clamped. */
 	UPROPERTY(config, EditAnywhere, meta=(ClampMin="0.0013", UIMin = "0.0013", ClampMax="1.0", UIMax="1.0"), Category=Framerate)
 	float MaxPhysicsDeltaTime;
@@ -75,14 +83,18 @@ class ENGINE_API UPhysicsSettings: public UObject
 	TArray<FPhysicalSurfaceName> PhysicalSurfaces;
 
 public:
+
 	static UPhysicsSettings * Get() { return CastChecked<UPhysicsSettings>(UPhysicsSettings::StaticClass()->GetDefaultObject()); }
 
-	virtual void	PostInitProperties() OVERRIDE;
+	virtual void PostInitProperties() OVERRIDE;
 
 #if WITH_EDITOR
+
 	virtual bool CanEditChange( const UProperty* Property ) const OVERRIDE;
+
 	/** Load Material Type data from INI file **/
 	/** this changes displayname meta data. That means we won't need it outside of editor*/
 	void LoadSurfaceType();
+
 #endif // WITH_EDITOR
-};
\ No newline at end of file
+};
diff --git a/Engine/Source/Runtime/Engine/Classes/Sound/AudioSettings.h b/Engine/Source/Runtime/Engine/Classes/Sound/AudioSettings.h
--- a/Engine/Source/Runtime/Engine/Classes/Sound/AudioSettings.h
+++ b/Engine/Source/Runtime/Engine/Classes/Sound/AudioSettings.h
@@ -4,8 +4,13 @@
 
 #include "AudioSettings.generated.h"
 
-UCLASS(config=Engine)
-class ENGINE_API UAudioSettings : public UObject
+
+/**
+ * Implements audio related settings.
+ */
+UCLASS(config=Engine, defaultconfig)
+class ENGINE_API UAudioSettings
+	: public UObject
 {
 	GENERATED_UCLASS_BODY()
 
diff --git a/Engine/Source/Runtime/EngineSettings/Classes/GameMapsSettings.h b/Engine/Source/Runtime/EngineSettings/Classes/GameMapsSettings.h
--- a/Engine/Source/Runtime/EngineSettings/Classes/GameMapsSettings.h
+++ b/Engine/Source/Runtime/EngineSettings/Classes/GameMapsSettings.h
@@ -8,6 +8,7 @@
 
 #include "GameMapsSettings.generated.h"
 
+
 UENUM()
 namespace ETwoPlayerSplitScreenType
 {
@@ -18,6 +19,7 @@ namespace ETwoPlayerSplitScreenType
 	};
 }
 
+
 UENUM()
 namespace EThreePlayerSplitScreenType
 {
@@ -28,7 +30,8 @@ namespace EThreePlayerSplitScreenType
 	};
 }
 
-UCLASS(config=Engine)
+
+UCLASS(config=Engine, defaultconfig)
 class ENGINESETTINGS_API UGameMapsSettings
 	: public UObject
 {
diff --git a/Engine/Source/Runtime/IOS/IOSRuntimeSettings/Classes/IOSRuntimeSettings.h b/Engine/Source/Runtime/IOS/IOSRuntimeSettings/Classes/IOSRuntimeSettings.h
--- a/Engine/Source/Runtime/IOS/IOSRuntimeSettings/Classes/IOSRuntimeSettings.h
+++ b/Engine/Source/Runtime/IOS/IOSRuntimeSettings/Classes/IOSRuntimeSettings.h
@@ -8,7 +8,7 @@
 /**
  * Implements the settings for the iOS target platform.
  */
-UCLASS(config = Engine)
+UCLASS(config=Engine, defaultconfig)
 class IOSRUNTIMESETTINGS_API UIOSRuntimeSettings : public UObject
 {
 public:
diff --git a/Engine/Source/Runtime/MoviePlayer/Public/MoviePlayerSettings.h b/Engine/Source/Runtime/MoviePlayer/Public/MoviePlayerSettings.h
--- a/Engine/Source/Runtime/MoviePlayer/Public/MoviePlayerSettings.h
+++ b/Engine/Source/Runtime/MoviePlayer/Public/MoviePlayerSettings.h
@@ -4,11 +4,13 @@
 
 #include "MoviePlayerSettings.generated.h"
 
+
 /**
  * Implements the settings for the Windows target platform.
  */
-UCLASS(config=Game)
-class MOVIEPLAYER_API UMoviePlayerSettings : public UObject
+UCLASS(config=Game, defaultconfig)
+class MOVIEPLAYER_API UMoviePlayerSettings
+	: public UObject
 {
 public:
 
@@ -19,14 +21,10 @@ public:
 	bool bWaitForMoviesToComplete;
 
 	/** If enabled startup movies can be skipped by the user when a mouse button is pressed */
-	UPROPERTY(globalconfig, EditAnywhere, Category = "Movies")
+	UPROPERTY(globalconfig, EditAnywhere, Category="Movies")
 	bool bMoviesAreSkippable;
 
-	/**
-	 * Movies to play on startup. Note that these must be in your game's Game/Content/Movies directory.
-	 */
-	UPROPERTY(globalconfig, EditAnywhere, Category="Movies", meta=(FilePathFilter = "mp4"))
+	/** Movies to play on startup. Note that these must be in your game's Game/Content/Movies directory. */
+	UPROPERTY(globalconfig, EditAnywhere, Category="Movies", meta=(FilePathFilter="mp4"))
 	TArray<FFilePath> StartupMovies;
-
-
 };
