Vehicle anim instance and wheel handler for vehicle

diff --git a/Engine/Source/Editor/AnimGraph/Classes/AnimGraphNode_WheelHandler.h b/Engine/Source/Editor/AnimGraph/Classes/AnimGraphNode_WheelHandler.h
new file mode 100644
--- /dev/null
+++ b/Engine/Source/Editor/AnimGraph/Classes/AnimGraphNode_WheelHandler.h
@@ -0,0 +1,30 @@
+// Copyright 1998-2014 Epic Games, Inc. All Rights Reserved.
+
+#pragma once
+
+#include "AnimGraphNode_SkeletalControlBase.h"
+#include "AnimGraphNode_WheelHandler.generated.h"
+
+UCLASS(MinimalAPI)
+class UAnimGraphNode_WheelHandler : public UAnimGraphNode_SkeletalControlBase
+{
+	GENERATED_UCLASS_BODY()
+
+	UPROPERTY(EditAnywhere, Category=Settings)
+	FAnimNode_WheelHandler Node;
+
+public:
+	// UEdGraphNode interface
+	virtual FText GetNodeTitle(ENodeTitleType::Type TitleType) const OVERRIDE;
+	virtual FString GetTooltip() const OVERRIDE;
+	virtual FString GetKeywords() const OVERRIDE;
+	// validate if this is within VehicleAnimInstance
+	virtual void ValidateAnimNodePostCompile(class FCompilerResultsLog& MessageLog, class UAnimBlueprintGeneratedClass* CompiledClass, int32 CompiledNodeIndex) OVERRIDE;
+	virtual void GetMenuEntries(FGraphContextMenuBuilder& ContextMenuBuilder) const OVERRIDE;
+	// End of UEdGraphNode interface
+
+protected:
+	// UAnimGraphNode_SkeletalControlBase interface
+	virtual FText GetControllerDescription() const OVERRIDE;
+	// End of UAnimGraphNode_SkeletalControlBase interface
+};
diff --git a/Engine/Source/Editor/AnimGraph/Private/AnimGraphNode_WheelHandler.cpp b/Engine/Source/Editor/AnimGraph/Private/AnimGraphNode_WheelHandler.cpp
new file mode 100644
--- /dev/null
+++ b/Engine/Source/Editor/AnimGraph/Private/AnimGraphNode_WheelHandler.cpp
@@ -0,0 +1,59 @@
+// Copyright 1998-2014 Epic Games, Inc. All Rights Reserved.
+
+#include "AnimGraphPrivatePCH.h"
+#include "AnimGraphNode_WheelHandler.h"
+
+/////////////////////////////////////////////////////
+// UAnimGraphNode_WheelHandler
+
+#define LOCTEXT_NAMESPACE "A3Nodes"
+
+UAnimGraphNode_WheelHandler::UAnimGraphNode_WheelHandler(const FPostConstructInitializeProperties& PCIP)
+	: Super(PCIP)
+{
+}
+
+FText UAnimGraphNode_WheelHandler::GetControllerDescription() const
+{
+	return LOCTEXT("AnimGraphNode_WheelHandler", "Wheel Handler for WheeledVehicle");
+}
+
+FString UAnimGraphNode_WheelHandler::GetKeywords() const
+{
+	return TEXT("Modify, Wheel, Vehicle");
+}
+
+FString UAnimGraphNode_WheelHandler::GetTooltip() const
+{
+	return LOCTEXT("AnimGraphNode_WheelHandler_Tooltip", "This alters the wheel transform based on set up in Wheeled Vehicle. This only works when the owner is WheeledVehicle.").ToString();
+}
+
+FText UAnimGraphNode_WheelHandler::GetNodeTitle(ENodeTitleType::Type TitleType) const
+{
+	// we don't have any run-time information, so it's limited to print anymore than whaat it is
+	// it would be nice to print more data such as name of bones for wheels, but it's not available in Persona
+	return FText(LOCTEXT("AnimGraphNode_WheelHandler_Title", "Wheel Handler"));
+}
+
+void UAnimGraphNode_WheelHandler::ValidateAnimNodePostCompile(class FCompilerResultsLog& MessageLog, class UAnimBlueprintGeneratedClass* CompiledClass, int32 CompiledNodeIndex)
+{
+	// we only support vehicle anim instance
+	if ( CompiledClass->IsChildOf(UVehicleAnimInstance::StaticClass())  == false )
+	{
+		MessageLog.Error(TEXT("@@ is only allowwed in VehicleAnimInstance. If this is for vehicle, please change parent to be VehicleAnimInstancen (Reparent Class)."), this);
+	}
+}
+
+void UAnimGraphNode_WheelHandler::GetMenuEntries(FGraphContextMenuBuilder& ContextMenuBuilder) const
+{
+	UAnimBlueprint* Blueprint = Cast<UAnimBlueprint> (FBlueprintEditorUtils::FindBlueprintForGraph(ContextMenuBuilder.CurrentGraph));
+
+	if (Blueprint &&  Blueprint->ParentClass->IsChildOf(UVehicleAnimInstance::StaticClass()))
+	{
+		Super::GetMenuEntries( ContextMenuBuilder );
+	}
+
+	// else we don't show
+}
+
+#undef LOCTEXT_NAMESPACE
diff --git a/Engine/Source/Runtime/Engine/Classes/Animation/BoneControllers/AnimNode_WheelHandler.h b/Engine/Source/Runtime/Engine/Classes/Animation/BoneControllers/AnimNode_WheelHandler.h
new file mode 100644
--- /dev/null
+++ b/Engine/Source/Runtime/Engine/Classes/Animation/BoneControllers/AnimNode_WheelHandler.h
@@ -0,0 +1,47 @@
+// Copyright 1998-2014 Epic Games, Inc. All Rights Reserved.
+
+#pragma once
+
+#include "AnimNode_SkeletalControlBase.h"
+#include "AnimNode_WheelHandler.generated.h"
+
+struct FWheelSimulator
+{
+	int32					WheelIndex;
+	FBoneReference			BoneReference;
+	FRotator				RotOffset;
+	FVector					LocOffset;
+};
+
+/**
+ *	Simple controller that replaces or adds to the translation/rotation of a single bone.
+ */
+USTRUCT()
+struct ENGINE_API FAnimNode_WheelHandler : public FAnimNode_SkeletalControlBase
+{
+	GENERATED_USTRUCT_BODY()
+
+	/** Current Asset being played **/
+	UPROPERTY(transient)
+	class UWheeledVehicleMovementComponent* VehicleSimComponent;
+
+	TArray<FWheelSimulator>			WheelSimulators;
+
+	FAnimNode_WheelHandler();
+
+	// FAnimNode_Base interface
+	virtual void GatherDebugData(FNodeDebugData& DebugData) OVERRIDE;
+	// End of FAnimNode_Base interface
+
+	// FAnimNode_SkeletalControlBase interface
+	virtual void EvaluateBoneTransforms(USkeletalMeshComponent* SkelComp, const FBoneContainer & RequiredBones, FA2CSPose& MeshBases, TArray<FBoneTransform>& OutBoneTransforms) OVERRIDE;
+	virtual bool IsValidToEvaluate(const USkeleton * Skeleton, const FBoneContainer & RequiredBones) OVERRIDE;
+	virtual void Update(const FAnimationUpdateContext& Context) OVERRIDE;
+	virtual void Initialize(const FAnimationInitializeContext& Context) OVERRIDE;
+	// End of FAnimNode_SkeletalControlBase interface
+
+private:
+	// FAnimNode_SkeletalControlBase interface
+	virtual void InitializeBoneReferences(const FBoneContainer & RequiredBones) OVERRIDE;
+	// End of FAnimNode_SkeletalControlBase interface
+};
diff --git a/Engine/Source/Runtime/Engine/Classes/Vehicles/VehicleAnimInstance.h b/Engine/Source/Runtime/Engine/Classes/Vehicles/VehicleAnimInstance.h
new file mode 100644
--- /dev/null
+++ b/Engine/Source/Runtime/Engine/Classes/Vehicles/VehicleAnimInstance.h
@@ -0,0 +1,20 @@
+// Copyright 1998-2014 Epic Games, Inc. All Rights Reserved.
+
+#pragma once
+/**
+ *
+ */
+#include "VehicleAnimInstance.generated.h"
+
+UCLASS(transient)
+class ENGINE_API UVehicleAnimInstance : public UAnimInstance
+{
+	GENERATED_UCLASS_BODY()
+
+	/** Makes a montage jump to the end of a named section. */
+	UFUNCTION(BlueprintCallable, Category="Animation")
+	class AWheeledVehicle * GetVehicle();
+};
+
+
+
diff --git a/Engine/Source/Runtime/Engine/Private/Animation/AnimNode_WheelHandler.cpp b/Engine/Source/Runtime/Engine/Private/Animation/AnimNode_WheelHandler.cpp
new file mode 100644
--- /dev/null
+++ b/Engine/Source/Runtime/Engine/Private/Animation/AnimNode_WheelHandler.cpp
@@ -0,0 +1,156 @@
+// Copyright 1998-2014 Epic Games, Inc. All Rights Reserved.
+
+#include "EnginePrivate.h"
+#include "AnimationRuntime.h"
+
+/////////////////////////////////////////////////////
+// FAnimNode_WheelHandler
+
+FAnimNode_WheelHandler::FAnimNode_WheelHandler()
+{
+}
+
+void FAnimNode_WheelHandler::GatherDebugData(FNodeDebugData& DebugData)
+{
+	FString DebugLine = DebugData.GetNodeName(this);
+
+	DebugLine += "(";
+	AddDebugNodeData(DebugLine);
+	DebugLine += ")";
+
+	DebugData.AddDebugItem(DebugLine);
+	for(const auto & WheelSim : WheelSimulators)
+	{
+		if (WheelSim.BoneReference.BoneIndex != INDEX_NONE)
+		{
+			DebugLine = FString::Printf(TEXT("\t[%d] Bone: %s , Rotation Offset : %s, Location Offset : %s"), 
+				WheelSim.WheelIndex, *WheelSim.BoneReference.BoneName.ToString(), *WheelSim.RotOffset.ToString(), *WheelSim.LocOffset.ToString());
+		}
+		else
+		{
+			DebugLine = FString::Printf(TEXT("\t[%d] Bone: %s (invalid bone)"),
+					WheelSim.WheelIndex, *WheelSim.BoneReference.BoneName.ToString());
+		}
+
+		DebugData.AddDebugItem(DebugLine);
+	}
+
+	ComponentPose.GatherDebugData(DebugData);
+}
+
+void FAnimNode_WheelHandler::EvaluateBoneTransforms(USkeletalMeshComponent* SkelComp, const FBoneContainer & RequiredBones, FA2CSPose& MeshBases, TArray<FBoneTransform>& OutBoneTransforms)
+{
+	check(OutBoneTransforms.Num() == 0);
+
+	for(const auto & WheelSim : WheelSimulators)
+	{
+		if(WheelSim.BoneReference.IsValid(RequiredBones))
+		{
+			// the way we apply transform is same as FMatrix or FTransform
+			// we apply scale first, and rotation, and translation
+			// if you'd like to translate first, you'll need two nodes that first node does translate and second nodes to rotate. 
+			FTransform NewBoneTM = MeshBases.GetComponentSpaceTransform(WheelSim.BoneReference.BoneIndex);
+
+			FAnimationRuntime::ConvertCSTransformToBoneSpace(SkelComp, MeshBases, NewBoneTM, WheelSim.BoneReference.BoneIndex, BCS_ComponentSpace);
+			
+			// Apply rotation offset
+			const FQuat BoneQuat(WheelSim.RotOffset);
+			NewBoneTM.SetRotation(BoneQuat * NewBoneTM.GetRotation());
+
+			// Apply loc offset
+			NewBoneTM.AddToTranslation(WheelSim.LocOffset);
+
+			// Convert back to Component Space.
+			FAnimationRuntime::ConvertBoneSpaceTransformToCS(SkelComp, MeshBases, NewBoneTM, WheelSim.BoneReference.BoneIndex, BCS_ComponentSpace);
+
+			// add back to it
+			OutBoneTransforms.Add( FBoneTransform(WheelSim.BoneReference.BoneIndex, NewBoneTM) );
+		}
+	}
+}
+
+bool FAnimNode_WheelHandler::IsValidToEvaluate(const USkeleton * Skeleton, const FBoneContainer & RequiredBones) 
+{
+	// if both bones are valid
+	for(const auto & WheelSim : WheelSimulators)
+	{
+		// if one of them is valid
+		if (WheelSim.BoneReference.IsValid(RequiredBones) == true)
+		{
+			return true;
+		}
+	}
+
+	return false;
+}
+
+void FAnimNode_WheelHandler::InitializeBoneReferences(const FBoneContainer & RequiredBones) 
+{
+	for (auto & WheelSim : WheelSimulators)
+	{
+		WheelSim.BoneReference.Initialize(RequiredBones);
+	}
+
+	// sort by bone indices
+	WheelSimulators.Sort([](FWheelSimulator L, FWheelSimulator R) { return L.BoneReference.BoneIndex < R.BoneReference.BoneIndex; });
+}
+
+void FAnimNode_WheelHandler::Update(const FAnimationUpdateContext& Context)
+{
+	if(VehicleSimComponent)
+	{
+		for(auto & WheelSim : WheelSimulators)
+		{
+			if(VehicleSimComponent->Wheels.IsValidIndex(WheelSim.WheelIndex))
+			{
+				UVehicleWheel* Wheel = VehicleSimComponent->Wheels[WheelSim.WheelIndex];
+				if(Wheel != nullptr)
+				{
+					WheelSim.RotOffset.Pitch	= Wheel->GetRotationAngle();
+					WheelSim.RotOffset.Yaw		= Wheel->GetSteerAngle();
+					WheelSim.RotOffset.Roll		= 0.f;
+
+					WheelSim.LocOffset.X		= 0.f;
+					WheelSim.LocOffset.Y		= 0.f;
+					WheelSim.LocOffset.Z		= Wheel->GetSuspensionOffset();
+				}
+			}
+		}
+	}
+
+	FAnimNode_SkeletalControlBase::Update(Context);
+}
+
+void FAnimNode_WheelHandler::Initialize(const FAnimationInitializeContext& Context)
+{
+	// TODO: only check vehicle anim instance
+	// UVehicleAnimInstance
+	AWheeledVehicle * Vehicle = Cast<AWheeledVehicle> ((CastChecked<USkeletalMeshComponent> (Context.AnimInstance->GetOuter()))->GetOwner());
+	;
+	// we only support vehicle for this class
+	if(Vehicle != nullptr)
+	{
+		VehicleSimComponent = Vehicle->GetVehicleMovementComponent();
+
+		int32 NumOfwheels = VehicleSimComponent->WheelSetups.Num();
+		if(NumOfwheels > 0)
+		{
+			WheelSimulators.Empty(NumOfwheels);
+			WheelSimulators.AddZeroed(NumOfwheels);
+			// now add wheel data
+			for(int32 WheelIndex = 0; WheelIndex<WheelSimulators.Num(); ++WheelIndex)
+			{
+				FWheelSimulator & WheelSim = WheelSimulators[WheelIndex];
+				const FWheelSetup & WheelSetup = VehicleSimComponent->WheelSetups[WheelIndex];
+
+				// set data
+				WheelSim.WheelIndex = WheelIndex;
+				WheelSim.BoneReference.BoneName = WheelSetup.BoneName;
+				WheelSim.LocOffset = FVector::ZeroVector;
+				WheelSim.RotOffset = FRotator::ZeroRotator;
+			}
+		}
+	}
+
+	FAnimNode_SkeletalControlBase::Initialize(Context);
+}
\ No newline at end of file
diff --git a/Engine/Source/Runtime/Engine/Private/Vehicles/VehicleAnimInstance.cpp b/Engine/Source/Runtime/Engine/Private/Vehicles/VehicleAnimInstance.cpp
new file mode 100644
--- /dev/null
+++ b/Engine/Source/Runtime/Engine/Private/Vehicles/VehicleAnimInstance.cpp
@@ -0,0 +1,28 @@
+// Copyright 1998-2014 Epic Games, Inc. All Rights Reserved.
+
+/*=============================================================================
+	UVehicleAnimInstance.cpp: Single Node Tree Instance 
+	Only plays one animation at a time. 
+=============================================================================*/ 
+
+#include "EnginePrivate.h"
+#include "AnimationRuntime.h"
+
+/////////////////////////////////////////////////////
+// UVehicleAnimInstance
+/////////////////////////////////////////////////////
+
+UVehicleAnimInstance::UVehicleAnimInstance(const class FPostConstructInitializeProperties& PCIP)
+	: Super(PCIP)
+{
+}
+
+class AWheeledVehicle * UVehicleAnimInstance::GetVehicle()
+{
+	return Cast<AWheeledVehicle> (GetOwningActor());
+}
+
+/////////////////////////////////////////////////////
+/////////////////////////////////////////////////////
+/////////////////////////////////////////////////////
+
