Fix ShaderCompiler hang in rare cases when a SCW PID gets reused by another process

diff --git a/Engine/Source/Runtime/Core/Private/GenericPlatform/GenericPlatformProcess.cpp b/Engine/Source/Runtime/Core/Private/GenericPlatform/GenericPlatformProcess.cpp
--- a/Engine/Source/Runtime/Core/Private/GenericPlatform/GenericPlatformProcess.cpp
+++ b/Engine/Source/Runtime/Core/Private/GenericPlatform/GenericPlatformProcess.cpp
@@ -188,6 +188,12 @@ bool FGenericPlatformProcess::IsApplicationRunning( const TCHAR* ProcName )
 	return false;
 }
 
+FString FGenericPlatformProcess::GetApplicationName( uint32 ProcessId )
+{
+	UE_LOG(LogHAL, Fatal, TEXT("FGenericPlatformProcess::GetApplicationName not implemented on this platform"));
+	return FString(TEXT(""));
+}
+
 bool FGenericPlatformProcess::IsThisApplicationForeground()
 {
 	UE_LOG(LogHAL, Fatal, TEXT("FGenericPlatformProcess::IsThisApplicationForeground not implemented on this platform"));
diff --git a/Engine/Source/Runtime/Core/Private/Linux/LinuxPlatformProcess.cpp b/Engine/Source/Runtime/Core/Private/Linux/LinuxPlatformProcess.cpp
--- a/Engine/Source/Runtime/Core/Private/Linux/LinuxPlatformProcess.cpp
+++ b/Engine/Source/Runtime/Core/Private/Linux/LinuxPlatformProcess.cpp
@@ -126,6 +126,23 @@ const TCHAR* FLinuxPlatformProcess::ExecutableName(bool bRemoveExtension)
 	return CachedResult;
 }
 
+FString FLinuxPlatformProcess::GetApplicationName( uint32 ProcessId )
+{
+	FString Output = TEXT("");
+
+	const int32 ReadLinkSize = 1024;	
+	char ReadLinkCmd[ReadLinkSize] = {0};
+	FCStringAnsi::Sprintf(ReadLinkCmd, "/proc/%d/exe", ProcessId);
+	
+	char ProcessPath[ PlatformProcessLimits::MaxBaseDirLength ] = {0};
+	int32 Ret = readlink(ReadLinkCmd, ProcessPath, ARRAY_COUNT(ProcessPath) - 1);
+	if (Ret != -1)
+	{
+		Output = ANSI_TO_TCHAR(ProcessPath);
+	}
+	return Output;
+}
+
 FPipeHandle::~FPipeHandle()
 {
 	close(PipeDesc);
diff --git a/Engine/Source/Runtime/Core/Private/Mac/MacPlatformProcess.cpp b/Engine/Source/Runtime/Core/Private/Mac/MacPlatformProcess.cpp
--- a/Engine/Source/Runtime/Core/Private/Mac/MacPlatformProcess.cpp
+++ b/Engine/Source/Runtime/Core/Private/Mac/MacPlatformProcess.cpp
@@ -7,6 +7,7 @@
 #include "CorePrivate.h"
 #include "../../Launch/Resources/Version.h"
 #include <mach-o/dyld.h>
+#include <libproc.h>
 
 void* FMacPlatformProcess::GetDllHandle( const TCHAR* Filename )
 {
@@ -509,6 +510,7 @@ bool FMacPlatformProcess::IsApplicationRunning( const TCHAR* ProcName )
 	}
 	return false;
 }
+
 bool FMacPlatformProcess::IsApplicationRunning( uint32 ProcessId )
 {
 	errno = 0;
@@ -516,6 +518,19 @@ bool FMacPlatformProcess::IsApplicationRunning( uint32 ProcessId )
 	return errno == 0;
 }
 
+FString FMacPlatformProcess::GetApplicationName( uint32 ProcessId )
+{
+	FString Output = TEXT("");
+
+	char Buffer[MAX_PATH];
+	int32 Ret = proc_pidpath(ProcessId, Buffer, sizeof(Buffer));
+	if (Ret > 0)
+	{
+		Output = ANSI_TO_TCHAR(Buffer);
+	}
+	return Output;
+}
+
 bool FMacPlatformProcess::IsThisApplicationForeground()
 {
 	SCOPED_AUTORELEASE_POOL;
diff --git a/Engine/Source/Runtime/Core/Private/Windows/WindowsPlatformProcess.cpp b/Engine/Source/Runtime/Core/Private/Windows/WindowsPlatformProcess.cpp
--- a/Engine/Source/Runtime/Core/Private/Windows/WindowsPlatformProcess.cpp
+++ b/Engine/Source/Runtime/Core/Private/Windows/WindowsPlatformProcess.cpp
@@ -430,6 +430,25 @@ bool FWindowsPlatformProcess::IsApplicationRunning( const TCHAR* ProcName )
 	return false;
 }
 
+FString FWindowsPlatformProcess::GetApplicationName( uint32 ProcessId )
+{
+	FString Output = TEXT("");
+	HANDLE ProcessHandle = OpenProcess(PROCESS_QUERY_INFORMATION, false, ProcessId);
+	if (ProcessHandle != NULL)
+	{
+		const int32 ProcessNameBufferSize = 4096;
+		TCHAR ProcessNameBuffer[ProcessNameBufferSize];
+		
+		int32 InOutSize = ProcessNameBufferSize;
+		checkAtCompileTime(sizeof(::DWORD) == sizeof(int32), "DWORD size doesn't match int32.  Is it the future or the past?");
+		QueryFullProcessImageName(ProcessHandle, 0, ProcessNameBuffer, (PDWORD)(&InOutSize));
+		Output = ProcessNameBuffer;
+	}
+
+	return Output;
+}
+
+
 bool FWindowsPlatformProcess::IsThisApplicationForeground()
 {
 	uint32 ForegroundProcess;
diff --git a/Engine/Source/Runtime/Core/Public/GenericPlatform/GenericPlatformProcess.h b/Engine/Source/Runtime/Core/Public/GenericPlatform/GenericPlatformProcess.h
--- a/Engine/Source/Runtime/Core/Public/GenericPlatform/GenericPlatformProcess.h
+++ b/Engine/Source/Runtime/Core/Public/GenericPlatform/GenericPlatformProcess.h
@@ -327,8 +327,10 @@ struct CORE_API FGenericPlatformProcess
 	static bool IsApplicationRunning( uint32 ProcessId );
 	/** Returns true if the specified application is running */
 	static bool IsApplicationRunning( const TCHAR* ProcName );
+	/** Returns the Name of process given by the PID.  Returns Empty string "" if PID not found. */
+	static FString GetApplicationName( uint32 ProcessId );
 	/** Returns true if the specified application has a visible window, and that window is active/has focus/is selected */
-	static bool IsThisApplicationForeground();
+	static bool IsThisApplicationForeground();	
 
 	/**
 	 * Executes a process, returning the return code, stdout, and stderr. This
diff --git a/Engine/Source/Runtime/Core/Public/Mac/MacPlatformProcess.h b/Engine/Source/Runtime/Core/Public/Mac/MacPlatformProcess.h
--- a/Engine/Source/Runtime/Core/Public/Mac/MacPlatformProcess.h
+++ b/Engine/Source/Runtime/Core/Public/Mac/MacPlatformProcess.h
@@ -53,6 +53,7 @@ struct CORE_API FMacPlatformProcess : public FGenericPlatformProcess
 	static bool GetProcReturnCode( FProcHandle & ProcHandle, int32* ReturnCode );
 	static bool IsApplicationRunning( const TCHAR* ProcName );
 	static bool IsApplicationRunning( uint32 ProcessId );
+	static FString GetApplicationName( uint32 ProcessId );
 	static bool IsThisApplicationForeground();
 	static void LaunchFileInDefaultExternalApplication( const TCHAR* FileName, const TCHAR* Parms = NULL, ELaunchVerb::Type Verb = ELaunchVerb::Open );
 	static void ExploreFolder( const TCHAR* FilePath );
diff --git a/Engine/Source/Runtime/Core/Public/Windows/WindowsPlatformProcess.h b/Engine/Source/Runtime/Core/Public/Windows/WindowsPlatformProcess.h
--- a/Engine/Source/Runtime/Core/Public/Windows/WindowsPlatformProcess.h
+++ b/Engine/Source/Runtime/Core/Public/Windows/WindowsPlatformProcess.h
@@ -94,6 +94,7 @@ struct CORE_API FWindowsPlatformProcess : public FGenericPlatformProcess
 	static bool GetProcReturnCode( FProcHandle & ProcHandle, int32* ReturnCode );
 	static bool IsApplicationRunning( uint32 ProcessId );
 	static bool IsApplicationRunning( const TCHAR* ProcName );
+	static FString GetApplicationName( uint32 ProcessId );	
 	static bool IsThisApplicationForeground();
 	static bool ExecProcess( const TCHAR* URL, const TCHAR* Params, int32* OutReturnCode, FString* OutStdOut, FString* OutStdErr );
 	static bool ExecElevatedProcess(const TCHAR* URL, const TCHAR* Params, int32* OutReturnCode);
diff --git a/Engine/Source/Runtime/Engine/Private/ShaderCompiler/ShaderCompiler.cpp b/Engine/Source/Runtime/Engine/Private/ShaderCompiler/ShaderCompiler.cpp
--- a/Engine/Source/Runtime/Engine/Private/ShaderCompiler/ShaderCompiler.cpp
+++ b/Engine/Source/Runtime/Engine/Private/ShaderCompiler/ShaderCompiler.cpp
@@ -298,7 +298,7 @@ struct FShaderCompileWorkerInfo
 	uint32 WorkerAppId;
 
 	/** Tracks whether tasks have been issued to the worker. */
-	bool bIssuedTasksToWorker;
+	bool bIssuedTasksToWorker;	
 
 	/** Whether the worker has been launched for this set of tasks. */
 	bool bLaunchedWorker;
@@ -321,7 +321,7 @@ struct FShaderCompileWorkerInfo
 
 	FShaderCompileWorkerInfo() :
 		WorkerAppId(0),
-		bIssuedTasksToWorker(false),
+		bIssuedTasksToWorker(false),		
 		bLaunchedWorker(false),
 		bComplete(false),
 #if PLATFORM_SUPPORTS_NAMED_PIPES
@@ -345,7 +345,7 @@ struct FShaderCompileWorkerInfo
 			if (bWorkerForPipeWasLaunched)
 			{
 				// Make sure it's alive
-				if (!FPlatformProcess::IsApplicationRunning(WorkerAppId))
+				if (!FShaderCompilingManager::IsShaderCompilerWorkerRunning(WorkerAppId))
 				{
 					bAllocNameForPipe = true;
 					bWorkerForPipeWasLaunched = false;
@@ -548,8 +548,9 @@ int32 FShaderCompileThreadRunnable::PullTasksFromQueue()
 						CurrentWorkerInfo.QueuedJobs.Add(Manager->CompileQueue[JobIndex]);
 					}
 
-					// Update the worker state as having new tasks that need to be issued
-					CurrentWorkerInfo.bIssuedTasksToWorker = false;
+					// Update the worker state as having new tasks that need to be issued					
+					// don't reset worker app ID, because the shadercompilerworkers don't shutdown immediately after finishing a single job queue.
+					CurrentWorkerInfo.bIssuedTasksToWorker = false;					
 					CurrentWorkerInfo.bLaunchedWorker = false;
 					CurrentWorkerInfo.StartTime = FPlatformTime::Seconds();
 					NumActiveThreads++;
@@ -694,7 +695,7 @@ void FShaderCompileThreadRunnable::LaunchWorkerIfNeeded(FShaderCompileWorkerInfo
 			LastCheckForWorkersTime = CurrentTime;
 		}
 
-		if (bCheckForWorkerRunning && !FPlatformProcess::IsApplicationRunning(CurrentWorkerInfo.WorkerAppId))
+		if (bCheckForWorkerRunning && !FShaderCompilingManager::IsShaderCompilerWorkerRunning(CurrentWorkerInfo.WorkerAppId))
 		{
 			// Worker died, so clear this pipe and make a new one
 			CurrentWorkerInfo.PipeWorker.DestroyPipe();
@@ -711,13 +712,15 @@ void FShaderCompileThreadRunnable::LaunchWorkerIfNeeded(FShaderCompileWorkerInfo
 
 		// Store the Id with this thread so that we will know not to launch it again
 		const FString& PipeName = CurrentWorkerInfo.PipeWorker.NamedPipe.GetName();
-		CurrentWorkerInfo.WorkerAppId = Manager->LaunchWorker(WorkingDirectory, Manager->ProcessId, WorkerIndex, PipeName, PipeName, true, !GShaderPipeConfig.bReuseNamedPipeAndProcess);
+		CurrentWorkerInfo.WorkerAppId = Manager->LaunchWorker(WorkingDirectory, Manager->ProcessId, WorkerIndex, PipeName, PipeName, true, !GShaderPipeConfig.bReuseNamedPipeAndProcess);		
 		CurrentWorkerInfo.bLaunchedWorker = true;
 		CurrentWorkerInfo.bWorkerForPipeWasLaunched = true;
 	}
 #endif
 }
 
+
+
 bool FShaderCompileThreadRunnable::LaunchWorkersIfNeeded()
 {
 	const double CurrentTime = FPlatformTime::Seconds();
@@ -747,12 +750,13 @@ bool FShaderCompileThreadRunnable::LaunchWorkersIfNeeded()
 		else
 #endif	// PLATFORM_SUPPORTS_NAMED_PIPES
 		{
-			if (CurrentWorkerInfo.WorkerAppId == 0 || (bCheckForWorkerRunning && !FPlatformProcess::IsApplicationRunning(CurrentWorkerInfo.WorkerAppId)))
+			if (CurrentWorkerInfo.WorkerAppId == 0 || (bCheckForWorkerRunning && !FShaderCompilingManager::IsShaderCompilerWorkerRunning(CurrentWorkerInfo.WorkerAppId)))
 			{
 				bool bLaunchAgain = true;
 
 				// Detect when the worker has exited due to fatal error
-				if (CurrentWorkerInfo.bLaunchedWorker && !FPlatformProcess::IsApplicationRunning(CurrentWorkerInfo.WorkerAppId))
+				// bLaunchedWorker check here is necessary to distinguish between 'process isn't running because it crashed' and 'process isn't running because it exited cleanly and the outputfile was already consumed'
+				if (CurrentWorkerInfo.WorkerAppId != 0 && CurrentWorkerInfo.bLaunchedWorker)
 				{
 					const FString WorkingDirectory = Manager->AbsoluteShaderBaseWorkingDirectory + FString::FromInt(WorkerIndex) + TEXT("/");
 					const FString OutputFileNameAndPath = WorkingDirectory + TEXT("WorkerOutputOnly.out");
@@ -760,7 +764,7 @@ bool FShaderCompileThreadRunnable::LaunchWorkersIfNeeded()
 					if (FPlatformFileManager::Get().GetPlatformFile().FileExists(*OutputFileNameAndPath))
 					{
 						// If the worker is no longer running but it successfully wrote out the output, no need to assert
-						bLaunchAgain = false;
+						bLaunchAgain = false;						
 					}
 					else
 					{
@@ -769,6 +773,9 @@ bool FShaderCompileThreadRunnable::LaunchWorkersIfNeeded()
 						bAbandonWorkers = true;
 						break;
 					}
+
+					// shader compiler exited one way or another, so clear out the stale PID.
+					CurrentWorkerInfo.WorkerAppId = 0;
 				}
 
 				if (bLaunchAgain)
@@ -778,7 +785,7 @@ bool FShaderCompileThreadRunnable::LaunchWorkersIfNeeded()
 					FString OutputFileName(TEXT("WorkerOutputOnly.out"));
 
 					// Store the Id with this thread so that we will know not to launch it again
-					CurrentWorkerInfo.WorkerAppId = Manager->LaunchWorker(WorkingDirectory, Manager->ProcessId, WorkerIndex, InputFileName, OutputFileName, false, false);
+					CurrentWorkerInfo.WorkerAppId = Manager->LaunchWorker(WorkingDirectory, Manager->ProcessId, WorkerIndex, InputFileName, OutputFileName, false, false);					
 					CurrentWorkerInfo.bLaunchedWorker = true;
 				}
 			}
@@ -1892,6 +1899,19 @@ void FShaderCompilingManager::ProcessAsyncResults(bool bLimitExecutionTime, bool
 	}
 }
 
+bool FShaderCompilingManager::IsShaderCompilerWorkerRunning(uint32 ProcessId)
+{
+	bool bRunning = false;
+	// make sure the PID we're inspecting is actually a shader compiler worker.
+	FString ProcessName = FPlatformProcess::GetApplicationName(ProcessId);
+
+	if (ProcessName.Find(FString(TEXT("ShaderCompileWorker"))) != -1)
+	{
+		bRunning = FPlatformProcess::IsApplicationRunning(ProcessId);
+	}
+	return bRunning;
+}
+
 /** Enqueues a shader compile job with GShaderCompilingManager. */
 void GlobalBeginCompileShader(
 	const FString& DebugGroupName,
diff --git a/Engine/Source/Runtime/Engine/Public/ShaderCompiler.h b/Engine/Source/Runtime/Engine/Public/ShaderCompiler.h
--- a/Engine/Source/Runtime/Engine/Public/ShaderCompiler.h
+++ b/Engine/Source/Runtime/Engine/Public/ShaderCompiler.h
@@ -304,6 +304,11 @@ public:
 	 *		This must be done before using global shaders for rendering.
 	 */
 	ENGINE_API void ProcessAsyncResults(bool bLimitExecutionTime, bool bBlockOnGlobalShaderCompletion);
+
+	/**
+	 * Returns true if the given PID is actually a ShaderCompilerWorker, and it's still running.
+	 */
+	static bool IsShaderCompilerWorkerRunning(uint32 ProcessId);
 };
 
 /** The global shader compiling thread manager. */
