Cleaned up the debug information

diff --git a/Engine/Source/Runtime/Engine/Private/Animation/AnimNode_WheelHandler.cpp b/Engine/Source/Runtime/Engine/Private/Animation/AnimNode_WheelHandler.cpp
--- a/Engine/Source/Runtime/Engine/Private/Animation/AnimNode_WheelHandler.cpp
+++ b/Engine/Source/Runtime/Engine/Private/Animation/AnimNode_WheelHandler.cpp
@@ -23,12 +23,12 @@ void FAnimNode_WheelHandler::GatherDebugData(FNodeDebugData& DebugData)
 	{
 		if (WheelSim.BoneReference.BoneIndex != INDEX_NONE)
 		{
-			DebugLine = FString::Printf(TEXT("\t[%d] Bone: %s , Rotation Offset : %s, Location Offset : %s"), 
+			DebugLine = FString::Printf(TEXT(" [Wheel Index : %d] Bone: %s , Rotation Offset : %s, Location Offset : %s"), 
 				WheelSim.WheelIndex, *WheelSim.BoneReference.BoneName.ToString(), *WheelSim.RotOffset.ToString(), *WheelSim.LocOffset.ToString());
 		}
 		else
 		{
-			DebugLine = FString::Printf(TEXT("\t[%d] Bone: %s (invalid bone)"),
+			DebugLine = FString::Printf(TEXT(" [Wheel Index : %d] Bone: %s (invalid bone)"),
 					WheelSim.WheelIndex, *WheelSim.BoneReference.BoneName.ToString());
 		}
 
