#ue4 - Added keywords for CharacterMovementComponent::GetCurrentAcceleration to make it easier to find in the BP editor.

diff --git a/Engine/Source/Runtime/Engine/Classes/GameFramework/CharacterMovementComponent.h b/Engine/Source/Runtime/Engine/Classes/GameFramework/CharacterMovementComponent.h
--- a/Engine/Source/Runtime/Engine/Classes/GameFramework/CharacterMovementComponent.h
+++ b/Engine/Source/Runtime/Engine/Classes/GameFramework/CharacterMovementComponent.h
@@ -753,7 +753,7 @@ public:
 	virtual float GetModifiedMaxAcceleration() const;
 
 	/** @return Current acceleration, computed from input vector each update. */
-	UFUNCTION(BlueprintCallable, Category="Pawn|Components|CharacterMovement")
+	UFUNCTION(BlueprintCallable, Category="Pawn|Components|CharacterMovement", meta=(Keywords="Acceleration GetAcceleration"))
 	FVector GetCurrentAcceleration() const;
 
 	/** @return Modifier [0..1] which affects max speed, based on the magnitude of the input vector. */
