Landscape Spline : Removed NULL case for control points/segments

diff --git a/Engine/Source/Runtime/Engine/Classes/Landscape/LandscapeSplinesComponent.h b/Engine/Source/Runtime/Engine/Classes/Landscape/LandscapeSplinesComponent.h
--- a/Engine/Source/Runtime/Engine/Classes/Landscape/LandscapeSplinesComponent.h
+++ b/Engine/Source/Runtime/Engine/Classes/Landscape/LandscapeSplinesComponent.h
@@ -46,6 +46,7 @@ public:
 #endif
 
 	// Begin UObject Interface
+	virtual void PostLoad() OVERRIDE;
 #if WITH_EDITOR
 	virtual void PostEditChangeProperty(FPropertyChangedEvent& PropertyChangedEvent) OVERRIDE;
 #endif // WITH_EDITOR
diff --git a/Engine/Source/Runtime/Engine/Private/Landscape/LandscapeSplines.cpp b/Engine/Source/Runtime/Engine/Private/Landscape/LandscapeSplines.cpp
--- a/Engine/Source/Runtime/Engine/Private/Landscape/LandscapeSplines.cpp
+++ b/Engine/Source/Runtime/Engine/Private/Landscape/LandscapeSplines.cpp
@@ -417,6 +417,35 @@ bool ULandscapeSplinesComponent::ModifySplines(bool bAlwaysMarkDirty /*= true*/)
 	return bSavedToTransactionBuffer;
 }
 
+void ULandscapeSplinesComponent::PostLoad()
+{
+	Super::PostLoad();
+
+	for (int32 iControlPoint = 0; iControlPoint < ControlPoints.Num(); )
+	{
+		if (ULandscapeSplineControlPoint* ControlPoint = ControlPoints[iControlPoint])
+		{
+			iControlPoint++;
+		}
+		else
+		{
+			ControlPoints.RemoveAt(iControlPoint);
+		}
+	}
+
+	for (int32 iSegment = 0; iSegment < Segments.Num(); )
+	{
+		if (ULandscapeSplineSegment* Segment = Segments[iSegment])
+		{
+			iSegment++;
+		}
+		else
+		{
+			Segments.RemoveAt(iSegment);
+		}
+	}
+}
+
 #if WITH_EDITOR
 void ULandscapeSplinesComponent::PostEditChangeProperty(FPropertyChangedEvent& PropertyChangedEvent)
 {
