fix for TTP #334607 "LIVE: EDITOR: NAVMESH: Navmesh does not update at runtime when meshes are moved, rotated, or scaled"

diff --git a/Engine/Source/Runtime/Engine/Classes/Components/PrimitiveComponent.h b/Engine/Source/Runtime/Engine/Classes/Components/PrimitiveComponent.h
--- a/Engine/Source/Runtime/Engine/Classes/Components/PrimitiveComponent.h
+++ b/Engine/Source/Runtime/Engine/Classes/Components/PrimitiveComponent.h
@@ -1166,7 +1166,8 @@ public:
 	// End UObject interface.
 
 	//Begin USceneComponent Interface
-	virtual bool MoveComponent( const FVector& Delta, const FRotator& NewRotation, bool bSweep, FHitResult* OutHit=NULL, EMoveComponentFlags MoveFlags = MOVECOMP_NoFlags ) OVERRIDE;
+	virtual void SetRelativeScale3D(FVector NewScale3D) OVERRIDE FINAL;
+	virtual bool MoveComponent(const FVector& Delta, const FRotator& NewRotation, bool bSweep, FHitResult* OutHit = NULL, EMoveComponentFlags MoveFlags = MOVECOMP_NoFlags) OVERRIDE;
 	virtual bool IsWorldGeometry() const OVERRIDE;
 	virtual ECollisionEnabled::Type GetCollisionEnabled() const OVERRIDE;
 	virtual ECollisionResponse GetCollisionResponseToChannel(ECollisionChannel Channel) const OVERRIDE;
diff --git a/Engine/Source/Runtime/Engine/Classes/Components/SceneComponent.h b/Engine/Source/Runtime/Engine/Classes/Components/SceneComponent.h
--- a/Engine/Source/Runtime/Engine/Classes/Components/SceneComponent.h
+++ b/Engine/Source/Runtime/Engine/Classes/Components/SceneComponent.h
@@ -223,7 +223,7 @@ public:
 
 	/** Set the non-uniform of this component relative to its parent */
 	UFUNCTION(BlueprintCallable, Category="Utilities|Orientation")
-	void SetRelativeScale3D(FVector NewScale3D);
+	virtual void SetRelativeScale3D(FVector NewScale3D);
 
 	/** Adds a delta to the translation of this component relative to its parent */
 	UFUNCTION(BlueprintCallable, Category="Utilities|Orientation")
diff --git a/Engine/Source/Runtime/Engine/Private/PrimitiveComponent.cpp b/Engine/Source/Runtime/Engine/Private/PrimitiveComponent.cpp
--- a/Engine/Source/Runtime/Engine/Private/PrimitiveComponent.cpp
+++ b/Engine/Source/Runtime/Engine/Private/PrimitiveComponent.cpp
@@ -1181,6 +1181,20 @@ FCollisionShape UPrimitiveComponent::GetCollisionShape(float Inflation) const
 	return FCollisionShape::MakeBox(Bounds.BoxExtent + Inflation);
 }
 
+void UPrimitiveComponent::SetRelativeScale3D(FVector NewScale3D)
+{
+	const FVector OldScale3D = RelativeScale3D;
+	Super::SetRelativeScale3D(NewScale3D);
+
+	AActor* Actor = GetOwner();
+	if (Actor && OldScale3D != RelativeScale3D && IsRegistered() && IsNavigationRelevant() && World != NULL && World->IsGameWorld() && World->GetNetMode() < ENetMode::NM_Client)
+	{
+		if (UNavigationSystem* NavSys = World->GetNavigationSystem())
+		{
+			NavSys->UpdateNavOctree(Actor);
+		}
+	}
+}
 
 bool UPrimitiveComponent::MoveComponent( const FVector& Delta, const FRotator& NewRotation, bool bSweep, FHitResult* OutHit, EMoveComponentFlags MoveFlags)
 {
@@ -1478,6 +1492,14 @@ bool UPrimitiveComponent::MoveComponent( const FVector& Delta, const FRotator& N
 			// still need to do this even if bGenerateOverlapEvents is false for this component, since we could have child components where it is true
 			UpdateOverlaps(&PendingOverlaps, true, OverlapsAtEndLocationPtr);
 		}
+
+		if (Actor && IsRegistered() && IsNavigationRelevant() && World != NULL && World->IsGameWorld() && World->GetNetMode() < ENetMode::NM_Client)
+		{
+			if (UNavigationSystem* NavSys = World->GetNavigationSystem())
+			{
+				NavSys->UpdateNavOctree(Actor);
+			}
+		}
 	}
 
 #if defined(PERF_SHOW_MOVECOMPONENT_TAKING_LONG_TIME) || LOOKING_FOR_PERF_ISSUES
