Fixed exploitable issues with ServerReplicateData frm GameDebuggingComponent.

diff --git a/Engine/Source/Runtime/Engine/Private/Debug/GameplayDebuggingComponent.cpp b/Engine/Source/Runtime/Engine/Private/Debug/GameplayDebuggingComponent.cpp
--- a/Engine/Source/Runtime/Engine/Private/Debug/GameplayDebuggingComponent.cpp
+++ b/Engine/Source/Runtime/Engine/Private/Debug/GameplayDebuggingComponent.cpp
@@ -194,7 +194,7 @@ bool UGameplayDebuggingComponent::ServerEnableTargetSelection_Validate(bool bEna
 void UGameplayDebuggingComponent::ServerEnableTargetSelection_Implementation(bool bEnable)
 {
 	bEnabledTargetSelection = bEnable;
-	if (bEnabledTargetSelection && GetWorld()->GetNetMode() < NM_Client)
+	if (bEnabledTargetSelection && World && World->GetNetMode() < NM_Client)
 	{
 		NextTargrtSelectionTime = 0;
 		SelectTargetToDebug();
@@ -205,7 +205,7 @@ void UGameplayDebuggingComponent::TickComponent(float DeltaTime, enum ELevelTick
 {
 	Super::TickComponent(DeltaTime, TickType, ThisTickFunction);
 
-	if (GetWorld()->GetNetMode() < NM_Client)
+	if (World && World->GetNetMode() < NM_Client)
 	{
 		if (bEnabledTargetSelection)
 		{
@@ -532,16 +532,18 @@ void UGameplayDebuggingComponent::ServerReplicateData_Implementation( EDebugComp
 		break;
 
 	case EDebugComponentMessage::ActivateDataView:
+		if (ReplicateViewDataCounters.IsValidIndex(DataView))
 		{
 			ReplicateViewDataCounters[DataView] += 1;
-			break;
 		}
+		break;
 
 	case EDebugComponentMessage::DeactivateDataView:
+		if (ReplicateViewDataCounters.IsValidIndex(DataView))
 		{
 			ReplicateViewDataCounters[DataView] = FMath::Max(0, ReplicateViewDataCounters[DataView] - 1);
-			break;
 		}
+		break;
 
 	default:
 		break;
@@ -616,7 +618,7 @@ FArchive& operator<<(FArchive& Ar, EQSDebug::FQueryData& Data)
 void UGameplayDebuggingComponent::OnRep_UpdateEQS()
 {
 	// decode scoring data
-	if (GetWorld()->GetNetMode() == NM_Client)
+	if (World && World->GetNetMode() == NM_Client)
 	{
 		TArray<uint8> UncompressedBuffer;
 		int32 UncompressedSize = 0;
@@ -746,7 +748,7 @@ void UGameplayDebuggingComponent::CollectEQSData()
 		1000.0f * (Timer2 - Timer1),
 		FMath::TruncToInt(100.0f * EQSRepData.Num() / UncompressedBuffer.Num()), 1000.0f * (Timer3 - Timer2));
 
-	if (GetWorld()->GetNetMode() != NM_DedicatedServer)
+	if (World && World->GetNetMode() != NM_DedicatedServer)
 	{
 		OnRep_UpdateEQS();
 	}
@@ -1033,7 +1035,7 @@ void UGameplayDebuggingComponent::ServerCollectNavmeshData_Implementation(FVecto
 		FMath::TruncToInt(100.0f * NavmeshRepData.Num() / UncompressedBuffer.Num()), 1000.0f * (Timer3 - Timer2));
 #endif
 
-	if (GetWorld()->GetNetMode() != NM_DedicatedServer)
+	if (World && World->GetNetMode() != NM_DedicatedServer)
 	{
 		OnRep_UpdateNavmesh();
 	}
