Character movement can now be affected by external forces

Added accumulator PendingMomentumToApply. This gets added to in calls to AddMomentum (now exposed to Blueprint) & applied each step.
Transferred logic for movement state change to where the momentum is applied, rather than doing it immediately on add.
Moved logic that constrained vertical velocity into ACharacter::ApplyDamageMomentum, as it was closely linked to damage application, rather than general forces.
Added new functions AddRadialForce and AddRadialImpulse. Updated Radial Force components to also apply to movement components.

TTP# 333553 - GF: CharacterMovement: investigate a good way to accumulate and apply external forces in CharacterMovement

diff --git a/Engine/Source/Runtime/Engine/Classes/GameFramework/CharacterMovementComponent.h b/Engine/Source/Runtime/Engine/Classes/GameFramework/CharacterMovementComponent.h
--- a/Engine/Source/Runtime/Engine/Classes/GameFramework/CharacterMovementComponent.h
+++ b/Engine/Source/Runtime/Engine/Classes/GameFramework/CharacterMovementComponent.h
@@ -381,6 +381,10 @@ protected:
 	UPROPERTY()
 	FVector Acceleration;
 
+	/** Accumulated momentum added this tick */
+	UPROPERTY()
+	FVector PendingMomentumToApply;
+
 	/**
 	 * Modifier to applied to values such as acceleration and max speed due to analog input.
 	 */
@@ -627,6 +631,8 @@ public:
 	virtual bool IsSwimming() const OVERRIDE;
 	virtual bool IsFlying() const OVERRIDE;
 	virtual float GetGravityZ() const OVERRIDE;
+	virtual void AddRadialForce(const FVector& Origin, float Radius, float Strength, enum ERadialImpulseFalloff Falloff) OVERRIDE;
+	virtual void AddRadialImpulse(const FVector& Origin, float Radius, float Strength, enum ERadialImpulseFalloff Falloff, bool bVelChange) OVERRIDE;
 	//END UMovementComponent Interface
 
 	//BEGIN UNavMovementComponent Interface
@@ -796,6 +802,9 @@ public:
 	/** Applies repulsion force to all touched components */
 	void ApplyRepulsionForce(float DeltaTime);
 	
+	/** Applies momentum accumulated through Addmomentum() */
+	void ApplyAccumulatedMomentum(float DeltaSeconds);	
+
 	/** 
 	 * Handle start swimming functionality
 	 * @param OldLocation - Location on last tick
@@ -882,8 +891,16 @@ public:
 	/** @Return MovementMode string */
 	virtual FString GetMovementName();
 
-	/** Add velocity based on imparted momentum */
-	virtual void AddMomentum( FVector const& Momentum, FVector const& LocationToApplyMomentum, bool bMassIndependent );
+	/** 
+	 * Add velocity based on imparted momentum. Momentum is accumulated each tick and applied together
+	 * so multiple calls to this function will accumulate.
+	 * Note that changing the momentum of characters like this can change the movement mode
+	 * 
+	 * @param	Momentum			Momentum to apply.
+	 * @param	bMassIndependent	Whether or not the momentum should be divided by mass before application.
+	 */
+	UFUNCTION(BlueprintCallable, Category="Pawn|Components|CharacterMovement")
+	virtual void AddMomentum( FVector InMomentum, bool bMassIndependent = false );
 
 	/**
 	 * Draw important variables on canvas.  Character will call DisplayDebug() on the current ViewTarget when the ShowDebug exec is used
diff --git a/Engine/Source/Runtime/Engine/Classes/GameFramework/MovementComponent.h b/Engine/Source/Runtime/Engine/Classes/GameFramework/MovementComponent.h
--- a/Engine/Source/Runtime/Engine/Classes/GameFramework/MovementComponent.h
+++ b/Engine/Source/Runtime/Engine/Classes/GameFramework/MovementComponent.h
@@ -214,6 +214,25 @@ public:
 	 */
 	virtual void TwoWallAdjust(FVector &Delta, const FHitResult& Hit, const FVector &OldHitNormal) const;
 
+	/**
+	 * Adds force from radial force components.
+	 * @param	Origin		The origin of the force
+	 * @param	Radius		The radius in which the force will be applied
+	 * @param	Strength	The strength of the force
+	 * @param	Falloff		The falloff from the force's origin
+	 */
+	virtual void AddRadialForce(const FVector& Origin, float Radius, float Strength, enum ERadialImpulseFalloff Falloff);
+
+	/**
+	 * Adds impulse from radial force components.
+	 * @param	Origin		The origin of the force
+	 * @param	Radius		The radius in which the force will be applied
+	 * @param	Strength	The strength of the force
+	 * @param	Falloff		The falloff from the force's origin
+	 * @param	bVelChange	If true, the Strength is taken as a change in velocity instead of an impulse (ie. mass will have no affect).
+	 */
+	virtual void AddRadialImpulse(const FVector& Origin, float Radius, float Strength, enum ERadialImpulseFalloff Falloff, bool bVelChange);
+
 	/**
 	 * Sets the normal of the plane that constrains movement, if plane constraint is enabled.
 	 * @param PlaneNormal	- The normal of the plane. If non-zero, it will be normalized.
diff --git a/Engine/Source/Runtime/Engine/Private/Character.cpp b/Engine/Source/Runtime/Engine/Private/Character.cpp
--- a/Engine/Source/Runtime/Engine/Private/Character.cpp
+++ b/Engine/Source/Runtime/Engine/Private/Character.cpp
@@ -327,7 +327,22 @@ void ACharacter::ApplyDamageMomentum(float DamageTaken, FDamageEvent const& Dama
 
 		FVector Impulse = ImpulseDir * ImpulseScale;
 		bool const bMassIndependentImpulse = !DmgTypeCDO->bScaleMomentumByMass;
-		CharacterMovement->AddMomentum(Impulse, HitInfo.ImpactPoint, bMassIndependentImpulse);
+
+		// limit Z momentum added if already going up faster than jump (to avoid blowing character way up into the sky)
+		{
+			FVector MassScaledImpulse = Impulse;
+			if(!bMassIndependentImpulse && CharacterMovement->Mass > SMALL_NUMBER)
+			{
+				MassScaledImpulse = MassScaledImpulse / CharacterMovement->Mass;
+			}
+
+			if ( (CharacterMovement->Velocity.Z > GetDefault<UCharacterMovementComponent>(CharacterMovement->GetClass())->JumpZVelocity) && (MassScaledImpulse.Z > 0.f) )
+			{
+				Impulse.Z *= 0.5f;
+			}
+		}
+
+		CharacterMovement->AddMomentum(Impulse, bMassIndependentImpulse);
 	}
 }
 
diff --git a/Engine/Source/Runtime/Engine/Private/CharacterMovementComponent.cpp b/Engine/Source/Runtime/Engine/Private/CharacterMovementComponent.cpp
--- a/Engine/Source/Runtime/Engine/Private/CharacterMovementComponent.cpp
+++ b/Engine/Source/Runtime/Engine/Private/CharacterMovementComponent.cpp
@@ -1163,6 +1163,8 @@ void UCharacterMovementComponent::PerformMovement(float DeltaSeconds)
 
 	UpdateBasedMovement(DeltaSeconds);
 
+	ApplyAccumulatedMomentum(DeltaSeconds);
+
 	FVector OldVelocity = Velocity;
 	FVector OldLocation = CharacterOwner->GetActorLocation();
 
@@ -1852,7 +1854,7 @@ void UCharacterMovementComponent::CalcVelocity(float DeltaTime, float Friction,
 			ApplyVelocityBraking(DeltaTime, Friction, BrakingDeceleration);
 	
 			// Don't allow braking to lower us below max speed if we started above it.
-			if (bVelocityOverMax && Velocity.SizeSquared() < FMath::Square(MaxSpeed))
+			if (FVector::DotProduct(Acceleration, OldVelocity) > 0.0f && bVelocityOverMax && Velocity.SizeSquared() < FMath::Square(MaxSpeed))
 			{
 				Velocity = SafeNormalPrecise(OldVelocity) * MaxSpeed;
 			}
@@ -3433,12 +3435,12 @@ bool UCharacterMovementComponent::CheckWaterJump(FVector CheckPoint, FVector& Wa
 	return false;
 }
 
-void UCharacterMovementComponent::AddMomentum( FVector const& Momentum, FVector const& LocationToApplyMomentum, bool bMassIndependent )
+void UCharacterMovementComponent::AddMomentum( FVector InMomentum, bool bMassIndependent )
 {
-	if (HasValidData() && !Momentum.IsZero())
+	if (HasValidData() && !InMomentum.IsZero())
 	{
 		// handle scaling by mass
-		FVector FinalMomentum = Momentum;
+		FVector FinalMomentum = InMomentum;
 		if ( !bMassIndependent)
 		{
 			if (Mass > SMALL_NUMBER)
@@ -3451,18 +3453,7 @@ void UCharacterMovementComponent::AddMomentum( FVector const& Momentum, FVector
 			}
 		}
 
-		if ( IsMovingOnGround() && FinalMomentum.Z > 0.f )
-		{
-			SetMovementMode(MOVE_Falling);
-		}
-
-		if ( (Velocity.Z > GetDefault<UCharacterMovementComponent>(GetClass())->JumpZVelocity) && (FinalMomentum.Z > 0.f) )
-		{
-			// limit Z momentum added if already going up faster than jump (to avoid blowing character way up into the sky)
-			FinalMomentum.Z *= 0.5f;
-		}
-
-		Velocity += FinalMomentum;
+		PendingMomentumToApply += FinalMomentum;
 	}
 }
 
@@ -5544,6 +5535,46 @@ void UCharacterMovementComponent::ApplyRepulsionForce( float DeltaTime )
 	}
 }
 
+void UCharacterMovementComponent::ApplyAccumulatedMomentum(float DeltaSeconds)
+{
+	// check to see if applied momentum is enough to overcome gravity
+	if ( IsMovingOnGround() && (PendingMomentumToApply.Z + (GetGravityZ() * DeltaSeconds) > SMALL_NUMBER) )
+	{
+		SetMovementMode(MOVE_Falling);
+	}
+
+	Velocity += PendingMomentumToApply;
+
+	PendingMomentumToApply = FVector::ZeroVector;
+}
+
+void UCharacterMovementComponent::AddRadialForce(const FVector& Origin, float Radius, float Strength, enum ERadialImpulseFalloff Falloff)
+{
+	AddRadialImpulse(Origin, Radius, Strength, Falloff, true);
+}
+ 
+void UCharacterMovementComponent::AddRadialImpulse(const FVector& Origin, float Radius, float Strength, enum ERadialImpulseFalloff Falloff, bool bVelChange)
+{
+	FVector Delta = UpdatedComponent->GetComponentLocation() - Origin;
+	const float DeltaMagnitude = Delta.Size();
+
+	// Do nothing if outside radius
+	if(DeltaMagnitude > Radius)
+	{
+		return;
+	}
+
+	Delta = Delta.SafeNormal();
+
+	float ImpulseMagnitude = Strength;
+	if (Falloff == RIF_Linear && Radius > 0.0f)
+	{
+		ImpulseMagnitude *= (1.0f - (DeltaMagnitude / Radius));
+	}
+
+	AddMomentum(Delta * ImpulseMagnitude, bVelChange);
+}
+
 FNetworkPredictionData_Client_Character::FNetworkPredictionData_Client_Character()
 	: ClientUpdateTime(0.f)
 	, CurrentTimeStamp(0.f)
diff --git a/Engine/Source/Runtime/Engine/Private/MovementComponent.cpp b/Engine/Source/Runtime/Engine/Private/MovementComponent.cpp
--- a/Engine/Source/Runtime/Engine/Private/MovementComponent.cpp
+++ b/Engine/Source/Runtime/Engine/Private/MovementComponent.cpp
@@ -453,5 +453,14 @@ void UMovementComponent::TwoWallAdjust(FVector& OutDelta, const FHitResult& Hit,
 	OutDelta = Delta;
 }
 
+void UMovementComponent::AddRadialForce(const FVector& Origin, float Radius, float Strength, enum ERadialImpulseFalloff Falloff)
+{
+	// Default implementation does nothing
+}
+
+void UMovementComponent::AddRadialImpulse(const FVector& Origin, float Radius, float Strength, enum ERadialImpulseFalloff Falloff, bool bVelChange)
+{
+	// Default implementation does nothing
+}
 
 #undef LOCTEXT_NAMESPACE
diff --git a/Engine/Source/Runtime/Engine/Private/PhysicsEngine/RadialForceComponent.cpp b/Engine/Source/Runtime/Engine/Private/PhysicsEngine/RadialForceComponent.cpp
--- a/Engine/Source/Runtime/Engine/Private/PhysicsEngine/RadialForceComponent.cpp
+++ b/Engine/Source/Runtime/Engine/Private/PhysicsEngine/RadialForceComponent.cpp
@@ -39,7 +39,19 @@ void URadialForceComponent::TickComponent(float DeltaTime, enum ELevelTick TickT
 			UPrimitiveComponent* PokeComp = Overlaps[OverlapIdx].Component.Get();
 			if(PokeComp)
 			{
-				PokeComp->AddRadialForce( Origin, Radius, ForceStrength, Falloff );				
+				PokeComp->AddRadialForce( Origin, Radius, ForceStrength, Falloff );
+
+				// see if this is a target for a movement component
+				TArray<UMovementComponent*> MovementComponents;
+				PokeComp->GetOwner()->GetComponents<UMovementComponent>(MovementComponents);
+				for(const auto& MovementComponent : MovementComponents)
+				{
+					if(MovementComponent->UpdatedComponent == PokeComp)
+					{
+						MovementComponent->AddRadialForce( Origin, Radius, ForceStrength, Falloff );
+						break;
+					}
+				}
 			}
 		}
 	}
@@ -85,6 +97,18 @@ void URadialForceComponent::FireImpulse()
 
 			// Do impulse after
 			PokeComp->AddRadialImpulse( Origin, Radius, ImpulseStrength, Falloff, bImpulseVelChange );
+
+			// see if this is a target for a movement component
+			TArray<UMovementComponent*> MovementComponents;
+			PokeComp->GetOwner()->GetComponents<UMovementComponent>(MovementComponents);
+			for(const auto& MovementComponent : MovementComponents)
+			{
+				if(MovementComponent->UpdatedComponent == PokeComp)
+				{
+					MovementComponent->AddRadialImpulse( Origin, Radius, ImpulseStrength, Falloff, bImpulseVelChange );
+					break;
+				}
+			}
 		}
 	}
 }
