Added option for 'flat bottomed' collision to character components.

When performing ground sweeps, we substitute a box (tightly containing the capsule) for the usual capsule sweeps.

TTP# 333862 - Option for box-shaped player collision

diff --git a/Engine/Source/Runtime/Engine/Classes/GameFramework/CharacterMovementComponent.h b/Engine/Source/Runtime/Engine/Classes/GameFramework/CharacterMovementComponent.h
--- a/Engine/Source/Runtime/Engine/Classes/GameFramework/CharacterMovementComponent.h
+++ b/Engine/Source/Runtime/Engine/Classes/GameFramework/CharacterMovementComponent.h
@@ -516,6 +516,15 @@ public:
 	UPROPERTY(Category="Character Movement", EditAnywhere, BlueprintReadWrite, AdvancedDisplay)
 	uint32 bAlwaysCheckFloor:1;
 
+	/**
+	 * Performs floor checks as if the character is using a box collider.
+	 * This avoids the situation where characters slowly slide down a ledge (as their capsule 'balances' on the edge).
+	 * Note that sweep checks perfomed using the box collision are axis-aligned and use a bounding box that tightly 
+	 * encloses the capsule.
+	 */
+	UPROPERTY(Category="Character Movement", EditAnywhere, BlueprintReadWrite, AdvancedDisplay)
+	uint32 bUseFlatBaseForFloorChecks:1;
+
 	/** Used to prevent reentry of JumpOff() */
 	UPROPERTY()
 	uint32 bPerformingJumpOff:1;
diff --git a/Engine/Source/Runtime/Engine/Private/CharacterMovementComponent.cpp b/Engine/Source/Runtime/Engine/Private/CharacterMovementComponent.cpp
--- a/Engine/Source/Runtime/Engine/Private/CharacterMovementComponent.cpp
+++ b/Engine/Source/Runtime/Engine/Private/CharacterMovementComponent.cpp
@@ -2412,7 +2412,19 @@ void UCharacterMovementComponent::PhysFalling(float deltaTime, int32 Iterations)
 				InitCollisionParams(CapsuleQuery, ResponseParam);
 				const FVector PawnLocation = CharacterOwner->GetActorLocation();
 				const ECollisionChannel CollisionChannel = UpdatedComponent->GetCollisionObjectType();
-				const bool bHit = GetWorld()->SweepSingle( Result, PawnLocation, PawnLocation + TestWalk, FQuat::Identity, CollisionChannel, GetPawnCapsuleCollisionShape(SHRINK_None), CapsuleQuery, ResponseParam );
+
+				bool bHit = false;
+				const FCollisionShape CapsuleShape = GetPawnCapsuleCollisionShape(SHRINK_None);
+				if(bUseFlatBaseForFloorChecks)
+				{
+					const FCollisionShape BoxShape = FCollisionShape::MakeBox(FVector(CapsuleShape.GetCapsuleRadius(), CapsuleShape.GetCapsuleRadius(), CapsuleShape.GetCapsuleHalfHeight()));
+					GetWorld()->SweepSingle( Result, PawnLocation, PawnLocation + TestWalk, FQuat::Identity, CollisionChannel, BoxShape, CapsuleQuery, ResponseParam );
+				}
+				else
+				{
+					GetWorld()->SweepSingle( Result, PawnLocation, PawnLocation + TestWalk, FQuat::Identity, CollisionChannel, CapsuleShape, CapsuleQuery, ResponseParam );
+				}
+
 				if (bHit)
 				{
 					// Only matters if we can't walk there
@@ -3639,7 +3651,15 @@ void UCharacterMovementComponent::ComputeFloorDist(const FVector& CapsuleLocatio
 		FCollisionShape CapsuleShape = FCollisionShape::MakeCapsule(SweepRadius, PawnHalfHeight - ShrinkHeight);
 
 		FHitResult Hit(1.f);
-		bBlockingHit = GetWorld()->SweepSingle(Hit, CapsuleLocation, CapsuleLocation + FVector(0.f,0.f,-TraceDist), FQuat::Identity, CollisionChannel, CapsuleShape, QueryParams, ResponseParam);
+		if(bUseFlatBaseForFloorChecks)
+		{
+			const FCollisionShape BoxShape = FCollisionShape::MakeBox(FVector(CapsuleShape.GetCapsuleRadius(), CapsuleShape.GetCapsuleRadius(), CapsuleShape.GetCapsuleHalfHeight()));
+			bBlockingHit = GetWorld()->SweepSingle(Hit, CapsuleLocation, CapsuleLocation + FVector(0.f,0.f,-TraceDist), FQuat::Identity, CollisionChannel, BoxShape, QueryParams, ResponseParam);
+		}
+		else
+		{
+			bBlockingHit = GetWorld()->SweepSingle(Hit, CapsuleLocation, CapsuleLocation + FVector(0.f,0.f,-TraceDist), FQuat::Identity, CollisionChannel, CapsuleShape, QueryParams, ResponseParam);
+		}
 
 		if (bBlockingHit)
 		{
@@ -3653,7 +3673,16 @@ void UCharacterMovementComponent::ComputeFloorDist(const FVector& CapsuleLocatio
 				CapsuleShape.Capsule.Radius = FMath::Max(0.f, CapsuleShape.Capsule.Radius - SWEEP_EDGE_REJECT_DISTANCE - KINDA_SMALL_NUMBER);
 				CapsuleShape.Capsule.HalfHeight = FMath::Max(PawnHalfHeight - ShrinkHeight, 0.1f);
 				Hit.Reset(1.f, false);
-				bBlockingHit = GetWorld()->SweepSingle(Hit, CapsuleLocation, CapsuleLocation + FVector(0.f,0.f,-TraceDist), FQuat::Identity, CollisionChannel, CapsuleShape, QueryParams, ResponseParam);
+
+				if(bUseFlatBaseForFloorChecks)
+				{
+					const FCollisionShape BoxShape = FCollisionShape::MakeBox(FVector(CapsuleShape.GetCapsuleRadius(), CapsuleShape.GetCapsuleRadius(), CapsuleShape.GetCapsuleHalfHeight()));
+					bBlockingHit = GetWorld()->SweepSingle(Hit, CapsuleLocation, CapsuleLocation + FVector(0.f,0.f,-TraceDist), FQuat::Identity, CollisionChannel, BoxShape, QueryParams, ResponseParam);
+				}
+				else
+				{
+					bBlockingHit = GetWorld()->SweepSingle(Hit, CapsuleLocation, CapsuleLocation + FVector(0.f,0.f,-TraceDist), FQuat::Identity, CollisionChannel, CapsuleShape, QueryParams, ResponseParam);
+				}
 			}
 
 			// Reduce hit distance by ShrinkHeight because we shrank the capsule for the trace.
