Make Texture Streaming use the Platform's actual texture size, rather than a cross-platform alignment and padding agnostic size.

diff --git a/Engine/Source/Runtime/Engine/Classes/Engine/Texture2D.h b/Engine/Source/Runtime/Engine/Classes/Engine/Texture2D.h
--- a/Engine/Source/Runtime/Engine/Classes/Engine/Texture2D.h
+++ b/Engine/Source/Runtime/Engine/Classes/Engine/Texture2D.h
@@ -281,7 +281,7 @@ public:
 	 * @param	Enum	Which mips to calculate size for.
 	 * @return	Total size of all specified mips, in bytes
 	 */
-	virtual uint32 CalcTextureMemorySizeEnum( ETextureMipCount Enum ) const;
+	virtual uint32 CalcTextureMemorySizeEnum( ETextureMipCount Enum ) const OVERRIDE;
 
 	/**
 	 *	Get the CRC of the source art pixels.
diff --git a/Engine/Source/Runtime/Engine/Classes/Engine/TextureCube.h b/Engine/Source/Runtime/Engine/Classes/Engine/TextureCube.h
--- a/Engine/Source/Runtime/Engine/Classes/Engine/TextureCube.h
+++ b/Engine/Source/Runtime/Engine/Classes/Engine/TextureCube.h
@@ -80,7 +80,7 @@ public:
 	 * @param	Enum	Which mips to calculate size for.
 	 * @return	Total size of all specified mips, in bytes
 	 */
-	virtual uint32 CalcTextureMemorySizeEnum( ETextureMipCount Enum ) const;
+	virtual uint32 CalcTextureMemorySizeEnum( ETextureMipCount Enum ) const OVERRIDE;
 };
 
 
diff --git a/Engine/Source/Runtime/Engine/Private/ContentStreaming.cpp b/Engine/Source/Runtime/Engine/Private/ContentStreaming.cpp
--- a/Engine/Source/Runtime/Engine/Private/ContentStreaming.cpp
+++ b/Engine/Source/Runtime/Engine/Private/ContentStreaming.cpp
@@ -4801,6 +4801,7 @@ void FStreamingManagerTexture::DumpTextureGroupStats( bool bDetailedStats )
 		FTextureGroupStats& Stat = TextureGroupStats[Texture->LODGroup];
 		FTextureGroupStats& Waste = TextureGroupWaste[Texture->LODGroup];
 		UTexture2D* Texture2D = Cast<UTexture2D>(Texture);
+		uint32 TextureAlign = 0;
 		if ( Texture2D && Texture2D->StreamingIndex >= 0 )
 		{
 			FStreamingTexture& StreamingTexture = GetStreamingTexture( Texture2D );
@@ -4808,9 +4809,13 @@ void FStreamingManagerTexture::DumpTextureGroupStats( bool bDetailedStats )
 			Stat.CurrentTextureSize += StreamingTexture.GetSize( StreamingTexture.ResidentMips );
 			Stat.WantedTextureSize += StreamingTexture.GetSize( StreamingTexture.WantedMips );
 			Stat.MaxTextureSize += StreamingTexture.GetSize( StreamingTexture.MaxAllowedMips );
-			int32 WasteCurrent = StreamingTexture.GetSize( StreamingTexture.ResidentMips ) - CalcTextureSize( Texture2D->GetSizeX(), Texture2D->GetSizeY(), Texture2D->GetPixelFormat(), StreamingTexture.ResidentMips );
-			int32 WasteWanted = StreamingTexture.GetSize( StreamingTexture.WantedMips ) - CalcTextureSize( Texture2D->GetSizeX(), Texture2D->GetSizeY(), Texture2D->GetPixelFormat(), StreamingTexture.WantedMips );
-			int32 WasteMaxSize = StreamingTexture.GetSize( StreamingTexture.MaxAllowedMips ) - CalcTextureSize( Texture2D->GetSizeX(), Texture2D->GetSizeY(), Texture2D->GetPixelFormat(), StreamingTexture.MaxAllowedMips );
+			
+			int32 WasteCurrent = StreamingTexture.GetSize( StreamingTexture.ResidentMips ) - RHICalcTexture2DPlatformSize(Texture2D->GetSizeX(), Texture2D->GetSizeY(), Texture2D->GetPixelFormat(), StreamingTexture.ResidentMips, 1, 0, TextureAlign);			
+
+			int32 WasteWanted = StreamingTexture.GetSize( StreamingTexture.WantedMips ) - RHICalcTexture2DPlatformSize(Texture2D->GetSizeX(), Texture2D->GetSizeY(), Texture2D->GetPixelFormat(), StreamingTexture.WantedMips, 1, 0, TextureAlign);			
+
+			int32 WasteMaxSize = StreamingTexture.GetSize( StreamingTexture.MaxAllowedMips ) - RHICalcTexture2DPlatformSize(Texture2D->GetSizeX(), Texture2D->GetSizeY(), Texture2D->GetPixelFormat(), StreamingTexture.MaxAllowedMips, 1, 0, TextureAlign);			
+
 			Waste.NumTextures++;
 			Waste.CurrentTextureSize += FMath::Max(WasteCurrent,0);
 			Waste.WantedTextureSize += FMath::Max(WasteWanted,0);
@@ -4823,8 +4828,9 @@ void FStreamingManagerTexture::DumpTextureGroupStats( bool bDetailedStats )
 			Stat.NumNonStreamingTextures++;
 			Stat.NonStreamingSize += TextureSize;
 			if ( Texture2D && Texture2D->Resource )
-			{
-				int32 WastedSize = TextureSize - CalcTextureSize( Texture2D->GetSizeX(), Texture2D->GetSizeY(), Texture2D->GetPixelFormat(), Texture2D->GetNumMips() );
+			{				
+				int32 WastedSize = TextureSize - RHICalcTexture2DPlatformSize(Texture2D->GetSizeX(), Texture2D->GetSizeY(), Texture2D->GetPixelFormat(), Texture2D->GetNumMips(), 1, 0, TextureAlign);				
+
 				Waste.NumNonStreamingTextures++;
 				Waste.NonStreamingSize += FMath::Max(WastedSize, 0);
 			}
diff --git a/Engine/Source/Runtime/Engine/Private/Texture2D.cpp b/Engine/Source/Runtime/Engine/Private/Texture2D.cpp
--- a/Engine/Source/Runtime/Engine/Private/Texture2D.cpp
+++ b/Engine/Source/Runtime/Engine/Private/Texture2D.cpp
@@ -694,12 +694,12 @@ int32 UTexture2D::CalcTextureMemorySize( int32 MipCount ) const
 		EPixelFormat Format = GetPixelFormat();
 
 		// Figure out what the first mip to use is.
-		int32 FirstMip	= FMath::Max( 0, NumMips - MipCount );
-		// Iterate over all relevant miplevels and sum up their size.
-		for( int32 MipIndex=FirstMip; MipIndex<NumMips; MipIndex++ )
-		{
-			Size += CalcTextureMipMapSize(SizeX, SizeY, Format, MipIndex);
-		}
+		int32 FirstMip	= FMath::Max( 0, NumMips - MipCount );		
+		FIntPoint MipExtents = CalcMipMapExtent(SizeX, SizeY, Format, FirstMip);
+
+		uint32 TextureAlign = 0;
+		uint64 TextureSize = RHICalcTexture2DPlatformSize(MipExtents.X, MipExtents.Y, Format, MipCount, 1, 0, TextureAlign);
+		Size = (int32)TextureSize;
 	}
 	return Size;
 }
diff --git a/Engine/Source/Runtime/Engine/Private/TextureCube.cpp b/Engine/Source/Runtime/Engine/Private/TextureCube.cpp
--- a/Engine/Source/Runtime/Engine/Private/TextureCube.cpp
+++ b/Engine/Source/Runtime/Engine/Private/TextureCube.cpp
@@ -115,16 +115,17 @@ uint32 UTextureCube::CalcTextureMemorySize( int32 MipCount ) const
 		int32 NumMips = GetNumMips();
 		EPixelFormat Format = GetPixelFormat();
 
+		ensureMsgf(SizeX == SizeY, TEXT("Cubemap faces expected to be square.  Actual sizes are: %i, %i"), SizeX, SizeY);
+
 		// Figure out what the first mip to use is.
-		int32 FirstMip	= FMath::Max( 0, NumMips - MipCount );
-		// Iterate over all relevant miplevels and sum up their size.
-		for( int32 MipIndex=FirstMip; MipIndex<NumMips; MipIndex++ )
-		{
-			Size += CalcTextureMipMapSize(SizeX, SizeY, Format, MipIndex);
-		}
+		int32 FirstMip	= FMath::Max( 0, NumMips - MipCount );		
+		FIntPoint MipExtents = CalcMipMapExtent(SizeX, SizeY, Format, FirstMip);
+		
+		uint32 TextureAlign = 0;
+		uint64 TextureSize = RHICalcTextureCubePlatformSize(MipExtents.X, Format, NumMips, 0, TextureAlign);
+		Size = (uint32)TextureSize;
 	}
-	// a cubemap has 6 sides
-	return Size * 6;
+	return Size;
 }
 
 
diff --git a/Engine/Source/Runtime/OpenGLDrv/Private/OpenGLTexture.cpp b/Engine/Source/Runtime/OpenGLDrv/Private/OpenGLTexture.cpp
--- a/Engine/Source/Runtime/OpenGLDrv/Private/OpenGLTexture.cpp
+++ b/Engine/Source/Runtime/OpenGLDrv/Private/OpenGLTexture.cpp
@@ -173,6 +173,24 @@ void OpenGLTextureDeleted( FRHITexture* Texture )
 	}
 }
 
+uint64 FOpenGLDynamicRHI::RHICalcTexture2DPlatformSize(uint32 SizeX, uint32 SizeY, uint8 Format, uint32 NumMips, uint32 NumSamples, uint32 Flags, uint32& OutAlign)
+{
+	OutAlign = 0;
+	return CalcTextureSize(SizeX, SizeY, (EPixelFormat)Format, NumMips);
+}
+
+uint64 FOpenGLDynamicRHI::RHICalcTexture3DPlatformSize(uint32 SizeX, uint32 SizeY, uint32 SizeZ, uint8 Format, uint32 NumMips, uint32 Flags, uint32& OutAlign)
+{
+	OutAlign = 0;
+	return CalcTextureSize3D(SizeX, SizeY, SizeZ, (EPixelFormat)Format, NumMips);
+}
+
+uint64 FOpenGLDynamicRHI::RHICalcTextureCubePlatformSize(uint32 Size, uint8 Format, uint32 NumMips, uint32 Flags,	uint32& OutAlign)
+{
+	OutAlign = 0;
+	return CalcTextureSize(Size, Size, (EPixelFormat)Format, NumMips) * 6;
+}
+
 /**
  * Retrieves texture memory stats. Unsupported with this allocator.
  *
diff --git a/Engine/Source/Runtime/RHI/Public/RHIMethods.h b/Engine/Source/Runtime/RHI/Public/RHIMethods.h
--- a/Engine/Source/Runtime/RHI/Public/RHIMethods.h
+++ b/Engine/Source/Runtime/RHI/Public/RHIMethods.h
@@ -364,6 +364,66 @@ DEFINE_RHIMETHOD_2(
 	,
 	);
 
+/**
+ * Computes the total size of a 2D texture with the specified parameters.
+ *
+ * @param SizeX - width of the texture to compute
+ * @param SizeY - height of the texture to compute
+ * @param Format - EPixelFormat texture format
+ * @param NumMips - number of mips to compute or 0 for full mip pyramid
+ * @param NumSamples - number of MSAA samples, usually 1
+ * @param Flags - ETextureCreateFlags creation flags
+ * @param OutAlign - Alignment required for this texture.  Output parameter.
+ */
+DEFINE_RHIMETHOD_7(
+	uint64,	RHICalcTexture2DPlatformSize,
+	uint32, SizeX,
+	uint32, SizeY,	
+	uint8, Format,
+	uint32, NumMips,
+	uint32, NumSamples,
+	uint32, Flags,
+	uint32&, OutAlign,
+	return, OutAlign = 0; return 0);
+
+/**
+* Computes the total size of a 3D texture with the specified parameters.
+* @param SizeX - width of the texture to create
+* @param SizeY - height of the texture to create
+* @param SizeZ - depth of the texture to create
+* @param Format - EPixelFormat texture format
+* @param NumMips - number of mips to generate or 0 for full mip pyramid
+* @param Flags - ETextureCreateFlags creation flags
+* @param OutAlign - Alignment required for this texture.  Output parameter.
+*/
+DEFINE_RHIMETHOD_7(
+	uint64,RHICalcTexture3DPlatformSize,
+	uint32,SizeX,
+	uint32,SizeY,
+	uint32,SizeZ,
+	uint8,Format,
+	uint32,NumMips,
+	uint32,Flags,	
+	uint32&, OutAlign,
+	return, OutAlign = 0; return 0);
+
+/**
+* Computes the total size of a cube texture with the specified parameters.
+* @param Size - width/height of the texture to create
+* @param Format - EPixelFormat texture format
+* @param NumMips - number of mips to generate or 0 for full mip pyramid
+* @param Flags - ETextureCreateFlags creation flags
+* @param OutAlign - Alignment required for this texture.  Output parameter.
+*/
+DEFINE_RHIMETHOD_5(
+	uint64,RHICalcTextureCubePlatformSize,
+	uint32,Size,
+	uint8,Format,
+	uint32,NumMips,
+	uint32,Flags,	
+	uint32&, OutAlign,
+	return, OutAlign = 0; return 0);
+
 /**
  * Retrieves texture memory stats.
  * safe to call on the main thread
@@ -754,6 +814,7 @@ DEFINE_RHIMETHOD_5(
 		return,return new FRHITextureCube(Size,NumMips,(EPixelFormat)Format,Flags);
 	);
 
+
 /**
 * Locks an RHI texture's mip-map for read/write operations on the CPU
 * @param Texture - the RHI texture resource to lock
diff --git a/Engine/Source/Runtime/Windows/D3D11RHI/Private/D3D11Texture.cpp b/Engine/Source/Runtime/Windows/D3D11RHI/Private/D3D11Texture.cpp
--- a/Engine/Source/Runtime/Windows/D3D11RHI/Private/D3D11Texture.cpp
+++ b/Engine/Source/Runtime/Windows/D3D11RHI/Private/D3D11Texture.cpp
@@ -208,6 +208,24 @@ FD3D11Texture3D::~FD3D11Texture3D()
 	D3D11TextureDeleted( *this );
 }
 
+uint64 FD3D11DynamicRHI::RHICalcTexture2DPlatformSize(uint32 SizeX, uint32 SizeY, uint8 Format, uint32 NumMips, uint32 NumSamples, uint32 Flags, uint32& OutAlign)
+{
+	OutAlign = 0;
+	return CalcTextureSize(SizeX, SizeY, (EPixelFormat)Format, NumMips);
+}
+
+uint64 FD3D11DynamicRHI::RHICalcTexture3DPlatformSize(uint32 SizeX, uint32 SizeY, uint32 SizeZ, uint8 Format, uint32 NumMips, uint32 Flags, uint32& OutAlign)
+{
+	OutAlign = 0;
+	return CalcTextureSize3D(SizeX, SizeY, SizeZ, (EPixelFormat)Format, NumMips);
+}
+
+uint64 FD3D11DynamicRHI::RHICalcTextureCubePlatformSize(uint32 Size, uint8 Format, uint32 NumMips, uint32 Flags,	uint32& OutAlign)
+{
+	OutAlign = 0;
+	return CalcTextureSize(Size, Size, (EPixelFormat)Format, NumMips) * 6;
+}
+
 /**
  * Retrieves texture memory stats. 
  */
