Move Tags in to their own category so hiding Actor category does not prevent tags from being used.
Also made them not an Advanced Display since they are now only thing in their category so that seems weird.

diff --git a/Engine/Source/Runtime/Engine/Classes/GameFramework/Actor.h b/Engine/Source/Runtime/Engine/Classes/GameFramework/Actor.h
--- a/Engine/Source/Runtime/Engine/Classes/GameFramework/Actor.h
+++ b/Engine/Source/Runtime/Engine/Classes/GameFramework/Actor.h
@@ -485,7 +485,7 @@ public:
 	uint32 bIgnoresOriginShifting:1;
 	
 	/** Array of tags that can be used for grouping and categorizing. Can also be accessed from scripting */
-	UPROPERTY(EditAnywhere, AdvancedDisplay, BlueprintReadWrite, Category=Actor)
+	UPROPERTY(EditAnywhere, BlueprintReadWrite, Category=Tags)
 	TArray<FName> Tags;
 
 	/** Bitflag to represent which views this actor is hidden in, via per-view layer visibility */
