Fixed issue causing objects with LOD Index of -1 to not render (Brushes / Landscapes), test refactored to allow checking first mesh instead of defaulting to LOD0.
#ttp 334838 - BLOCKER: CRITICAL: Regression: Landscape: Not rendering unless selected in scene outliner
#ttp 334824 - BLOCKER: CRITICAL: REGRESSION: EDITOR: Placing a BSP Into the level will only render Wireframe

diff --git a/Engine/Source/Runtime/Engine/Private/Landscape/LandscapeRender.cpp b/Engine/Source/Runtime/Engine/Private/Landscape/LandscapeRender.cpp
--- a/Engine/Source/Runtime/Engine/Private/Landscape/LandscapeRender.cpp
+++ b/Engine/Source/Runtime/Engine/Private/Landscape/LandscapeRender.cpp
@@ -1104,7 +1104,7 @@ void FLandscapeComponentSceneProxy::DrawStaticElements(FStaticPrimitiveDrawInter
 		BatchElement->MaxVertexIndex = SharedBuffers->IndexRanges[LOD].MaxIndexFull;
 	}
 
-	PDI->DrawMesh(MeshBatch, 0.0f);
+	PDI->DrawMesh(MeshBatch, FLT_MAX);
 }
 
 uint64 FLandscapeVertexFactory::GetStaticBatchElementVisibility( const class FSceneView& View, const struct FMeshBatch* Batch ) const
diff --git a/Engine/Source/Runtime/Engine/Private/ModelRender.cpp b/Engine/Source/Runtime/Engine/Private/ModelRender.cpp
--- a/Engine/Source/Runtime/Engine/Private/ModelRender.cpp
+++ b/Engine/Source/Runtime/Engine/Private/ModelRender.cpp
@@ -453,7 +453,7 @@ public:
 					BatchElement.MaxVertexIndex = ModelElement.MaxVertexIndex;
 					MeshElement.Type = PT_TriangleList;
 					MeshElement.DepthPriorityGroup = PrimitiveDPG;
-					PDI->DrawMesh(MeshElement, 0.0f);
+					PDI->DrawMesh(MeshElement, FLT_MAX);
 				}
 			}
 		}
diff --git a/Engine/Source/Runtime/Engine/Private/SceneManagement.cpp b/Engine/Source/Runtime/Engine/Private/SceneManagement.cpp
--- a/Engine/Source/Runtime/Engine/Private/SceneManagement.cpp
+++ b/Engine/Source/Runtime/Engine/Private/SceneManagement.cpp
@@ -135,7 +135,7 @@ int8 ComputeLODForMeshes( const TIndirectArray<class FStaticMesh>& StaticMeshes,
 
 		const float ScreenSize = ComputeBoundsScreenSize(Origin, SphereRadius, View);
 
-		for(int32 MeshIndex = NumMeshes-1 ; MeshIndex > 0 ; --MeshIndex)
+		for(int32 MeshIndex = NumMeshes-1 ; MeshIndex >= 0 ; --MeshIndex)
 		{
 			const FStaticMesh& Mesh = StaticMeshes[MeshIndex];
 
@@ -143,7 +143,7 @@ int8 ComputeLODForMeshes( const TIndirectArray<class FStaticMesh>& StaticMeshes,
 			{
 				float MeshScreenSize = Mesh.ScreenSize * ScreenSizeScale;
 
-				if(MeshScreenSize > ScreenSize)
+				if(MeshScreenSize >= ScreenSize)
 				{
 					LODToRender = Mesh.LODIndex;
 					break;
diff --git a/Engine/Source/Runtime/Engine/Private/StaticMeshRender.cpp b/Engine/Source/Runtime/Engine/Private/StaticMeshRender.cpp
--- a/Engine/Source/Runtime/Engine/Private/StaticMeshRender.cpp
+++ b/Engine/Source/Runtime/Engine/Private/StaticMeshRender.cpp
@@ -372,7 +372,7 @@ void FStaticMeshSceneProxy::DrawStaticElements(FStaticPrimitiveDrawInterface* PD
 				FMeshBatch MeshElement;
 				if(GetMeshElement(LODIndex, SectionIndex, PrimitiveDPG, MeshElement, bUseSelectedMaterial, bUseHoveredMaterial))
 				{
-					PDI->DrawMesh(MeshElement, 0.0f);
+					PDI->DrawMesh(MeshElement, FLT_MAX);
 				}
 			}
 		} 
