Fixed random crash on client when using remote recompile shaders.
Crash was caused by FShaderResources being cleaned up during serialization before serialization was over.

diff --git a/Engine/Source/Runtime/Engine/Private/Materials/MaterialShader.cpp b/Engine/Source/Runtime/Engine/Private/Materials/MaterialShader.cpp
--- a/Engine/Source/Runtime/Engine/Private/Materials/MaterialShader.cpp
+++ b/Engine/Source/Runtime/Engine/Private/Materials/MaterialShader.cpp
@@ -861,16 +861,13 @@ void FMaterialShaderMap::SaveForRemoteRecompile(FArchive& Ar, const TMap<FString
 					FShader *Shader = ShaderIt.Value();
 					FShaderResourceId ShaderId = Shader->GetResourceId();
 
-					// UE_LOG(LogMaterial, Display, TEXT("Shader Outputhash %s vs %s vs %s"), *ShaderId.OutputHash.ToString(), *Shader->GetOutputHash().ToString(), *Shader->GetResourceOutputHash().ToString() );
 					// skip this shader if the Id was already on the client (ie, it didn't change)
-					if (ClientResourceIds.Contains(ShaderId) == false ) // || true)
+					if (ClientResourceIds.Contains(ShaderId) == false )
 					{
 						// lookup the resource by ID
 						FShaderResource* Resource = FShaderResource::FindShaderResourceById(ShaderId);
 						// add it if it's unique
 						UniqueResources.AddUnique(Resource);
-
-						// UE_LOG(LogMaterial, Display, TEXT("Saving material resource id %d %s"), GetTypeHash(ShaderId), *Resource->GetOutputHash().ToString() );
 					}
 					else
 					{
@@ -930,6 +927,9 @@ void FMaterialShaderMap::LoadForRemoteRecompile(FArchive& Ar, EShaderPlatform Sh
 	int32 NumResources;
 	Ar << NumResources;
 
+	// KeepAliveReferences keeps resources alive until we are finished serializing in this function
+	TArray<TRefCountPtr<FShaderResource> > KeepAliveReferences;
+
 	// load and register the resources
 	for (int32 Index = 0; Index < NumResources; Index++)
 	{
@@ -937,18 +937,19 @@ void FMaterialShaderMap::LoadForRemoteRecompile(FArchive& Ar, EShaderPlatform Sh
 		FShaderResource* Resource = new FShaderResource();
 		Resource->Serialize(Ar);
 
-		// UE_LOG(LogMaterial, Display, TEXT("Loaded material resource id %d"), GetTypeHash(Resource->GetId()));
-
 		// if this Id is already in memory, that means that this is a repeated resource and so we skip it
 		if (FShaderResource::FindShaderResourceById(Resource->GetId()) != NULL)
 		{
-			// UE_LOG(LogMaterial, Display, TEXT("Found ShaderResourceById %d"), GetTypeHash(Resource->GetId()));
 			delete Resource;
 		}
 		// otherwise, it's a new resource, so we register it for the maps to find below
 		else
 		{
 			Resource->Register();
+
+			// Keep this guy alive until we finish serializing all the FShaders in
+			// The FShaders which are discarded may cause these resources to be discarded 
+			KeepAliveReferences.Add( Resource );
 		}
 	}
 
diff --git a/Engine/Source/Runtime/ShaderCore/Public/Shader.h b/Engine/Source/Runtime/ShaderCore/Public/Shader.h
--- a/Engine/Source/Runtime/ShaderCore/Public/Shader.h
+++ b/Engine/Source/Runtime/ShaderCore/Public/Shader.h
@@ -82,8 +82,8 @@ public:
 	SHADERCORE_API void Serialize(FArchive& Ar);
 
 	// Reference counting.
-	void AddRef();
-	void Release();
+	SHADERCORE_API void AddRef();
+	SHADERCORE_API void Release();
 
 	SHADERCORE_API void Register();
 
