expose engine max top rotation speed to blueprint

diff --git a/Engine/Source/Runtime/Engine/Classes/Vehicles/WheeledVehicleMovementComponent.h b/Engine/Source/Runtime/Engine/Classes/Vehicles/WheeledVehicleMovementComponent.h
--- a/Engine/Source/Runtime/Engine/Classes/Vehicles/WheeledVehicleMovementComponent.h
+++ b/Engine/Source/Runtime/Engine/Classes/Vehicles/WheeledVehicleMovementComponent.h
@@ -323,6 +323,10 @@ class ENGINE_API UWheeledVehicleMovementComponent : public UPawnMovementComponen
 	UFUNCTION(BlueprintCallable, Category="Game|Components|WheeledVehicleMovement")
 	float GetEngineRotationSpeed() const;
 
+	/** Get current engine's max rotation speed */
+	UFUNCTION(BlueprintCallable, Category="Game|Components|WheeledVehicleMovement")
+	float GetEngineMaxRotationSpeed() const;
+
 	/** Get current gear */
 	UFUNCTION(BlueprintCallable, Category="Game|Components|WheeledVehicleMovement")
 	int32 GetCurrentGear() const;
diff --git a/Engine/Source/Runtime/Engine/Private/Vehicles/WheeledVehicleMovementComponent.cpp b/Engine/Source/Runtime/Engine/Private/Vehicles/WheeledVehicleMovementComponent.cpp
--- a/Engine/Source/Runtime/Engine/Private/Vehicles/WheeledVehicleMovementComponent.cpp
+++ b/Engine/Source/Runtime/Engine/Private/Vehicles/WheeledVehicleMovementComponent.cpp
@@ -934,6 +934,11 @@ float UWheeledVehicleMovementComponent::GetEngineRotationSpeed() const
 	return 0.0f;
 }
 
+float UWheeledVehicleMovementComponent::GetEngineMaxRotationSpeed() const
+{
+	return MaxEngineRPM;
+}
+
 #if WITH_PHYSX
 
 int32 UWheeledVehicleMovementComponent::GearToPhysXGear(const int32 Gear) const
