add negative scale warning for convex body setup

diff --git a/Engine/Source/Runtime/Engine/Private/PhysicsEngine/BodySetup.cpp b/Engine/Source/Runtime/Engine/Private/PhysicsEngine/BodySetup.cpp
--- a/Engine/Source/Runtime/Engine/Private/PhysicsEngine/BodySetup.cpp
+++ b/Engine/Source/Runtime/Engine/Private/PhysicsEngine/BodySetup.cpp
@@ -417,6 +417,10 @@ void UBodySetup::AddShapesToRigidActor(PxRigidActor* PDestActor, FVector& Scale3
 				PConvexGeom.convexMesh = bUseNegX ? ConvexElem->ConvexMeshNegX : ConvexElem->ConvexMesh;
 				PConvexGeom.scale.scale = U2PVector(Scale3DAbs * ConvexElem->GetTransform().GetScale3D().GetAbs());
 				FTransform ConvexTransform = ConvexElem->GetTransform();
+				if (ConvexTransform.GetScale3D().X < 0 || ConvexTransform.GetScale3D().Y < 0 || ConvexTransform.GetScale3D().Z < 0)
+				{
+					UE_LOG(LogPhysics, Warning, TEXT("AddShapesToRigidActor: [%s] ConvexElem[%d] has negative scale. Not currently supported"), *GetPathNameSafe(GetOuter()), i);
+				}
 				if (ConvexTransform.IsValid())
 				{
 					PxTransform PElementTransform = U2PTransform(RelativeTM ? ConvexTransform * *RelativeTM : ConvexTransform);
