Support touch input on hitboxes

diff --git a/Engine/Source/Runtime/Engine/Classes/GameFramework/HUD.h b/Engine/Source/Runtime/Engine/Classes/GameFramework/HUD.h
--- a/Engine/Source/Runtime/Engine/Classes/GameFramework/HUD.h
+++ b/Engine/Source/Runtime/Engine/Classes/GameFramework/HUD.h
@@ -119,7 +119,7 @@ public:
 	 * @param	InCanvas		Canvas on which to render.
 	 * @param	InColor			Color to render the box.
 	 */
-	void Draw( class FCanvas* InCanvas, const FLinearColor& InColor );
+	void Draw( class FCanvas* InCanvas, const FLinearColor& InColor ) const;
 	
 	/** Get the name of this hitbox.  */
 	const FName& GetName() const { return Name;};
@@ -512,7 +512,7 @@ public:
 	/** Array of hitboxes that have been hit for this frame. */
 	TArray< class FHUDHitBox* >	HitBoxHits;
 
-	/** Set of hitbox (by name) that are currently moused over */
+	/** Set of hitbox (by name) that are currently moused over or have a touch contacting them */
 	TSet< FName > HitBoxesOver;
 
 	/**
@@ -525,30 +525,38 @@ public:
 	 * Update the list of hitboxes and dispatch events for any hits.
 	 * @param	InEventType	Input event that triggered the call.
 	 */
-	bool UpdateAndDispatchHitBoxClickEvents(EInputEvent InEventType);
+	bool UpdateAndDispatchHitBoxClickEvents(const FVector2D ClickLocation, const EInputEvent InEventType, const bool bDispatchOverOutEvent);
 	
 	/**
 	 * Update a the list of hitboxes that have been hit this frame.
 	 * @param	Canvas	Canvas on which to render debug boxes.
 	 */
-	void UpdateHitBoxCandidates( const FVector2D& InHitLocation );
+	void UpdateHitBoxCandidates( TArray<FVector2D> InContactPoints );
 	
 	/** Have any hitboxes been hit this frame. */
 	bool AnyCurrentHitBoxHits() const;	
 
 	/**
 	 * Find the first hitbox containing the given coordinates.
-	 * @param	InHitLocation Coordinates to check
-	 * @return returns the hitbox at the given coordinates or NULL if none match.
+	 * @param	InHitLocation		Coordinates to check
+	 * @param	bConsumingInput		If true will return the first hitbox that would consume input at this coordinate
+	 * @return	returns the hitbox at the given coordinates or NULL if none match.
 	 */
-	const FHUDHitBox* GetHitBoxAtCoordinates( const FVector2D& InHitLocation ) const;
+	const FHUDHitBox* GetHitBoxAtCoordinates( FVector2D InHitLocation, bool bConsumingInput = false ) const;
 
+	/**
+	* Finds all the hitboxes containing the given coordinates.
+	* @param	InHitLocation		Coordinates to check
+	* @param	OutHitBoxes			Reference parameter with hit box references at the coordinates
+	*/
+	void GetHitBoxesAtCoordinates(FVector2D InHitLocation, TArray<const FHUDHitBox*>& OutHitBoxes) const;
+	
 	/**
 	 * Return the hitbox with the given name
 	 * @param	InName	Name of required hitbox
 	 * @return returns the hitbox with the given name NULL if none match.
 	 */
-	const FHUDHitBox* GetHitBoxWithName( const FName& InName ) const;
+	const FHUDHitBox* GetHitBoxWithName( const FName InName ) const;
 
 	//=============================================================================
 	// UIBlur overrides
@@ -566,6 +574,10 @@ public:
 
 protected:
 	bool IsCanvasValid_WarnIfNot() const;
+
+private:
+	// Helper function to deal with screen offset and splitscreen mapping of coordinates to HUD
+	FVector2D GetCoordinateOffset() const;
 };
 
 
diff --git a/Engine/Source/Runtime/Engine/Private/HUD.cpp b/Engine/Source/Runtime/Engine/Private/HUD.cpp
--- a/Engine/Source/Runtime/Engine/Private/HUD.cpp
+++ b/Engine/Source/Runtime/Engine/Private/HUD.cpp
@@ -69,6 +69,39 @@ void AHUD::NotifyBindPostProcessEffects()
 	// overload with custom code e.g. getting material pointers to control effects for gameplay.
 }
 
+FVector2D AHUD::GetCoordinateOffset() const
+{
+	FVector2D Offset;
+
+	ULocalPlayer* LocalPlayer = Cast<ULocalPlayer>(GetOwningPlayerController()->Player);
+
+	if (LocalPlayer)
+	{
+		// Create a view family for the game viewport
+		FSceneViewFamilyContext ViewFamily(FSceneViewFamily::ConstructionValues(
+			LocalPlayer->ViewportClient->Viewport,
+			GetWorld()->Scene,
+			LocalPlayer->ViewportClient->EngineShowFlags)
+			.SetRealtimeUpdate(true));
+
+		// Calculate a view where the player is to update the streaming from the players start location
+		FVector ViewLocation;
+		FRotator ViewRotation;
+		FSceneView* SceneView = LocalPlayer->CalcSceneView(&ViewFamily, /*out*/ ViewLocation, /*out*/ ViewRotation, LocalPlayer->ViewportClient->Viewport);
+
+		if (SceneView)
+		{
+			Offset.X = (SceneView->ViewRect.Min.X - SceneView->UnconstrainedViewRect.Min.X) // This accounts for the borders when the aspect ratio is locked
+							- SceneView->UnconstrainedViewRect.Min.X;						// And this will deal with the viewport offset if its a split screen
+
+			Offset.Y = (SceneView->ViewRect.Min.Y - SceneView->UnconstrainedViewRect.Min.Y)
+							- SceneView->UnconstrainedViewRect.Min.Y;
+		}
+	}
+
+	return Offset;
+}
+
 void AHUD::PostRender()
 {
 	// Set up delta time
@@ -103,30 +136,32 @@ void AHUD::PostRender()
 
 		if (LocalPlayer)
 		{
-			FVector2D MousePosition = LocalPlayer->ViewportClient->GetMousePosition();
+			TArray<FVector2D> ContactPoints;
+			ContactPoints.Add(LocalPlayer->ViewportClient->GetMousePosition());
 
-			// Create a view family for the game viewport
-			FSceneViewFamilyContext ViewFamily( FSceneViewFamily::ConstructionValues(
-				LocalPlayer->ViewportClient->Viewport,
-				GetWorld()->Scene,
-				LocalPlayer->ViewportClient->EngineShowFlags )
-				.SetRealtimeUpdate(true) );
+			for (int32 FingerIndex = 0; FingerIndex < EKeys::NUM_TOUCH_KEYS; ++FingerIndex)
+			{
+				FVector2D TouchLocation;
+				bool bPressed = false;
+
+				GetOwningPlayerController()->GetInputTouchState((ETouchIndex::Type)FingerIndex, TouchLocation.X, TouchLocation.Y, bPressed);
+
+				if (bPressed)
+				{
+					ContactPoints.Add(TouchLocation);
+				}
+			}
 
-			// Calculate a view where the player is to update the streaming from the players start location
-			FVector ViewLocation;
-			FRotator ViewRotation;
-			FSceneView* SceneView = LocalPlayer->CalcSceneView( &ViewFamily, /*out*/ ViewLocation, /*out*/ ViewRotation, LocalPlayer->ViewportClient->Viewport );
+			FVector2D ContactPointOffset = GetCoordinateOffset();
 
-			if (SceneView)
+			if (!ContactPointOffset.IsZero())
 			{
-				// This accounts for the borders when the aspect ratio is locked
-				MousePosition.X -= (SceneView->ViewRect.Min.X - SceneView->UnconstrainedViewRect.Min.X);
-				MousePosition.Y -= (SceneView->ViewRect.Min.Y - SceneView->UnconstrainedViewRect.Min.Y);
-				// And this will deal with the viewport offset if its a split screen
-				MousePosition.X -= SceneView->UnconstrainedViewRect.Min.X;
-				MousePosition.Y -= SceneView->UnconstrainedViewRect.Min.Y;
+				for (FVector2D& ContactPoint : ContactPoints)
+				{
+					ContactPoint -= ContactPointOffset;
+				}
 			}
-			UpdateHitBoxCandidates( MousePosition );
+			UpdateHitBoxCandidates( ContactPoints );
 		}
 	}
 	
@@ -628,42 +663,54 @@ APawn* AHUD::GetOwningPawn() const
 
 void AHUD::RenderHitBoxes( FCanvas* InCanvas )
 {
-	for (int32 iBox = 0; iBox < HitBoxMap.Num() ; iBox++)
+	for (const FHUDHitBox& HitBox : HitBoxMap)
 	{
 		FLinearColor BoxColor = FLinearColor::White;
-		if( HitBoxHits.Contains(&HitBoxMap[iBox]) )
+		if( HitBoxHits.Contains(const_cast<FHUDHitBox*>(&HitBox)))
 		{
 			BoxColor = FLinearColor::Red;
 		}
-		HitBoxMap[iBox].Draw( InCanvas, BoxColor );
+		HitBox.Draw( InCanvas, BoxColor );
 	}
 }
 
-void AHUD::UpdateHitBoxCandidates( const FVector2D& InHitLocation )
+void AHUD::UpdateHitBoxCandidates( TArray<FVector2D> InContactPoints )
 {
 	HitBoxHits.Empty();
-	for (int32 iBox = 0; iBox < HitBoxMap.Num() ; iBox++)
+	for (FHUDHitBox& HitBox : HitBoxMap)
 	{
-		if( HitBoxMap[iBox].Contains( InHitLocation ) )
+		bool bAdded = false;
+		for (int32 ContactPointIndex = InContactPoints.Num() - 1; ContactPointIndex >= 0; --ContactPointIndex)
+		{
+			if (HitBox.Contains(InContactPoints[ContactPointIndex]))
+			{
+				if (!bAdded)
+				{
+					HitBoxHits.Add(&HitBox);
+					bAdded = true;
+				}
+				if (HitBox.ConsumesInput())
+				{
+					InContactPoints.RemoveAtSwap(ContactPointIndex);
+				}
+				else
+				{
+					break;
+				}
+			}
+		}
+		if (InContactPoints.Num() == 0)
 		{
-			HitBoxHits.Add(&HitBoxMap[iBox]);
+			break;
 		}
 	}
-	struct FHitBoxSort
-	{
-		FORCEINLINE bool operator()( const FHUDHitBox& A, const FHUDHitBox& B ) const { return B.GetPriority() < A.GetPriority(); }
-	};
-
-	// Now sort the hits
-	HitBoxHits.Sort( FHitBoxSort() );
 
 	TSet<FName> NotOverHitBoxes = HitBoxesOver;
 	TArray<FName> NewlyOverHitBoxes;
 
 	// Now figure out which boxes we are over and deal with begin/end cursor over messages 
-	for (int32 HitIndex = 0; HitIndex < HitBoxHits.Num(); ++HitIndex )
+	for (FHUDHitBox* HitBox : HitBoxHits)
 	{
-		FHUDHitBox* HitBox = HitBoxHits[HitIndex];
 		const FName HitBoxName = HitBox->GetName();
 		if (HitBoxesOver.Contains(HitBoxName))
 		{
@@ -673,10 +720,6 @@ void AHUD::UpdateHitBoxCandidates( const FVector2D& InHitLocation )
 		{
 			NewlyOverHitBoxes.AddUnique(HitBoxName);
 		}
-		if (HitBox->ConsumesInput())
-		{
-			break; //Early out if this box blocks further consideration
-		}
 	}
 
 	// Dispatch the end cursor over messages
@@ -696,19 +739,35 @@ void AHUD::UpdateHitBoxCandidates( const FVector2D& InHitLocation )
 	}
 }
 
-const FHUDHitBox* AHUD::GetHitBoxAtCoordinates( const FVector2D& InHitLocation ) const
+const FHUDHitBox* AHUD::GetHitBoxAtCoordinates( FVector2D InHitLocation, const bool bIsConsumingInput ) const
 {
-	for (int32 iBox = 0; iBox < HitBoxMap.Num() ; iBox++)
+	InHitLocation -= GetCoordinateOffset();
+
+	for (const FHUDHitBox& HitBox : HitBoxMap)
 	{
-		if( HitBoxMap[iBox].Contains( InHitLocation ) )
+		if( (!bIsConsumingInput || HitBox.ConsumesInput()) && HitBox.Contains( InHitLocation ) )
 		{
-			return &HitBoxMap[iBox];
+			return &HitBox;
 		}
 	}
 	return NULL;
 }
 
-const FHUDHitBox* AHUD::GetHitBoxWithName( const FName& InName ) const
+void AHUD::GetHitBoxesAtCoordinates(FVector2D InHitLocation, TArray<const FHUDHitBox*>& OutHitBoxes) const
+{
+	InHitLocation -= GetCoordinateOffset();
+	OutHitBoxes.Empty();
+
+	for (const FHUDHitBox& HitBox : HitBoxMap)
+	{
+		if (HitBox.Contains(InHitLocation))
+		{
+			OutHitBoxes.Add(&HitBox);
+		}
+	}
+}
+
+const FHUDHitBox* AHUD::GetHitBoxWithName( const FName InName ) const
 {
 	for (int32 iBox = 0; iBox < HitBoxMap.Num() ; iBox++)
 	{
@@ -740,32 +799,61 @@ const TArray<FIntRect>& AHUD::GetUIBlurRectangles() const
 	return UIBlurOverrideRectangles;
 }
 
-bool AHUD::UpdateAndDispatchHitBoxClickEvents(EInputEvent InEventType)
+bool AHUD::UpdateAndDispatchHitBoxClickEvents(FVector2D ClickLocation, const EInputEvent InEventType, const bool bDispatchOverOutEvent)
 {
+	ClickLocation -= GetCoordinateOffset();
+
+	bool bHit = false;
 	for (FHUDHitBox* HitBoxHit : HitBoxHits)
 	{
-		if(InEventType == IE_Pressed || InEventType == IE_DoubleClick)
-		{
-			ReceiveHitBoxClick( HitBoxHit->GetName() );
-		}
-		else if(InEventType == IE_Released)
+		if (HitBoxHit->Contains(ClickLocation))
 		{
-			ReceiveHitBoxRelease( HitBoxHit->GetName() );
-		}
+			bHit = true;
 
-		if( HitBoxHit->ConsumesInput() == true )
-		{
-			break;	//Early out if this box consumed the click
+			if (InEventType == IE_Pressed || InEventType == IE_DoubleClick)
+			{
+				ReceiveHitBoxClick(HitBoxHit->GetName());
+				if (bDispatchOverOutEvent)
+				{
+					ReceiveHitBoxBeginCursorOver(HitBoxHit->GetName());
+				}
+			}
+			else if (InEventType == IE_Released)
+			{
+				ReceiveHitBoxRelease(HitBoxHit->GetName());
+				if (bDispatchOverOutEvent)
+				{
+					ReceiveHitBoxEndCursorOver(HitBoxHit->GetName());
+				}
+			}
+
+			if (HitBoxHit->ConsumesInput() == true)
+			{
+				break;	//Early out if this box consumed the click
+			}
 		}
 	}
-	return AnyCurrentHitBoxHits();
+	return bHit;
 }
 
 void AHUD::AddHitBox(FVector2D Position, FVector2D Size, FName Name, bool bConsumesInput, int32 Priority)
 {	
 	if( GetHitBoxWithName(Name) == NULL )
 	{
-		HitBoxMap.Add( FHUDHitBox( Position, Size, Name, bConsumesInput, Priority ) );
+		bool bAdded = false;
+		for (int32 Index = 0; Index < HitBoxMap.Num(); ++Index)
+		{
+			if (HitBoxMap[Index].GetPriority() < Priority)
+			{
+				HitBoxMap.Insert(FHUDHitBox(Position, Size, Name, bConsumesInput, Priority), Index);
+				bAdded = true;
+				break;
+			}
+		}
+		if (!bAdded)
+		{
+			HitBoxMap.Add(FHUDHitBox(Position, Size, Name, bConsumesInput, Priority));
+		}
 	}
 	else
 	{
@@ -789,12 +877,12 @@ bool AHUD::IsCanvasValid_WarnIfNot() const
 /////////////////
 
 FHUDHitBox::FHUDHitBox( FVector2D InCoords, FVector2D InSize, const FName& InName, bool bInConsumesInput, int32 InPriority )
+	: Coords(InCoords)
+	, Size(InSize)
+	, Name(InName)
+	, bConsumesInput(bInConsumesInput)
+	, Priority(InPriority)
 {
-	Coords = InCoords;
-	Size = InSize;
-	Name = InName;
-	bConsumesInput = bInConsumesInput;
-	Priority = InPriority;
 }
 
 bool FHUDHitBox::Contains( FVector2D InCoords ) const
@@ -810,7 +898,7 @@ bool FHUDHitBox::Contains( FVector2D InCoords ) const
 	return bResult;
 }
 
-void FHUDHitBox::Draw( FCanvas* InCanvas, const FLinearColor& InColor )
+void FHUDHitBox::Draw( FCanvas* InCanvas, const FLinearColor& InColor ) const
 {
 	FCanvasBoxItem	BoxItem( Coords, Size );
 	BoxItem.SetColor( InColor );
diff --git a/Engine/Source/Runtime/Engine/Private/PlayerController.cpp b/Engine/Source/Runtime/Engine/Private/PlayerController.cpp
--- a/Engine/Source/Runtime/Engine/Private/PlayerController.cpp
+++ b/Engine/Source/Runtime/Engine/Private/PlayerController.cpp
@@ -1749,7 +1749,7 @@ void APlayerController::DeprojectScreenPositionToWorld(float ScreenX, float Scre
 bool APlayerController::GetHitResultAtScreenPosition(const FVector2D ScreenPosition, const ECollisionChannel TraceChannel, bool bTraceComplex, FHitResult& HitResult) const
 {
 	// Early out if we clicked on a HUD hitbox
-	if( GetHUD() != NULL && GetHUD()->AnyCurrentHitBoxHits() )
+	if( GetHUD() != NULL && GetHUD()->GetHitBoxAtCoordinates(ScreenPosition, true) )
 	{
 		return false;
 	}
@@ -1787,7 +1787,7 @@ bool APlayerController::GetHitResultAtScreenPosition(const FVector2D ScreenPosit
 bool APlayerController::GetHitResultAtScreenPosition(const FVector2D ScreenPosition, const ETraceTypeQuery TraceChannel, bool bTraceComplex, FHitResult& HitResult) const
 {
 	// Early out if we clicked on a HUD hitbox
-	if( GetHUD() != NULL && GetHUD()->AnyCurrentHitBoxHits() )
+	if (GetHUD() != NULL && GetHUD()->GetHitBoxAtCoordinates(ScreenPosition, true))
 	{
 		return false;
 	}
@@ -1825,7 +1825,7 @@ bool APlayerController::GetHitResultAtScreenPosition(const FVector2D ScreenPosit
 bool APlayerController::GetHitResultAtScreenPosition(const FVector2D ScreenPosition, const TArray<TEnumAsByte<EObjectTypeQuery> > & ObjectTypes, bool bTraceComplex, FHitResult& HitResult) const
 {
 	// Early out if we clicked on a HUD hitbox
-	if( GetHUD() != NULL && GetHUD()->AnyCurrentHitBoxHits() )
+	if (GetHUD() != NULL && GetHUD()->GetHitBoxAtCoordinates(ScreenPosition, true))
 	{
 		return false;
 	}
@@ -1963,6 +1963,7 @@ bool APlayerController::InputKey(FKey Key, EInputEvent EventType, float AmountDe
 		// TODO: Allow click key(s?) to be defined
 		if (bEnableClickEvents && Key == EKeys::LeftMouseButton)
 		{
+			const FVector2D MousePosition = CastChecked<ULocalPlayer>(Player)->ViewportClient->GetMousePosition();
 			UPrimitiveComponent* ClickedPrimitive = NULL;
 			if (bEnableMouseOverEvents)
 			{
@@ -1970,11 +1971,8 @@ bool APlayerController::InputKey(FKey Key, EInputEvent EventType, float AmountDe
 			}
 			else
 			{
-				FVector2D MouseLocation;
-				GetInputMouseDelta(MouseLocation.X, MouseLocation.Y);
-
 				FHitResult HitResult;
-				const bool bHit = GetHitResultAtScreenPosition(CastChecked<ULocalPlayer>(Player)->ViewportClient->GetMousePosition(), CurrentClickTraceChannel, true, HitResult);
+				const bool bHit = GetHitResultAtScreenPosition(MousePosition, CurrentClickTraceChannel, true, HitResult);
 				if (bHit)
 				{
 					ClickedPrimitive = HitResult.Component.Get();
@@ -1982,7 +1980,7 @@ bool APlayerController::InputKey(FKey Key, EInputEvent EventType, float AmountDe
 			}
 			if( GetHUD() )
 			{
-				if( GetHUD()->UpdateAndDispatchHitBoxClickEvents(EventType) )
+				if (GetHUD()->UpdateAndDispatchHitBoxClickEvents(MousePosition, EventType, false))
 				{
 					ClickedPrimitive = NULL;
 				}
@@ -2028,6 +2026,17 @@ bool APlayerController::InputTouch(uint32 Handle, ETouchType::Type Type, const F
 			UPrimitiveComponent* PreviousComponent = CurrentTouchablePrimitives[Handle].Get();
 			UPrimitiveComponent* CurrentComponent = (bHit ? HitResult.Component.Get() : NULL);
 
+			if (GetHUD())
+			{
+				if (Type == ETouchType::Began || Type == ETouchType::Ended)
+				{
+					if (GetHUD()->UpdateAndDispatchHitBoxClickEvents(TouchLocation, (Type == ETouchType::Began ? EInputEvent::IE_Pressed : EInputEvent::IE_Released), true))
+					{
+						CurrentComponent = NULL;
+					}
+				}
+			}
+
 			switch(Type)
 			{
 			case ETouchType::Began:
