Merge branch 'master' of http://github.com/EpicGames/UnrealEngine

diff --git a/Engine/Source/Editor/BlueprintGraph/Classes/EdGraphSchema_K2.h b/Engine/Source/Editor/BlueprintGraph/Classes/EdGraphSchema_K2.h
--- a/Engine/Source/Editor/BlueprintGraph/Classes/EdGraphSchema_K2.h
+++ b/Engine/Source/Editor/BlueprintGraph/Classes/EdGraphSchema_K2.h
@@ -662,7 +662,7 @@ public:
 	/** Can this function be overridden by kismet (either placed as event or new function graph created) */
 	static bool CanKismetOverrideFunction(const UFunction* Function);
 
-	/** returns friendly signiture name if possible or Removes any mangling to get the unmangled signiture name of the function */
+	/** returns friendly signiture name if possible or Removes any mangling to get the unmangled signature name of the function */
 	static FString GetFriendlySignitureName(const UFunction* Function);
 
 	static bool IsAllowableBlueprintVariableType(const class UEnum* InEnum);
@@ -723,6 +723,14 @@ public:
 	 */
 	virtual void CreateFunctionGraphTerminators(UEdGraph& Graph, UClass* Class) const;
 
+	/** 
+	 * Populate new function graph with entry and possibly return node
+	 * 
+	 * @param	Graph			Graph to add the function terminators to
+	 * @param	FunctionSignature	The function signature to mimic when creating the inputs and outputs for the function.
+	 */
+	virtual void CreateFunctionGraphTerminators(UEdGraph& Graph, UFunction* FunctionSignature) const;
+
 	/**
 	 * Converts a pin type into a fully qualified string (e.g., object'ObjectName').
 	 *
diff --git a/Engine/Source/Editor/BlueprintGraph/Private/EdGraphSchema_K2.cpp b/Engine/Source/Editor/BlueprintGraph/Private/EdGraphSchema_K2.cpp
--- a/Engine/Source/Editor/BlueprintGraph/Private/EdGraphSchema_K2.cpp
+++ b/Engine/Source/Editor/BlueprintGraph/Private/EdGraphSchema_K2.cpp
@@ -728,56 +728,56 @@ void UEdGraphSchema_K2::GetContextMenuActions(const UEdGraph* CurrentGraph, cons
 		MenuBuilder->BeginSection("EdGraphSchemaPinActions", LOCTEXT("PinActionsMenuHeader", "Pin Actions"));
 		{
 			if (!bIsDebugging)
-		    {
-			    // Break pin links
-			    if (InGraphPin->LinkedTo.Num() > 1)
-			    {
-				    MenuBuilder->AddMenuEntry( FGraphEditorCommands::Get().BreakPinLinks );
-			    }
-    
-			    // Add the change pin type action, if this is a select node
-			    if (InGraphNode->IsA(UK2Node_Select::StaticClass()))
-			    {
-				    MenuBuilder->AddMenuEntry(FGraphEditorCommands::Get().ChangePinType);
-			    }
-    
-			    // add sub menu for break link to
-			    if (InGraphPin->LinkedTo.Num() > 0)
-			    {
-				    if(InGraphPin->LinkedTo.Num() > 1)
-				    {
-					    MenuBuilder->AddSubMenu(
-						    LOCTEXT("BreakLinkTo", "Break Link To..."),
-						    LOCTEXT("BreakSpecificLinks", "Break a specific link..."),
-						    FNewMenuDelegate::CreateUObject( (UEdGraphSchema_K2*const)this, &UEdGraphSchema_K2::GetBreakLinkToSubMenuActions, const_cast<UEdGraphPin*>(InGraphPin)));
+			{
+				// Break pin links
+				if (InGraphPin->LinkedTo.Num() > 1)
+				{
+					MenuBuilder->AddMenuEntry( FGraphEditorCommands::Get().BreakPinLinks );
+				}
+	
+				// Add the change pin type action, if this is a select node
+				if (InGraphNode->IsA(UK2Node_Select::StaticClass()))
+				{
+					MenuBuilder->AddMenuEntry(FGraphEditorCommands::Get().ChangePinType);
+				}
+	
+				// add sub menu for break link to
+				if (InGraphPin->LinkedTo.Num() > 0)
+				{
+					if(InGraphPin->LinkedTo.Num() > 1)
+					{
+						MenuBuilder->AddSubMenu(
+							LOCTEXT("BreakLinkTo", "Break Link To..."),
+							LOCTEXT("BreakSpecificLinks", "Break a specific link..."),
+							FNewMenuDelegate::CreateUObject( (UEdGraphSchema_K2*const)this, &UEdGraphSchema_K2::GetBreakLinkToSubMenuActions, const_cast<UEdGraphPin*>(InGraphPin)));
 
 						MenuBuilder->AddSubMenu(
 							LOCTEXT("JumpToConnection", "Jump to Connection..."),
 							LOCTEXT("JumpToSpecificConnection", "Jump to specific connection..."),
 							FNewMenuDelegate::CreateUObject( (UEdGraphSchema_K2*const)this, &UEdGraphSchema_K2::GetJumpToConnectionSubMenuActions, const_cast<UEdGraphPin*>(InGraphPin)));
-				    }
-				    else
-				    {
-					    ((UEdGraphSchema_K2*const)this)->GetBreakLinkToSubMenuActions(*MenuBuilder, const_cast<UEdGraphPin*>(InGraphPin));
+					}
+					else
+					{
+						((UEdGraphSchema_K2*const)this)->GetBreakLinkToSubMenuActions(*MenuBuilder, const_cast<UEdGraphPin*>(InGraphPin));
 						((UEdGraphSchema_K2*const)this)->GetJumpToConnectionSubMenuActions(*MenuBuilder, const_cast<UEdGraphPin*>(InGraphPin));
 					}
-			    }
-    
-			    // Conditionally add the var promotion pin if this is an output pin and it's not an exec pin
-			    if (InGraphPin->PinType.PinCategory != PC_Exec)
-			    {
-				    MenuBuilder->AddMenuEntry( FGraphEditorCommands::Get().PromoteToVariable );
-			    }
-    
-			    // Conditionally add the execution path pin removal if this is an execution branching node
-			    if( InGraphPin->Direction == EGPD_Output && InGraphPin->GetOwningNode())
-			    {
-				    if ( InGraphPin->GetOwningNode()->IsA(UK2Node_ExecutionSequence::StaticClass()) ||  InGraphPin->GetOwningNode()->IsA(UK2Node_Switch::StaticClass()) )
-				    {
-					    MenuBuilder->AddMenuEntry( FGraphEditorCommands::Get().RemoveExecutionPin );
-				    }
-			    }
-		    }
+				}
+	
+				// Conditionally add the var promotion pin if this is an output pin and it's not an exec pin
+				if (InGraphPin->PinType.PinCategory != PC_Exec)
+				{
+					MenuBuilder->AddMenuEntry( FGraphEditorCommands::Get().PromoteToVariable );
+				}
+	
+				// Conditionally add the execution path pin removal if this is an execution branching node
+				if( InGraphPin->Direction == EGPD_Output && InGraphPin->GetOwningNode())
+				{
+					if ( InGraphPin->GetOwningNode()->IsA(UK2Node_ExecutionSequence::StaticClass()) ||  InGraphPin->GetOwningNode()->IsA(UK2Node_Switch::StaticClass()) )
+					{
+						MenuBuilder->AddMenuEntry( FGraphEditorCommands::Get().RemoveExecutionPin );
+					}
+				}
+			}
 		}
 		MenuBuilder->EndSection();
 
@@ -2172,9 +2172,54 @@ void UEdGraphSchema_K2::CreateFunctionGraphTerminators(UEdGraph& Graph, UClass*
 			UEdGraphPin* EntryNodeExec = FindExecutionPin(*EntryNode, EGPD_Output);
 			UEdGraphPin* ResultNodeExec = FindExecutionPin(*ReturnNode, EGPD_Input);
 			EntryNodeExec->MakeLinkTo(ResultNodeExec);
+		}
+	}
+}
+
+void UEdGraphSchema_K2::CreateFunctionGraphTerminators(UEdGraph& Graph, UFunction* FunctionSignature) const
+{
+	const FName GraphName = Graph.GetFName();
 
+	UBlueprint* Blueprint = FBlueprintEditorUtils::FindBlueprintForGraphChecked(&Graph);
+	check(Blueprint->BlueprintType != BPTYPE_MacroLibrary);
+
+	// Create a function entry node
+	FGraphNodeCreator<UK2Node_FunctionEntry> FunctionEntryCreator(Graph);
+	UK2Node_FunctionEntry* EntryNode = FunctionEntryCreator.CreateNode();
+	EntryNode->SignatureClass = NULL;
+	EntryNode->SignatureName = GraphName;
+	FunctionEntryCreator.Finalize();
+
+	EntryNode->CreatePinsForFunctionEntryExit(FunctionSignature, /*bIsFunctionEntry=*/ true);
+
+	// See if any function params are marked as out
+	bool bHasOutParam =  false;
+	for ( TFieldIterator<UProperty> It(FunctionSignature); It && ( It->PropertyFlags & CPF_Parm ); ++It )
+	{
+		if ( It->PropertyFlags & CPF_OutParm )
+		{
+			bHasOutParam = true;
+			break;
 		}
 	}
+
+	if ( bHasOutParam )
+	{
+		FGraphNodeCreator<UK2Node_FunctionResult> NodeCreator(Graph);
+		UK2Node_FunctionResult* ReturnNode = NodeCreator.CreateNode();
+		ReturnNode->SignatureClass = NULL;
+		ReturnNode->SignatureName = GraphName;
+		ReturnNode->NodePosX = EntryNode->NodePosX + EntryNode->NodeWidth + 256;
+		ReturnNode->NodePosY = EntryNode->NodePosY;
+		NodeCreator.Finalize();
+
+		ReturnNode->CreatePinsForFunctionEntryExit(FunctionSignature, /*bIsFunctionEntry=*/ false);
+
+		// Auto-connect the pins for entry and exit, so that by default the signature is properly generated
+		UEdGraphPin* EntryNodeExec = FindExecutionPin(*EntryNode, EGPD_Output);
+		UEdGraphPin* ResultNodeExec = FindExecutionPin(*ReturnNode, EGPD_Input);
+		EntryNodeExec->MakeLinkTo(ResultNodeExec);
+	}
 }
 
 bool UEdGraphSchema_K2::ConvertPropertyToPinType(const UProperty* Property, /*out*/ FEdGraphPinType& TypeOut) const
@@ -2417,9 +2462,9 @@ void UEdGraphSchema_K2::GetVariableTypeTree( TArray< TSharedPtr<FPinTypeTreeInfo
 	UScriptStruct* VectorStruct = FindObjectChecked<UScriptStruct>(UObject::StaticClass(), TEXT("Vector"));
 	UScriptStruct* RotatorStruct = FindObjectChecked<UScriptStruct>(UObject::StaticClass(), TEXT("Rotator"));
 	UScriptStruct* TransformStruct = FindObjectChecked<UScriptStruct>(UObject::StaticClass(), TEXT("Transform"));
- 	TypeTree.Add( MakeShareable( new FPinTypeTreeInfo(PC_Struct, VectorStruct, LOCTEXT("VectorType", "A 3D vector").ToString()) ) );
- 	TypeTree.Add( MakeShareable( new FPinTypeTreeInfo(PC_Struct, RotatorStruct, LOCTEXT("RotatorType", "A 3D rotation").ToString()) ) );
- 	TypeTree.Add( MakeShareable( new FPinTypeTreeInfo(PC_Struct, TransformStruct, LOCTEXT("TransformType", "A 3D transformation, including translation, rotation and 3D scale.").ToString()) ) );
+	TypeTree.Add( MakeShareable( new FPinTypeTreeInfo(PC_Struct, VectorStruct, LOCTEXT("VectorType", "A 3D vector").ToString()) ) );
+	TypeTree.Add( MakeShareable( new FPinTypeTreeInfo(PC_Struct, RotatorStruct, LOCTEXT("RotatorType", "A 3D rotation").ToString()) ) );
+	TypeTree.Add( MakeShareable( new FPinTypeTreeInfo(PC_Struct, TransformStruct, LOCTEXT("TransformType", "A 3D transformation, including translation, rotation and 3D scale.").ToString()) ) );
 	
 	// Add wildcard type
 	if (bAllowWildCard)
@@ -2737,7 +2782,7 @@ bool UEdGraphSchema_K2::CanEncapuslateNode(UEdGraphNode const& TestNode) const
 {
 	// Can't encapsulate entry points (may relax this restriction in the future, but it makes sense for now)
 	return !TestNode.IsA(UK2Node_FunctionTerminator::StaticClass()) && 
-		    TestNode.GetClass() != UK2Node_Tunnel::StaticClass(); //Tunnel nodes getting sucked into collapsed graphs fails badly, want to allow derived types though(composite node/Macroinstances)
+			TestNode.GetClass() != UK2Node_Tunnel::StaticClass(); //Tunnel nodes getting sucked into collapsed graphs fails badly, want to allow derived types though(composite node/Macroinstances)
 }
 
 void UEdGraphSchema_K2::HandleGraphBeingDeleted(UEdGraph& GraphBeingRemoved) const
diff --git a/Engine/Source/Editor/GameProjectGeneration/GameProjectGeneration.Build.cs b/Engine/Source/Editor/GameProjectGeneration/GameProjectGeneration.Build.cs
--- a/Engine/Source/Editor/GameProjectGeneration/GameProjectGeneration.Build.cs
+++ b/Engine/Source/Editor/GameProjectGeneration/GameProjectGeneration.Build.cs
@@ -28,7 +28,8 @@ public class GameProjectGeneration : ModuleRules
 				"SlateCore",
                 "EditorStyle",
                 "SourceControl",
-                "UnrealEd",
+ 				"TargetPlatform",
+				"UnrealEd",
 				"DesktopPlatform"
 			}
 		);
diff --git a/Engine/Source/Editor/GameProjectGeneration/Private/GameProjectUtils.cpp b/Engine/Source/Editor/GameProjectGeneration/Private/GameProjectUtils.cpp
--- a/Engine/Source/Editor/GameProjectGeneration/Private/GameProjectUtils.cpp
+++ b/Engine/Source/Editor/GameProjectGeneration/Private/GameProjectUtils.cpp
@@ -11,6 +11,7 @@
 #include "EngineBuildSettings.h"
 
 #include "DesktopPlatformModule.h"
+#include "TargetPlatform.h"
 
 #define LOCTEXT_NAMESPACE "GameProjectUtils"
 
@@ -84,6 +85,12 @@ bool GameProjectUtils::IsValidProjectFileForCreation(const FString& ProjectFile,
 		return false;
 	}
 
+	if (NameContainsUnderscoreAndXB1Installed(BaseProjectFile))
+	{
+		OutFailReason = LOCTEXT( "ProjectNameContainsIllegalCharactersOnXB1", "Project names may not contain an underscore when the Xbox One XDK is installed." );
+		return false;
+	}
+
 	if ( !FPaths::ValidatePath(FPaths::GetPath(ProjectFile), &OutFailReason) )
 	{
 		return false;
@@ -164,6 +171,12 @@ bool GameProjectUtils::OpenProject(const FString& ProjectFile, FText& OutFailRea
 		return false;
 	}
 
+	if (NameContainsUnderscoreAndXB1Installed(BaseProjectFile))
+	{
+		OutFailReason = LOCTEXT( "ProjectNameContainsIllegalCharactersOnXB1", "Project names may not contain an underscore when the Xbox One XDK is installed." );
+		return false;
+	}
+
 	if ( !FPaths::ValidatePath(FPaths::GetPath(ProjectFile), &OutFailReason) )
 	{
 		return false;
@@ -1089,6 +1102,31 @@ bool GameProjectUtils::NameContainsOnlyLegalCharacters(const FString& TestName,
 	return !bContainsIllegalCharacters;
 }
 
+bool GameProjectUtils::NameContainsUnderscoreAndXB1Installed(const FString& TestName)
+{
+	bool bContainsIllegalCharacters = false;
+
+	// Only allow alphanumeric characters in the project name
+	for ( int32 CharIdx = 0 ; CharIdx < TestName.Len() ; ++CharIdx )
+	{
+		const FString& Char = TestName.Mid( CharIdx, 1 );
+		if ( Char == TEXT("_") )
+		{
+			const ITargetPlatform* Platform = GetTargetPlatformManager()->FindTargetPlatform(TEXT("XboxOne"));
+			if (Platform)
+			{
+				FString NotInstalledDocLink;
+				if (Platform->IsSdkInstalled(true, NotInstalledDocLink))
+				{
+					bContainsIllegalCharacters = true;
+				}
+			}
+		}
+	}
+
+	return bContainsIllegalCharacters;
+}
+
 bool GameProjectUtils::ProjectFileExists(const FString& ProjectFile)
 {
 	return FPlatformFileManager::Get().GetPlatformFile().FileExists(*ProjectFile);
diff --git a/Engine/Source/Editor/GameProjectGeneration/Private/GameProjectUtils.h b/Engine/Source/Editor/GameProjectGeneration/Private/GameProjectUtils.h
--- a/Engine/Source/Editor/GameProjectGeneration/Private/GameProjectUtils.h
+++ b/Engine/Source/Editor/GameProjectGeneration/Private/GameProjectUtils.h
@@ -111,6 +111,9 @@ private:
 	/** Checks the name for illegal characters */
 	static bool NameContainsOnlyLegalCharacters(const FString& TestName, FString& OutIllegalCharacters);
 
+	/** Checks the name for an underscore and the existence of XB1 XDK */
+	static bool NameContainsUnderscoreAndXB1Installed(const FString& TestName);
+
 	/** Returns true if the project file exists on disk */
 	static bool ProjectFileExists(const FString& ProjectFile);
 
diff --git a/Engine/Source/Editor/Kismet/Private/BlueprintDetailsCustomization.cpp b/Engine/Source/Editor/Kismet/Private/BlueprintDetailsCustomization.cpp
--- a/Engine/Source/Editor/Kismet/Private/BlueprintDetailsCustomization.cpp
+++ b/Engine/Source/Editor/Kismet/Private/BlueprintDetailsCustomization.cpp
@@ -1335,7 +1335,7 @@ void FBlueprintVarActionDetails::OnChangeReplication(TSharedPtr<FString> ItemSel
 				if (!FuncGraph)
 				{
 					FuncGraph = FBlueprintEditorUtils::CreateNewGraph(GetBlueprintObj(), FName(*NewFuncName), UEdGraph::StaticClass(), UEdGraphSchema_K2::StaticClass());
-					FBlueprintEditorUtils::AddFunctionGraph(GetBlueprintObj(), FuncGraph, false, NULL);
+					FBlueprintEditorUtils::AddFunctionGraph<UClass>(GetBlueprintObj(), FuncGraph, false, NULL);
 				}
 
 				if (FuncGraph)
diff --git a/Engine/Source/Editor/Kismet/Private/BlueprintEditor.cpp b/Engine/Source/Editor/Kismet/Private/BlueprintEditor.cpp
--- a/Engine/Source/Editor/Kismet/Private/BlueprintEditor.cpp
+++ b/Engine/Source/Editor/Kismet/Private/BlueprintEditor.cpp
@@ -3333,7 +3333,7 @@ bool FBlueprintEditor::CanCollapseSelectionToFunction(TSet<class UEdGraphNode*>&
 
 	// Create a function graph
 	UEdGraph* FunctionGraph = FBlueprintEditorUtils::CreateNewGraph(GetBlueprintObj(), FBlueprintEditorUtils::FindUniqueKismetName(GetBlueprintObj(), TEXT("TempGraph")), UEdGraph::StaticClass(), UEdGraphSchema_K2::StaticClass());
-	FBlueprintEditorUtils::AddFunctionGraph(GetBlueprintObj(), FunctionGraph, /*bIsUserCreated=*/ true, NULL);
+	FBlueprintEditorUtils::AddFunctionGraph<UClass>(GetBlueprintObj(), FunctionGraph, /*bIsUserCreated=*/ true, NULL);
 
 	const UEdGraphSchema_K2* K2Schema = GetDefault<UEdGraphSchema_K2>();
 
@@ -4984,7 +4984,7 @@ UEdGraph* FBlueprintEditor::CollapseSelectionToFunction(TSet<class UEdGraphNode*
 	FName DocumentName = FBlueprintEditorUtils::FindUniqueKismetName(GetBlueprintObj(), TEXT("NewFunction"));
 
 	NewGraph = FBlueprintEditorUtils::CreateNewGraph(GetBlueprintObj(), DocumentName, UEdGraph::StaticClass(), UEdGraphSchema_K2::StaticClass());
-	FBlueprintEditorUtils::AddFunctionGraph(GetBlueprintObj(), NewGraph, /*bIsUserCreated=*/ true, NULL);
+	FBlueprintEditorUtils::AddFunctionGraph<UClass>(GetBlueprintObj(), NewGraph, /*bIsUserCreated=*/ true, NULL);
 
 	TArray<UK2Node_FunctionEntry*> EntryNodes;
 	NewGraph->GetNodesOfClass(EntryNodes);
@@ -5385,7 +5385,7 @@ void FBlueprintEditor::OnAddNewDelegate()
 	NewGraph->bEditable = false;
 
 	K2Schema->CreateDefaultNodesForGraph(*NewGraph);
-	K2Schema->CreateFunctionGraphTerminators(*NewGraph, NULL);
+	K2Schema->CreateFunctionGraphTerminators(*NewGraph, (UClass*)NULL);
 	K2Schema->AddExtraFunctionFlags(NewGraph, (FUNC_BlueprintCallable|FUNC_BlueprintEvent|FUNC_Public));
 	K2Schema->MarkFunctionEntryAsEditable(NewGraph, true);
 
@@ -5449,7 +5449,7 @@ void FBlueprintEditor::NewDocument_OnClicked(ECreatedDocumentType GraphType)
 	if (GraphType == CGT_NewFunctionGraph)
 	{
 		NewGraph = FBlueprintEditorUtils::CreateNewGraph(GetBlueprintObj(), DocumentName, UEdGraph::StaticClass(), UEdGraphSchema_K2::StaticClass());
-		FBlueprintEditorUtils::AddFunctionGraph(GetBlueprintObj(), NewGraph, /*bIsUserCreated=*/ true, NULL);
+		FBlueprintEditorUtils::AddFunctionGraph<UClass>(GetBlueprintObj(), NewGraph, /*bIsUserCreated=*/ true, NULL);
 	}
 	else if (GraphType == CGT_NewMacroGraph)
 	{
diff --git a/Engine/Source/Editor/UMGEditor/Private/DetailCustomizations.cpp b/Engine/Source/Editor/UMGEditor/Private/DetailCustomizations.cpp
--- a/Engine/Source/Editor/UMGEditor/Private/DetailCustomizations.cpp
+++ b/Engine/Source/Editor/UMGEditor/Private/DetailCustomizations.cpp
@@ -63,7 +63,7 @@ void FBlueprintWidgetCustomization::CustomizeDetails( IDetailLayoutBuilder& Deta
 				PropertyRow.GetDefaultWidgets(NameWidget, ValueWidget, Row);
 
 				TSharedRef<SComboButton> DelegateButton = SNew(SComboButton)
-				.OnGetMenuContent(this, &FBlueprintWidgetCustomization::OnGenerateDelegateMenu, ConstPropertyHandle)
+				.OnGetMenuContent(this, &FBlueprintWidgetCustomization::OnGenerateDelegateMenu, ConstPropertyHandle, Property->SignatureFunction)
 				.ContentPadding(0)
 				.ButtonContent()
 				[
@@ -104,14 +104,15 @@ void FBlueprintWidgetCustomization::CustomizeDetails( IDetailLayoutBuilder& Deta
 	}
 }
 
-void FBlueprintWidgetCustomization::RefreshBlueprintFunctionCache()
+void FBlueprintWidgetCustomization::RefreshBlueprintFunctionCache(const UFunction* DelegateSignature)
 {
 	const UWidgetGraphSchema* Schema = GetDefault<UWidgetGraphSchema>();
 	const FSlateFontInfo DetailFontInfo = IDetailLayoutBuilder::GetDetailFont();
 
 	BlueprintFunctionCache.Reset();
 
-	//TODO UMG Filter the list of functions to only be functions matching the signature.
+	// Get the current skeleton class, think header for the blueprint.
+	UBlueprintGeneratedClass* SkeletonClass = Cast<UBlueprintGeneratedClass>(Blueprint->SkeletonGeneratedClass);
 
 	// Grab functions implemented by the blueprint
 	for ( int32 i = 0; i < Blueprint->FunctionGraphs.Num(); i++ )
@@ -125,37 +126,62 @@ void FBlueprintWidgetCustomization::RefreshBlueprintFunctionCache()
 			continue;
 		}
 
-		FGraphDisplayInfo DisplayInfo;
-		Graph->GetSchema()->GetGraphDisplayInformation(*Graph, DisplayInfo);
+		FName FunctionFName = Graph->GetFName();
 
-		TSharedPtr<FunctionInfo> NewFuncAction = MakeShareable(new FunctionInfo());
-		NewFuncAction->DisplayName = DisplayInfo.PlainName;
-		NewFuncAction->Tooltip = DisplayInfo.Tooltip;
-		NewFuncAction->FuncName = Graph->GetFName();
-		NewFuncAction->EdGraph = Graph;
+		UFunction* Function = SkeletonClass->FindFunctionByName(FunctionFName, EIncludeSuperFlag::IncludeSuper);
+		if ( Function == NULL )
+		{
+			continue;
+		}
+
+		// We ignore CPF_ReturnParm because all that matters for binding to script functions is that the number of out parameters match.
+		if ( Function->IsSignatureCompatibleWith(DelegateSignature, UFunction::GetDefaultIgnoredSignatureCompatibilityFlags() | CPF_ReturnParm) )
+		{
+			FGraphDisplayInfo DisplayInfo;
+			Graph->GetSchema()->GetGraphDisplayInformation(*Graph, DisplayInfo);
+
+			TSharedPtr<FunctionInfo> NewFuncAction = MakeShareable(new FunctionInfo());
+			NewFuncAction->DisplayName = DisplayInfo.PlainName;
+			NewFuncAction->Tooltip = DisplayInfo.Tooltip;
+			NewFuncAction->FuncName = FunctionFName;
+			NewFuncAction->EdGraph = Graph;
 
-		BlueprintFunctionCache.Add(NewFuncAction);
+			BlueprintFunctionCache.Add(NewFuncAction);
+		}
 	}
 }
 
-TSharedRef<SWidget> FBlueprintWidgetCustomization::OnGenerateDelegateMenu(TSharedRef<IPropertyHandle> PropertyHandle)
+TSharedRef<SWidget> FBlueprintWidgetCustomization::OnGenerateDelegateMenu(TSharedRef<IPropertyHandle> PropertyHandle, UFunction* DelegateSignature)
 {
-	RefreshBlueprintFunctionCache();
+	RefreshBlueprintFunctionCache(DelegateSignature);
 
 	const bool bInShouldCloseWindowAfterMenuSelection = true;
 	FMenuBuilder MenuBuilder(bInShouldCloseWindowAfterMenuSelection, NULL);
 
+	//TODO UMG Enable or disable Remove Binding if needed.
+
 	MenuBuilder.BeginSection("ClearBinding");
 	{
 		MenuBuilder.AddMenuEntry(
-			LOCTEXT("FilterListResetFilters", "Remove Binding"),
-			LOCTEXT("FilterListResetToolTip", "Removes the current binding"),
+			LOCTEXT("RemoveBinding", "Remove Binding"),
+			LOCTEXT("RemoveBindingToolTip", "Removes the current binding"),
 			FSlateIcon(),
 			FUIAction(FExecuteAction::CreateSP(this, &FBlueprintWidgetCustomization::HandleRemoveBinding, PropertyHandle))
 			);
 	}
 	MenuBuilder.EndSection(); //ClearBinding
 
+	MenuBuilder.BeginSection("CreateBinding");
+	{
+		MenuBuilder.AddMenuEntry(
+			LOCTEXT("CreateBinding", "Create Binding"),
+			LOCTEXT("CreateBindingToolTip", "Creates a new function for this property to be bound to"),
+			FSlateIcon(),
+			FUIAction(FExecuteAction::CreateSP(this, &FBlueprintWidgetCustomization::HandleCreateAndAddBinding, PropertyHandle, DelegateSignature))
+			);
+	}
+	MenuBuilder.EndSection(); //CreateBinding
+
 	MenuBuilder.BeginSection("Functions");
 	{
 		for ( TSharedPtr<FunctionInfo>& Function : BlueprintFunctionCache )
@@ -257,4 +283,31 @@ void FBlueprintWidgetCustomization::HandleAddBinding(TSharedRef<IPropertyHandle>
 	FBlueprintEditorUtils::MarkBlueprintAsStructurallyModified(Blueprint);
 }
 
+void FBlueprintWidgetCustomization::HandleCreateAndAddBinding(TSharedRef<IPropertyHandle> PropertyHandle, UFunction* DelegateSignature)
+{
+	const FScopedTransaction Transaction(LOCTEXT("CreateDelegate", "Create Binding"));
+
+	Blueprint->Modify();
+
+	// Create the function graph.
+	FString FunctionName = FString(TEXT("Get")) + PropertyHandle->GetProperty()->GetName();
+	UEdGraph* FunctionGraph = FBlueprintEditorUtils::CreateNewGraph(Blueprint, FBlueprintEditorUtils::FindUniqueKismetName(Blueprint, FunctionName), UEdGraph::StaticClass(), UEdGraphSchema_K2::StaticClass());
+
+	bool bUserCreated = true;
+	FBlueprintEditorUtils::AddFunctionGraph(Blueprint, FunctionGraph, bUserCreated, DelegateSignature);
+
+	//TODO UMG Once we start binding event functions, pure needs to be variable, so maybe pass it in.
+	
+	// All bound delegates should be pure.
+	const UEdGraphSchema_K2* Schema_K2 = Cast<UEdGraphSchema_K2>(FunctionGraph->GetSchema());
+	Schema_K2->AddExtraFunctionFlags(FunctionGraph, FUNC_BlueprintPure);
+
+	// Add the binding to the blueprint
+	TSharedPtr<FunctionInfo> SelectedFunction = MakeShareable(new FunctionInfo());
+	SelectedFunction->FuncName = FunctionGraph->GetFName();
+	SelectedFunction->EdGraph = FunctionGraph;
+
+	HandleAddBinding(PropertyHandle, SelectedFunction);
+}
+
 #undef LOCTEXT_NAMESPACE
diff --git a/Engine/Source/Editor/UMGEditor/Private/DetailCustomizations.h b/Engine/Source/Editor/UMGEditor/Private/DetailCustomizations.h
--- a/Engine/Source/Editor/UMGEditor/Private/DetailCustomizations.h
+++ b/Engine/Source/Editor/UMGEditor/Private/DetailCustomizations.h
@@ -29,13 +29,15 @@ private:
 		UEdGraph* EdGraph;
 	};
 
-	void RefreshBlueprintFunctionCache();
+	void RefreshBlueprintFunctionCache(const UFunction* DelegateSignature);
 
-	TSharedRef<SWidget> OnGenerateDelegateMenu(TSharedRef<IPropertyHandle> PropertyHandle);
+	TSharedRef<SWidget> OnGenerateDelegateMenu(TSharedRef<IPropertyHandle> PropertyHandle, UFunction* DelegateSignature);
 
 	FText GetCurrentBindingText(TSharedRef<IPropertyHandle> PropertyHandle) const;
+
 	void HandleRemoveBinding(TSharedRef<IPropertyHandle> PropertyHandle);
 	void HandleAddBinding(TSharedRef<IPropertyHandle> PropertyHandle, TSharedPtr<FunctionInfo> SelectedFunction);
+	void HandleCreateAndAddBinding(TSharedRef<IPropertyHandle> PropertyHandle, UFunction* DelegateSignature);
 
 private:
 
diff --git a/Engine/Source/Editor/UnrealEd/Private/Commandlets/CookCommandlet.cpp b/Engine/Source/Editor/UnrealEd/Private/Commandlets/CookCommandlet.cpp
--- a/Engine/Source/Editor/UnrealEd/Private/Commandlets/CookCommandlet.cpp
+++ b/Engine/Source/Editor/UnrealEd/Private/Commandlets/CookCommandlet.cpp
@@ -442,7 +442,8 @@ bool UCookCommandlet::SaveCookedPackage( UPackage* Package, uint32 SaveFlags, bo
 					}
 				}
 
-				if( PlatFilename.Len() >= PLATFORM_MAX_FILEPATH_LENGTH )
+				const FString FullFilename = FPaths::ConvertRelativePathToFull( PlatFilename );
+				if( FullFilename.Len() >= PLATFORM_MAX_FILEPATH_LENGTH )
 				{
 					UE_LOG( LogCookCommandlet, Error, TEXT( "Couldn't save package, filename is too long :%s" ), *PlatFilename );
 					bSavedCorrectly = false;
diff --git a/Engine/Source/Editor/UnrealEd/Private/CookOnTheFlyServer.cpp b/Engine/Source/Editor/UnrealEd/Private/CookOnTheFlyServer.cpp
--- a/Engine/Source/Editor/UnrealEd/Private/CookOnTheFlyServer.cpp
+++ b/Engine/Source/Editor/UnrealEd/Private/CookOnTheFlyServer.cpp
@@ -638,7 +638,8 @@ bool UCookOnTheFlyServer::SaveCookedPackage( UPackage* Package, uint32 SaveFlags
 					World->PersistentLevel->OwningWorld = World;
 				}
 
-				if( PlatFilename.Len() >= PLATFORM_MAX_FILEPATH_LENGTH )
+				const FString FullFilename = FPaths::ConvertRelativePathToFull( PlatFilename );
+				if( FullFilename.Len() >= PLATFORM_MAX_FILEPATH_LENGTH )
 				{
 					UE_LOG( LogCookOnTheFly, Error, TEXT( "Couldn't save package, filename is too long :%s" ), *PlatFilename );
 					bSavedCorrectly = false;
diff --git a/Engine/Source/Editor/UnrealEd/Private/Kismet2/BlueprintEditorUtils.cpp b/Engine/Source/Editor/UnrealEd/Private/Kismet2/BlueprintEditorUtils.cpp
--- a/Engine/Source/Editor/UnrealEd/Private/Kismet2/BlueprintEditorUtils.cpp
+++ b/Engine/Source/Editor/UnrealEd/Private/Kismet2/BlueprintEditorUtils.cpp
@@ -1541,40 +1541,6 @@ UEdGraph* FBlueprintEditorUtils::CreateNewGraph(UObject* ParentScope, const FNam
 	return NewGraph;
 }
 
-
-void FBlueprintEditorUtils::AddFunctionGraph(UBlueprint* Blueprint, class UEdGraph* Graph, bool bIsUserCreated, UClass* SignatureFromClass)
-{
-	// Give the schema a chance to fill out any required nodes (like the entry node or results node)
-	const UEdGraphSchema* Schema = Graph->GetSchema();
-	const UEdGraphSchema_K2* K2Schema = Cast<const UEdGraphSchema_K2>(Graph->GetSchema());
-
-	Schema->CreateDefaultNodesForGraph(*Graph);
-
-	if (K2Schema != NULL)
-	{
-		K2Schema->CreateFunctionGraphTerminators(*Graph, SignatureFromClass);
-
-		if (bIsUserCreated)
-		{
-			// We need to flag the entry node to make sure that the compiled function is callable from Kismet2
-			int32 ExtraFunctionFlags = (FUNC_BlueprintCallable | FUNC_BlueprintEvent | FUNC_Public);
-			if (BPTYPE_FunctionLibrary == Blueprint->BlueprintType)
-			{
-				ExtraFunctionFlags |= FUNC_Static;
-			}
-			K2Schema->AddExtraFunctionFlags(Graph, ExtraFunctionFlags);
-			K2Schema->MarkFunctionEntryAsEditable(Graph, true);
-		}
-	}
-
-	Blueprint->FunctionGraphs.Add(Graph);
-
-	// Potentially adjust variable names for any child blueprints
-	ValidateBlueprintChildVariables(Blueprint, Graph->GetFName());
-
-	FBlueprintEditorUtils::MarkBlueprintAsStructurallyModified(Blueprint);
-}
-
 UFunction* FBlueprintEditorUtils::FindFunctionInImplementedInterfaces(const UBlueprint* Blueprint, const FName& FunctionName, bool * bOutInvalidInterface )
 {
 	if(Blueprint)
@@ -2047,9 +2013,9 @@ void FBlueprintEditorUtils::GetDependentBlueprints(UBlueprint* Blueprint, TArray
 
 	//@TODO PRETEST MERGE:  Swap this when main is merged to pretest!
 #if 0
- 	TArray<UObject*> AllBlueprints;
- 	bool const bIncludeDerivedClasses = true;
- 	GetObjectsOfClass(UBlueprint::StaticClass(), AllBlueprints, bIncludeDerivedClasses );
+	TArray<UObject*> AllBlueprints;
+	bool const bIncludeDerivedClasses = true;
+	GetObjectsOfClass(UBlueprint::StaticClass(), AllBlueprints, bIncludeDerivedClasses );
 
 	for ( auto ObjIt = AllBlueprints.CreateConstIterator(); ObjIt; ++ObjIt )
 #else
diff --git a/Engine/Source/Editor/UnrealEd/Public/GameModeInfoCustomizer.h b/Engine/Source/Editor/UnrealEd/Public/GameModeInfoCustomizer.h
--- a/Engine/Source/Editor/UnrealEd/Public/GameModeInfoCustomizer.h
+++ b/Engine/Source/Editor/UnrealEd/Public/GameModeInfoCustomizer.h
@@ -174,7 +174,10 @@ public:
 			return NULL;
 		}
 
-		const UClass* GameModeClass = FindObject<UClass>(ANY_PACKAGE, *ClassName);
+		FString StrippedClassName = ClassName;
+		ConstructorHelpers::StripObjectClass(StrippedClassName);
+
+		const UClass* GameModeClass = FindObject<UClass>(ANY_PACKAGE, *StrippedClassName);
 		if (!GameModeClass)
 		{
 			GameModeClass = LoadObject<UClass>(NULL, *ClassName);
diff --git a/Engine/Source/Editor/UnrealEd/Public/Kismet2/BlueprintEditorUtils.h b/Engine/Source/Editor/UnrealEd/Public/Kismet2/BlueprintEditorUtils.h
--- a/Engine/Source/Editor/UnrealEd/Public/Kismet2/BlueprintEditorUtils.h
+++ b/Engine/Source/Editor/UnrealEd/Public/Kismet2/BlueprintEditorUtils.h
@@ -149,8 +149,43 @@ public:
 	 */
 	static class UEdGraph* CreateNewGraph(UObject* ParentScope, const FName& GraphName, TSubclassOf<class UEdGraph> GraphClass, TSubclassOf<class UEdGraphSchema> SchemaClass);
 
-	/** Adds a function graph to this blueprint.  If bIsUserCreated is true, the entry/exit nodes will be editable. SignatureFromClass is used to find signature for entry/exit nodes if using an existing signature. */
-	static void AddFunctionGraph(UBlueprint* Blueprint, class UEdGraph* Graph,  bool bIsUserCreated, UClass* SignatureFromClass);
+	/** 
+	 * Adds a function graph to this blueprint.  If bIsUserCreated is true, the entry/exit nodes will be editable. SignatureFromObject is used to find signature for entry/exit nodes if using an existing signature.
+	 * The template argument SignatureType should be UClass or UFunction.
+	 */
+	template <typename SignatureType>
+	static void AddFunctionGraph(UBlueprint* Blueprint, class UEdGraph* Graph, bool bIsUserCreated, SignatureType* SignatureFromObject)
+	{
+		// Give the schema a chance to fill out any required nodes (like the entry node or results node)
+		const UEdGraphSchema* Schema = Graph->GetSchema();
+		const UEdGraphSchema_K2* K2Schema = Cast<const UEdGraphSchema_K2>(Graph->GetSchema());
+
+		Schema->CreateDefaultNodesForGraph(*Graph);
+
+		if ( K2Schema != NULL )
+		{
+			K2Schema->CreateFunctionGraphTerminators(*Graph, SignatureFromObject);
+
+			if ( bIsUserCreated )
+			{
+				// We need to flag the entry node to make sure that the compiled function is callable from Kismet2
+				int32 ExtraFunctionFlags = ( FUNC_BlueprintCallable | FUNC_BlueprintEvent | FUNC_Public );
+				if ( BPTYPE_FunctionLibrary == Blueprint->BlueprintType )
+				{
+					ExtraFunctionFlags |= FUNC_Static;
+				}
+				K2Schema->AddExtraFunctionFlags(Graph, ExtraFunctionFlags);
+				K2Schema->MarkFunctionEntryAsEditable(Graph, true);
+			}
+		}
+
+		Blueprint->FunctionGraphs.Add(Graph);
+
+		// Potentially adjust variable names for any child blueprints
+		ValidateBlueprintChildVariables(Blueprint, Graph->GetFName());
+
+		FBlueprintEditorUtils::MarkBlueprintAsStructurallyModified(Blueprint);
+	}
 
 	/** Adds a macro graph to this blueprint.  If bIsUserCreated is true, the entry/exit nodes will be editable. SignatureFromClass is used to find signature for entry/exit nodes if using an existing signature. */
 	static void AddMacroGraph(UBlueprint* Blueprint, class UEdGraph* Graph,  bool bIsUserCreated, UClass* SignatureFromClass);
diff --git a/Engine/Source/Runtime/Core/Private/Mac/MacApplication.cpp b/Engine/Source/Runtime/Core/Private/Mac/MacApplication.cpp
--- a/Engine/Source/Runtime/Core/Private/Mac/MacApplication.cpp
+++ b/Engine/Source/Runtime/Core/Private/Mac/MacApplication.cpp
@@ -222,7 +222,7 @@ void FMacApplication::ProcessEvent( NSEvent* Event )
 
 	FSlateCocoaWindow* NativeWindow = FindEventWindow( Event );
 	TSharedPtr< FMacWindow > CurrentEventWindow = FindWindowByNSWindow( Windows, NativeWindow );
-	if( !CurrentEventWindow.IsValid() && LastEventWindow.IsValid() )
+	if( NativeWindow && !CurrentEventWindow.IsValid() && LastEventWindow.IsValid() )
 	{
 		CurrentEventWindow = LastEventWindow;
 		NativeWindow = CurrentEventWindow->GetWindowHandle();
@@ -587,6 +587,16 @@ FSlateCocoaWindow* FMacApplication::FindEventWindow( NSEvent* Event )
 
 	FSlateCocoaWindow* EventWindow = (FSlateCocoaWindow*)[Event window];
 
+	if ([Event type] == NSMouseMoved)
+	{
+		// Ignore windows owned by other applications
+		NSInteger WindowNumber = [NSWindow windowNumberAtPoint:[NSEvent mouseLocation] belowWindowWithWindowNumber:0];
+		if ([NSApp windowWithWindowNumber:WindowNumber] == NULL)
+		{
+			return NULL;
+		}
+	}
+
 	if( IsMouseEvent )
 	{
 		if( DraggedWindow )
diff --git a/Engine/Source/Runtime/CoreUObject/Private/UObject/Class.cpp b/Engine/Source/Runtime/CoreUObject/Private/UObject/Class.cpp
--- a/Engine/Source/Runtime/CoreUObject/Private/UObject/Class.cpp
+++ b/Engine/Source/Runtime/CoreUObject/Private/UObject/Class.cpp
@@ -3666,6 +3666,14 @@ void UFunction::Link(FArchive& Ar, bool bRelinkExistingProperties)
 }
 
 bool UFunction::IsSignatureCompatibleWith(const UFunction* OtherFunction) const
+{
+	// The script compiler pointlessly sets CPF_EditInline on components passed as arguments
+	const uint64 IgnoreFlags = UFunction::GetDefaultIgnoredSignatureCompatibilityFlags();
+
+	return IsSignatureCompatibleWith(OtherFunction, IgnoreFlags);
+}
+
+bool UFunction::IsSignatureCompatibleWith(const UFunction* OtherFunction, uint64 IgnoreFlags) const
 {
 	// Early out if they're exactly the same function
 	if (this == OtherFunction)
@@ -3673,9 +3681,6 @@ bool UFunction::IsSignatureCompatibleWith(const UFunction* OtherFunction) const
 		return true;
 	}
 
-	// The script compiler pointlessly sets CPF_EditInline on components passed as arguments
-	const uint64 IgnoreFlags = CPF_EditInline | CPF_ExportObject | CPF_InstancedReference | CPF_ContainsInstancedReference | CPF_ComputedFlags | CPF_ConstParm;//@TODO: UCREMOVAL: CPF_ConstParm added as a hack to get blueprints compiling with a const DamageType parameter.
-
 	// Run thru the parameter property chains to compare each property
 	TFieldIterator<UProperty> IteratorA(this);
 	TFieldIterator<UProperty> IteratorB(OtherFunction);
diff --git a/Engine/Source/Runtime/CoreUObject/Public/UObject/Class.h b/Engine/Source/Runtime/CoreUObject/Public/UObject/Class.h
--- a/Engine/Source/Runtime/CoreUObject/Public/UObject/Class.h
+++ b/Engine/Source/Runtime/CoreUObject/Public/UObject/Class.h
@@ -139,7 +139,7 @@ class COREUOBJECT_API UField : public UObject
 	
 	/**
 	 * Find the metadata value associated with the key
- 	 * and return bool  
+	 * and return bool  
 	 * @param Key The key to lookup in the metadata
 	 * @return return true if the value was true (case insensitive)
 	 */
@@ -1140,6 +1140,7 @@ public:
 	{
 		return (FunctionFlags&FlagsToCheck) != 0 || FlagsToCheck == FUNC_AllFlags;
 	}
+
 	/**
 	 * Used to safely check whether all of the passed in flags are set.
 	 *
@@ -1151,6 +1152,16 @@ public:
 		return ((FunctionFlags & FlagsToCheck) == FlagsToCheck);
 	}
 
+	/**
+	 * Returns the flags that are ignored by default when comparing function signatures.
+	 */
+	FORCEINLINE static uint64 GetDefaultIgnoredSignatureCompatibilityFlags()
+	{
+		//@TODO: UCREMOVAL: CPF_ConstParm added as a hack to get blueprints compiling with a const DamageType parameter.
+		const uint64 IgnoreFlags = CPF_EditInline | CPF_ExportObject | CPF_InstancedReference | CPF_ContainsInstancedReference | CPF_ComputedFlags | CPF_ConstParm;
+		return IgnoreFlags;
+	}
+
 	/**
 	 * Determines if two functions have an identical signature (note: currently doesn't allow
 	 * matches with class parameters that differ only in how derived they are; there is no
@@ -1161,6 +1172,18 @@ public:
 	 * @return	true if function signatures are compatible.
 	 */
 	COREUOBJECT_API bool IsSignatureCompatibleWith(const UFunction* OtherFunction) const;
+
+	/**
+	 * Determines if two functions have an identical signature (note: currently doesn't allow
+	 * matches with class parameters that differ only in how derived they are; there is no
+	 * directionality to the call)
+	 *
+	 * @param	OtherFunction	Function to compare this function against.
+	 * @param   IgnoreFlags     Custom flags to ignore when comparing parameters between the functions.
+	 *
+	 * @return	true if function signatures are compatible.
+	 */
+	COREUOBJECT_API bool IsSignatureCompatibleWith(const UFunction* OtherFunction, uint64 IgnoreFlags) const;
 };
 
 /*-----------------------------------------------------------------------------
@@ -1766,7 +1789,7 @@ public:
 	 * default if it doesn't exist yet.
 	 *
 	 * @return	name of the class specific ini file
- 	 */
+	 */
 	const FString GetConfigName() const;
 
 	UClass* GetSuperClass() const
diff --git a/Engine/Source/Runtime/Engine/Classes/GameFramework/HUD.h b/Engine/Source/Runtime/Engine/Classes/GameFramework/HUD.h
--- a/Engine/Source/Runtime/Engine/Classes/GameFramework/HUD.h
+++ b/Engine/Source/Runtime/Engine/Classes/GameFramework/HUD.h
@@ -119,7 +119,7 @@ public:
 	 * @param	InCanvas		Canvas on which to render.
 	 * @param	InColor			Color to render the box.
 	 */
-	void Draw( class FCanvas* InCanvas, const FLinearColor& InColor );
+	void Draw( class FCanvas* InCanvas, const FLinearColor& InColor ) const;
 	
 	/** Get the name of this hitbox.  */
 	const FName& GetName() const { return Name;};
@@ -512,7 +512,7 @@ public:
 	/** Array of hitboxes that have been hit for this frame. */
 	TArray< class FHUDHitBox* >	HitBoxHits;
 
-	/** Set of hitbox (by name) that are currently moused over */
+	/** Set of hitbox (by name) that are currently moused over or have a touch contacting them */
 	TSet< FName > HitBoxesOver;
 
 	/**
@@ -525,30 +525,38 @@ public:
 	 * Update the list of hitboxes and dispatch events for any hits.
 	 * @param	InEventType	Input event that triggered the call.
 	 */
-	bool UpdateAndDispatchHitBoxClickEvents(EInputEvent InEventType);
+	bool UpdateAndDispatchHitBoxClickEvents(const FVector2D ClickLocation, const EInputEvent InEventType, const bool bDispatchOverOutEvent);
 	
 	/**
 	 * Update a the list of hitboxes that have been hit this frame.
 	 * @param	Canvas	Canvas on which to render debug boxes.
 	 */
-	void UpdateHitBoxCandidates( const FVector2D& InHitLocation );
+	void UpdateHitBoxCandidates( TArray<FVector2D> InContactPoints );
 	
 	/** Have any hitboxes been hit this frame. */
 	bool AnyCurrentHitBoxHits() const;	
 
 	/**
 	 * Find the first hitbox containing the given coordinates.
-	 * @param	InHitLocation Coordinates to check
-	 * @return returns the hitbox at the given coordinates or NULL if none match.
+	 * @param	InHitLocation		Coordinates to check
+	 * @param	bConsumingInput		If true will return the first hitbox that would consume input at this coordinate
+	 * @return	returns the hitbox at the given coordinates or NULL if none match.
 	 */
-	const FHUDHitBox* GetHitBoxAtCoordinates( const FVector2D& InHitLocation ) const;
+	const FHUDHitBox* GetHitBoxAtCoordinates( FVector2D InHitLocation, bool bConsumingInput = false ) const;
 
+	/**
+	* Finds all the hitboxes containing the given coordinates.
+	* @param	InHitLocation		Coordinates to check
+	* @param	OutHitBoxes			Reference parameter with hit box references at the coordinates
+	*/
+	void GetHitBoxesAtCoordinates(FVector2D InHitLocation, TArray<const FHUDHitBox*>& OutHitBoxes) const;
+	
 	/**
 	 * Return the hitbox with the given name
 	 * @param	InName	Name of required hitbox
 	 * @return returns the hitbox with the given name NULL if none match.
 	 */
-	const FHUDHitBox* GetHitBoxWithName( const FName& InName ) const;
+	const FHUDHitBox* GetHitBoxWithName( const FName InName ) const;
 
 	//=============================================================================
 	// UIBlur overrides
@@ -566,6 +574,10 @@ public:
 
 protected:
 	bool IsCanvasValid_WarnIfNot() const;
+
+private:
+	// Helper function to deal with screen offset and splitscreen mapping of coordinates to HUD
+	FVector2D GetCoordinateOffset() const;
 };
 
 
diff --git a/Engine/Source/Runtime/Engine/Private/HUD.cpp b/Engine/Source/Runtime/Engine/Private/HUD.cpp
--- a/Engine/Source/Runtime/Engine/Private/HUD.cpp
+++ b/Engine/Source/Runtime/Engine/Private/HUD.cpp
@@ -69,6 +69,39 @@ void AHUD::NotifyBindPostProcessEffects()
 	// overload with custom code e.g. getting material pointers to control effects for gameplay.
 }
 
+FVector2D AHUD::GetCoordinateOffset() const
+{
+	FVector2D Offset;
+
+	ULocalPlayer* LocalPlayer = Cast<ULocalPlayer>(GetOwningPlayerController()->Player);
+
+	if (LocalPlayer)
+	{
+		// Create a view family for the game viewport
+		FSceneViewFamilyContext ViewFamily(FSceneViewFamily::ConstructionValues(
+			LocalPlayer->ViewportClient->Viewport,
+			GetWorld()->Scene,
+			LocalPlayer->ViewportClient->EngineShowFlags)
+			.SetRealtimeUpdate(true));
+
+		// Calculate a view where the player is to update the streaming from the players start location
+		FVector ViewLocation;
+		FRotator ViewRotation;
+		FSceneView* SceneView = LocalPlayer->CalcSceneView(&ViewFamily, /*out*/ ViewLocation, /*out*/ ViewRotation, LocalPlayer->ViewportClient->Viewport);
+
+		if (SceneView)
+		{
+			Offset.X = (SceneView->ViewRect.Min.X - SceneView->UnconstrainedViewRect.Min.X) // This accounts for the borders when the aspect ratio is locked
+							- SceneView->UnconstrainedViewRect.Min.X;						// And this will deal with the viewport offset if its a split screen
+
+			Offset.Y = (SceneView->ViewRect.Min.Y - SceneView->UnconstrainedViewRect.Min.Y)
+							- SceneView->UnconstrainedViewRect.Min.Y;
+		}
+	}
+
+	return Offset;
+}
+
 void AHUD::PostRender()
 {
 	// Set up delta time
@@ -103,30 +136,32 @@ void AHUD::PostRender()
 
 		if (LocalPlayer)
 		{
-			FVector2D MousePosition = LocalPlayer->ViewportClient->GetMousePosition();
+			TArray<FVector2D> ContactPoints;
+			ContactPoints.Add(LocalPlayer->ViewportClient->GetMousePosition());
 
-			// Create a view family for the game viewport
-			FSceneViewFamilyContext ViewFamily( FSceneViewFamily::ConstructionValues(
-				LocalPlayer->ViewportClient->Viewport,
-				GetWorld()->Scene,
-				LocalPlayer->ViewportClient->EngineShowFlags )
-				.SetRealtimeUpdate(true) );
+			for (int32 FingerIndex = 0; FingerIndex < EKeys::NUM_TOUCH_KEYS; ++FingerIndex)
+			{
+				FVector2D TouchLocation;
+				bool bPressed = false;
+
+				GetOwningPlayerController()->GetInputTouchState((ETouchIndex::Type)FingerIndex, TouchLocation.X, TouchLocation.Y, bPressed);
+
+				if (bPressed)
+				{
+					ContactPoints.Add(TouchLocation);
+				}
+			}
 
-			// Calculate a view where the player is to update the streaming from the players start location
-			FVector ViewLocation;
-			FRotator ViewRotation;
-			FSceneView* SceneView = LocalPlayer->CalcSceneView( &ViewFamily, /*out*/ ViewLocation, /*out*/ ViewRotation, LocalPlayer->ViewportClient->Viewport );
+			FVector2D ContactPointOffset = GetCoordinateOffset();
 
-			if (SceneView)
+			if (!ContactPointOffset.IsZero())
 			{
-				// This accounts for the borders when the aspect ratio is locked
-				MousePosition.X -= (SceneView->ViewRect.Min.X - SceneView->UnconstrainedViewRect.Min.X);
-				MousePosition.Y -= (SceneView->ViewRect.Min.Y - SceneView->UnconstrainedViewRect.Min.Y);
-				// And this will deal with the viewport offset if its a split screen
-				MousePosition.X -= SceneView->UnconstrainedViewRect.Min.X;
-				MousePosition.Y -= SceneView->UnconstrainedViewRect.Min.Y;
+				for (FVector2D& ContactPoint : ContactPoints)
+				{
+					ContactPoint -= ContactPointOffset;
+				}
 			}
-			UpdateHitBoxCandidates( MousePosition );
+			UpdateHitBoxCandidates( ContactPoints );
 		}
 	}
 	
@@ -628,42 +663,54 @@ APawn* AHUD::GetOwningPawn() const
 
 void AHUD::RenderHitBoxes( FCanvas* InCanvas )
 {
-	for (int32 iBox = 0; iBox < HitBoxMap.Num() ; iBox++)
+	for (const FHUDHitBox& HitBox : HitBoxMap)
 	{
 		FLinearColor BoxColor = FLinearColor::White;
-		if( HitBoxHits.Contains(&HitBoxMap[iBox]) )
+		if( HitBoxHits.Contains(const_cast<FHUDHitBox*>(&HitBox)))
 		{
 			BoxColor = FLinearColor::Red;
 		}
-		HitBoxMap[iBox].Draw( InCanvas, BoxColor );
+		HitBox.Draw( InCanvas, BoxColor );
 	}
 }
 
-void AHUD::UpdateHitBoxCandidates( const FVector2D& InHitLocation )
+void AHUD::UpdateHitBoxCandidates( TArray<FVector2D> InContactPoints )
 {
 	HitBoxHits.Empty();
-	for (int32 iBox = 0; iBox < HitBoxMap.Num() ; iBox++)
+	for (FHUDHitBox& HitBox : HitBoxMap)
 	{
-		if( HitBoxMap[iBox].Contains( InHitLocation ) )
+		bool bAdded = false;
+		for (int32 ContactPointIndex = InContactPoints.Num() - 1; ContactPointIndex >= 0; --ContactPointIndex)
+		{
+			if (HitBox.Contains(InContactPoints[ContactPointIndex]))
+			{
+				if (!bAdded)
+				{
+					HitBoxHits.Add(&HitBox);
+					bAdded = true;
+				}
+				if (HitBox.ConsumesInput())
+				{
+					InContactPoints.RemoveAtSwap(ContactPointIndex);
+				}
+				else
+				{
+					break;
+				}
+			}
+		}
+		if (InContactPoints.Num() == 0)
 		{
-			HitBoxHits.Add(&HitBoxMap[iBox]);
+			break;
 		}
 	}
-	struct FHitBoxSort
-	{
-		FORCEINLINE bool operator()( const FHUDHitBox& A, const FHUDHitBox& B ) const { return B.GetPriority() < A.GetPriority(); }
-	};
-
-	// Now sort the hits
-	HitBoxHits.Sort( FHitBoxSort() );
 
 	TSet<FName> NotOverHitBoxes = HitBoxesOver;
 	TArray<FName> NewlyOverHitBoxes;
 
 	// Now figure out which boxes we are over and deal with begin/end cursor over messages 
-	for (int32 HitIndex = 0; HitIndex < HitBoxHits.Num(); ++HitIndex )
+	for (FHUDHitBox* HitBox : HitBoxHits)
 	{
-		FHUDHitBox* HitBox = HitBoxHits[HitIndex];
 		const FName HitBoxName = HitBox->GetName();
 		if (HitBoxesOver.Contains(HitBoxName))
 		{
@@ -673,10 +720,6 @@ void AHUD::UpdateHitBoxCandidates( const FVector2D& InHitLocation )
 		{
 			NewlyOverHitBoxes.AddUnique(HitBoxName);
 		}
-		if (HitBox->ConsumesInput())
-		{
-			break; //Early out if this box blocks further consideration
-		}
 	}
 
 	// Dispatch the end cursor over messages
@@ -696,19 +739,35 @@ void AHUD::UpdateHitBoxCandidates( const FVector2D& InHitLocation )
 	}
 }
 
-const FHUDHitBox* AHUD::GetHitBoxAtCoordinates( const FVector2D& InHitLocation ) const
+const FHUDHitBox* AHUD::GetHitBoxAtCoordinates( FVector2D InHitLocation, const bool bIsConsumingInput ) const
 {
-	for (int32 iBox = 0; iBox < HitBoxMap.Num() ; iBox++)
+	InHitLocation -= GetCoordinateOffset();
+
+	for (const FHUDHitBox& HitBox : HitBoxMap)
 	{
-		if( HitBoxMap[iBox].Contains( InHitLocation ) )
+		if( (!bIsConsumingInput || HitBox.ConsumesInput()) && HitBox.Contains( InHitLocation ) )
 		{
-			return &HitBoxMap[iBox];
+			return &HitBox;
 		}
 	}
 	return NULL;
 }
 
-const FHUDHitBox* AHUD::GetHitBoxWithName( const FName& InName ) const
+void AHUD::GetHitBoxesAtCoordinates(FVector2D InHitLocation, TArray<const FHUDHitBox*>& OutHitBoxes) const
+{
+	InHitLocation -= GetCoordinateOffset();
+	OutHitBoxes.Empty();
+
+	for (const FHUDHitBox& HitBox : HitBoxMap)
+	{
+		if (HitBox.Contains(InHitLocation))
+		{
+			OutHitBoxes.Add(&HitBox);
+		}
+	}
+}
+
+const FHUDHitBox* AHUD::GetHitBoxWithName( const FName InName ) const
 {
 	for (int32 iBox = 0; iBox < HitBoxMap.Num() ; iBox++)
 	{
@@ -740,32 +799,61 @@ const TArray<FIntRect>& AHUD::GetUIBlurRectangles() const
 	return UIBlurOverrideRectangles;
 }
 
-bool AHUD::UpdateAndDispatchHitBoxClickEvents(EInputEvent InEventType)
+bool AHUD::UpdateAndDispatchHitBoxClickEvents(FVector2D ClickLocation, const EInputEvent InEventType, const bool bDispatchOverOutEvent)
 {
+	ClickLocation -= GetCoordinateOffset();
+
+	bool bHit = false;
 	for (FHUDHitBox* HitBoxHit : HitBoxHits)
 	{
-		if(InEventType == IE_Pressed || InEventType == IE_DoubleClick)
-		{
-			ReceiveHitBoxClick( HitBoxHit->GetName() );
-		}
-		else if(InEventType == IE_Released)
+		if (HitBoxHit->Contains(ClickLocation))
 		{
-			ReceiveHitBoxRelease( HitBoxHit->GetName() );
-		}
+			bHit = true;
 
-		if( HitBoxHit->ConsumesInput() == true )
-		{
-			break;	//Early out if this box consumed the click
+			if (InEventType == IE_Pressed || InEventType == IE_DoubleClick)
+			{
+				ReceiveHitBoxClick(HitBoxHit->GetName());
+				if (bDispatchOverOutEvent)
+				{
+					ReceiveHitBoxBeginCursorOver(HitBoxHit->GetName());
+				}
+			}
+			else if (InEventType == IE_Released)
+			{
+				ReceiveHitBoxRelease(HitBoxHit->GetName());
+				if (bDispatchOverOutEvent)
+				{
+					ReceiveHitBoxEndCursorOver(HitBoxHit->GetName());
+				}
+			}
+
+			if (HitBoxHit->ConsumesInput() == true)
+			{
+				break;	//Early out if this box consumed the click
+			}
 		}
 	}
-	return AnyCurrentHitBoxHits();
+	return bHit;
 }
 
 void AHUD::AddHitBox(FVector2D Position, FVector2D Size, FName Name, bool bConsumesInput, int32 Priority)
 {	
 	if( GetHitBoxWithName(Name) == NULL )
 	{
-		HitBoxMap.Add( FHUDHitBox( Position, Size, Name, bConsumesInput, Priority ) );
+		bool bAdded = false;
+		for (int32 Index = 0; Index < HitBoxMap.Num(); ++Index)
+		{
+			if (HitBoxMap[Index].GetPriority() < Priority)
+			{
+				HitBoxMap.Insert(FHUDHitBox(Position, Size, Name, bConsumesInput, Priority), Index);
+				bAdded = true;
+				break;
+			}
+		}
+		if (!bAdded)
+		{
+			HitBoxMap.Add(FHUDHitBox(Position, Size, Name, bConsumesInput, Priority));
+		}
 	}
 	else
 	{
@@ -789,12 +877,12 @@ bool AHUD::IsCanvasValid_WarnIfNot() const
 /////////////////
 
 FHUDHitBox::FHUDHitBox( FVector2D InCoords, FVector2D InSize, const FName& InName, bool bInConsumesInput, int32 InPriority )
+	: Coords(InCoords)
+	, Size(InSize)
+	, Name(InName)
+	, bConsumesInput(bInConsumesInput)
+	, Priority(InPriority)
 {
-	Coords = InCoords;
-	Size = InSize;
-	Name = InName;
-	bConsumesInput = bInConsumesInput;
-	Priority = InPriority;
 }
 
 bool FHUDHitBox::Contains( FVector2D InCoords ) const
@@ -810,7 +898,7 @@ bool FHUDHitBox::Contains( FVector2D InCoords ) const
 	return bResult;
 }
 
-void FHUDHitBox::Draw( FCanvas* InCanvas, const FLinearColor& InColor )
+void FHUDHitBox::Draw( FCanvas* InCanvas, const FLinearColor& InColor ) const
 {
 	FCanvasBoxItem	BoxItem( Coords, Size );
 	BoxItem.SetColor( InColor );
diff --git a/Engine/Source/Runtime/Engine/Private/PlayerController.cpp b/Engine/Source/Runtime/Engine/Private/PlayerController.cpp
--- a/Engine/Source/Runtime/Engine/Private/PlayerController.cpp
+++ b/Engine/Source/Runtime/Engine/Private/PlayerController.cpp
@@ -1749,7 +1749,7 @@ void APlayerController::DeprojectScreenPositionToWorld(float ScreenX, float Scre
 bool APlayerController::GetHitResultAtScreenPosition(const FVector2D ScreenPosition, const ECollisionChannel TraceChannel, bool bTraceComplex, FHitResult& HitResult) const
 {
 	// Early out if we clicked on a HUD hitbox
-	if( GetHUD() != NULL && GetHUD()->AnyCurrentHitBoxHits() )
+	if( GetHUD() != NULL && GetHUD()->GetHitBoxAtCoordinates(ScreenPosition, true) )
 	{
 		return false;
 	}
@@ -1787,7 +1787,7 @@ bool APlayerController::GetHitResultAtScreenPosition(const FVector2D ScreenPosit
 bool APlayerController::GetHitResultAtScreenPosition(const FVector2D ScreenPosition, const ETraceTypeQuery TraceChannel, bool bTraceComplex, FHitResult& HitResult) const
 {
 	// Early out if we clicked on a HUD hitbox
-	if( GetHUD() != NULL && GetHUD()->AnyCurrentHitBoxHits() )
+	if (GetHUD() != NULL && GetHUD()->GetHitBoxAtCoordinates(ScreenPosition, true))
 	{
 		return false;
 	}
@@ -1825,7 +1825,7 @@ bool APlayerController::GetHitResultAtScreenPosition(const FVector2D ScreenPosit
 bool APlayerController::GetHitResultAtScreenPosition(const FVector2D ScreenPosition, const TArray<TEnumAsByte<EObjectTypeQuery> > & ObjectTypes, bool bTraceComplex, FHitResult& HitResult) const
 {
 	// Early out if we clicked on a HUD hitbox
-	if( GetHUD() != NULL && GetHUD()->AnyCurrentHitBoxHits() )
+	if (GetHUD() != NULL && GetHUD()->GetHitBoxAtCoordinates(ScreenPosition, true))
 	{
 		return false;
 	}
@@ -1963,6 +1963,7 @@ bool APlayerController::InputKey(FKey Key, EInputEvent EventType, float AmountDe
 		// TODO: Allow click key(s?) to be defined
 		if (bEnableClickEvents && Key == EKeys::LeftMouseButton)
 		{
+			const FVector2D MousePosition = CastChecked<ULocalPlayer>(Player)->ViewportClient->GetMousePosition();
 			UPrimitiveComponent* ClickedPrimitive = NULL;
 			if (bEnableMouseOverEvents)
 			{
@@ -1970,11 +1971,8 @@ bool APlayerController::InputKey(FKey Key, EInputEvent EventType, float AmountDe
 			}
 			else
 			{
-				FVector2D MouseLocation;
-				GetInputMouseDelta(MouseLocation.X, MouseLocation.Y);
-
 				FHitResult HitResult;
-				const bool bHit = GetHitResultAtScreenPosition(CastChecked<ULocalPlayer>(Player)->ViewportClient->GetMousePosition(), CurrentClickTraceChannel, true, HitResult);
+				const bool bHit = GetHitResultAtScreenPosition(MousePosition, CurrentClickTraceChannel, true, HitResult);
 				if (bHit)
 				{
 					ClickedPrimitive = HitResult.Component.Get();
@@ -1982,7 +1980,7 @@ bool APlayerController::InputKey(FKey Key, EInputEvent EventType, float AmountDe
 			}
 			if( GetHUD() )
 			{
-				if( GetHUD()->UpdateAndDispatchHitBoxClickEvents(EventType) )
+				if (GetHUD()->UpdateAndDispatchHitBoxClickEvents(MousePosition, EventType, false))
 				{
 					ClickedPrimitive = NULL;
 				}
@@ -2028,6 +2026,17 @@ bool APlayerController::InputTouch(uint32 Handle, ETouchType::Type Type, const F
 			UPrimitiveComponent* PreviousComponent = CurrentTouchablePrimitives[Handle].Get();
 			UPrimitiveComponent* CurrentComponent = (bHit ? HitResult.Component.Get() : NULL);
 
+			if (GetHUD())
+			{
+				if (Type == ETouchType::Began || Type == ETouchType::Ended)
+				{
+					if (GetHUD()->UpdateAndDispatchHitBoxClickEvents(TouchLocation, (Type == ETouchType::Began ? EInputEvent::IE_Pressed : EInputEvent::IE_Released), true))
+					{
+						CurrentComponent = NULL;
+					}
+				}
+			}
+
 			switch(Type)
 			{
 			case ETouchType::Began:
