Radial damage now defaults to "unblocked" unless the visibility trace finds an explicit blockage.  Use bCanBeDamaged to toggle damagability.

diff --git a/Engine/Source/Runtime/Engine/Private/GameplayStatics.cpp b/Engine/Source/Runtime/Engine/Private/GameplayStatics.cpp
--- a/Engine/Source/Runtime/Engine/Private/GameplayStatics.cpp
+++ b/Engine/Source/Runtime/Engine/Private/GameplayStatics.cpp
@@ -125,7 +125,7 @@ static bool ComponentIsVisibleFrom(UPrimitiveComponent const* VictimComp, FVecto
 	LineParams.AddIgnoredActors( IgnoreActors );
 
 	// Do a trace from origin to middle of box
-	UWorld* World = VictimComp->GetWorld();
+	UWorld* const World = VictimComp->GetWorld();
 	check(World);
 
 	FVector const TraceEnd = VictimComp->Bounds.Origin;
@@ -154,8 +154,8 @@ static bool ComponentIsVisibleFrom(UPrimitiveComponent const* VictimComp, FVecto
 		}
 	}
 		
-	// didn't hit anything, including the victim component.  assume not visible.
-	return false;
+	// didn't hit anything, assume nothing blocking the damage and victim is consequently visible
+	return true;
 }
 
 bool UGameplayStatics::ApplyRadialDamage(UObject* WorldContextObject, float BaseDamage, const FVector& Origin, float DamageRadius, TSubclassOf<UDamageType> DamageTypeClass, const TArray<AActor*>& IgnoreActors, AActor* DamageCauser, AController* InstigatedByController, bool bDoFullDamage )
