Fix blend issue with 0 weighted Layered blend per bone

#TTP 332438 -  LayeredBlendPerBone not updating correctly when changing Blend Weight from 1.0 to 0.0

diff --git a/Engine/Source/Runtime/Engine/Private/Animation/AnimationRuntime.cpp b/Engine/Source/Runtime/Engine/Private/Animation/AnimationRuntime.cpp
--- a/Engine/Source/Runtime/Engine/Private/Animation/AnimationRuntime.cpp
+++ b/Engine/Source/Runtime/Engine/Private/Animation/AnimationRuntime.cpp
@@ -939,7 +939,7 @@ void FAnimationRuntime::UpdateDesiredBoneWeight(const TArray<FPerBoneBlendWeight
 	// in the future, cache this outside
 	ensure (TargetBoneBlendWeights.Num() == SrcBoneBlendWeights.Num());
 
-	FMemory::Memset( TargetBoneBlendWeights.GetTypedData(), 0, TargetBoneBlendWeights.Num());
+	FMemory::Memset(TargetBoneBlendWeights.GetTypedData(), 0, TargetBoneBlendWeights.Num() * sizeof(FPerBoneBlendWeight));
 
 	// go through skeleton tree requiredboneindices
 	const TArray<FBoneIndexType> & RequiredBoneIndices = RequiredBones.GetBoneIndicesArray();
