* fixed mouse leave when in game mode/PIE

diff --git a/Engine/Source/Runtime/Engine/Private/Slate/SceneViewport.cpp b/Engine/Source/Runtime/Engine/Private/Slate/SceneViewport.cpp
--- a/Engine/Source/Runtime/Engine/Private/Slate/SceneViewport.cpp
+++ b/Engine/Source/Runtime/Engine/Private/Slate/SceneViewport.cpp
@@ -386,6 +386,8 @@ void FSceneViewport::OnMouseEnter( const FGeometry& MyGeometry, const FPointerEv
 void FSceneViewport::OnMouseLeave( const FPointerEvent& MouseEvent )
 {
 	ViewportClient->MouseLeave( this );
+	
+	CachedMousePos = FIntPoint(-1, -1);
 }
 
 FReply FSceneViewport::OnMouseMove( const FGeometry& InGeometry, const FPointerEvent& InMouseEvent )
