Renaming the "Orientation" BP category to "Transformations"

diff --git a/Engine/Source/Runtime/Engine/Classes/AI/Navigation/RecastNavMesh.h b/Engine/Source/Runtime/Engine/Classes/AI/Navigation/RecastNavMesh.h
--- a/Engine/Source/Runtime/Engine/Classes/AI/Navigation/RecastNavMesh.h
+++ b/Engine/Source/Runtime/Engine/Classes/AI/Navigation/RecastNavMesh.h
@@ -363,7 +363,7 @@ struct FTileSetItem
 	}
 };
 
-UCLASS(config=Engine, defaultconfig, hidecategories=(Transform,"Utilities|Orientation",Actor,Layers,Replication))
+UCLASS(config=Engine, defaultconfig, hidecategories=(Transform,"Utilities|Transformation",Actor,Layers,Replication))
 class ENGINE_API ARecastNavMesh : public ANavigationData
 {
 	GENERATED_UCLASS_BODY()
@@ -883,4 +883,4 @@ bool ARecastNavMesh::NavMeshRaycast(const ANavigationData* Self, const FVector&
 	FRaycastResult Result;
 	return NavMeshRaycast(Self, RayStart, RayEnd, HitLocation, QueryFilter, Querier, Result);
 }
-#endif // WITH_RECAST
\ No newline at end of file
+#endif // WITH_RECAST
diff --git a/Engine/Source/Runtime/Engine/Classes/Atmosphere/AtmosphericFog.h b/Engine/Source/Runtime/Engine/Classes/Atmosphere/AtmosphericFog.h
--- a/Engine/Source/Runtime/Engine/Classes/Atmosphere/AtmosphericFog.h
+++ b/Engine/Source/Runtime/Engine/Classes/Atmosphere/AtmosphericFog.h
@@ -4,7 +4,7 @@
 #pragma once
 #include "AtmosphericFog.generated.h"
 
-UCLASS(showcategories=(Movement, Rendering, "Utilities|Orientation", "Input|MouseInput", "Input|TouchInput"), ClassGroup=Fog, hidecategories=(Info,Object,Input), MinimalAPI)
+UCLASS(showcategories=(Movement, Rendering, "Utilities|Transformation", "Input|MouseInput", "Input|TouchInput"), ClassGroup=Fog, hidecategories=(Info,Object,Input), MinimalAPI)
 class AAtmosphericFog : public AInfo
 {
 	GENERATED_UCLASS_BODY()
diff --git a/Engine/Source/Runtime/Engine/Classes/Components/SceneComponent.h b/Engine/Source/Runtime/Engine/Classes/Components/SceneComponent.h
--- a/Engine/Source/Runtime/Engine/Classes/Components/SceneComponent.h
+++ b/Engine/Source/Runtime/Engine/Classes/Components/SceneComponent.h
@@ -202,90 +202,90 @@ protected:
 public:
 
 	/** Set the location of this component relative to its parent */
-	UFUNCTION(BlueprintCallable, Category="Utilities|Orientation")
+	UFUNCTION(BlueprintCallable, Category="Utilities|Transformation")
 	void SetRelativeLocation(FVector NewLocation, bool bSweep=false);
 
 	/** Set the rotation of this component relative to its parent */
-	UFUNCTION(BlueprintCallable, Category="Utilities|Orientation")
+	UFUNCTION(BlueprintCallable, Category="Utilities|Transformation")
 	void SetRelativeRotation(FRotator NewRotation, bool bSweep=false);
 
 	/** Set the transform of this component relative to its parent */
-	UFUNCTION(BlueprintCallable, Category="Utilities|Orientation")
+	UFUNCTION(BlueprintCallable, Category="Utilities|Transformation")
 	void SetRelativeTransform(const FTransform& NewTransform, bool bSweep=false);
 
 	/** Returns the transform of this component relative to its parent */
-	UFUNCTION(BlueprintCallable, Category="Utilities|Orientation")
+	UFUNCTION(BlueprintCallable, Category="Utilities|Transformation")
 	FTransform GetRelativeTransform();
 
 	/** Set the transform of this component relative to its parent */
-	UFUNCTION(BlueprintCallable, Category="Utilities|Orientation")
+	UFUNCTION(BlueprintCallable, Category="Utilities|Transformation")
 	void ResetRelativeTransform();
 
 	/** Set the non-uniform of this component relative to its parent */
-	UFUNCTION(BlueprintCallable, Category="Utilities|Orientation")
+	UFUNCTION(BlueprintCallable, Category="Utilities|Transformation")
 	virtual void SetRelativeScale3D(FVector NewScale3D);
 
 	/** Adds a delta to the translation of this component relative to its parent */
-	UFUNCTION(BlueprintCallable, Category="Utilities|Orientation")
+	UFUNCTION(BlueprintCallable, Category="Utilities|Transformation")
 	void AddRelativeLocation(FVector DeltaLocation, bool bSweep=false);
 
 	/** Adds a delta the rotation of this component relative to its parent */
-	UFUNCTION(BlueprintCallable, Category="Utilities|Orientation")
+	UFUNCTION(BlueprintCallable, Category="Utilities|Transformation")
 	void AddRelativeRotation(FRotator DeltaRotation, bool bSweep=false);
 
 	/** Adds a delta to the location of this component in its local reference frame */
-	UFUNCTION(BlueprintCallable, Category="Utilities|Orientation")
+	UFUNCTION(BlueprintCallable, Category="Utilities|Transformation")
 	void AddLocalOffset(FVector DeltaLocation, bool bSweep=false);
 	
 	/** Adds a delta to the rotation of this component in its local reference frame */
-	UFUNCTION(BlueprintCallable, Category="Utilities|Orientation")
+	UFUNCTION(BlueprintCallable, Category="Utilities|Transformation")
 	void AddLocalRotation(FRotator DeltaRotation, bool bSweep=false);
 
 	/** Adds a delta to the rotation of this component in its local reference frame */
-	UFUNCTION(BlueprintCallable, Category="Utilities|Orientation")
+	UFUNCTION(BlueprintCallable, Category="Utilities|Transformation")
 	void AddLocalTransform(const FTransform& DeltaTransform, bool bSweep=false);
 
 	/** Set the relative location of this component to put it at the supplied location in world space. */
-	UFUNCTION(BlueprintCallable, Category="Utilities|Orientation")
+	UFUNCTION(BlueprintCallable, Category="Utilities|Transformation")
 	void SetWorldLocation(FVector NewLocation, bool bSweep=false);
 
 	/** Set the relative rotation of this component to put it at the supplied orientation in world space. */
-	UFUNCTION(BlueprintCallable, Category="Utilities|Orientation")
+	UFUNCTION(BlueprintCallable, Category="Utilities|Transformation")
 	void SetWorldRotation(FRotator NewRotation, bool bSweep=false);
 
-	UFUNCTION(BlueprintCallable, Category="Utilities|Orientation")
+	UFUNCTION(BlueprintCallable, Category="Utilities|Transformation")
 	void SetWorldScale3D(FVector NewScale);
 
-	UFUNCTION(BlueprintCallable, Category="Utilities|Orientation")
+	UFUNCTION(BlueprintCallable, Category="Utilities|Transformation")
 	/** Set the transform of this component in world space. */
 	void SetWorldTransform(const FTransform& NewTransform, bool bSweep=false);
 
 	/** Return location of the component, in world space */
-	UFUNCTION(BlueprintCallable, meta=(FriendlyName = "GetWorldLocation"), Category="Utilities|Orientation")
+	UFUNCTION(BlueprintCallable, meta=(FriendlyName = "GetWorldLocation"), Category="Utilities|Transformation")
 	FVector K2_GetComponentLocation() const;
 
 	/** Returns rotation of the component, in world space. */
-	UFUNCTION(BlueprintCallable, meta=(FriendlyName = "GetWorldRotation"), Category="Utilities|Orientation")
+	UFUNCTION(BlueprintCallable, meta=(FriendlyName = "GetWorldRotation"), Category="Utilities|Transformation")
 	FRotator K2_GetComponentRotation() const;
 	
 	/** Returns scale of the component, in world space. */
-	UFUNCTION(BlueprintCallable, meta=(FriendlyName = "GetWorldScale"), Category="Utilities|Orientation")
+	UFUNCTION(BlueprintCallable, meta=(FriendlyName = "GetWorldScale"), Category="Utilities|Transformation")
 	FVector K2_GetComponentScale() const;
 
 	/** Get the current component-to-world transform for this component */
-	UFUNCTION(BlueprintCallable, meta=(FriendlyName = "GetWorldTransform"), Category="Utilities|Orientation")
+	UFUNCTION(BlueprintCallable, meta=(FriendlyName = "GetWorldTransform"), Category="Utilities|Transformation")
 	FTransform K2_GetComponentToWorld() const;
 
 	/** Get the forward (X) unit direction vector from this component, in world space.  */
-	UFUNCTION(BlueprintCallable, Category="Utilities|Orientation")
+	UFUNCTION(BlueprintCallable, Category="Utilities|Transformation")
 	FVector GetForwardVector() const;
 
 	/** Get the up (Z) unit direction vector from this component, in world space.  */
-	UFUNCTION(BlueprintCallable, Category="Utilities|Orientation")
+	UFUNCTION(BlueprintCallable, Category="Utilities|Transformation")
 	FVector GetUpVector() const;
 
 	/** Get the right (Y) unit direction vector from this component, in world space.  */
-	UFUNCTION(BlueprintCallable, Category="Utilities|Orientation")
+	UFUNCTION(BlueprintCallable, Category="Utilities|Transformation")
 	FVector GetRightVector() const;
 
 	/** Returns whether the specified body is currently using physics simulation */
@@ -297,7 +297,7 @@ public:
 	virtual bool IsAnySimulatingPhysics() const;
 
 	/** Get a pointer to the USceneComponent we are attached to */
-	UFUNCTION(BlueprintCallable, Category="Utilities|Orientation")
+	UFUNCTION(BlueprintCallable, Category="Utilities|Transformation")
 	class USceneComponent* GetAttachParent() const;
 
 	/** Gets all parent components up to and including the root component */
@@ -320,18 +320,18 @@ public:
 	 *   Attach this component to another scene component, optionally at a named socket. It is valid to call this on components whether or not they have been Registered.
 	 *   @param bMaintainWorldTransform	If true, update the relative location/rotation of this component to keep its world position the same
 	 */
-	UFUNCTION(BlueprintCallable, Category="Utilities|Orientation", meta=(AttachType="KeepRelativeOffset"))
+	UFUNCTION(BlueprintCallable, Category="Utilities|Transformation", meta=(AttachType="KeepRelativeOffset"))
 	void AttachTo(class USceneComponent* InParent, FName InSocketName = NAME_None, EAttachLocation::Type AttachType = EAttachLocation::KeepRelativeOffset);
 
 	/** Zeroes out the relative transform of this component, and calls AttachTo(). Useful for attaching directly to a scene component or socket location  */
-	UFUNCTION(BlueprintCallable, meta=(DeprecatedFunction, DeprecationMessage = "Use AttachTo with EAttachLocation::SnapToTarget option instead"), Category="Utilities|Orientation")
+	UFUNCTION(BlueprintCallable, meta=(DeprecatedFunction, DeprecationMessage = "Use AttachTo with EAttachLocation::SnapToTarget option instead"), Category="Utilities|Transformation")
 	void SnapTo(class USceneComponent* InParent, FName InSocketName = NAME_None);
 
 	/** 
 	 *	Detach this component from whatever it is attached to 
 	 *   @param bMaintainWorldTransform	If true, update the relative location/rotation of this component to keep its world position the same *	
 	 */
-	UFUNCTION(BlueprintCallable, Category="Utilities|Orientation")
+	UFUNCTION(BlueprintCallable, Category="Utilities|Transformation")
 	virtual void DetachFromParent(bool bMaintainWorldPosition = false);
 
 	/** 
@@ -395,7 +395,7 @@ public:
 	 * Get Component Velocity vector
 	 * @return Socket transform in world space if socket is found. Otherwise it will return component's transform in world space.
 	 */
-	UFUNCTION(BlueprintCallable, Category="Utilities|Orientation")
+	UFUNCTION(BlueprintCallable, Category="Utilities|Transformation")
 	virtual FVector GetComponentVelocity() const;
 
 	/** 
@@ -629,7 +629,7 @@ public:
 	void SetRelativeLocationAndRotation(FVector NewLocation, FRotator NewRotation, bool bSweep=false);
 
 	/** Set which parts of the relative transform should be relative to parent, and which should be relative to world */
-	UFUNCTION(BlueprintCallable, Category="Utilities|Orientation")
+	UFUNCTION(BlueprintCallable, Category="Utilities|Transformation")
 	void SetAbsolute(bool bNewAbsoluteLocation = false, bool bNewAbsoluteRotation = false, bool bNewAbsoluteScale = false);
 
 	/** Set the relative location & FRotator rotation of this component to put it at the supplied pose in world space. */
diff --git a/Engine/Source/Runtime/Engine/Classes/Engine/ExponentialHeightFog.h b/Engine/Source/Runtime/Engine/Classes/Engine/ExponentialHeightFog.h
--- a/Engine/Source/Runtime/Engine/Classes/Engine/ExponentialHeightFog.h
+++ b/Engine/Source/Runtime/Engine/Classes/Engine/ExponentialHeightFog.h
@@ -4,7 +4,7 @@
 #pragma once
 #include "ExponentialHeightFog.generated.h"
 
-UCLASS(showcategories=(Movement, Rendering, "Utilities|Orientation"), ClassGroup=Fog, MinimalAPI)
+UCLASS(showcategories=(Movement, Rendering, "Utilities|Transformation"), ClassGroup=Fog, MinimalAPI)
 class AExponentialHeightFog : public AInfo
 {
 	GENERATED_UCLASS_BODY()
diff --git a/Engine/Source/Runtime/Engine/Classes/Engine/LevelScriptActor.h b/Engine/Source/Runtime/Engine/Classes/Engine/LevelScriptActor.h
--- a/Engine/Source/Runtime/Engine/Classes/Engine/LevelScriptActor.h
+++ b/Engine/Source/Runtime/Engine/Classes/Engine/LevelScriptActor.h
@@ -4,7 +4,7 @@
 
 #include "LevelScriptActor.generated.h"
 
-UCLASS(notplaceable, meta=(KismetHideOverrides = "ReceiveAnyDamage,ReceivePointDamage,ReceiveRadialDamage,ReceiveActorBeginOverlap,ReceiveActorEndOverlap,ReceiveHit,ReceiveDestroyed,ReceiveActorBeginCursorOver,ReceiveActorEndCursorOver,ReceiveActorOnClicked,ReceiveActorOnReleased,ReceiveActorOnInputTouchBegin,ReceiveActorOnInputTouchEnd,ReceiveActorOnInputTouchEnter,ReceiveActorOnInputTouchLeave"), HideCategories=(Collision,Rendering,"Utilities|Orientation"))
+UCLASS(notplaceable, meta=(KismetHideOverrides = "ReceiveAnyDamage,ReceivePointDamage,ReceiveRadialDamage,ReceiveActorBeginOverlap,ReceiveActorEndOverlap,ReceiveHit,ReceiveDestroyed,ReceiveActorBeginCursorOver,ReceiveActorEndCursorOver,ReceiveActorOnClicked,ReceiveActorOnReleased,ReceiveActorOnInputTouchBegin,ReceiveActorOnInputTouchEnd,ReceiveActorOnInputTouchEnter,ReceiveActorOnInputTouchLeave"), HideCategories=(Collision,Rendering,"Utilities|Transformation"))
 class ENGINE_API ALevelScriptActor : public AActor
 {
 	GENERATED_UCLASS_BODY()
diff --git a/Engine/Source/Runtime/Engine/Classes/Engine/WindDirectionalSource.h b/Engine/Source/Runtime/Engine/Classes/Engine/WindDirectionalSource.h
--- a/Engine/Source/Runtime/Engine/Classes/Engine/WindDirectionalSource.h
+++ b/Engine/Source/Runtime/Engine/Classes/Engine/WindDirectionalSource.h
@@ -3,7 +3,7 @@
 
 #include "WindDirectionalSource.generated.h"
 
-UCLASS(ClassGroup=Wind, showcategories=(Rendering, "Utilities|Orientation"))
+UCLASS(ClassGroup=Wind, showcategories=(Rendering, "Utilities|Transformation"))
 class ENGINE_API AWindDirectionalSource : public AInfo
 {
 	GENERATED_UCLASS_BODY()
diff --git a/Engine/Source/Runtime/Engine/Classes/GameFramework/Actor.h b/Engine/Source/Runtime/Engine/Classes/GameFramework/Actor.h
--- a/Engine/Source/Runtime/Engine/Classes/GameFramework/Actor.h
+++ b/Engine/Source/Runtime/Engine/Classes/GameFramework/Actor.h
@@ -602,11 +602,11 @@ public:
 	class UClass* GetActorClass() const;
 
 	/** Get the actor to world transform */
-	UFUNCTION(BlueprintCallable, meta=(FriendlyName = "GetActorTransform"), Category="Utilities|Orientation")
+	UFUNCTION(BlueprintCallable, meta=(FriendlyName = "GetActorTransform"), Category="Utilities|Transformation")
 	FTransform GetTransform() const;
 
 	/** Returns location of the RootComponent of this Actor*/
-	UFUNCTION(BlueprintCallable, meta=(FriendlyName = "GetActorLocation"), Category="Utilities|Orientation")
+	UFUNCTION(BlueprintCallable, meta=(FriendlyName = "GetActorLocation"), Category="Utilities|Transformation")
 	FVector K2_GetActorLocation() const;
 
 	/** 
@@ -615,23 +615,23 @@ public:
 	 *	@param bSweep		Should we sweep to the destination location. If true, will stop short of the target if blocked by something.
 	 *	@return	Whether the location was successfully set if not swept, or whether movement occurred if swept.
 	 */
-	UFUNCTION(BlueprintCallable, meta=(FriendlyName = "SetActorLocation"), Category="Utilities|Orientation")
+	UFUNCTION(BlueprintCallable, meta=(FriendlyName = "SetActorLocation"), Category="Utilities|Transformation")
 	bool K2_SetActorLocation(FVector NewLocation, bool bSweep=false);
 
 	/** Returns rotation of the RootComponent of this Actor */
-	UFUNCTION(BlueprintCallable, meta=(FriendlyName = "GetActorRotation"), Category="Utilities|Orientation")
+	UFUNCTION(BlueprintCallable, meta=(FriendlyName = "GetActorRotation"), Category="Utilities|Transformation")
 	FRotator K2_GetActorRotation() const;
 
 	/** Get the forward (X) vector (length 1.0) from this Actor, in world space.  */
-	UFUNCTION(BlueprintCallable, Category = "Utilities|Orientation")
+	UFUNCTION(BlueprintCallable, Category = "Utilities|Transformation")
 	FVector GetActorForwardVector() const;
 
 	/** Get the up (Z) vector (length 1.0) from this Actor, in world space.  */
-	UFUNCTION(BlueprintCallable, Category = "Utilities|Orientation")
+	UFUNCTION(BlueprintCallable, Category = "Utilities|Transformation")
 	FVector GetActorUpVector() const;
 
 	/** Get the right (Y) vector (length 1.0) from this Actor, in world space.  */
-	UFUNCTION(BlueprintCallable, Category = "Utilities|Orientation")
+	UFUNCTION(BlueprintCallable, Category = "Utilities|Transformation")
 	FVector GetActorRightVector() const;
 
 	/**
@@ -642,11 +642,11 @@ public:
 	void GetActorBounds(bool bOnlyCollidingComponents, FVector& Origin, FVector& BoxExtent) const;
 
 	/** Returns the RootComponent of this Actor */
-	UFUNCTION(BlueprintCallable, meta=(FriendlyName = "GetRootComponent"), Category="Utilities|Orientation")
+	UFUNCTION(BlueprintCallable, meta=(FriendlyName = "GetRootComponent"), Category="Utilities|Transformation")
 	class USceneComponent* K2_GetRootComponent() const;
 
 	/** Returns velocity (in cm/s (Unreal Units/second) of the rootcomponent if it is either using physics or has an associated MovementComponent */
-	UFUNCTION(BlueprintCallable, Category="Utilities|Orientation")
+	UFUNCTION(BlueprintCallable, Category="Utilities|Transformation")
 	virtual FVector GetVelocity() const;
 
 	/** 
@@ -662,7 +662,7 @@ public:
 	 *	@param	NewRotation	The new rotation for the Actor.
 	 *	@return	Whether the rotation was successfully set.
 	 */
-	UFUNCTION(BlueprintCallable, Category="Utilities|Orientation")
+	UFUNCTION(BlueprintCallable, Category="Utilities|Transformation")
 	bool SetActorRotation(FRotator NewRotation);
 
 	/** 
@@ -672,50 +672,50 @@ public:
 	 *	@param bSweep		Should we sweep to the destination location. If true, will stop short of the target if blocked by something.
 	 *	@return	Whether the rotation was successfully set.
 	 */
-	UFUNCTION(BlueprintCallable, Category="Utilities|Orientation")
+	UFUNCTION(BlueprintCallable, Category="Utilities|Transformation")
 	bool SetActorLocationAndRotation(const FVector& NewLocation, FRotator NewRotation, bool bSweep=false);
 
 	/** Set the Actor's world-space scale. */
-	UFUNCTION(BlueprintCallable, Category="Utilities|Orientation")
+	UFUNCTION(BlueprintCallable, Category="Utilities|Transformation")
 	void SetActorScale3D(const FVector& NewScale3D);
 
 	/** Returns the distance from this Actor to OtherActor. */
-	UFUNCTION(BlueprintCallable, Category="Utilities|Orientation")
+	UFUNCTION(BlueprintCallable, Category="Utilities|Transformation")
 	float GetDistanceTo(AActor* OtherActor);
 
 	/** Returns the distance from this Actor to OtherActor, ignoring Z. */
-	UFUNCTION(BlueprintCallable, Category="Utilities|Orientation")
+	UFUNCTION(BlueprintCallable, Category="Utilities|Transformation")
 	float GetHorizontalDistanceTo(AActor* OtherActor);
 
 	/** Returns the distance from this Actor to OtherActor, ignoring XY. */
-	UFUNCTION(BlueprintCallable, Category="Utilities|Orientation")
+	UFUNCTION(BlueprintCallable, Category="Utilities|Transformation")
 	float GetVerticalDistanceTo(AActor* OtherActor);
 
 	/** Returns the dot product from this Actor to OtherActor. Returns -2.0 on failure. Returns 0.0 for coincidental actors. */
-	UFUNCTION(BlueprintCallable, Category = "Utilities|Orientation")
+	UFUNCTION(BlueprintCallable, Category = "Utilities|Transformation")
 	float GetDotProductTo(AActor* OtherActor);
 
 	/** Returns the dot product from this Actor to OtherActor, ignoring Z. Returns -2.0 on failure. Returns 0.0 for coincidental actors. */
-	UFUNCTION(BlueprintCallable, Category = "Utilities|Orientation")
+	UFUNCTION(BlueprintCallable, Category = "Utilities|Transformation")
 	float GetHorizontalDotProductTo(AActor* OtherActor);
 
 	/** 
 	 *	Set the Actors transform to the specified one.
 	 *	@param bSweep		Should we sweep to the destination location. If true, will stop short of the target if blocked by something.
 	 */
-	UFUNCTION(BlueprintCallable, Category="Utilities|Orientation")
+	UFUNCTION(BlueprintCallable, Category="Utilities|Transformation")
 	bool SetActorTransform(const FTransform& NewTransform, bool bSweep=false);
 
 	/** Adds a delta to the location of this component in its local reference frame */
-	UFUNCTION(BlueprintCallable, Category="Utilities|Orientation")
+	UFUNCTION(BlueprintCallable, Category="Utilities|Transformation")
 	void AddActorLocalOffset(FVector DeltaLocation, bool bSweep=false);
 
 	/** Adds a delta to the rotation of this component in its local reference frame */
-	UFUNCTION(BlueprintCallable, Category="Utilities|Orientation")
+	UFUNCTION(BlueprintCallable, Category="Utilities|Transformation")
 	void AddActorLocalRotation(FRotator DeltaRotation, bool bSweep=false);
 
 	/** Adds a delta to the transform of this component in its local reference frame */
-	UFUNCTION(BlueprintCallable, Category="Utilities|Orientation")
+	UFUNCTION(BlueprintCallable, Category="Utilities|Transformation")
 	void AddActorLocalTransform(const FTransform& NewTransform, bool bSweep=false);
 
 	/**
@@ -723,7 +723,7 @@ public:
 	 * @param NewRelativeLocation	New relative location to set the actor's RootComponent to
 	 * @param bSweep				Should we sweep to the destination location. If true, will stop short of the target if blocked by something
 	 */
-	UFUNCTION(BlueprintCallable, Category="Utilities|Orientation")
+	UFUNCTION(BlueprintCallable, Category="Utilities|Transformation")
 	void SetActorRelativeLocation(FVector NewRelativeLocation, bool bSweep=false);
 
 	/**
@@ -731,7 +731,7 @@ public:
 	 * @param NewRelativeRotation		New relative rotation to set the actor's RootComponent to
 	 * @param bSweep					Should we sweep to the destination rotation. If true, will stop short of the target if blocked by something
 	 */
-	UFUNCTION(BlueprintCallable, Category="Utilities|Orientation")
+	UFUNCTION(BlueprintCallable, Category="Utilities|Transformation")
 	void SetActorRelativeRotation(FRotator NewRelativeRotation, bool bSweep=false);
 
 	/**
@@ -739,14 +739,14 @@ public:
 	 * @param NewRelativeTransform		New relative transform to set the actor's RootComponent to
 	 * @param bSweep					Should we sweep to the destination transform. If true, will stop short of the target if blocked by something
 	 */
-	UFUNCTION(BlueprintCallable, Category="Utilities|Orientation")
+	UFUNCTION(BlueprintCallable, Category="Utilities|Transformation")
 	void SetActorRelativeTransform(const FTransform& NewRelativeTransform, bool bSweep=false);
 
 	/**
 	 * Set the actor's RootComponent to the specified relative scale 3d
 	 * @param NewRelativeScale	New scale to set the actor's RootComponent to
 	 */
-	UFUNCTION(BlueprintCallable, Category="Utilities|Orientation")
+	UFUNCTION(BlueprintCallable, Category="Utilities|Transformation")
 	void SetActorRelativeScale3D(FVector NewRelativeScale);
 
 	/**
@@ -788,7 +788,7 @@ public:
 	 *  Attaches the RootComponent of this Actor to the supplied component, optionally at a named socket. It is not valid to call this on components that are not Registered. 
 	 *   @param AttachLocationType	Type of attachment, AbsoluteWorld to keep its world position, RelativeOffset to keep the object's relative offset and SnapTo to snap to the new parent.
 	 */
-	UFUNCTION(BlueprintCallable, meta=(FriendlyName = "AttachActorToComponent", AttachLocationType="KeepRelativeOffset"), Category="Utilities|Orientation")
+	UFUNCTION(BlueprintCallable, meta=(FriendlyName = "AttachActorToComponent", AttachLocationType="KeepRelativeOffset"), Category="Utilities|Transformation")
 	void AttachRootComponentTo(class USceneComponent* InParent, FName InSocketName = NAME_None, EAttachLocation::Type AttachLocationType = EAttachLocation::KeepRelativeOffset);
 
 	/**
@@ -797,21 +797,21 @@ public:
 	 * @param InSocketName				Socket name to attach to, if any
 	 * @param AttachLocationType	Type of attachment, AbsoluteWorld to keep its world position, RelativeOffset to keep the object's relative offset and SnapTo to snap to the new parent.
 	 */
-	UFUNCTION(BlueprintCallable, meta=(FriendlyName = "AttachActorToActor", AttachLocationType="KeepRelativeOffset"), Category="Utilities|Orientation")
+	UFUNCTION(BlueprintCallable, meta=(FriendlyName = "AttachActorToActor", AttachLocationType="KeepRelativeOffset"), Category="Utilities|Transformation")
 	void AttachRootComponentToActor(AActor* InParentActor, FName InSocketName = NAME_None, EAttachLocation::Type AttachLocationType = EAttachLocation::KeepRelativeOffset);
 
 	/** 
 	 *  Snap the RootComponent of this Actor to the supplied Actor's root component, optionally at a named socket. It is not valid to call this on components that are not Registered. 
 	 *  If InSocketName == NAME_None, it will attach to origin of the InParentActor. 
 	 */
-	UFUNCTION(BlueprintCallable, meta=(DeprecatedFunction, DeprecationMessage = "Use AttachRootComponentTo with EAttachLocation::SnapToTarget option instead", FriendlyName = "SnapActorTo"), Category="Utilities|Orientation")
+	UFUNCTION(BlueprintCallable, meta=(DeprecatedFunction, DeprecationMessage = "Use AttachRootComponentTo with EAttachLocation::SnapToTarget option instead", FriendlyName = "SnapActorTo"), Category="Utilities|Transformation")
 	void SnapRootComponentTo(AActor* InParentActor, FName InSocketName = NAME_None);
 
 	/** 
 	 *  Detaches the RootComponent of this Actor from any SceneComponent it is currently attached to. 
 	 *   @param bMaintainWorldTransform	If true, update the relative location/rotation of this component to keep its world position the same
 	 */
-	UFUNCTION(BlueprintCallable, meta=(FriendlyName = "DetachActorFromActor"), Category="Utilities|Orientation")
+	UFUNCTION(BlueprintCallable, meta=(FriendlyName = "DetachActorFromActor"), Category="Utilities|Transformation")
 	void DetachRootComponentFromParent(bool bMaintainWorldPosition = true);
 
 	/** 
@@ -1582,7 +1582,7 @@ public:
 	 * @param DestRotation The target rotation at the destination
 	 * @return true if the actor has been successfully moved, or false if it couldn't fit.
 	 */
-	UFUNCTION(BlueprintCallable, meta=( FriendlyName="Teleport", Keywords = "Move Position" ), Category="Utilities|Orientation")
+	UFUNCTION(BlueprintCallable, meta=( FriendlyName="Teleport", Keywords = "Move Position" ), Category="Utilities|Transformation")
 	bool K2_TeleportTo( FVector DestLocation, FRotator DestRotation );
 
 	/** Called from TeleportTo() when teleport succeeds */
diff --git a/Engine/Source/Runtime/Engine/Classes/GameFramework/Controller.h b/Engine/Source/Runtime/Engine/Classes/GameFramework/Controller.h
--- a/Engine/Source/Runtime/Engine/Classes/GameFramework/Controller.h
+++ b/Engine/Source/Runtime/Engine/Classes/GameFramework/Controller.h
@@ -28,7 +28,7 @@ class UPathFollowingComponent;
 UDELEGATE(BlueprintAuthorityOnly)
 DECLARE_DYNAMIC_MULTICAST_DELEGATE_FourParams( FInstigatedAnyDamageSignature, float, Damage, const class UDamageType*, DamageType, class AActor*, DamagedActor, class AActor*, DamageCauser );
 
-UCLASS(abstract, notplaceable, NotBlueprintable, HideCategories=(Collision,Rendering,"Utilities|Orientation")) 
+UCLASS(abstract, notplaceable, NotBlueprintable, HideCategories=(Collision,Rendering,"Utilities|Transformation")) 
 class ENGINE_API AController : public AActor, public INavAgentInterface
 {
 	GENERATED_UCLASS_BODY()
diff --git a/Engine/Source/Runtime/Engine/Classes/GameFramework/Info.h b/Engine/Source/Runtime/Engine/Classes/GameFramework/Info.h
--- a/Engine/Source/Runtime/Engine/Classes/GameFramework/Info.h
+++ b/Engine/Source/Runtime/Engine/Classes/GameFramework/Info.h
@@ -8,7 +8,7 @@
 #pragma once
 #include "Info.generated.h"
 
-UCLASS(abstract, hidecategories=(Input, Movement, Collision, Rendering, "Utilities|Orientation"), showcategories=("Input|MouseInput", "Input|TouchInput"), MinimalAPI, NotBlueprintable)
+UCLASS(abstract, hidecategories=(Input, Movement, Collision, Rendering, "Utilities|Transformation"), showcategories=("Input|MouseInput", "Input|TouchInput"), MinimalAPI, NotBlueprintable)
 class AInfo : public AActor
 {
 	GENERATED_UCLASS_BODY()
diff --git a/Engine/Source/Runtime/Engine/Classes/Kismet/GameplayStatics.h b/Engine/Source/Runtime/Engine/Classes/Kismet/GameplayStatics.h
--- a/Engine/Source/Runtime/Engine/Classes/Kismet/GameplayStatics.h
+++ b/Engine/Source/Runtime/Engine/Classes/Kismet/GameplayStatics.h
@@ -44,7 +44,7 @@ class ENGINE_API UGameplayStatics : public UBlueprintFunctionLibrary
 	// --- Actor functions ------------------------------
 
 	/** Find the average location (centroid) of an array of Actors */
-	UFUNCTION(BlueprintCallable, Category="Utilities|Orientation")
+	UFUNCTION(BlueprintCallable, Category="Utilities|Transformation")
 	static FVector GetActorArrayAverageLocation(const TArray<AActor*>& Actors);
 
 	/** Bind the bounds of an array of Actors */
diff --git a/Engine/Source/Runtime/Engine/Classes/Landscape/LandscapeProxy.h b/Engine/Source/Runtime/Engine/Classes/Landscape/LandscapeProxy.h
--- a/Engine/Source/Runtime/Engine/Classes/Landscape/LandscapeProxy.h
+++ b/Engine/Source/Runtime/Engine/Classes/Landscape/LandscapeProxy.h
@@ -169,7 +169,7 @@ namespace ELandscapeLayerPaintingRestriction
 	};
 }
 
-UCLASS(dependson=UEngineTypes, NotPlaceable, hidecategories=(Display, Attachment, Physics, Debug, Lighting, LOD), showcategories=(Rendering, "Utilities|Orientation"), MinimalAPI)
+UCLASS(dependson=UEngineTypes, NotPlaceable, hidecategories=(Display, Attachment, Physics, Debug, Lighting, LOD), showcategories=(Rendering, "Utilities|Transformation"), MinimalAPI)
 class ALandscapeProxy : public AActor, public INavRelevantActorInterface
 {
 	GENERATED_UCLASS_BODY()
