"ToggleDrawEvents" now prints its state to the log better

diff --git a/Engine/Source/Runtime/Engine/Private/LocalPlayer.cpp b/Engine/Source/Runtime/Engine/Private/LocalPlayer.cpp
--- a/Engine/Source/Runtime/Engine/Private/LocalPlayer.cpp
+++ b/Engine/Source/Runtime/Engine/Private/LocalPlayer.cpp
@@ -15,6 +15,7 @@
 
 
 DEFINE_LOG_CATEGORY(LogPlayerManagement);
+DEFINE_LOG_CATEGORY_STATIC(LogEngine, Log, All);
 
 #if !(UE_BUILD_SHIPPING || UE_BUILD_TEST)
 
@@ -943,12 +944,12 @@ bool ULocalPlayer::HandleToggleDrawEventsCommand( const TCHAR* Cmd, FOutputDevic
 	if( GEmitDrawEvents )
 	{
 		GEmitDrawEvents = false;
-		UE_LOG(LogPlayerManagement, Warning, TEXT("Draw events are now DISABLED"));
+		UE_LOG(LogEngine, Warning, TEXT("Draw events are now DISABLED"));
 	}
 	else
 	{
 		GEmitDrawEvents = true;
-		UE_LOG(LogPlayerManagement, Warning, TEXT("Draw events are now ENABLED"));
+		UE_LOG(LogEngine, Warning, TEXT("Draw events are now ENABLED"));
 	}
 #endif
 	return true;
