Settings: NavMesh setting workaround for not being able to have different default values in CDO constructor and Default.ini

diff --git a/Engine/Source/Runtime/Engine/Private/AI/Navigation/NavigationData.cpp b/Engine/Source/Runtime/Engine/Private/AI/Navigation/NavigationData.cpp
--- a/Engine/Source/Runtime/Engine/Private/AI/Navigation/NavigationData.cpp
+++ b/Engine/Source/Runtime/Engine/Private/AI/Navigation/NavigationData.cpp
@@ -16,7 +16,7 @@ void FNavDataGenerator::TriggerGeneration()
 ANavigationData::ANavigationData(const class FPostConstructInitializeProperties& PCIP)
 	: Super(PCIP)
 	, bEnableDrawing(false)
-	, bRebuildAtRuntime(true)
+	, bRebuildAtRuntime(false)
 	, DataVersion(NAVMESHVER_LATEST)
 	, FindPathImplementation(NULL)
 	, FindHierarchicalPathImplementation(NULL)
