Settings: forcing NavMesh data settings to save to Default INI

diff --git a/Engine/Source/Runtime/Engine/Classes/AI/Navigation/NavigationData.h b/Engine/Source/Runtime/Engine/Classes/AI/Navigation/NavigationData.h
--- a/Engine/Source/Runtime/Engine/Classes/AI/Navigation/NavigationData.h
+++ b/Engine/Source/Runtime/Engine/Classes/AI/Navigation/NavigationData.h
@@ -22,7 +22,7 @@ struct FSupportedAreaData
  *	Represents abstract Navigation Data (sub-classed as NavMesh, NavGraph, etc)
  *	Used as a common interface for all navigation types handled by NavigationSystem
  */
-UCLASS(config=Engine, dependson=UNavigationSystem, NotBlueprintable, abstract)
+UCLASS(config=Engine, defaultconfig, dependson=UNavigationSystem, NotBlueprintable, abstract)
 class ENGINE_API ANavigationData : public AActor
 {
 	GENERATED_UCLASS_BODY()
