Merging

//depot/UE4-TappyChicken/Engine/Source/...

to //depot/UE4/Engine/Source/...

Fixes from TappyChicken

diff --git a/Engine/Source/Runtime/Core/Private/IOS/EAGLView.cpp b/Engine/Source/Runtime/Core/Private/IOS/EAGLView.cpp
--- a/Engine/Source/Runtime/Core/Private/IOS/EAGLView.cpp
+++ b/Engine/Source/Runtime/Core/Private/IOS/EAGLView.cpp
@@ -86,12 +86,13 @@ void GetInputStack(/*TArray<FPointerEvent>& TouchEvents*/)
 		self.contentScaleFactor = bIsForOnDevice ? RequestedContentScaleFactor : 1.0f;
 		UE_LOG(LogIOS, Log, TEXT("Setting contentScaleFactor to %0.4f (optimal = %0.4f)"), self.contentScaleFactor, NativeScale);
 
-		if (self.contentScaleFactor == 1.0f || self.contentScaleFactor == 2.0f)
-		{
-			UE_LOG(LogIOS,Log,TEXT("Setting layer filter to NEAREST"));
-			CAEAGLLayer *EaglLayer = (CAEAGLLayer *)self.layer;
-			EaglLayer.magnificationFilter = kCAFilterNearest;
-		}
+        // This is causing a pretty large slow down on iPad 3 and 4 using iOS 6, for now going to comment it out
+//		if (self.contentScaleFactor == 1.0f || self.contentScaleFactor == 2.0f)
+//		{
+//			UE_LOG(LogIOS,Log,TEXT("Setting layer filter to NEAREST"));
+//			CAEAGLLayer *EaglLayer = (CAEAGLLayer *)self.layer;
+//			EaglLayer.magnificationFilter = kCAFilterNearest;
+//		}
 
 		// Create our standard displayable surface
 		glGenRenderbuffers(1, &OnScreenColorRenderBuffer);
diff --git a/Engine/Source/Runtime/Core/Private/IOS/IOSPlatformProcess.cpp b/Engine/Source/Runtime/Core/Private/IOS/IOSPlatformProcess.cpp
--- a/Engine/Source/Runtime/Core/Private/IOS/IOSPlatformProcess.cpp
+++ b/Engine/Source/Runtime/Core/Private/IOS/IOSPlatformProcess.cpp
@@ -48,7 +48,7 @@ void FIOSPlatformProcess::LaunchURL( const TCHAR* URL, const TCHAR* Parms, FStri
 
 void FIOSPlatformProcess::SetRealTimeMode()
 {
-	if ([IOSAppDelegate GetDelegate].OSVersion < 7)
+	if ([IOSAppDelegate GetDelegate].OSVersion < 7 && FPlatformMisc::NumberOfCores() > 1)
 	{
 		mach_timebase_info_data_t TimeBaseInfo;
 		mach_timebase_info( &TimeBaseInfo );
diff --git a/Engine/Source/Runtime/Engine/Classes/Engine/World.h b/Engine/Source/Runtime/Engine/Classes/Engine/World.h
--- a/Engine/Source/Runtime/Engine/Classes/Engine/World.h
+++ b/Engine/Source/Runtime/Engine/Classes/Engine/World.h
@@ -744,6 +744,9 @@ public:
 	/** Whether we are in the middle of ticking actors/components or not														*/
 	bool										bInTick;
 
+    /** Whether we have already built the collision tree or not                                                                 */
+    bool                                        bIsBuilt;
+    
 	/** We are in the middle of actor ticking, so add tasks for newly spawned actors											*/
 	bool										bTickNewlySpawned;
 
diff --git a/Engine/Source/Runtime/Engine/Private/LevelTick.cpp b/Engine/Source/Runtime/Engine/Private/LevelTick.cpp
--- a/Engine/Source/Runtime/Engine/Private/LevelTick.cpp
+++ b/Engine/Source/Runtime/Engine/Private/LevelTick.cpp
@@ -1110,6 +1110,9 @@ void UWorld::Tick( ELevelTick TickType, float DeltaSeconds )
 
 		SCOPE_CYCLE_COUNTER(STAT_TickTime);
 		RunTickGroup(TG_PrePhysics);
+        bInTick = false;
+        EnsureCollisionTreeIsBuilt();
+        bInTick = true;
 #if EXPERIMENTAL_PARALLEL_CODE
 		{			
 			RunTickGroup(TG_ParallelAnimWork);			
diff --git a/Engine/Source/Runtime/Engine/Private/World.cpp b/Engine/Source/Runtime/Engine/Private/World.cpp
--- a/Engine/Source/Runtime/Engine/Private/World.cpp
+++ b/Engine/Source/Runtime/Engine/Private/World.cpp
@@ -50,6 +50,7 @@ FWorldDelegates::FOnLevelChanged FWorldDelegates::LevelRemovedFromWorld;
 UWorld::UWorld( const class FPostConstructInitializeProperties& PCIP )
 :	UObject(PCIP)
 ,	TickTaskLevel(FTickTaskManagerInterface::Get().AllocateTickTaskLevel())
+,   bIsBuilt(false)
 ,	NextTravelType(TRAVEL_Relative)
 {
 	TimerManager = new FTimerManager();
@@ -63,6 +64,7 @@ UWorld::UWorld( const class FPostConstructInitializeProperties& PCIP,const FURL&
 :	UObject(PCIP)
 ,	URL(InURL)
 ,	TickTaskLevel(FTickTaskManagerInterface::Get().AllocateTickTaskLevel())
+,   bIsBuilt(false)
 ,	NextTravelType(TRAVEL_Relative)
 {
 	SetFlags( RF_Transactional );
@@ -1115,7 +1117,7 @@ void UWorld::ModifyLevel(ULevel* Level)
 
 void UWorld::EnsureCollisionTreeIsBuilt()
 {
-	if (bInTick)
+	if (bInTick || bIsBuilt)
 	{
 		// Current implementation of collision tree rebuild ticks physics scene and can not be called during world tick
 		return;
@@ -1137,7 +1139,7 @@ void UWorld::EnsureCollisionTreeIsBuilt()
 	{
 		SetupPhysicsTickFunctions(0.1f);
 		PhysicsScene->StartFrame();
-		PhysicsScene->WaitPhysScenes();
+        PhysicsScene->WaitPhysScenes();
 		PhysicsScene->EndFrame(NULL);
 	}
 
@@ -1145,6 +1147,7 @@ void UWorld::EnsureCollisionTreeIsBuilt()
 	{
 		PhysicsScene->SetIsStaticLoading(false);
 	}
+    bIsBuilt = true;
 }
 
 void UWorld::InvalidateModelGeometry(ULevel* InLevel)
