332767	PUNTABLE: EDITOR: PERSONA: Blendspace parameters axes tooltip does not match data

diff --git a/Engine/Source/Runtime/Engine/Classes/Animation/BlendSpaceBase.h b/Engine/Source/Runtime/Engine/Classes/Animation/BlendSpaceBase.h
--- a/Engine/Source/Runtime/Engine/Classes/Animation/BlendSpaceBase.h
+++ b/Engine/Source/Runtime/Engine/Classes/Animation/BlendSpaceBase.h
@@ -200,7 +200,7 @@ class UBlendSpaceBase : public UAnimationAsset
 	GENERATED_UCLASS_BODY()
 
 protected:
-	/** Blend Parameters for each axis, for now we only allow 2 max **/
+	/** Blend Parameters for each axis. **/
 	UPROPERTY(Category=BlendSpace, VisibleAnywhere)
 	struct FBlendParameter BlendParameters[3];
 
