Settings: hiding undesired NavMesh settings categories

diff --git a/Engine/Source/Runtime/Engine/Classes/AI/Navigation/RecastNavMesh.h b/Engine/Source/Runtime/Engine/Classes/AI/Navigation/RecastNavMesh.h
--- a/Engine/Source/Runtime/Engine/Classes/AI/Navigation/RecastNavMesh.h
+++ b/Engine/Source/Runtime/Engine/Classes/AI/Navigation/RecastNavMesh.h
@@ -363,7 +363,8 @@ struct FTileSetItem
 	}
 };
 
-UCLASS(config=Engine, defaultconfig, hidecategories=(Transform,"Utilities|Transformation",Actor,Layers,Replication))
+
+UCLASS(config=Engine, defaultconfig, hidecategories=(Input,Rendering,Tags,Transform,"Utilities|Transformation",Actor,Layers,Replication))
 class ENGINE_API ARecastNavMesh : public ANavigationData
 {
 	GENERATED_UCLASS_BODY()
