Simple Look At Control : it doesn't handle anything else but look at bone/ target point with look at axis.

diff --git a/Engine/Source/Editor/AnimGraph/Classes/AnimGraphNode_LookAt.h b/Engine/Source/Editor/AnimGraph/Classes/AnimGraphNode_LookAt.h
new file mode 100644
--- /dev/null
+++ b/Engine/Source/Editor/AnimGraph/Classes/AnimGraphNode_LookAt.h
@@ -0,0 +1,31 @@
+// Copyright 1998-2014 Epic Games, Inc. All Rights Reserved.
+
+#pragma once
+
+#include "AnimGraphNode_SkeletalControlBase.h"
+#include "AnimGraphNode_LookAt.generated.h"
+
+UCLASS(MinimalAPI)
+class UAnimGraphNode_LookAt : public UAnimGraphNode_SkeletalControlBase
+{
+	GENERATED_UCLASS_BODY()
+
+	UPROPERTY(EditAnywhere, Category=Settings)
+	FAnimNode_LookAt Node;
+
+public:
+	// UEdGraphNode interface
+	virtual FText GetNodeTitle(ENodeTitleType::Type TitleType) const OVERRIDE;
+	virtual FString GetTooltip() const OVERRIDE;
+	virtual FString GetKeywords() const OVERRIDE;
+	// End of UEdGraphNode interface
+
+	// UAnimGraphNode_SkeletalControlBase interface
+	ANIMGRAPH_API virtual void Draw(FPrimitiveDrawInterface* PDI, USkeletalMeshComponent* SkelMeshComp) const OVERRIDE;
+	// End of UAnimGraphNode_SkeletalControlBase interface
+
+protected:
+	// UAnimGraphNode_SkeletalControlBase interface
+	virtual FText GetControllerDescription() const OVERRIDE;
+	// End of UAnimGraphNode_SkeletalControlBase interface
+};
diff --git a/Engine/Source/Editor/AnimGraph/Private/AnimGraphNode_LookAt.cpp b/Engine/Source/Editor/AnimGraph/Private/AnimGraphNode_LookAt.cpp
new file mode 100644
--- /dev/null
+++ b/Engine/Source/Editor/AnimGraph/Private/AnimGraphNode_LookAt.cpp
@@ -0,0 +1,59 @@
+// Copyright 1998-2014 Epic Games, Inc. All Rights Reserved.
+
+#include "AnimGraphPrivatePCH.h"
+#include "AnimGraphNode_LookAt.h"
+
+/////////////////////////////////////////////////////
+// UAnimGraphNode_LookAt
+
+#define LOCTEXT_NAMESPACE "AnimGraph_LookAt"
+
+UAnimGraphNode_LookAt::UAnimGraphNode_LookAt(const FPostConstructInitializeProperties& PCIP)
+	: Super(PCIP)
+{
+}
+
+FText UAnimGraphNode_LookAt::GetControllerDescription() const
+{
+	return LOCTEXT("LookAtNode", "Look At");
+}
+
+FString UAnimGraphNode_LookAt::GetKeywords() const
+{
+	return TEXT("Look At, Follow, Trace, Track");
+}
+
+FString UAnimGraphNode_LookAt::GetTooltip() const
+{
+	return LOCTEXT("AnimGraphNode_LookAt_Tooltip", "This node allow a bone to trace or follow another bone").ToString();
+}
+
+FText UAnimGraphNode_LookAt::GetNodeTitle(ENodeTitleType::Type TitleType) const
+{
+	FFormatNamedArguments Args;
+	Args.Add(TEXT("ControllerDescription"), GetControllerDescription());
+	Args.Add(TEXT("BoneName"), FText::FromName(Node.BoneToModify.BoneName));
+
+	if(TitleType == ENodeTitleType::ListView)
+	{
+		return FText::Format(LOCTEXT("AnimGraphNode_LookAt_Title", "{ControllerDescription} - Bone: {BoneName}"), Args);
+	}
+	else
+	{
+		return FText::Format(LOCTEXT("AnimGraphNode_LookAt_Title", "{ControllerDescription}\nBone: {BoneName}"), Args);
+	}
+}
+
+void UAnimGraphNode_LookAt::Draw(FPrimitiveDrawInterface* PDI, USkeletalMeshComponent* SkelMeshComp) const
+{
+	// this is not accurate debugging
+	// since i can't access the transient data of run-time instance, I have to calculate data from other way
+	// this technically means a frame delay, but it would be sufficient for debugging purpose. 
+	FVector BoneLocation = SkelMeshComp->GetSocketLocation(Node.BoneToModify.BoneName);
+	FVector TargetLocation = (Node.LookAtBone.BoneName != NAME_None)? SkelMeshComp->GetSocketLocation(Node.LookAtBone.BoneName) : Node.LookAtLocation;
+
+	DrawWireStar(PDI, TargetLocation, 5.f, FLinearColor(1, 0, 0), SDPG_Foreground);
+	DrawDashedLine(PDI, BoneLocation, TargetLocation, FLinearColor(1, 1, 0), 2.f, SDPG_Foreground);
+}
+
+#undef LOCTEXT_NAMESPACE
diff --git a/Engine/Source/Runtime/Engine/Classes/Animation/BoneControllers/AnimNode_LookAt.h b/Engine/Source/Runtime/Engine/Classes/Animation/BoneControllers/AnimNode_LookAt.h
new file mode 100644
--- /dev/null
+++ b/Engine/Source/Runtime/Engine/Classes/Animation/BoneControllers/AnimNode_LookAt.h
@@ -0,0 +1,67 @@
+// Copyright 1998-2014 Epic Games, Inc. All Rights Reserved.
+
+#pragma once
+
+#include "AnimNode_SkeletalControlBase.h"
+#include "AnimNode_LookAt.generated.h"
+
+UENUM()
+namespace EAxisOption
+{
+	enum Type
+	{
+		X,
+		Y,
+		Z, 
+		X_Neg, 
+		Y_Neg, 
+		Z_Neg
+	};
+}
+
+/**
+ *	Simple controller that make a bone to look at the point or another bone
+ */
+USTRUCT()
+struct ENGINE_API FAnimNode_LookAt : public FAnimNode_SkeletalControlBase
+{
+	GENERATED_USTRUCT_BODY()
+
+	/** Name of bone to control. This is the main bone chain to modify from. **/
+	UPROPERTY(EditAnywhere, BlueprintReadWrite, Category=SkeletalControl) 
+	FBoneReference BoneToModify;
+
+	/** Target Bone to look at - you can't use LookAtLocation as alternative as you'll get a delay on bone location if you query directly **/
+	UPROPERTY(EditAnywhere, BlueprintReadWrite, Category=SkeletalControl)
+	FBoneReference LookAtBone;
+
+	/** Target Location in world space if LookAtBone is empty */
+	UPROPERTY(EditAnywhere, BlueprintReadWrite, Category=SkeletalControl)
+	FVector LookAtLocation;
+
+	/** Look at axis, which axis to align to look at point */
+	UPROPERTY(EditAnywhere, BlueprintReadWrite, Category=SkeletalControl) 
+	TEnumAsByte<EAxisOption::Type>	LookAtAxis;
+
+	/** Debug transient data */
+	FVector CurrentLookAtLocation;
+
+	// in the future, it would be nice to have more options, -i.e. lag, interpolation speed
+	FAnimNode_LookAt();
+
+	// FAnimNode_Base interface
+	virtual void GatherDebugData(FNodeDebugData& DebugData) OVERRIDE;
+	// End of FAnimNode_Base interface
+
+	// FAnimNode_SkeletalControlBase interface
+	virtual void EvaluateBoneTransforms(USkeletalMeshComponent* SkelComp, const FBoneContainer & RequiredBones, FA2CSPose& MeshBases, TArray<FBoneTransform>& OutBoneTransforms) OVERRIDE;
+	virtual bool IsValidToEvaluate(const USkeleton * Skeleton, const FBoneContainer & RequiredBones) OVERRIDE;
+	// End of FAnimNode_SkeletalControlBase interface
+
+private:
+	// FAnimNode_SkeletalControlBase interface
+	virtual void InitializeBoneReferences(const FBoneContainer & RequiredBones) OVERRIDE;
+	// End of FAnimNode_SkeletalControlBase interface
+
+	FVector GetAlignVector(FTransform & Transform, EAxisOption::Type AxisOption);
+};
diff --git a/Engine/Source/Runtime/Engine/Private/Animation/AnimNode_LookAt.cpp b/Engine/Source/Runtime/Engine/Private/Animation/AnimNode_LookAt.cpp
new file mode 100644
--- /dev/null
+++ b/Engine/Source/Runtime/Engine/Private/Animation/AnimNode_LookAt.cpp
@@ -0,0 +1,110 @@
+// Copyright 1998-2014 Epic Games, Inc. All Rights Reserved.
+
+#include "EnginePrivate.h"
+#include "AnimationRuntime.h"
+
+/////////////////////////////////////////////////////
+// FAnimNode_LookAt
+
+FAnimNode_LookAt::FAnimNode_LookAt()
+	: LookAtLocation(FVector::ZeroVector)
+	, LookAtAxis(EAxis::X)
+	, CurrentLookAtLocation(FVector::ZeroVector)
+{
+}
+
+void FAnimNode_LookAt::GatherDebugData(FNodeDebugData& DebugData)
+{
+	FString DebugLine = DebugData.GetNodeName(this);
+
+	DebugLine += "(";
+	AddDebugNodeData(DebugLine);
+	if (LookAtBone.BoneIndex != INDEX_NONE)
+	{
+		DebugLine += FString::Printf(TEXT(" Target: %s, Look At Bone: %s, Location : %s)"), *BoneToModify.BoneName.ToString(), *LookAtBone.BoneName.ToString(), *CurrentLookAtLocation.ToString());
+	}
+	else
+	{
+		DebugLine += FString::Printf(TEXT(" Target: %s, Look At Location : %s, Location : %s)"), *BoneToModify.BoneName.ToString(), *LookAtLocation.ToString(), *CurrentLookAtLocation.ToString());
+	}
+
+	DebugData.AddDebugItem(DebugLine);
+
+	ComponentPose.GatherDebugData(DebugData);
+}
+
+void FAnimNode_LookAt::EvaluateBoneTransforms(USkeletalMeshComponent* SkelComp, const FBoneContainer & RequiredBones, FA2CSPose& MeshBases, TArray<FBoneTransform>& OutBoneTransforms)
+{
+	check(OutBoneTransforms.Num() == 0);
+
+	FTransform ComponentBoneTransform = MeshBases.GetComponentSpaceTransform(BoneToModify.BoneIndex);
+
+	// get target location
+	FVector TargetLocationInComponentSpace;
+	if (LookAtBone.IsValid(RequiredBones))
+	{
+		FTransform LookAtTransform  = MeshBases.GetComponentSpaceTransform(LookAtBone.BoneIndex);
+		TargetLocationInComponentSpace = LookAtTransform.GetLocation();
+	}
+	else
+	{
+		TargetLocationInComponentSpace = SkelComp->ComponentToWorld.InverseTransformPosition(LookAtLocation);
+	}
+	
+	CurrentLookAtLocation = TargetLocationInComponentSpace;
+
+	// lookat vector
+	FVector LookAtVector = GetAlignVector(ComponentBoneTransform, LookAtAxis);
+	// flip to target vector if it wasnt negative
+	bool bShouldFlip = LookAtAxis == EAxisOption::X_Neg || LookAtAxis == EAxisOption::Y_Neg || LookAtAxis == EAxisOption::Z_Neg;
+	FVector ToTarget = CurrentLookAtLocation - ComponentBoneTransform.GetLocation();
+	ToTarget.Normalize();
+	if (bShouldFlip)
+	{
+		ToTarget *= -1.f;
+	}
+	
+	// get delta rotation
+	FQuat DeltaRot = FQuat::FindBetween(LookAtVector, ToTarget);
+
+	// transform current rotation to delta rotation
+	FQuat CurrentRot = ComponentBoneTransform.GetRotation();
+	FQuat NewRotation = DeltaRot * CurrentRot;
+	ComponentBoneTransform.SetRotation(NewRotation);
+
+	OutBoneTransforms.Add( FBoneTransform(BoneToModify.BoneIndex, ComponentBoneTransform) );
+}
+
+bool FAnimNode_LookAt::IsValidToEvaluate(const USkeleton * Skeleton, const FBoneContainer & RequiredBones) 
+{
+	// if both bones are valid
+	return (BoneToModify.IsValid(RequiredBones) && 
+		// or if name isn't set (use Look At Location) or Look at bone is valid 
+		// do not call isValid since that means if look at bone isn't in LOD, we won't evaluate
+		// we still should evaluate as long as the BoneToModify is valid even LookAtBone isn't included in required bones
+		(LookAtBone.BoneName == NAME_None || LookAtBone.BoneIndex != INDEX_NONE) );
+}
+
+void FAnimNode_LookAt::InitializeBoneReferences(const FBoneContainer & RequiredBones) 
+{
+	BoneToModify.Initialize(RequiredBones);
+	LookAtBone.Initialize(RequiredBones);
+}
+
+FVector FAnimNode_LookAt::GetAlignVector(FTransform & Transform, EAxisOption::Type AxisOption)
+{
+	switch (AxisOption)
+	{
+	case EAxisOption::X:
+	case EAxisOption::X_Neg:
+		return Transform.GetUnitAxis(EAxis::X);
+	case EAxisOption::Y:
+	case EAxisOption::Y_Neg:
+		return Transform.GetUnitAxis(EAxis::Y);
+	case EAxisOption::Z:
+	case EAxisOption::Z_Neg:
+		return Transform.GetUnitAxis(EAxis::Z);
+	}
+
+	return FVector(1.f, 0.f, 0.f);
+}
\ No newline at end of file
