Make UMaterialInstance::UpdateStaticPermutation call InitStaticPermutation only for instances actually dependent on the current instance, rather than for all instances with the same parent UMaterial.

Fixes TTP#333845 LANDSCAPE: Importing a layers for a level split into multiple sublevels causes many thousands of shaders to be compiled

diff --git a/Engine/Source/Runtime/Engine/Private/Materials/MaterialInstance.cpp b/Engine/Source/Runtime/Engine/Private/Materials/MaterialInstance.cpp
--- a/Engine/Source/Runtime/Engine/Private/Materials/MaterialInstance.cpp
+++ b/Engine/Source/Runtime/Engine/Private/Materials/MaterialInstance.cpp
@@ -1924,7 +1924,7 @@ void UMaterialInstance::UpdateStaticPermutation(const FStaticParameterSet& NewPa
 		// This will flush the rendering thread which is necessary before changing bHasStaticPermutationResource, since the RT is reading from that directly
 		// The update context will also make sure any dependent MI's with static parameters get recompiled
 		FMaterialUpdateContext MaterialUpdateContext;
-		MaterialUpdateContext.AddMaterial(GetMaterial());
+		MaterialUpdateContext.AddMaterialInstance(this);
 		bHasStaticPermutationResource = bWantsStaticPermutationResource;
 		StaticParameters = CompareParameters;
 
diff --git a/Engine/Source/Runtime/Engine/Private/Materials/MaterialShared.cpp b/Engine/Source/Runtime/Engine/Private/Materials/MaterialShared.cpp
--- a/Engine/Source/Runtime/Engine/Private/Materials/MaterialShared.cpp
+++ b/Engine/Source/Runtime/Engine/Private/Materials/MaterialShared.cpp
@@ -1937,6 +1937,18 @@ FMaterialUpdateContext::FMaterialUpdateContext(uint32 Options, EShaderPlatform I
 	ShaderPlatform = InShaderPlatform;
 }
 
+void FMaterialUpdateContext::AddMaterial(UMaterial* Material)
+{
+	UpdatedMaterials.Add(Material);
+	UpdatedMaterialInterfaces.Add(Material);
+}
+
+void FMaterialUpdateContext::AddMaterialInstance(UMaterialInstance* Instance)
+{
+	UpdatedMaterials.Add(Instance->GetMaterial());
+	UpdatedMaterialInterfaces.Add(Instance);
+}
+
 FMaterialUpdateContext::~FMaterialUpdateContext()
 {
 	double StartTime = FPlatformTime::Seconds();
@@ -1989,7 +2001,15 @@ FMaterialUpdateContext::~FMaterialUpdateContext()
 
 		if (UpdatedMaterials.Contains(BaseMaterial))
 		{
-			InstancesToUpdate.Add(CurrentMaterialInstance);
+			// Check to see if this instance is dependent on any of the material interfaces we directly updated.
+			for (auto InterfaceIt = UpdatedMaterialInterfaces.CreateConstIterator(); InterfaceIt; ++InterfaceIt)
+			{
+				if (CurrentMaterialInstance->IsDependent(*InterfaceIt))
+				{
+					InstancesToUpdate.Add(CurrentMaterialInstance);
+					break;
+				}
+			}
 		}
 	}
 
@@ -2047,10 +2067,11 @@ FMaterialUpdateContext::~FMaterialUpdateContext()
 	}
 
 	double EndTime = FPlatformTime::Seconds();
-	UE_LOG(LogMaterial,Log,
-		TEXT("%f seconds spent updating %d materials, %d instances, %d with static permutations."),
-		(float)(EndTime-StartTime),
+	UE_LOG(LogMaterial, Log,
+		TEXT("%f seconds spent updating %d materials, %d interfaces, %d instances, %d with static permutations."),
+		(float)(EndTime - StartTime),
 		UpdatedMaterials.Num(),
+		UpdatedMaterialInterfaces.Num(),
 		InstancesToUpdate.Num(),
 		NumInstancesWithStaticPermutations
 		);
diff --git a/Engine/Source/Runtime/Engine/Public/MaterialShared.h b/Engine/Source/Runtime/Engine/Public/MaterialShared.h
--- a/Engine/Source/Runtime/Engine/Public/MaterialShared.h
+++ b/Engine/Source/Runtime/Engine/Public/MaterialShared.h
@@ -1808,8 +1808,10 @@ protected:
  */
 class FMaterialUpdateContext
 {
-	/** Materials updated within this context. */
+	/** UMaterial parents of any UMaterialInterfaces updated within this context. */
 	TSet<UMaterial*> UpdatedMaterials;
+	/** Materials updated within this context. */
+	TSet<UMaterialInterface*> UpdatedMaterialInterfaces;
 	/** Active global component reregister context, if any. */
 	TScopedPointer<class FGlobalComponentReregisterContext> ComponentReregisterContext;
 	/** The shader platform that was being processed - can control if we need to update components */
@@ -1844,10 +1846,10 @@ public:
 	ENGINE_API ~FMaterialUpdateContext();
 
 	/** Add a material that has been updated to the context. */
-	void AddMaterial(UMaterial* Material)
-	{
-		UpdatedMaterials.Add(Material);
-	}
+	ENGINE_API void AddMaterial(UMaterial* Material);
+
+	/** Adds a material instance that has been updated to the context. */
+	ENGINE_API void AddMaterialInstance(UMaterialInstance* Instance);
 };
 
 /**
