Fix wheel bodies of vehicle from ragdolling. The bodies of the wheel is automatically disabled

diff --git a/Engine/Source/Runtime/Engine/Classes/Vehicles/WheeledVehicleMovementComponent.h b/Engine/Source/Runtime/Engine/Classes/Vehicles/WheeledVehicleMovementComponent.h
--- a/Engine/Source/Runtime/Engine/Classes/Vehicles/WheeledVehicleMovementComponent.h
+++ b/Engine/Source/Runtime/Engine/Classes/Vehicles/WheeledVehicleMovementComponent.h
@@ -280,6 +280,9 @@ class ENGINE_API UWheeledVehicleMovementComponent : public UPawnMovementComponen
 	/** Draw debug lines for the wheels and suspension */
 	virtual void DrawDebugLines();
 
+	/** Skeletal mesh needs some special handling in the vehicle case */
+	virtual void FixupSkeletalMesh();
+
 #if WITH_EDITOR
 	/** Respond to a property change in editor */
 	virtual void PostEditChangeProperty(struct FPropertyChangedEvent& PropertyChangedEvent) OVERRIDE;
diff --git a/Engine/Source/Runtime/Engine/Private/Vehicles/WheeledVehicleMovementComponent.cpp b/Engine/Source/Runtime/Engine/Private/Vehicles/WheeledVehicleMovementComponent.cpp
--- a/Engine/Source/Runtime/Engine/Private/Vehicles/WheeledVehicleMovementComponent.cpp
+++ b/Engine/Source/Runtime/Engine/Private/Vehicles/WheeledVehicleMovementComponent.cpp
@@ -3,6 +3,7 @@
 #include "EnginePrivate.h"
 #include "Net/UnrealNetwork.h"
 #include "MessageLog.h"
+#include "PhysicsAssetUtils.h"
 
 #if WITH_PHYSX
 #include "../PhysicsEngine/PhysXSupport.h"
@@ -481,15 +482,6 @@ FVector UWheeledVehicleMovementComponent::GetWheelRestingPosition( const FWheelS
 	if ( WheelSetup.BoneName != NAME_None )
 	{
 		USkinnedMeshComponent* Mesh = GetMesh();
-
-		//warn if wheels have a body associated with it
-		FBodyInstance * WheelBody = Mesh->GetBodyInstance(WheelSetup.BoneName);
-		if (WheelBody)
-		{
-			UE_LOG(LogPhysics, Warning, TEXT("Wheel %s has a body instance associated with it. These are created by vehicle"), *WheelSetup.BoneName.ToString());
-			FMessageLog("PIE").Warning(FText::Format(LOCTEXT("WheelHasBodyWarning", "The wheel bone {0} already has a body setup in the physics asset. The vehicle simulation sets this body up automatically, the wheel body should be removed"), FText::FromString(WheelSetup.BoneName.ToString())));
-		}
-
 		if ( Mesh && Mesh->SkeletalMesh )
 		{
 			const FVector BonePosition = Mesh->SkeletalMesh->GetComposedRefPoseMatrix( WheelSetup.BoneName ).GetOrigin() * Mesh->RelativeScale3D;
@@ -561,6 +553,7 @@ void UWheeledVehicleMovementComponent::CreatePhysicsState()
 
 		if ( PhysScene && PhysScene->GetVehicleManager() )
 		{
+			FixupSkeletalMesh();	//this must come first because it may trigger a rebuild of the skeletal mesh which will invalidate the vehicle setup
 			CreateVehicle();
 
 			if ( PVehicle )
@@ -1215,6 +1208,62 @@ void UWheeledVehicleMovementComponent::DrawDebug( UCanvas* Canvas, float& YL, fl
 	DrawDebugLines();
 }
 
+void UWheeledVehicleMovementComponent::FixupSkeletalMesh()
+{
+	if (USkinnedMeshComponent* Mesh = GetMesh())
+	{
+		if (UPhysicsAsset * PhysicsAsset = Mesh->GetPhysicsAsset())
+		{
+			for (int32 WheelIdx = 0; WheelIdx < WheelSetups.Num(); ++WheelIdx)
+			{
+				FWheelSetup& WheelSetup = WheelSetups[WheelIdx];
+				if (WheelSetup.BoneName != NAME_None)
+				{
+					int32 BodySetupIdx = PhysicsAsset->FindBodyIndex(WheelSetup.BoneName);
+					
+					if (BodySetupIdx >= 0)
+					{
+						if (UBodySetup * BodySetup = PhysicsAsset->BodySetup[BodySetupIdx])
+						{
+							BodySetup->PhysicsType = PhysType_Fixed; //we want wheel to be animated
+							BodySetup->CollisionReponse = EBodyCollisionResponse::BodyCollision_Disabled;	//and have no collision
+
+							//we also need to make sure the wheel body has no constraints with the parent body
+							TArray<int32> WheelConstraints;
+							FPhysicsAssetUtils::BodyFindConstraints(PhysicsAsset, BodySetupIdx, WheelConstraints);
+							for (int32 ConstraintIdx = 0; ConstraintIdx < WheelConstraints.Num(); ++ConstraintIdx)
+							{
+								UPhysicsConstraintTemplate * ConstraintTemplate = PhysicsAsset->ConstraintSetup[WheelConstraints[ConstraintIdx]];
+								FConstraintInstance & ConstraintInstance = ConstraintTemplate->DefaultInstance;
+								ConstraintInstance.LinearXMotion = LCM_Free;
+								ConstraintInstance.LinearYMotion = LCM_Free;
+								ConstraintInstance.LinearZMotion = LCM_Free;
+								ConstraintInstance.AngularSwing1Motion = ACM_Free;
+								ConstraintInstance.AngularSwing2Motion = ACM_Free;
+								ConstraintInstance.AngularTwistMotion = ACM_Free;
+							}
+						}
+					}
+					else
+					{
+						UE_LOG(LogPhysics, Warning, TEXT("UWheeledVehicleMovementComponent::FixupSkeletalMesh Wheel [%d] references a bone [%s] that has no body in its skeletal mesh."), WheelIdx, *WheelSetup.BoneName.ToString());
+						FMessageLog("PIE").Warning(FText::Format(LOCTEXT("WheelHasBodyWarning", "Cannot find body for wheel bone {0}."), FText::FromString(WheelSetup.BoneName.ToString())));
+					}
+
+				}
+			}
+		}
+
+		// recreate skeletal mesh using new template
+		if (Mesh->IsPhysicsStateCreated())
+		{
+			Mesh->RecreatePhysicsState();
+		}
+	}
+
+
+}
+
 void UWheeledVehicleMovementComponent::DrawDebugLines()
 {
 	if ( PVehicle == NULL )
