Make SkyLight blueprintable

diff --git a/Engine/Source/Runtime/Engine/Classes/Engine/SkyLight.h b/Engine/Source/Runtime/Engine/Classes/Engine/SkyLight.h
--- a/Engine/Source/Runtime/Engine/Classes/Engine/SkyLight.h
+++ b/Engine/Source/Runtime/Engine/Classes/Engine/SkyLight.h
@@ -4,7 +4,7 @@
 
 #include "SkyLight.generated.h"
 
-UCLASS(ClassGroup=Lights, hidecategories=(Input,Collision,Replication,Info), showcategories=("Input|MouseInput", "Input|TouchInput"), ConversionRoot)
+UCLASS(ClassGroup=Lights, hidecategories=(Input,Collision,Replication,Info), showcategories=("Input|MouseInput", "Input|TouchInput"), ConversionRoot, Blueprintable)
 class ENGINE_API ASkyLight : public AInfo
 {
 	GENERATED_UCLASS_BODY()
