Back out changelist 2055240

Decided there was not a need at this time to mark Directional, Point and SpotLight as ENGINE_API. Seems like rather than extending those classes, users should use blueprints or create their own actor class and add light components to it.

diff --git a/Engine/Source/Runtime/Engine/Classes/Engine/DirectionalLight.h b/Engine/Source/Runtime/Engine/Classes/Engine/DirectionalLight.h
--- a/Engine/Source/Runtime/Engine/Classes/Engine/DirectionalLight.h
+++ b/Engine/Source/Runtime/Engine/Classes/Engine/DirectionalLight.h
@@ -4,8 +4,8 @@
 
 #include "DirectionalLight.generated.h"
 
-UCLASS(ClassGroup=(Lights, DirectionalLights), meta=(ChildCanTick))
-class ENGINE_API ADirectionalLight : public ALight
+UCLASS(ClassGroup=(Lights, DirectionalLights), MinimalAPI, meta=(ChildCanTick))
+class ADirectionalLight : public ALight
 {
 	GENERATED_UCLASS_BODY()
 
diff --git a/Engine/Source/Runtime/Engine/Classes/Engine/PointLight.h b/Engine/Source/Runtime/Engine/Classes/Engine/PointLight.h
--- a/Engine/Source/Runtime/Engine/Classes/Engine/PointLight.h
+++ b/Engine/Source/Runtime/Engine/Classes/Engine/PointLight.h
@@ -4,8 +4,8 @@
 
 #include "PointLight.generated.h"
 
-UCLASS(ClassGroup=(Lights, PointLights), meta=(ChildCanTick))
-class ENGINE_API APointLight : public ALight
+UCLASS(ClassGroup=(Lights, PointLights),MinimalAPI, meta=(ChildCanTick))
+class APointLight : public ALight
 {
 	GENERATED_UCLASS_BODY()
 
@@ -14,9 +14,9 @@ class ENGINE_API APointLight : public ALight
 
 	// BEGIN DEPRECATED (use component functions now in level script)
 	UFUNCTION(BlueprintCallable, Category="Rendering|Lighting", meta=(DeprecatedFunction))
-	void SetRadius(float NewRadius);
+	ENGINE_API void SetRadius(float NewRadius);
 	UFUNCTION(BlueprintCallable, Category="Rendering|Lighting", meta=(DeprecatedFunction))
-	void SetLightFalloffExponent(float NewLightFalloffExponent);
+	ENGINE_API void SetLightFalloffExponent(float NewLightFalloffExponent);
 	// END DEPRECATED
 
 #if WITH_EDITOR
diff --git a/Engine/Source/Runtime/Engine/Classes/Engine/SpotLight.h b/Engine/Source/Runtime/Engine/Classes/Engine/SpotLight.h
--- a/Engine/Source/Runtime/Engine/Classes/Engine/SpotLight.h
+++ b/Engine/Source/Runtime/Engine/Classes/Engine/SpotLight.h
@@ -5,8 +5,8 @@
 
 #include "SpotLight.generated.h"
 
-UCLASS(ClassGroup=(Lights, SpotLights), meta=(ChildCanTick))
-class ENGINE_API ASpotLight : public ALight
+UCLASS(ClassGroup=(Lights, SpotLights), MinimalAPI, meta=(ChildCanTick))
+class ASpotLight : public ALight
 {
 	GENERATED_UCLASS_BODY()
 
