Clear GPU particle system textures on creation to prevent random massive particles when GPU particle emitters spawn.

diff --git a/Engine/Source/Runtime/Engine/Private/Particles/ParticleGpuSimulation.cpp b/Engine/Source/Runtime/Engine/Private/Particles/ParticleGpuSimulation.cpp
--- a/Engine/Source/Runtime/Engine/Private/Particles/ParticleGpuSimulation.cpp
+++ b/Engine/Source/Runtime/Engine/Private/Particles/ParticleGpuSimulation.cpp
@@ -180,6 +180,9 @@ public:
 			PositionTextureRHI
 			);
 
+		RHISetRenderTarget( PositionTextureTargetRHI, FTextureRHIRef() );
+		RHIClear( true, FLinearColor::Black, false, 1.0f, false, 0, FIntRect() );
+
 		// 16-bit per channel RGBA texture for velocity.
 		check( !IsValidRef( VelocityTextureTargetRHI ) );
 		check( !IsValidRef( VelocityTextureRHI ) );
@@ -195,6 +198,9 @@ public:
 			VelocityTextureTargetRHI,
 			VelocityTextureRHI
 			);
+	
+		RHISetRenderTarget( VelocityTextureTargetRHI, FTextureRHIRef()); 
+		RHIClear( true, FLinearColor::Black, false, 1.0f, false, 0, FIntRect() );
 	}
 
 	/**
