Landscape Spline : Removed NULL case for control points/segments

diff --git a/Engine/Source/Runtime/Engine/Private/Landscape/LandscapeSplines.cpp b/Engine/Source/Runtime/Engine/Private/Landscape/LandscapeSplines.cpp
--- a/Engine/Source/Runtime/Engine/Private/Landscape/LandscapeSplines.cpp
+++ b/Engine/Source/Runtime/Engine/Private/Landscape/LandscapeSplines.cpp
@@ -421,29 +421,8 @@ void ULandscapeSplinesComponent::PostLoad()
 {
 	Super::PostLoad();
 
-	for (int32 iControlPoint = 0; iControlPoint < ControlPoints.Num(); )
-	{
-		if (ULandscapeSplineControlPoint* ControlPoint = ControlPoints[iControlPoint])
-		{
-			iControlPoint++;
-		}
-		else
-		{
-			ControlPoints.RemoveAt(iControlPoint);
-		}
-	}
-
-	for (int32 iSegment = 0; iSegment < Segments.Num(); )
-	{
-		if (ULandscapeSplineSegment* Segment = Segments[iSegment])
-		{
-			iSegment++;
-		}
-		else
-		{
-			Segments.RemoveAt(iSegment);
-		}
-	}
+	ControlPoints.Remove(NULL);
+	Segments.Remove(NULL);
 }
 
 #if WITH_EDITOR
