Fix a crash when an FSceneView doesn't contain a scene

diff --git a/Engine/Source/Runtime/Engine/Private/SceneView.cpp b/Engine/Source/Runtime/Engine/Private/SceneView.cpp
--- a/Engine/Source/Runtime/Engine/Private/SceneView.cpp
+++ b/Engine/Source/Runtime/Engine/Private/SceneView.cpp
@@ -1106,7 +1106,14 @@ void FSceneView::ConfigureBufferVisualizationSettings()
 
 ERHIFeatureLevel::Type FSceneView::GetFeatureLevel() const
 { 
-	return Family->Scene->GetFeatureLevel();
+	if (Family->Scene)
+	{
+		return Family->Scene->GetFeatureLevel();
+	}
+	else
+	{
+		return GRHIFeatureLevel;
+	}
 }
 
 FSceneViewFamily::FSceneViewFamily( const ConstructionValues& CVS )
