Speculative fix for crash detailed in TTP 334843 - TD: MP_CRATER_TURRET_4: Host is crashing with an "Assertion failed: !ContainsNaN()" crash during gameplay

diff --git a/Engine/Source/Runtime/Core/Public/Math/UnrealMath.h b/Engine/Source/Runtime/Core/Public/Math/UnrealMath.h
--- a/Engine/Source/Runtime/Core/Public/Math/UnrealMath.h
+++ b/Engine/Source/Runtime/Core/Public/Math/UnrealMath.h
@@ -437,6 +437,36 @@ FORCEINLINE_DEBUGGABLE U FMath::CubicCRSplineInterp(const U& P0, const U& P1, co
 	return  ( ( L012 * ((T2 - T) * InvT2MinusT1) ) + ( L123 * ((T - T1) * InvT2MinusT1) ) );
 }
 
+template< class U >
+FORCEINLINE_DEBUGGABLE U FMath::CubicCRSplineInterpSafe(const U& P0, const U& P1, const U& P2, const U& P3, const float T0, const float T1, const float T2, const float T3, const float T)
+{
+	//Based on http://www.cemyuksel.com/research/catmullrom_param/catmullrom.pdf 
+	float T1MinusT0 = (T1 - T0);
+	float T2MinusT1 = (T2 - T1);
+	float T3MinusT2 = (T3 - T2);
+	float T2MinusT0 = (T2 - T0);
+	float T3MinusT1 = (T3 - T1);
+	if (FMath::IsNearlyZero(T1MinusT0) || FMath::IsNearlyZero(T2MinusT1) || FMath::IsNearlyZero(T3MinusT2) || FMath::IsNearlyZero(T2MinusT0) || FMath::IsNearlyZero(T3MinusT1))
+	{
+		//There's going to be a divide by zero here so just bail out and return P1
+		return P1;
+	}
+
+	float InvT1MinusT0 = 1.0f / T1MinusT0;
+	U L01 = (P0 * ((T1 - T) * InvT1MinusT0)) + (P1 * ((T - T0) * InvT1MinusT0));
+	float InvT2MinusT1 = 1.0f / T2MinusT1;
+	U L12 = (P1 * ((T2 - T) * InvT2MinusT1)) + (P2 * ((T - T1) * InvT2MinusT1));
+	float InvT3MinusT2 = 1.0f / T3MinusT2;
+	U L23 = (P2 * ((T3 - T) * InvT3MinusT2)) + (P3 * ((T - T2) * InvT3MinusT2));
+
+	float InvT2MinusT0 = 1.0f / T2MinusT0;
+	U L012 = (L01 * ((T2 - T) * InvT2MinusT0)) + (L12 * ((T - T0) * InvT2MinusT0));
+	float InvT3MinusT1 = 1.0f / T3MinusT1;
+	U L123 = (L12 * ((T3 - T) * InvT3MinusT1)) + (L23 * ((T - T1) * InvT3MinusT1));
+
+	return  ((L012 * ((T2 - T) * InvT2MinusT1)) + (L123 * ((T - T1) * InvT2MinusT1)));
+}
+
 /**
  * Performs a sphere vs box intersection test using Arvo's algorithm:
  *
diff --git a/Engine/Source/Runtime/Core/Public/Math/UnrealMathUtility.h b/Engine/Source/Runtime/Core/Public/Math/UnrealMathUtility.h
--- a/Engine/Source/Runtime/Core/Public/Math/UnrealMathUtility.h
+++ b/Engine/Source/Runtime/Core/Public/Math/UnrealMathUtility.h
@@ -562,6 +562,9 @@ struct FMath : public FPlatformMath
 	 */
 	template< class U > static U CubicCRSplineInterp(const U& P0, const U& P1, const U& P2, const U& P3, const float T0, const float T1, const float T2, const float T3, const float T);
 
+	/* Same as CubicCRSplineInterp but with additional saftey checks. If the checks fail P1 is returned. **/
+	template< class U > static U CubicCRSplineInterpSafe(const U& P0, const U& P1, const U& P2, const U& P3, const float T0, const float T1, const float T2, const float T3, const float T);
+	
 	// Special-case interpolation
 
 	/** Interpolate a normal vector Current to Target, by interpolating the angle between those vectors with constant step. */
diff --git a/Engine/Source/Runtime/Engine/Private/Particles/ParticleSystemRender.cpp b/Engine/Source/Runtime/Engine/Private/Particles/ParticleSystemRender.cpp
--- a/Engine/Source/Runtime/Engine/Private/Particles/ParticleSystemRender.cpp
+++ b/Engine/Source/Runtime/Engine/Private/Particles/ParticleSystemRender.cpp
@@ -6435,17 +6435,17 @@ struct FAnimTrailParticleRenderData
 			float T = CurrT + ((PrevT - CurrT) * InterpFactor);
 				
 			//Interpolate locations
-			FVector Location = FMath::CubicCRSplineInterp(PrevPrevLocation, PrevParticle->Location, Particle->Location, NextLocation, PrevPrevT, PrevT, CurrT, NextT, T);
-			FVector InterpDir = FMath::CubicCRSplineInterp(PrevPrevDirection, PrevPayload->Direction, Payload->Direction, NextDirection, PrevPrevT, PrevT, CurrT, NextT, T);
+			FVector Location = FMath::CubicCRSplineInterpSafe(PrevPrevLocation, PrevParticle->Location, Particle->Location, NextLocation, PrevPrevT, PrevT, CurrT, NextT, T);
+			FVector InterpDir = FMath::CubicCRSplineInterpSafe(PrevPrevDirection, PrevPayload->Direction, Payload->Direction, NextDirection, PrevPrevT, PrevT, CurrT, NextT, T);
 			InterpDir.Normalize();
-			float InterpLength = FMath::CubicCRSplineInterp(PrevPrevLength, PrevPayload->Length, Payload->Length, NextLength, PrevPrevT, PrevT, CurrT, NextT, T);
-			OutTileU = FMath::CubicCRSplineInterp(PrevPrevTiledU, PrevPayload->TiledU, Payload->TiledU, NextTiledU, PrevPrevT, PrevT, CurrT, NextT, T);
-			OutSize = FMath::CubicCRSplineInterp(PrevPrevSize, PrevParticle->Size.X * Source.Scale.X, Particle->Size.X * Source.Scale.X, NextSize, PrevPrevT, PrevT, CurrT, NextT, T);
-			OutColor = FMath::CubicCRSplineInterp(PrevPrevColor, PrevParticle->Color, Particle->Color, NextColor, PrevPrevT, PrevT, CurrT, NextT, T);
+			float InterpLength = FMath::CubicCRSplineInterpSafe(PrevPrevLength, PrevPayload->Length, Payload->Length, NextLength, PrevPrevT, PrevT, CurrT, NextT, T);
+			OutTileU = FMath::CubicCRSplineInterpSafe(PrevPrevTiledU, PrevPayload->TiledU, Payload->TiledU, NextTiledU, PrevPrevT, PrevT, CurrT, NextT, T);
+			OutSize = FMath::CubicCRSplineInterpSafe(PrevPrevSize, PrevParticle->Size.X * Source.Scale.X, Particle->Size.X * Source.Scale.X, NextSize, PrevPrevT, PrevT, CurrT, NextT, T);
+			OutColor = FMath::CubicCRSplineInterpSafe(PrevPrevColor, PrevParticle->Color, Particle->Color, NextColor, PrevPrevT, PrevT, CurrT, NextT, T);
 
 			if( OutDynamicParameters )
 			{
-				*OutDynamicParameters = FMath::CubicCRSplineInterp(PrevPrevDynPayload->DynamicParameterValue,
+				*OutDynamicParameters = FMath::CubicCRSplineInterpSafe(PrevPrevDynPayload->DynamicParameterValue,
 					PrevDynPayload->DynamicParameterValue,
 					CurrDynPayload->DynamicParameterValue,
 					NextDynPayload->DynamicParameterValue, PrevPrevT, PrevT, CurrT, NextT, T);
