Added timed particle notify to play and stop looping particle effects in an animation
#ttp 334635 - TD: Global: need ability to turn on/off a looping emitter in anim notifies.

diff --git a/Engine/Source/Editor/Persona/Private/AnimNotifyDetails.cpp b/Engine/Source/Editor/Persona/Private/AnimNotifyDetails.cpp
--- a/Engine/Source/Editor/Persona/Private/AnimNotifyDetails.cpp
+++ b/Engine/Source/Editor/Persona/Private/AnimNotifyDetails.cpp
@@ -33,6 +33,10 @@ void FAnimNotifyDetails::CustomizeDetails( IDetailLayoutBuilder& DetailBuilder )
 		AddBoneNameProperty(DetailBuilder, DetailObjectClass, TEXT("FirstSocketName"), TEXT("Trail"));
 		AddBoneNameProperty(DetailBuilder, DetailObjectClass, TEXT("SecondSocketName"), TEXT("Trail"));
 	}
+	else if(DetailObjectClass->GetName().Find(TEXT("AnimNotifyState_TimedParticleEffect")) != INDEX_NONE)
+	{
+		AddBoneNameProperty(DetailBuilder, DetailObjectClass, TEXT("SocketName"), TEXT("ParticleSystem"));
+	}
 }
 
 void FAnimNotifyDetails::AddBoneNameProperty(IDetailLayoutBuilder& DetailBuilder, const UClass* PropertyClass, const TCHAR* PropertyName, const TCHAR* CategoryName)
diff --git a/Engine/Source/Runtime/Engine/Classes/Animation/AnimNotifies/AnimNotifyState_TimedParticleEffect.h b/Engine/Source/Runtime/Engine/Classes/Animation/AnimNotifies/AnimNotifyState_TimedParticleEffect.h
new file mode 100644
--- /dev/null
+++ b/Engine/Source/Runtime/Engine/Classes/Animation/AnimNotifies/AnimNotifyState_TimedParticleEffect.h
@@ -0,0 +1,56 @@
+// Copyright 1998-2014 Epic Games, Inc. All Rights Reserved.
+
+#pragma once
+#include "AnimNotifyState_TimedParticleEffect.generated.h"
+
+// Timed Particle Effect Notify
+// Allows a looping particle effect to be played in an animation that will activate
+// at the beginning of the notify and deactivate at the end.
+UCLASS(MinimalAPI, Blueprintable, meta = (DisplayName = "Timed Particle Effect"))
+class UAnimNotifyState_TimedParticleEffect : public UAnimNotifyState
+{
+	GENERATED_UCLASS_BODY()
+
+	// The particle system template to use when spawning the particle component
+	UPROPERTY(EditAnywhere, Category = ParticleSystem, meta = (ToolTip = "The particle system to spawn for the notify state"))
+	UParticleSystem* PSTemplate;
+
+	// The socket within our mesh component to attach to when we spawn the particle componnet
+	UPROPERTY(EditAnywhere, Category = ParticleSystem, meta = (ToolTip = "The socket or bone to attach the system to"))
+	FName SocketName;
+
+	// Offset from the socket / bone location
+	UPROPERTY(EditAnywhere, Category = ParticleSystem, meta = (ToolTip = "Offset from the socket or bone to place the particle system"))
+	FVector LocationOffset;
+
+	// Offset from the socket / bone rotation
+	UPROPERTY(EditAnywhere, Category = ParticleSystem, meta = (ToolTip = "Rotation offset from the socket or bone for the particle system"))
+	FRotator RotationOffset;
+
+	// Whether or not we destroy the component at the end of the notify or instead just stop
+	// the emitters.
+	UPROPERTY(EditAnywhere, Category = ParticleSystem, meta = (DisplayName = "Destroy Immediately", ToolTip = "Whether the particle system should be immediately destroyed at the end of the notify state or be allowed to finish"))
+	bool bDestroyAtEnd;
+
+#if WITH_EDITORONLY_DATA
+	// The following arrays are used to handle property changes during a state. Because we can't
+	// store any stateful data here we can't know which emitter is ours. The best metric we have
+	// is an emitter on our Mesh Component with the same template and socket name we have defined.
+	// Because these can change at any time we need to track previous versions when we are in an
+	// editor build. Refactor when stateful data is possible, tracking our component instead.
+	UPROPERTY(transient)
+	TArray<UParticleSystem*> PreviousPSTemplates;
+
+	UPROPERTY(transient)
+	TArray<FName> PreviousSocketNames;
+
+	virtual void PreEditChange(UProperty* PropertyAboutToChange) OVERRIDE;
+#endif
+
+	virtual void NotifyBegin(class USkeletalMeshComponent * MeshComp, class UAnimSequence * AnimSeq) OVERRIDE;
+	virtual void NotifyTick(class USkeletalMeshComponent * MeshComp, class UAnimSequence * AnimSeq, float FrameDeltaTime) OVERRIDE;
+	virtual void NotifyEnd(class USkeletalMeshComponent * MeshComp, class UAnimSequence * AnimSeq) OVERRIDE;
+
+private:
+	bool ValidateParameters(USkeletalMeshComponent* MeshComp);
+};
\ No newline at end of file
diff --git a/Engine/Source/Runtime/Engine/Private/Animation/AnimNotifyState_TimedParticleEffect.cpp b/Engine/Source/Runtime/Engine/Private/Animation/AnimNotifyState_TimedParticleEffect.cpp
new file mode 100644
--- /dev/null
+++ b/Engine/Source/Runtime/Engine/Private/Animation/AnimNotifyState_TimedParticleEffect.cpp
@@ -0,0 +1,106 @@
+// Copyright 1998-2014 Epic Games, Inc. All Rights Reserved.
+
+#include "EnginePrivate.h"
+
+UAnimNotifyState_TimedParticleEffect::UAnimNotifyState_TimedParticleEffect(const class FPostConstructInitializeProperties& PCIP)
+: Super(PCIP)
+{
+	PSTemplate = nullptr;
+	LocationOffset.Set(0.0f, 0.0f, 0.0f);
+	RotationOffset = FRotator(0.0f, 0.0f, 0.0f);
+}
+
+void UAnimNotifyState_TimedParticleEffect::NotifyBegin(USkeletalMeshComponent * MeshComp, class UAnimSequence * AnimSeq)
+{
+	// Only spawn if we've got valid params
+	if(ValidateParameters(MeshComp))
+	{
+		UParticleSystemComponent* NewComponent = UGameplayStatics::SpawnEmitterAttached(PSTemplate, MeshComp, SocketName, LocationOffset, RotationOffset);
+	}
+
+	Received_NotifyBegin(MeshComp, AnimSeq);
+}
+
+void UAnimNotifyState_TimedParticleEffect::NotifyTick(USkeletalMeshComponent * MeshComp, class UAnimSequence * AnimSeq, float FrameDeltaTime)
+{
+	Received_NotifyTick(MeshComp, AnimSeq, FrameDeltaTime);
+}
+
+void UAnimNotifyState_TimedParticleEffect::NotifyEnd(USkeletalMeshComponent * MeshComp, class UAnimSequence * AnimSeq)
+{
+	TArray<USceneComponent*> Children;
+	MeshComp->GetChildrenComponents(false, Children);
+
+	for(USceneComponent* Component : Children)
+	{
+		if(UParticleSystemComponent* ParticleComponent = Cast<UParticleSystemComponent>(Component))
+		{
+			bool bSocketMatch = ParticleComponent->AttachSocketName == SocketName;
+			bool bTemplateMatch = ParticleComponent->Template == PSTemplate;
+
+#if WITH_EDITORONLY_DATA
+			// In editor someone might have changed our parameters while we're ticking; so check 
+			// previous known parameters too.
+			bSocketMatch |= PreviousSocketNames.Contains(ParticleComponent->AttachSocketName);
+			bTemplateMatch |= PreviousPSTemplates.Contains(ParticleComponent->Template);
+#endif
+
+			if(bSocketMatch && bTemplateMatch)
+			{
+				// Either destroy the component or deactivate it to have it's active particles finish.
+				// The component will auto destroy once all particle are gone.
+				if(bDestroyAtEnd)
+				{
+					ParticleComponent->DestroyComponent();
+				}
+				else
+				{
+					ParticleComponent->bAutoDestroy = true;
+					ParticleComponent->DeactivateSystem();
+				}
+
+#if WITH_EDITORONLY_DATA
+				// No longer need to track previous values as we've found our component
+				// and removed it.
+				PreviousPSTemplates.Empty();
+				PreviousSocketNames.Empty();
+#endif
+				// Removed a component, no need to continue
+				break;
+			}
+		}
+	}
+
+	Received_NotifyEnd(MeshComp, AnimSeq);
+}
+
+bool UAnimNotifyState_TimedParticleEffect::ValidateParameters(USkeletalMeshComponent* MeshComp)
+{
+	bool bValid = true;
+
+	if(!PSTemplate)
+	{
+		bValid = false;
+	}
+	else if(!MeshComp->DoesSocketExist(SocketName) && MeshComp->GetBoneIndex(SocketName) == INDEX_NONE)
+	{
+		bValid = false;
+	}
+
+	return bValid;
+}
+
+#if WITH_EDITORONLY_DATA
+void UAnimNotifyState_TimedParticleEffect::PreEditChange(UProperty* PropertyAboutToChange)
+{
+	if(PropertyAboutToChange->GetName() == GET_MEMBER_NAME_STRING_CHECKED(UAnimNotifyState_TimedParticleEffect, PSTemplate) && PSTemplate != NULL)
+	{
+		PreviousPSTemplates.Add(PSTemplate);
+	}
+
+	if(PropertyAboutToChange->GetName() == GET_MEMBER_NAME_STRING_CHECKED(UAnimNotifyState_TimedParticleEffect, SocketName) && SocketName != FName(TEXT("None")))
+	{
+		PreviousSocketNames.Add(SocketName);
+	}
+}
+#endif
