Delegates: Cleaned up abusive/legacy usages of the Delegate API

Please use AddRaw, AddSP, AddUObject, , BindRaw, BindSP, RemoveAll, Unbind etc.

diff --git a/Engine/Plugins/Editor/IntroTutorials/Source/IntroTutorials/Private/InteractiveTutorials.cpp b/Engine/Plugins/Editor/IntroTutorials/Source/IntroTutorials/Private/InteractiveTutorials.cpp
--- a/Engine/Plugins/Editor/IntroTutorials/Source/IntroTutorials/Private/InteractiveTutorials.cpp
+++ b/Engine/Plugins/Editor/IntroTutorials/Source/IntroTutorials/Private/InteractiveTutorials.cpp
@@ -30,7 +30,7 @@ void FInteractiveTutorials::SetupEditorHooks()
 	USelection::SelectionChangedEvent.AddSP(this, &FInteractiveTutorials::OnDelegateTriggered<FOnSelectionChanged, FBaseTriggerDelegate::OnSelectionChanged>);
 	USelection::SelectObjectEvent.AddSP(this, &FInteractiveTutorials::OnDelegateTriggered<FOnSelectObject, FBaseTriggerDelegate::OnSelectObject>);
 
-	FCoreDelegates::OnObjectPropertyChanged.Add( FCoreDelegates::FOnObjectPropertyChanged::FDelegate::CreateSP( this, &FInteractiveTutorials::OnDelegateTriggered<FObjectPropertyChangedTrigger, FBaseTriggerDelegate::ObjectPropertyChanged> ) );
+	FCoreDelegates::OnObjectPropertyChanged.AddSP( this, &FInteractiveTutorials::OnDelegateTriggered<FObjectPropertyChangedTrigger, FBaseTriggerDelegate::ObjectPropertyChanged> );
 
 	FAssetEditorManager::Get().OnAssetEditorRequestedOpen().AddSP(this, &FInteractiveTutorials::OnDelegateTriggered<FAssetEditorRequestedOpenTrigger, FBaseTriggerDelegate::AssetEditorRequestedOpen>);
 
diff --git a/Engine/Source/Developer/CrashTracker/Private/CrashTrackerModule.cpp b/Engine/Source/Developer/CrashTracker/Private/CrashTrackerModule.cpp
--- a/Engine/Source/Developer/CrashTracker/Private/CrashTrackerModule.cpp
+++ b/Engine/Source/Developer/CrashTracker/Private/CrashTrackerModule.cpp
@@ -133,8 +133,8 @@ void FCrashTrackerModule::StartupModule()
 	EventLogger = MakeShareable( new FCrashTrackerEventLogger(VideoCapture) );
 	FSlateApplication::Get().SetSlateUILogger(EventLogger);
 	// hook up event logger to global assert/ensure hook
-	FCoreDelegates::OnHandleSystemEnsure.Add(FCoreDelegates::FOnHandleSystemEnsure::FDelegate::CreateRaw(EventLogger.Get(), &FCrashTrackerEventLogger::OnHandleEnsure));
-	FCoreDelegates::OnHandleSystemError.Add(FCoreDelegates::FOnHandleSystemError::FDelegate::CreateRaw(EventLogger.Get(), &FCrashTrackerEventLogger::OnHandleError));
+	FCoreDelegates::OnHandleSystemEnsure.AddRaw(EventLogger.Get(), &FCrashTrackerEventLogger::OnHandleEnsure);
+	FCoreDelegates::OnHandleSystemError.AddRaw(EventLogger.Get(), &FCrashTrackerEventLogger::OnHandleError);
 	
 	if (VideoCapture.IsValid())
 	{
@@ -151,8 +151,8 @@ void FCrashTrackerModule::StartupModule()
 void FCrashTrackerModule::ShutdownModule()
 {
 	// Make sure we remove the hook to prevent calling this when it's garbage
-	FCoreDelegates::OnHandleSystemEnsure.Remove(FCoreDelegates::FOnHandleSystemEnsure::FDelegate::CreateRaw(EventLogger.Get(), &FCrashTrackerEventLogger::OnHandleEnsure));
-	FCoreDelegates::OnHandleSystemError.Remove(FCoreDelegates::FOnHandleSystemError::FDelegate::CreateRaw(EventLogger.Get(), &FCrashTrackerEventLogger::OnHandleError));
+	FCoreDelegates::OnHandleSystemEnsure.RemoveAll(EventLogger.Get());
+	FCoreDelegates::OnHandleSystemError.RemoveAll(EventLogger.Get());
 
 	VideoCapture.Reset();
 }
diff --git a/Engine/Source/Developer/InternationalizationSettings/Private/InternationalizationSettingsModelDetails.cpp b/Engine/Source/Developer/InternationalizationSettings/Private/InternationalizationSettingsModelDetails.cpp
--- a/Engine/Source/Developer/InternationalizationSettings/Private/InternationalizationSettingsModelDetails.cpp
+++ b/Engine/Source/Developer/InternationalizationSettings/Private/InternationalizationSettingsModelDetails.cpp
@@ -102,7 +102,7 @@ void FInternationalizationSettingsModelDetails::CustomizeDetails( IDetailLayoutB
 	}
 	check(Model.IsValid());
 
-	Model->OnSettingChanged().Add( UInternationalizationSettingsModel::FSettingChangedEvent::FDelegate::CreateRaw(this, &FInternationalizationSettingsModelDetails::OnSettingsChanged) );
+	Model->OnSettingChanged().AddRaw(this, &FInternationalizationSettingsModelDetails::OnSettingsChanged);
 
 	// If the saved culture is not the same as the actual current culture, a restart is needed to sync them fully and properly.
 	FString SavedCultureName = Model->GetCultureName();
diff --git a/Engine/Source/Editor/ContentBrowser/Private/AssetViewWidgets.cpp b/Engine/Source/Editor/ContentBrowser/Private/AssetViewWidgets.cpp
--- a/Engine/Source/Editor/ContentBrowser/Private/AssetViewWidgets.cpp
+++ b/Engine/Source/Editor/ContentBrowser/Private/AssetViewWidgets.cpp
@@ -1034,7 +1034,7 @@ void SAssetListItem::Construct( const FArguments& InArgs )
 	SetForceMipLevelsToBeResident(true);
 
 	// listen for asset loads so we can force mips to stream in if required
-	FCoreDelegates::OnAssetLoaded.Add(FCoreDelegates::FOnAssetLoaded::FDelegate::CreateSP(this, &SAssetViewItem::HandleAssetLoaded));
+	FCoreDelegates::OnAssetLoaded.AddSP(this, &SAssetViewItem::HandleAssetLoaded);
 }
 END_SLATE_FUNCTION_BUILD_OPTIMIZATION
 
@@ -1228,7 +1228,7 @@ void SAssetTileItem::Construct( const FArguments& InArgs )
 	SetForceMipLevelsToBeResident(true);
 
 	// listen for asset loads so we can force mips to stream in if required
-	FCoreDelegates::OnAssetLoaded.Add(FCoreDelegates::FOnAssetLoaded::FDelegate::CreateSP(this, &SAssetViewItem::HandleAssetLoaded));
+	FCoreDelegates::OnAssetLoaded.AddSP(this, &SAssetViewItem::HandleAssetLoaded);
 }
 END_SLATE_FUNCTION_BUILD_OPTIMIZATION
 
diff --git a/Engine/Source/Editor/ContentBrowser/Private/ContentBrowserSingleton.cpp b/Engine/Source/Editor/ContentBrowser/Private/ContentBrowserSingleton.cpp
--- a/Engine/Source/Editor/ContentBrowser/Private/ContentBrowserSingleton.cpp
+++ b/Engine/Source/Editor/ContentBrowser/Private/ContentBrowserSingleton.cpp
@@ -41,13 +41,13 @@ FContentBrowserSingleton::FContentBrowserSingleton()
 	FGlobalTabmanager::Get()->AddLegacyTabType( "MajorContentBrowserTab", "ContentBrowserTab2" );
 
 	// Register to be notified when properties are edited
-	FCoreDelegates::OnObjectPropertyChanged.Add( FCoreDelegates::FOnObjectPropertyChanged::FDelegate::CreateRaw(this, &FContentBrowserSingleton::OnObjectPropertyChanged) );
+	FCoreDelegates::OnObjectPropertyChanged.AddRaw(this, &FContentBrowserSingleton::OnObjectPropertyChanged);
 	FEditorDelegates::LoadSelectedAssetsIfNeeded.AddRaw(this, &FContentBrowserSingleton::OnEditorLoadSelectedAssetsIfNeeded);
 }
 
 FContentBrowserSingleton::~FContentBrowserSingleton()
 {
-	FCoreDelegates::OnObjectPropertyChanged.Remove( FCoreDelegates::FOnObjectPropertyChanged::FDelegate::CreateRaw(this, &FContentBrowserSingleton::OnObjectPropertyChanged) );
+	FCoreDelegates::OnObjectPropertyChanged.RemoveAll(this);
 	FEditorDelegates::LoadSelectedAssetsIfNeeded.RemoveAll(this);
 
 	if ( FSlateApplication::IsInitialized() )
diff --git a/Engine/Source/Editor/ContentBrowser/Private/SAssetView.cpp b/Engine/Source/Editor/ContentBrowser/Private/SAssetView.cpp
--- a/Engine/Source/Editor/ContentBrowser/Private/SAssetView.cpp
+++ b/Engine/Source/Editor/ContentBrowser/Private/SAssetView.cpp
@@ -80,11 +80,11 @@ void SAssetView::Construct( const FArguments& InArgs )
 	CollectionManagerModule.Get().OnCollectionRenamed().AddSP( this, &SAssetView::OnCollectionRenamed );
 
 	// Listen for when assets are loaded or changed to update item data
-	FCoreDelegates::OnAssetLoaded.Add(FCoreDelegates::FOnAssetLoaded::FDelegate::CreateSP(this, &SAssetView::OnAssetLoaded));
-	FCoreDelegates::OnObjectPropertyChanged.Add(FCoreDelegates::FOnObjectPropertyChanged::FDelegate::CreateSP(this, &SAssetView::OnObjectPropertyChanged));
+	FCoreDelegates::OnAssetLoaded.AddSP(this, &SAssetView::OnAssetLoaded);
+	FCoreDelegates::OnObjectPropertyChanged.AddSP(this, &SAssetView::OnObjectPropertyChanged);
 
 	// Listen for when view settings are changed
-	UContentBrowserSettings::OnSettingChanged().Add(UContentBrowserSettings::FSettingChangedEvent::FDelegate::CreateSP(this, &SAssetView::HandleSettingChanged));
+	UContentBrowserSettings::OnSettingChanged().AddSP(this, &SAssetView::HandleSettingChanged);
 
 	// Get desktop metrics
 	FDisplayMetrics DisplayMetrics;
diff --git a/Engine/Source/Editor/ContentBrowser/Private/SContentBrowser.cpp b/Engine/Source/Editor/ContentBrowser/Private/SContentBrowser.cpp
--- a/Engine/Source/Editor/ContentBrowser/Private/SContentBrowser.cpp
+++ b/Engine/Source/Editor/ContentBrowser/Private/SContentBrowser.cpp
@@ -41,7 +41,7 @@ void SContentBrowser::Construct( const FArguments& InArgs, const FName& InInstan
 	bIsLocked = InArgs._InitiallyLocked;
 
 	// Listen for when view settings are changed
-	UContentBrowserSettings::OnSettingChanged().Add(UContentBrowserSettings::FSettingChangedEvent::FDelegate::CreateSP(this, &SContentBrowser::HandleSettingChanged));
+	UContentBrowserSettings::OnSettingChanged().AddSP(this, &SContentBrowser::HandleSettingChanged);
 
 	HistoryManager.SetOnApplyHistoryData(FOnApplyHistoryData::CreateSP(this, &SContentBrowser::OnApplyHistoryData));
 	HistoryManager.SetOnUpdateHistoryData(FOnUpdateHistoryData::CreateSP(this, &SContentBrowser::OnUpdateHistoryData));
diff --git a/Engine/Source/Editor/SceneOutliner/Private/SSceneOutliner.cpp b/Engine/Source/Editor/SceneOutliner/Private/SSceneOutliner.cpp
--- a/Engine/Source/Editor/SceneOutliner/Private/SSceneOutliner.cpp
+++ b/Engine/Source/Editor/SceneOutliner/Private/SSceneOutliner.cpp
@@ -432,7 +432,7 @@ namespace SceneOutliner
 		GEditor->RegisterForUndo( this );
 
 		// Register to be notified when properties are edited
-		FCoreDelegates::OnActorLabelChanged.Add( FCoreDelegates::FOnActorLabelChanged::FDelegate::CreateRaw(this, &SSceneOutliner::OnActorLabelChanged) );
+		FCoreDelegates::OnActorLabelChanged.AddRaw(this, &SSceneOutliner::OnActorLabelChanged);
 
 		auto& Folders = FActorFolders::Get();
 		Folders.OnFolderCreate.AddSP(this, &SSceneOutliner::OnBroadcastFolderCreate);
@@ -457,7 +457,7 @@ namespace SceneOutliner
 		FWorldDelegates::LevelAddedToWorld.RemoveAll( this );
 		FWorldDelegates::LevelRemovedFromWorld.RemoveAll( this );
 
-		FCoreDelegates::OnActorLabelChanged.Remove( FCoreDelegates::FOnActorLabelChanged::FDelegate::CreateRaw(this, &SSceneOutliner::OnActorLabelChanged) );
+		FCoreDelegates::OnActorLabelChanged.RemoveAll(this);
 
 		auto& Folders = FActorFolders::Get();
 		Folders.OnFolderCreate.RemoveAll(this);
diff --git a/Engine/Source/Editor/Sequencer/Private/SequencerObjectChangeListener.cpp b/Engine/Source/Editor/Sequencer/Private/SequencerObjectChangeListener.cpp
--- a/Engine/Source/Editor/Sequencer/Private/SequencerObjectChangeListener.cpp
+++ b/Engine/Source/Editor/Sequencer/Private/SequencerObjectChangeListener.cpp
@@ -10,9 +10,9 @@ DEFINE_LOG_CATEGORY_STATIC(LogSequencerTools, Log, All);
 FSequencerObjectChangeListener::FSequencerObjectChangeListener( TSharedRef<ISequencer> InSequencer )
 	: Sequencer( InSequencer )
 {
-	FCoreDelegates::OnPreObjectPropertyChanged.Add( FCoreDelegates::FOnPreObjectPropertyChanged::FDelegate::CreateRaw( this, &FSequencerObjectChangeListener::OnObjectPreEditChange ) );
-	FCoreDelegates::OnObjectPropertyChanged.Add( FCoreDelegates::FOnObjectPropertyChanged::FDelegate::CreateRaw( this, &FSequencerObjectChangeListener::OnObjectPostEditChange ) );
-	//GEditor->OnPreActorMoved.Add( FOnObjectPropertyChanged::FDelegate::CreateRaw( this, &FSequencerObjectChangeListener::OnActorPreEditMove ) );
+	FCoreDelegates::OnPreObjectPropertyChanged.AddRaw(this, &FSequencerObjectChangeListener::OnObjectPreEditChange);
+	FCoreDelegates::OnObjectPropertyChanged.AddRaw(this, &FSequencerObjectChangeListener::OnObjectPostEditChange);
+	//GEditor->OnPreActorMoved.AddRaw(this, &FSequencerObjectChangeListener::OnActorPreEditMove);
 	GEditor->OnActorMoved().AddRaw( this, &FSequencerObjectChangeListener::OnActorPostEditMove );
 }
 
diff --git a/Engine/Source/Editor/StaticMeshEditor/Private/SStaticMeshEditorViewport.cpp b/Engine/Source/Editor/StaticMeshEditor/Private/SStaticMeshEditorViewport.cpp
--- a/Engine/Source/Editor/StaticMeshEditor/Private/SStaticMeshEditorViewport.cpp
+++ b/Engine/Source/Editor/StaticMeshEditor/Private/SStaticMeshEditorViewport.cpp
@@ -43,13 +43,13 @@ void SStaticMeshEditorViewport::Construct(const FArguments& InArgs)
 
 	SetPreviewMesh(StaticMesh);
 
-	FCoreDelegates::OnObjectPropertyChanged.Add( FCoreDelegates::FOnObjectPropertyChanged::FDelegate::CreateRaw(this, &SStaticMeshEditorViewport::OnObjectPropertyChanged) );
+	FCoreDelegates::OnObjectPropertyChanged.AddRaw(this, &SStaticMeshEditorViewport::OnObjectPropertyChanged);
 
 }
 
 SStaticMeshEditorViewport::~SStaticMeshEditorViewport()
 {
-	FCoreDelegates::OnObjectPropertyChanged.Remove( FCoreDelegates::FOnObjectPropertyChanged::FDelegate::CreateRaw(this, &SStaticMeshEditorViewport::OnObjectPropertyChanged) );
+	FCoreDelegates::OnObjectPropertyChanged.RemoveAll(this);
 	if (EditorViewportClient.IsValid())
 	{
 		EditorViewportClient->Viewport = NULL;
diff --git a/Engine/Source/Editor/TranslationEditor/Private/TranslationEditor.h b/Engine/Source/Editor/TranslationEditor/Private/TranslationEditor.h
--- a/Engine/Source/Editor/TranslationEditor/Private/TranslationEditor.h
+++ b/Engine/Source/Editor/TranslationEditor/Private/TranslationEditor.h
@@ -25,7 +25,7 @@ public:
 		for (UTranslationUnit* TranslationUnit : DataManager->GetAllTranslationsArray())
 		{
 			// Set up a property changed event to trigger a write of the translation data when TranslationUnit property changes
-			TranslationUnit->OnPropertyChanged().Add(UTranslationUnit::FTranslationUnitPropertyChangedEvent::FDelegate::CreateSP(DataManager, &FTranslationDataManager::HandlePropertyChanged));
+			TranslationUnit->OnPropertyChanged().AddSP(DataManager, &FTranslationDataManager::HandlePropertyChanged);
 		}
 
 		return TranslationEditor;
diff --git a/Engine/Source/Editor/UMGEditor/Private/SUMGEditorTree.cpp b/Engine/Source/Editor/UMGEditor/Private/SUMGEditorTree.cpp
--- a/Engine/Source/Editor/UMGEditor/Private/SUMGEditorTree.cpp
+++ b/Engine/Source/Editor/UMGEditor/Private/SUMGEditorTree.cpp
@@ -20,7 +20,7 @@ void SUMGEditorTree::Construct(const FArguments& InArgs, TSharedPtr<FBlueprintEd
 	UWidgetBlueprint* Blueprint = GetBlueprint();
 	Blueprint->OnChanged().AddSP(this, &SUMGEditorTree::OnBlueprintChanged);
 
-	FCoreDelegates::OnObjectPropertyChanged.Add( FCoreDelegates::FOnObjectPropertyChanged::FDelegate::CreateRaw(this, &SUMGEditorTree::OnObjectPropertyChanged) );
+	FCoreDelegates::OnObjectPropertyChanged.AddRaw(this, &SUMGEditorTree::OnObjectPropertyChanged);
 
 	ChildSlot
 	[
@@ -72,7 +72,7 @@ SUMGEditorTree::~SUMGEditorTree()
 		Blueprint->OnChanged().RemoveAll(this);
 	}
 
-	FCoreDelegates::OnObjectPropertyChanged.Remove( FCoreDelegates::FOnObjectPropertyChanged::FDelegate::CreateRaw(this, &SUMGEditorTree::OnObjectPropertyChanged) );
+	FCoreDelegates::OnObjectPropertyChanged.RemoveAll(this);
 }
 
 UWidgetBlueprint* SUMGEditorTree::GetBlueprint() const
diff --git a/Engine/Source/Editor/UMGEditor/Private/SUMGEditorWidgetTemplates.cpp b/Engine/Source/Editor/UMGEditor/Private/SUMGEditorWidgetTemplates.cpp
--- a/Engine/Source/Editor/UMGEditor/Private/SUMGEditorWidgetTemplates.cpp
+++ b/Engine/Source/Editor/UMGEditor/Private/SUMGEditorWidgetTemplates.cpp
@@ -68,7 +68,7 @@ void SUMGEditorWidgetTemplates::Construct(const FArguments& InArgs, TSharedPtr<F
 
 	BuildWidgetList();
 
-	//FCoreDelegates::OnObjectPropertyChanged.Add( FCoreDelegates::FOnObjectPropertyChanged::FDelegate::CreateRaw(this, &SUMGEditorWidgetTemplates::OnObjectPropertyChanged) );
+	//FCoreDelegates::OnObjectPropertyChanged.AddRaw(this, &SUMGEditorWidgetTemplates::OnObjectPropertyChanged);
 
 	ChildSlot
 	[
@@ -87,7 +87,7 @@ void SUMGEditorWidgetTemplates::Construct(const FArguments& InArgs, TSharedPtr<F
 
 SUMGEditorWidgetTemplates::~SUMGEditorWidgetTemplates()
 {
-	//FCoreDelegates::OnObjectPropertyChanged.Remove( FCoreDelegates::FOnObjectPropertyChanged::FDelegate::CreateRaw(this, &SUMGEditorWidgetTemplates::OnObjectPropertyChanged) );
+	//FCoreDelegates::OnObjectPropertyChanged.RemoveAll(this);
 }
 
 UWidgetBlueprint* SUMGEditorWidgetTemplates::GetBlueprint() const
diff --git a/Engine/Source/Editor/UnrealEd/Private/AssetThumbnail.cpp b/Engine/Source/Editor/UnrealEd/Private/AssetThumbnail.cpp
--- a/Engine/Source/Editor/UnrealEd/Private/AssetThumbnail.cpp
+++ b/Engine/Source/Editor/UnrealEd/Private/AssetThumbnail.cpp
@@ -685,14 +685,14 @@ FAssetThumbnailPool::FAssetThumbnailPool( uint32 InNumInPool, const TAttribute<b
 	, MaxFrameTimeAllowance( InMaxFrameTimeAllowance )
 	, MaxRealTimeThumbnailsPerFrame( InMaxRealTimeThumbnailsPerFrame )
 {
-	FCoreDelegates::OnObjectPropertyChanged.Add(FCoreDelegates::FOnObjectPropertyChanged::FDelegate::CreateRaw(this, &FAssetThumbnailPool::OnObjectPropertyChanged));
-	FCoreDelegates::OnAssetLoaded.Add(FCoreDelegates::FOnAssetLoaded::FDelegate::CreateRaw(this, &FAssetThumbnailPool::OnAssetLoaded));
+	FCoreDelegates::OnObjectPropertyChanged.AddRaw(this, &FAssetThumbnailPool::OnObjectPropertyChanged);
+	FCoreDelegates::OnAssetLoaded.AddRaw(this, &FAssetThumbnailPool::OnAssetLoaded);
 }
 
 FAssetThumbnailPool::~FAssetThumbnailPool()
 {
-	FCoreDelegates::OnObjectPropertyChanged.Remove(FCoreDelegates::FOnObjectPropertyChanged::FDelegate::CreateRaw(this, &FAssetThumbnailPool::OnObjectPropertyChanged));
-	FCoreDelegates::OnAssetLoaded.Remove(FCoreDelegates::FOnAssetLoaded::FDelegate::CreateRaw(this, &FAssetThumbnailPool::OnAssetLoaded));
+	FCoreDelegates::OnObjectPropertyChanged.RemoveAll(this);
+	FCoreDelegates::OnAssetLoaded.RemoveAll(this);
 
 	// Release all the texture resources
 	ReleaseResources();
diff --git a/Engine/Source/Editor/UnrealEd/Private/Commandlets/ChunkManifestGenerator.cpp b/Engine/Source/Editor/UnrealEd/Private/Commandlets/ChunkManifestGenerator.cpp
--- a/Engine/Source/Editor/UnrealEd/Private/Commandlets/ChunkManifestGenerator.cpp
+++ b/Engine/Source/Editor/UnrealEd/Private/Commandlets/ChunkManifestGenerator.cpp
@@ -184,7 +184,7 @@ void FChunkManifestGenerator::Initialize(bool InGenerateChunks)
 	// Hook up game delegate
 	if (bGenerateChunks)
 	{
-		FCoreDelegates::OnAssetLoaded.Add(FCoreDelegates::FOnAssetLoaded::FDelegate::CreateRaw(this, &FChunkManifestGenerator::OnAssetLoaded));
+		FCoreDelegates::OnAssetLoaded.AddRaw(this, &FChunkManifestGenerator::OnAssetLoaded);
 	}
 }
 
diff --git a/Engine/Source/Editor/UnrealEd/Private/UnrealEdEngine.cpp b/Engine/Source/Editor/UnrealEd/Private/UnrealEdEngine.cpp
--- a/Engine/Source/Editor/UnrealEd/Private/UnrealEdEngine.cpp
+++ b/Engine/Source/Editor/UnrealEd/Private/UnrealEdEngine.cpp
@@ -100,8 +100,8 @@ void UUnrealEdEngine::Init(IEngineLoop* InEngineLoop)
 		CookServer->Initialize( false, false, false, true, true );
 		CookServer->StartNetworkFileServer( false );
 
-		FCoreDelegates::OnObjectPropertyChanged.Add( FCoreDelegates::FOnObjectPropertyChanged::FDelegate::CreateUObject(CookServer, &UCookOnTheFlyServer::OnObjectPropertyChanged) );
-		FCoreDelegates::OnObjectModified.Add( FCoreDelegates::FOnObjectModified::FDelegate::CreateUObject(CookServer, &UCookOnTheFlyServer::OnObjectModified) );
+		FCoreDelegates::OnObjectPropertyChanged.AddUObject(CookServer, &UCookOnTheFlyServer::OnObjectPropertyChanged);
+		FCoreDelegates::OnObjectModified.AddUObject(CookServer, &UCookOnTheFlyServer::OnObjectModified);
 	}
 }
 
diff --git a/Engine/Source/Editor/UnrealEd/Private/UnrealEdMisc.cpp b/Engine/Source/Editor/UnrealEd/Private/UnrealEdMisc.cpp
--- a/Engine/Source/Editor/UnrealEd/Private/UnrealEdMisc.cpp
+++ b/Engine/Source/Editor/UnrealEd/Private/UnrealEdMisc.cpp
@@ -86,7 +86,7 @@ void FUnrealEdMisc::OnInit()
 #endif // USE_UNIT_TESTS
 
 	/** Delegate that gets called when a script exception occurs */
-	FBlueprintCoreDelegates::OnScriptException.Add(FBlueprintCoreDelegates::FOnScriptDebuggingEvent::FDelegate::CreateStatic(&FKismetDebugUtilities::OnScriptException));
+	FBlueprintCoreDelegates::OnScriptException.AddStatic(&FKismetDebugUtilities::OnScriptException);
 	
 	FEditorDelegates::ChangeEditorMode.AddRaw(this, &FUnrealEdMisc::OnEditorChangeMode);
 	FCoreDelegates::PreModal.AddRaw(this, &FUnrealEdMisc::OnEditorPreModal);
diff --git a/Engine/Source/Runtime/Engine/Private/ApplicationLifecycleComponent.cpp b/Engine/Source/Runtime/Engine/Private/ApplicationLifecycleComponent.cpp
--- a/Engine/Source/Runtime/Engine/Private/ApplicationLifecycleComponent.cpp
+++ b/Engine/Source/Runtime/Engine/Private/ApplicationLifecycleComponent.cpp
@@ -13,20 +13,20 @@ void UApplicationLifecycleComponent::OnRegister()
 {
 	Super::OnRegister();
 
-	FCoreDelegates::ApplicationWillDeactivateDelegate.Add( FCoreDelegates::FApplicationLifetimeDelegate::FDelegate::CreateUObject(this, &UApplicationLifecycleComponent::ApplicationWillDeactivateDelegate_Handler) );
-	FCoreDelegates::ApplicationHasReactivatedDelegate.Add( FCoreDelegates::FApplicationLifetimeDelegate::FDelegate::CreateUObject(this, &UApplicationLifecycleComponent::ApplicationHasReactivatedDelegate_Handler) );
-	FCoreDelegates::ApplicationWillEnterBackgroundDelegate.Add( FCoreDelegates::FApplicationLifetimeDelegate::FDelegate::CreateUObject(this, &UApplicationLifecycleComponent::ApplicationWillEnterBackgroundDelegate_Handler) );
-	FCoreDelegates::ApplicationHasEnteredForegroundDelegate.Add( FCoreDelegates::FApplicationLifetimeDelegate::FDelegate::CreateUObject(this, &UApplicationLifecycleComponent::ApplicationHasEnteredForegroundDelegate_Handler) );
-	FCoreDelegates::ApplicationWillTerminateDelegate.Add( FCoreDelegates::FApplicationLifetimeDelegate::FDelegate::CreateUObject(this, &UApplicationLifecycleComponent::ApplicationWillTerminateDelegate_Handler) );
+	FCoreDelegates::ApplicationWillDeactivateDelegate.AddUObject(this, &UApplicationLifecycleComponent::ApplicationWillDeactivateDelegate_Handler);
+	FCoreDelegates::ApplicationHasReactivatedDelegate.AddUObject(this, &UApplicationLifecycleComponent::ApplicationHasReactivatedDelegate_Handler);
+	FCoreDelegates::ApplicationWillEnterBackgroundDelegate.AddUObject(this, &UApplicationLifecycleComponent::ApplicationWillEnterBackgroundDelegate_Handler);
+	FCoreDelegates::ApplicationHasEnteredForegroundDelegate.AddUObject(this, &UApplicationLifecycleComponent::ApplicationHasEnteredForegroundDelegate_Handler);
+	FCoreDelegates::ApplicationWillTerminateDelegate.AddUObject(this, &UApplicationLifecycleComponent::ApplicationWillTerminateDelegate_Handler);
 }
 
 void UApplicationLifecycleComponent::OnUnregister()
 {
 	Super::OnUnregister();
 	
- 	FCoreDelegates::ApplicationWillDeactivateDelegate.RemoveUObject(this, &UApplicationLifecycleComponent::ApplicationWillDeactivateDelegate_Handler);
- 	FCoreDelegates::ApplicationHasReactivatedDelegate.RemoveUObject(this, &UApplicationLifecycleComponent::ApplicationHasReactivatedDelegate_Handler);
- 	FCoreDelegates::ApplicationWillEnterBackgroundDelegate.RemoveUObject(this, &UApplicationLifecycleComponent::ApplicationWillEnterBackgroundDelegate_Handler);
- 	FCoreDelegates::ApplicationHasEnteredForegroundDelegate.RemoveUObject(this, &UApplicationLifecycleComponent::ApplicationHasEnteredForegroundDelegate_Handler);
- 	FCoreDelegates::ApplicationWillTerminateDelegate.RemoveUObject(this, &UApplicationLifecycleComponent::ApplicationWillTerminateDelegate_Handler);
+ 	FCoreDelegates::ApplicationWillDeactivateDelegate.RemoveAll(this);
+ 	FCoreDelegates::ApplicationHasReactivatedDelegate.RemoveAll(this);
+ 	FCoreDelegates::ApplicationWillEnterBackgroundDelegate.RemoveAll(this);
+ 	FCoreDelegates::ApplicationHasEnteredForegroundDelegate.RemoveAll(this);
+ 	FCoreDelegates::ApplicationWillTerminateDelegate.RemoveAll(this);
 }
diff --git a/Engine/Source/Runtime/Engine/Private/LevelTick.cpp b/Engine/Source/Runtime/Engine/Private/LevelTick.cpp
--- a/Engine/Source/Runtime/Engine/Private/LevelTick.cpp
+++ b/Engine/Source/Runtime/Engine/Private/LevelTick.cpp
@@ -176,7 +176,7 @@ void FDetailedTickStats::EndObject( UObject* Object, float DeltaTime, bool bForS
 		{
 			GCCallBackRegistered = true;
 			// register callback so that we can avoid finding the wrong stats for new objects reusing memory that used to be associated with a different object
-			FCoreDelegates::PreGarbageCollect.Add(FSimpleMulticastDelegate::FDelegate::CreateRaw(this, &FDetailedTickStats::OnPreGarbageCollect));
+			FCoreDelegates::PreGarbageCollect.AddRaw(this, &FDetailedTickStats::OnPreGarbageCollect);
 		}
 
 		FTickStats NewTickStats;
diff --git a/Engine/Source/Runtime/Engine/Private/PhysicsEngine/PhysLevel.cpp b/Engine/Source/Runtime/Engine/Private/PhysicsEngine/PhysLevel.cpp
--- a/Engine/Source/Runtime/Engine/Private/PhysicsEngine/PhysLevel.cpp
+++ b/Engine/Source/Runtime/Engine/Private/PhysicsEngine/PhysLevel.cpp
@@ -254,7 +254,7 @@ void InitGamePhys()
 
 	GPhysCommandHandler = new FPhysCommandHandler();
 
-	FCoreDelegates::PreGarbageCollect.Add(FSimpleMulticastDelegate::FDelegate::CreateRaw(GPhysCommandHandler, &FPhysCommandHandler::Flush));
+	FCoreDelegates::PreGarbageCollect.AddRaw(GPhysCommandHandler, &FPhysCommandHandler::Flush);
 
 	// Init Extensions
 	PxInitExtensions(*GPhysXSDK);
diff --git a/Engine/Source/Runtime/Online/OnlineSubsystemSteam/Private/VoiceEngineSteam.cpp b/Engine/Source/Runtime/Online/OnlineSubsystemSteam/Private/VoiceEngineSteam.cpp
--- a/Engine/Source/Runtime/Online/OnlineSubsystemSteam/Private/VoiceEngineSteam.cpp
+++ b/Engine/Source/Runtime/Online/OnlineSubsystemSteam/Private/VoiceEngineSteam.cpp
@@ -318,7 +318,7 @@ uint32 FVoiceEngineSteam::SubmitRemoteVoiceData(const FUniqueNetId& RemoteTalker
 		QueuedData->AudioComponent = CreateVoiceAudioComponent(SteamUserPtr->GetVoiceOptimalSampleRate());
 		if (QueuedData->AudioComponent)
 		{
-			QueuedData->AudioComponent->OnAudioFinishedNative.Add(FOnAudioFinishedNative::FDelegate::CreateRaw(this, &FVoiceEngineSteam::OnAudioFinished));
+			QueuedData->AudioComponent->OnAudioFinishedNative.AddRaw(this, &FVoiceEngineSteam::OnAudioFinished);
 			QueuedData->AudioComponent->Play();
 		}
 	}
diff --git a/Engine/Source/Runtime/Online/OnlineSubsystemUtils/Private/VoiceEngineImpl.cpp b/Engine/Source/Runtime/Online/OnlineSubsystemUtils/Private/VoiceEngineImpl.cpp
--- a/Engine/Source/Runtime/Online/OnlineSubsystemUtils/Private/VoiceEngineImpl.cpp
+++ b/Engine/Source/Runtime/Online/OnlineSubsystemUtils/Private/VoiceEngineImpl.cpp
@@ -365,7 +365,7 @@ uint32 FVoiceEngineImpl::SubmitRemoteVoiceData(const FUniqueNetId& RemoteTalkerI
 		QueuedData->AudioComponent = CreateVoiceAudioComponent(VOICE_SAMPLE_RATE);
 		if (QueuedData->AudioComponent)
 		{
-			QueuedData->AudioComponent->OnAudioFinishedNative.Add(FOnAudioFinishedNative::FDelegate::CreateRaw(this, &FVoiceEngineImpl::OnAudioFinished));
+			QueuedData->AudioComponent->OnAudioFinishedNative.AddRaw(this, &FVoiceEngineImpl::OnAudioFinished);
 			QueuedData->AudioComponent->Play();
 		}
 	}
