modify wheel constraints of the skeletalmesh instance only, not the actual template

diff --git a/Engine/Source/Runtime/Engine/Private/Vehicles/WheeledVehicleMovementComponent.cpp b/Engine/Source/Runtime/Engine/Private/Vehicles/WheeledVehicleMovementComponent.cpp
--- a/Engine/Source/Runtime/Engine/Private/Vehicles/WheeledVehicleMovementComponent.cpp
+++ b/Engine/Source/Runtime/Engine/Private/Vehicles/WheeledVehicleMovementComponent.cpp
@@ -552,7 +552,7 @@ void UWheeledVehicleMovementComponent::CreatePhysicsState()
 
 		if ( PhysScene && PhysScene->GetVehicleManager() )
 		{
-			FixupSkeletalMesh();	//this must come first because it may trigger a rebuild of the skeletal mesh which will invalidate the vehicle setup
+			FixupSkeletalMesh();
 			CreateVehicle();
 
 			if ( PVehicle )
@@ -1209,7 +1209,7 @@ void UWheeledVehicleMovementComponent::DrawDebug( UCanvas* Canvas, float& YL, fl
 
 void UWheeledVehicleMovementComponent::FixupSkeletalMesh()
 {
-	if (USkinnedMeshComponent* Mesh = GetMesh())
+	if (USkeletalMeshComponent * Mesh = Cast<USkeletalMeshComponent>(GetMesh()))
 	{
 		if (UPhysicsAsset * PhysicsAsset = Mesh->GetPhysicsAsset())
 		{
@@ -1224,40 +1224,30 @@ void UWheeledVehicleMovementComponent::FixupSkeletalMesh()
 					{
 						if (UBodySetup * BodySetup = PhysicsAsset->BodySetup[BodySetupIdx])
 						{
-							BodySetup->PhysicsType = PhysType_Fixed; //we want wheel to be animated
-							BodySetup->CollisionReponse = EBodyCollisionResponse::BodyCollision_Disabled;	//and have no collision
+							if (BodySetup[BodySetupIdx].PhysicsType != PhysType_Fixed)
+							{
+								UE_LOG(LogPhysics, Warning, TEXT("UWheeledVehicleMovementComponent::FixupSkeletalMesh Wheel [%d] has a body that is not set to Fixed."), WheelIdx);
+								FMessageLog("PIE").Warning(FText::Format(LOCTEXT("WheelHasBodyWarning", "Wheel [{0}] is using a body which is not set to Fixed. Please change this in the physics asset."), FText::FromString(WheelSetup.BoneName.ToString())));
+							}
+
+							Mesh->Bodies[BodySetupIdx]->SetResponseToAllChannels(ECR_Ignore);	//turn off collision for wheel automatically
 
-							//we also need to make sure the wheel body has no constraints with the parent body
+							//and get rid of constraints on the wheels. TODO: right now we remove all wheel constraints, we probably only want to remove parent constraints
 							TArray<int32> WheelConstraints;
 							PhysicsAsset->BodyFindConstraints(BodySetupIdx, WheelConstraints);
 							for (int32 ConstraintIdx = 0; ConstraintIdx < WheelConstraints.Num(); ++ConstraintIdx)
 							{
-								UPhysicsConstraintTemplate * ConstraintTemplate = PhysicsAsset->ConstraintSetup[WheelConstraints[ConstraintIdx]];
-								FConstraintInstance & ConstraintInstance = ConstraintTemplate->DefaultInstance;
-								ConstraintInstance.LinearXMotion = LCM_Free;
-								ConstraintInstance.LinearYMotion = LCM_Free;
-								ConstraintInstance.LinearZMotion = LCM_Free;
-								ConstraintInstance.AngularSwing1Motion = ACM_Free;
-								ConstraintInstance.AngularSwing2Motion = ACM_Free;
-								ConstraintInstance.AngularTwistMotion = ACM_Free;
+								FConstraintInstance * ConstraintInstance = Mesh->Constraints[WheelConstraints[ConstraintIdx]];
+								ConstraintInstance->TermConstraint();
 							}
 						}
 					}
-					else
-					{
-						UE_LOG(LogPhysics, Warning, TEXT("UWheeledVehicleMovementComponent::FixupSkeletalMesh Wheel [%d] references a bone [%s] that has no body in its skeletal mesh."), WheelIdx, *WheelSetup.BoneName.ToString());
-						FMessageLog("PIE").Warning(FText::Format(LOCTEXT("WheelHasBodyWarning", "Cannot find body for wheel bone {0}."), FText::FromString(WheelSetup.BoneName.ToString())));
-					}
-
 				}
 			}
 		}
 
-		// recreate skeletal mesh using new template
-		if (Mesh->IsPhysicsStateCreated())
-		{
-			Mesh->RecreatePhysicsState();
-		}
+		Mesh->KinematicBonesUpdateType = EKinematicBonesUpdateToPhysics::SkipFixedAndSimulatingBones;	//since the wheels use physx we need to skip fixed bones from updating
+		
 	}
 
 
