Moving the blueprint socket functions out of animation and into the more general "Utilities|Transformation" category.
(TTP #327814)

diff --git a/Engine/Source/Runtime/Engine/Classes/Components/SceneComponent.h b/Engine/Source/Runtime/Engine/Classes/Components/SceneComponent.h
--- a/Engine/Source/Runtime/Engine/Classes/Components/SceneComponent.h
+++ b/Engine/Source/Runtime/Engine/Classes/Components/SceneComponent.h
@@ -338,7 +338,7 @@ public:
 	 * Gets the names of all the sockets on the component.
 	 * @return Get the names of all the sockets on the component.
 	 */
-	UFUNCTION(BlueprintCallable, Category="Animation", meta=(Keywords="Bone"))
+	UFUNCTION(BlueprintCallable, Category="Utilities|Transformation", meta=(Keywords="Bone"))
 	virtual TArray<FName> GetAllSocketNames() const;
 
 	/** 
@@ -346,7 +346,7 @@ public:
 	 * @param InSocketName Name of the socket to get the transform 
 	 * @return Socket transform in world space if socket if found. Otherwise it will return component's transform in world space.
 	 */
-	UFUNCTION(BlueprintCallable, Category="Animation", meta=(Keywords="Bone"))
+	UFUNCTION(BlueprintCallable, Category="Utilities|Transformation", meta=(Keywords="Bone"))
 	virtual FTransform GetSocketTransform(FName InSocketName, ERelativeTransformSpace TransformSpace = RTS_World) const;
 
 	/** 
@@ -354,7 +354,7 @@ public:
 	 * @param InSocketName Name of the socket to get the transform 
 	 * @return Socket transform in world space if socket if found. Otherwise it will return component's transform in world space.
 	 */
-	UFUNCTION(BlueprintCallable, Category="Animation", meta=(Keywords="Bone"))
+	UFUNCTION(BlueprintCallable, Category="Utilities|Transformation", meta=(Keywords="Bone"))
 	virtual FVector GetSocketLocation(FName InSocketName) const;
 
 	/** 
@@ -362,7 +362,7 @@ public:
 	 * @param InSocketName Name of the socket to get the transform 
 	 * @return Socket transform in world space if socket if found. Otherwise it will return component's transform in world space.
 	 */
-	UFUNCTION(BlueprintCallable, Category="Animation", meta=(Keywords="Bone"))
+	UFUNCTION(BlueprintCallable, Category="Utilities|Transformation", meta=(Keywords="Bone"))
 	virtual FRotator GetSocketRotation(FName InSocketName) const;
 
 	/** 
@@ -370,7 +370,7 @@ public:
 	 * @param InSocketName Name of the socket to get the transform 
 	 * @return Socket transform in world space if socket if found. Otherwise it will return component's transform in world space.
 	 */
-	UFUNCTION(BlueprintCallable, Category="Animation", meta=(Keywords="Bone"))
+	UFUNCTION(BlueprintCallable, Category="Utilities|Transformation", meta=(Keywords="Bone"))
 	virtual FQuat GetSocketQuaternion(FName InSocketName) const;
 
 	/** 
@@ -378,7 +378,7 @@ public:
 	 * @param InSocketName Name of the socket to get the transform 
 	 * @return true if the socket with the given name exists. Otherwise, return false
 	 */
-	UFUNCTION(BlueprintCallable, Category="Animation", meta=(Keywords="Bone"))
+	UFUNCTION(BlueprintCallable, Category="Utilities|Transformation", meta=(Keywords="Bone"))
 	virtual bool DoesSocketExist(FName InSocketName) const;
 
 	/**
