Fix incorrect capsule half length conversion from CapsuleComponent to FKSphylElem
Fix incorrect capsule half length conversion from FKSphylElem to PxCapsuleGeometry

diff --git a/Engine/Source/Runtime/Engine/Private/CapsuleComponent.cpp b/Engine/Source/Runtime/Engine/Private/CapsuleComponent.cpp
--- a/Engine/Source/Runtime/Engine/Private/CapsuleComponent.cpp
+++ b/Engine/Source/Runtime/Engine/Private/CapsuleComponent.cpp
@@ -155,7 +155,7 @@ void UCapsuleComponent::UpdateBodySetup()
 	
 	SE->SetTransform(FTransform::Identity);
 	SE->Radius = CapsuleRadius;
-	SE->Length = 2 * CapsuleHalfHeight;
+	SE->Length = 2 * FMath::Max(CapsuleHalfHeight - CapsuleRadius, 1.f);	//SphylElem uses height from center of capsule spheres, but UCapsuleComponent uses halfHeight from end of the sphere
 }
 
 bool UCapsuleComponent::IsZeroExtent() const
diff --git a/Engine/Source/Runtime/Engine/Private/PhysicsEngine/BodySetup.cpp b/Engine/Source/Runtime/Engine/Private/PhysicsEngine/BodySetup.cpp
--- a/Engine/Source/Runtime/Engine/Private/PhysicsEngine/BodySetup.cpp
+++ b/Engine/Source/Runtime/Engine/Private/PhysicsEngine/BodySetup.cpp
@@ -375,9 +375,6 @@ void UBodySetup::AddShapesToRigidActor(PxRigidActor* PDestActor, FVector& Scale3
 				float Radius = FMath::Max(SphylElem->Radius * ScaleRadius, 0.1f);
 				float HalfHeight = SphylElem->Length * ScaleLength * 0.5f;
 
-				// now find half height without the caps
-				HalfHeight = FMath::Max<float>(HalfHeight - Radius, 1.f);
-
 				PxCapsuleGeometry PCapsuleGeom;
 				PCapsuleGeom.halfHeight = HalfHeight;
 				PCapsuleGeom.radius = Radius;
