#ue4 - Optimize UPrimitiveComponent::AreAllCollideableDescendantsRelative()

- Reordered some code to better cache performance and to perform cheaper checks that fail more often earlier.
- Removed unnecessary Cast<>.

diff --git a/Engine/Source/Runtime/Engine/Private/PrimitiveComponent.cpp b/Engine/Source/Runtime/Engine/Private/PrimitiveComponent.cpp
--- a/Engine/Source/Runtime/Engine/Private/PrimitiveComponent.cpp
+++ b/Engine/Source/Runtime/Engine/Private/PrimitiveComponent.cpp
@@ -1983,25 +1983,25 @@ bool UPrimitiveComponent::AreAllCollideableDescendantsRelative(bool bAllowCached
 		TArray<USceneComponent*, TInlineAllocator<16>> ComponentStack;
 		ComponentStack.Reserve(FMath::Max(16, AttachChildren.Num()));
 
+		const ECollisionChannel CollisionChannel = GetCollisionObjectType();
 		ComponentStack.Append(AttachChildren);
 		while (ComponentStack.Num() > 0)
 		{
 			USceneComponent* const CurrentComp = ComponentStack.Pop();
 			if (CurrentComp)
 			{
-				ComponentStack.Append(CurrentComp->AttachChildren);
-
 				// Can we possibly collide with the component?
-				UPrimitiveComponent* const PrimComp = Cast<UPrimitiveComponent>(CurrentComp);
-				if (PrimComp && PrimComp->IsCollisionEnabled() && PrimComp->GetCollisionResponseToChannel(GetCollisionObjectType()) != ECR_Ignore)
+				if (CurrentComp->bAbsoluteLocation || CurrentComp->bAbsoluteRotation)
 				{
-					if (PrimComp->bAbsoluteRotation || PrimComp->bAbsoluteLocation)
+					if (CurrentComp->IsCollisionEnabled() && CurrentComp->GetCollisionResponseToChannel(CollisionChannel) != ECR_Ignore)
 					{
 						MutableThis->bCachedAllCollideableDescendantsRelative = false;
 						MutableThis->LastCheckedAllCollideableDescendantsTime = MyWorld->GetTimeSeconds();
 						return false;
 					}
 				}
+
+				ComponentStack.Append(CurrentComp->AttachChildren);
 			}
 		}
 	}
