#ue4 - Cleaned up comments within UPrimitiveComponent::AreAllCollideableDescendantsRelative()

diff --git a/Engine/Source/Runtime/Engine/Private/PrimitiveComponent.cpp b/Engine/Source/Runtime/Engine/Private/PrimitiveComponent.cpp
--- a/Engine/Source/Runtime/Engine/Private/PrimitiveComponent.cpp
+++ b/Engine/Source/Runtime/Engine/Private/PrimitiveComponent.cpp
@@ -1990,9 +1990,10 @@ bool UPrimitiveComponent::AreAllCollideableDescendantsRelative(bool bAllowCached
 			USceneComponent* const CurrentComp = ComponentStack.Pop();
 			if (CurrentComp)
 			{
-				// Can we possibly collide with the component?
+				// Is the component not using relative position?
 				if (CurrentComp->bAbsoluteLocation || CurrentComp->bAbsoluteRotation)
 				{
+					// Can we possibly collide with the component?
 					if (CurrentComp->IsCollisionEnabled() && CurrentComp->GetCollisionResponseToChannel(CollisionChannel) != ECR_Ignore)
 					{
 						MutableThis->bCachedAllCollideableDescendantsRelative = false;
