fix convex physx mesh collision drawing

diff --git a/Engine/Source/Editor/PhAT/Private/PhATRender.cpp b/Engine/Source/Editor/PhAT/Private/PhATRender.cpp
--- a/Engine/Source/Editor/PhAT/Private/PhATRender.cpp
+++ b/Engine/Source/Editor/PhAT/Private/PhATRender.cpp
@@ -179,15 +179,13 @@ void UPhATEdSkeletalMeshComponent::RenderAssetTools(const FSceneView* View, clas
 				}
 
 				FTransform ElemTM = GetPrimitiveTransform(BoneTM, i, KPT_Convex, j, 1.f);
-				FVector Scale3D = ElemTM.GetScale3D();
-				ElemTM.SetScale3D(FVector(1.0f, 1.0f, 1.0f));
 
 				//convex doesn't have solid draw so render lines if we're in hitTestAndBodyMode
 				if(!bHitTest || bHitTestAndBodyMode)
 				{
 					if (CollisionViewMode == FPhATSharedData::PRM_Solid || CollisionViewMode == FPhATSharedData::PRM_Wireframe)
 					{
-						AggGeom->ConvexElems[j].DrawElemWire(PDI, ElemTM, Scale3D, GetPrimitiveColor(i, KPT_Convex, j));
+						AggGeom->ConvexElems[j].DrawElemWire(PDI, ElemTM, GetPrimitiveColor(i, KPT_Convex, j));
 					}
 				}
 				
diff --git a/Engine/Source/Runtime/Engine/Classes/PhysicsEngine/BodySetup.h b/Engine/Source/Runtime/Engine/Classes/PhysicsEngine/BodySetup.h
--- a/Engine/Source/Runtime/Engine/Classes/PhysicsEngine/BodySetup.h
+++ b/Engine/Source/Runtime/Engine/Classes/PhysicsEngine/BodySetup.h
@@ -78,7 +78,7 @@ struct FKConvexElem
 		, ConvexMeshNegX(NULL)
 		{}
 
-		ENGINE_API void	DrawElemWire(class FPrimitiveDrawInterface* PDI, const FTransform& ElemTM, const FVector& Scale3D, const FColor Color);
+		ENGINE_API void	DrawElemWire(class FPrimitiveDrawInterface* PDI, const FTransform& ElemTM, const FColor Color);
 		void	AddCachedSolidConvexGeom(TArray<FDynamicMeshVertex>& VertexBuffer, TArray<int32>& IndexBuffer, const FColor VertexColor);
 
 		/** Reset the hull to empty all arrays */
diff --git a/Engine/Source/Runtime/Engine/Private/PhysicsEngine/PhysDrawing.cpp b/Engine/Source/Runtime/Engine/Private/PhysicsEngine/PhysDrawing.cpp
--- a/Engine/Source/Runtime/Engine/Private/PhysicsEngine/PhysDrawing.cpp
+++ b/Engine/Source/Runtime/Engine/Private/PhysicsEngine/PhysDrawing.cpp
@@ -238,14 +238,10 @@ void FKSphylElem::DrawElemSolid(FPrimitiveDrawInterface* PDI, const FTransform&
 // FKConvexElem
 /////////////////////////////////////////////////////////////////////////////////////
 
-// NB: ElemTM is assumed to have no scaling in it!
-void FKConvexElem::DrawElemWire(FPrimitiveDrawInterface* PDI, const FTransform& ElemTM, const FVector& Scale3D, const FColor Color)
+void FKConvexElem::DrawElemWire(FPrimitiveDrawInterface* PDI, const FTransform& ElemTM, const FColor Color)
 {
 #if WITH_PHYSX
 
-	FTransform LocalToWorld = ElemTM;
-	LocalToWorld.SetScale3D(Scale3D);
-
 	PxConvexMesh* Mesh = ConvexMesh;
 
 	if(Mesh)
@@ -258,7 +254,7 @@ void FKConvexElem::DrawElemWire(FPrimitiveDrawInterface* PDI, const FTransform&
 		TransformedVerts.AddUninitialized(NbVerts);
 		for(PxU32 i=0; i<NbVerts; i++)
 		{
-			TransformedVerts[i] = LocalToWorld.TransformPosition( P2UVector(Vertices[i]) );
+			TransformedVerts[i] = ElemTM.TransformPosition(P2UVector(Vertices[i]));
 		}
 						
 		const PxU8* PIndexBuffer = Mesh->getIndexBuffer();
@@ -524,7 +520,7 @@ void FKAggregateGeom::DrawAggGeom(FPrimitiveDrawInterface* PDI, const FTransform
 				FTransform ElemTM = ConvexElems[i].GetTransform();
 				ElemTM.ScaleTranslation(Scale3D);
 				ElemTM *= ParentTM;
-				ConvexElems[i].DrawElemWire(PDI, ElemTM, Scale3D, ConvexColor);
+				ConvexElems[i].DrawElemWire(PDI, ElemTM, ConvexColor);
 
 			}
 		}
