Fix divide by zero in two kismet math library functions

#UE4

diff --git a/Engine/Source/Runtime/Engine/Private/KismetMathLibrary.cpp b/Engine/Source/Runtime/Engine/Private/KismetMathLibrary.cpp
--- a/Engine/Source/Runtime/Engine/Private/KismetMathLibrary.cpp
+++ b/Engine/Source/Runtime/Engine/Private/KismetMathLibrary.cpp
@@ -462,6 +462,11 @@ float UKismetMathLibrary::SignOfFloat(float A)
 
 float UKismetMathLibrary::NormalizeToRange(float Value, float RangeMin, float RangeMax)
 {
+	if (RangeMin == RangeMax)
+	{
+		return RangeMin;
+	}
+
 	if (RangeMin > RangeMax)
 	{
 		Swap(RangeMin, RangeMax);
@@ -471,6 +476,11 @@ float UKismetMathLibrary::NormalizeToRange(float Value, float RangeMin, float Ra
 
 float UKismetMathLibrary::MapRange(float Value, float InRangeA, float InRangeB, float OutRangeA, float OutRangeB)
 {
+	if (InRangeB == InRangeA)
+	{
+		return OutRangeA;
+	}
+
 	return (Value - InRangeA) * (OutRangeB - OutRangeA) / (InRangeB - InRangeA) + OutRangeA;
 }
 
