#ue4 - Reverted one minor bit of the previous optimization to UPrimitiveComponent::AreAllCollideableDescendantsRelative()

- GetCollisionObjectType() only needed in rare cases that absolut rotation/location are being used.

diff --git a/Engine/Source/Runtime/Engine/Private/PrimitiveComponent.cpp b/Engine/Source/Runtime/Engine/Private/PrimitiveComponent.cpp
--- a/Engine/Source/Runtime/Engine/Private/PrimitiveComponent.cpp
+++ b/Engine/Source/Runtime/Engine/Private/PrimitiveComponent.cpp
@@ -1983,7 +1983,6 @@ bool UPrimitiveComponent::AreAllCollideableDescendantsRelative(bool bAllowCached
 		TArray<USceneComponent*, TInlineAllocator<16>> ComponentStack;
 		ComponentStack.Reserve(FMath::Max(16, AttachChildren.Num()));
 
-		const ECollisionChannel CollisionChannel = GetCollisionObjectType();
 		ComponentStack.Append(AttachChildren);
 		while (ComponentStack.Num() > 0)
 		{
@@ -1994,7 +1993,7 @@ bool UPrimitiveComponent::AreAllCollideableDescendantsRelative(bool bAllowCached
 				if (CurrentComp->bAbsoluteLocation || CurrentComp->bAbsoluteRotation)
 				{
 					// Can we possibly collide with the component?
-					if (CurrentComp->IsCollisionEnabled() && CurrentComp->GetCollisionResponseToChannel(CollisionChannel) != ECR_Ignore)
+					if (CurrentComp->IsCollisionEnabled() && CurrentComp->GetCollisionResponseToChannel(GetCollisionObjectType()) != ECR_Ignore)
 					{
 						MutableThis->bCachedAllCollideableDescendantsRelative = false;
 						MutableThis->LastCheckedAllCollideableDescendantsTime = MyWorld->GetTimeSeconds();
