Fixed NAN issue in BlendSpace code with Pete K

diff --git a/Engine/Source/Runtime/Engine/Private/Animation/BlendSpaceBase.cpp b/Engine/Source/Runtime/Engine/Private/Animation/BlendSpaceBase.cpp
--- a/Engine/Source/Runtime/Engine/Private/Animation/BlendSpaceBase.cpp
+++ b/Engine/Source/Runtime/Engine/Private/Animation/BlendSpaceBase.cpp
@@ -123,11 +123,17 @@ void UBlendSpaceBase::TickAssetPlayerInstance(const FAnimTickRecord& Instance, c
 				// apply speed change if you want, 
 				if (AxisToScale == BSA_X)
 				{
-					FilterMultiplier = Instance.BlendSpacePosition.X/BlendInput.X;
+					if (BlendInput.X != 0.f)
+					{
+						FilterMultiplier = Instance.BlendSpacePosition.X / BlendInput.X;
+					}
 				}
 				else if (AxisToScale == BSA_Y)
 				{
-					FilterMultiplier = Instance.BlendSpacePosition.Y/BlendInput.Y;
+					if (BlendInput.Y != 0.f)
+					{
+						FilterMultiplier = Instance.BlendSpacePosition.Y / BlendInput.Y;
+					}
 				}
 			}
 
@@ -139,14 +145,21 @@ void UBlendSpaceBase::TickAssetPlayerInstance(const FAnimTickRecord& Instance, c
 				// apply speed change if you want, 
 				if (AxisToScale == BSA_X)
 				{
-					FilterMultiplier *= BlendInput.X/ClampedInput.X;
+					if (ClampedInput.X != 0.f)
+					{
+						FilterMultiplier *= BlendInput.X / ClampedInput.X;
+					}
 				}
 				else if (AxisToScale == BSA_Y)
 				{
-					FilterMultiplier *= BlendInput.Y/ClampedInput.Y;
+					if (ClampedInput.Y != 0.f)
+					{
+						FilterMultiplier *= BlendInput.Y / ClampedInput.Y;
+					}
 				}
 			}
 
+
 			MoveDelta *= FilterMultiplier;
 			UE_LOG(LogAnimation, Log, TEXT("BlendSpace(%s) - BlendInput(%s) : FilteredBlendInput(%s), FilterMultiplier(%0.2f)"), *GetName(), *Instance.BlendSpacePosition.ToString(), *BlendInput.ToString(), FilterMultiplier );
 		}
diff --git a/Engine/Source/Runtime/Engine/Public/AnimInterpFilter.h b/Engine/Source/Runtime/Engine/Public/AnimInterpFilter.h
--- a/Engine/Source/Runtime/Engine/Public/AnimInterpFilter.h
+++ b/Engine/Source/Runtime/Engine/Public/AnimInterpFilter.h
@@ -113,6 +113,7 @@ public:
 		CurrentStackIndex = 0;
 		TimeDuration = WindowDuration;
 		CurrentTime = 0.f;
+		LastOutput = 0.f;
 	}
 
 	float GetFilteredData(float Input, float DeltaTime);
