mesh particle CPU rendering optimizations for ES2

diff --git a/Engine/Source/Runtime/Engine/Private/GPUSkinVertexFactory.cpp b/Engine/Source/Runtime/Engine/Private/GPUSkinVertexFactory.cpp
--- a/Engine/Source/Runtime/Engine/Private/GPUSkinVertexFactory.cpp
+++ b/Engine/Source/Runtime/Engine/Private/GPUSkinVertexFactory.cpp
@@ -492,7 +492,7 @@ public:
 
 		bool bIsGPUCached = false;
 
-		if (GEnableGPUSkinCache && GGPUSkinCache.SetVertexStreamFromCache(BatchElement.GPUSkinCacheKey, Shader, VertexFactory, BatchElement.MinVertexIndex, GPrevPerBoneMotionBlur.IsLocked(), GPUSkinCacheStreamFloatOffset, GPUSkinCacheStreamStride, GPUSkinCacheStreamBuffer))
+		if (GEnableGPUSkinCache && GGPUSkinCache.SetVertexStreamFromCache(BatchElement.UserIndex, Shader, VertexFactory, BatchElement.MinVertexIndex, GPrevPerBoneMotionBlur.IsLocked(), GPUSkinCacheStreamFloatOffset, GPUSkinCacheStreamStride, GPUSkinCacheStreamBuffer))
 		{
 			bIsGPUCached = true;
 		}
diff --git a/Engine/Source/Runtime/Engine/Private/MeshParticleVertexFactory.cpp b/Engine/Source/Runtime/Engine/Private/MeshParticleVertexFactory.cpp
--- a/Engine/Source/Runtime/Engine/Private/MeshParticleVertexFactory.cpp
+++ b/Engine/Source/Runtime/Engine/Private/MeshParticleVertexFactory.cpp
@@ -37,17 +37,17 @@ public:
 
 	virtual void SetMesh(FShader* Shader,const FVertexFactory* VertexFactory,const FSceneView& View,const FMeshBatchElement& BatchElement,uint32 DataFlags) const OVERRIDE
 	{
+		const bool bInstanced = GRHIFeatureLevel >= ERHIFeatureLevel::SM3;
 		FMeshParticleVertexFactory* MeshParticleVF = (FMeshParticleVertexFactory*)VertexFactory;
 		FVertexShaderRHIParamRef VertexShaderRHI = Shader->GetVertexShader();
 		SetUniformBufferParameter( VertexShaderRHI, Shader->GetUniformBufferParameter<FMeshParticleUniformParameters>(), MeshParticleVF->GetUniformBuffer() );
 
-		FMeshParticleInstanceVertex* Vertex = MeshParticleVF->GetInstanceDataCPU(); 
-		FMeshParticleInstanceVertexDynamicParameter* DynamicVertex = MeshParticleVF->GetDynamicInstanceDataCPU();
-
-		const bool bInstanced = GRHIFeatureLevel >= ERHIFeatureLevel::SM3;
-
-		if(!bInstanced)
+		if (!bInstanced)
 		{
+			FMeshParticleVertexFactory::FBatchParametersCPU* BatchParameters = (FMeshParticleVertexFactory::FBatchParametersCPU*)BatchElement.UserData;
+			const FMeshParticleInstanceVertex* Vertex = BatchParameters->InstanceBuffer + BatchElement.UserIndex;
+			const FMeshParticleInstanceVertexDynamicParameter* DynamicVertex = BatchParameters->DynamicParameterBuffer + BatchElement.UserIndex;
+
 			SetShaderValue(VertexShaderRHI, Transform1, Vertex->Transform[0]);
 			SetShaderValue(VertexShaderRHI, Transform2, Vertex->Transform[1]);
 			SetShaderValue(VertexShaderRHI, Transform3, Vertex->Transform[2]);
diff --git a/Engine/Source/Runtime/Engine/Private/Particles/ParticleSystemRender.cpp b/Engine/Source/Runtime/Engine/Private/Particles/ParticleSystemRender.cpp
--- a/Engine/Source/Runtime/Engine/Private/Particles/ParticleSystemRender.cpp
+++ b/Engine/Source/Runtime/Engine/Private/Particles/ParticleSystemRender.cpp
@@ -1712,172 +1712,45 @@ int32 FDynamicMeshEmitterData::Render(FParticleSystemSceneProxy* Proxy, FPrimiti
 	int32 NumDraws = 0;
 	if (Source.EmitterRenderMode == ERM_Normal)
 	{
-		const FStaticMeshLODResources& LODModel = StaticMesh->RenderData->LODResources[0];
-		TArray<int32> ValidSectionIndices;
-
-		bool bNoValidElements = true;
-
-		for (int32 LODIndex = 0; LODIndex < 1; LODIndex++)
-		{
-			for (int32 ElementIndex = 0; ElementIndex < LODModel.Sections.Num(); ElementIndex++)
-			{
-				FMaterialRenderProxy* MaterialProxy = MeshMaterials[ElementIndex]->GetRenderProxy(bSelected);
-				check(MaterialProxy);
-
-				// If the material is ignored by the PDI (or not there at all...), 
-				// do not add it to the list of valid elements.
-				if ((MaterialProxy && !PDI->IsMaterialIgnored(MaterialProxy, View->GetFeatureLevel())) || View->Family->EngineShowFlags.Wireframe)
-				{
-					ValidSectionIndices.Add(ElementIndex);
-					bNoValidElements = false;
-				}
-				else
-				{
-					ValidSectionIndices.Add(-1);
-				}
-			}
-		}
-
-		if (bNoValidElements == true)
-		{
-			// No valid materials... quick out
-			return 0;
-		}
-
 		const bool bIsWireframe = AllowDebugViewmodes() && View->Family->EngineShowFlags.Wireframe;
 
 		// Find the vertex allocation for this view.
-		FGlobalDynamicVertexBuffer::FAllocation* Allocation = NULL;
 		const int32 ViewIndex = View->Family->Views.Find( View );
 		
 		// Only unsorted instance data will exist for unsorted emitters and must be used for all views being rendered. 
 		const int32 InstanceBufferIndex = ((Source.SortMode == PSORTMODE_None) ? 0 : ViewIndex);
+
+		FMeshParticleVertexFactory* MeshVertexFactory = (FMeshParticleVertexFactory*)VertexFactory;
+
 		if(bInstanced)
 		{
 			check(InstanceDataAllocations.IsValidIndex( InstanceBufferIndex ) != false);
 			check(InstanceDataAllocations[InstanceBufferIndex].IsValid());
 			check(DynamicParameterDataAllocations.IsValidIndex( InstanceBufferIndex ) != false);
 			check(!bUsesDynamicParameter || DynamicParameterDataAllocations[InstanceBufferIndex].IsValid());
-		}
 
-		// Calculate the number of particles that must be drawn.
-		int32 ParticleCount = Source.ActiveParticleCount;
-		if ((Source.MaxDrawCount >= 0) && (ParticleCount > Source.MaxDrawCount))
-		{
-			ParticleCount = Source.MaxDrawCount;
+			// Set the particle instance data. This buffer is generated in PreRenderView.
+			MeshVertexFactory->SetInstanceBuffer(InstanceDataAllocations[InstanceBufferIndex].VertexBuffer, InstanceDataAllocations[InstanceBufferIndex].VertexOffset , GetDynamicVertexStride());
+			MeshVertexFactory->SetDynamicParameterBuffer(DynamicParameterDataAllocations[InstanceBufferIndex].VertexBuffer, DynamicParameterDataAllocations[InstanceBufferIndex].VertexOffset , GetDynamicParameterVertexStride());
 		}
 
-		FMeshBatch Mesh;
-		FMeshBatchElement& BatchElement = Mesh.Elements[0];
-		Mesh.VertexFactory = VertexFactory;
-		Mesh.DynamicVertexData = NULL;
-		Mesh.LCI = NULL;
-		Mesh.UseDynamicData = false;
-		Mesh.ReverseCulling = Proxy->IsLocalToWorldDeterminantNegative();
-		Mesh.CastShadow = Proxy->GetCastShadow();
-		Mesh.DepthPriorityGroup = (ESceneDepthPriorityGroup)Proxy->GetDepthPriorityGroup(View);
-
-		//@todo. Handle LODs.
-		for (int32 LODIndex = 0; LODIndex < 1; LODIndex++)
+		check(InstanceBufferIndex < (FirstBatchForView.Num() - 1));
+			
+		int32 FirstBatch = FirstBatchForView[InstanceBufferIndex];
+		int32 LastBatchPlusOne = FirstBatchForView[InstanceBufferIndex + 1];
+		for (int32 BatchIndex = FirstBatch; BatchIndex < LastBatchPlusOne; ++BatchIndex)
 		{
-			for (int32 ValidIndex = 0; ValidIndex < ValidSectionIndices.Num(); ValidIndex++)
-			{
-				int32 SectionIndex = ValidSectionIndices[ValidIndex];
-				if (SectionIndex == -1)
-				{
-					continue;
-				}
-
-				FMaterialRenderProxy* MaterialProxy = MeshMaterials[SectionIndex]->GetRenderProxy(bSelected);
-				const FStaticMeshSection& Section = LODModel.Sections[SectionIndex];
-				if ((Section.NumTriangles == 0) || (MaterialProxy == NULL))
-				{
-					//@todo. This should never occur, but it does occasionally.
-					continue;
-				}
-
-				// Draw the static mesh Sections.
-				FMeshParticleVertexFactory* MeshVertexFactory = (FMeshParticleVertexFactory*)VertexFactory;
-				MeshVertexFactory->SetUniformBuffer( UniformBuffer );
-
-				if(bInstanced)
-				{
-					// Set the particle instance data. This buffer is generated in PreRenderView.
-					MeshVertexFactory->SetInstanceBuffer(InstanceDataAllocations[InstanceBufferIndex].VertexBuffer, InstanceDataAllocations[InstanceBufferIndex].VertexOffset , GetDynamicVertexStride());
-					MeshVertexFactory->SetDynamicParameterBuffer(DynamicParameterDataAllocations[InstanceBufferIndex].VertexBuffer, DynamicParameterDataAllocations[InstanceBufferIndex].VertexOffset , GetDynamicParameterVertexStride());
-				}
-
-				BatchElement.PrimitiveUniformBufferResource = &Proxy->GetWorldSpacePrimitiveUniformBuffer();
-				BatchElement.FirstIndex = Section.FirstIndex;
-				BatchElement.MinVertexIndex = Section.MinVertexIndex;
-				BatchElement.MaxVertexIndex = Section.MaxVertexIndex;
-
-				BatchElement.NumInstances = bInstanced ? ParticleCount : 1;
-
-				if (bIsWireframe)
-				{
-					if( LODModel.WireframeIndexBuffer.IsInitialized()
-						&& !(RHISupportsTessellation(GRHIShaderPlatform) && Mesh.VertexFactory->GetType()->SupportsTessellationShaders())
-						)
-					{
-						Mesh.Type = PT_LineList;
-						Mesh.MaterialRenderProxy = Proxy->GetDeselectedWireframeMatInst();
-						BatchElement.FirstIndex = 0;
-						BatchElement.IndexBuffer = &LODModel.WireframeIndexBuffer;
-						BatchElement.NumPrimitives = LODModel.WireframeIndexBuffer.GetNumIndices() / 2;
-						
-					}
-					else
-					{
-						Mesh.Type = PT_TriangleList;
-						Mesh.MaterialRenderProxy = MeshMaterials[SectionIndex]->GetRenderProxy(bSelected);
-						Mesh.bWireframe = true;	
-						BatchElement.FirstIndex = 0;
-						BatchElement.IndexBuffer = &LODModel.IndexBuffer;
-						BatchElement.NumPrimitives = LODModel.IndexBuffer.GetNumIndices() / 3;
-					}
-				}
-				else
-				{
-					Mesh.Type = PT_TriangleList;
-					Mesh.MaterialRenderProxy = MeshMaterials[SectionIndex]->GetRenderProxy(bSelected);
-					BatchElement.IndexBuffer = &LODModel.IndexBuffer;
-					BatchElement.FirstIndex = Section.FirstIndex;
-					BatchElement.NumPrimitives = Section.NumTriangles;
-				}
-
-				if(bInstanced)
-				{
-					NumDraws += DrawRichMesh(
-						PDI, 
-						Mesh, 
-						FLinearColor(1.0f, 0.0f, 0.0f),	//WireframeColor,
-						FLinearColor(1.0f, 1.0f, 0.0f),	//LevelColor,
-						FLinearColor(1.0f, 1.0f, 1.0f),	//PropertyColor,		
-						Proxy,
-						GIsEditor && (View->Family->EngineShowFlags.Selection) ? bSelected : false,
-						bIsWireframe
-						);
-				}
-				else
-				{
-					for(int32 Particle = 0; Particle < ParticleCount; ++Particle)
-					{
-						MeshVertexFactory->SetInstanceDataCPU(InstanceDataAllocationsCPU.GetTypedData()+Particle);
-						MeshVertexFactory->SetDynamicInstanceDataCPU(DynamicParameterDataAllocationsCPU.GetTypedData()+Particle);
-						NumDraws += DrawRichMesh(
-							PDI, 
-							Mesh, 
-							FLinearColor(1.0f, 0.0f, 0.0f),	//WireframeColor,
-							FLinearColor(1.0f, 1.0f, 0.0f),	//LevelColor,
-							FLinearColor(1.0f, 1.0f, 1.0f),	//PropertyColor,		
-							Proxy,
-							GIsEditor && (View->Family->EngineShowFlags.Selection) ? bSelected : false,
-							bIsWireframe
-							);
-					}
-				}
-			}
+			FMeshBatch& Mesh = *MeshBatches[BatchIndex];
+			NumDraws += DrawRichMesh(
+				PDI, 
+				Mesh, 
+				FLinearColor(1.0f, 0.0f, 0.0f),	//WireframeColor,
+				FLinearColor(1.0f, 1.0f, 0.0f),	//LevelColor,
+				FLinearColor(1.0f, 1.0f, 1.0f),	//PropertyColor,		
+				Proxy,
+				GIsEditor && (View->Family->EngineShowFlags.Selection) ? bSelected : false,
+				bIsWireframe
+				);
 		}
 	}
 	else if (Source.EmitterRenderMode == ERM_Point)
@@ -1953,8 +1826,10 @@ void FDynamicMeshEmitterData::PreRenderView(FParticleSystemSceneProxy* Proxy, co
 		}
 		else
 		{
-			InstanceDataAllocationsCPU.Init( ParticleCount );
-			DynamicParameterDataAllocationsCPU.Init( ParticleCount );
+			InstanceDataAllocationsCPU.Reset(ParticleCount);
+			InstanceDataAllocationsCPU.AddUninitialized(ParticleCount);
+			DynamicParameterDataAllocationsCPU.Reset(ParticleCount);
+			DynamicParameterDataAllocationsCPU.AddUninitialized(ParticleCount);
 
 			// Fill instance buffer.
 			GetInstanceData((void*) InstanceDataAllocationsCPU.GetTypedData(), (void*) DynamicParameterDataAllocationsCPU.GetTypedData(), Proxy, ViewFamily->Views[ViewIndex]);
@@ -1967,12 +1842,117 @@ void FDynamicMeshEmitterData::PreRenderView(FParticleSystemSceneProxy* Proxy, co
 		}
 	}
 
+	// Setup the vertex factory.
+	FMeshParticleVertexFactory* MeshVertexFactory = (FMeshParticleVertexFactory*)VertexFactory;
+	MeshVertexFactory->SetUniformBuffer( UniformBuffer );
+
 	// Update the primitive uniform buffer
 	if (LastFramePreRendered != FrameNumber)
 	{
 		Proxy->UpdateWorldSpacePrimitiveUniformBuffer();
 		LastFramePreRendered = FrameNumber;
 	}
+
+	// Setup the mesh elements we will render with later on.
+	MeshBatches.Reset();
+	FirstBatchForView.Reset();
+	if (Source.EmitterRenderMode == ERM_Normal)
+	{
+		const FStaticMeshLODResources& LODModel = StaticMesh->RenderData->LODResources[0];
+
+		check(MeshBatches.Num() == 0);
+		FirstBatchForView.Add(0);
+		for ( int32 ViewIndex = 0; ViewIndex < ViewCount; ++ViewIndex )
+		{
+			const FSceneView* View = ViewFamily->Views[ViewIndex];
+			const bool bIsWireframe = AllowDebugViewmodes() && View->Family->EngineShowFlags.Wireframe;
+
+			//@todo. Handle LODs.
+			for (int32 LODIndex = 0; LODIndex < 1; LODIndex++)
+			{
+				for (int32 SectionIndex = 0; SectionIndex < LODModel.Sections.Num(); SectionIndex++)
+				{
+					FMeshBatch* MeshPtr = Proxy->GetPooledMeshBatch();
+					FMeshBatch& Mesh = *MeshPtr;
+					Mesh.Elements.Reset();
+					FMeshBatchElement& BatchElement = *new(Mesh.Elements) FMeshBatchElement();
+					Mesh.VertexFactory = VertexFactory;
+					Mesh.DynamicVertexData = NULL;
+					Mesh.LCI = NULL;
+					Mesh.UseDynamicData = false;
+					Mesh.ReverseCulling = Proxy->IsLocalToWorldDeterminantNegative();
+					Mesh.CastShadow = Proxy->GetCastShadow();
+					Mesh.DepthPriorityGroup = (ESceneDepthPriorityGroup)Proxy->GetDepthPriorityGroup(View);
+
+					FMaterialRenderProxy* MaterialProxy = MeshMaterials[SectionIndex]->GetRenderProxy(bSelected);
+					const FStaticMeshSection& Section = LODModel.Sections[SectionIndex];
+					if ((Section.NumTriangles == 0) || (MaterialProxy == NULL))
+					{
+						//@todo. This should never occur, but it does occasionally.
+						continue;
+					}
+
+					BatchElement.PrimitiveUniformBufferResource = &Proxy->GetWorldSpacePrimitiveUniformBuffer();
+					BatchElement.FirstIndex = Section.FirstIndex;
+					BatchElement.MinVertexIndex = Section.MinVertexIndex;
+					BatchElement.MaxVertexIndex = Section.MaxVertexIndex;
+					BatchElement.NumInstances = bInstanced ? ParticleCount : 1;
+
+					if (bIsWireframe)
+					{
+						if( LODModel.WireframeIndexBuffer.IsInitialized()
+							&& !(RHISupportsTessellation(GRHIShaderPlatform) && Mesh.VertexFactory->GetType()->SupportsTessellationShaders())
+							)
+						{
+							Mesh.Type = PT_LineList;
+							Mesh.MaterialRenderProxy = Proxy->GetDeselectedWireframeMatInst();
+							BatchElement.FirstIndex = 0;
+							BatchElement.IndexBuffer = &LODModel.WireframeIndexBuffer;
+							BatchElement.NumPrimitives = LODModel.WireframeIndexBuffer.GetNumIndices() / 2;
+						
+						}
+						else
+						{
+							Mesh.Type = PT_TriangleList;
+							Mesh.MaterialRenderProxy = MeshMaterials[SectionIndex]->GetRenderProxy(bSelected);
+							Mesh.bWireframe = true;	
+							BatchElement.FirstIndex = 0;
+							BatchElement.IndexBuffer = &LODModel.IndexBuffer;
+							BatchElement.NumPrimitives = LODModel.IndexBuffer.GetNumIndices() / 3;
+						}
+					}
+					else
+					{
+						Mesh.Type = PT_TriangleList;
+						Mesh.MaterialRenderProxy = MeshMaterials[SectionIndex]->GetRenderProxy(bSelected);
+						BatchElement.IndexBuffer = &LODModel.IndexBuffer;
+						BatchElement.FirstIndex = Section.FirstIndex;
+						BatchElement.NumPrimitives = Section.NumTriangles;
+					}
+
+					if(!bInstanced)
+					{
+						FMeshParticleVertexFactory::FBatchParametersCPU* BatchParameters = new(MeshBatchParameters) FMeshParticleVertexFactory::FBatchParametersCPU();
+						BatchParameters->InstanceBuffer = InstanceDataAllocationsCPU.GetTypedData();
+						BatchParameters->DynamicParameterBuffer = DynamicParameterDataAllocationsCPU.GetTypedData();
+						BatchElement.UserData = BatchParameters;
+						BatchElement.UserIndex = 0;
+
+						Mesh.Elements.Reserve(ParticleCount);
+						for(int32 ParticleIndex = 1; ParticleIndex < ParticleCount; ++ParticleIndex)
+						{
+							FMeshBatchElement* NextElement = new(Mesh.Elements) FMeshBatchElement();
+							*NextElement = Mesh.Elements[0];
+							NextElement->UserIndex = ParticleIndex;
+						}
+					}
+
+					MeshBatches.Add(MeshPtr);
+				}
+			}
+			FirstBatchForView.Add(MeshBatches.Num());
+		}
+	}
 }
 
 void FDynamicMeshEmitterData::GatherSimpleLights(const FParticleSystemSceneProxy* Proxy, const FSceneViewFamily& ViewFamily, FSimpleLightArray& OutParticleLights) const
@@ -6833,6 +6813,7 @@ FParticleSystemSceneProxy::FParticleSystemSceneProxy(const UParticleSystemCompon
 		)
 	, PendingLODDistance(0.0f)
 	, LastFramePreRendered(-1)
+	, FirstFreeMeshBatch(0)
 {
 #if STATS
 	LastStatCaptureTime = FApp::GetCurrentTime();
@@ -6848,6 +6829,21 @@ FParticleSystemSceneProxy::~FParticleSystemSceneProxy()
 	DynamicData = NULL;
 }
 
+FMeshBatch* FParticleSystemSceneProxy::GetPooledMeshBatch()
+{
+	FMeshBatch* Batch = NULL;
+	if (FirstFreeMeshBatch < MeshBatchPool.Num())
+	{
+		Batch = &MeshBatchPool[FirstFreeMeshBatch];
+	}
+	else
+	{
+		Batch = new(MeshBatchPool) FMeshBatch();
+	}
+	FirstFreeMeshBatch++;
+	return Batch;
+}
+
 // FPrimitiveSceneProxy interface.
 
 /** 
@@ -7246,6 +7242,7 @@ void FParticleSystemSceneProxy::ProcessPreRenderView(const FSceneView* View, int
  */
 void FParticleSystemSceneProxy::PreRenderView(const FSceneViewFamily* ViewFamily, const uint32 VisibilityMap, int32 FrameNumber)
 {
+	FirstFreeMeshBatch = 0;
 	for (int32 ViewIndex = 0; ViewIndex < ViewFamily->Views.Num(); ViewIndex++)
 	{
 		ProcessPreRenderView(ViewFamily->Views[ViewIndex], FrameNumber);
diff --git a/Engine/Source/Runtime/Engine/Private/SkeletalMesh.cpp b/Engine/Source/Runtime/Engine/Private/SkeletalMesh.cpp
--- a/Engine/Source/Runtime/Engine/Private/SkeletalMesh.cpp
+++ b/Engine/Source/Runtime/Engine/Private/SkeletalMesh.cpp
@@ -4204,7 +4204,7 @@ void FSkeletalMeshSceneProxy::DrawDynamicElementsSection(FPrimitiveDrawInterface
 	BatchElement.IndexBuffer = LODModel.MultiSizeIndexContainer.GetIndexBuffer();
 	BatchElement.MaxVertexIndex = LODModel.NumVertices - 1;
 
-	BatchElement.GPUSkinCacheKey = MeshObject->GPUSkinCacheKeys[Section.ChunkIndex];
+	BatchElement.UserIndex = MeshObject->GPUSkinCacheKeys[Section.ChunkIndex];
 
 	const bool bRequiresAdjacencyInformation = RequiresAdjacencyInformation( SectionElementInfo.Material, Mesh.VertexFactory->GetType() );
 	if ( bRequiresAdjacencyInformation )
diff --git a/Engine/Source/Runtime/Engine/Public/MeshParticleVertexFactory.h b/Engine/Source/Runtime/Engine/Public/MeshParticleVertexFactory.h
--- a/Engine/Source/Runtime/Engine/Public/MeshParticleVertexFactory.h
+++ b/Engine/Source/Runtime/Engine/Public/MeshParticleVertexFactory.h
@@ -65,6 +65,12 @@ public:
 		}
 	};
 
+	struct FBatchParametersCPU
+	{
+		const struct FMeshParticleInstanceVertex* InstanceBuffer;
+		const struct FMeshParticleInstanceVertexDynamicParameter* DynamicParameterBuffer;
+	};
+
 	/** Default constructor. */
 	FMeshParticleVertexFactory( EParticleVertexFactoryType InType = PVFT_Mesh )
 		: FParticleVertexFactoryBase(InType)
@@ -125,29 +131,6 @@ public:
 	 */
 	void SetDynamicParameterBuffer(const FVertexBuffer* InDynamicParameterBuffer, uint32 StreamOffset, uint32 Stride);
 
-
-	FORCEINLINE void SetInstanceDataCPU(struct FMeshParticleInstanceVertex* Ptr) 
-	{
-		InstanceDataCPU = Ptr;
-	}
-
-	FORCEINLINE void SetDynamicInstanceDataCPU(struct FMeshParticleInstanceVertexDynamicParameter* Ptr) 
-	{
-		DynamicInstanceDataCPU = Ptr;
-	}
-
-	FORCEINLINE struct FMeshParticleInstanceVertex* GetInstanceDataCPU()
-	{
-		return InstanceDataCPU;
-	}
-
-	FORCEINLINE struct FMeshParticleInstanceVertexDynamicParameter* GetDynamicInstanceDataCPU()
-	{
-		return DynamicInstanceDataCPU;
-	}
-
-
-
 	/**
 	* Copy the data from another vertex factory
 	* @param Other - factory to copy from
@@ -166,9 +149,4 @@ protected:
 
 	/** Uniform buffer with mesh particle paramters. */
 	FUniformBufferRHIParamRef MeshParticleUniformBuffer;
-
-	// ES2 needs to fetch instance data from ne
-	struct FMeshParticleInstanceVertex* InstanceDataCPU; 
-	struct FMeshParticleInstanceVertexDynamicParameter* DynamicInstanceDataCPU;
-
 };
diff --git a/Engine/Source/Runtime/Engine/Public/ParticleHelper.h b/Engine/Source/Runtime/Engine/Public/ParticleHelper.h
--- a/Engine/Source/Runtime/Engine/Public/ParticleHelper.h
+++ b/Engine/Source/Runtime/Engine/Public/ParticleHelper.h
@@ -1992,7 +1992,16 @@ struct FDynamicMeshEmitterData : public FDynamicSpriteEmitterDataBase
 
 	UStaticMesh*		StaticMesh;
 	TArray<UMaterialInterface*, TInlineAllocator<2> > MeshMaterials;
-		
+
+	/** Mesh batches used for rendering, built in PreRenderView. */
+	TArray<FMeshBatch*, TInlineAllocator<4> > MeshBatches;
+
+	/** Mesh batch parameters used when instancing is not allowed. */
+	TArray<FMeshParticleVertexFactory::FBatchParametersCPU, TInlineAllocator<4> > MeshBatchParameters;
+
+	/** The first mesh batches to render for a given view. */
+	TArray<int32, TInlineAllocator<4> > FirstBatchForView;
+
 	/** Particle instance data allocations (ES2). */
 	TArray<FMeshParticleInstanceVertex> InstanceDataAllocationsCPU;
 
@@ -2675,6 +2684,9 @@ public:
 	FColoredMaterialRenderProxy* GetSelectedWireframeMatInst()		{	return &SelectedWireframeMaterialInstance;		}
 	FColoredMaterialRenderProxy* GetDeselectedWireframeMatInst()	{	return &DeselectedWireframeMaterialInstance;	}
 
+	/** Gets a mesh batch from the pool. */
+	FMeshBatch* GetPooledMeshBatch();
+
 protected:
 
 	/**
@@ -2712,6 +2724,10 @@ protected:
 	/** The primitive's uniform buffer. */
 	TUniformBuffer<FPrimitiveUniformShaderParameters> WorldSpacePrimitiveUniformBuffer;
 
+	/** Pool for holding FMeshBatches to reduce allocations. */
+	TIndirectArray<FMeshBatch, TInlineAllocator<4> > MeshBatchPool;
+	int32 FirstFreeMeshBatch;
+
 	friend struct FDynamicSpriteEmitterDataBase;
 };
 
diff --git a/Engine/Source/Runtime/Engine/Public/SceneManagement.h b/Engine/Source/Runtime/Engine/Public/SceneManagement.h
--- a/Engine/Source/Runtime/Engine/Public/SceneManagement.h
+++ b/Engine/Source/Runtime/Engine/Public/SceneManagement.h
@@ -1274,7 +1274,7 @@ struct FMeshBatchElement
 	uint32 NumInstances;
 	uint32 MinVertexIndex;
 	uint32 MaxVertexIndex;
-	int32 GPUSkinCacheKey;	// -1 if not using GPU skin cache
+	int32 UserIndex;
 	void* UserData;
 
 	/** 
@@ -1288,7 +1288,7 @@ struct FMeshBatchElement
 	:	PrimitiveUniformBufferResource(NULL)
 	,	IndexBuffer(NULL)
 	,	NumInstances(1)
-	,	GPUSkinCacheKey(-1)
+	,	UserIndex(-1)
 	,	UserData(NULL)
 	,	DynamicIndexData(NULL)
 	{
