ES2 shader optimization workflow improvements

1) Show instruction counts for lightmapped shader without distance field shadows in the material editor.
2) Add r.CompileMaterialsForShaderFormat to explicitly request the engine initiate compiles for a certain shader format.

diff --git a/Engine/Source/Runtime/Engine/Private/Materials/Material.cpp b/Engine/Source/Runtime/Engine/Private/Materials/Material.cpp
--- a/Engine/Source/Runtime/Engine/Private/Materials/Material.cpp
+++ b/Engine/Source/Runtime/Engine/Private/Materials/Material.cpp
@@ -1303,6 +1303,13 @@ EBlendMode UMaterial::GetBlendModeFromString(const TCHAR* InBlendModeStr)
 	return BLEND_Opaque;
 }
 
+static FAutoConsoleVariable GCompileMaterialsForShaderFormatCVar(
+	TEXT("r.CompileMaterialsForShaderFormat"),
+	TEXT(""),
+	TEXT("When enabled, compile materials for this shader format in addition to those for the running platform.\n")
+	TEXT("Note that these shaders are compiled and immediately tossed. This is only useful when directly inspecting output via r.DebugDumpShaderInfo.")
+	);
+
 void UMaterial::CacheResourceShadersForRendering(bool bRegenerateId)
 {
 	if (bRegenerateId)
@@ -1330,6 +1337,23 @@ void UMaterial::CacheResourceShadersForRendering(bool bRegenerateId)
 			CacheShadersForResources(ShaderPlatform, ResourcesToCache, true);
 		}
 
+		FString AdditionalFormatToCache = GCompileMaterialsForShaderFormatCVar->GetString();
+		if (!AdditionalFormatToCache.IsEmpty())
+		{
+			EShaderPlatform AdditionalPlatform = ShaderFormatToLegacyShaderPlatform(FName(*AdditionalFormatToCache));
+			if (AdditionalPlatform != SP_NumPlatforms)
+			{
+				ResourcesToCache.Reset();
+				CacheResourceShadersForCooking(AdditionalPlatform,ResourcesToCache);
+				for (int32 i = 0; i < ResourcesToCache.Num(); ++i)
+				{
+					FMaterialResource* Resource = ResourcesToCache[i];
+					delete Resource;
+				}
+				ResourcesToCache.Reset();
+			}
+		}
+
 		RecacheUniformExpressions();
 	}
 }
diff --git a/Engine/Source/Runtime/Engine/Private/Materials/MaterialShared.cpp b/Engine/Source/Runtime/Engine/Private/Materials/MaterialShared.cpp
--- a/Engine/Source/Runtime/Engine/Private/Materials/MaterialShared.cpp
+++ b/Engine/Source/Runtime/Engine/Private/Materials/MaterialShared.cpp
@@ -957,8 +957,12 @@ void FMaterialResource::GetRepresentativeShaderTypesAndDescriptions(TArray<FStri
 			if (IsUsedWithStaticLighting())
 			{
 				//lit materials are usually lightmapped
-				new (ShaderTypeNames)FString(FString::Printf(TEXT("TBasePassForForwardShadingPSTDistanceFieldShadowsAndLightMapPolicyLQ%s"), ShaderSuffix));
+				new (ShaderTypeNames)FString(FString::Printf(TEXT("TBasePassForForwardShadingPSTLightMapPolicyLQ%s"), ShaderSuffix));
 				new (ShaderTypeDescriptions)FString(FString::Printf(TEXT("Mobile base pass shader with static lighting%s"), DescSuffix));
+
+				// + distance field shadows
+				new (ShaderTypeNames)FString(FString::Printf(TEXT("TBasePassForForwardShadingPSTDistanceFieldShadowsAndLightMapPolicyLQ%s"), ShaderSuffix));
+				new (ShaderTypeDescriptions)FString(FString::Printf(TEXT("Mobile base pass shader with distance field shadows%s"), DescSuffix));
 			}
 
 			//also show a dynamically lit shader
diff --git a/Engine/Source/Runtime/RHI/Private/RHI.cpp b/Engine/Source/Runtime/RHI/Private/RHI.cpp
--- a/Engine/Source/Runtime/RHI/Private/RHI.cpp
+++ b/Engine/Source/Runtime/RHI/Private/RHI.cpp
@@ -276,8 +276,7 @@ EShaderPlatform ShaderFormatToLegacyShaderPlatform(FName ShaderFormat)
 	if (ShaderFormat == NAME_OPENGL_ES2)		return SP_OPENGL_ES2;
 	if (ShaderFormat == NAME_OPENGL_ES2_WEBGL)	return SP_OPENGL_ES2_WEBGL;
 	if (ShaderFormat == NAME_OPENGL_ES2_IOS)	return SP_OPENGL_ES2_IOS;
-	check(0);
-	return SP_PCD3D_SM5;
+	return SP_NumPlatforms;
 }
 
 
