remove unneeded assert

diff --git a/Engine/Source/Runtime/Engine/Private/Particles/ParticleSystemRender.cpp b/Engine/Source/Runtime/Engine/Private/Particles/ParticleSystemRender.cpp
--- a/Engine/Source/Runtime/Engine/Private/Particles/ParticleSystemRender.cpp
+++ b/Engine/Source/Runtime/Engine/Private/Particles/ParticleSystemRender.cpp
@@ -1579,7 +1579,6 @@ void FDynamicMeshEmitterData::Init( bool bInSelected,
 	StaticMesh = InStaticMesh;
 	check(StaticMesh);
 
-	check(Source.ActiveParticleCount < 16 * 1024);	// TTP #33375
 	check(Source.ParticleStride < 2 * 1024);	// TTP #3375
 
 	InEmitterInstance->GetMeshMaterials(
