UBlueprint and GlassGeneratedBy shouldn't be used outside EDITOR

diff --git a/Engine/Source/Runtime/CoreUObject/Private/UObject/LinkerLoad.cpp b/Engine/Source/Runtime/CoreUObject/Private/UObject/LinkerLoad.cpp
--- a/Engine/Source/Runtime/CoreUObject/Private/UObject/LinkerLoad.cpp
+++ b/Engine/Source/Runtime/CoreUObject/Private/UObject/LinkerLoad.cpp
@@ -2810,7 +2810,7 @@ void ULinkerLoad::Preload( UObject* Object )
 #endif
 
 				// It's ok now to call PostLoad on blueprint CDOs
-				if (Object->HasAnyFlags(RF_ClassDefaultObject) && Object->GetClass()->ClassGeneratedBy != NULL)
+				if (Object->HasAnyFlags(RF_ClassDefaultObject) && Object->GetClass()->HasAnyClassFlags(CLASS_CompiledFromBlueprint))
 				{
 					Object->SetFlags(RF_NeedPostLoad|RF_WasLoaded);
 					GObjLoaded.Add( Object );
diff --git a/Engine/Source/Runtime/Engine/Private/BlueprintGeneratedClass.cpp b/Engine/Source/Runtime/Engine/Private/BlueprintGeneratedClass.cpp
--- a/Engine/Source/Runtime/Engine/Private/BlueprintGeneratedClass.cpp
+++ b/Engine/Source/Runtime/Engine/Private/BlueprintGeneratedClass.cpp
@@ -44,7 +44,7 @@ void UBlueprintGeneratedClass::PostLoad()
 			CurrObj->MarkPendingKill();
 		}
 	}
-
+#if WITH_EDITORONLY_DATA
 	if (GetLinkerUE4Version() < VER_UE4_CLASS_NOTPLACEABLE_ADDED)
 	{
 		// Make sure the placeable flag is correct for all blueprints
@@ -54,6 +54,7 @@ void UBlueprintGeneratedClass::PostLoad()
 			ClassFlags &= ~CLASS_NotPlaceable;
 		}
 	}
+#endif // WITH_EDITORONLY_DATA
 }
 
 #if WITH_EDITOR
diff --git a/Engine/Source/Runtime/Engine/Private/TimelineTemplate.cpp b/Engine/Source/Runtime/Engine/Private/TimelineTemplate.cpp
--- a/Engine/Source/Runtime/Engine/Private/TimelineTemplate.cpp
+++ b/Engine/Source/Runtime/Engine/Private/TimelineTemplate.cpp
@@ -216,6 +216,7 @@ void UTimelineTemplate::PostDuplicate(bool bDuplicateForPIE)
 
 	UObject* NewCurveOuter = GetOuter();
 
+#if WITH_EDITOR
 	// The outer should always be a BlueprintGeneratedClass, however the Curves get put into the Blueprint, however if they belong to the Transient package, that's where the duplicated Blueprint should be.
 	if(GetOutermost() != GetTransientPackage())
 	{
@@ -224,6 +225,7 @@ void UTimelineTemplate::PostDuplicate(bool bDuplicateForPIE)
 			NewCurveOuter = BPClass->ClassGeneratedBy;
 		}
 	}
+#endif // WITH_EDITOR
 
 	for(TArray<struct FTTFloatTrack>::TIterator It = FloatTracks.CreateIterator();It;++It)
 	{
diff --git a/Engine/Source/Runtime/Engine/Private/UnrealEngine.cpp b/Engine/Source/Runtime/Engine/Private/UnrealEngine.cpp
--- a/Engine/Source/Runtime/Engine/Private/UnrealEngine.cpp
+++ b/Engine/Source/Runtime/Engine/Private/UnrealEngine.cpp
@@ -9434,6 +9434,7 @@ void UEngine::CopyPropertiesForUnrelatedObjects(UObject* OldObject, UObject* New
 						if (CDOInst)
 						{
 							ReferenceReplacementMap.Add(CDOInst, NewInstance);
+#if WITH_EDITOR
 							if (Class->ClassGeneratedBy && Cast<UBlueprint>(Class->ClassGeneratedBy)->SkeletonGeneratedClass)
 							{
 								UObject *CDOInstS = Cast<UBlueprint>(Class->ClassGeneratedBy)->SkeletonGeneratedClass->GetDefaultSubobjectByName(NewInstance->GetFName());
@@ -9443,6 +9444,7 @@ void UEngine::CopyPropertiesForUnrelatedObjects(UObject* OldObject, UObject* New
 								}
 
 							}
+#endif // WITH_EDITOR
 						}
 						else
 						{
