Adding some code to trace down TTP 334549 (nulls in spline controlpoints array), and general cleanup of the splines code.

diff --git a/Engine/Source/Runtime/Engine/Classes/Landscape/LandscapeSplinesComponent.h b/Engine/Source/Runtime/Engine/Classes/Landscape/LandscapeSplinesComponent.h
--- a/Engine/Source/Runtime/Engine/Classes/Landscape/LandscapeSplinesComponent.h
+++ b/Engine/Source/Runtime/Engine/Classes/Landscape/LandscapeSplinesComponent.h
@@ -39,6 +39,7 @@ class ULandscapeSplinesComponent : public UPrimitiveComponent
 
 
 public:
+	void CheckSplinesValid();
 	bool ModifySplines(bool bAlwaysMarkDirty = true);
 
 #if WITH_EDITOR
diff --git a/Engine/Source/Runtime/Engine/Private/Landscape/LandscapeEdit.cpp b/Engine/Source/Runtime/Engine/Private/Landscape/LandscapeEdit.cpp
--- a/Engine/Source/Runtime/Engine/Private/Landscape/LandscapeEdit.cpp
+++ b/Engine/Source/Runtime/Engine/Private/Landscape/LandscapeEdit.cpp
@@ -507,7 +507,7 @@ void ULandscapeComponent::UpdateCollisionHeightData(const FColor* HeightmapTextu
 			{
 				DominantLayerData.AddUninitialized(FMath::Square(CollisionSizeVerts));
 
-				uint8* SrcDominantLayerData = (uint8*)CollisionComp->DominantLayerData.Lock(LOCK_READ_WRITE);
+				const uint8* SrcDominantLayerData = (uint8*)CollisionComp->DominantLayerData.Lock(LOCK_READ_ONLY);
 				FMemory::Memcpy( DominantLayerData.GetTypedData(), SrcDominantLayerData, sizeof(uint8)*FMath::Square(CollisionSizeVerts) );
 				CollisionComp->DominantLayerData.Unlock();
 			}
@@ -843,11 +843,11 @@ void ULandscapeComponent::UpdateCollisionLayerData(TArray<FColor*>& WeightmapTex
 							int32 TexY = CollisionSubsectionSizeVerts * SubsectionY + VertY;
 							int32 DataOffset = (TexX + TexY * MipSizeU) * sizeof(FColor);
 					
-							int32 DominantLayer = 255; // 255 as invalid value
+							uint8 DominantLayer = 255; // 255 as invalid value
 							int32 DominantWeight = 0;
 							for( int32 LayerIdx=0;LayerIdx<CandidateDataPtrs.Num();LayerIdx++ )
 							{
-								uint8 LayerWeight = CandidateDataPtrs[LayerIdx][DataOffset];
+								const uint8 LayerWeight = CandidateDataPtrs[LayerIdx][DataOffset];
 
 								if (LayerIdx == DataLayerIdx) // Override value for hole
 								{
@@ -3300,11 +3300,20 @@ void ALandscapeProxy::PostEditMove(bool bFinished)
 {
 	// This point is only reached when Copy and Pasted
 	Super::PostEditMove(bFinished);
+
 	if (bFinished)
 	{
 		ULandscapeInfo::RecreateLandscapeInfo(GetWorld(), true);
 		RecreateComponentsState();
 	}
+
+	if (bFinished)
+	{
+		if (SplineComponent)
+		{
+			SplineComponent->CheckSplinesValid();
+		}
+	}
 }
 
 void ALandscapeProxy::PostEditImport()
@@ -3345,32 +3354,6 @@ void ALandscape::PostEditMove(bool bFinished)
 	}
 
 	Super::PostEditMove(bFinished);
-
-	if( bFinished )
-	{
-		if (SplineComponent)
-		{
-			for (auto It = SplineComponent->ControlPoints.CreateConstIterator(); It; ++It)
-			{
-				check((*It)->GetOuterULandscapeSplinesComponent() == SplineComponent);
-
-				for (auto It2 = (*It)->ConnectedSegments.CreateConstIterator(); It2; ++It2)
-				{
-					check((*It2).Segment->GetOuterULandscapeSplinesComponent() == SplineComponent);
-				}
-			}
-
-			for (auto It = SplineComponent->Segments.CreateConstIterator(); It; ++It)
-			{
-				check((*It)->GetOuterULandscapeSplinesComponent() == SplineComponent);
-
-				for (int32 i2 = 0; i2 < 2; i2++)
-				{
-					check((*It)->Connections[i2].ControlPoint->GetOuterULandscapeSplinesComponent() == SplineComponent);
-				}
-			}
-		}
-	}
 }
 #endif	//WITH_EDITOR
 
diff --git a/Engine/Source/Runtime/Engine/Private/Landscape/LandscapeSplines.cpp b/Engine/Source/Runtime/Engine/Private/Landscape/LandscapeSplines.cpp
--- a/Engine/Source/Runtime/Engine/Private/Landscape/LandscapeSplines.cpp
+++ b/Engine/Source/Runtime/Engine/Private/Landscape/LandscapeSplines.cpp
@@ -58,50 +58,40 @@ public:
 		bDrawFalloff(Component->bShowSplineEditorMesh)
 	{
 		Segments.Reserve(Component->Segments.Num());
-		for (int32 i = 0; i < Component->Segments.Num(); i++)
-		{
-			if (ULandscapeSplineSegment* Segment = Component->Segments[i])
-			{
-
-				FSegmentProxy SegmentProxy;
-				SegmentProxy.Owner = Segment;
-				SegmentProxy.HitProxy = NULL;
-				SegmentProxy.Points = Segment->GetPoints();
-				SegmentProxy.bSelected = Segment->IsSplineSelected();
-				Segments.Add(SegmentProxy);
-			}
+		for (ULandscapeSplineSegment* Segment : Component->Segments)
+		{
+			FSegmentProxy SegmentProxy;
+			SegmentProxy.Owner = Segment;
+			SegmentProxy.HitProxy = NULL;
+			SegmentProxy.Points = Segment->GetPoints();
+			SegmentProxy.bSelected = Segment->IsSplineSelected();
+			Segments.Add(SegmentProxy);
 		}
 
 		ControlPoints.Reserve(Component->ControlPoints.Num());
-		for (int32 i = 0; i < Component->ControlPoints.Num(); i++)
-		{
-			if (ULandscapeSplineControlPoint* ControlPoint = Component->ControlPoints[i])
-			{
-
-				FControlPointProxy ControlPointProxy;
-				ControlPointProxy.Owner = ControlPoint;
-				ControlPointProxy.HitProxy = NULL;
-				ControlPointProxy.Location = ControlPoint->Location;
-				ControlPointProxy.Points = ControlPoint->GetPoints();
-				ControlPointProxy.SpriteScale = FMath::Clamp<float>(ControlPoint->Width != 0 ? ControlPoint->Width / 2 : ControlPoint->SideFalloff / 4, 10, 1000);
-				ControlPointProxy.bSelected = ControlPoint->IsSplineSelected();
-				ControlPoints.Add(ControlPointProxy);
-			}
+		for (ULandscapeSplineControlPoint* ControlPoint : Component->ControlPoints)
+		{
+			FControlPointProxy ControlPointProxy;
+			ControlPointProxy.Owner = ControlPoint;
+			ControlPointProxy.HitProxy = NULL;
+			ControlPointProxy.Location = ControlPoint->Location;
+			ControlPointProxy.Points = ControlPoint->GetPoints();
+			ControlPointProxy.SpriteScale = FMath::Clamp<float>(ControlPoint->Width != 0 ? ControlPoint->Width / 2 : ControlPoint->SideFalloff / 4, 10, 1000);
+			ControlPointProxy.bSelected = ControlPoint->IsSplineSelected();
+			ControlPoints.Add(ControlPointProxy);
 		}
 	}
 
 	virtual HHitProxy* CreateHitProxies(UPrimitiveComponent* Component, TArray<TRefCountPtr<HHitProxy> >& OutHitProxies)
 	{
 		OutHitProxies.Reserve(OutHitProxies.Num() + Segments.Num() + ControlPoints.Num());
-		for (int32 iSegment = 0; iSegment < Segments.Num(); iSegment++)
+		for (FSegmentProxy& Segment : Segments)
 		{
-			FSegmentProxy& Segment = Segments[iSegment];
 			Segment.HitProxy = new HLandscapeSplineProxy_Segment(Segment.Owner);
 			OutHitProxies.Add(Segment.HitProxy);
 		}
-		for (int32 iControlPoint = 0; iControlPoint < ControlPoints.Num(); iControlPoint++)
+		for (FControlPointProxy& ControlPoint : ControlPoints)
 		{
-			FControlPointProxy& ControlPoint = ControlPoints[iControlPoint];
 			ControlPoint.HitProxy = new HLandscapeSplineProxy_ControlPoint(ControlPoint.Owner);
 			OutHitProxies.Add(ControlPoint.HitProxy);
 		}
@@ -119,10 +109,8 @@ public:
 		const FLinearColor SelectedSplineColor = GEngine->GetSelectedMaterialColor();
 		const FLinearColor SelectedControlPointSpriteColor = FLinearColor::White + (GEngine->GetSelectedMaterialColor() * GEngine->SelectionHighlightIntensityBillboards * 10); // copied from FSpriteSceneProxy::DrawDynamicElements()
 
-		for (int32 iSegment = 0; iSegment < Segments.Num(); iSegment++)
+		for (const FSegmentProxy& Segment : Segments)
 		{
-			const FSegmentProxy& Segment = Segments[iSegment];
-
 			const FLinearColor SegmentColor = Segment.bSelected ? SelectedSplineColor : SplineColor;
 
 			FLandscapeSplineInterpPoint OldPoint = Segment.Points[0];
@@ -131,7 +119,7 @@ public:
 			OldPoint.Right        = LocalToWorld.TransformPosition(OldPoint.Right);
 			OldPoint.FalloffLeft  = LocalToWorld.TransformPosition(OldPoint.FalloffLeft);
 			OldPoint.FalloffRight = LocalToWorld.TransformPosition(OldPoint.FalloffRight);
-			for(int32 i = 1; i < Segment.Points.Num(); i++)
+			for (int32 i = 1; i < Segment.Points.Num(); i++)
 			{
 				FLandscapeSplineInterpPoint NewPoint = Segment.Points[i];
 				NewPoint.Center       = LocalToWorld.TransformPosition(NewPoint.Center);
@@ -163,10 +151,8 @@ public:
 			}
 		}
 
-		for (int32 iControlPoint = 0; iControlPoint < ControlPoints.Num(); iControlPoint++)
+		for (const FControlPointProxy& ControlPoint : ControlPoints)
 		{
-			const FControlPointProxy& ControlPoint = ControlPoints[iControlPoint];
-
 			const float ControlPointSpriteScale = LocalToWorld.GetScaleVector().X * ControlPoint.SpriteScale;
 			const FVector ControlPointLocation = LocalToWorld.TransformPosition(ControlPoint.Location) + FVector(0, 0, ControlPointSpriteScale * 0.75f);
 
@@ -218,9 +204,9 @@ public:
 				//OldPoint.FalloffLeft  = LocalToWorld.TransformPosition(OldPoint.FalloffLeft);
 				OldPoint.FalloffRight = LocalToWorld.TransformPosition(OldPoint.FalloffRight);
 
-				for(int32 i = 0; i < ControlPoint.Points.Num(); i++)
+				for (const FLandscapeSplineInterpPoint& Point : ControlPoint.Points)
 				{
-					FLandscapeSplineInterpPoint NewPoint = ControlPoint.Points[i];
+					FLandscapeSplineInterpPoint NewPoint = Point;
 					NewPoint.Center = LocalToWorld.TransformPosition(NewPoint.Center);
 					NewPoint.Left   = LocalToWorld.TransformPosition(NewPoint.Left);
 					NewPoint.Right  = LocalToWorld.TransformPosition(NewPoint.Right);
@@ -278,9 +264,8 @@ public:
 	uint32 GetAllocatedSize( void ) const
 	{
 		uint32 AllocatedSize = FPrimitiveSceneProxy::GetAllocatedSize() + Segments.GetAllocatedSize();
-		for (int32 iSegment = 0; iSegment < Segments.Num(); iSegment++)
+		for (const FSegmentProxy& Segment : Segments)
 		{
-			const FSegmentProxy& Segment = Segments[iSegment];
 			AllocatedSize += Segment.Points.GetAllocatedSize();
 		}
 		return AllocatedSize;
@@ -322,41 +307,53 @@ ULandscapeSplinesComponent::ULandscapeSplinesComponent(const class FPostConstruc
 	//RelativeScale3D = FVector(1/100.0f, 1/100.0f, 1/100.0f); // cancel out landscape scale. The scale is set up when component is created, but for a default landscape it's this
 }
 
-void ULandscapeSplinesComponent::OnRegister()
+void ULandscapeSplinesComponent::CheckSplinesValid()
 {
-	Super::OnRegister();
+#if DO_CHECK
+	// This shouldn't happen, but it has somehow (TTP #334549) so we have to fix it
+	ensure(!ControlPoints.Contains(nullptr));
+	ensure(!Segments.Contains(nullptr));
+
+	// Remove all null control points/segments
+	ControlPoints.Remove(nullptr);
+	Segments.Remove(nullptr);
 
-	for (int32 iControlPoint = 0; iControlPoint < ControlPoints.Num(); iControlPoint++)
+	// Check for cross-spline connections, as this is a potential source of nulls
+	for (ULandscapeSplineControlPoint* ControlPoint : ControlPoints)
 	{
-		if (ULandscapeSplineControlPoint* ControlPoint = ControlPoints[iControlPoint])
+		ensure(ControlPoint->GetOuterULandscapeSplinesComponent() == this);
+		for (const FLandscapeSplineConnection& Connection : ControlPoint->ConnectedSegments)
 		{
-			ControlPoint->RegisterComponents();
+			ensure(Connection.Segment->GetOuterULandscapeSplinesComponent() == this);
 		}
 	}
-	for (int32 iSegment = 0; iSegment < Segments.Num(); iSegment++)
+	for (ULandscapeSplineSegment* Segment : Segments)
 	{
-		if (ULandscapeSplineSegment* Segment = Segments[iSegment])
+		ensure(Segment->GetOuterULandscapeSplinesComponent() == this);
+		for (const FLandscapeSplineSegmentConnection& Connection : Segment->Connections)
 		{
-			Segment->RegisterComponents();
+			ensure(Connection.ControlPoint->GetOuterULandscapeSplinesComponent() == this);
 		}
 	}
+#endif
+}
+
+void ULandscapeSplinesComponent::OnRegister()
+{
+	CheckSplinesValid();
+
+	Super::OnRegister();
 }
 
 void ULandscapeSplinesComponent::OnUnregister()
 {
-	for (int32 iControlPoint = 0; iControlPoint < ControlPoints.Num(); iControlPoint++)
+	for (ULandscapeSplineControlPoint* ControlPoint : ControlPoints)
 	{
-		if (ULandscapeSplineControlPoint* ControlPoint = ControlPoints[iControlPoint])
-		{
-			ControlPoint->UnregisterComponents();
-		}
+		ControlPoint->UnregisterComponents();
 	}
-	for (int32 iSegment = 0; iSegment < Segments.Num(); iSegment++)
+	for (ULandscapeSplineSegment* Segment : Segments)
 	{
-		if (ULandscapeSplineSegment* Segment = Segments[iSegment])
-		{
-			Segment->UnregisterComponents();
-		}
+		Segment->UnregisterComponents();
 	}
 
 	Super::OnUnregister();
@@ -365,6 +362,8 @@ void ULandscapeSplinesComponent::OnUnregister()
 #if WITH_EDITOR
 FPrimitiveSceneProxy* ULandscapeSplinesComponent::CreateSceneProxy()
 {
+	CheckSplinesValid();
+
 	return new FLandscapeSplinesSceneProxy(this);
 }
 #endif
@@ -373,17 +372,18 @@ FBoxSphereBounds ULandscapeSplinesComponent::CalcBounds(const FTransform & Local
 {
 	FBox NewBoundsCalc(0);
 
-	for (int32 iControlPoint = 0; iControlPoint < ControlPoints.Num(); iControlPoint++)
+	for (ULandscapeSplineControlPoint* ControlPoint : ControlPoints)
 	{
-		if (ULandscapeSplineControlPoint* ControlPoint = ControlPoints[iControlPoint])
+		// TTP #334549: Somehow we're getting nulls in the ControlPoints array
+		if (ControlPoint)
 		{
 			NewBoundsCalc += ControlPoint->GetBounds();
 		}
 	}
 
-	for (int32 iSegment = 0; iSegment < Segments.Num(); iSegment++)
+	for (ULandscapeSplineSegment* Segment : Segments)
 	{
-		if (ULandscapeSplineSegment* Segment = Segments[iSegment])
+		if (Segment)
 		{
 			NewBoundsCalc += Segment->GetBounds();
 		}
@@ -397,21 +397,13 @@ bool ULandscapeSplinesComponent::ModifySplines(bool bAlwaysMarkDirty /*= true*/)
 {
 	bool bSavedToTransactionBuffer = Modify(bAlwaysMarkDirty);
 
-	for (int32 ControlPointIndex = 0; ControlPointIndex < ControlPoints.Num(); ControlPointIndex++)
+	for (ULandscapeSplineControlPoint* ControlPoint : ControlPoints)
 	{
-		ULandscapeSplineControlPoint* ControlPoint = ControlPoints[ControlPointIndex];
-		if (ControlPoint)
-		{
-			bSavedToTransactionBuffer = ControlPoint->Modify(bAlwaysMarkDirty) || bSavedToTransactionBuffer;
-		}
+		bSavedToTransactionBuffer = ControlPoint->Modify(bAlwaysMarkDirty) || bSavedToTransactionBuffer;
 	}
-	for (int32 SegmentIndex = 0; SegmentIndex < Segments.Num(); SegmentIndex++)
+	for (ULandscapeSplineSegment* Segment : Segments)
 	{
-		ULandscapeSplineSegment* Segment = Segments[SegmentIndex];
-		if (Segment)
-		{
-			bSavedToTransactionBuffer = Segment->Modify(bAlwaysMarkDirty) || bSavedToTransactionBuffer;
-		}
+		bSavedToTransactionBuffer = Segment->Modify(bAlwaysMarkDirty) || bSavedToTransactionBuffer;
 	}
 
 	return bSavedToTransactionBuffer;
@@ -421,8 +413,7 @@ void ULandscapeSplinesComponent::PostLoad()
 {
 	Super::PostLoad();
 
-	ControlPoints.Remove(NULL);
-	Segments.Remove(NULL);
+	CheckSplinesValid();
 }
 
 #if WITH_EDITOR
@@ -432,19 +423,13 @@ void ULandscapeSplinesComponent::PostEditChangeProperty(FPropertyChangedEvent& P
 
 	const bool bUpdateCollision = PropertyChangedEvent.ChangeType != EPropertyChangeType::Interactive;
 
-	for (int32 iControlPoint = 0; iControlPoint < ControlPoints.Num(); iControlPoint++)
+	for (ULandscapeSplineControlPoint* ControlPoint : ControlPoints)
 	{
-		if (ULandscapeSplineControlPoint* ControlPoint = ControlPoints[iControlPoint])
-		{
-			ControlPoint->UpdateSplinePoints(bUpdateCollision);
-		}
+		ControlPoint->UpdateSplinePoints(bUpdateCollision);
 	}
-	for (int32 iSegment = 0; iSegment < Segments.Num(); iSegment++)
+	for (ULandscapeSplineSegment* Segment : Segments)
 	{
-		if (ULandscapeSplineSegment* Segment = Segments[iSegment])
-		{
-			Segment->UpdateSplinePoints(bUpdateCollision);
-		}
+		Segment->UpdateSplinePoints(bUpdateCollision);
 	}
 }
 
@@ -452,13 +437,9 @@ void ULandscapeSplinesComponent::ShowSplineEditorMesh(bool bShow)
 {
 	bShowSplineEditorMesh = bShow;
 
-	for (int32 SegmentIndex = 0; SegmentIndex < Segments.Num(); SegmentIndex++)
+	for (ULandscapeSplineSegment* Segment : Segments)
 	{
-		ULandscapeSplineSegment* Segment = Segments[SegmentIndex];
-		if (Segment)
-		{
-			Segment->UpdateSplineEditorMesh();
-		}
+		Segment->UpdateSplineEditorMesh();
 	}
 
 	MarkRenderStateDirty();
@@ -541,10 +522,8 @@ FName ULandscapeSplineControlPoint::GetBestConnectionTo(FVector Destination) con
 		const UStaticMesh* Mesh = MeshComponent->StaticMesh;
 		const FVector MeshScale = MeshComponent->RelativeScale3D;
 #endif
-		for (auto It = Mesh->Sockets.CreateConstIterator(); It; ++It)
+		for (const UStaticMeshSocket* Socket : Mesh->Sockets)
 		{
-			const UStaticMeshSocket* Socket = *It;
-
 			FTransform SocketTransform = FTransform(Socket->RelativeRotation, Socket->RelativeLocation) * FTransform(Rotation, Location, MeshScale);
 			FVector SocketLocation = SocketTransform.GetTranslation();
 			FRotator SocketRotation = SocketTransform.GetRotation().Rotator();
@@ -628,20 +607,20 @@ void ULandscapeSplineControlPoint::AutoCalcRotation()
 
 	FRotator Delta = FRotator::ZeroRotator;
 
-	for (auto It = ConnectedSegments.CreateConstIterator(); It; It++)
+	for (const FLandscapeSplineConnection& Connection : ConnectedSegments)
 	{
 		// Get the start and end location/rotation of this connection
 		FVector StartLocation; FRotator StartRotation;
-		this->GetConnectionLocationAndRotation((*It).GetNearConnection().SocketName, StartLocation, StartRotation);
+		this->GetConnectionLocationAndRotation(Connection.GetNearConnection().SocketName, StartLocation, StartRotation);
 		FVector StartLocalLocation; FRotator StartLocalRotation;
-		this->GetConnectionLocalLocationAndRotation((*It).GetNearConnection().SocketName, StartLocalLocation, StartLocalRotation);
+		this->GetConnectionLocalLocationAndRotation(Connection.GetNearConnection().SocketName, StartLocalLocation, StartLocalRotation);
 		FVector EndLocation; FRotator EndRotation;
-		(*It).GetFarConnection().ControlPoint->GetConnectionLocationAndRotation((*It).GetFarConnection().SocketName, EndLocation, EndRotation);
+		Connection.GetFarConnection().ControlPoint->GetConnectionLocationAndRotation(Connection.GetFarConnection().SocketName, EndLocation, EndRotation);
 
 		// Find the delta between the direction of the tangent at the connection point and
 		// the direction to the other end's control point
 		FQuat SocketLocalRotation = StartLocalRotation.Quaternion();
-		if (FMath::Sign((*It).GetNearConnection().TangentLen) < 0)
+		if (FMath::Sign(Connection.GetNearConnection().TangentLen) < 0)
 		{
 			SocketLocalRotation = SocketLocalRotation * FRotator(0, 180, 0).Quaternion();
 		}
@@ -662,22 +641,22 @@ void ULandscapeSplineControlPoint::AutoCalcRotation()
 
 void ULandscapeSplineControlPoint::AutoFlipTangents()
 {
-	for (auto It = ConnectedSegments.CreateConstIterator(); It; It++)
+	for (const FLandscapeSplineConnection& Connection : ConnectedSegments)
 	{
-		(*It).Segment->AutoFlipTangents();
+		Connection.Segment->AutoFlipTangents();
 	}
 }
 
 void ULandscapeSplineControlPoint::AutoSetConnections(bool bIncludingValid)
 {
-	for (auto It = ConnectedSegments.CreateConstIterator(); It; It++)
+	for (const FLandscapeSplineConnection& Connection : ConnectedSegments)
 	{
-		FLandscapeSplineSegmentConnection& NearConnection = (*It).GetNearConnection();
+		FLandscapeSplineSegmentConnection& NearConnection = Connection.GetNearConnection();
 		if (bIncludingValid ||
 			(Mesh != NULL && Mesh->FindSocket(NearConnection.SocketName) == NULL) ||
 			(Mesh == NULL && NearConnection.SocketName != NAME_None))
 		{
-			FLandscapeSplineSegmentConnection& FarConnection = (*It).GetFarConnection();
+			FLandscapeSplineSegmentConnection& FarConnection = Connection.GetFarConnection();
 			FVector EndLocation; FRotator EndRotation;
 			FarConnection.ControlPoint->GetConnectionLocationAndRotation(FarConnection.SocketName, EndLocation, EndRotation);
 
@@ -822,10 +801,8 @@ void ULandscapeSplineControlPoint::UpdateSplinePoints(bool bUpdateCollision)
 	{
 		Points.Reset(ConnectedSegments.Num());
 
-		for (auto It = ConnectedSegments.CreateConstIterator(); It; It++)
+		for (const FLandscapeSplineConnection& Connection : ConnectedSegments)
 		{
-			const FLandscapeSplineConnection& Connection = (*It);
-
 			FVector StartLocation; FRotator StartRotation;
 			GetConnectionLocationAndRotation(Connection.GetNearConnection().SocketName, StartLocation, StartRotation);
 
@@ -863,10 +840,10 @@ void ULandscapeSplineControlPoint::UpdateSplinePoints(bool bUpdateCollision)
 
 	// Update Bounds
 	Bounds = FBox(0);
-	for (int32 i = 0; i < Points.Num(); i++)
+	for (const FLandscapeSplineInterpPoint& Point : Points)
 	{
-		Bounds += Points[i].FalloffLeft;
-		Bounds += Points[i].FalloffRight;
+		Bounds += Point.FalloffLeft;
+		Bounds += Point.FalloffRight;
 	}
 
 	OuterSplines->MarkRenderStateDirty();
@@ -877,9 +854,9 @@ void ULandscapeSplineControlPoint::UpdateSplinePoints(bool bUpdateCollision)
 		bNavDirty = false;
 	}
 
-	for (auto It = ConnectedSegments.CreateConstIterator(); It; It++)
+	for (const FLandscapeSplineConnection& Connection : ConnectedSegments)
 	{
-		(*It).Segment->UpdateSplinePoints(bUpdateCollision);
+		Connection.Segment->UpdateSplinePoints(bUpdateCollision);
 	}
 }
 
@@ -1009,9 +986,8 @@ void ULandscapeSplineSegment::Serialize(FArchive& Ar)
 #if WITH_EDITOR
 	if (Ar.UE4Ver() < VER_UE4_SPLINE_MESH_ORIENTATION)
 	{
-		for (auto It = SplineMeshes.CreateIterator(); It; It++)
+		for (FLandscapeSplineMeshEntry& MeshEntry : SplineMeshes)
 		{
-			FLandscapeSplineMeshEntry& MeshEntry = *It;
 			switch (MeshEntry.Orientation_DEPRECATED)
 			{
 			case LSMO_XUp:
@@ -1056,9 +1032,9 @@ void ULandscapeSplineSegment::PostLoad()
 		}
 	}
 
-	for (auto It = MeshComponents.CreateConstIterator(); It; It++)
+	for (auto* MeshComponent : MeshComponents)
 	{
-		(*It)->SetFlags(RF_TextExportTransient);
+		MeshComponent->SetFlags(RF_TextExportTransient);
 	}
 #endif
 }
@@ -1070,10 +1046,10 @@ void ULandscapeSplineSegment::SetSplineSelected(bool bInSelected)
 	bSelected = bInSelected;
 	GetOuterULandscapeSplinesComponent()->MarkRenderStateDirty();
 
-	for (auto It = MeshComponents.CreateConstIterator(); It; It++)
+	for (auto* MeshComponent : MeshComponents)
 	{
-		(*It)->bSelected = bInSelected;
-		(*It)->PushSelectionToProxy();
+		MeshComponent->bSelected = bInSelected;
+		MeshComponent->PushSelectionToProxy();
 	}
 }
 
@@ -1098,17 +1074,17 @@ void ULandscapeSplineSegment::AutoFlipTangents()
 
 void ULandscapeSplineSegment::RegisterComponents()
 {
-	for (auto It = MeshComponents.CreateConstIterator(); It; It++)
+	for (auto* MeshComponent : MeshComponents)
 	{
-		(*It)->RegisterComponent();
+		MeshComponent->RegisterComponent();
 	}
 }
 
 void ULandscapeSplineSegment::UnregisterComponents()
 {
-	for (auto It = MeshComponents.CreateConstIterator(); It; It++)
+	for (auto* MeshComponent : MeshComponents)
 	{
-		(*It)->UnregisterComponent();
+		MeshComponent->UnregisterComponent();
 	}
 }
 
@@ -1236,7 +1212,7 @@ void ULandscapeSplineSegment::UpdateSplinePoints(bool bUpdateCollision)
 	NumPoints = FMath::Clamp(NumPoints, 1, 1000);
 
 	float OldKeyTime = 0;
-	for(int32 i = 0; i < SplineInfo.Points.Num(); i++)
+	for (int32 i = 0; i < SplineInfo.Points.Num(); i++)
 	{
 		const float NewKeyTime = SplineInfo.Points[i].InVal;
 		const float NewKeyCosInterp = 0.5f - 0.5f * FMath::Cos(NewKeyTime * PI);
@@ -1253,13 +1229,13 @@ void ULandscapeSplineSegment::UpdateSplinePoints(bool bUpdateCollision)
 		const float NewKeyStartEndFalloff = FMath::Min((StartFalloffFraction > 0 ? NewKeyTime / StartFalloffFraction : 1.0f), (EndFalloffFraction > 0 ? (1 - NewKeyTime) / EndFalloffFraction : 1.0f));
 
 		// If not the first keypoint, interp from the last keypoint.
-		if(i > 0)
+		if (i > 0)
 		{
 			int32 NumSteps = FMath::CeilToInt( (NewKeyTime - OldKeyTime) * NumPoints );
 			float DrawSubstep = (NewKeyTime - OldKeyTime) / NumSteps;
 
 			// Add a point for each substep, except the ends because that's the point added outside the interp'ing.
-			for(int32 j = 1; j < NumSteps; j++)
+			for (int32 j = 1; j < NumSteps; j++)
 			{
 				const float NewTime = OldKeyTime + j*DrawSubstep;
 				const float NewCosInterp = 0.5f - 0.5f * FMath::Cos(NewTime * PI);
@@ -1292,10 +1268,10 @@ void ULandscapeSplineSegment::UpdateSplinePoints(bool bUpdateCollision)
 
 	// Update Bounds
 	Bounds = FBox(0);
-	for (int32 i = 0; i < Points.Num(); i++)
+	for (const FLandscapeSplineInterpPoint& Point : Points)
 	{
-		Bounds += Points[i].FalloffLeft;
-		Bounds += Points[i].FalloffRight;
+		Bounds += Point.FalloffLeft;
+		Bounds += Point.FalloffRight;
 	}
 
 	OuterSplines->MarkRenderStateDirty();
@@ -1303,12 +1279,11 @@ void ULandscapeSplineSegment::UpdateSplinePoints(bool bUpdateCollision)
 	// Spline mesh components
 	TArray<const FLandscapeSplineMeshEntry*> UsableMeshes;
 	UsableMeshes.Reserve(SplineMeshes.Num());
-	for (int32 i = 0; i < SplineMeshes.Num(); i++)
+	for (const FLandscapeSplineMeshEntry& MeshEntry : SplineMeshes)
 	{
-		const FLandscapeSplineMeshEntry* MeshEntry = &SplineMeshes[i];
-		if (MeshEntry->IsValid())
+		if (MeshEntry.IsValid())
 		{
-			UsableMeshes.Add(MeshEntry);
+			UsableMeshes.Add(&MeshEntry);
 		}
 	}
 
@@ -1325,9 +1300,9 @@ void ULandscapeSplineSegment::UpdateSplinePoints(bool bUpdateCollision)
 	}
 
 	OuterSplines->Modify();
-	for (auto It = MeshComponents.CreateConstIterator(); It; It++)
+	for (auto* MeshComponent : MeshComponents)
 	{
-		(*It)->Modify();
+		MeshComponent->Modify();
 	}
 	UnregisterComponents();
 
@@ -1339,9 +1314,9 @@ void ULandscapeSplineSegment::UpdateSplinePoints(bool bUpdateCollision)
 		struct FMeshSettings
 		{
 			const float T;
-			const FLandscapeSplineMeshEntry* MeshEntry;
+			const FLandscapeSplineMeshEntry* const MeshEntry;
 
-			FMeshSettings(float InT, const FLandscapeSplineMeshEntry* InMeshEntry) :
+			FMeshSettings(float InT, const FLandscapeSplineMeshEntry* const InMeshEntry) :
 				T(InT), MeshEntry(InMeshEntry) { }
 		};
 		TArray<FMeshSettings> MeshSettings;
@@ -1623,9 +1598,9 @@ void ULandscapeSplineSegment::UpdateSplinePoints(bool bUpdateCollision)
 				OuterSplines->GetWorld()->GetNavigationSystem()->UnregisterNavigationRelevantActor(OuterSplines->GetOwner());
 			}
 
-			for (auto It = MeshComponents.CreateConstIterator(); It; It++)
+			for (auto* MeshComponent : MeshComponents)
 			{
-				(*It)->DestroyComponent();
+				MeshComponent->DestroyComponent();
 			}
 			MeshComponents.Empty();
 			bNavDirty = true;
@@ -1643,9 +1618,8 @@ void ULandscapeSplineSegment::UpdateSplineEditorMesh()
 {
 	ULandscapeSplinesComponent* OuterSplines = CastChecked<ULandscapeSplinesComponent>(GetOuter());
 
-	for (auto It = MeshComponents.CreateConstIterator(); It; ++It)
+	for (auto* MeshComponent : MeshComponents)
 	{
-		USplineMeshComponent* MeshComponent = (*It);
 		if (MeshComponent->bHiddenInGame)
 		{
 			MeshComponent->SetVisibility(OuterSplines->bShowSplineEditorMesh);
@@ -1671,10 +1645,8 @@ void ULandscapeSplineSegment::DeleteSplinePoints()
 		OuterSplines->GetWorld()->GetNavigationSystem()->UnregisterNavigationRelevantActor(OuterSplines->GetOwner());
 	}
 
-	for (auto It = MeshComponents.CreateConstIterator(); It; It++)
+	for (auto* MeshComponent : MeshComponents)
 	{
-		USplineMeshComponent* MeshComponent = (*It);
-
 		MeshComponent->Modify();
 		MeshComponent->UnregisterComponent();
 		MeshComponent->DestroyComponent();
@@ -1821,11 +1793,11 @@ bool ULandscapeSplineSegment::Modify(bool bAlwaysMarkDirty /*= true*/)
 {
 	bool bSavedToTransactionBuffer = Super::Modify(bAlwaysMarkDirty);
 
-	for (auto MeshComponentIt = MeshComponents.CreateConstIterator(); MeshComponentIt; MeshComponentIt++)
+	for (auto* MeshComponent : MeshComponents)
 	{
-		if (*MeshComponentIt)
+		if (MeshComponent)
 		{
-			bSavedToTransactionBuffer = (*MeshComponentIt)->Modify(bAlwaysMarkDirty) || bSavedToTransactionBuffer;
+			bSavedToTransactionBuffer = MeshComponent->Modify(bAlwaysMarkDirty) || bSavedToTransactionBuffer;
 		}
 	}
 
@@ -1849,9 +1821,8 @@ void ULandscapeSplineSegment::PostDuplicate(bool bDuplicateForPIE)
 	{
 		ULandscapeSplinesComponent* OuterSplines = CastChecked<ULandscapeSplinesComponent>(GetOuter());
 
-		for (auto It = MeshComponents.CreateIterator(); It; ++It)
+		for (auto*& MeshComponent : MeshComponents)
 		{
-			USplineMeshComponent*& MeshComponent = (*It);
 			if (MeshComponent->GetOuter() != OuterSplines->GetOuter())
 			{
 				MeshComponent = DuplicateObject<USplineMeshComponent>(MeshComponent, OuterSplines->GetOuter(), *MeshComponent->GetName());
@@ -2031,12 +2002,10 @@ public:
 		FStaticMeshSceneProxy(InComponent)
 	{
 		// make sure all the materials are okay to be rendered as a spline mesh
-		for (int32 LODIndex = 0; LODIndex < LODs.Num(); LODIndex++)
+		for (FStaticMeshSceneProxy::FLODInfo& LODInfo : LODs)
 		{
-			FStaticMeshSceneProxy::FLODInfo& LODInfo = LODs[LODIndex];
-			for (int32 SectionIndex = 0; SectionIndex < LODInfo.Sections.Num(); SectionIndex++)
+			for (FStaticMeshSceneProxy::FLODInfo::FSectionInfo& Section : LODInfo.Sections)
 			{
-				FStaticMeshSceneProxy::FLODInfo::FSectionInfo& Section = LODInfo.Sections[SectionIndex];
 				if (!Section.Material->CheckMaterialUsage(MATUSAGE_SplineMesh))
 				{
 					Section.Material = UMaterial::GetDefaultMaterial(MD_Surface);
@@ -2071,9 +2040,9 @@ public:
 	{
 		ReleaseResources();
 
-		for (int32 LODIndex = 0; LODIndex < LODResources.Num(); LODIndex++)
+		for (FSplineMeshSceneProxy::FLODResources& LODResource : LODResources)
 		{
-			delete LODResources[LODIndex].VertexFactory;
+			delete LODResource.VertexFactory;
 		}
 		LODResources.Empty();
 	}
@@ -2315,9 +2284,9 @@ void FSplineMeshSceneProxy::InitResources(USplineMeshComponent* InComponent, int
 
 void FSplineMeshSceneProxy::ReleaseResources()
 {
-	for (int32 LODIndex = 0; LODIndex < LODResources.Num(); LODIndex++)
+	for (FSplineMeshSceneProxy::FLODResources& LODResource : LODResources)
 	{
-		LODResources[LODIndex].VertexFactory->ReleaseResource();
+		LODResource.VertexFactory->ReleaseResource();
 	}
 }
 
@@ -2525,9 +2494,9 @@ bool USplineMeshComponent::GetPhysicsTriMeshData(struct FTriMeshCollisionData* C
 		FVector Mask = FVector(1,1,1);
 		GetAxisValue(Mask, ForwardAxis) = 0;
 
-		for (int32 i = 0; i < CollisionData->Vertices.Num(); i++)
+		for (FVector& CollisionVert : CollisionData->Vertices)
 		{
-			CollisionData->Vertices[i] = CalcSliceTransform(GetAxisValue(CollisionData->Vertices[i], ForwardAxis)).TransformPosition(CollisionData->Vertices[i] * Mask);
+			CollisionVert = CalcSliceTransform(GetAxisValue(CollisionVert, ForwardAxis)).TransformPosition(CollisionVert * Mask);
 		}
 		return true;
 	}
@@ -2663,40 +2632,41 @@ void USplineMeshComponent::RecreateCollision()
 			GetAxisValue(Mask, ForwardAxis) = 0;
 
 			// distortion of a sphere can't be done nicely, so we just transform the origin and size
-			for (auto ItElement = BodySetup->AggGeom.SphereElems.CreateIterator(); ItElement; ++ItElement)
+			for (FKSphereElem& SphereElem : BodySetup->AggGeom.SphereElems)
 			{
-				const float Z = GetAxisValue(ItElement->Center, ForwardAxis);
+				const float Z = GetAxisValue(SphereElem.Center, ForwardAxis);
 				FTransform SliceTransform = CalcSliceTransform(Z);
-				ItElement->Center *= Mask;
+				SphereElem.Center *= Mask;
 
-				ItElement->Radius *= SliceTransform.GetMaximumAxisScale();
+				SphereElem.Radius *= SliceTransform.GetMaximumAxisScale();
 
 				SliceTransform.RemoveScaling();
-				ItElement->Center = SliceTransform.TransformPosition( ItElement->Center );
+				SphereElem.Center = SliceTransform.TransformPosition(SphereElem.Center);
 			}
 
 			// distortion of a sphyl can't be done nicely, so we just transform the origin and size
-			for (auto ItElement = BodySetup->AggGeom.SphylElems.CreateIterator(); ItElement; ++ItElement)
+			for (FKSphylElem& SphylElem : BodySetup->AggGeom.SphylElems)
 			{
-				const float Z = GetAxisValue(ItElement->Center, ForwardAxis);
+				const float Z = GetAxisValue(SphylElem.Center, ForwardAxis);
 				FTransform SliceTransform = CalcSliceTransform(Z);
-				ItElement->Center *= Mask;
+				SphylElem.Center *= Mask;
 
-				FTransform TM = ItElement->GetTransform();
-				ItElement->Length = (TM * SliceTransform).TransformVector(FVector(0, 0, ItElement->Length)).Size();
-				ItElement->Radius *= SliceTransform.GetMaximumAxisScale();
+				FTransform TM = SphylElem.GetTransform();
+				SphylElem.Length = (TM * SliceTransform).TransformVector(FVector(0, 0, SphylElem.Length)).Size();
+				SphylElem.Radius *= SliceTransform.GetMaximumAxisScale();
 
 				SliceTransform.RemoveScaling();
-				ItElement->SetTransform( TM * SliceTransform );
+				SphylElem.SetTransform(TM * SliceTransform);
 			}
 
 			// Convert boxes to convex hulls to better respect distortion
-			for (auto ItElement = BodySetup->AggGeom.BoxElems.CreateIterator(); ItElement; ++ItElement)
+			for (FKBoxElem& BoxElem : BodySetup->AggGeom.BoxElems)
 			{
 				FKConvexElem& ConvexElem = *new(BodySetup->AggGeom.ConvexElems) FKConvexElem();
 
-				const FVector Radii = FVector(ItElement->X / 2, ItElement->Y / 2, ItElement->Z / 2);
-				FTransform ElementTM = ItElement->GetTransform();
+				const FVector Radii = FVector(BoxElem.X / 2, BoxElem.Y / 2, BoxElem.Z / 2);
+				const FTransform ElementTM = BoxElem.GetTransform();
+				ConvexElem.VertexData.Empty(8);
 				ConvexElem.VertexData.Add(ElementTM.TransformPosition(Radii * FVector(-1,-1,-1)));
 				ConvexElem.VertexData.Add(ElementTM.TransformPosition(Radii * FVector(-1,-1, 1)));
 				ConvexElem.VertexData.Add(ElementTM.TransformPosition(Radii * FVector(-1, 1,-1)));
@@ -2711,11 +2681,11 @@ void USplineMeshComponent::RecreateCollision()
 			BodySetup->AggGeom.BoxElems.Empty();
 
 			// transform the points of the convex hulls into spline space
-			for (auto ItElement = BodySetup->AggGeom.ConvexElems.CreateIterator(); ItElement; ++ItElement)
+			for (FKConvexElem& ConvexElem : BodySetup->AggGeom.ConvexElems)
 			{
-				for (auto It = (*ItElement).VertexData.CreateIterator(); It; ++It)
+				for (FVector& Point : ConvexElem.VertexData)
 				{
-					(*It) = CalcSliceTransform(GetAxisValue(*It, ForwardAxis)).TransformPosition((*It) * Mask);
+					Point = CalcSliceTransform(GetAxisValue(Point, ForwardAxis)).TransformPosition(Point * Mask);
 				}
 			}
 
