Added a new clothing option to SkeletalMeshComponent for local space cloth simulation

diff --git a/Engine/Source/Runtime/Engine/Classes/Components/SkeletalMeshComponent.h b/Engine/Source/Runtime/Engine/Classes/Components/SkeletalMeshComponent.h
--- a/Engine/Source/Runtime/Engine/Classes/Components/SkeletalMeshComponent.h
+++ b/Engine/Source/Runtime/Engine/Classes/Components/SkeletalMeshComponent.h
@@ -303,6 +303,15 @@ public:
 	/** can't collide with part of attached children if total collision volumes exceed 16 capsules or 32 planes per convex */
 	UPROPERTY(EditAnywhere, BlueprintReadWrite, Category=Clothing)
 	uint32 bCollideWithAttachedChildren:1;
+	/**
+	 * It's worth trying this option when you feel that the current cloth simulation is unstable.
+	 * The scale of the actor is maintained during the simulation. 
+	 * It is possible to add the inertia effects to the simulation, through the inertiaScale parameter of the clothing material. 
+	 * So with an inertiaScale of 1.0 there should be no visible difference between local space and global space simulation. 
+	 * Known issues: - Currently there's simulation issues when this feature is used in 3.x (DE4076) So if localSpaceSim is enabled there's no inertia effect when the global pose of the clothing actor changes.
+	 */
+	UPROPERTY(EditAnywhere, BlueprintReadWrite, Category = Clothing)
+	uint32 bLocalSpaceSimulation : 1;
 	/** reset the clothing after moving the clothing position (called teleport) */
 	UPROPERTY(EditAnywhere, BlueprintReadWrite, Category=Clothing)
 	uint32 bResetAfterTeleport:1;
diff --git a/Engine/Source/Runtime/Engine/Private/SkeletalMeshComponentPhysics.cpp b/Engine/Source/Runtime/Engine/Private/SkeletalMeshComponentPhysics.cpp
--- a/Engine/Source/Runtime/Engine/Private/SkeletalMeshComponentPhysics.cpp
+++ b/Engine/Source/Runtime/Engine/Private/SkeletalMeshComponentPhysics.cpp
@@ -1817,6 +1817,8 @@ bool USkeletalMeshComponent::CreateClothingActor(int32 AssetIndex, TSharedPtr<FC
 	verify(NxParameterized::setParamF32(*ActorDesc,"lodWeights.bias",0));
 	verify(NxParameterized::setParamF32(*ActorDesc,"lodWeights.benefitsBias",0));
 
+	verify(NxParameterized::setParamBool(*ActorDesc, "localSpaceSim", bLocalSpaceSimulation));
+
 	// Initialize the global pose
 
 	FMatrix UGlobalPose = ComponentToWorld.ToMatrixWithScale();
@@ -3000,7 +3002,7 @@ void USkeletalMeshComponent::TickClothing(float DeltaTime)
 	}
 	else
 	{
-	ValidateClothingActors();
+		ValidateClothingActors();
 		UpdateClothState(DeltaTime);
 	}
 #if 0 //if turn on this flag, you can check which primitive objects are activated for collision detection
@@ -3062,10 +3064,33 @@ void USkeletalMeshComponent::GetUpdateClothSimulationData(TArray<FClothSimulData
 				const physx::PxVec3* Vertices = ClothingActor->getSimulationPositions();
 				const physx::PxVec3* Normals = ClothingActor->getSimulationNormals();
 
-				for(uint32 VertexIndex=0; VertexIndex<NumSimulVertices; VertexIndex++)
+				// In case of Local space simulation, need to transform simulated positions with the internal bone matrix
+				if (bLocalSpaceSimulation)
 				{
-					ClothData.ClothSimulPositions[VertexIndex] = P2UVector(Vertices[VertexIndex]);
-					ClothData.ClothSimulNormals[VertexIndex] = P2UVector(Normals[VertexIndex]);
+					check(ClothingActors[ActorIndex].ParentClothingAsset.IsValid());
+
+					const NxParameterized::Interface* AssetParams = ClothingActors[ActorIndex].ParentClothingAsset->GetAsset()->getAssetNxParameterized();
+					uint32 InternalRootBoneIndex;
+					NxParameterized::getParamU32(*AssetParams, "rootBoneIndex", InternalRootBoneIndex);
+					check(InternalRootBoneIndex >= 0);
+					FName BoneName = ClothingActors[ActorIndex].ParentClothingAsset->GetConvertedBoneName(InternalRootBoneIndex);					
+					int32 BoneIndex = GetBoneIndex(BoneName);
+					check(BoneIndex >= 0);
+					FMatrix ClothRootBoneMatrix = GetBoneMatrix(BoneIndex);
+
+					for (uint32 VertexIndex = 0; VertexIndex < NumSimulVertices; VertexIndex++)
+					{						
+						ClothData.ClothSimulPositions[VertexIndex] = ClothRootBoneMatrix.TransformPosition(P2UVector(Vertices[VertexIndex]));
+						ClothData.ClothSimulNormals[VertexIndex] = ClothRootBoneMatrix.TransformVector(P2UVector(Normals[VertexIndex]));
+					}
+				}
+				else
+				{
+					for (uint32 VertexIndex = 0; VertexIndex < NumSimulVertices; VertexIndex++)
+					{
+						ClothData.ClothSimulPositions[VertexIndex] = P2UVector(Vertices[VertexIndex]);
+						ClothData.ClothSimulNormals[VertexIndex] = P2UVector(Normals[VertexIndex]);
+					}
 				}
 			}
 		}
