Engine - Moved implementation  to the source file due to linker issues

diff --git a/Engine/Source/Runtime/Engine/Private/Animation/AnimEncoding_PerTrackCompression.cpp b/Engine/Source/Runtime/Engine/Private/Animation/AnimEncoding_PerTrackCompression.cpp
--- a/Engine/Source/Runtime/Engine/Private/Animation/AnimEncoding_PerTrackCompression.cpp
+++ b/Engine/Source/Runtime/Engine/Private/Animation/AnimEncoding_PerTrackCompression.cpp
@@ -747,3 +747,109 @@ void AEFPerTrackCompressionCodec::GetBoneAtomScale(
 	}
 }
 
+#if USE_ANIMATION_CODEC_BATCH_SOLVER
+
+/**
+* Decompress all requested rotation components from an Animation Sequence
+*
+* @param	Atoms			The FTransform array to fill in.
+* @param	DesiredPairs	Array of requested bone information
+* @param	Seq				The animation sequence to use.
+* @param	Time			Current time to solve for.
+* @param	bLooping		True when looping the stream in intended.
+* @return					None.
+*/
+void AEFPerTrackCompressionCodec::GetPoseRotations(
+	FTransformArray& Atoms,
+	const BoneTrackArray& DesiredPairs,
+	const UAnimSequence& Seq,
+	float Time,
+	bool bLooping )
+{
+	const int32 PairCount = DesiredPairs.Num();
+	const float RelativePos = Time / Seq.SequenceLength;
+
+	for( int32 PairIndex = 0; PairIndex < PairCount; ++PairIndex )
+	{
+		const BoneTrackPair& Pair = DesiredPairs[PairIndex];
+		const int32 TrackIndex = Pair.TrackIndex;
+		const int32 AtomIndex = Pair.AtomIndex;
+		FTransform& BoneAtom = Atoms[AtomIndex];
+
+		const int32* RESTRICT TrackData = Seq.CompressedTrackOffsets.GetTypedData() + (TrackIndex * 2);
+		const int32 RotKeysOffset = *(TrackData + 1);
+
+		GetBoneAtomRotation( BoneAtom, Seq, RotKeysOffset, Time, RelativePos, bLooping );
+	}
+}
+
+/**
+* Decompress all requested translation components from an Animation Sequence
+*
+* @param	Atoms			The FTransform array to fill in.
+* @param	DesiredPairs	Array of requested bone information
+* @param	Seq				The animation sequence to use.
+* @param	Time			Current time to solve for.
+* @param	bLooping		True when looping the stream in intended.
+* @return					None.
+*/
+void AEFPerTrackCompressionCodec::GetPoseTranslations(
+	FTransformArray& Atoms,
+	const BoneTrackArray& DesiredPairs,
+	const UAnimSequence& Seq,
+	float Time,
+	bool bLooping )
+{
+	const int32 PairCount = DesiredPairs.Num();
+	const float RelativePos = Time / Seq.SequenceLength;
+
+	for( int32 PairIndex = 0; PairIndex < PairCount; ++PairIndex )
+	{
+		const BoneTrackPair& Pair = DesiredPairs[PairIndex];
+		const int32 TrackIndex = Pair.TrackIndex;
+		const int32 AtomIndex = Pair.AtomIndex;
+		FTransform& BoneAtom = Atoms[AtomIndex];
+
+		const int32* RESTRICT TrackData = Seq.CompressedTrackOffsets.GetTypedData() + (TrackIndex * 2);
+		const int32 PosKeysOffset = *(TrackData + 0);
+
+		GetBoneAtomTranslation( BoneAtom, Seq, PosKeysOffset, Time, RelativePos, bLooping );
+	}
+}
+
+/**
+* Decompress all requested Scale components from an Animation Sequence
+*
+* @param	Atoms			The FTransform array to fill in.
+* @param	DesiredPairs	Array of requested bone information
+* @param	Seq				The animation sequence to use.
+* @param	Time			Current time to solve for.
+* @param	bLooping		True when looping the stream in intended.
+* @return					None.
+*/
+void AEFPerTrackCompressionCodec::GetPoseScales(
+	FTransformArray& Atoms,
+	const BoneTrackArray& DesiredPairs,
+	const UAnimSequence& Seq,
+	float Time,
+	bool bLooping )
+{
+	check( Seq.CompressedScaleOffsets.IsValid() );
+
+	const int32 PairCount = DesiredPairs.Num();
+	const float RelativePos = Time / Seq.SequenceLength;
+
+	for( int32 PairIndex = 0; PairIndex < PairCount; ++PairIndex )
+	{
+		const BoneTrackPair& Pair = DesiredPairs[PairIndex];
+		const int32 TrackIndex = Pair.TrackIndex;
+		const int32 AtomIndex = Pair.AtomIndex;
+		FTransform& BoneAtom = Atoms[AtomIndex];
+
+		const int32 ScaleKeysOffset = Seq.CompressedScaleOffsets.GetOffsetData( TrackIndex, 0 );
+
+		GetBoneAtomScale( BoneAtom, Seq, ScaleKeysOffset, Time, RelativePos, bLooping );
+	}
+}
+
+#endif // USE_ANIMATION_CODEC_BATCH_SOLVER
\ No newline at end of file
diff --git a/Engine/Source/Runtime/Engine/Public/AnimEncoding_PerTrackCompression.h b/Engine/Source/Runtime/Engine/Public/AnimEncoding_PerTrackCompression.h
--- a/Engine/Source/Runtime/Engine/Public/AnimEncoding_PerTrackCompression.h
+++ b/Engine/Source/Runtime/Engine/Public/AnimEncoding_PerTrackCompression.h
@@ -187,119 +187,3 @@ protected:
 		bool bLooping);
 };
 
-
-
-
-
-
-
-
-
-
-#if USE_ANIMATION_CODEC_BATCH_SOLVER
-
-/**
- * Decompress all requested rotation components from an Animation Sequence
- *
- * @param	Atoms			The FTransform array to fill in.
- * @param	DesiredPairs	Array of requested bone information
- * @param	Seq				The animation sequence to use.
- * @param	Time			Current time to solve for.
- * @param	bLooping		True when looping the stream in intended.
- * @return					None. 
- */
-FORCEINLINE_DEBUGGABLE void AEFPerTrackCompressionCodec::GetPoseRotations(
-	FTransformArray& Atoms, 
-	const BoneTrackArray& DesiredPairs,
-	const UAnimSequence& Seq,
-	float Time,
-	bool bLooping)
-{
-	const int32 PairCount = DesiredPairs.Num();
-	const float RelativePos = Time / Seq.SequenceLength;
-
-	for (int32 PairIndex = 0; PairIndex < PairCount; ++PairIndex)
-	{
-		const BoneTrackPair& Pair = DesiredPairs[PairIndex];
-		const int32 TrackIndex = Pair.TrackIndex;
-		const int32 AtomIndex = Pair.AtomIndex;
-		FTransform& BoneAtom = Atoms[AtomIndex];
-
-		const int32* RESTRICT TrackData = Seq.CompressedTrackOffsets.GetTypedData() + (TrackIndex*2);
-		const int32 RotKeysOffset	= *(TrackData+1);
-
-		GetBoneAtomRotation(BoneAtom, Seq, RotKeysOffset, Time, RelativePos, bLooping);
-	}
-}
-
-/**
- * Decompress all requested translation components from an Animation Sequence
- *
- * @param	Atoms			The FTransform array to fill in.
- * @param	DesiredPairs	Array of requested bone information
- * @param	Seq				The animation sequence to use.
- * @param	Time			Current time to solve for.
- * @param	bLooping		True when looping the stream in intended.
- * @return					None. 
- */
-FORCEINLINE_DEBUGGABLE void AEFPerTrackCompressionCodec::GetPoseTranslations(	
-	FTransformArray& Atoms, 
-	const BoneTrackArray& DesiredPairs,
-	const UAnimSequence& Seq,
-	float Time,
-	bool bLooping)
-{
-	const int32 PairCount = DesiredPairs.Num();
-	const float RelativePos = Time / Seq.SequenceLength;
-
-	for (int32 PairIndex = 0; PairIndex < PairCount; ++PairIndex)
-	{
-		const BoneTrackPair& Pair = DesiredPairs[PairIndex];
-		const int32 TrackIndex = Pair.TrackIndex;
-		const int32 AtomIndex = Pair.AtomIndex;
-		FTransform& BoneAtom = Atoms[AtomIndex];
-
-		const int32* RESTRICT TrackData = Seq.CompressedTrackOffsets.GetTypedData() + (TrackIndex*2);
-		const int32 PosKeysOffset	= *(TrackData+0);
-
-		GetBoneAtomTranslation(BoneAtom, Seq, PosKeysOffset, Time, RelativePos, bLooping);
-	}
-}
-
-/**
- * Decompress all requested Scale components from an Animation Sequence
- *
- * @param	Atoms			The FTransform array to fill in.
- * @param	DesiredPairs	Array of requested bone information
- * @param	Seq				The animation sequence to use.
- * @param	Time			Current time to solve for.
- * @param	bLooping		True when looping the stream in intended.
- * @return					None. 
- */
-FORCEINLINE_DEBUGGABLE void AEFPerTrackCompressionCodec::GetPoseScales(	
-	FTransformArray& Atoms, 
-	const BoneTrackArray& DesiredPairs,
-	const UAnimSequence& Seq,
-	float Time,
-	bool bLooping)
-{
-	check (Seq.CompressedScaleOffsets.IsValid());
-
-	const int32 PairCount = DesiredPairs.Num();
-	const float RelativePos = Time / Seq.SequenceLength;
-
-	for (int32 PairIndex = 0; PairIndex < PairCount; ++PairIndex)
-	{
-		const BoneTrackPair& Pair = DesiredPairs[PairIndex];
-		const int32 TrackIndex = Pair.TrackIndex;
-		const int32 AtomIndex = Pair.AtomIndex;
-		FTransform& BoneAtom = Atoms[AtomIndex];
-
-		const int32 ScaleKeysOffset	= Seq.CompressedScaleOffsets.GetOffsetData(TrackIndex, 0);
-
-		GetBoneAtomScale(BoneAtom, Seq, ScaleKeysOffset, Time, RelativePos, bLooping);
-	}
-}
-
-
-#endif // USE_ANIMATION_CODEC_BATCH_SOLVER
