ue4 - Fixed performance issue with navmesh updates from actors (issue found by ZakM).

diff --git a/Engine/Source/Runtime/Engine/Private/PrimitiveComponent.cpp b/Engine/Source/Runtime/Engine/Private/PrimitiveComponent.cpp
--- a/Engine/Source/Runtime/Engine/Private/PrimitiveComponent.cpp
+++ b/Engine/Source/Runtime/Engine/Private/PrimitiveComponent.cpp
@@ -1187,7 +1187,7 @@ void UPrimitiveComponent::SetRelativeScale3D(FVector NewScale3D)
 	Super::SetRelativeScale3D(NewScale3D);
 
 	AActor* Actor = GetOwner();
-	if (Actor && OldScale3D != RelativeScale3D && IsRegistered() && IsNavigationRelevant() && World != NULL && World->IsGameWorld() && World->GetNetMode() < ENetMode::NM_Client)
+	if (Actor && OldScale3D != RelativeScale3D && IsRegistered() && Actor->IsNavigationRelevant() && CanEverAffectNavigation() && IsNavigationRelevant() && World != NULL && World->IsGameWorld() && World->GetNetMode() < ENetMode::NM_Client)
 	{
 		if (UNavigationSystem* NavSys = World->GetNavigationSystem())
 		{
@@ -1493,7 +1493,7 @@ bool UPrimitiveComponent::MoveComponent( const FVector& Delta, const FRotator& N
 			UpdateOverlaps(&PendingOverlaps, true, OverlapsAtEndLocationPtr);
 		}
 
-		if (Actor && IsRegistered() && IsNavigationRelevant() && World != NULL && World->IsGameWorld() && World->GetNetMode() < ENetMode::NM_Client)
+		if (Actor && IsRegistered() && Actor->IsNavigationRelevant() && CanEverAffectNavigation() && IsNavigationRelevant() && World != NULL && World->IsGameWorld() && World->GetNetMode() < ENetMode::NM_Client)
 		{
 			if (UNavigationSystem* NavSys = World->GetNavigationSystem())
 			{
