[335081] Fix for a crash when attempting to open a new Blueprint derived from a TriggerBox.

- Fixed up a code merge issue from CL# 19774656 (based on original change in CL# 1948728).

diff --git a/Engine/Source/Runtime/Engine/Private/TriggerActors.cpp b/Engine/Source/Runtime/Engine/Private/TriggerActors.cpp
--- a/Engine/Source/Runtime/Engine/Private/TriggerActors.cpp
+++ b/Engine/Source/Runtime/Engine/Private/TriggerActors.cpp
@@ -90,7 +90,6 @@ ATriggerBox::ATriggerBox(const class FPostConstructInitializeProperties& PCIP)
 	static FName CollisionProfileName(TEXT("Trigger"));
 	BoxCollisionComponent->SetCollisionProfileName(CollisionProfileName);
 
-	TSubobjectPtr<UBillboardComponent> SpriteComponent = PCIP.CreateEditorOnlyDefaultSubobject<UBillboardComponent>(this, TEXT("Sprite"));
 	if (SpriteComponent)
 	{
 		SpriteComponent->AttachParent = BoxCollisionComponent;
@@ -134,7 +133,6 @@ ATriggerSphere::ATriggerSphere(const class FPostConstructInitializeProperties& P
 	static FName CollisionProfileName(TEXT("Trigger"));
 	SphereCollisionComponent->SetCollisionProfileName(CollisionProfileName);
 
-	TSubobjectPtr<UBillboardComponent> SpriteComponent = PCIP.CreateEditorOnlyDefaultSubobject<UBillboardComponent>(this, TEXT("Sprite"));
 	if (SpriteComponent)
 	{
 		SpriteComponent->AttachParent = SphereCollisionComponent;
