Re-enable named events for external profilers.

diff --git a/Engine/Source/Runtime/Engine/Private/UnrealEngine.cpp b/Engine/Source/Runtime/Engine/Private/UnrealEngine.cpp
--- a/Engine/Source/Runtime/Engine/Private/UnrealEngine.cpp
+++ b/Engine/Source/Runtime/Engine/Private/UnrealEngine.cpp
@@ -9957,11 +9957,13 @@ bool UEngine::ToggleStatNamedEvents(UWorld* World, FCommonViewportClient* Viewpo
 	if (ViewportClient->IsStatEnabled(TEXT("NamedEvents")))
 	{
 		StatsMasterEnableAdd();
+		GCycleStatsShouldEmitNamedEvents = true;
 	}
 	// Disable emission of named events and force-enabling cycle stats.
 	else
 	{
-		StatsMasterEnableSubtract();
+		GCycleStatsShouldEmitNamedEvents = false;
+		StatsMasterEnableSubtract();		
 	}
 	return false;
 }
