Networking fix for vehicles. Move networking code out of substep and into pretick TTP#335110

diff --git a/Engine/Source/Runtime/Engine/Classes/Vehicles/WheeledVehicleMovementComponent.h b/Engine/Source/Runtime/Engine/Classes/Vehicles/WheeledVehicleMovementComponent.h
--- a/Engine/Source/Runtime/Engine/Classes/Vehicles/WheeledVehicleMovementComponent.h
+++ b/Engine/Source/Runtime/Engine/Classes/Vehicles/WheeledVehicleMovementComponent.h
@@ -258,8 +258,8 @@ class ENGINE_API UWheeledVehicleMovementComponent : public UPawnMovementComponen
 	/** Tick this vehicle sim right before input is sent to the vehicle system  */
 	virtual void TickVehicle( float DeltaTime );
 
-	/** Updates the vehicle tuning data. Used when designer tweaks values while game is running */
-	virtual void UpdateTuning();
+	/** Updates the vehicle tuning and other state such as user input. */
+	virtual void PreTick(float DeltaTime);
 
 	/** Updates the forces of drag acting on the vehicle */
 	virtual void UpdateDrag( float DeltaTime );
diff --git a/Engine/Source/Runtime/Engine/Private/PhysicsEngine/PhysScene.cpp b/Engine/Source/Runtime/Engine/Private/PhysicsEngine/PhysScene.cpp
--- a/Engine/Source/Runtime/Engine/Private/PhysicsEngine/PhysScene.cpp
+++ b/Engine/Source/Runtime/Engine/Private/PhysicsEngine/PhysScene.cpp
@@ -413,7 +413,14 @@ void FPhysScene::TickPhysScene(uint32 SceneType, FGraphEventRef& InOutCompletion
 
 	if (VehicleManager && SceneType == PST_Sync)
 	{
-		VehicleManager->UpdateTuning();
+		float TickTime = AveragedFrameTime[SceneType];
+#if WITH_SUBSTEPPING
+		if (IsSubstepping())
+		{
+			TickTime = UseSyncTime(SceneType) ? SyncDeltaSeconds : DeltaSeconds;
+		}
+#endif
+		VehicleManager->PreTick(TickTime);
 #if WITH_SUBSTEPPING
 		if (IsSubstepping() == false)
 #endif
diff --git a/Engine/Source/Runtime/Engine/Private/Vehicles/PhysXVehicleManager.cpp b/Engine/Source/Runtime/Engine/Private/Vehicles/PhysXVehicleManager.cpp
--- a/Engine/Source/Runtime/Engine/Private/Vehicles/PhysXVehicleManager.cpp
+++ b/Engine/Source/Runtime/Engine/Private/Vehicles/PhysXVehicleManager.cpp
@@ -291,11 +291,11 @@ void FPhysXVehicleManager::Update( float DeltaTime )
 #endif //PX_DEBUG_VEHICLE_ON
 }
 
-void FPhysXVehicleManager::UpdateTuning()
+void FPhysXVehicleManager::PreTick( float DeltaTime )
 {
 	for (int32 i = 0; i < Vehicles.Num(); ++i)
 	{
-		Vehicles[i]->UpdateTuning();
+		Vehicles[i]->PreTick(DeltaTime);
 	}
 }
 
diff --git a/Engine/Source/Runtime/Engine/Private/Vehicles/PhysXVehicleManager.h b/Engine/Source/Runtime/Engine/Private/Vehicles/PhysXVehicleManager.h
--- a/Engine/Source/Runtime/Engine/Private/Vehicles/PhysXVehicleManager.h
+++ b/Engine/Source/Runtime/Engine/Private/Vehicles/PhysXVehicleManager.h
@@ -58,8 +58,8 @@ public:
 	void Update( float DeltaTime );
 	
 	/**
-	 * Update vehicle tuning. Used when designer tweaks values while the game is running */
-	void UpdateTuning();
+	 * Update vehicle tuning and other state such as input */
+	void PreTick( float DeltaTime );
 
 private:
 
diff --git a/Engine/Source/Runtime/Engine/Private/Vehicles/WheeledVehicleMovementComponent.cpp b/Engine/Source/Runtime/Engine/Private/Vehicles/WheeledVehicleMovementComponent.cpp
--- a/Engine/Source/Runtime/Engine/Private/Vehicles/WheeledVehicleMovementComponent.cpp
+++ b/Engine/Source/Runtime/Engine/Private/Vehicles/WheeledVehicleMovementComponent.cpp
@@ -654,7 +654,6 @@ void UWheeledVehicleMovementComponent::TickVehicle( float DeltaTime )
 		APawn* MyOwner = Cast<APawn>(UpdatedComponent->GetOwner());
 		if (MyOwner)
 		{
-			UpdateState(DeltaTime);
 			UpdateSimulation(DeltaTime);
 		}
 	}
@@ -686,8 +685,18 @@ void UWheeledVehicleMovementComponent::UpdateDrag(float DeltaTime)
 	}
 }
 
-void UWheeledVehicleMovementComponent::UpdateTuning()
+void UWheeledVehicleMovementComponent::PreTick(float DeltaTime)
 {
+	// movement updates and replication
+	if (PVehicle && UpdatedComponent)
+	{
+		APawn* MyOwner = Cast<APawn>(UpdatedComponent->GetOwner());
+		if (MyOwner)
+		{
+			UpdateState(DeltaTime);
+		}
+	}
+
 	if (VehicleSetupTag != FPhysXVehicleManager::VehicleSetupTag)
 	{
 		RecreatePhysicsState();
@@ -697,7 +706,6 @@ void UWheeledVehicleMovementComponent::UpdateTuning()
 void UWheeledVehicleMovementComponent::UpdateSimulation( float DeltaTime )
 {
 }
-
 #endif // WITH_PHYSX
 
 void UWheeledVehicleMovementComponent::UpdateState( float DeltaTime )
