Make Input* functions virtual in PlayerController and never pass directly to PlayerInput from GameviewportClient

diff --git a/Engine/Source/Runtime/Engine/Classes/GameFramework/PlayerController.h b/Engine/Source/Runtime/Engine/Classes/GameFramework/PlayerController.h
--- a/Engine/Source/Runtime/Engine/Classes/GameFramework/PlayerController.h
+++ b/Engine/Source/Runtime/Engine/Classes/GameFramework/PlayerController.h
@@ -951,8 +951,12 @@ public:
 	/** Removes given inputcomponent from the input stack (regardless of if it's the top, actually). */
 	bool PopInputComponent(UInputComponent* Input);
 
-	bool InputKey(FKey Key, EInputEvent EventType, float AmountDepressed, bool bGamepad);
-	bool InputTouch(uint32 Handle, ETouchType::Type Type, const FVector2D& TouchLocation, FDateTime DeviceTimestamp, uint32 TouchpadIndex);
+	virtual void FlushPressedKeys();
+
+	virtual bool InputKey(FKey Key, EInputEvent EventType, float AmountDepressed, bool bGamepad);
+	virtual bool InputTouch(uint32 Handle, ETouchType::Type Type, const FVector2D& TouchLocation, FDateTime DeviceTimestamp, uint32 TouchpadIndex);
+	virtual bool InputAxis(FKey Key, float Delta, float DeltaTime, int32 NumSamples, bool bGamepad);
+	virtual bool InputMotion(const FVector& Tilt, const FVector& RotationRate, const FVector& Gravity, const FVector& Acceleration);
 
 	/** Associate a new UPlayer with this PlayerController. */
 	virtual void SetPlayer(UPlayer* Player);
diff --git a/Engine/Source/Runtime/Engine/Classes/GameFramework/PlayerInput.h b/Engine/Source/Runtime/Engine/Classes/GameFramework/PlayerInput.h
--- a/Engine/Source/Runtime/Engine/Classes/GameFramework/PlayerInput.h
+++ b/Engine/Source/Runtime/Engine/Classes/GameFramework/PlayerInput.h
@@ -369,8 +369,8 @@ public:
 	virtual void PostInitProperties() OVERRIDE;
 
 	void FlushPressedKeys();
-	bool InputKey(FKey Key, enum EInputEvent Event, float AmountDepressed = 1.f, bool bGamepad = false );
-	bool InputAxis(FKey Key, float Delta, float DeltaTime, int32 NumSamples=1, bool bGamepad = false);
+	bool InputKey(FKey Key, enum EInputEvent Event, float AmountDepressed, bool bGamepad);
+	bool InputAxis(FKey Key, float Delta, float DeltaTime, int32 NumSamples, bool bGamepad);
 	bool InputTouch(uint32 Handle, ETouchType::Type Type, const FVector2D& TouchLocation, FDateTime DeviceTimestamp, uint32 TouchpadIndex);
 	bool InputMotion(const FVector& Tilt, const FVector& RotationRate, const FVector& Gravity, const FVector& Acceleration);
 	bool InputGesture(const FKey Gesture, const EInputEvent Event, const float Value);
diff --git a/Engine/Source/Runtime/Engine/Private/GameViewportClient.cpp b/Engine/Source/Runtime/Engine/Private/GameViewportClient.cpp
--- a/Engine/Source/Runtime/Engine/Private/GameViewportClient.cpp
+++ b/Engine/Source/Runtime/Engine/Private/GameViewportClient.cpp
@@ -242,9 +242,9 @@ bool UGameViewportClient::InputAxis(FViewport* InViewport, int32 ControllerId, F
 		if (!bResult)
 		{
 			ULocalPlayer* const TargetPlayer = GEngine->GetLocalPlayerFromControllerId(this, ControllerId);
-			if (TargetPlayer && TargetPlayer->PlayerController && TargetPlayer->PlayerController->PlayerInput)
+			if (TargetPlayer && TargetPlayer->PlayerController)
 			{
-				bResult = TargetPlayer->PlayerController->PlayerInput->InputAxis(Key, Delta, DeltaTime, NumSamples, bGamepad);
+				bResult = TargetPlayer->PlayerController->InputAxis(Key, Delta, DeltaTime, NumSamples, bGamepad);
 			}
 		}
 
@@ -310,9 +310,9 @@ bool UGameViewportClient::InputMotion(FViewport* InViewport, int32 ControllerId,
 	bool bResult = false;
 
 	ULocalPlayer* const TargetPlayer = GEngine->GetLocalPlayerFromControllerId(this, ControllerId);
-	if (TargetPlayer && TargetPlayer->PlayerController && TargetPlayer->PlayerController->PlayerInput)
+	if (TargetPlayer && TargetPlayer->PlayerController)
 	{
-		bResult = TargetPlayer->PlayerController->PlayerInput->InputMotion(Tilt, RotationRate, Gravity, Acceleration);
+		bResult = TargetPlayer->PlayerController->InputMotion(Tilt, RotationRate, Gravity, Acceleration);
 	}
 
 	return bResult;
@@ -1114,9 +1114,9 @@ void UGameViewportClient::LostFocus(FViewport* InViewport)
 	for( FConstPlayerControllerIterator Iterator = GetWorld()->GetPlayerControllerIterator(); Iterator; ++Iterator )
 	{
 		APlayerController* PlayerController = *Iterator;
-		if (PlayerController && PlayerController->PlayerInput)
+		if (PlayerController)
 		{
-			PlayerController->PlayerInput->FlushPressedKeys();
+			PlayerController->FlushPressedKeys();
 		}
 	}
 }
diff --git a/Engine/Source/Runtime/Engine/Private/PlayerController.cpp b/Engine/Source/Runtime/Engine/Private/PlayerController.cpp
--- a/Engine/Source/Runtime/Engine/Private/PlayerController.cpp
+++ b/Engine/Source/Runtime/Engine/Private/PlayerController.cpp
@@ -1952,6 +1952,14 @@ void APlayerController::PlayerTick( float DeltaTime )
 	}
 }
 
+void APlayerController::FlushPressedKeys()
+{
+	if (PlayerInput)
+	{
+		PlayerInput->FlushPressedKeys();
+	}
+}
+
 bool APlayerController::InputKey(FKey Key, EInputEvent EventType, float AmountDepressed, bool bGamepad)
 {
 	bool bResult = false;
@@ -2008,6 +2016,18 @@ bool APlayerController::InputKey(FKey Key, EInputEvent EventType, float AmountDe
 	return bResult;
 }
 
+bool APlayerController::InputAxis(FKey Key, float Delta, float DeltaTime, int32 NumSamples, bool bGamepad)
+{
+	bool bResult = false;
+	
+	if (PlayerInput)
+	{
+		bResult = PlayerInput->InputAxis(Key, Delta, DeltaTime, NumSamples, bGamepad);
+	}
+
+	return bResult;
+}
+
 bool APlayerController::InputTouch(uint32 Handle, ETouchType::Type Type, const FVector2D& TouchLocation, FDateTime DeviceTimestamp, uint32 TouchpadIndex)
 {
 	bool bResult = false;
@@ -2086,6 +2106,18 @@ bool APlayerController::InputTouch(uint32 Handle, ETouchType::Type Type, const F
 	return bResult;
 }
 
+bool APlayerController::InputMotion(const FVector& Tilt, const FVector& RotationRate, const FVector& Gravity, const FVector& Acceleration)
+{
+	bool bResult = false;
+
+	if (PlayerInput)
+	{
+		bResult = PlayerInput->InputMotion(Tilt, RotationRate, Gravity, Acceleration);
+	}
+
+	return bResult;
+}
+
 bool APlayerController::ShouldShowMouseCursor() const
 {
 	return bShowMouseCursor;
diff --git a/Engine/Source/Runtime/Engine/Private/UserInterface/PlayerInput.cpp b/Engine/Source/Runtime/Engine/Private/UserInterface/PlayerInput.cpp
--- a/Engine/Source/Runtime/Engine/Private/UserInterface/PlayerInput.cpp
+++ b/Engine/Source/Runtime/Engine/Private/UserInterface/PlayerInput.cpp
@@ -62,7 +62,8 @@ void UPlayerInput::FlushPressedKeys()
 
 			for ( int32 KeyIndex = 0; KeyIndex < PressedKeys.Num(); KeyIndex++ )
 			{
-				InputKey(PressedKeys[KeyIndex], IE_Released, 0);
+				FKey& Key = PressedKeys[KeyIndex];
+				InputKey(Key, IE_Released, 0, Key.IsGamepadKey());
 			}
 		}
 	}
