Fix stat sounds -debug regression introduced in CL#2058522
TTP#335055

diff --git a/Engine/Source/Editor/UnrealEd/Public/LevelEditorViewport.h b/Engine/Source/Editor/UnrealEd/Public/LevelEditorViewport.h
--- a/Engine/Source/Editor/UnrealEd/Public/LevelEditorViewport.h
+++ b/Engine/Source/Editor/UnrealEd/Public/LevelEditorViewport.h
@@ -576,7 +576,7 @@ public:
 	 */
 	virtual ESoundShowFlags::Type GetSoundShowFlags() const OVERRIDE
 	{ 
-		return ESoundShowFlags::Disabled;
+		return SoundShowFlags;
 	}
 
 	/**
diff --git a/Engine/Source/Runtime/Engine/Classes/Engine/GameViewportClient.h b/Engine/Source/Runtime/Engine/Classes/Engine/GameViewportClient.h
--- a/Engine/Source/Runtime/Engine/Classes/Engine/GameViewportClient.h
+++ b/Engine/Source/Runtime/Engine/Classes/Engine/GameViewportClient.h
@@ -659,7 +659,7 @@ public:
 	 */
 	virtual ESoundShowFlags::Type GetSoundShowFlags() const OVERRIDE
 	{ 
-		return ESoundShowFlags::Disabled;
+		return SoundShowFlags;
 	}
 
 	/**
