GameFramework: Fix BlendRot (used in SmoothTargetViewRotation) to return a float instead of int32 (still very magic)

diff --git a/Engine/Source/Runtime/Engine/Private/PlayerController.cpp b/Engine/Source/Runtime/Engine/Private/PlayerController.cpp
--- a/Engine/Source/Runtime/Engine/Private/PlayerController.cpp
+++ b/Engine/Source/Runtime/Engine/Private/PlayerController.cpp
@@ -466,30 +466,41 @@ void APlayerController::ForceSingleNetUpdateFor(AActor* Target)
 	}
 }
 
-/** worker function for APlayerController::SmoothTargetViewRotation() */
-int32 BlendRot(float DeltaTime, float BlendC, float NewC)
+void APlayerController::SmoothTargetViewRotation(APawn* TargetPawn, float DeltaSeconds)
 {
-	if ( FMath::Abs(BlendC - NewC) > 180.f )
+	struct FBlendHelper
 	{
-		if ( BlendC > NewC )
-			NewC += 360.f;
-		else
-			BlendC += 360.f;
-	}
-	if ( FMath::Abs(BlendC - NewC) > 22.57 )
-		BlendC = NewC;
-	else
-		BlendC = BlendC + (NewC - BlendC) * FMath::Min(1.f,24.f * DeltaTime);
+		/** worker function for APlayerController::SmoothTargetViewRotation() */
+		static float BlendRotation(float DeltaTime, float BlendC, float NewC)
+		{
+			if (FMath::Abs(BlendC - NewC) > 180.f)
+			{
+				if (BlendC > NewC)
+				{
+					NewC += 360.f;
+				}
+				else
+				{
+					BlendC += 360.f;
+				}
+			}
 
-	return FRotator::ClampAxis(BlendC);
-}
+			if (FMath::Abs(BlendC - NewC) > 22.57f)
+			{
+				BlendC = NewC;
+			}
+			else
+			{
+				BlendC = BlendC + (NewC - BlendC) * FMath::Min(1.f, 24.f * DeltaTime);
+			}
 
+			return FRotator::ClampAxis(BlendC);
+		}
+	};
 
-void APlayerController::SmoothTargetViewRotation(APawn* TargetPawn, float DeltaSeconds)
-{
-	BlendedTargetViewRotation.Pitch = BlendRot(DeltaSeconds, BlendedTargetViewRotation.Pitch, FRotator::ClampAxis(TargetViewRotation.Pitch));
-	BlendedTargetViewRotation.Yaw = BlendRot(DeltaSeconds, BlendedTargetViewRotation.Yaw, FRotator::ClampAxis(TargetViewRotation.Yaw));
-	BlendedTargetViewRotation.Roll = BlendRot(DeltaSeconds, BlendedTargetViewRotation.Roll, FRotator::ClampAxis(TargetViewRotation.Roll));
+	BlendedTargetViewRotation.Pitch = FBlendHelper::BlendRotation(DeltaSeconds, BlendedTargetViewRotation.Pitch, FRotator::ClampAxis(TargetViewRotation.Pitch));
+	BlendedTargetViewRotation.Yaw = FBlendHelper::BlendRotation(DeltaSeconds, BlendedTargetViewRotation.Yaw, FRotator::ClampAxis(TargetViewRotation.Yaw));
+	BlendedTargetViewRotation.Roll = FBlendHelper::BlendRotation(DeltaSeconds, BlendedTargetViewRotation.Roll, FRotator::ClampAxis(TargetViewRotation.Roll));
 }
 
 
