Additional check to ensure LODResources element isn't NULL

#ttp 334031 - CRASH: Client crash occurs through natural gameplay

diff --git a/Engine/Source/Runtime/Engine/Private/StaticMesh.cpp b/Engine/Source/Runtime/Engine/Private/StaticMesh.cpp
--- a/Engine/Source/Runtime/Engine/Private/StaticMesh.cpp
+++ b/Engine/Source/Runtime/Engine/Private/StaticMesh.cpp
@@ -1211,6 +1211,7 @@ bool UStaticMesh::HasValidRenderData() const
 {
 	return RenderData != NULL
 		&& RenderData->LODResources.Num() > 0
+		&& RenderData->LODResources.GetTypedData() != NULL
 		&& RenderData->LODResources[0].VertexBuffer.GetNumVertices() > 0;
 }
 
