Landscape : Fixed LODDistanceFactor applying problem

diff --git a/Engine/Source/Runtime/Engine/Private/Landscape/LandscapeRender.cpp b/Engine/Source/Runtime/Engine/Private/Landscape/LandscapeRender.cpp
--- a/Engine/Source/Runtime/Engine/Private/Landscape/LandscapeRender.cpp
+++ b/Engine/Source/Runtime/Engine/Private/Landscape/LandscapeRender.cpp
@@ -587,7 +587,7 @@ FLandscapeComponentSceneProxy::FLandscapeComponentSceneProxy(ULandscapeComponent
 
 	LODDistance = FMath::Sqrt(2.f * FMath::Square((float)SubsectionSizeQuads)) * LANDSCAPE_LOD_DISTANCE_FACTOR / InComponent->GetLandscapeProxy()->LODDistanceFactor; // vary in 0...1
 	DistDiff = -FMath::Sqrt(2.f * FMath::Square(0.5f*(float)SubsectionSizeQuads));
-	LODDistanceFactor = InComponent->GetLandscapeProxy()->LODDistanceFactor * 0.33f;
+	LODDistanceFactor = FMath::Min(InComponent->GetLandscapeProxy()->LODDistanceFactor, 3.f) * 0.33f;
 
 	if (InComponent->StaticLightingResolution > 0.f)
 	{
@@ -2362,7 +2362,7 @@ void ULandscapeComponent::GetStreamingTextureInfo(TArray<FStreamingTexturePrimit
 
 void ALandscapeProxy::ChangeLODDistanceFactor(float InLODDistanceFactor)
 {
-	LODDistanceFactor = FMath::Clamp<float>(InLODDistanceFactor, 0.1f, 3.f);
+	LODDistanceFactor = FMath::Clamp<float>(InLODDistanceFactor, 0.1f, 10.f);
 	
 	if (LandscapeComponents.Num())
 	{
@@ -2377,7 +2377,7 @@ void ALandscapeProxy::ChangeLODDistanceFactor(float InLODDistanceFactor)
 			LandscapeChangeLODDistanceFactorCommand,
 			FLandscapeComponentSceneProxy**, Proxies, Proxies,
 			int32, CompNum, CompNum,
-			FVector2D, InLODDistanceFactors, FVector2D(FMath::Sqrt(2.f * FMath::Square((float)SubsectionSizeQuads)) * LANDSCAPE_LOD_DISTANCE_FACTOR / LODDistanceFactor, LODDistanceFactor * 0.33f),
+			FVector2D, InLODDistanceFactors, FVector2D(FMath::Sqrt(2.f * FMath::Square((float)SubsectionSizeQuads)) * LANDSCAPE_LOD_DISTANCE_FACTOR / LODDistanceFactor, FMath::Max(LODDistanceFactor, 3.f) * 0.33f),
 		{
 			for (int32 Idx = 0; Idx < CompNum; ++Idx)
 			{
