Temp crash fix for GetTrailEmitters. Code is reworked in an upcoming change list which bypasses the issue anyway but this will patch it till then.

diff --git a/Engine/Source/Runtime/Engine/Private/Particles/ParticleComponents.cpp b/Engine/Source/Runtime/Engine/Private/Particles/ParticleComponents.cpp
--- a/Engine/Source/Runtime/Engine/Private/Particles/ParticleComponents.cpp
+++ b/Engine/Source/Runtime/Engine/Private/Particles/ParticleComponents.cpp
@@ -4733,10 +4733,11 @@ void UParticleSystemComponent::GetTrailEmitters(UAnimNotifyState* InAnimNotifySt
 	int32 NumEmitters = EmitterInstances.Num();
 	for (int32 i = 0; i < NumEmitters; ++i)
 	{
-		if (EmitterInstances[i]->IsTrailEmitter())
+		FParticleEmitterInstance* Instance = EmitterInstances[i];
+		if (Instance && Instance->IsTrailEmitter())
 		{
-			FParticleAnimTrailEmitterInstance* Instance = (FParticleAnimTrailEmitterInstance*)(EmitterInstances[i]);
-			if (Instance && ((bIncludeUnassociated && Instance->AnimNotifyState == NULL) || Instance->AnimNotifyState == InAnimNotifyState))
+			FParticleAnimTrailEmitterInstance* TrailInstance = (FParticleAnimTrailEmitterInstance*)(Instance);
+			if ((bIncludeUnassociated && TrailInstance->AnimNotifyState == NULL) || TrailInstance->AnimNotifyState == InAnimNotifyState)
 			{
 				OutTrailEmitters.Add(Instance);
 			}
