BP: Pin connections accept inherited structures, as long as they have the same layout.
ttp#335142 Blueprints can no longer implicitly convert from FVector_NetQuantize to FVector

diff --git a/Engine/Source/Editor/BlueprintGraph/Private/EdGraphSchema_K2.cpp b/Engine/Source/Editor/BlueprintGraph/Private/EdGraphSchema_K2.cpp
--- a/Engine/Source/Editor/BlueprintGraph/Private/EdGraphSchema_K2.cpp
+++ b/Engine/Source/Editor/BlueprintGraph/Private/EdGraphSchema_K2.cpp
@@ -2630,7 +2630,7 @@ bool UEdGraphSchema_K2::ArePinTypesCompatible(const FEdGraphPinType& Output, con
 			{
 				if (Output.PinCategory == PC_Struct)
 				{
-					return (OutputObject == InputObject);
+					return OutputObject->IsChildOf(InputObject) && FStructUtils::TheSameLayout(OutputObject, InputObject);
 				}
 
 				// Special Case:  Cannot mix interface and non-interface calls, because the pointer size is different under the hood
diff --git a/Engine/Source/Runtime/Engine/Classes/Engine/BlueprintGeneratedClass.h b/Engine/Source/Runtime/Engine/Classes/Engine/BlueprintGeneratedClass.h
--- a/Engine/Source/Runtime/Engine/Classes/Engine/BlueprintGeneratedClass.h
+++ b/Engine/Source/Runtime/Engine/Classes/Engine/BlueprintGeneratedClass.h
@@ -132,6 +132,15 @@ private:
 	int32 WriteIndex;
 };
 
+struct FStructUtils
+{
+private:
+	static bool ArePropertiesTheSame(const UProperty* A, const UProperty* B);
+
+public:
+	// does structures have exactly the same memory layout
+	ENGINE_API static bool TheSameLayout(const UStruct* StructA, const UStruct* StructB);
+};
 
 USTRUCT()
 struct ENGINE_API FBlueprintDebugData
diff --git a/Engine/Source/Runtime/Engine/Private/BlueprintGeneratedClass.cpp b/Engine/Source/Runtime/Engine/Private/BlueprintGeneratedClass.cpp
--- a/Engine/Source/Runtime/Engine/Private/BlueprintGeneratedClass.cpp
+++ b/Engine/Source/Runtime/Engine/Private/BlueprintGeneratedClass.cpp
@@ -68,122 +68,72 @@ UClass* UBlueprintGeneratedClass::GetAuthoritativeClass()
 	return (GeneratingBP->GeneratedClass != NULL) ? GeneratingBP->GeneratedClass : this;
 }
 
-struct FConditionalRecompileClassHepler
+bool FStructUtils::ArePropertiesTheSame(const UProperty* A, const UProperty* B)
 {
-	enum ENeededAction
+	if (A == B)
 	{
-		None,
-		StaticLink,
-		Recompile,
-	};
+		return true;
+	}
 
-	static bool SameTypeProperties(const UProperty* A, const UProperty* B)
+	if (!A != !B) //one of properties is null
 	{
-		check(A && B);
-		const UClass* ClassA = A->GetClass();
-		const UClass* ClassB = B->GetClass();
-		if (ClassA != ClassB)
-		{
-			return false;
-		}
-
-		if (A->IsA<UObjectPropertyBase>()
-			&&  (CastChecked<UObjectPropertyBase>(A)->PropertyClass != CastChecked<UObjectPropertyBase>(B)->PropertyClass))
-		{
-			return false;
-		}
-
-		if (A->IsA<UClassProperty>() && 
-			(CastChecked<UClassProperty>(A)->MetaClass != CastChecked<UClassProperty>(B)->MetaClass))
-		{
-			return false;
-		}
-		else if (A->IsA<UAssetClassProperty>() && 
-			(CastChecked<UAssetClassProperty>(A)->MetaClass != CastChecked<UAssetClassProperty>(B)->MetaClass))
-		{
-			return false;
-		}
-		else if (A->IsA<UInterfaceProperty>() && 
-			(CastChecked<UInterfaceProperty>(A)->InterfaceClass != CastChecked<UInterfaceProperty>(B)->InterfaceClass))
-		{
-			return false;
-		}
-		else if (A->IsA<UArrayProperty>() && 
-			(!SameTypeProperties(CastChecked<UArrayProperty>(A)->Inner, CastChecked<UArrayProperty>(B)->Inner)))
-		{
-			return false;
-		}
-		else if (A->IsA<UStructProperty>() && 
-			(CastChecked<UStructProperty>(A)->Struct != CastChecked<UStructProperty>(B)->Struct))
-		{
-			return false;
-		}
-		else if (A->IsA<UDelegateProperty>() && 
-			(CastChecked<UDelegateProperty>(A)->SignatureFunction != CastChecked<UDelegateProperty>(B)->SignatureFunction))
-		{
-			return false;
-		}
-		else if (A->IsA<UMulticastDelegateProperty>() && 
-			(CastChecked<UMulticastDelegateProperty>(A)->SignatureFunction != CastChecked<UMulticastDelegateProperty>(B)->SignatureFunction))
-		{
-			return false;
-		}
+		return false;
+	}
 
-		return true;
+	if (A->GetSize() != B->GetSize())
+	{
+		return false;
 	}
 
-	static bool ArePropertiesTheSame(const UProperty* A, const UProperty* B)
+	if (A->GetOffset_ForGC() != B->GetOffset_ForGC())
 	{
-		if (A == B)
-		{
-			return true;
-		}
+		return false;
+	}
 
-		if (!A != !B) //one of properties is null
-		{
-			return false;
-		}
+	if (!A->SameType(B))
+	{
+		return false;
+	}
 
-		if (A->GetSize() != B->GetSize())
-		{
-			return false;
-		}
+	return true;
+}
 
-		if (A->GetOffset_ForGC() != B->GetOffset_ForGC())
-		{
-			return false;
-		}
+bool FStructUtils::TheSameLayout(const UStruct* StructA, const UStruct* StructB)
+{
+	bool bResult = false;
+	if (StructA && StructB)
+	{
+		const UProperty* PropertyA = StructA->PropertyLink;
+		const UProperty* PropertyB = StructB->PropertyLink;
 
-		if (!A->SameType(B))
+		bResult = true;
+		while (bResult)
 		{
-			return false;
+			bResult = ArePropertiesTheSame(PropertyA, PropertyB);
+			PropertyA = PropertyA ? PropertyA->PropertyLinkNext : NULL;
+			PropertyB = PropertyB ? PropertyB->PropertyLinkNext : NULL;
+			if (PropertyA == PropertyB)
+			{
+				break;
+			}
 		}
-
-		return true;
 	}
+	return bResult;
+}
+
+struct FConditionalRecompileClassHepler
+{
+	enum ENeededAction
+	{
+		None,
+		StaticLink,
+		Recompile,
+	};
 
 	static bool HasTheSameLayoutAsParent(const UStruct* Struct)
 	{
-		bool bResult = false;
 		const UStruct* Parent = Struct ? Struct->GetSuperStruct() : NULL;
-		if (Parent)
-		{
-			const UProperty* Property = Struct->PropertyLink;
-			const UProperty* SuperProperty = Parent->PropertyLink;
-
-			bResult = true;
-			while (bResult)
-			{
-				bResult = ArePropertiesTheSame(Property, SuperProperty);
-				Property = Property ? Property->PropertyLinkNext : NULL;
-				SuperProperty = SuperProperty ? SuperProperty->PropertyLinkNext : NULL;
-				if (Property == SuperProperty)
-				{
-					break;
-				}
-			}
-		}
-		return bResult;
+		return FStructUtils::TheSameLayout(Struct, Parent);
 	}
 
 	static ENeededAction IsConditionalRecompilationNecessary(const UBlueprint* GeneratingBP)
