Fixed compilation error

diff --git a/Engine/Source/Runtime/Engine/Private/Particles/ParticleComponents.cpp b/Engine/Source/Runtime/Engine/Private/Particles/ParticleComponents.cpp
--- a/Engine/Source/Runtime/Engine/Private/Particles/ParticleComponents.cpp
+++ b/Engine/Source/Runtime/Engine/Private/Particles/ParticleComponents.cpp
@@ -4739,7 +4739,7 @@ void UParticleSystemComponent::GetTrailEmitters(UAnimNotifyState* InAnimNotifySt
 			FParticleAnimTrailEmitterInstance* TrailInstance = (FParticleAnimTrailEmitterInstance*)(Instance);
 			if ((bIncludeUnassociated && TrailInstance->AnimNotifyState == NULL) || TrailInstance->AnimNotifyState == InAnimNotifyState)
 			{
-				OutTrailEmitters.Add(Instance);
+				OutTrailEmitters.Add(TrailInstance);
 			}
 		}
 	}
