Alternate fix for ttp 328395. Rather than clear the GPU particle textures on creation which isn't allowed, clear them just before first use.

diff --git a/Engine/Source/Runtime/Engine/Private/Particles/ParticleGpuSimulation.cpp b/Engine/Source/Runtime/Engine/Private/Particles/ParticleGpuSimulation.cpp
--- a/Engine/Source/Runtime/Engine/Private/Particles/ParticleGpuSimulation.cpp
+++ b/Engine/Source/Runtime/Engine/Private/Particles/ParticleGpuSimulation.cpp
@@ -156,6 +156,8 @@ public:
 	FTexture2DRHIRef VelocityTextureTargetRHI;
 	FTexture2DRHIRef VelocityTextureRHI;
 
+	bool bTexturesCleared;
+
 	/**
 	 * Initialize RHI resources used for particle simulation.
 	 */
@@ -195,6 +197,8 @@ public:
 			VelocityTextureTargetRHI,
 			VelocityTextureRHI
 			);
+
+		bTexturesCleared = false;
 	}
 
 	/**
@@ -3953,6 +3957,27 @@ void FFXSystem::SimulateGPUParticles(
 	FParticleStateTextures& CurrentStateTextures = ParticleSimulationResources->GetCurrentStateTextures();
 	FParticleStateTextures& PrevStateTextures = ParticleSimulationResources->GetPreviousStateTextures();
 
+	// On some platforms, the textures are filled with garbage after creation, so we need to clear them to black the first time we use them
+	if ( !CurrentStateTextures.bTexturesCleared )
+	{
+		RHISetRenderTarget( CurrentStateTextures.PositionTextureTargetRHI, FTextureRHIRef() );
+		RHIClear( true, FLinearColor::Black, false, 1.0f, false, 0, FIntRect() );
+		RHISetRenderTarget( CurrentStateTextures.VelocityTextureTargetRHI, FTextureRHIRef() );
+		RHIClear( true, FLinearColor::Black, false, 1.0f, false, 0, FIntRect() );
+
+		CurrentStateTextures.bTexturesCleared = true;
+	}
+
+	if ( !PrevStateTextures.bTexturesCleared )
+	{
+		RHISetRenderTarget( PrevStateTextures.PositionTextureTargetRHI, FTextureRHIRef() );
+		RHIClear( true, FLinearColor::Black, false, 1.0f, false, 0, FIntRect() );
+		RHISetRenderTarget( PrevStateTextures.VelocityTextureTargetRHI, FTextureRHIRef() );
+		RHIClear( true, FLinearColor::Black, false, 1.0f, false, 0, FIntRect() );
+
+		PrevStateTextures.bTexturesCleared = true;
+	}
+
 	// Setup render states.
 	FTextureRHIParamRef RenderTargets[3] =
 	{
