- removed deprecated OSS tests

diff --git a/Engine/Source/Runtime/Engine/Private/Tests/OSSTests.cpp b/Engine/Source/Runtime/Engine/Private/Tests/OSSTests.cpp
--- a/Engine/Source/Runtime/Engine/Private/Tests/OSSTests.cpp
+++ b/Engine/Source/Runtime/Engine/Private/Tests/OSSTests.cpp
@@ -28,26 +28,6 @@ bool FHttpTest::RunTest(const FString& Parameters)
 	return true;
 }
 
-IMPLEMENT_SIMPLE_AUTOMATION_TEST(FMCPAccountCreationTest,"OSS.MCP.Test Account Creation", EAutomationTestFlags::ATF_Editor )
-
-bool FMCPAccountCreationTest::RunTest(const FString& Parameters)
-{
-	ADD_LATENT_AUTOMATION_COMMAND(FEngineWaitLatentCommand(3.0));
-	ADD_LATENT_AUTOMATION_COMMAND(FExecStringLatentCommand(TEXT("online sub=mcp test accountcreation [REDACTED_EMAIL] automationguy Epic3234")));
-
-	return true;
-}
-
-IMPLEMENT_SIMPLE_AUTOMATION_TEST(FMCPAccountDeletionTest,"OSS.MCP.Test Account Deletion", EAutomationTestFlags::ATF_Editor )
-
-bool FMCPAccountDeletionTest::RunTest(const FString& Parameters)
-{
-	ADD_LATENT_AUTOMATION_COMMAND(FEngineWaitLatentCommand(3.0));
-	ADD_LATENT_AUTOMATION_COMMAND(FExecStringLatentCommand(TEXT("online sub=mcp test deleteaccount [REDACTED_EMAIL] Epic3234")));
-
-	return true;
-}
-
 IMPLEMENT_SIMPLE_AUTOMATION_TEST(FMCPFriendsTest,"OSS.MCP.Test Friends", EAutomationTestFlags::ATF_Editor )
 
 bool FMCPFriendsTest::RunTest(const FString& Parameters)
