Fix NULL pointer exception when trying to open LiveEditor

diff --git a/Engine/Source/Runtime/Engine/Private/LevelTick.cpp b/Engine/Source/Runtime/Engine/Private/LevelTick.cpp
--- a/Engine/Source/Runtime/Engine/Private/LevelTick.cpp
+++ b/Engine/Source/Runtime/Engine/Private/LevelTick.cpp
@@ -1124,8 +1124,11 @@ void UWorld::Tick( ELevelTick TickType, float DeltaSeconds )
 		RunTickGroup(TG_StartPhysics); 
 		RunTickGroup(TG_DuringPhysics, false); // No wait here, we should run until idle though. We don't care if all of the async ticks are done before we start running post-phys stuff
 		TickGroup = TG_EndPhysics; // set this here so the current tick group is correct during collision notifies, though I am not sure it matters. 'cause of the false up there^^^
-		RunTickGroup(TG_EndPhysics); 
-		PhysicsScene->DeferredCommandHandler.Flush();
+		RunTickGroup(TG_EndPhysics);
+		if ( PhysicsScene != NULL )
+		{
+			PhysicsScene->DeferredCommandHandler.Flush();
+		}
 		RunTickGroup(TG_PreCloth);
 		RunTickGroup(TG_StartCloth);
 		RunTickGroup(TG_EndCloth);
