Settings: Fix for setting Project Settings as default clears redirects (TTP# 334046)

To fix the loss of redirects we added corresponding UProperties in UEngine.
To fix the corruption of the Default INI on shutdown we are now only saving to Default from the Settings UI.

diff --git a/Engine/Source/Developer/Settings/Private/SettingsSection.h b/Engine/Source/Developer/Settings/Private/SettingsSection.h
--- a/Engine/Source/Developer/Settings/Private/SettingsSection.h
+++ b/Engine/Source/Developer/Settings/Private/SettingsSection.h
@@ -228,7 +228,14 @@ public:
 
 		if (SettingsObject.IsValid())
 		{
-			SettingsObject->SaveConfig();
+			if (SettingsObject->GetClass()->HasAnyClassFlags(CLASS_DefaultConfig))
+			{
+				SettingsObject->UpdateDefaultConfigFile();
+			}
+			else
+			{
+				SettingsObject->SaveConfig();
+			}
 
 			return true;
 		}
diff --git a/Engine/Source/Runtime/CoreUObject/Private/UObject/Obj.cpp b/Engine/Source/Runtime/CoreUObject/Private/UObject/Obj.cpp
--- a/Engine/Source/Runtime/CoreUObject/Private/UObject/Obj.cpp
+++ b/Engine/Source/Runtime/CoreUObject/Private/UObject/Obj.cpp
@@ -1537,17 +1537,6 @@ void UObject::SaveConfig( uint64 Flags, const TCHAR* InFilename, FConfigCacheIni
 		return;
 	}
 
-	if (GetClass()->HasAnyClassFlags(CLASS_DefaultConfig))
-	{
-		// no INI file specified - save to corresponding Default INI
-		if ((InFilename == NULL) && (Config == GConfig))
-		{
-			UpdateDefaultConfigFile();
-
-			return;
-		}
-	}
-
 	uint32 PropagationFlags = UE4::LCPF_None;
 
 	const FString Filename
diff --git a/Engine/Source/Runtime/Engine/Classes/Engine/Engine.h b/Engine/Source/Runtime/Engine/Classes/Engine/Engine.h
--- a/Engine/Source/Runtime/Engine/Classes/Engine/Engine.h
+++ b/Engine/Source/Runtime/Engine/Classes/Engine/Engine.h
@@ -462,6 +462,67 @@ namespace EMatineeCaptureType
 	};
 }
 
+USTRUCT()
+struct FGameNameRedirect
+{
+	GENERATED_USTRUCT_BODY()
+
+	UPROPERTY()
+	FString OldGameName;
+
+	UPROPERTY()
+	FString NewGameName;
+};
+
+
+USTRUCT()
+struct FClassRedirect
+{
+	GENERATED_USTRUCT_BODY()
+
+	UPROPERTY()
+	FString ObjectName;
+
+	UPROPERTY()
+	FString OldClassName;
+
+	UPROPERTY()
+	FString NewClassName;
+
+	UPROPERTY()
+	FString OldSubobjName;
+
+	UPROPERTY()
+	FString NewSubobjName;
+
+	UPROPERTY()
+	bool InstanceOnly;
+};
+
+USTRUCT()
+struct FStructRedirect
+{
+	GENERATED_USTRUCT_BODY()
+
+	UPROPERTY()
+	FString OldStructName;
+
+	UPROPERTY()
+	FString NewStructName;
+};
+
+USTRUCT()
+struct FPluginRedirect
+{
+	GENERATED_USTRUCT_BODY()
+
+	UPROPERTY()
+	FString OldPluginName;
+
+	UPROPERTY()
+	FString NewPluginName;
+};
+
 class IAnalyticsProvider;
 
 DECLARE_DELEGATE_OneParam(FBeginStreamingPauseDelegate, FViewport*);
@@ -945,6 +1006,18 @@ public:
 	UPROPERTY(globalconfig)
 	FStringAssetReference DefaultPhysMaterialName;
 
+	UPROPERTY(config)
+	TArray<FGameNameRedirect> ActiveGameNameRedirects;
+
+	UPROPERTY(config)
+	TArray<FClassRedirect> ActiveClassRedirects;
+
+	UPROPERTY(config)
+	TArray<FPluginRedirect> ActivePluginRedirects;
+
+	UPROPERTY(config)
+	TArray<FStructRedirect> ActiveStructRedirects;
+
 	/** Texture used for pre-integrated skin shading */
 	UPROPERTY()
 	class UTexture2D* PreIntegratedSkinBRDFTexture;
