Fix vehicle ignoring bUseGearAutoBox

diff --git a/Engine/Source/Runtime/Engine/Private/Vehicles/WheeledVehicleMovementComponent4W.cpp b/Engine/Source/Runtime/Engine/Private/Vehicles/WheeledVehicleMovementComponent4W.cpp
--- a/Engine/Source/Runtime/Engine/Private/Vehicles/WheeledVehicleMovementComponent4W.cpp
+++ b/Engine/Source/Runtime/Engine/Private/Vehicles/WheeledVehicleMovementComponent4W.cpp
@@ -278,10 +278,7 @@ void UWheeledVehicleMovementComponent4W::SetupVehicle()
 	PVehicle = PVehicleDrive4W;
 	PVehicleDrive = PVehicleDrive4W;
 
-	// MSS Using PxVehicle built-in automatic transmission for now. No reverse yet, have to fix that.
-	//SetUseAutoGears(bUseGearAutoBox);
-	SetUseAutoGears(true);
-	//SetTargetGear(1, true);
+	SetUseAutoGears(TransmissionSetup.bUseGearAutoBox);
 }
 
 void UWheeledVehicleMovementComponent4W::UpdateSimulation(float DeltaTime)
