Expose the HUD Deproject method to Blueprints
#ttp 332113 Blueprints: Can someone add Blueprint support for Deproject
reviewer Thomas.Sarkanen

diff --git a/Engine/Source/Runtime/Engine/Classes/GameFramework/HUD.h b/Engine/Source/Runtime/Engine/Classes/GameFramework/HUD.h
--- a/Engine/Source/Runtime/Engine/Classes/GameFramework/HUD.h
+++ b/Engine/Source/Runtime/Engine/Classes/GameFramework/HUD.h
@@ -398,9 +398,9 @@ public:
 	FVector Project(FVector Location);
 	
 	/** Transforms a 2D screen location into a 3D location and direction */
+	UFUNCTION(BlueprintPure, Category=HUD)
 	void Deproject(float ScreenX, float ScreenY, FVector& WorldPosition, FVector& WorldDirection);
 
-
 	/**
 	 * Add a hitbox to the hud
 	 * @param Position			Coordinates of the top left of the hit box.
