If we're aren't playing-in-editor then don't clear the CachedMousePos when the mouse leaves as this can be triggered by a right-click.
#ttp 333924 CHECKIN: CRITICAL: Regression: EDITOR: Place Actor location incorrect after deleting SkySphere (fix-fail)
reviewer Thomas.Sarkanen

diff --git a/Engine/Source/Runtime/Engine/Private/Slate/SceneViewport.cpp b/Engine/Source/Runtime/Engine/Private/Slate/SceneViewport.cpp
--- a/Engine/Source/Runtime/Engine/Private/Slate/SceneViewport.cpp
+++ b/Engine/Source/Runtime/Engine/Private/Slate/SceneViewport.cpp
@@ -387,7 +387,10 @@ void FSceneViewport::OnMouseLeave( const FPointerEvent& MouseEvent )
 {
 	ViewportClient->MouseLeave( this );
 	
-	CachedMousePos = FIntPoint(-1, -1);
+	if ( IsPlayInEditorViewport() )
+	{
+		CachedMousePos = FIntPoint(-1, -1);
+	}
 }
 
 FReply FSceneViewport::OnMouseMove( const FGeometry& InGeometry, const FPointerEvent& InMouseEvent )
