Speculative fix for unrepeatable crash bug - guard against overrunning the CachedData array.
#ttp 334315 LIVE: CASCADE: Crash when marqee selecting in the Curve Editor section of the Cascade Editor
reviewer Andrew.Brown

diff --git a/Engine/Source/Runtime/Engine/Private/UnrealClient.cpp b/Engine/Source/Runtime/Engine/Private/UnrealClient.cpp
--- a/Engine/Source/Runtime/Engine/Private/UnrealClient.cpp
+++ b/Engine/Source/Runtime/Engine/Private/UnrealClient.cpp
@@ -1231,11 +1231,20 @@ const TArray<FColor>& FViewport::GetRawHitProxyData(FIntRect InRect)
 void FViewport::GetHitProxyMap(FIntRect InRect,TArray<HHitProxy*>& OutMap)
 {
 	const TArray<FColor>& CachedData = GetRawHitProxyData(InRect);
-
+	if (CachedData.Num()==0)
+	{
+		return;
+	}
+	
 	// Map the hit proxy map surface data to hit proxies.
 	OutMap.Empty(InRect.Width() * InRect.Height());
 	for(int32 Y = InRect.Min.Y; Y < InRect.Max.Y; Y++)
 	{
+		if (!CachedData.IsValidIndex(int32(Y * SizeX)))
+		{
+			break;
+		}
+
 		const FColor* SourceData = &CachedData[Y * SizeX];
 		for(int32 X = InRect.Min.X ;X < InRect.Max.X; X++)
 		{
