* fixed optimization that got disabled for LPV to be working as before if LPV is off (faster shadow map rendering)

diff --git a/Engine/Source/Runtime/Engine/Private/StaticMeshRender.cpp b/Engine/Source/Runtime/Engine/Private/StaticMeshRender.cpp
--- a/Engine/Source/Runtime/Engine/Private/StaticMeshRender.cpp
+++ b/Engine/Source/Runtime/Engine/Private/StaticMeshRender.cpp
@@ -338,6 +338,29 @@ HHitProxy* FStaticMeshSceneProxy::CreateHitProxies(UPrimitiveComponent* Componen
 }
 #endif // WITH_EDITOR
 
+// use for render thread only
+bool UseLightPropagationVolumeRT()
+{
+	if(!IsFeatureLevelSupported(GRHIShaderPlatform, ERHIFeatureLevel::SM5))
+	{
+		return false;
+	}
+
+	// todo: better we get the engine LPV state not the cvar (later we want to make it changeable at runtime)
+	static const auto* CVar = IConsoleManager::Get().FindTConsoleVariableDataInt(TEXT("r.LightPropagationVolume"));
+	check(CVar);
+
+	int32 Value = CVar->GetValueOnRenderThread();
+
+	return Value != 0;
+}
+
+inline bool AllowShadowOnlyMesh()
+{
+	// todo: later we should refine that (only if occlusion feature in LPV is on, only if inside a cascade, if shadow casting is disabled it should look at bUseEmissiveForDynamicAreaLighting)
+	return !UseLightPropagationVolumeRT();
+}
+
 void FStaticMeshSceneProxy::DrawStaticElements(FStaticPrimitiveDrawInterface* PDI)
 {
 	checkSlow(IsInRenderingThread());
@@ -391,7 +414,9 @@ void FStaticMeshSceneProxy::DrawStaticElements(FStaticPrimitiveDrawInterface* PD
 					const FLODInfo& ProxyLODInfo = LODs[LODIndex];
 
 					// The shadow-only mesh can be used only if all elements cast shadows and use opaque materials with no vertex modification.
-					bool bSafeToUseShadowOnlyMesh = true;
+					// In some cases (e.g. LPV) we don't want the optimization
+					bool bSafeToUseShadowOnlyMesh = AllowShadowOnlyMesh();
+
 					bool bAnySectionCastsShadow = false;
 					for (int32 SectionIndex = 0; bSafeToUseShadowOnlyMesh && SectionIndex < LODModel.Sections.Num(); SectionIndex++)
 					{
@@ -405,8 +430,6 @@ void FStaticMeshSceneProxy::DrawStaticElements(FStaticPrimitiveDrawInterface* PD
 						bAnySectionCastsShadow |= Section.bCastShadow;
 					}
 
-					bSafeToUseShadowOnlyMesh = false;
-
 					if (bAnySectionCastsShadow && bSafeToUseShadowOnlyMesh)
 					{
 						FMeshBatch MeshElement;
