Engine: Split non-EdGraph.h types out of EdGraph.cpp into individual files

diff --git a/Engine/Source/Runtime/Engine/Private/EdGraph/EdGraph.cpp b/Engine/Source/Runtime/Engine/Private/EdGraph/EdGraph.cpp
--- a/Engine/Source/Runtime/Engine/Private/EdGraph/EdGraph.cpp
+++ b/Engine/Source/Runtime/Engine/Private/EdGraph/EdGraph.cpp
@@ -1,9 +1,5 @@
 // Copyright 1998-2014 Epic Games, Inc. All Rights Reserved.
 
-/*=============================================================================
-	EdGraph.cpp
-=============================================================================*/
-
 #include "EnginePrivate.h"
 #include "EngineClasses.h"
 #include "BlueprintUtilities.h"
@@ -18,136 +14,6 @@
 
 DEFINE_LOG_CATEGORY_STATIC(LogEdGraph, Log, All);
 
-/////////////////////////////////////////////////////
-// FGraphActionListBuilderBase
-
-void FGraphActionListBuilderBase::AddAction( const TSharedPtr<FEdGraphSchemaAction>& NewAction, FString const& Category/* = TEXT("") */ )
-{
-	Entries.Add( ActionGroup( NewAction, Category ) );
-}
-
-void FGraphActionListBuilderBase::AddActionList( const TArray<TSharedPtr<FEdGraphSchemaAction> >& NewActions, FString const& Category/* = TEXT("") */ )
-{
-	Entries.Add( ActionGroup( NewActions, Category ) );
-}
-
-void FGraphActionListBuilderBase::Append( FGraphActionListBuilderBase& Other )
-{
-	Entries.Append( Other.Entries );
-}
-
-int32 FGraphActionListBuilderBase::GetNumActions() const
-{
-	return Entries.Num();
-}
-
-FGraphActionListBuilderBase::ActionGroup& FGraphActionListBuilderBase::GetAction( int32 Index )
-{
-	return Entries[Index];
-}
-
-void FGraphActionListBuilderBase::Empty()
-{
-	Entries.Empty();
-}
-
-/////////////////////////////////////////////////////
-// FGraphActionListBuilderBase::GraphAction
-
-FGraphActionListBuilderBase::ActionGroup::ActionGroup( TSharedPtr<FEdGraphSchemaAction> InAction, FString const& CategoryPrefix/* = TEXT("") */ )
-	: RootCategory(CategoryPrefix)
-{
-	Actions.Add( InAction );
-}
-
-FGraphActionListBuilderBase::ActionGroup::ActionGroup( const TArray< TSharedPtr<FEdGraphSchemaAction> >& InActions, FString const& CategoryPrefix/* = TEXT("") */ )
-	: RootCategory(CategoryPrefix)
-{
-	Actions = InActions;
-}
-
-void FGraphActionListBuilderBase::ActionGroup::GetCategoryChain(TArray<FString>& HierarchyOut) const
-{
-	static FString const CategoryDelim("|");
-	RootCategory.ParseIntoArray(&HierarchyOut, *CategoryDelim, true);
-
-	if (Actions.Num() > 0)
-	{
-		TArray<FString> SubCategoryChain;
-		Actions[0]->Category.ParseIntoArray(&SubCategoryChain, *CategoryDelim, true);
-
-		HierarchyOut.Append(SubCategoryChain);
-	}
-
-	for (FString& Category : HierarchyOut)
-	{
-		Category.Trim();
-	}
-}
-
-void FGraphActionListBuilderBase::ActionGroup::PerformAction( class UEdGraph* ParentGraph, TArray<UEdGraphPin*>& FromPins, const FVector2D Location )
-{
-	for ( int32 ActionIndex = 0; ActionIndex < Actions.Num(); ActionIndex++ )
-	{
-		TSharedPtr<FEdGraphSchemaAction> CurrentAction = Actions[ActionIndex];
-		if ( CurrentAction.IsValid() )
-		{
-			CurrentAction->PerformAction( ParentGraph, FromPins, Location );
-		}
-	}
-}
-
-/////////////////////////////////////////////////////
-// FCategorizedGraphActionListBuilder
-
-static FString ConcatCategories(FString const& RootCategory, FString const& SubCategory)
-{
-	FString ConcatedCategory = RootCategory;
-	if (!SubCategory.IsEmpty() && !RootCategory.IsEmpty())
-	{
-		ConcatedCategory += TEXT("|");
-	}
-	ConcatedCategory += SubCategory;
-
-	return ConcatedCategory;
-}
-
-FCategorizedGraphActionListBuilder::FCategorizedGraphActionListBuilder(FString const& CategoryIn/* = TEXT("")*/)
-	: Category(CategoryIn)
-{
-}
-
-void FCategorizedGraphActionListBuilder::AddAction(TSharedPtr<FEdGraphSchemaAction> const& NewAction, FString const& CategoryIn/* = TEXT("")*/)
-{
-	FGraphActionListBuilderBase::AddAction(NewAction, ConcatCategories(Category, CategoryIn));
-}
-
-void FCategorizedGraphActionListBuilder::AddActionList(TArray<TSharedPtr<FEdGraphSchemaAction> > const& NewActions, FString const& CategoryIn/* = TEXT("")*/)
-{
-	FGraphActionListBuilderBase::AddActionList(NewActions, ConcatCategories(Category, CategoryIn));
-}
-
-/////////////////////////////////////////////////////
-// FGraphNodeContextMenuBuilder
-
-FGraphNodeContextMenuBuilder::FGraphNodeContextMenuBuilder(const UEdGraph* InGraph, const UEdGraphNode* InNode, const UEdGraphPin* InPin, FMenuBuilder* InMenuBuilder, bool bInDebuggingMode)
-	: Blueprint(NULL)
-	, Graph(InGraph)
-	, Node(InNode)
-	, Pin(InPin)
-	, MenuBuilder(InMenuBuilder)
-	, bIsDebugging(bInDebuggingMode)
-{
-#if WITH_EDITOR
-	Blueprint = FBlueprintEditorUtils::FindBlueprintForGraph(Graph);
-#endif
-
-	if (Pin != NULL)
-	{
-		Node = Pin->GetOwningNode();
-	}
-}
-
 /////////////////////////////////////////////////////
 // FGraphReference
 
@@ -304,7 +170,6 @@ void UEdGraph::NotifyGraphChanged(const FEdGraphEditAction& InAction)
 	OnGraphChanged.Broadcast(InAction);
 }
 
-
 void UEdGraph::MoveNodesToAnotherGraph(UEdGraph* DestinationGraph, bool bIsLoading)
 {
 	// Move one node over at a time
@@ -388,1148 +253,5 @@ void UEdGraph::RemovePropertyChangedNotifier(const FOnPropertyChanged::FDelegate
 #endif	//WITH_EDITOR
 
 /////////////////////////////////////////////////////
-// UEdGraphNode
-
-UEdGraphNode::UEdGraphNode(const class FPostConstructInitializeProperties& PCIP)
-	: Super(PCIP)
-	, AdvancedPinDisplay(ENodeAdvancedPins::NoPins)
-{
-
-#if WITH_EDITORONLY_DATA
-	bCanResizeNode = false;
-#endif // WITH_EDITORONLY_DATA
-}
-
-#if WITH_EDITOR
-UEdGraphPin* UEdGraphNode::CreatePin(EEdGraphPinDirection Dir, const FString& PinCategory, const FString& PinSubCategory, UObject* PinSubCategoryObject, bool bIsArray, bool bIsReference, const FString& PinName, bool bIsConst /*= false*/, int32 Index /*= INDEX_NONE*/)
-{
-#if 0
-	UEdGraphPin* NewPin = AllocatePinFromPool(this);
-#else
-	UEdGraphPin* NewPin = NewObject<UEdGraphPin>(this);
-#endif
-	NewPin->PinName = PinName;
-	NewPin->Direction = Dir;
-	NewPin->PinType.PinCategory = PinCategory;
-	NewPin->PinType.PinSubCategory = PinSubCategory;
-	NewPin->PinType.PinSubCategoryObject = PinSubCategoryObject;
-	NewPin->PinType.bIsArray = bIsArray;
-	NewPin->PinType.bIsReference = bIsReference;
-	NewPin->PinType.bIsConst = bIsConst;
-	NewPin->SetFlags(RF_Transactional);
-
-	Modify(false);
-	if ( Pins.IsValidIndex( Index ) )
-	{
-		Pins.Insert(NewPin, Index);
-	}
-	else
-	{
-		Pins.Add(NewPin);
-	}
-	return NewPin;
-}
-
-UEdGraphPin* UEdGraphNode::FindPin(const FString& PinName) const
-{
-	for(int32 PinIdx=0; PinIdx<Pins.Num(); PinIdx++)
-	{
-		if( Pins[PinIdx]->PinName == PinName )
-		{
-			return Pins[PinIdx];
-		}
-	}
-
-	return NULL;
-}
-
-UEdGraphPin* UEdGraphNode::FindPinChecked(const FString& PinName) const
-{
-	UEdGraphPin* Result = FindPin(PinName);
-	check(Result != NULL);
-	return Result;
-}
-
-void UEdGraphNode::DiscardPin(UEdGraphPin* Pin)
-{
-	check( Pin );
-
-	Modify();
-	Pins.Remove( Pin );
-}
-
-void UEdGraphNode::BreakAllNodeLinks()
-{
-	TSet<UEdGraphNode*> NodeList;
-
-	NodeList.Add(this);
-
-	// Iterate over each pin and break all links
-	for(int32 PinIdx=0; PinIdx<Pins.Num(); PinIdx++)
-	{
-		Pins[PinIdx]->BreakAllPinLinks();
-		NodeList.Add(Pins[PinIdx]->GetOwningNode());
-	}
-
-	// Send all nodes that received a new pin connection a notification
-	for (auto It = NodeList.CreateConstIterator(); It; ++It)
-	{
-		UEdGraphNode* Node = (*It);
-		Node->NodeConnectionListChanged();
-	}
-}
-
-void UEdGraphNode::GetPinHoverText(const UEdGraphPin& Pin, FString& HoverTextOut) const
-{
-	ensure(Pin.GetOwningNode() == this);
-	HoverTextOut = Pin.PinToolTip;
-}
-
-void UEdGraphNode::SnapToGrid(float GridSnapSize)
-{
-	NodePosX = GridSnapSize * FMath::RoundToInt(NodePosX/GridSnapSize);
-	NodePosY = GridSnapSize * FMath::RoundToInt(NodePosY/GridSnapSize);
-}
-
-void UEdGraphNode::DestroyNode()
-{
-	UEdGraph* ParentGraph = GetGraph();
-	check(ParentGraph);
-
-	// Remove the node - this will break all links. Will be GC'd after this.
-	ParentGraph->RemoveNode(this);
-}
-
-const class UEdGraphSchema* UEdGraphNode::GetSchema() const
-{
-	return GetGraph()->GetSchema();
-}
-
-FLinearColor UEdGraphNode::GetNodeTitleColor() const
-{
-	return FLinearColor(0.4f, 0.62f, 1.0f);
-}
-
-FLinearColor UEdGraphNode::GetNodeCommentColor() const
-{
-	return FLinearColor::White;
-}
-
-FString UEdGraphNode::GetTooltip() const
-{
-	return GetClass()->GetToolTipText().ToString();
-}
-
-FString UEdGraphNode::GetDocumentationExcerptName() const
-{
-	// Default the node to searching for an excerpt named for the C++ node class name, including the U prefix.
-	// This is done so that the excerpt name in the doc file can be found by find-in-files when searching for the full class name.
-	UClass* MyClass = GetClass();
-	return FString::Printf(TEXT("%s%s"), MyClass->GetPrefixCPP(), *MyClass->GetName());
-}
-
-
-FString UEdGraphNode::GetDescriptiveCompiledName() const
-{
-	return GetFName().GetPlainNameString();
-}
-
-bool UEdGraphNode::IsDeprecated() const
-{
-	return GetClass()->HasAnyClassFlags(CLASS_Deprecated);
-}
-
-FString UEdGraphNode::GetDeprecationMessage() const
-{
-	return NSLOCTEXT("EdGraphCompiler", "NodeDeprecated_Warning", "@@ is deprecated; please replace or remove it.").ToString();
-}
-
-// Array of pooled pins
-TArray<UEdGraphPin*> UEdGraphNode::PooledPins;
-
-void UEdGraphNode::AddReferencedObjects(UObject* InThis, FReferenceCollector& Collector)
-{
-	UEdGraphNode* This = CastChecked<UEdGraphNode>(InThis);	
-
-	// Only register the pool once per GC pass
-	if (This->HasAnyFlags(RF_ClassDefaultObject))
-	{
-		if (This->GetClass() == UEdGraphNode::StaticClass())
-		{
-			for (int32 Index = 0; Index < PooledPins.Num(); ++Index)
-			{
-				Collector.AddReferencedObject(PooledPins[Index], This);
-			}
-		}
-	}
-	Super::AddReferencedObjects(This, Collector);
-}
-
-void UEdGraphNode::PostLoad()
-{
-	Super::PostLoad();
-
-	// Create Guid if not present (and not CDO)
-	if(!NodeGuid.IsValid() && !IsTemplate() && GetLinker() && GetLinker()->IsPersistent() && GetLinker()->IsLoading())
-	{
-		// _Should_ have a guid on all nodes after this version
-		if(GetLinkerUE4Version() >= VER_UE4_ADD_EDGRAPHNODE_GUID)
-		{
-			UE_LOG(LogBlueprint, Warning, TEXT("Node '%s' missing NodeGuid."), *GetPathName());
-		}
-
-		// Generate new one
-		CreateNewGuid();
-	}
-}
-
-void UEdGraphNode::CreateNewGuid()
-{
-	NodeGuid = FGuid::NewGuid();
-}
-
-void UEdGraphNode::FindDiffs( class UEdGraphNode* OtherNode, struct FDiffResults& Results ) 
-{
-}
-
-UEdGraphPin* UEdGraphNode::AllocatePinFromPool(UEdGraphNode* OuterNode)
-{
-	if (PooledPins.Num() > 0)
-	{
-		UEdGraphPin* Result = PooledPins.Pop();
-		Result->Rename(NULL, OuterNode);
-		return Result;
-	}
-	else
-	{
-		UEdGraphPin* Result = NewObject<UEdGraphPin>(OuterNode);
-		return Result;
-	}
-}
-
-void UEdGraphNode::ReturnPinToPool(UEdGraphPin* OldPin)
-{
-	check(OldPin);
-
-	check(!OldPin->HasAnyFlags(RF_NeedLoad));
-	OldPin->ClearFlags(RF_NeedPostLoadSubobjects|RF_NeedPostLoad);
-
-	OldPin->ResetToDefaults();
-
-	PooledPins.Add(OldPin);
-}
-
-bool UEdGraphNode::CanDuplicateNode() const
-{
-	return true;
-}
-
-bool UEdGraphNode::CanUserDeleteNode() const
-{
-	return true;
-}
-
-FText UEdGraphNode::GetNodeTitle(ENodeTitleType::Type TitleType) const
-{
-	return FText::FromString(GetClass()->GetName());
-}
-
-FString UEdGraphNode::GetNodeNativeTitle(ENodeTitleType::Type TitleType) const
-{
-	FText NodeTitle = GetNodeTitle(TitleType);
-	if(const FString* SourceString = FTextInspector::GetSourceString(NodeTitle))
-	{
-		return *SourceString;
-	}
-	return NodeTitle.ToString();
-}
-
-FText UEdGraphNode::GetNodeSearchTitle() const
-{
-	TArray< FText > SearchMetadata;
-
-	FText LocalizedTitle = GetNodeTitle(ENodeTitleType::ListView);
-	FText NativeTitle = FText::FromString(GetNodeNativeTitle(ENodeTitleType::ListView));
-
-	if(FInternationalization::Get().GetCurrentCulture()->GetName() == TEXT("en_US"))
-	{
-		ensure( LocalizedTitle.CompareTo(NativeTitle) == 0 );
-	}
-
-	FFormatNamedArguments Args;
-	Args.Add(TEXT("LocalizedTitle"), LocalizedTitle);
-	Args.Add(TEXT("NativeTitle"), NativeTitle);
-	return FText::Format(LOCTEXT("NodeSearchTitle", "{LocalizedTitle} {NativeTitle}"), Args);
-}
-
-UObject* UEdGraphNode::GetJumpTargetForDoubleClick() const
-{
-	return NULL;
-}
-
-FString UEdGraphNode::GetPinDisplayName(const UEdGraphPin* Pin) const
-{
-	return GetSchema()->GetPinDisplayName(Pin);
-}
-
-int32 UEdGraphNode::GetPinIndex(UEdGraphPin* Pin) const
-{
-	return Pins.Find(Pin);
-}
-
-#endif	//#if WITH_EDITOR
-
-
-/////////////////////////////////////////////////////
-// UEdGraphNode_Comment
-
-UEdGraphNode_Comment::UEdGraphNode_Comment(const class FPostConstructInitializeProperties& PCIP)
-	: Super(PCIP)
-{
-	NodeWidth = 400.0f;
-	NodeHeight = 100.0f;
-	CommentColor = FLinearColor::White;
-	bColorCommentBubble = false;
-	MoveMode = ECommentBoxMode::GroupMovement;
-
-#if WITH_EDITORONLY_DATA
-	bCanResizeNode = true;
-	bCanRenameNode = true;
-#endif // WITH_EDITORONLY_DATA
-}
-
-void UEdGraphNode_Comment::AddReferencedObjects(UObject* InThis, FReferenceCollector& Collector) 
-{
-#if WITH_EDITOR
-	UEdGraphNode_Comment* This = CastChecked<UEdGraphNode_Comment>(InThis);
-	for (auto It = This->NodesUnderComment.CreateIterator(); It; ++It)
-	{
-		Collector.AddReferencedObject(*It, This);
-	}
-#endif // WITH_EDITOR
-
-	Super::AddReferencedObjects(InThis, Collector);
-}
-
-#if WITH_EDITOR
-void UEdGraphNode_Comment::PostPlacedNewNode()
-{
-	//@TODO: Consider making the default value we use here a preference
-	// This is done here instead of in the constructor so we can later change the default for newly placed
-	// instances without changing all of the existing ones (due to delta serialization)
-	MoveMode = ECommentBoxMode::GroupMovement;
-
-	NodeComment = NSLOCTEXT("K2Node", "CommentBlock_NewEmptyComment", "Comment").ToString();
-}
-
-FString UEdGraphNode_Comment::GetTooltip() const
-{
-	return FString::Printf(*NSLOCTEXT("K2Node", "CommentBlock_Tooltip", "Comment:\n%s").ToString(), *NodeComment);
-}
-
-FString UEdGraphNode_Comment::GetDocumentationLink() const
-{
-	return TEXT("Shared/GraphNodes/Common");
-}
-
-FString UEdGraphNode_Comment::GetDocumentationExcerptName() const
-{
-	return TEXT("UEdGraphNode_Comment");
-}
-
-FText UEdGraphNode_Comment::GetNodeTitle(ENodeTitleType::Type TitleType) const
-{
-	if(TitleType == ENodeTitleType::ListView)
-	{
-		return NSLOCTEXT("K2Node", "CommentBlock_Title", "Comment");
-	}
-
-	return FText::FromString(NodeComment);
-}
-
-FString UEdGraphNode_Comment::GetPinNameOverride(const UEdGraphPin& Pin) const
-{
-	return GetNodeTitle(ENodeTitleType::ListView).ToString();
-}
-
-FLinearColor UEdGraphNode_Comment::GetNodeCommentColor() const
-{
-	// Only affects the 'zoomed out' comment bubble color, not the box itself
-	return (bColorCommentBubble)
-		? CommentColor 
-		: FLinearColor::White;
-}
-
-void UEdGraphNode_Comment::ResizeNode(const FVector2D& NewSize)
-{
-	if (bCanResizeNode) 
-	{
-		NodeHeight = NewSize.Y;
-		NodeWidth = NewSize.X;
-	}
-}
-
-void UEdGraphNode_Comment::AddNodeUnderComment(UObject* Object)
-{
-	NodesUnderComment.Add(Object);
-}
-
-void UEdGraphNode_Comment::ClearNodesUnderComment()
-{
-	NodesUnderComment.Empty();
-}
-
-void UEdGraphNode_Comment::SetBounds(const class FSlateRect& Rect)
-{
-	NodePosX = Rect.Left;
-	NodePosY = Rect.Top;
-
-	FVector2D Size = Rect.GetSize();
-	NodeWidth = Size.X;
-	NodeHeight = Size.Y;
-}
-
-const FCommentNodeSet& UEdGraphNode_Comment::GetNodesUnderComment() const
-{
-	return NodesUnderComment;
-}
-
-void UEdGraphNode_Comment::OnRenameNode(const FString& NewName)
-{
-	NodeComment = NewName;
-}
-
-TSharedPtr<class INameValidatorInterface> UEdGraphNode_Comment::MakeNameValidator() const
-{
-	// Comments can be duplicated, etc...
-	return MakeShareable(new FDummyNameValidator(EValidatorResult::Ok));
-}
-#endif // WITH_EDITOR
-
-
-/////////////////////////////////////////////////////
-// UEdGraphPin
-
-UEdGraphPin::UEdGraphPin(const class FPostConstructInitializeProperties& PCIP)
-	: Super(PCIP)
-{
-#if WITH_EDITORONLY_DATA
-	bHidden = false;
-	bNotConnectable = false;
-	bAdvancedView = false;
-#endif // WITH_EDITORONLY_DATA
-}
-
-void UEdGraphPin::MakeLinkTo(UEdGraphPin* ToPin)
-{
-	Modify();
-
-	if (ToPin != NULL)
-	{
-		ToPin->Modify();
-
-		// Make sure we don't already link to it
-		if (!LinkedTo.Contains(ToPin))
-		{
-			UEdGraphNode* MyNode = GetOwningNode();
-
-			// Check that the other pin does not link to us
-			ensureMsg(!ToPin->LinkedTo.Contains(this), *GetLinkInfoString( LOCTEXT("MakeLinkTo", "MakeLinkTo").ToString(), LOCTEXT("IsLinked", "is linked with pin").ToString(), ToPin));			    
-			ensureMsg(MyNode->GetOuter() == ToPin->GetOwningNode()->GetOuter(), *GetLinkInfoString( LOCTEXT("MakeLinkTo", "MakeLinkTo").ToString(), LOCTEXT("OuterMismatch", "has a different outer than pin").ToString(), ToPin)); // Ensure both pins belong to the same graph
-
-			// Add to both lists
-			LinkedTo.Add(ToPin);
-			ToPin->LinkedTo.Add(this);
-		}
-	}
-}
-
-void UEdGraphPin::BreakLinkTo(UEdGraphPin* ToPin)
-{
-	Modify();
-
-	if (ToPin != NULL)
-	{
-		ToPin->Modify();
-
-		// If we do indeed link to the passed in pin...
-		if (LinkedTo.Contains(ToPin))
-		{
-			LinkedTo.Remove(ToPin);
-
-			// Check that the other pin links to us
-			ensureMsg(ToPin->LinkedTo.Contains(this), *GetLinkInfoString( LOCTEXT("BreakLinkTo", "BreakLinkTo").ToString(), LOCTEXT("NotLinked", "not reciprocally linked with pin").ToString(), ToPin) );
-			ToPin->LinkedTo.Remove(this);
-		}
-		else
-		{
-			// Check that the other pin does not link to us
-			ensureMsg(!ToPin->LinkedTo.Contains(this), *GetLinkInfoString( LOCTEXT("MakeLinkTo", "MakeLinkTo").ToString(), LOCTEXT("IsLinked", "is linked with pin").ToString(), ToPin));
-		}
-	}
-}
-
-void UEdGraphPin::BreakAllPinLinks()
-{
-	TArray<UEdGraphPin*> LinkedToCopy = LinkedTo;
-
-	for (int32 LinkIdx = 0; LinkIdx < LinkedToCopy.Num(); LinkIdx++)
-	{
-		BreakLinkTo(LinkedToCopy[LinkIdx]);
-	}
-}
-
-
-void UEdGraphPin::CopyPersistentDataFromOldPin(const UEdGraphPin& SourcePin)
-{
-	// The name matches already, doesn't get copied here
-	// The PinType, Direction, and bNotConnectable are properties generated from the schema
-
-	// Only move the default value if it was modified; inherit the new default value otherwise
-	if (SourcePin.DefaultValue != SourcePin.AutogeneratedDefaultValue || SourcePin.DefaultObject != NULL)
-	{
-		DefaultObject = SourcePin.DefaultObject;
-		DefaultValue = SourcePin.DefaultValue;
-	}
-
-	const bool bTextValueWasModified = (SourcePin.DefaultTextValue.ToString() != SourcePin.AutogeneratedDefaultValue);
-	if (bTextValueWasModified)
-	{
-		DefaultTextValue = SourcePin.DefaultTextValue;
-	}
-
-	// Copy the links
-	for (int32 LinkIndex = 0; LinkIndex < SourcePin.LinkedTo.Num(); ++LinkIndex)
-	{
-		UEdGraphPin* OtherPin = SourcePin.LinkedTo[LinkIndex];
-		check(NULL != OtherPin);
-		
-		Modify();
-		OtherPin->Modify();
-
-		OtherPin->MakeLinkTo(this);
-	}
-
-#if WITH_EDITORONLY_DATA
-	// Copy advanced visibility property, since it can be changed by user
-	bAdvancedView = SourcePin.bAdvancedView;
-#endif // WITH_EDITORONLY_DATA
-}
-
-const class UEdGraphSchema* UEdGraphPin::GetSchema() const
-{
-#if WITH_EDITOR
-	return GetOwningNode()->GetGraph()->GetSchema();
-#else
-	return NULL;
-#endif	//WITH_EDITOR
-}
-
-FString UEdGraphPin::GetDefaultAsString() const
-{
-	if(DefaultObject != NULL)
-	{
-		return DefaultObject->GetFullName();
-	}
-	else
-	{
-		return DefaultValue;
-	}
-}
-
-const FString UEdGraphPin::GetLinkInfoString( const FString& InFunctionName, const FString& InInfoData, const UEdGraphPin* InToPin ) const
-{
-#if WITH_EDITOR
-	const FString FromPinName = PinName;
-	const UEdGraphNode* FromPinNode = Cast<UEdGraphNode>(GetOuter());
-	const FString FromPinNodeName = (FromPinNode != NULL) ? FromPinNode->GetNodeTitle(ENodeTitleType::ListView).ToString() : TEXT("Unknown");
-	
-	const FString ToPinName = InToPin->PinName;
-	const UEdGraphNode* ToPinNode = Cast<UEdGraphNode>(InToPin->GetOuter());
-	const FString ToPinNodeName = (ToPinNode != NULL) ? ToPinNode->GetNodeTitle(ENodeTitleType::ListView).ToString() : TEXT("Unknown");
-	const FString LinkInfo = FString::Printf( TEXT("UEdGraphPin::%s Pin '%s' on node '%s' %s '%s' on node '%s'"), *InFunctionName, *ToPinName, *ToPinNodeName, *InInfoData, *FromPinName, *FromPinNodeName);
-	return LinkInfo;
-#else
-	return FString();
-#endif
-}
-
-FString FGraphDisplayInfo::GetNotesAsString() const
-{
-	FString Result;
-
-	if (Notes.Num() > 0)
-	{
-		Result = TEXT("(");
-		for (int32 NoteIndex = 0; NoteIndex < Notes.Num(); ++NoteIndex)
-		{
-			if (NoteIndex > 0)
-			{
-				Result += TEXT(", ");
-			}
-			Result += Notes[NoteIndex];
-		}
-		Result += TEXT(")");
-	}
-
-	return Result;
-}
-
-/////////////////////////////////////////////////////
-// FEdGraphSchemaAction_NewNode
-
-#define SNAP_GRID (16) // @todo ensure this is the same as SNodePanel::GetSnapGridSize()
-
-namespace 
-{
-	// Maximum distance a drag can be off a node edge to require 'push off' from node
-	const int32 NodeDistance = 60;
-}
-
-UEdGraphNode* FEdGraphSchemaAction_NewNode::CreateNode(class UEdGraph* ParentGraph, UEdGraphPin* FromPin, const FVector2D Location, class UEdGraphNode* NodeTemplate)
-{
-	// Duplicate template node to create new node
-	UEdGraphNode* ResultNode = NULL;
-
-#if WITH_EDITOR
-	ResultNode = DuplicateObject<UEdGraphNode>(NodeTemplate, ParentGraph);
-	ResultNode->SetFlags(RF_Transactional);
-
-	ParentGraph->AddNode(ResultNode, true);
-
-	ResultNode->CreateNewGuid();
-	ResultNode->PostPlacedNewNode();
-	ResultNode->AllocateDefaultPins();
-	ResultNode->AutowireNewNode(FromPin);
-
-	// For input pins, new node will generally overlap node being dragged off
-	// Work out if we want to visually push away from connected node
-	int32 XLocation = Location.X;
-	if (FromPin && FromPin->Direction == EGPD_Input)
-	{
-		UEdGraphNode* PinNode = FromPin->GetOwningNode();
-		const float XDelta = FMath::Abs(PinNode->NodePosX - Location.X);
-
-		if (XDelta < NodeDistance)
-		{
-			// Set location to edge of current node minus the max move distance
-			// to force node to push off from connect node enough to give selection handle
-			XLocation = PinNode->NodePosX - NodeDistance;
-		}
-	}
-
-	ResultNode->NodePosX = XLocation;
-	ResultNode->NodePosY = Location.Y;
-	ResultNode->SnapToGrid(SNAP_GRID);
-#endif // WITH_EDITOR
-
-	return ResultNode;
-}
-
-UEdGraphNode* FEdGraphSchemaAction_NewNode::PerformAction(class UEdGraph* ParentGraph, UEdGraphPin* FromPin, const FVector2D Location, bool bSelectNewNode/* = true*/)
-{
-	UEdGraphNode* ResultNode = NULL;
-
-#if WITH_EDITOR
-	// If there is a template, we actually use it
-	if (NodeTemplate != NULL)
-	{
-		const FScopedTransaction Transaction(LOCTEXT("AddNode", "Add Node"));
-		ParentGraph->Modify();
-		if (FromPin)
-		{
-			FromPin->Modify();
-		}
-
-		ResultNode = CreateNode(ParentGraph, FromPin, Location, NodeTemplate);
-	}
-#endif // WITH_EDITOR
-
-	return ResultNode;
-}
-
-UEdGraphNode* FEdGraphSchemaAction_NewNode::PerformAction(class UEdGraph* ParentGraph, TArray<UEdGraphPin*>& FromPins, const FVector2D Location, bool bSelectNewNode/* = true*/) 
-{
-	UEdGraphNode* ResultNode = NULL;
-
-#if WITH_EDITOR
-	if (FromPins.Num() > 0)
-	{
-		ResultNode = PerformAction(ParentGraph, FromPins[0], Location, bSelectNewNode);
-
-		// Try autowiring the rest of the pins
-		for (int32 Index = 1; Index < FromPins.Num(); ++Index)
-		{
-			ResultNode->AutowireNewNode(FromPins[Index]);
-		}
-	}
-	else
-	{
-		ResultNode = PerformAction(ParentGraph, NULL, Location, bSelectNewNode);
-	}
-#endif // WITH_EDITOR
-
-	return ResultNode;
-}
-
-void FEdGraphSchemaAction_NewNode::AddReferencedObjects( FReferenceCollector& Collector )
-{
-	FEdGraphSchemaAction::AddReferencedObjects( Collector );
-
-	// These don't get saved to disk, but we want to make sure the objects don't get GC'd while the action array is around
-	Collector.AddReferencedObject( NodeTemplate );
-}
-
-/////////////////////////////////////////////////////
-// UEdGraphSchema
-
-UEdGraphSchema::UEdGraphSchema(const class FPostConstructInitializeProperties& PCIP)
-	: Super(PCIP)
-{
-}
-
-bool UEdGraphSchema::TryCreateConnection(UEdGraphPin* PinA, UEdGraphPin* PinB) const
-{
-	const FPinConnectionResponse Response = CanCreateConnection(PinA, PinB);
-	bool bModified = false;
-
-	switch (Response.Response)
-	{
-	case CONNECT_RESPONSE_MAKE:
-		PinA->Modify();
-		PinB->Modify();
-		PinA->MakeLinkTo(PinB);
-		bModified = true;
-		break;
-
-	case CONNECT_RESPONSE_BREAK_OTHERS_A:
-		PinA->Modify();
-		PinB->Modify();
-		PinA->BreakAllPinLinks();
-		PinA->MakeLinkTo(PinB);
-		bModified = true;
-		break;
-
-	case CONNECT_RESPONSE_BREAK_OTHERS_B:
-		PinA->Modify();
-		PinB->Modify();
-		PinB->BreakAllPinLinks();
-		PinA->MakeLinkTo(PinB);
-		bModified = true;
-		break;
-
-	case CONNECT_RESPONSE_BREAK_OTHERS_AB:
-		PinA->Modify();
-		PinB->Modify();
-		PinA->BreakAllPinLinks();
-		PinB->BreakAllPinLinks();
-		PinA->MakeLinkTo(PinB);
-		bModified = true;
-		break;
-
-	case CONNECT_RESPONSE_MAKE_WITH_CONVERSION_NODE:
-		bModified = CreateAutomaticConversionNodeAndConnections(PinA, PinB);
-		break;
-		break;
-
-	case CONNECT_RESPONSE_DISALLOW:
-	default:
-		break;
-	}
-
-#if WITH_EDITOR
-	if (bModified)
-	{
-		PinA->GetOwningNode()->PinConnectionListChanged(PinA);
-		PinB->GetOwningNode()->PinConnectionListChanged(PinB);
-	}
-#endif	//#if WITH_EDITOR
-
-	return bModified;
-}
-
-bool UEdGraphSchema::CreateAutomaticConversionNodeAndConnections(UEdGraphPin* PinA, UEdGraphPin* PinB) const
-{
-	return false;
-}
-
-void UEdGraphSchema::TrySetDefaultValue(UEdGraphPin& Pin, const FString& NewDefaultValue) const
-{
-	Pin.DefaultValue = NewDefaultValue;
-
-#if WITH_EDITOR
-	UEdGraphNode* Node = Pin.GetOwningNode();
-	check(Node);
-	Node->PinDefaultValueChanged(&Pin);
-#endif	//#if WITH_EDITOR
-}
-
-void UEdGraphSchema::TrySetDefaultObject(UEdGraphPin& Pin, UObject* NewDefaultObject) const
-{
-	Pin.DefaultObject = NewDefaultObject;
-
-#if WITH_EDITOR
-	UEdGraphNode* Node = Pin.GetOwningNode();
-	check(Node);
-	Node->PinDefaultValueChanged(&Pin);
-#endif	//#if WITH_EDITOR
-}
-
-void UEdGraphSchema::TrySetDefaultText(UEdGraphPin& InPin, const FText& InNewDefaultText) const
-{
-	if(InNewDefaultText.IsEmpty())
-	{
-		InPin.DefaultTextValue = InNewDefaultText;
-	}
-	else
-	{
-#if WITH_EDITOR
-		InPin.DefaultTextValue = FText::ChangeKey(TEXT(""), InPin.GetOwningNode()->NodeGuid.ToString() + TEXT("_") + InPin.PinName, InNewDefaultText);
-#endif
-	}
-
-#if WITH_EDITOR
-	UEdGraphNode* Node = InPin.GetOwningNode();
-	check(Node);
-	Node->PinDefaultValueChanged(&InPin);
-#endif	//#if WITH_EDITOR
-}
-
-void UEdGraphSchema::BreakNodeLinks(UEdGraphNode& TargetNode) const
-{
-#if WITH_EDITOR
-	for (TArray<UEdGraphPin*>::TIterator PinIt(TargetNode.Pins); PinIt; ++PinIt)
-	{
-		BreakPinLinks(*(*PinIt), false);
-	}
-	TargetNode.NodeConnectionListChanged();
-#endif	//#if WITH_EDITOR
-}
-
-void UEdGraphSchema::BreakPinLinks(UEdGraphPin& TargetPin, bool bSendsNodeNotifcation) const
-{
-#if WITH_EDITOR
-	// Copy the old pin links
-	TArray<class UEdGraphPin*> OldLinkedTo(TargetPin.LinkedTo);
-#endif
-
-	TargetPin.BreakAllPinLinks();
-
-#if WITH_EDITOR
-	TSet<UEdGraphNode*> NodeList;
-	// Notify this node
-	TargetPin.GetOwningNode()->PinConnectionListChanged(&TargetPin);
-	NodeList.Add(TargetPin.GetOwningNode());
-
-	// As well as all other nodes that were connected
-	for (TArray<UEdGraphPin*>::TIterator PinIt(OldLinkedTo); PinIt; ++PinIt)
-	{
-		UEdGraphPin* OtherPin = *PinIt;
-		UEdGraphNode* OtherNode = OtherPin->GetOwningNode();
-		OtherNode->PinConnectionListChanged(OtherPin);
-		NodeList.Add(OtherNode);
-	}
-
-	if (bSendsNodeNotifcation)
-	{
-		// Send all nodes that received a new pin connection a notification
-		for (auto It = NodeList.CreateConstIterator(); It; ++It)
-		{
-			UEdGraphNode* Node = (*It);
-			Node->NodeConnectionListChanged();
-		}
-	}
-#endif	//#if WITH_EDITOR
-}
-
-void UEdGraphSchema::BreakSinglePinLink(UEdGraphPin* SourcePin, UEdGraphPin* TargetPin)
-{
-	SourcePin->BreakLinkTo(TargetPin);
-
-#if WITH_EDITOR
-	TargetPin->GetOwningNode()->PinConnectionListChanged(TargetPin);
-	SourcePin->GetOwningNode()->PinConnectionListChanged(SourcePin);
-
-	TargetPin->GetOwningNode()->NodeConnectionListChanged();
-	SourcePin->GetOwningNode()->NodeConnectionListChanged();
-#endif	//#if WITH_EDITOR
-}
-
-FPinConnectionResponse UEdGraphSchema::MovePinLinks(UEdGraphPin& MoveFromPin, UEdGraphPin& MoveToPin, bool bIsIntermeadiateMove) const
-{
-#if WITH_EDITOR
-	ensureMsg(bIsIntermeadiateMove || !MoveToPin.GetOwningNode()->GetGraph()->HasAnyFlags(RF_Transient),
-		TEXT("When moving to an intermeadiate pin, use FKismetCompilerContext::MovePinLinksToIntermediate() instead of UEdGraphSchema::MovePinLinks()"));
-#endif // #if WITH_EDITOR
-
-	FPinConnectionResponse FinalResponse = FPinConnectionResponse(CONNECT_RESPONSE_MAKE, TEXT(""));
-	// First copy the current set of links
-	TArray<UEdGraphPin*> CurrentLinks = MoveFromPin.LinkedTo;
-	// Then break all links at pin we are moving from
-	MoveFromPin.BreakAllPinLinks();
-	// Try and make each new connection
-	for (int32 i=0; i<CurrentLinks.Num(); i++)
-	{
-		UEdGraphPin* NewLink = CurrentLinks[i];
-		FPinConnectionResponse Response = CanCreateConnection(&MoveToPin, NewLink);
-		if(Response.CanSafeConnect())
-		{
-			MoveToPin.MakeLinkTo(NewLink);
-		}
-		else
-		{
-			FinalResponse = Response;
-		}
-	}
-	// Move over the default values
-	MoveToPin.DefaultValue = MoveFromPin.DefaultValue;
-	MoveToPin.DefaultObject = MoveFromPin.DefaultObject;
-	MoveToPin.DefaultTextValue = MoveFromPin.DefaultTextValue;
-	return FinalResponse;
-}
-
-FPinConnectionResponse UEdGraphSchema::CopyPinLinks(UEdGraphPin& CopyFromPin, UEdGraphPin& CopyToPin, bool bIsIntermeadiateCopy) const
-{
-#if WITH_EDITOR
-	ensureMsg(bIsIntermeadiateCopy || !CopyToPin.GetOwningNode()->GetGraph()->HasAnyFlags(RF_Transient),
-		TEXT("When copying to an intermeadiate pin, use FKismetCompilerContext::CopyPinLinksToIntermediate() instead of UEdGraphSchema::CopyPinLinks()"));
-#endif // #if WITH_EDITOR
-
-	FPinConnectionResponse FinalResponse = FPinConnectionResponse(CONNECT_RESPONSE_MAKE, TEXT(""));
-	for (int32 i=0; i<CopyFromPin.LinkedTo.Num(); i++)
-	{
-		UEdGraphPin* NewLink = CopyFromPin.LinkedTo[i];
-		FPinConnectionResponse Response = CanCreateConnection(&CopyToPin, NewLink);
-		if (Response.CanSafeConnect())
-		{
-			CopyToPin.MakeLinkTo(NewLink);
-		}
-		else
-		{
-			FinalResponse = Response;
-		}
-	}
-
-	CopyToPin.DefaultValue = CopyFromPin.DefaultValue;
-	CopyToPin.DefaultObject = CopyFromPin.DefaultObject;
-	CopyToPin.DefaultTextValue = CopyFromPin.DefaultTextValue;
-	return FinalResponse;
-}
-
-FString UEdGraphSchema::GetPinDisplayName(const UEdGraphPin* Pin) const
-{
-	check(Pin != NULL);
-	return !Pin->PinFriendlyName.IsEmpty() ? Pin->PinFriendlyName.ToString() : Pin->PinName;
-}
-
-void UEdGraphSchema::ConstructBasicPinTooltip(UEdGraphPin const& Pin, FString const& PinDescription, FString& TooltipOut) const
-{
-	TooltipOut = PinDescription;
-}
-
-void UEdGraphSchema::GetGraphContextActions(FGraphContextMenuBuilder& ContextMenuBuilder) const
-{
-#if WITH_EDITOR
-	// Run thru all nodes and add any menu items they want to add
-	for (TObjectIterator<UClass> ClassIt; ClassIt; ++ClassIt)
-	{
-		UClass* Class = *ClassIt;
-		if (Class->IsChildOf(UEdGraphNode::StaticClass()) && !Class->HasAnyClassFlags(CLASS_Abstract | CLASS_Deprecated))
-		{
-			const UEdGraphNode* ClassCDO = Class->GetDefaultObject<UEdGraphNode>();
-
-			if (ClassCDO->CanCreateUnderSpecifiedSchema(this))
-			{
-				ClassCDO->GetMenuEntries(ContextMenuBuilder);
-			}
-		}
-	}
-#endif
-}
-
-void UEdGraphSchema::ReconstructNode(UEdGraphNode& TargetNode, bool bIsBatchRequest/*=false*/) const
-{
-#if WITH_EDITOR
-	TargetNode.ReconstructNode();
-#endif	//#if WITH_EDITOR
-}
-
-void UEdGraphSchema::GetGraphDisplayInformation(const UEdGraph& Graph, /*out*/ FGraphDisplayInfo& DisplayInfo) const
-{
-	DisplayInfo.PlainName = FText::FromString( Graph.GetName() );
-	DisplayInfo.DisplayName = DisplayInfo.PlainName;
-}
-
-void UEdGraphSchema::GetContextMenuActions(const UEdGraph* CurrentGraph, const UEdGraphNode* InGraphNode, const UEdGraphPin* InGraphPin, class FMenuBuilder* MenuBuilder, bool bIsDebugging) const
-{
-#if WITH_EDITOR
-	FGraphNodeContextMenuBuilder Context(CurrentGraph, InGraphNode, InGraphPin, MenuBuilder, bIsDebugging);
-
-	if (Context.Node)
-	{
-		Context.Node->GetContextMenuActions(Context);
-	}
-
-	if (InGraphNode)
-	{
-		// Helper to do the node comment editing
-		struct Local
-		{
-			// Called by the EditableText widget to get the current comment for the node
-			static FString GetNodeComment(TWeakObjectPtr<UEdGraphNode> NodeWeakPtr)
-			{
-				if (UEdGraphNode* SelectedNode = NodeWeakPtr.Get())
-				{
-					return SelectedNode->NodeComment;
-				}
-				return FString();
-			}
-
-			// Called by the EditableText widget when the user types a new comment for the selected node
-			static void OnNodeCommentTextCommitted(const FText& NewText, ETextCommit::Type /*CommitInfo*/, TWeakObjectPtr<UEdGraphNode> NodeWeakPtr)
-			{
-				// Apply the change to the selected actor
-				UEdGraphNode* SelectedNode = NodeWeakPtr.Get();
-				FString NewString = NewText.ToString();
-				if (SelectedNode && !SelectedNode->NodeComment.Equals( NewString, ESearchCase::CaseSensitive ))
-				{
-					// send property changed events
-					const FScopedTransaction Transaction( LOCTEXT("EditNodeComment", "Change Node Comment") );
-					SelectedNode->Modify();
-					UProperty* NodeCommentProperty = FindField<UProperty>(SelectedNode->GetClass(), "NodeComment");
-					if(NodeCommentProperty != NULL)
-					{
-						SelectedNode->PreEditChange(NodeCommentProperty);
-
-						SelectedNode->NodeComment = NewString;
-
-						FPropertyChangedEvent NodeCommentPropertyChangedEvent(NodeCommentProperty);
-						SelectedNode->PostEditChangeProperty(NodeCommentPropertyChangedEvent);
-					}
-				}
-
-				FSlateApplication::Get().DismissAllMenus();
-			}
-		};
-
-		if( !InGraphPin )
-		{
-			int32 SelectionCount = GetNodeSelectionCount(CurrentGraph);
-		
-			if( SelectionCount == 1 )
-			{
-				// Node comment area
-				TSharedRef<SHorizontalBox> NodeCommentBox = SNew(SHorizontalBox);
-
-				MenuBuilder->BeginSection("GraphNodeComment", LOCTEXT("NodeCommentMenuHeader", "Node Comment"));
-				{
-					MenuBuilder->AddWidget(NodeCommentBox, FText::GetEmpty() );
-				}
-				TWeakObjectPtr<UEdGraphNode> SelectedNodeWeakPtr = InGraphNode;
-
-				FText NodeCommentText;
-				if ( UEdGraphNode* SelectedNode = SelectedNodeWeakPtr.Get() )
-				{
-					NodeCommentText = FText::FromString( SelectedNode->NodeComment );
-				}
-
-			const FSlateBrush* NodeIcon = FCoreStyle::Get().GetDefaultBrush();//@TODO: FActorIconFinder::FindIconForActor(SelectedActors(0).Get());
-
-
-				// Comment label
-				NodeCommentBox->AddSlot()
-				.VAlign(VAlign_Center)
-				.FillWidth( 1.0f )
-				[
-					SNew(SEditableTextBox)
-					.Text( NodeCommentText )
-					.ToolTipText(LOCTEXT("NodeComment_ToolTip", "Comment for this node"))
-					.OnTextCommitted_Static(&Local::OnNodeCommentTextCommitted, SelectedNodeWeakPtr)
-					.MinDesiredWidth( 120.f )
-					//.IsReadOnly( !SelectedActors(0)->IsActorLabelEditable() )
-					.SelectAllTextWhenFocused( true )
-					.RevertTextOnEscape( true )
-				];
-				MenuBuilder->EndSection();
-			}
-			else if( SelectionCount > 1 )
-			{
-				struct SCommentUtility
-				{
-					static void CreateComment(const UEdGraphSchema* Schema, UEdGraph* Graph)
-					{
-						if( Schema && Graph )
-						{
-							TSharedPtr<FEdGraphSchemaAction> Action = Schema->GetCreateCommentAction();
-
-							if( Action.IsValid() )
-							{
-								Action->PerformAction(Graph, NULL, FVector2D());
-							}
-						}
-					}
-				};
-
-				MenuBuilder->BeginSection("SchemaActionComment", LOCTEXT("MultiCommentHeader", "Comment Group"));
-				MenuBuilder->AddMenuEntry(	LOCTEXT("MultiCommentDesc", "Create Comment from Selection"),
-											LOCTEXT("CommentToolTip", "Create a resizable comment box around selection."),
-											FSlateIcon(), 
-											FUIAction( FExecuteAction::CreateStatic( SCommentUtility::CreateComment, this, const_cast<UEdGraph*>( CurrentGraph ))));
-				MenuBuilder->EndSection();
-			}
-		}
-	}
-#endif
-}
-
-bool FEdGraphPinType::Serialize(FArchive& Ar)
-{
-	if (Ar.UE4Ver() < VER_UE4_EDGRAPHPINTYPE_SERIALIZATION)
-	{
-		return false;
-	}
-
-	Ar << PinCategory;
-	Ar << PinSubCategory;
-
-	// See: FArchive& operator<<( FArchive& Ar, FWeakObjectPtr& WeakObjectPtr )
-	// The PinSubCategoryObject should be serialized into the package.
-	if(!Ar.IsObjectReferenceCollector() || Ar.IsModifyingWeakAndStrongReferences() || Ar.IsPersistent())
-	{
-		UObject* Object = PinSubCategoryObject.Get(true);
-		Ar << Object;
-		if( Ar.IsLoading() || Ar.IsModifyingWeakAndStrongReferences() )
-		{
-			PinSubCategoryObject = Object;
-		}
-	}
-
-	Ar << bIsArray;
-	Ar << bIsReference;
-	Ar << bIsWeakPointer;
-	
-	if (Ar.UE4Ver() >= VER_UE4_MEMBERREFERENCE_IN_PINTYPE)
-	{
-		Ar << PinSubCategoryMemberReference;
-	}
-	else if (Ar.IsLoading() && Ar.IsPersistent())
-	{
-		if ((PinCategory == TEXT("delegate")) || (PinCategory == TEXT("mcdelegate")))
-		{
-			if (const UFunction* Signature = Cast<const UFunction>(PinSubCategoryObject.Get()))
-			{
-				PinSubCategoryMemberReference.MemberName = Signature->GetFName();
-				PinSubCategoryMemberReference.MemberParentClass = Signature->GetOwnerClass();
-				PinSubCategoryObject = NULL;
-			}
-			else
-			{
-				ensure(true);
-			}
-		}
-	}
-
-	return true;
-}
 
 #undef LOCTEXT_NAMESPACE
diff --git a/Engine/Source/Runtime/Engine/Private/EdGraph/EdGraphNode.cpp b/Engine/Source/Runtime/Engine/Private/EdGraph/EdGraphNode.cpp
new file mode 100644
--- /dev/null
+++ b/Engine/Source/Runtime/Engine/Private/EdGraph/EdGraphNode.cpp
@@ -0,0 +1,330 @@
+// Copyright 1998-2014 Epic Games, Inc. All Rights Reserved.
+
+#include "EnginePrivate.h"
+#include "EngineClasses.h"
+#include "BlueprintUtilities.h"
+#if WITH_EDITOR
+#include "Editor/UnrealEd/Public/Kismet2/BlueprintEditorUtils.h"
+#include "Slate.h"
+#include "ScopedTransaction.h"
+#include "Editor/UnrealEd/Public/Kismet2/Kismet2NameValidators.h"
+#endif
+
+#define LOCTEXT_NAMESPACE "EdGraph"
+
+/////////////////////////////////////////////////////
+// FGraphNodeContextMenuBuilder
+
+FGraphNodeContextMenuBuilder::FGraphNodeContextMenuBuilder(const UEdGraph* InGraph, const UEdGraphNode* InNode, const UEdGraphPin* InPin, FMenuBuilder* InMenuBuilder, bool bInDebuggingMode)
+	: Blueprint(NULL)
+	, Graph(InGraph)
+	, Node(InNode)
+	, Pin(InPin)
+	, MenuBuilder(InMenuBuilder)
+	, bIsDebugging(bInDebuggingMode)
+{
+#if WITH_EDITOR
+	Blueprint = FBlueprintEditorUtils::FindBlueprintForGraph(Graph);
+#endif
+
+	if (Pin != NULL)
+	{
+		Node = Pin->GetOwningNode();
+	}
+}
+
+/////////////////////////////////////////////////////
+// UEdGraphNode
+
+UEdGraphNode::UEdGraphNode(const class FPostConstructInitializeProperties& PCIP)
+	: Super(PCIP)
+	, AdvancedPinDisplay(ENodeAdvancedPins::NoPins)
+{
+
+#if WITH_EDITORONLY_DATA
+	bCanResizeNode = false;
+#endif // WITH_EDITORONLY_DATA
+}
+
+#if WITH_EDITOR
+UEdGraphPin* UEdGraphNode::CreatePin(EEdGraphPinDirection Dir, const FString& PinCategory, const FString& PinSubCategory, UObject* PinSubCategoryObject, bool bIsArray, bool bIsReference, const FString& PinName, bool bIsConst /*= false*/, int32 Index /*= INDEX_NONE*/)
+{
+#if 0
+	UEdGraphPin* NewPin = AllocatePinFromPool(this);
+#else
+	UEdGraphPin* NewPin = NewObject<UEdGraphPin>(this);
+#endif
+	NewPin->PinName = PinName;
+	NewPin->Direction = Dir;
+	NewPin->PinType.PinCategory = PinCategory;
+	NewPin->PinType.PinSubCategory = PinSubCategory;
+	NewPin->PinType.PinSubCategoryObject = PinSubCategoryObject;
+	NewPin->PinType.bIsArray = bIsArray;
+	NewPin->PinType.bIsReference = bIsReference;
+	NewPin->PinType.bIsConst = bIsConst;
+	NewPin->SetFlags(RF_Transactional);
+
+	Modify(false);
+	if ( Pins.IsValidIndex( Index ) )
+	{
+		Pins.Insert(NewPin, Index);
+	}
+	else
+	{
+		Pins.Add(NewPin);
+	}
+	return NewPin;
+}
+
+UEdGraphPin* UEdGraphNode::FindPin(const FString& PinName) const
+{
+	for(int32 PinIdx=0; PinIdx<Pins.Num(); PinIdx++)
+	{
+		if( Pins[PinIdx]->PinName == PinName )
+		{
+			return Pins[PinIdx];
+		}
+	}
+
+	return NULL;
+}
+
+UEdGraphPin* UEdGraphNode::FindPinChecked(const FString& PinName) const
+{
+	UEdGraphPin* Result = FindPin(PinName);
+	check(Result != NULL);
+	return Result;
+}
+
+void UEdGraphNode::DiscardPin(UEdGraphPin* Pin)
+{
+	check( Pin );
+
+	Modify();
+	Pins.Remove( Pin );
+}
+
+void UEdGraphNode::BreakAllNodeLinks()
+{
+	TSet<UEdGraphNode*> NodeList;
+
+	NodeList.Add(this);
+
+	// Iterate over each pin and break all links
+	for(int32 PinIdx=0; PinIdx<Pins.Num(); PinIdx++)
+	{
+		Pins[PinIdx]->BreakAllPinLinks();
+		NodeList.Add(Pins[PinIdx]->GetOwningNode());
+	}
+
+	// Send all nodes that received a new pin connection a notification
+	for (auto It = NodeList.CreateConstIterator(); It; ++It)
+	{
+		UEdGraphNode* Node = (*It);
+		Node->NodeConnectionListChanged();
+	}
+}
+
+void UEdGraphNode::GetPinHoverText(const UEdGraphPin& Pin, FString& HoverTextOut) const
+{
+	ensure(Pin.GetOwningNode() == this);
+	HoverTextOut = Pin.PinToolTip;
+}
+
+void UEdGraphNode::SnapToGrid(float GridSnapSize)
+{
+	NodePosX = GridSnapSize * FMath::RoundToInt(NodePosX/GridSnapSize);
+	NodePosY = GridSnapSize * FMath::RoundToInt(NodePosY/GridSnapSize);
+}
+
+void UEdGraphNode::DestroyNode()
+{
+	UEdGraph* ParentGraph = GetGraph();
+	check(ParentGraph);
+
+	// Remove the node - this will break all links. Will be GC'd after this.
+	ParentGraph->RemoveNode(this);
+}
+
+const class UEdGraphSchema* UEdGraphNode::GetSchema() const
+{
+	return GetGraph()->GetSchema();
+}
+
+FLinearColor UEdGraphNode::GetNodeTitleColor() const
+{
+	return FLinearColor(0.4f, 0.62f, 1.0f);
+}
+
+FLinearColor UEdGraphNode::GetNodeCommentColor() const
+{
+	return FLinearColor::White;
+}
+
+FString UEdGraphNode::GetTooltip() const
+{
+	return GetClass()->GetToolTipText().ToString();
+}
+
+FString UEdGraphNode::GetDocumentationExcerptName() const
+{
+	// Default the node to searching for an excerpt named for the C++ node class name, including the U prefix.
+	// This is done so that the excerpt name in the doc file can be found by find-in-files when searching for the full class name.
+	UClass* MyClass = GetClass();
+	return FString::Printf(TEXT("%s%s"), MyClass->GetPrefixCPP(), *MyClass->GetName());
+}
+
+
+FString UEdGraphNode::GetDescriptiveCompiledName() const
+{
+	return GetFName().GetPlainNameString();
+}
+
+bool UEdGraphNode::IsDeprecated() const
+{
+	return GetClass()->HasAnyClassFlags(CLASS_Deprecated);
+}
+
+FString UEdGraphNode::GetDeprecationMessage() const
+{
+	return NSLOCTEXT("EdGraphCompiler", "NodeDeprecated_Warning", "@@ is deprecated; please replace or remove it.").ToString();
+}
+
+// Array of pooled pins
+TArray<UEdGraphPin*> UEdGraphNode::PooledPins;
+
+void UEdGraphNode::AddReferencedObjects(UObject* InThis, FReferenceCollector& Collector)
+{
+	UEdGraphNode* This = CastChecked<UEdGraphNode>(InThis);	
+
+	// Only register the pool once per GC pass
+	if (This->HasAnyFlags(RF_ClassDefaultObject))
+	{
+		if (This->GetClass() == UEdGraphNode::StaticClass())
+		{
+			for (int32 Index = 0; Index < PooledPins.Num(); ++Index)
+			{
+				Collector.AddReferencedObject(PooledPins[Index], This);
+			}
+		}
+	}
+	Super::AddReferencedObjects(This, Collector);
+}
+
+void UEdGraphNode::PostLoad()
+{
+	Super::PostLoad();
+
+	// Create Guid if not present (and not CDO)
+	if(!NodeGuid.IsValid() && !IsTemplate() && GetLinker() && GetLinker()->IsPersistent() && GetLinker()->IsLoading())
+	{
+		// _Should_ have a guid on all nodes after this version
+		if(GetLinkerUE4Version() >= VER_UE4_ADD_EDGRAPHNODE_GUID)
+		{
+			UE_LOG(LogBlueprint, Warning, TEXT("Node '%s' missing NodeGuid."), *GetPathName());
+		}
+
+		// Generate new one
+		CreateNewGuid();
+	}
+}
+
+void UEdGraphNode::CreateNewGuid()
+{
+	NodeGuid = FGuid::NewGuid();
+}
+
+void UEdGraphNode::FindDiffs( class UEdGraphNode* OtherNode, struct FDiffResults& Results ) 
+{
+}
+
+UEdGraphPin* UEdGraphNode::AllocatePinFromPool(UEdGraphNode* OuterNode)
+{
+	if (PooledPins.Num() > 0)
+	{
+		UEdGraphPin* Result = PooledPins.Pop();
+		Result->Rename(NULL, OuterNode);
+		return Result;
+	}
+	else
+	{
+		UEdGraphPin* Result = NewObject<UEdGraphPin>(OuterNode);
+		return Result;
+	}
+}
+
+void UEdGraphNode::ReturnPinToPool(UEdGraphPin* OldPin)
+{
+	check(OldPin);
+
+	check(!OldPin->HasAnyFlags(RF_NeedLoad));
+	OldPin->ClearFlags(RF_NeedPostLoadSubobjects|RF_NeedPostLoad);
+
+	OldPin->ResetToDefaults();
+
+	PooledPins.Add(OldPin);
+}
+
+bool UEdGraphNode::CanDuplicateNode() const
+{
+	return true;
+}
+
+bool UEdGraphNode::CanUserDeleteNode() const
+{
+	return true;
+}
+
+FText UEdGraphNode::GetNodeTitle(ENodeTitleType::Type TitleType) const
+{
+	return FText::FromString(GetClass()->GetName());
+}
+
+FString UEdGraphNode::GetNodeNativeTitle(ENodeTitleType::Type TitleType) const
+{
+	FText NodeTitle = GetNodeTitle(TitleType);
+	if(const FString* SourceString = FTextInspector::GetSourceString(NodeTitle))
+	{
+		return *SourceString;
+	}
+	return NodeTitle.ToString();
+}
+
+FText UEdGraphNode::GetNodeSearchTitle() const
+{
+	TArray< FText > SearchMetadata;
+
+	FText LocalizedTitle = GetNodeTitle(ENodeTitleType::ListView);
+	FText NativeTitle = FText::FromString(GetNodeNativeTitle(ENodeTitleType::ListView));
+
+	if(FInternationalization::Get().GetCurrentCulture()->GetName() == TEXT("en_US"))
+	{
+		ensure( LocalizedTitle.CompareTo(NativeTitle) == 0 );
+	}
+
+	FFormatNamedArguments Args;
+	Args.Add(TEXT("LocalizedTitle"), LocalizedTitle);
+	Args.Add(TEXT("NativeTitle"), NativeTitle);
+	return FText::Format(LOCTEXT("NodeSearchTitle", "{LocalizedTitle} {NativeTitle}"), Args);
+}
+
+UObject* UEdGraphNode::GetJumpTargetForDoubleClick() const
+{
+	return NULL;
+}
+
+FString UEdGraphNode::GetPinDisplayName(const UEdGraphPin* Pin) const
+{
+	return GetSchema()->GetPinDisplayName(Pin);
+}
+
+int32 UEdGraphNode::GetPinIndex(UEdGraphPin* Pin) const
+{
+	return Pins.Find(Pin);
+}
+
+#endif	//#if WITH_EDITOR
+
+/////////////////////////////////////////////////////
+
+#undef LOCTEXT_NAMESPACE
\ No newline at end of file
diff --git a/Engine/Source/Runtime/Engine/Private/EdGraph/EdGraphNode_Comment.cpp b/Engine/Source/Runtime/Engine/Private/EdGraph/EdGraphNode_Comment.cpp
new file mode 100644
--- /dev/null
+++ b/Engine/Source/Runtime/Engine/Private/EdGraph/EdGraphNode_Comment.cpp
@@ -0,0 +1,143 @@
+// Copyright 1998-2014 Epic Games, Inc. All Rights Reserved.
+
+#include "EnginePrivate.h"
+#include "EngineClasses.h"
+#include "BlueprintUtilities.h"
+#if WITH_EDITOR
+#include "Editor/UnrealEd/Public/Kismet2/BlueprintEditorUtils.h"
+#include "Slate.h"
+#include "ScopedTransaction.h"
+#include "Editor/UnrealEd/Public/Kismet2/Kismet2NameValidators.h"
+#endif
+
+#define LOCTEXT_NAMESPACE "EdGraph"
+
+/////////////////////////////////////////////////////
+// UEdGraphNode_Comment
+
+UEdGraphNode_Comment::UEdGraphNode_Comment(const class FPostConstructInitializeProperties& PCIP)
+	: Super(PCIP)
+{
+	NodeWidth = 400.0f;
+	NodeHeight = 100.0f;
+	CommentColor = FLinearColor::White;
+	bColorCommentBubble = false;
+	MoveMode = ECommentBoxMode::GroupMovement;
+
+#if WITH_EDITORONLY_DATA
+	bCanResizeNode = true;
+	bCanRenameNode = true;
+#endif // WITH_EDITORONLY_DATA
+}
+
+void UEdGraphNode_Comment::AddReferencedObjects(UObject* InThis, FReferenceCollector& Collector) 
+{
+#if WITH_EDITOR
+	UEdGraphNode_Comment* This = CastChecked<UEdGraphNode_Comment>(InThis);
+	for (auto It = This->NodesUnderComment.CreateIterator(); It; ++It)
+	{
+		Collector.AddReferencedObject(*It, This);
+	}
+#endif // WITH_EDITOR
+
+	Super::AddReferencedObjects(InThis, Collector);
+}
+
+#if WITH_EDITOR
+void UEdGraphNode_Comment::PostPlacedNewNode()
+{
+	//@TODO: Consider making the default value we use here a preference
+	// This is done here instead of in the constructor so we can later change the default for newly placed
+	// instances without changing all of the existing ones (due to delta serialization)
+	MoveMode = ECommentBoxMode::GroupMovement;
+
+	NodeComment = NSLOCTEXT("K2Node", "CommentBlock_NewEmptyComment", "Comment").ToString();
+}
+
+FString UEdGraphNode_Comment::GetTooltip() const
+{
+	return FString::Printf(*NSLOCTEXT("K2Node", "CommentBlock_Tooltip", "Comment:\n%s").ToString(), *NodeComment);
+}
+
+FString UEdGraphNode_Comment::GetDocumentationLink() const
+{
+	return TEXT("Shared/GraphNodes/Common");
+}
+
+FString UEdGraphNode_Comment::GetDocumentationExcerptName() const
+{
+	return TEXT("UEdGraphNode_Comment");
+}
+
+FText UEdGraphNode_Comment::GetNodeTitle(ENodeTitleType::Type TitleType) const
+{
+	if(TitleType == ENodeTitleType::ListView)
+	{
+		return NSLOCTEXT("K2Node", "CommentBlock_Title", "Comment");
+	}
+
+	return FText::FromString(NodeComment);
+}
+
+FString UEdGraphNode_Comment::GetPinNameOverride(const UEdGraphPin& Pin) const
+{
+	return GetNodeTitle(ENodeTitleType::ListView).ToString();
+}
+
+FLinearColor UEdGraphNode_Comment::GetNodeCommentColor() const
+{
+	// Only affects the 'zoomed out' comment bubble color, not the box itself
+	return (bColorCommentBubble)
+		? CommentColor 
+		: FLinearColor::White;
+}
+
+void UEdGraphNode_Comment::ResizeNode(const FVector2D& NewSize)
+{
+	if (bCanResizeNode) 
+	{
+		NodeHeight = NewSize.Y;
+		NodeWidth = NewSize.X;
+	}
+}
+
+void UEdGraphNode_Comment::AddNodeUnderComment(UObject* Object)
+{
+	NodesUnderComment.Add(Object);
+}
+
+void UEdGraphNode_Comment::ClearNodesUnderComment()
+{
+	NodesUnderComment.Empty();
+}
+
+void UEdGraphNode_Comment::SetBounds(const class FSlateRect& Rect)
+{
+	NodePosX = Rect.Left;
+	NodePosY = Rect.Top;
+
+	FVector2D Size = Rect.GetSize();
+	NodeWidth = Size.X;
+	NodeHeight = Size.Y;
+}
+
+const FCommentNodeSet& UEdGraphNode_Comment::GetNodesUnderComment() const
+{
+	return NodesUnderComment;
+}
+
+void UEdGraphNode_Comment::OnRenameNode(const FString& NewName)
+{
+	NodeComment = NewName;
+}
+
+TSharedPtr<class INameValidatorInterface> UEdGraphNode_Comment::MakeNameValidator() const
+{
+	// Comments can be duplicated, etc...
+	return MakeShareable(new FDummyNameValidator(EValidatorResult::Ok));
+}
+#endif // WITH_EDITOR
+
+/////////////////////////////////////////////////////
+
+#undef LOCTEXT_NAMESPACE
diff --git a/Engine/Source/Runtime/Engine/Private/EdGraph/EdGraphPin.cpp b/Engine/Source/Runtime/Engine/Private/EdGraph/EdGraphPin.cpp
new file mode 100644
--- /dev/null
+++ b/Engine/Source/Runtime/Engine/Private/EdGraph/EdGraphPin.cpp
@@ -0,0 +1,238 @@
+// Copyright 1998-2014 Epic Games, Inc. All Rights Reserved.
+
+#include "EnginePrivate.h"
+#include "EngineClasses.h"
+#include "BlueprintUtilities.h"
+#if WITH_EDITOR
+#include "Editor/UnrealEd/Public/Kismet2/BlueprintEditorUtils.h"
+#include "Slate.h"
+#include "ScopedTransaction.h"
+#include "Editor/UnrealEd/Public/Kismet2/Kismet2NameValidators.h"
+#endif
+
+#define LOCTEXT_NAMESPACE "EdGraph"
+
+/////////////////////////////////////////////////////
+// FEdGraphPinType
+
+bool FEdGraphPinType::Serialize(FArchive& Ar)
+{
+	if (Ar.UE4Ver() < VER_UE4_EDGRAPHPINTYPE_SERIALIZATION)
+	{
+		return false;
+	}
+
+	Ar << PinCategory;
+	Ar << PinSubCategory;
+
+	// See: FArchive& operator<<( FArchive& Ar, FWeakObjectPtr& WeakObjectPtr )
+	// The PinSubCategoryObject should be serialized into the package.
+	if(!Ar.IsObjectReferenceCollector() || Ar.IsModifyingWeakAndStrongReferences() || Ar.IsPersistent())
+	{
+		UObject* Object = PinSubCategoryObject.Get(true);
+		Ar << Object;
+		if( Ar.IsLoading() || Ar.IsModifyingWeakAndStrongReferences() )
+		{
+			PinSubCategoryObject = Object;
+		}
+	}
+
+	Ar << bIsArray;
+	Ar << bIsReference;
+	Ar << bIsWeakPointer;
+	
+	if (Ar.UE4Ver() >= VER_UE4_MEMBERREFERENCE_IN_PINTYPE)
+	{
+		Ar << PinSubCategoryMemberReference;
+	}
+	else if (Ar.IsLoading() && Ar.IsPersistent())
+	{
+		if ((PinCategory == TEXT("delegate")) || (PinCategory == TEXT("mcdelegate")))
+		{
+			if (const UFunction* Signature = Cast<const UFunction>(PinSubCategoryObject.Get()))
+			{
+				PinSubCategoryMemberReference.MemberName = Signature->GetFName();
+				PinSubCategoryMemberReference.MemberParentClass = Signature->GetOwnerClass();
+				PinSubCategoryObject = NULL;
+			}
+			else
+			{
+				ensure(true);
+			}
+		}
+	}
+
+	return true;
+}
+
+/////////////////////////////////////////////////////
+// UEdGraphPin
+
+UEdGraphPin::UEdGraphPin(const class FPostConstructInitializeProperties& PCIP)
+	: Super(PCIP)
+{
+#if WITH_EDITORONLY_DATA
+	bHidden = false;
+	bNotConnectable = false;
+	bAdvancedView = false;
+#endif // WITH_EDITORONLY_DATA
+}
+
+void UEdGraphPin::MakeLinkTo(UEdGraphPin* ToPin)
+{
+	Modify();
+
+	if (ToPin != NULL)
+	{
+		ToPin->Modify();
+
+		// Make sure we don't already link to it
+		if (!LinkedTo.Contains(ToPin))
+		{
+			UEdGraphNode* MyNode = GetOwningNode();
+
+			// Check that the other pin does not link to us
+			ensureMsg(!ToPin->LinkedTo.Contains(this), *GetLinkInfoString( LOCTEXT("MakeLinkTo", "MakeLinkTo").ToString(), LOCTEXT("IsLinked", "is linked with pin").ToString(), ToPin));			    
+			ensureMsg(MyNode->GetOuter() == ToPin->GetOwningNode()->GetOuter(), *GetLinkInfoString( LOCTEXT("MakeLinkTo", "MakeLinkTo").ToString(), LOCTEXT("OuterMismatch", "has a different outer than pin").ToString(), ToPin)); // Ensure both pins belong to the same graph
+
+			// Add to both lists
+			LinkedTo.Add(ToPin);
+			ToPin->LinkedTo.Add(this);
+		}
+	}
+}
+
+void UEdGraphPin::BreakLinkTo(UEdGraphPin* ToPin)
+{
+	Modify();
+
+	if (ToPin != NULL)
+	{
+		ToPin->Modify();
+
+		// If we do indeed link to the passed in pin...
+		if (LinkedTo.Contains(ToPin))
+		{
+			LinkedTo.Remove(ToPin);
+
+			// Check that the other pin links to us
+			ensureMsg(ToPin->LinkedTo.Contains(this), *GetLinkInfoString( LOCTEXT("BreakLinkTo", "BreakLinkTo").ToString(), LOCTEXT("NotLinked", "not reciprocally linked with pin").ToString(), ToPin) );
+			ToPin->LinkedTo.Remove(this);
+		}
+		else
+		{
+			// Check that the other pin does not link to us
+			ensureMsg(!ToPin->LinkedTo.Contains(this), *GetLinkInfoString( LOCTEXT("MakeLinkTo", "MakeLinkTo").ToString(), LOCTEXT("IsLinked", "is linked with pin").ToString(), ToPin));
+		}
+	}
+}
+
+void UEdGraphPin::BreakAllPinLinks()
+{
+	TArray<UEdGraphPin*> LinkedToCopy = LinkedTo;
+
+	for (int32 LinkIdx = 0; LinkIdx < LinkedToCopy.Num(); LinkIdx++)
+	{
+		BreakLinkTo(LinkedToCopy[LinkIdx]);
+	}
+}
+
+
+void UEdGraphPin::CopyPersistentDataFromOldPin(const UEdGraphPin& SourcePin)
+{
+	// The name matches already, doesn't get copied here
+	// The PinType, Direction, and bNotConnectable are properties generated from the schema
+
+	// Only move the default value if it was modified; inherit the new default value otherwise
+	if (SourcePin.DefaultValue != SourcePin.AutogeneratedDefaultValue || SourcePin.DefaultObject != NULL)
+	{
+		DefaultObject = SourcePin.DefaultObject;
+		DefaultValue = SourcePin.DefaultValue;
+	}
+
+	const bool bTextValueWasModified = (SourcePin.DefaultTextValue.ToString() != SourcePin.AutogeneratedDefaultValue);
+	if (bTextValueWasModified)
+	{
+		DefaultTextValue = SourcePin.DefaultTextValue;
+	}
+
+	// Copy the links
+	for (int32 LinkIndex = 0; LinkIndex < SourcePin.LinkedTo.Num(); ++LinkIndex)
+	{
+		UEdGraphPin* OtherPin = SourcePin.LinkedTo[LinkIndex];
+		check(NULL != OtherPin);
+		
+		Modify();
+		OtherPin->Modify();
+
+		OtherPin->MakeLinkTo(this);
+	}
+
+#if WITH_EDITORONLY_DATA
+	// Copy advanced visibility property, since it can be changed by user
+	bAdvancedView = SourcePin.bAdvancedView;
+#endif // WITH_EDITORONLY_DATA
+}
+
+const class UEdGraphSchema* UEdGraphPin::GetSchema() const
+{
+#if WITH_EDITOR
+	return GetOwningNode()->GetGraph()->GetSchema();
+#else
+	return NULL;
+#endif	//WITH_EDITOR
+}
+
+FString UEdGraphPin::GetDefaultAsString() const
+{
+	if(DefaultObject != NULL)
+	{
+		return DefaultObject->GetFullName();
+	}
+	else
+	{
+		return DefaultValue;
+	}
+}
+
+const FString UEdGraphPin::GetLinkInfoString( const FString& InFunctionName, const FString& InInfoData, const UEdGraphPin* InToPin ) const
+{
+#if WITH_EDITOR
+	const FString FromPinName = PinName;
+	const UEdGraphNode* FromPinNode = Cast<UEdGraphNode>(GetOuter());
+	const FString FromPinNodeName = (FromPinNode != NULL) ? FromPinNode->GetNodeTitle(ENodeTitleType::ListView).ToString() : TEXT("Unknown");
+	
+	const FString ToPinName = InToPin->PinName;
+	const UEdGraphNode* ToPinNode = Cast<UEdGraphNode>(InToPin->GetOuter());
+	const FString ToPinNodeName = (ToPinNode != NULL) ? ToPinNode->GetNodeTitle(ENodeTitleType::ListView).ToString() : TEXT("Unknown");
+	const FString LinkInfo = FString::Printf( TEXT("UEdGraphPin::%s Pin '%s' on node '%s' %s '%s' on node '%s'"), *InFunctionName, *ToPinName, *ToPinNodeName, *InInfoData, *FromPinName, *FromPinNodeName);
+	return LinkInfo;
+#else
+	return FString();
+#endif
+}
+
+FString FGraphDisplayInfo::GetNotesAsString() const
+{
+	FString Result;
+
+	if (Notes.Num() > 0)
+	{
+		Result = TEXT("(");
+		for (int32 NoteIndex = 0; NoteIndex < Notes.Num(); ++NoteIndex)
+		{
+			if (NoteIndex > 0)
+			{
+				Result += TEXT(", ");
+			}
+			Result += Notes[NoteIndex];
+		}
+		Result += TEXT(")");
+	}
+
+	return Result;
+}
+
+/////////////////////////////////////////////////////
+
+#undef LOCTEXT_NAMESPACE
diff --git a/Engine/Source/Runtime/Engine/Private/EdGraph/EdGraphSchema.cpp b/Engine/Source/Runtime/Engine/Private/EdGraph/EdGraphSchema.cpp
new file mode 100644
--- /dev/null
+++ b/Engine/Source/Runtime/Engine/Private/EdGraph/EdGraphSchema.cpp
@@ -0,0 +1,631 @@
+// Copyright 1998-2014 Epic Games, Inc. All Rights Reserved.
+
+#include "EnginePrivate.h"
+#include "EngineClasses.h"
+#include "BlueprintUtilities.h"
+#if WITH_EDITOR
+#include "Editor/UnrealEd/Public/Kismet2/BlueprintEditorUtils.h"
+#include "Slate.h"
+#include "ScopedTransaction.h"
+#include "Editor/UnrealEd/Public/Kismet2/Kismet2NameValidators.h"
+#endif
+
+#define LOCTEXT_NAMESPACE "EdGraph"
+
+/////////////////////////////////////////////////////
+// FGraphActionListBuilderBase
+
+void FGraphActionListBuilderBase::AddAction( const TSharedPtr<FEdGraphSchemaAction>& NewAction, FString const& Category/* = TEXT("") */ )
+{
+	Entries.Add( ActionGroup( NewAction, Category ) );
+}
+
+void FGraphActionListBuilderBase::AddActionList( const TArray<TSharedPtr<FEdGraphSchemaAction> >& NewActions, FString const& Category/* = TEXT("") */ )
+{
+	Entries.Add( ActionGroup( NewActions, Category ) );
+}
+
+void FGraphActionListBuilderBase::Append( FGraphActionListBuilderBase& Other )
+{
+	Entries.Append( Other.Entries );
+}
+
+int32 FGraphActionListBuilderBase::GetNumActions() const
+{
+	return Entries.Num();
+}
+
+FGraphActionListBuilderBase::ActionGroup& FGraphActionListBuilderBase::GetAction( int32 Index )
+{
+	return Entries[Index];
+}
+
+void FGraphActionListBuilderBase::Empty()
+{
+	Entries.Empty();
+}
+
+/////////////////////////////////////////////////////
+// FGraphActionListBuilderBase::GraphAction
+
+FGraphActionListBuilderBase::ActionGroup::ActionGroup( TSharedPtr<FEdGraphSchemaAction> InAction, FString const& CategoryPrefix/* = TEXT("") */ )
+	: RootCategory(CategoryPrefix)
+{
+	Actions.Add( InAction );
+}
+
+FGraphActionListBuilderBase::ActionGroup::ActionGroup( const TArray< TSharedPtr<FEdGraphSchemaAction> >& InActions, FString const& CategoryPrefix/* = TEXT("") */ )
+	: RootCategory(CategoryPrefix)
+{
+	Actions = InActions;
+}
+
+void FGraphActionListBuilderBase::ActionGroup::GetCategoryChain(TArray<FString>& HierarchyOut) const
+{
+	static FString const CategoryDelim("|");
+	RootCategory.ParseIntoArray(&HierarchyOut, *CategoryDelim, true);
+
+	if (Actions.Num() > 0)
+	{
+		TArray<FString> SubCategoryChain;
+		Actions[0]->Category.ParseIntoArray(&SubCategoryChain, *CategoryDelim, true);
+
+		HierarchyOut.Append(SubCategoryChain);
+	}
+
+	for (FString& Category : HierarchyOut)
+	{
+		Category.Trim();
+	}
+}
+
+void FGraphActionListBuilderBase::ActionGroup::PerformAction( class UEdGraph* ParentGraph, TArray<UEdGraphPin*>& FromPins, const FVector2D Location )
+{
+	for ( int32 ActionIndex = 0; ActionIndex < Actions.Num(); ActionIndex++ )
+	{
+		TSharedPtr<FEdGraphSchemaAction> CurrentAction = Actions[ActionIndex];
+		if ( CurrentAction.IsValid() )
+		{
+			CurrentAction->PerformAction( ParentGraph, FromPins, Location );
+		}
+	}
+}
+
+/////////////////////////////////////////////////////
+// FCategorizedGraphActionListBuilder
+
+static FString ConcatCategories(FString const& RootCategory, FString const& SubCategory)
+{
+	FString ConcatedCategory = RootCategory;
+	if (!SubCategory.IsEmpty() && !RootCategory.IsEmpty())
+	{
+		ConcatedCategory += TEXT("|");
+	}
+	ConcatedCategory += SubCategory;
+
+	return ConcatedCategory;
+}
+
+FCategorizedGraphActionListBuilder::FCategorizedGraphActionListBuilder(FString const& CategoryIn/* = TEXT("")*/)
+	: Category(CategoryIn)
+{
+}
+
+void FCategorizedGraphActionListBuilder::AddAction(TSharedPtr<FEdGraphSchemaAction> const& NewAction, FString const& CategoryIn/* = TEXT("")*/)
+{
+	FGraphActionListBuilderBase::AddAction(NewAction, ConcatCategories(Category, CategoryIn));
+}
+
+void FCategorizedGraphActionListBuilder::AddActionList(TArray<TSharedPtr<FEdGraphSchemaAction> > const& NewActions, FString const& CategoryIn/* = TEXT("")*/)
+{
+	FGraphActionListBuilderBase::AddActionList(NewActions, ConcatCategories(Category, CategoryIn));
+}
+
+/////////////////////////////////////////////////////
+// FEdGraphSchemaAction_NewNode
+
+#define SNAP_GRID (16) // @todo ensure this is the same as SNodePanel::GetSnapGridSize()
+
+namespace 
+{
+	// Maximum distance a drag can be off a node edge to require 'push off' from node
+	const int32 NodeDistance = 60;
+}
+
+UEdGraphNode* FEdGraphSchemaAction_NewNode::CreateNode(class UEdGraph* ParentGraph, UEdGraphPin* FromPin, const FVector2D Location, class UEdGraphNode* NodeTemplate)
+{
+	// Duplicate template node to create new node
+	UEdGraphNode* ResultNode = NULL;
+
+#if WITH_EDITOR
+	ResultNode = DuplicateObject<UEdGraphNode>(NodeTemplate, ParentGraph);
+	ResultNode->SetFlags(RF_Transactional);
+
+	ParentGraph->AddNode(ResultNode, true);
+
+	ResultNode->CreateNewGuid();
+	ResultNode->PostPlacedNewNode();
+	ResultNode->AllocateDefaultPins();
+	ResultNode->AutowireNewNode(FromPin);
+
+	// For input pins, new node will generally overlap node being dragged off
+	// Work out if we want to visually push away from connected node
+	int32 XLocation = Location.X;
+	if (FromPin && FromPin->Direction == EGPD_Input)
+	{
+		UEdGraphNode* PinNode = FromPin->GetOwningNode();
+		const float XDelta = FMath::Abs(PinNode->NodePosX - Location.X);
+
+		if (XDelta < NodeDistance)
+		{
+			// Set location to edge of current node minus the max move distance
+			// to force node to push off from connect node enough to give selection handle
+			XLocation = PinNode->NodePosX - NodeDistance;
+		}
+	}
+
+	ResultNode->NodePosX = XLocation;
+	ResultNode->NodePosY = Location.Y;
+	ResultNode->SnapToGrid(SNAP_GRID);
+#endif // WITH_EDITOR
+
+	return ResultNode;
+}
+
+UEdGraphNode* FEdGraphSchemaAction_NewNode::PerformAction(class UEdGraph* ParentGraph, UEdGraphPin* FromPin, const FVector2D Location, bool bSelectNewNode/* = true*/)
+{
+	UEdGraphNode* ResultNode = NULL;
+
+#if WITH_EDITOR
+	// If there is a template, we actually use it
+	if (NodeTemplate != NULL)
+	{
+		const FScopedTransaction Transaction(LOCTEXT("AddNode", "Add Node"));
+		ParentGraph->Modify();
+		if (FromPin)
+		{
+			FromPin->Modify();
+		}
+
+		ResultNode = CreateNode(ParentGraph, FromPin, Location, NodeTemplate);
+	}
+#endif // WITH_EDITOR
+
+	return ResultNode;
+}
+
+UEdGraphNode* FEdGraphSchemaAction_NewNode::PerformAction(class UEdGraph* ParentGraph, TArray<UEdGraphPin*>& FromPins, const FVector2D Location, bool bSelectNewNode/* = true*/) 
+{
+	UEdGraphNode* ResultNode = NULL;
+
+#if WITH_EDITOR
+	if (FromPins.Num() > 0)
+	{
+		ResultNode = PerformAction(ParentGraph, FromPins[0], Location, bSelectNewNode);
+
+		// Try autowiring the rest of the pins
+		for (int32 Index = 1; Index < FromPins.Num(); ++Index)
+		{
+			ResultNode->AutowireNewNode(FromPins[Index]);
+		}
+	}
+	else
+	{
+		ResultNode = PerformAction(ParentGraph, NULL, Location, bSelectNewNode);
+	}
+#endif // WITH_EDITOR
+
+	return ResultNode;
+}
+
+void FEdGraphSchemaAction_NewNode::AddReferencedObjects( FReferenceCollector& Collector )
+{
+	FEdGraphSchemaAction::AddReferencedObjects( Collector );
+
+	// These don't get saved to disk, but we want to make sure the objects don't get GC'd while the action array is around
+	Collector.AddReferencedObject( NodeTemplate );
+}
+
+/////////////////////////////////////////////////////
+// UEdGraphSchema
+
+UEdGraphSchema::UEdGraphSchema(const class FPostConstructInitializeProperties& PCIP)
+	: Super(PCIP)
+{
+}
+
+bool UEdGraphSchema::TryCreateConnection(UEdGraphPin* PinA, UEdGraphPin* PinB) const
+{
+	const FPinConnectionResponse Response = CanCreateConnection(PinA, PinB);
+	bool bModified = false;
+
+	switch (Response.Response)
+	{
+	case CONNECT_RESPONSE_MAKE:
+		PinA->Modify();
+		PinB->Modify();
+		PinA->MakeLinkTo(PinB);
+		bModified = true;
+		break;
+
+	case CONNECT_RESPONSE_BREAK_OTHERS_A:
+		PinA->Modify();
+		PinB->Modify();
+		PinA->BreakAllPinLinks();
+		PinA->MakeLinkTo(PinB);
+		bModified = true;
+		break;
+
+	case CONNECT_RESPONSE_BREAK_OTHERS_B:
+		PinA->Modify();
+		PinB->Modify();
+		PinB->BreakAllPinLinks();
+		PinA->MakeLinkTo(PinB);
+		bModified = true;
+		break;
+
+	case CONNECT_RESPONSE_BREAK_OTHERS_AB:
+		PinA->Modify();
+		PinB->Modify();
+		PinA->BreakAllPinLinks();
+		PinB->BreakAllPinLinks();
+		PinA->MakeLinkTo(PinB);
+		bModified = true;
+		break;
+
+	case CONNECT_RESPONSE_MAKE_WITH_CONVERSION_NODE:
+		bModified = CreateAutomaticConversionNodeAndConnections(PinA, PinB);
+		break;
+		break;
+
+	case CONNECT_RESPONSE_DISALLOW:
+	default:
+		break;
+	}
+
+#if WITH_EDITOR
+	if (bModified)
+	{
+		PinA->GetOwningNode()->PinConnectionListChanged(PinA);
+		PinB->GetOwningNode()->PinConnectionListChanged(PinB);
+	}
+#endif	//#if WITH_EDITOR
+
+	return bModified;
+}
+
+bool UEdGraphSchema::CreateAutomaticConversionNodeAndConnections(UEdGraphPin* PinA, UEdGraphPin* PinB) const
+{
+	return false;
+}
+
+void UEdGraphSchema::TrySetDefaultValue(UEdGraphPin& Pin, const FString& NewDefaultValue) const
+{
+	Pin.DefaultValue = NewDefaultValue;
+
+#if WITH_EDITOR
+	UEdGraphNode* Node = Pin.GetOwningNode();
+	check(Node);
+	Node->PinDefaultValueChanged(&Pin);
+#endif	//#if WITH_EDITOR
+}
+
+void UEdGraphSchema::TrySetDefaultObject(UEdGraphPin& Pin, UObject* NewDefaultObject) const
+{
+	Pin.DefaultObject = NewDefaultObject;
+
+#if WITH_EDITOR
+	UEdGraphNode* Node = Pin.GetOwningNode();
+	check(Node);
+	Node->PinDefaultValueChanged(&Pin);
+#endif	//#if WITH_EDITOR
+}
+
+void UEdGraphSchema::TrySetDefaultText(UEdGraphPin& InPin, const FText& InNewDefaultText) const
+{
+	if(InNewDefaultText.IsEmpty())
+	{
+		InPin.DefaultTextValue = InNewDefaultText;
+	}
+	else
+	{
+#if WITH_EDITOR
+		InPin.DefaultTextValue = FText::ChangeKey(TEXT(""), InPin.GetOwningNode()->NodeGuid.ToString() + TEXT("_") + InPin.PinName, InNewDefaultText);
+#endif
+	}
+
+#if WITH_EDITOR
+	UEdGraphNode* Node = InPin.GetOwningNode();
+	check(Node);
+	Node->PinDefaultValueChanged(&InPin);
+#endif	//#if WITH_EDITOR
+}
+
+void UEdGraphSchema::BreakNodeLinks(UEdGraphNode& TargetNode) const
+{
+#if WITH_EDITOR
+	for (TArray<UEdGraphPin*>::TIterator PinIt(TargetNode.Pins); PinIt; ++PinIt)
+	{
+		BreakPinLinks(*(*PinIt), false);
+	}
+	TargetNode.NodeConnectionListChanged();
+#endif	//#if WITH_EDITOR
+}
+
+void UEdGraphSchema::BreakPinLinks(UEdGraphPin& TargetPin, bool bSendsNodeNotifcation) const
+{
+#if WITH_EDITOR
+	// Copy the old pin links
+	TArray<class UEdGraphPin*> OldLinkedTo(TargetPin.LinkedTo);
+#endif
+
+	TargetPin.BreakAllPinLinks();
+
+#if WITH_EDITOR
+	TSet<UEdGraphNode*> NodeList;
+	// Notify this node
+	TargetPin.GetOwningNode()->PinConnectionListChanged(&TargetPin);
+	NodeList.Add(TargetPin.GetOwningNode());
+
+	// As well as all other nodes that were connected
+	for (TArray<UEdGraphPin*>::TIterator PinIt(OldLinkedTo); PinIt; ++PinIt)
+	{
+		UEdGraphPin* OtherPin = *PinIt;
+		UEdGraphNode* OtherNode = OtherPin->GetOwningNode();
+		OtherNode->PinConnectionListChanged(OtherPin);
+		NodeList.Add(OtherNode);
+	}
+
+	if (bSendsNodeNotifcation)
+	{
+		// Send all nodes that received a new pin connection a notification
+		for (auto It = NodeList.CreateConstIterator(); It; ++It)
+		{
+			UEdGraphNode* Node = (*It);
+			Node->NodeConnectionListChanged();
+		}
+	}
+#endif	//#if WITH_EDITOR
+}
+
+void UEdGraphSchema::BreakSinglePinLink(UEdGraphPin* SourcePin, UEdGraphPin* TargetPin)
+{
+	SourcePin->BreakLinkTo(TargetPin);
+
+#if WITH_EDITOR
+	TargetPin->GetOwningNode()->PinConnectionListChanged(TargetPin);
+	SourcePin->GetOwningNode()->PinConnectionListChanged(SourcePin);
+
+	TargetPin->GetOwningNode()->NodeConnectionListChanged();
+	SourcePin->GetOwningNode()->NodeConnectionListChanged();
+#endif	//#if WITH_EDITOR
+}
+
+FPinConnectionResponse UEdGraphSchema::MovePinLinks(UEdGraphPin& MoveFromPin, UEdGraphPin& MoveToPin, bool bIsIntermeadiateMove) const
+{
+#if WITH_EDITOR
+	ensureMsg(bIsIntermeadiateMove || !MoveToPin.GetOwningNode()->GetGraph()->HasAnyFlags(RF_Transient),
+		TEXT("When moving to an intermeadiate pin, use FKismetCompilerContext::MovePinLinksToIntermediate() instead of UEdGraphSchema::MovePinLinks()"));
+#endif // #if WITH_EDITOR
+
+	FPinConnectionResponse FinalResponse = FPinConnectionResponse(CONNECT_RESPONSE_MAKE, TEXT(""));
+	// First copy the current set of links
+	TArray<UEdGraphPin*> CurrentLinks = MoveFromPin.LinkedTo;
+	// Then break all links at pin we are moving from
+	MoveFromPin.BreakAllPinLinks();
+	// Try and make each new connection
+	for (int32 i=0; i<CurrentLinks.Num(); i++)
+	{
+		UEdGraphPin* NewLink = CurrentLinks[i];
+		FPinConnectionResponse Response = CanCreateConnection(&MoveToPin, NewLink);
+		if(Response.CanSafeConnect())
+		{
+			MoveToPin.MakeLinkTo(NewLink);
+		}
+		else
+		{
+			FinalResponse = Response;
+		}
+	}
+	// Move over the default values
+	MoveToPin.DefaultValue = MoveFromPin.DefaultValue;
+	MoveToPin.DefaultObject = MoveFromPin.DefaultObject;
+	MoveToPin.DefaultTextValue = MoveFromPin.DefaultTextValue;
+	return FinalResponse;
+}
+
+FPinConnectionResponse UEdGraphSchema::CopyPinLinks(UEdGraphPin& CopyFromPin, UEdGraphPin& CopyToPin, bool bIsIntermeadiateCopy) const
+{
+#if WITH_EDITOR
+	ensureMsg(bIsIntermeadiateCopy || !CopyToPin.GetOwningNode()->GetGraph()->HasAnyFlags(RF_Transient),
+		TEXT("When copying to an intermeadiate pin, use FKismetCompilerContext::CopyPinLinksToIntermediate() instead of UEdGraphSchema::CopyPinLinks()"));
+#endif // #if WITH_EDITOR
+
+	FPinConnectionResponse FinalResponse = FPinConnectionResponse(CONNECT_RESPONSE_MAKE, TEXT(""));
+	for (int32 i=0; i<CopyFromPin.LinkedTo.Num(); i++)
+	{
+		UEdGraphPin* NewLink = CopyFromPin.LinkedTo[i];
+		FPinConnectionResponse Response = CanCreateConnection(&CopyToPin, NewLink);
+		if (Response.CanSafeConnect())
+		{
+			CopyToPin.MakeLinkTo(NewLink);
+		}
+		else
+		{
+			FinalResponse = Response;
+		}
+	}
+
+	CopyToPin.DefaultValue = CopyFromPin.DefaultValue;
+	CopyToPin.DefaultObject = CopyFromPin.DefaultObject;
+	CopyToPin.DefaultTextValue = CopyFromPin.DefaultTextValue;
+	return FinalResponse;
+}
+
+FString UEdGraphSchema::GetPinDisplayName(const UEdGraphPin* Pin) const
+{
+	check(Pin != NULL);
+	return !Pin->PinFriendlyName.IsEmpty() ? Pin->PinFriendlyName.ToString() : Pin->PinName;
+}
+
+void UEdGraphSchema::ConstructBasicPinTooltip(UEdGraphPin const& Pin, FString const& PinDescription, FString& TooltipOut) const
+{
+	TooltipOut = PinDescription;
+}
+
+void UEdGraphSchema::GetGraphContextActions(FGraphContextMenuBuilder& ContextMenuBuilder) const
+{
+#if WITH_EDITOR
+	// Run thru all nodes and add any menu items they want to add
+	for (TObjectIterator<UClass> ClassIt; ClassIt; ++ClassIt)
+	{
+		UClass* Class = *ClassIt;
+		if (Class->IsChildOf(UEdGraphNode::StaticClass()) && !Class->HasAnyClassFlags(CLASS_Abstract | CLASS_Deprecated))
+		{
+			const UEdGraphNode* ClassCDO = Class->GetDefaultObject<UEdGraphNode>();
+
+			if (ClassCDO->CanCreateUnderSpecifiedSchema(this))
+			{
+				ClassCDO->GetMenuEntries(ContextMenuBuilder);
+			}
+		}
+	}
+#endif
+}
+
+void UEdGraphSchema::ReconstructNode(UEdGraphNode& TargetNode, bool bIsBatchRequest/*=false*/) const
+{
+#if WITH_EDITOR
+	TargetNode.ReconstructNode();
+#endif	//#if WITH_EDITOR
+}
+
+void UEdGraphSchema::GetGraphDisplayInformation(const UEdGraph& Graph, /*out*/ FGraphDisplayInfo& DisplayInfo) const
+{
+	DisplayInfo.PlainName = FText::FromString( Graph.GetName() );
+	DisplayInfo.DisplayName = DisplayInfo.PlainName;
+}
+
+void UEdGraphSchema::GetContextMenuActions(const UEdGraph* CurrentGraph, const UEdGraphNode* InGraphNode, const UEdGraphPin* InGraphPin, class FMenuBuilder* MenuBuilder, bool bIsDebugging) const
+{
+#if WITH_EDITOR
+	FGraphNodeContextMenuBuilder Context(CurrentGraph, InGraphNode, InGraphPin, MenuBuilder, bIsDebugging);
+
+	if (Context.Node)
+	{
+		Context.Node->GetContextMenuActions(Context);
+	}
+
+	if (InGraphNode)
+	{
+		// Helper to do the node comment editing
+		struct Local
+		{
+			// Called by the EditableText widget to get the current comment for the node
+			static FString GetNodeComment(TWeakObjectPtr<UEdGraphNode> NodeWeakPtr)
+			{
+				if (UEdGraphNode* SelectedNode = NodeWeakPtr.Get())
+				{
+					return SelectedNode->NodeComment;
+				}
+				return FString();
+			}
+
+			// Called by the EditableText widget when the user types a new comment for the selected node
+			static void OnNodeCommentTextCommitted(const FText& NewText, ETextCommit::Type /*CommitInfo*/, TWeakObjectPtr<UEdGraphNode> NodeWeakPtr)
+			{
+				// Apply the change to the selected actor
+				UEdGraphNode* SelectedNode = NodeWeakPtr.Get();
+				FString NewString = NewText.ToString();
+				if (SelectedNode && !SelectedNode->NodeComment.Equals( NewString, ESearchCase::CaseSensitive ))
+				{
+					// send property changed events
+					const FScopedTransaction Transaction( LOCTEXT("EditNodeComment", "Change Node Comment") );
+					SelectedNode->Modify();
+					UProperty* NodeCommentProperty = FindField<UProperty>(SelectedNode->GetClass(), "NodeComment");
+					if(NodeCommentProperty != NULL)
+					{
+						SelectedNode->PreEditChange(NodeCommentProperty);
+
+						SelectedNode->NodeComment = NewString;
+
+						FPropertyChangedEvent NodeCommentPropertyChangedEvent(NodeCommentProperty);
+						SelectedNode->PostEditChangeProperty(NodeCommentPropertyChangedEvent);
+					}
+				}
+
+				FSlateApplication::Get().DismissAllMenus();
+			}
+		};
+
+		if( !InGraphPin )
+		{
+			int32 SelectionCount = GetNodeSelectionCount(CurrentGraph);
+		
+			if( SelectionCount == 1 )
+			{
+				// Node comment area
+				TSharedRef<SHorizontalBox> NodeCommentBox = SNew(SHorizontalBox);
+
+				MenuBuilder->BeginSection("GraphNodeComment", LOCTEXT("NodeCommentMenuHeader", "Node Comment"));
+				{
+					MenuBuilder->AddWidget(NodeCommentBox, FText::GetEmpty() );
+				}
+				TWeakObjectPtr<UEdGraphNode> SelectedNodeWeakPtr = InGraphNode;
+
+				FText NodeCommentText;
+				if ( UEdGraphNode* SelectedNode = SelectedNodeWeakPtr.Get() )
+				{
+					NodeCommentText = FText::FromString( SelectedNode->NodeComment );
+				}
+
+			const FSlateBrush* NodeIcon = FCoreStyle::Get().GetDefaultBrush();//@TODO: FActorIconFinder::FindIconForActor(SelectedActors(0).Get());
+
+
+				// Comment label
+				NodeCommentBox->AddSlot()
+				.VAlign(VAlign_Center)
+				.FillWidth( 1.0f )
+				[
+					SNew(SEditableTextBox)
+					.Text( NodeCommentText )
+					.ToolTipText(LOCTEXT("NodeComment_ToolTip", "Comment for this node"))
+					.OnTextCommitted_Static(&Local::OnNodeCommentTextCommitted, SelectedNodeWeakPtr)
+					.MinDesiredWidth( 120.f )
+					//.IsReadOnly( !SelectedActors(0)->IsActorLabelEditable() )
+					.SelectAllTextWhenFocused( true )
+					.RevertTextOnEscape( true )
+				];
+				MenuBuilder->EndSection();
+			}
+			else if( SelectionCount > 1 )
+			{
+				struct SCommentUtility
+				{
+					static void CreateComment(const UEdGraphSchema* Schema, UEdGraph* Graph)
+					{
+						if( Schema && Graph )
+						{
+							TSharedPtr<FEdGraphSchemaAction> Action = Schema->GetCreateCommentAction();
+
+							if( Action.IsValid() )
+							{
+								Action->PerformAction(Graph, NULL, FVector2D());
+							}
+						}
+					}
+				};
+
+				MenuBuilder->BeginSection("SchemaActionComment", LOCTEXT("MultiCommentHeader", "Comment Group"));
+				MenuBuilder->AddMenuEntry(	LOCTEXT("MultiCommentDesc", "Create Comment from Selection"),
+											LOCTEXT("CommentToolTip", "Create a resizable comment box around selection."),
+											FSlateIcon(), 
+											FUIAction( FExecuteAction::CreateStatic( SCommentUtility::CreateComment, this, const_cast<UEdGraph*>( CurrentGraph ))));
+				MenuBuilder->EndSection();
+			}
+		}
+	}
+#endif
+}
+
+#undef LOCTEXT_NAMESPACE
