* fixed name conflict in unity build

diff --git a/Engine/Source/Runtime/Engine/Private/StaticMeshRender.cpp b/Engine/Source/Runtime/Engine/Private/StaticMeshRender.cpp
--- a/Engine/Source/Runtime/Engine/Private/StaticMeshRender.cpp
+++ b/Engine/Source/Runtime/Engine/Private/StaticMeshRender.cpp
@@ -339,7 +339,7 @@ HHitProxy* FStaticMeshSceneProxy::CreateHitProxies(UPrimitiveComponent* Componen
 #endif // WITH_EDITOR
 
 // use for render thread only
-bool UseLightPropagationVolumeRT()
+bool UseLightPropagationVolumeRT2()
 {
 	if(!IsFeatureLevelSupported(GRHIShaderPlatform, ERHIFeatureLevel::SM5))
 	{
@@ -358,7 +358,7 @@ bool UseLightPropagationVolumeRT()
 inline bool AllowShadowOnlyMesh()
 {
 	// todo: later we should refine that (only if occlusion feature in LPV is on, only if inside a cascade, if shadow casting is disabled it should look at bUseEmissiveForDynamicAreaLighting)
-	return !UseLightPropagationVolumeRT();
+	return !UseLightPropagationVolumeRT2();
 }
 
 void FStaticMeshSceneProxy::DrawStaticElements(FStaticPrimitiveDrawInterface* PDI)
