* refactor: packing members

diff --git a/Engine/Source/Runtime/Engine/Public/SceneManagement.h b/Engine/Source/Runtime/Engine/Public/SceneManagement.h
--- a/Engine/Source/Runtime/Engine/Public/SceneManagement.h
+++ b/Engine/Source/Runtime/Engine/Public/SceneManagement.h
@@ -1302,29 +1302,38 @@ struct FMeshBatch
 {
 	TArray<FMeshBatchElement,TInlineAllocator<1> > Elements;
 
+	// used with DynamicVertexData
+	uint16 DynamicVertexStride;
+
+	/** LOD index of the mesh, used for fading LOD transitions. */
+	int8 LODIndex;
+
 	uint32 UseDynamicData : 1;
 	uint32 ReverseCulling : 1;
 	uint32 bDisableBackfaceCulling : 1;
 	uint32 CastShadow : 1;
 	uint32 bWireframe : 1;
+	// e.g. PT_TriangleList(default), PT_LineList, ..
 	uint32 Type : PT_NumBits;
+	// e.g. SDPG_World (default), SDPG_Foreground
 	uint32 DepthPriorityGroup : SDPG_NumBits;
 	uint32 bUseAsOccluder : 1;
 
+	// can be 0
 	const FLightCacheInterface* LCI;
 
 	/** 
 	 *	DynamicVertexData - pointer to user memory containing the vertex data.
 	 *	Used for rendering dynamic data directly.
+	 *  used with DynamicVertexStride
 	 */
 	const void* DynamicVertexData;
-	uint16 DynamicVertexStride;
 
+	// can be 0
 	const FVertexFactory* VertexFactory;
+	// can be 0
 	const FMaterialRenderProxy* MaterialRenderProxy;
 
-	/** LOD index of the mesh, used for fading LOD transitions. */
-	int8 LODIndex;
 
 	FORCEINLINE bool IsTranslucent() const
 	{
@@ -1382,18 +1391,20 @@ struct FMeshBatch
 
 	/** Default constructor. */
 	FMeshBatch()
-	:	UseDynamicData(false)
+	:	DynamicVertexStride(0)
+	,	LODIndex(INDEX_NONE)
+	,	UseDynamicData(false)
 	,	ReverseCulling(false)
 	,	bDisableBackfaceCulling(false)
 	,	CastShadow(true)
 	,	bWireframe(false)
 	,	Type(PT_TriangleList)
+	,	DepthPriorityGroup(SDPG_World)
 	,	bUseAsOccluder(true)
 	,	LCI(NULL)
 	,	DynamicVertexData(NULL)
 	,	VertexFactory(NULL)
 	,	MaterialRenderProxy(NULL)
-	,	LODIndex(INDEX_NONE)
 	{
 		// By default always add the first element.
 		new(Elements) FMeshBatchElement;
