Fix it so vehicle wheel with default body automatically fixes it up for the vehicle's instance

diff --git a/Engine/Source/Runtime/Engine/Private/PhysicsEngine/BodyInstance.cpp b/Engine/Source/Runtime/Engine/Private/PhysicsEngine/BodyInstance.cpp
--- a/Engine/Source/Runtime/Engine/Private/PhysicsEngine/BodyInstance.cpp
+++ b/Engine/Source/Runtime/Engine/Private/PhysicsEngine/BodyInstance.cpp
@@ -1410,7 +1410,7 @@ bool FBodyInstance::UpdateBodyScale(const FVector & inScale3D)
 #endif
 }
 
-void FBodyInstance::UpdateInstanceSimulatePhysics()
+void FBodyInstance::UpdateInstanceSimulatePhysics(bool bIgnoreParent)
 {
 #if WITH_PHYSX
 	PxRigidDynamic* PRigidDynamic = GetPxRigidDynamic();
@@ -1440,7 +1440,7 @@ void FBodyInstance::UpdateInstanceSimulatePhysics()
 }
 
 
-void FBodyInstance::SetInstanceSimulatePhysics(bool bSimulate, bool bMaintainPhysicsBlending)
+void FBodyInstance::SetInstanceSimulatePhysics(bool bSimulate, bool bMaintainPhysicsBlending, bool bIgnoreOwner)
 {
 #if WITH_PHYSX
 #if !(UE_BUILD_SHIPPING || UE_BUILD_TEST)
@@ -1480,19 +1480,19 @@ void FBodyInstance::SetInstanceSimulatePhysics(bool bSimulate, bool bMaintainPhy
 		}
 	}
 
-	UpdateInstanceSimulatePhysics();
+	UpdateInstanceSimulatePhysics(bIgnoreOwner);
 }
 
-bool FBodyInstance::IsInstanceSimulatingPhysics()
+bool FBodyInstance::IsInstanceSimulatingPhysics(bool bIgnoreOwner)
 {
 	// if I'm simulating or owner is simulating
-	return ShouldInstanceSimulatingPhysics() && IsValidBodyInstance();
+	return ShouldInstanceSimulatingPhysics(bIgnoreOwner) && IsValidBodyInstance();
 }
 
-bool FBodyInstance::ShouldInstanceSimulatingPhysics()
+bool FBodyInstance::ShouldInstanceSimulatingPhysics(bool bIgnoreOwner)
 {
 	// if I'm simulating or owner is simulating
-	if ( BodySetup.IsValid() && BodySetup.Get()->PhysicsType == PhysType_Default )
+	if (BodySetup.IsValid() && BodySetup.Get()->PhysicsType == PhysType_Default && !bIgnoreOwner)
 	{
 		// if derive from owner, and owner is simulating, this should simulate
 		if (OwnerComponent != NULL && OwnerComponent->BodyInstance.bSimulatePhysics)
diff --git a/Engine/Source/Runtime/Engine/Private/PhysicsEngine/PhysAnim.cpp b/Engine/Source/Runtime/Engine/Private/PhysicsEngine/PhysAnim.cpp
--- a/Engine/Source/Runtime/Engine/Private/PhysicsEngine/PhysAnim.cpp
+++ b/Engine/Source/Runtime/Engine/Private/PhysicsEngine/PhysAnim.cpp
@@ -106,12 +106,10 @@ void USkeletalMeshComponent::BlendPhysicsBones( TArray<FBoneIndexType>& InRequir
 			}
 #endif
 			BodyInstance = Bodies[BodyIndex];
-			UBodySetup* BodySetup = PhysicsAsset->BodySetup[BodyIndex];
 			// since we don't copy back to physics, we shouldn't blend fixed bones here if this set up is used
 			// if you'd like to use fixed bones to blend use SkipSimulatingBones
 			// you can simulate fixed bones by setting it to instance, but that case, please use SkipSimulating Bones. 
-			bool bSkipFixedBones =  bShouldSkipFixedBones
-				&& (BodySetup? BodySetup->PhysicsType == PhysType_Fixed : false);
+			bool bSkipFixedBones = bShouldSkipFixedBones && !BodyInstance->bSimulatePhysics;
 
 			if(!bSkipFixedBones && BodyInstance->IsValidBodyInstance())
 			{
@@ -305,11 +303,8 @@ void USkeletalMeshComponent::UpdateKinematicBonesToPhysics(bool bTeleport)
 			// If we have a physics body, and its kinematic...
 			FBodyInstance* BodyInst = Bodies[i];
 			check(BodyInst);
-			// special flag to check to see if we should update fixed bones or not
-			bool bSkipFixedBones =  bShouldSkipFixedBones
-				&& (BodyInst->BodySetup.IsValid()? BodyInst->BodySetup.Get()->PhysicsType == PhysType_Fixed : false);
 
-			if(!bSkipFixedBones && BodyInst->IsValidBodyInstance() && !BodyInst->IsInstanceSimulatingPhysics())
+			if (!bShouldSkipFixedBones && BodyInst->IsValidBodyInstance() && !BodyInst->IsInstanceSimulatingPhysics(true))
 			{
 				// Find the graphics bone index that corresponds to this physics body.
 				FName const BodyName = PhysicsAsset->BodySetup[i]->BoneName;
diff --git a/Engine/Source/Runtime/Engine/Private/Vehicles/WheeledVehicleMovementComponent.cpp b/Engine/Source/Runtime/Engine/Private/Vehicles/WheeledVehicleMovementComponent.cpp
--- a/Engine/Source/Runtime/Engine/Private/Vehicles/WheeledVehicleMovementComponent.cpp
+++ b/Engine/Source/Runtime/Engine/Private/Vehicles/WheeledVehicleMovementComponent.cpp
@@ -1230,16 +1230,17 @@ void UWheeledVehicleMovementComponent::FixupSkeletalMesh()
 					
 					if (BodySetupIdx >= 0)
 					{
+						FBodyInstance * BodyInstance = Mesh->Bodies[BodySetupIdx];
+						BodyInstance->SetResponseToAllChannels(ECR_Ignore);	//turn off collision for wheel automatically
+
 						if (UBodySetup * BodySetup = PhysicsAsset->BodySetup[BodySetupIdx])
 						{
-							if (BodySetup[BodySetupIdx].PhysicsType != PhysType_Fixed)
+
+							if (BodySetup->PhysicsType == PhysType_Default) 	//if they set it to unfixed we don't fixup because they are explicitely saying Unfixed
 							{
-								UE_LOG(LogPhysics, Warning, TEXT("UWheeledVehicleMovementComponent::FixupSkeletalMesh Wheel [%d] has a body that is not set to Fixed."), WheelIdx);
-								FMessageLog("PIE").Warning(FText::Format(LOCTEXT("WheelHasBodyWarning", "Wheel [{0}] is using a body which is not set to Fixed. Please change this in the physics asset."), FText::FromString(WheelSetup.BoneName.ToString())));
+								BodyInstance->SetInstanceSimulatePhysics(false, false, true);
 							}
 
-							Mesh->Bodies[BodySetupIdx]->SetResponseToAllChannels(ECR_Ignore);	//turn off collision for wheel automatically
-
 							//and get rid of constraints on the wheels. TODO: right now we remove all wheel constraints, we probably only want to remove parent constraints
 							TArray<int32> WheelConstraints;
 							PhysicsAsset->BodyFindConstraints(BodySetupIdx, WheelConstraints);
