Fix build

diff --git a/Engine/Source/Runtime/Engine/Classes/PhysicsEngine/BodyInstance.h b/Engine/Source/Runtime/Engine/Classes/PhysicsEngine/BodyInstance.h
--- a/Engine/Source/Runtime/Engine/Classes/PhysicsEngine/BodyInstance.h
+++ b/Engine/Source/Runtime/Engine/Classes/PhysicsEngine/BodyInstance.h
@@ -299,13 +299,13 @@ public:
 
 
 		/** Set this body to be fixed (kinematic) or not. */
-		void SetInstanceSimulatePhysics(bool bSimulate, bool bMaintainPhysicsBlending=false);
+		void SetInstanceSimulatePhysics(bool bSimulate, bool bMaintainPhysicsBlending=false, bool bIgnoreOwner=false);
 		/** Makes sure the current kinematic state matches the simulate flag */
-		void UpdateInstanceSimulatePhysics();
+		void UpdateInstanceSimulatePhysics(bool bIgnoreOwner = false);
 		/** Returns true if this body is simulating, false if it is fixed (kinematic) */
-		bool IsInstanceSimulatingPhysics();
+		bool IsInstanceSimulatingPhysics(bool bIgnoreOwner = false);
 		/** Should Simulate Physics **/
-		bool ShouldInstanceSimulatingPhysics();
+		bool ShouldInstanceSimulatingPhysics(bool bIgnoreOwner=false);
 		/** Returns whether this body is awake */
 		bool IsInstanceAwake();
 		/** Wake this body */
