Only unregister if we are changing the outer

diff --git a/Engine/Source/Runtime/Engine/Private/Actor.cpp b/Engine/Source/Runtime/Engine/Private/Actor.cpp
--- a/Engine/Source/Runtime/Engine/Private/Actor.cpp
+++ b/Engine/Source/Runtime/Engine/Private/Actor.cpp
@@ -677,8 +677,11 @@ bool AActor::IsActorTickEnabled() const
 
 bool AActor::Rename( const TCHAR* InName, UObject* NewOuter, ERenameFlags Flags )
 {
-	RegisterAllActorTickFunctions(false, true); // unregister all tick functions
-	UnregisterAllComponents();
+	if (NewOuter)
+	{
+		RegisterAllActorTickFunctions(false, true); // unregister all tick functions
+		UnregisterAllComponents();
+	}
 
 	bool bSuccess = Super::Rename( InName, NewOuter, Flags );
 
