Fixed initialization order

diff --git a/Engine/Source/Runtime/Engine/Private/SpeedTreeWind.cpp b/Engine/Source/Runtime/Engine/Private/SpeedTreeWind.cpp
--- a/Engine/Source/Runtime/Engine/Private/SpeedTreeWind.cpp
+++ b/Engine/Source/Runtime/Engine/Private/SpeedTreeWind.cpp
@@ -77,23 +77,23 @@ FSpeedTreeWind::SParams::SParams( ) :
 	m_fGlobalHeightExponent(2.0f),
 	m_fFrondRippleTile(10.0f),
 	m_fFrondRippleLightingScalar(1.0f),
-	m_fGustFrequency(0.0f),
-	m_fGustStrengthMin(0.5f),
-	m_fGustStrengthMax(1.0f),
-	m_fGustDurationMin(1.0f),
-	m_fGustDurationMax(4.0f),
-	m_fGustRiseScalar(1.0f),
-	m_fGustFallScalar(1.0f),
 	m_fRollingNoiseSize(0.005f),
 	m_fRollingNoiseTwist(9.0f),
 	m_fRollingNoiseTurbulence(32.0f),
 	m_fRollingNoisePeriod(0.4f),
 	m_fRollingNoiseSpeed(0.05f),
-	m_fRollingLeafRippleMin(0.5f),
-	m_fRollingLeafTumbleMin(0.5f),
 	m_fRollingBranchFieldMin(0.5f),
 	m_fRollingBranchLightingAdjust(0.5f),
-	m_fRollingBranchVerticalOffset(-0.5f)
+	m_fRollingBranchVerticalOffset(-0.5f),
+	m_fRollingLeafRippleMin(0.5f),
+	m_fRollingLeafTumbleMin(0.5f),
+	m_fGustFrequency(0.0f),
+	m_fGustStrengthMin(0.5f),
+	m_fGustStrengthMax(1.0f),
+	m_fGustDurationMin(1.0f),
+	m_fGustDurationMax(4.0f),
+	m_fGustRiseScalar(1.0f),
+	m_fGustFallScalar(1.0f)
 {
 	for (int32 i = 0; i < FSpeedTreeWind::NUM_OSC_COMPONENTS; ++i)
 		for (int32 j = 0; j < FSpeedTreeWind::NUM_WIND_POINTS_IN_CURVE; ++j)
