#ue4 - Fix recent changes to CanJump() causing BP compile errors for existing nodes.

- Changed CanJump() back to a BP callable function, and internally wrap a call to a BP native event, continuing to allow an overridden implementation.

diff --git a/Engine/Source/Runtime/Engine/Classes/GameFramework/Character.h b/Engine/Source/Runtime/Engine/Classes/GameFramework/Character.h
--- a/Engine/Source/Runtime/Engine/Classes/GameFramework/Character.h
+++ b/Engine/Source/Runtime/Engine/Classes/GameFramework/Character.h
@@ -134,7 +134,7 @@ struct FRepRelativeMovement
 // to a CharacterMovementComponent that handles movement of the collision capsule, and they
 // also have implementations of basic networking and input models.
 //=============================================================================
-UCLASS(abstract, config=Game, dependson=(AController, UCharacterMovementComponent), BlueprintType)
+UCLASS(abstract, config=Game, dependson=(AController, UCharacterMovementComponent), BlueprintType, hidecategories=("Pawn|Character|InternalEvents"))
 class ENGINE_API ACharacter : public APawn
 {
 	GENERATED_UCLASS_BODY()
@@ -307,6 +307,18 @@ public:
 
 	/**
 	 * Check if the character can jump in the current state.
+	 *
+	 * The default implementation may be overridden or extended by implementing the custom CanJump event.
+	 * 
+	 * @Return Whether the character can jump in the current state. 
+	 */
+	UFUNCTION(BlueprintCallable, Category="Pawn|Character")
+	bool CanJump() const;
+
+protected:
+
+	/**
+	 * Customizable event to check if the character can jump in the current state.
 	 * Default implementation returns true if the character is on the ground and not crouching,
 	 * has a valid CharacterMovementComponent and CanEverJump() returns true.
 	 * Default implementation also allows for 'hold to jump higher' functionality: 
@@ -315,8 +327,11 @@ public:
 	 * 
 	 * @Return Whether the character can jump in the current state. 
 	 */
-	UFUNCTION(BlueprintNativeEvent, Category="Pawn|Character")
-	bool CanJump() const;
+
+	UFUNCTION(BlueprintNativeEvent, Category="Pawn|Character|InternalEvents", meta=(FriendlyName="CanJump"))
+	bool CanJumpInternal() const;
+
+public:
 
 	/** True if currently jumping; i.e. jump key is held and the time it has been held is less than JumpMaxHoldTime */
 	UFUNCTION(BlueprintCallable, Category="Pawn|Character")
diff --git a/Engine/Source/Runtime/Engine/Private/Character.cpp b/Engine/Source/Runtime/Engine/Private/Character.cpp
--- a/Engine/Source/Runtime/Engine/Private/Character.cpp
+++ b/Engine/Source/Runtime/Engine/Private/Character.cpp
@@ -194,7 +194,12 @@ void ACharacter::Landed(const FHitResult& Hit)
 	OnLanded(Hit);
 }
 
-bool ACharacter::CanJump_Implementation() const
+bool ACharacter::CanJump() const
+{
+	return CanJumpInternal();
+}
+
+bool ACharacter::CanJumpInternal_Implementation() const
 {
 	const bool bCanHoldToJumpHigher = (GetJumpMaxHoldTime() > 0.0f) && IsJumping();
 
