UE4: Fixed a crash that would happen when you used a console command to play a sound, waited for a GC, then used it again.

diff --git a/Engine/Source/Runtime/Engine/Private/AudioDevice.cpp b/Engine/Source/Runtime/Engine/Private/AudioDevice.cpp
--- a/Engine/Source/Runtime/Engine/Private/AudioDevice.cpp
+++ b/Engine/Source/Runtime/Engine/Private/AudioDevice.cpp
@@ -583,23 +583,24 @@ bool FAudioDevice::HandleListSoundDurationsCommand( const TCHAR* Cmd, FOutputDev
 bool FAudioDevice::HandlePlaySoundCueCommand( const TCHAR* Cmd, FOutputDevice& Ar )
 {
 	// Stop any existing sound playing
-	if( TestAudioComponent == NULL )
+	if( !TestAudioComponent.IsValid() )
 	{
 		TestAudioComponent = ConstructObject<UAudioComponent>( UAudioComponent::StaticClass() );
 	}
 
-	if( TestAudioComponent != NULL )
+	UAudioComponent* AudioComp = TestAudioComponent.Get();
+	if( AudioComp != NULL )
 	{
-		TestAudioComponent->Stop();
+		AudioComp->Stop();
 
 		// Load up an arbitrary cue
 		USoundCue* Cue = LoadObject<USoundCue>( NULL, Cmd, NULL, LOAD_None, NULL );
 		if( Cue != NULL )
 		{
-			TestAudioComponent->Sound = Cue;
-			TestAudioComponent->bAllowSpatialization = false;
-			TestAudioComponent->bAutoDestroy = true;
-			TestAudioComponent->Play();
+			AudioComp->Sound = Cue;
+			AudioComp->bAllowSpatialization = false;
+			AudioComp->bAutoDestroy = true;
+			AudioComp->Play();
 
 			TArray<USoundNodeWavePlayer*> WavePlayers;
 			Cue->RecursiveFindNode<USoundNodeWavePlayer>( Cue->FirstNode, WavePlayers );
@@ -619,23 +620,24 @@ bool FAudioDevice::HandlePlaySoundCueCommand( const TCHAR* Cmd, FOutputDevice& A
 bool FAudioDevice::HandlePlaySoundWaveCommand( const TCHAR* Cmd, FOutputDevice& Ar )
 {
 	// Stop any existing sound playing
-	if( TestAudioComponent == NULL )
+	if( !TestAudioComponent.IsValid() )
 	{
 		TestAudioComponent = ConstructObject<UAudioComponent>( UAudioComponent::StaticClass() );
 	}
 
-	if( TestAudioComponent != NULL )
+	UAudioComponent* AudioComp = TestAudioComponent.Get();
+	if( AudioComp != NULL )
 	{
-		TestAudioComponent->Stop();
+		AudioComp->Stop();
 
 		// Load up an arbitrary wave
 		USoundWave* Wave = LoadObject<USoundWave>( NULL, Cmd, NULL, LOAD_None, NULL );
 		if( Wave != NULL )
 		{
-			TestAudioComponent->Sound = Wave;
-			TestAudioComponent->bAllowSpatialization = false;
-			TestAudioComponent->bAutoDestroy = true;
-			TestAudioComponent->Play();
+			AudioComp->Sound = Wave;
+			AudioComp->bAllowSpatialization = false;
+			AudioComp->bAutoDestroy = true;
+			AudioComp->Play();
 
 			Wave->LogSubtitle( Ar );
 		}
diff --git a/Engine/Source/Runtime/Engine/Public/AudioDevice.h b/Engine/Source/Runtime/Engine/Public/AudioDevice.h
--- a/Engine/Source/Runtime/Engine/Public/AudioDevice.h
+++ b/Engine/Source/Runtime/Engine/Public/AudioDevice.h
@@ -721,7 +721,7 @@ public:
 	uint64 CurrentTick;
 
 	/** An AudioComponent to play test sounds on */
-	class UAudioComponent* TestAudioComponent;
+	TWeakObjectPtr<class UAudioComponent> TestAudioComponent;
 
 	/** The debug state of the audio device */
 	TEnumAsByte<enum EDebugState> DebugState;
