Added a way to have BP-implementable EQS contexts #UE4

diff --git a/Engine/Source/Runtime/Engine/Classes/AI/EnvironmentQuery/Items/EnvQueryItemType_Actor.h b/Engine/Source/Runtime/Engine/Classes/AI/EnvironmentQuery/Items/EnvQueryItemType_Actor.h
--- a/Engine/Source/Runtime/Engine/Classes/AI/EnvironmentQuery/Items/EnvQueryItemType_Actor.h
+++ b/Engine/Source/Runtime/Engine/Classes/AI/EnvironmentQuery/Items/EnvQueryItemType_Actor.h
@@ -13,6 +13,7 @@ class ENGINE_API UEnvQueryItemType_Actor : public UEnvQueryItemType_ActorBase
 
 	static void SetContextHelper(struct FEnvQueryContextData& ContextData, const AActor* SingleActor);
 	static void SetContextHelper(struct FEnvQueryContextData& ContextData, const TArray<const AActor*>& MultipleActors);
+	static void SetContextHelper(struct FEnvQueryContextData& ContextData, const TArray<AActor*>& MultipleActors);
 
 	virtual FVector GetLocation(const uint8* RawData) const OVERRIDE;
 	virtual FRotator GetRotation(const uint8* RawData) const OVERRIDE;
diff --git a/Engine/Source/Runtime/Engine/Private/AI/EnvironmentQuery/Items/EnvQueryItemType_Actor.cpp b/Engine/Source/Runtime/Engine/Private/AI/EnvironmentQuery/Items/EnvQueryItemType_Actor.cpp
--- a/Engine/Source/Runtime/Engine/Private/AI/EnvironmentQuery/Items/EnvQueryItemType_Actor.cpp
+++ b/Engine/Source/Runtime/Engine/Private/AI/EnvironmentQuery/Items/EnvQueryItemType_Actor.cpp
@@ -52,9 +52,23 @@ void UEnvQueryItemType_Actor::SetContextHelper(struct FEnvQueryContextData& Cont
 	ContextData.RawData.Init(sizeof(FWeakObjectPtr) * MultipleActors.Num());
 
 	uint8* RawData = (uint8*)ContextData.RawData.GetTypedData();
-	for (int32 i = 0; i < MultipleActors.Num(); i++)
+	for (int32 ActorIndex = 0; ActorIndex < MultipleActors.Num(); ActorIndex++)
 	{
-		UEnvQueryItemType_Actor::SetValue(RawData, MultipleActors[i]);
+		UEnvQueryItemType_Actor::SetValue(RawData, MultipleActors[ActorIndex]);
+		RawData += sizeof(FWeakObjectPtr);
+	}
+}
+
+void UEnvQueryItemType_Actor::SetContextHelper(struct FEnvQueryContextData& ContextData, const TArray<AActor*>& MultipleActors)
+{
+	ContextData.ValueType = UEnvQueryItemType_Actor::StaticClass();
+	ContextData.NumValues = MultipleActors.Num();
+	ContextData.RawData.Init(sizeof(FWeakObjectPtr)* MultipleActors.Num());
+
+	uint8* RawData = (uint8*)ContextData.RawData.GetTypedData();
+	for (int32 ActorIndex = 0; ActorIndex < MultipleActors.Num(); ActorIndex++)
+	{
+		UEnvQueryItemType_Actor::SetValue(RawData, MultipleActors[ActorIndex]);
 		RawData += sizeof(FWeakObjectPtr);
 	}
 }
