CharacterMovementComponent: Added AddForce, changed AddMomentum to AddImpulse

Changed API functions around force & impulse accumulation to be more in-line with physics-based counterparts.

diff --git a/Engine/Config/BaseEngine.ini b/Engine/Config/BaseEngine.ini
--- a/Engine/Config/BaseEngine.ini
+++ b/Engine/Config/BaseEngine.ini
@@ -829,6 +829,9 @@ bUseStreamingPause=false
 +K2FieldRedirects=(OldFieldName="CharacterMovementComponent.WaterSpeed",NewPropertyName="CharacterMovementComponent.MaxSwimSpeed")
 +K2FieldRedirects=(OldFieldName="CharacterMovementComponent.CrouchedSpeedPercent",NewPropertyName="CharacterMovementComponent.CrouchedSpeedMultiplier")
 +K2FieldRedirects=(OldFieldName="CharacterMovementComponent.bOrientToMovement",NewPropertyName="CharacterMovementComponent.bOrientRotationToMovement")
++K2FieldRedirects=(OldFieldName="CharacterMovementComponent.AddMomentum",NewFieldName="CharacterMovementComponent.AddImpulse")
++K2ParamRedirects=(NodeName="/Script/BlueprintGraph.K2Node_CallFunction", OldParamName="CharacterMovementComponent.AddImpulse.InMomentum", NewParamName="Impulse")
++K2ParamRedirects=(NodeName="/Script/BlueprintGraph.K2Node_CallFunction", OldParamName="CharacterMovementComponent.AddImpulse.bMassIndependent", NewParamName="bVelocityChange")
 
 +K2FieldRedirects=(OldFieldName="PlayerController.LookRightScale",NewFieldName="PlayerController.InputYawScale")
 +K2FieldRedirects=(OldFieldName="PlayerController.LookUpScale",NewFieldName="PlayerController.InputPitchScale")
diff --git a/Engine/Source/Runtime/Engine/Classes/GameFramework/CharacterMovementComponent.h b/Engine/Source/Runtime/Engine/Classes/GameFramework/CharacterMovementComponent.h
--- a/Engine/Source/Runtime/Engine/Classes/GameFramework/CharacterMovementComponent.h
+++ b/Engine/Source/Runtime/Engine/Classes/GameFramework/CharacterMovementComponent.h
@@ -381,9 +381,13 @@ protected:
 	UPROPERTY()
 	FVector Acceleration;
 
-	/** Accumulated momentum added this tick */
+	/** Accumulated impulse added this tick */
 	UPROPERTY()
-	FVector PendingMomentumToApply;
+	FVector PendingImpulseToApply;
+
+	/** Accumulated force added this tick */
+	UPROPERTY()
+	FVector PendingForceToApply;
 
 	/**
 	 * Modifier to applied to values such as acceleration and max speed due to analog input.
@@ -812,7 +816,7 @@ public:
 	void ApplyRepulsionForce(float DeltaTime);
 	
 	/** Applies momentum accumulated through Addmomentum() */
-	void ApplyAccumulatedMomentum(float DeltaSeconds);	
+	void ApplyAccumulatedForces(float DeltaSeconds);	
 
 	/** 
 	 * Handle start swimming functionality
@@ -901,15 +905,30 @@ public:
 	virtual FString GetMovementName();
 
 	/** 
-	 * Add velocity based on imparted momentum. Momentum is accumulated each tick and applied together
+	 * Add impulse to character. Impulses are accumulated each tick and applied together
+	 * so multiple calls to this function will accumulate.
+	 * An impulse is an instantaneous force, usually applied once. If you want to continually apply
+	 * forces each frame, use AddForce().
+	 * Note that changing the momentum of characters like this can change the movement mode.
+	 * 
+	 * @param	Impulse				Impulse to apply.
+	 * @param	bVelocityChange		Whether or not the impulse is relative to mass.
+	 */
+	UFUNCTION(BlueprintCallable, Category="Pawn|Components|CharacterMovement")
+	virtual void AddImpulse( FVector Impulse, bool bVelocityChange = false );
+
+	/** 
+	 * Add force to character. Forces are accumulated each tick and applied together
 	 * so multiple calls to this function will accumulate.
-	 * Note that changing the momentum of characters like this can change the movement mode
+	 * Forces are scaled depending on timestep, so they can be applied each frame. If you want an
+	 * instantaneous force, use AdddImpulse.
+	 * Adding a force always takes the actor's mass into account.
+	 * Note that changing the momentum of characters like this can change the movement mode.
 	 * 
-	 * @param	Momentum			Momentum to apply.
-	 * @param	bMassIndependent	Whether or not the momentum should be divided by mass before application.
+	 * @param	Force			Force to apply.
 	 */
 	UFUNCTION(BlueprintCallable, Category="Pawn|Components|CharacterMovement")
-	virtual void AddMomentum( FVector InMomentum, bool bMassIndependent = false );
+	virtual void AddForce( FVector Force );
 
 	/**
 	 * Draw important variables on canvas.  Character will call DisplayDebug() on the current ViewTarget when the ShowDebug exec is used
diff --git a/Engine/Source/Runtime/Engine/Private/Character.cpp b/Engine/Source/Runtime/Engine/Private/Character.cpp
--- a/Engine/Source/Runtime/Engine/Private/Character.cpp
+++ b/Engine/Source/Runtime/Engine/Private/Character.cpp
@@ -347,7 +347,7 @@ void ACharacter::ApplyDamageMomentum(float DamageTaken, FDamageEvent const& Dama
 			}
 		}
 
-		CharacterMovement->AddMomentum(Impulse, bMassIndependentImpulse);
+		CharacterMovement->AddImpulse(Impulse, bMassIndependentImpulse);
 	}
 }
 
diff --git a/Engine/Source/Runtime/Engine/Private/CharacterMovementComponent.cpp b/Engine/Source/Runtime/Engine/Private/CharacterMovementComponent.cpp
--- a/Engine/Source/Runtime/Engine/Private/CharacterMovementComponent.cpp
+++ b/Engine/Source/Runtime/Engine/Private/CharacterMovementComponent.cpp
@@ -1163,7 +1163,7 @@ void UCharacterMovementComponent::PerformMovement(float DeltaSeconds)
 
 	UpdateBasedMovement(DeltaSeconds);
 
-	ApplyAccumulatedMomentum(DeltaSeconds);
+	ApplyAccumulatedForces(DeltaSeconds);
 
 	FVector OldVelocity = Velocity;
 	FVector OldLocation = CharacterOwner->GetActorLocation();
@@ -3447,25 +3447,40 @@ bool UCharacterMovementComponent::CheckWaterJump(FVector CheckPoint, FVector& Wa
 	return false;
 }
 
-void UCharacterMovementComponent::AddMomentum( FVector InMomentum, bool bMassIndependent )
+void UCharacterMovementComponent::AddImpulse( FVector Impulse, bool bVelocityChange )
 {
-	if (HasValidData() && !InMomentum.IsZero())
+	if (HasValidData() && !Impulse.IsZero())
 	{
 		// handle scaling by mass
-		FVector FinalMomentum = InMomentum;
-		if ( !bMassIndependent)
+		FVector FinalImpulse = Impulse;
+		if ( !bVelocityChange )
 		{
 			if (Mass > SMALL_NUMBER)
 			{
-				FinalMomentum = FinalMomentum / Mass;
+				FinalImpulse = FinalImpulse / Mass;
 			}
 			else
 			{
-				UE_LOG(LogCharacterMovement, Warning, TEXT("Attempt to apply momentum to zero or negative Mass in CharacterMovement"));
+				UE_LOG(LogCharacterMovement, Warning, TEXT("Attempt to apply impulse to zero or negative Mass in CharacterMovement"));
 			}
 		}
 
-		PendingMomentumToApply += FinalMomentum;
+		PendingImpulseToApply += FinalImpulse;
+	}
+}
+
+void UCharacterMovementComponent::AddForce( FVector Force )
+{
+	if (HasValidData() && !Force.IsZero())
+	{
+		if (Mass > SMALL_NUMBER)
+		{
+			PendingForceToApply += Force / Mass;
+		}
+		else
+		{
+			UE_LOG(LogCharacterMovement, Warning, TEXT("Attempt to apply force to zero or negative Mass in CharacterMovement"));
+		}
 	}
 }
 
@@ -5564,22 +5579,41 @@ void UCharacterMovementComponent::ApplyRepulsionForce( float DeltaTime )
 	}
 }
 
-void UCharacterMovementComponent::ApplyAccumulatedMomentum(float DeltaSeconds)
+void UCharacterMovementComponent::ApplyAccumulatedForces(float DeltaSeconds)
 {
 	// check to see if applied momentum is enough to overcome gravity
-	if ( IsMovingOnGround() && (PendingMomentumToApply.Z + (GetGravityZ() * DeltaSeconds) > SMALL_NUMBER) )
+	if ( IsMovingOnGround() && (PendingImpulseToApply.Z + (PendingForceToApply.Z * DeltaSeconds) + (GetGravityZ() * DeltaSeconds) > SMALL_NUMBER))
 	{
 		SetMovementMode(MOVE_Falling);
 	}
 
-	Velocity += PendingMomentumToApply;
+	Velocity += PendingImpulseToApply;
+	Velocity += PendingForceToApply * DeltaSeconds;
 
-	PendingMomentumToApply = FVector::ZeroVector;
+	PendingImpulseToApply = FVector::ZeroVector;
+	PendingForceToApply = FVector::ZeroVector;
 }
 
 void UCharacterMovementComponent::AddRadialForce(const FVector& Origin, float Radius, float Strength, enum ERadialImpulseFalloff Falloff)
 {
-	AddRadialImpulse(Origin, Radius, Strength, Falloff, true);
+	FVector Delta = UpdatedComponent->GetComponentLocation() - Origin;
+	const float DeltaMagnitude = Delta.Size();
+
+	// Do nothing if outside radius
+	if(DeltaMagnitude > Radius)
+	{
+		return;
+	}
+
+	Delta = Delta.SafeNormal();
+
+	float ForceMagnitude = Strength;
+	if (Falloff == RIF_Linear && Radius > 0.0f)
+	{
+		ForceMagnitude *= (1.0f - (DeltaMagnitude / Radius));
+	}
+
+	AddForce(Delta * ForceMagnitude);
 }
  
 void UCharacterMovementComponent::AddRadialImpulse(const FVector& Origin, float Radius, float Strength, enum ERadialImpulseFalloff Falloff, bool bVelChange)
@@ -5601,7 +5635,7 @@ void UCharacterMovementComponent::AddRadialImpulse(const FVector& Origin, float
 		ImpulseMagnitude *= (1.0f - (DeltaMagnitude / Radius));
 	}
 
-	AddMomentum(Delta * ImpulseMagnitude, bVelChange);
+	AddImpulse(Delta * ImpulseMagnitude, bVelChange);
 }
 
 FNetworkPredictionData_Client_Character::FNetworkPredictionData_Client_Character()
