Fixed a crash in PIE when traveling to map that has streaming levels

diff --git a/Engine/Source/Runtime/Engine/Private/LevelStreaming.cpp b/Engine/Source/Runtime/Engine/Private/LevelStreaming.cpp
--- a/Engine/Source/Runtime/Engine/Private/LevelStreaming.cpp
+++ b/Engine/Source/Runtime/Engine/Private/LevelStreaming.cpp
@@ -228,12 +228,12 @@ void ULevelStreaming::Serialize( FArchive& Ar )
 	
 	if (Ar.IsLoading())
 	{
-		if ((GetOutermost()->PackageFlags & PKG_PlayInEditor) != 0)
-		{
 #if WITH_EDITOR
+		if ((GetOutermost()->PackageFlags & PKG_PlayInEditor) != 0 && GetOutermost()->PIEInstanceID != INDEX_NONE)
+		{
 			RenameForPIE(GetOutermost()->PIEInstanceID);
-#endif
 		}
+#endif
 	}
 }
 
