Add ENGINE_API to UBodySetup::GetPhysMaterial()

diff --git a/Engine/Source/Runtime/Engine/Classes/PhysicsEngine/BodySetup.h b/Engine/Source/Runtime/Engine/Classes/PhysicsEngine/BodySetup.h
--- a/Engine/Source/Runtime/Engine/Classes/PhysicsEngine/BodySetup.h
+++ b/Engine/Source/Runtime/Engine/Classes/PhysicsEngine/BodySetup.h
@@ -503,7 +503,7 @@ public:
 	ENGINE_API float		CalculateMass(const UPrimitiveComponent* Component = NULL);
 
 	/** Returns the physics material used for this body. If none, specified, returns the default engine material. */
-	class UPhysicalMaterial* GetPhysMaterial() const;
+	ENGINE_API class UPhysicalMaterial* GetPhysMaterial() const;
 
 #if WITH_EDITOR
 	/** Clear all simple collision */
