- Add 'add' keyword to CombineRotators and ComposeTransforms
- Add 'bool' keyword to branch node

diff --git a/Engine/Source/Editor/BlueprintGraph/Private/K2Node_IfThenElse.cpp b/Engine/Source/Editor/BlueprintGraph/Private/K2Node_IfThenElse.cpp
--- a/Engine/Source/Editor/BlueprintGraph/Private/K2Node_IfThenElse.cpp
+++ b/Engine/Source/Editor/BlueprintGraph/Private/K2Node_IfThenElse.cpp
@@ -108,7 +108,7 @@ FString UK2Node_IfThenElse::GetTooltip() const
 
 FString UK2Node_IfThenElse::GetKeywords() const
 {
-	return *LOCTEXT("BranchKeywords", "if").ToString();
+	return *LOCTEXT("BranchKeywords", "if bool").ToString();
 }
 
 UEdGraphPin* UK2Node_IfThenElse::GetThenPin() const
diff --git a/Engine/Source/Runtime/Engine/Classes/Kismet/KismetMathLibrary.h b/Engine/Source/Runtime/Engine/Classes/Kismet/KismetMathLibrary.h
--- a/Engine/Source/Runtime/Engine/Classes/Kismet/KismetMathLibrary.h
+++ b/Engine/Source/Runtime/Engine/Classes/Kismet/KismetMathLibrary.h
@@ -523,7 +523,7 @@ class UKismetMathLibrary : public UBlueprintFunctionLibrary
 	static FRotator Multiply_RotatorFloat(FRotator A, float B);
 
 	/** Combine 2 rotations to give you the resulting rotation */
-	UFUNCTION(BlueprintPure, meta=(FriendlyName = "CombineRotators", Keywords="rotate rotation"), Category="Math|Rotator")
+	UFUNCTION(BlueprintPure, meta=(FriendlyName = "CombineRotators", Keywords="rotate rotation add"), Category="Math|Rotator")
 	static FRotator ComposeRotators(FRotator A, FRotator B);
 
 	/** Negate a rotator*/
@@ -863,7 +863,7 @@ class UKismetMathLibrary : public UBlueprintFunctionLibrary
 	static FVector InverseTransformDirection(const FTransform& T, FVector Direction);
 
 	/** Multiply two transforms in order */
-	UFUNCTION(BlueprintPure, meta=(ToolTip = "Multiply two transforms in order: A * B"), Category="Math|Transform")
+	UFUNCTION(BlueprintPure, meta=(ToolTip = "Multiply two transforms in order: A * B", Keywords="add"), Category="Math|Transform")
 	static FTransform ComposeTransforms(const FTransform& A, const FTransform& B);
 
 	/** 
