ActorComponent and WorldSettings now support AssetUserData
GitHub pull request 58 (ahatala) https://github.com/EpicGames/UnrealEngine/pull/58

diff --git a/Engine/Source/Runtime/Engine/Classes/Components/ActorComponent.h b/Engine/Source/Runtime/Engine/Classes/Components/ActorComponent.h
--- a/Engine/Source/Runtime/Engine/Classes/Components/ActorComponent.h
+++ b/Engine/Source/Runtime/Engine/Classes/Components/ActorComponent.h
@@ -4,10 +4,12 @@
 
 #include "ComponentInstanceDataCache.h"
 #include "Engine/EngineTypes.h"
+#include "Interfaces/Interface_AssetUserData.h"
+
 #include "ActorComponent.generated.h"
 
 UCLASS(DefaultToInstanced, abstract, hidecategories=(ComponentReplication))
-class ENGINE_API UActorComponent : public UObject
+class ENGINE_API UActorComponent : public UObject, public IInterface_AssetUserData
 {
 	GENERATED_UCLASS_BODY()
 
@@ -20,6 +22,11 @@ class ENGINE_API UActorComponent : public UObject
 	TArray<FName> ComponentTags;
 
 protected:
+
+	/** Array of user data stored with the component */
+	UPROPERTY()
+	TArray<UAssetUserData*> AssetUserData;
+
 	/** 
 	 *  Indicates if this ActorComponent is currently registered with a scene. 
 	 */
@@ -429,6 +436,12 @@ public:
 #endif // WITH_EDITOR
 	// End UObject interface.
 
+	// Begin IInterface_AssetUserData Interface
+	virtual void AddAssetUserData(UAssetUserData* InUserData) OVERRIDE;
+	virtual void RemoveUserDataOfClass(TSubclassOf<UAssetUserData> InUserDataClass) OVERRIDE;
+	virtual UAssetUserData* GetAssetUserDataOfClass(TSubclassOf<UAssetUserData> InUserDataClass) OVERRIDE;
+	// End IInterface_AssetUserData Interface
+
 	/** See if this component is currently registered */
 	FORCEINLINE bool IsRegistered() const
 	{
diff --git a/Engine/Source/Runtime/Engine/Classes/GameFramework/WorldSettings.h b/Engine/Source/Runtime/Engine/Classes/GameFramework/WorldSettings.h
--- a/Engine/Source/Runtime/Engine/Classes/GameFramework/WorldSettings.h
+++ b/Engine/Source/Runtime/Engine/Classes/GameFramework/WorldSettings.h
@@ -174,7 +174,7 @@ struct ENGINE_API FNetViewer
  * Actor containing all script accessible world properties.
  */
 UCLASS(config=game, showcategories=(Rendering), hidecategories=(Actor, Advanced, Display, Events, Object, Attachment, Info, Input, Blueprint, Layers), showcategories=("Input|MouseInput", "Input|TouchInput"), notplaceable, dependson=(APostProcessVolume, UMusicTrackDataStructures))
-class ENGINE_API AWorldSettings : public AInfo
+class ENGINE_API AWorldSettings : public AInfo, public IInterface_AssetUserData
 {
 	GENERATED_UCLASS_BODY()
 
@@ -402,6 +402,12 @@ class ENGINE_API AWorldSettings : public AInfo
 	/** Maximum number of bookmarks	*/
 	static const int32 MAX_BOOKMARK_NUMBER = 10;
 
+	protected:
+
+	/** Array of user data stored with the asset */
+	UPROPERTY()
+	TArray<UAssetUserData*> AssetUserData;
+
 public:
 	// Begin UObject interface.
 #if WITH_EDITOR
@@ -441,6 +447,13 @@ public:
 	 */	
 	virtual void NotifyMatchStarted();
 
+	// Begin IInterface_AssetUserData Interface
+	virtual void AddAssetUserData(UAssetUserData* InUserData) OVERRIDE;
+	virtual void RemoveUserDataOfClass(TSubclassOf<UAssetUserData> InUserDataClass) OVERRIDE;
+	virtual UAssetUserData* GetAssetUserDataOfClass(TSubclassOf<UAssetUserData> InUserDataClass) OVERRIDE;
+	// End IInterface_AssetUserData Interface
+
+
 private:
 
 	// Hidden functions that don't make sense to use on this class.
diff --git a/Engine/Source/Runtime/Engine/Classes/Interfaces/Interface_AssetUserData.h b/Engine/Source/Runtime/Engine/Classes/Interfaces/Interface_AssetUserData.h
--- a/Engine/Source/Runtime/Engine/Classes/Interfaces/Interface_AssetUserData.h
+++ b/Engine/Source/Runtime/Engine/Classes/Interfaces/Interface_AssetUserData.h
@@ -3,9 +3,10 @@
 #pragma once
 #include "Interface_AssetUserData.generated.h"
 
+class UAssetUserData;
 
-/** Interface for assets that can own UserData **/
-UINTERFACE(meta=(CannotImplementInterfaceInBlueprint))
+/** Interface for assets/objects that can own UserData **/
+UINTERFACE(MinimalApi, meta=(CannotImplementInterfaceInBlueprint))
 class UInterface_AssetUserData : public UInterface
 {
 	GENERATED_UINTERFACE_BODY()
diff --git a/Engine/Source/Runtime/Engine/Private/ActorComponent.cpp b/Engine/Source/Runtime/Engine/Private/ActorComponent.cpp
--- a/Engine/Source/Runtime/Engine/Private/ActorComponent.cpp
+++ b/Engine/Source/Runtime/Engine/Private/ActorComponent.cpp
@@ -1033,6 +1033,45 @@ bool UActorComponent::IsActive() const
 	return bIsActive;
 }
 
+void UActorComponent::AddAssetUserData(UAssetUserData* InUserData)
+{
+	if (InUserData != NULL)
+	{
+		UAssetUserData* ExistingData = GetAssetUserDataOfClass(InUserData->GetClass());
+		if (ExistingData != NULL)
+		{
+			AssetUserData.Remove(ExistingData);
+		}
+		AssetUserData.Add(InUserData);
+	}
+}
+
+UAssetUserData* UActorComponent::GetAssetUserDataOfClass(TSubclassOf<UAssetUserData> InUserDataClass)
+{
+	for (int32 DataIdx = 0; DataIdx < AssetUserData.Num(); DataIdx++)
+	{
+		UAssetUserData* Datum = AssetUserData[DataIdx];
+		if (Datum != NULL && Datum->IsA(InUserDataClass))
+		{
+			return Datum;
+		}
+	}
+	return NULL;
+}
+
+void UActorComponent::RemoveUserDataOfClass(TSubclassOf<UAssetUserData> InUserDataClass)
+{
+	for (int32 DataIdx = 0; DataIdx < AssetUserData.Num(); DataIdx++)
+	{
+		UAssetUserData* Datum = AssetUserData[DataIdx];
+		if (Datum != NULL && Datum->IsA(InUserDataClass))
+		{
+			AssetUserData.RemoveAt(DataIdx);
+			return;
+		}
+	}
+}
+
 void UActorComponent::SetNetAddressable()
 {
 	bNetAddressable = true;
diff --git a/Engine/Source/Runtime/Engine/Private/WorldSettings.cpp b/Engine/Source/Runtime/Engine/Private/WorldSettings.cpp
--- a/Engine/Source/Runtime/Engine/Private/WorldSettings.cpp
+++ b/Engine/Source/Runtime/Engine/Private/WorldSettings.cpp
@@ -165,6 +165,44 @@ void AWorldSettings::Serialize( FArchive& Ar )
 	}
 }
 
+void AWorldSettings::AddAssetUserData(UAssetUserData* InUserData)
+{
+	if (InUserData != NULL)
+	{
+		UAssetUserData* ExistingData = GetAssetUserDataOfClass(InUserData->GetClass());
+		if (ExistingData != NULL)
+		{
+			AssetUserData.Remove(ExistingData);
+		}
+		AssetUserData.Add(InUserData);
+	}
+}
+
+UAssetUserData* AWorldSettings::GetAssetUserDataOfClass(TSubclassOf<UAssetUserData> InUserDataClass)
+{
+	for (int32 DataIdx = 0; DataIdx < AssetUserData.Num(); DataIdx++)
+	{
+		UAssetUserData* Datum = AssetUserData[DataIdx];
+		if (Datum != NULL && Datum->IsA(InUserDataClass))
+		{
+			return Datum;
+		}
+	}
+	return NULL;
+}
+
+void AWorldSettings::RemoveUserDataOfClass(TSubclassOf<UAssetUserData> InUserDataClass)
+{
+	for (int32 DataIdx = 0; DataIdx < AssetUserData.Num(); DataIdx++)
+	{
+		UAssetUserData* Datum = AssetUserData[DataIdx];
+		if (Datum != NULL && Datum->IsA(InUserDataClass))
+		{
+			AssetUserData.RemoveAt(DataIdx);
+			return;
+		}
+	}
+}
 
 #if WITH_EDITOR
 
