Provide a default Component Instance Data matching behavior based on name and position in the serialized components array

diff --git a/Engine/Source/Runtime/Engine/Classes/Components/StaticMeshComponent.h b/Engine/Source/Runtime/Engine/Classes/Components/StaticMeshComponent.h
--- a/Engine/Source/Runtime/Engine/Classes/Components/StaticMeshComponent.h
+++ b/Engine/Source/Runtime/Engine/Classes/Components/StaticMeshComponent.h
@@ -92,84 +92,6 @@ private:
 	FStaticMeshComponentLODInfo &operator=( const FStaticMeshComponentLODInfo &rhs ) { check(0); return *this; }
 };
 
-/** Used to store lightmap data during RerunConstructionScripts */
-class FLightMapInstanceData : public FComponentInstanceDataBase
-{
-public:
-	static const FName LightMapInstanceDataTypeName;
-
-	virtual ~FLightMapInstanceData()
-	{}
-
-	// Begin FComponentInstanceDataBase interface
-	virtual FName GetDataTypeName() const OVERRIDE
-	{
-		return LightMapInstanceDataTypeName;
-	}
-	// End FComponentInstanceDataBase interface
-
-	/** Mesh being used by component */
-	class UStaticMesh*	StaticMesh;
-	/** Transform of instance */
-	FTransform		Transform;
-	/** Lightmaps from LODData */
-	TArray<FLightMapRef>	LODDataLightMap;
-	TArray<FShadowMapRef>	LODDataShadowMap;
-	TArray<FGuid> IrrelevantLights;
-};
-
-/**
- * Vertex color information stored during RerunConstructionScripts
- */
-class FVertexColorInstanceData : public FComponentInstanceDataBase
-{
-public:
-	static const FName VertexColorInstanceDataName;
-
-	/** FComponentInstanceDataBase interface */
-	virtual FName GetDataTypeName() const OVERRIDE
-	{
-		return VertexColorInstanceDataName;
-	}
-
-	/** Add vertex color data for a specified LOD before RerunConstructionScripts is called */
-	void AddVertexColorData( const struct FStaticMeshComponentLODInfo& LODInfo, uint32 LODIndex );
-
-	/** Re-apply vertex color data after RerunConstructionScripts is called */
-	bool ApplyVertexColorData( UStaticMeshComponent* StaticMeshComponent );
-
-	/** Check whether this vertex color data can match the specified component */
-	bool MatchesComponent( const UStaticMeshComponent* StaticMeshComponent ) const;
-
-	/** Vertex data stored per-LOD */
-	struct FVertexColorLODData
-	{
-		/** copy of painted vertex data */
-		TArray<FPaintedVertex> PaintedVertices;
-
-		/** Copy of vertex buffer colors */
-		TArray<FColor> VertexBufferColors;
-
-		/** Index of the LOD that this data came from */
-		uint32 LODIndex;
-
-		/** Check whether this contains valid data */
-		bool IsValid() const 
-		{ 
-			return PaintedVertices.Num() > 0 && VertexBufferColors.Num() > 0; 
-		}
-	};
-
-	/** Mesh being used by component */
-	class UStaticMesh* StaticMesh;
-
-	/** 'Stable' index into serialized components array */
-	int32 SerializedComponentsIndex;
-
-	/** Array of cached vertex colors for each LOD */
-	TArray<FVertexColorLODData> VertexColorLODs;
-};
-
 /** A StaticMeshComponent is a mesh that does not animate. */
 UCLASS(ClassGroup=(Rendering, Common), hidecategories=(Object,Activation,"Components|Activation"), ShowCategories=(Mobility), dependson=ULightmassPrimitiveSettingsObject, editinlinenew, meta=(BlueprintSpawnableComponent))
 class ENGINE_API UStaticMeshComponent : public UMeshComponent
diff --git a/Engine/Source/Runtime/Engine/Private/ComponentInstanceDataCache.cpp b/Engine/Source/Runtime/Engine/Private/ComponentInstanceDataCache.cpp
--- a/Engine/Source/Runtime/Engine/Private/ComponentInstanceDataCache.cpp
+++ b/Engine/Source/Runtime/Engine/Private/ComponentInstanceDataCache.cpp
@@ -2,77 +2,65 @@
 
 #include "EnginePrivate.h"
 
-#include "Components/ApplicationLifecycleComponent.h"
-#include "AI/BehaviorTree/BlackboardComponent.h"
-#include "AI/BrainComponent.h"
-#include "AI/BehaviorTree/BehaviorTreeComponent.h"
-#include "Debug/GameplayDebuggingControllerComponent.h"
-#include "Components/InputComponent.h"
-#include "GameFramework/MovementComponent.h"
-#include "Components/SceneComponent.h"
-#include "Components/LightComponentBase.h"
-#include "Components/LightComponent.h"
-#include "Components/PrimitiveComponent.h"
-#include "GameFramework/NavMovementComponent.h"
-#include "GameFramework/PawnMovementComponent.h"
-#include "GameFramework/CharacterMovementComponent.h"
-#include "GameFramework/FloatingPawnMovement.h"
-#include "GameFramework/SpectatorPawnMovement.h"
-#include "Vehicles/WheeledVehicleMovementComponent.h"
-#include "Vehicles/WheeledVehicleMovementComponent4W.h"
-#include "GameFramework/ProjectileMovementComponent.h"
-#include "GameFramework/RotatingMovementComponent.h"
-#include "AI/Navigation/NavigationComponent.h"
-#include "AI/Navigation/PathFollowingComponent.h"
-#include "Components/PawnNoiseEmitterComponent.h"
-#include "Components/PawnSensingComponent.h"
-#include "PhysicsEngine/PhysicsHandleComponent.h"
-#include "Atmosphere/AtmosphericFogComponent.h"
-#include "Components/AudioComponent.h"
-#include "Camera/CameraComponent.h"
-#include "Components/ChildActorComponent.h"
-#include "Components/ExponentialHeightFogComponent.h"
-#include "Components/SkyLightComponent.h"
-#include "AI/Navigation/NavigationGraphNodeComponent.h"
-#include "PhysicsEngine/PhysicsConstraintComponent.h"
-#include "PhysicsEngine/PhysicsThrusterComponent.h"
-#include "Components/PostProcessComponent.h"
-#include "Components/ArrowComponent.h"
-#include "Components/BillboardComponent.h"
-#include "Components/BrushComponent.h"
-#include "Components/DrawFrustumComponent.h"
-#include "AI/EnvironmentQuery/EQSRenderingComponent.h"
-#include "Debug/GameplayDebuggingComponent.h"
-#include "Components/LineBatchComponent.h"
-#include "Components/MaterialBillboardComponent.h"
-#include "Components/MeshComponent.h"
-#include "Components/SkinnedMeshComponent.h"
-#include "Components/DestructibleComponent.h"
-#include "Components/PoseableMeshComponent.h"
-#include "Components/SkeletalMeshComponent.h"
-#include "Components/StaticMeshComponent.h"
-#include "Components/InstancedStaticMeshComponent.h"
-#include "Components/InteractiveFoliageComponent.h"
-#include "Components/SplineMeshComponent.h"
-#include "Components/ModelComponent.h"
-#include "AI/Navigation/NavLinkRenderingComponent.h"
-#include "AI/Navigation/NavMeshRenderingComponent.h"
-#include "AI/Navigation/NavTestRenderingComponent.h"
-#include "Components/NiagaraComponent.h"
-#include "Components/ShapeComponent.h"
-#include "Components/BoxComponent.h"
-#include "Components/CapsuleComponent.h"
-#include "Components/SphereComponent.h"
-#include "Components/DrawSphereComponent.h"
-#include "Components/TextRenderComponent.h"
-#include "Components/VectorFieldComponent.h"
-#include "PhysicsEngine/RadialForceComponent.h"
-#include "GameFramework/SpringArmComponent.h"
-#include "Components/WindDirectionalSourceComponent.h"
-#include "Components/TimelineComponent.h"
-#include "Slate.h"
-#include "AI/NavDataGenerator.h"
-#include "OnlineSubsystemUtils.h"
+FComponentInstanceDataBase::FComponentInstanceDataBase(const UActorComponent* SourceComponent)
+{
+	check(SourceComponent);
+	SourceComponentName = SourceComponent->GetFName();
+	SourceComponentTypeSerializedIndex = -1;
+	
+	AActor* ComponentOwner = SourceComponent->GetOwner();
+	if (ComponentOwner)
+	{
+		bool bFound = false;
+		for (const UActorComponent* SerializedComponent : ComponentOwner->SerializedComponents)
+		{
+			if (SerializedComponent == SourceComponent)
+			{
+				++SourceComponentTypeSerializedIndex;
+				bFound = true;
+				break;
+			}
+			else if (SerializedComponent->GetClass() == SourceComponent->GetClass())
+			{
+				++SourceComponentTypeSerializedIndex;
+			}
+		}
+		if (!bFound)
+		{
+			SourceComponentTypeSerializedIndex = -1;
+		}
+	}
+}
+
+bool FComponentInstanceDataBase::MatchesComponent(const UActorComponent* Component) const
+{
+	bool bMatches = false;
+	if (Component)
+	{
+		if (Component->GetFName() == SourceComponentName)
+		{
+			bMatches = true;
+		}
+		else if (SourceComponentTypeSerializedIndex >= 0)
+		{
+			int32 FoundSerializedComponentsOfType = -1;
+			AActor* ComponentOwner = Component->GetOwner();
+			if (ComponentOwner)
+			{
+				for (const UActorComponent* SerializedComponent : ComponentOwner->SerializedComponents)
+				{
+					if (   (SerializedComponent->GetClass() == Component->GetClass())
+						&& (++FoundSerializedComponentsOfType == SourceComponentTypeSerializedIndex))
+					{
+						bMatches = true;
+						break;
+					}
+				}
+			}
+		}
+	}
+	return bMatches;
+}
 
 void FComponentInstanceDataCache::AddInstanceData(TSharedPtr<FComponentInstanceDataBase> NewData)
 {
diff --git a/Engine/Source/Runtime/Engine/Private/InstancedStaticMesh.cpp b/Engine/Source/Runtime/Engine/Private/InstancedStaticMesh.cpp
--- a/Engine/Source/Runtime/Engine/Private/InstancedStaticMesh.cpp
+++ b/Engine/Source/Runtime/Engine/Private/InstancedStaticMesh.cpp
@@ -1197,20 +1197,17 @@ class FInstancedStaticMeshSelectionData : public FComponentInstanceDataBase
 {
 public:
 	FInstancedStaticMeshSelectionData(const UInstancedStaticMeshComponent& InComponent)
-		: ComponentName(InComponent.GetFName())
+		: FComponentInstanceDataBase(&InComponent)
 		, SelectedInstances(InComponent.SelectedInstances)
 	{
 	}
 
 	virtual FName GetDataTypeName() const
 	{
-		return ComponentName;
+		return SourceComponentName;
 	}
 
 public:
-	/** Our components name */
-	FName ComponentName;
-
 	/** The cached selected instances */
 	TBitArray<> SelectedInstances;
 };
diff --git a/Engine/Source/Runtime/Engine/Private/StaticMeshComponent.cpp b/Engine/Source/Runtime/Engine/Private/StaticMeshComponent.cpp
--- a/Engine/Source/Runtime/Engine/Private/StaticMeshComponent.cpp
+++ b/Engine/Source/Runtime/Engine/Private/StaticMeshComponent.cpp
@@ -8,70 +8,120 @@
 
 #define LOCTEXT_NAMESPACE "StaticMeshComponent"
 
-const FName FVertexColorInstanceData::VertexColorInstanceDataName(TEXT("VertexColorInstanceData"));
-
-void FVertexColorInstanceData::AddVertexColorData( const FStaticMeshComponentLODInfo& LODInfo, uint32 LODIndex )
+/**
+* Vertex color information stored during RerunConstructionScripts
+*/
+class FVertexColorInstanceData : public FComponentInstanceDataBase
 {
-	if(VertexColorLODs.Num() <= (int32)LODIndex)
+public:
+
+	FVertexColorInstanceData(const UStaticMeshComponent* SourceComponent)
+		: FComponentInstanceDataBase(SourceComponent)
+		, StaticMesh(SourceComponent->StaticMesh)
 	{
-		VertexColorLODs.SetNum(LODIndex + 1);
 	}
-	FVertexColorLODData& VertexColorData = VertexColorLODs[LODIndex];
-	VertexColorData.LODIndex = LODIndex;
-	VertexColorData.PaintedVertices = LODInfo.PaintedVertices;
-	LODInfo.OverrideVertexColors->GetVertexColors( VertexColorData.VertexBufferColors );
-}
 
-bool FVertexColorInstanceData::ApplyVertexColorData( UStaticMeshComponent* StaticMeshComponent )
-{
-	bool bAppliedAnyData = false;
+	static const FName VertexColorInstanceDataName;
+
+	/** FComponentInstanceDataBase interface */
+	virtual FName GetDataTypeName() const OVERRIDE
+	{
+		return VertexColorInstanceDataName;
+	}
+
+	/** Add vertex color data for a specified LOD before RerunConstructionScripts is called */
+	void AddVertexColorData(const struct FStaticMeshComponentLODInfo& LODInfo, uint32 LODIndex)
+	{
+		if (VertexColorLODs.Num() <= (int32)LODIndex)
+		{
+			VertexColorLODs.SetNum(LODIndex + 1);
+		}
+		FVertexColorLODData& VertexColorData = VertexColorLODs[LODIndex];
+		VertexColorData.LODIndex = LODIndex;
+		VertexColorData.PaintedVertices = LODInfo.PaintedVertices;
+		LODInfo.OverrideVertexColors->GetVertexColors(VertexColorData.VertexBufferColors);
+	}
 
-	if(StaticMeshComponent != NULL)
+	/** Re-apply vertex color data after RerunConstructionScripts is called */
+	bool ApplyVertexColorData(UStaticMeshComponent* StaticMeshComponent)
 	{
-		StaticMeshComponent->SetLODDataCount(VertexColorLODs.Num(), VertexColorLODs.Num());
+		bool bAppliedAnyData = false;
 
-		for( int32 LODDataIndex = 0; LODDataIndex < VertexColorLODs.Num(); ++LODDataIndex )
+		if (StaticMeshComponent != NULL)
 		{
-			const FVertexColorLODData& VertexColorLODData = VertexColorLODs[LODDataIndex];
-			uint32 LODIndex = VertexColorLODData.LODIndex;
+			StaticMeshComponent->SetLODDataCount(VertexColorLODs.Num(), VertexColorLODs.Num());
 
-			if( StaticMeshComponent->LODData.IsValidIndex( LODIndex ) )
+			for (int32 LODDataIndex = 0; LODDataIndex < VertexColorLODs.Num(); ++LODDataIndex)
 			{
-				FStaticMeshComponentLODInfo& LODInfo = StaticMeshComponent->LODData[LODIndex];
-			
-				// Should not already have overriden vertex colors
-				if( LODInfo.OverrideVertexColors == NULL )
+				const FVertexColorLODData& VertexColorLODData = VertexColorLODs[LODDataIndex];
+				uint32 LODIndex = VertexColorLODData.LODIndex;
+
+				if (StaticMeshComponent->LODData.IsValidIndex(LODIndex))
 				{
-					LODInfo.PaintedVertices = VertexColorLODData.PaintedVertices;
+					FStaticMeshComponentLODInfo& LODInfo = StaticMeshComponent->LODData[LODIndex];
 
-					LODInfo.OverrideVertexColors = new FColorVertexBuffer;
-					LODInfo.OverrideVertexColors->InitFromColorArray( VertexColorLODData.VertexBufferColors );
+					// Should not already have overriden vertex colors
+					if (LODInfo.OverrideVertexColors == NULL)
+					{
+						LODInfo.PaintedVertices = VertexColorLODData.PaintedVertices;
+
+						LODInfo.OverrideVertexColors = new FColorVertexBuffer;
+						LODInfo.OverrideVertexColors->InitFromColorArray(VertexColorLODData.VertexBufferColors);
 
-					BeginInitResource( LODInfo.OverrideVertexColors );	
-					bAppliedAnyData = true;
+						BeginInitResource(LODInfo.OverrideVertexColors);
+						bAppliedAnyData = true;
+					}
 				}
 			}
-		}	
-	}
+		}
 
-	return bAppliedAnyData;
-}
+		return bAppliedAnyData;
+	}
 
-bool FVertexColorInstanceData::MatchesComponent( const UStaticMeshComponent* StaticMeshComponent ) const
-{
-	if(StaticMeshComponent != NULL)
+	/** Check whether this vertex color data can match the specified component */
+	virtual bool MatchesComponent(const UActorComponent* Component) const OVERRIDE
 	{
-		if(StaticMeshComponent->StaticMesh != StaticMesh)
+		const UStaticMeshComponent* StaticMeshComponent = CastChecked<const UStaticMeshComponent>(Component, ECastCheckedType::NullAllowed);
+		if (StaticMeshComponent != NULL)
 		{
-			return false;
+			if (StaticMeshComponent->StaticMesh != StaticMesh)
+			{
+				return false;
+			}
+
+			return FComponentInstanceDataBase::MatchesComponent(Component);
 		}
 
-		return SerializedComponentsIndex == StaticMeshComponent->GetSerializedComponentIndex();
+		// this component was not found
+		return false;
 	}
 
-	// this component was not found
-	return false;
-}
+	/** Vertex data stored per-LOD */
+	struct FVertexColorLODData
+	{
+		/** copy of painted vertex data */
+		TArray<FPaintedVertex> PaintedVertices;
+
+		/** Copy of vertex buffer colors */
+		TArray<FColor> VertexBufferColors;
+
+		/** Index of the LOD that this data came from */
+		uint32 LODIndex;
+
+		/** Check whether this contains valid data */
+		bool IsValid() const
+		{
+			return PaintedVertices.Num() > 0 && VertexBufferColors.Num() > 0;
+		}
+	};
+
+	/** Mesh being used by component */
+	class UStaticMesh* StaticMesh;
+
+	/** Array of cached vertex colors for each LOD */
+	TArray<FVertexColorLODData> VertexColorLODs;
+};
+const FName FVertexColorInstanceData::VertexColorInstanceDataName(TEXT("VertexColorInstanceData"));
 
 UStaticMeshComponent::UStaticMeshComponent(const class FPostConstructInitializeProperties& PCIP)
 	: Super(PCIP)
@@ -1460,7 +1510,31 @@ int32 UStaticMeshComponent::GetSerializedComponentIndex() const
 	return INDEX_NONE;
 }
 
-// Init type name static
+/** Used to store lightmap data during RerunConstructionScripts */
+class FLightMapInstanceData : public FComponentInstanceDataBase
+{
+public:
+	static const FName LightMapInstanceDataTypeName;
+
+	virtual ~FLightMapInstanceData()
+	{}
+
+	// Begin FComponentInstanceDataBase interface
+	virtual FName GetDataTypeName() const OVERRIDE
+	{
+		return LightMapInstanceDataTypeName;
+	}
+	// End FComponentInstanceDataBase interface
+
+	/** Mesh being used by component */
+	class UStaticMesh*	StaticMesh;
+	/** Transform of instance */
+	FTransform		Transform;
+	/** Lightmaps from LODData */
+	TArray<FLightMapRef>	LODDataLightMap;
+	TArray<FShadowMapRef>	LODDataShadowMap;
+	TArray<FGuid> IrrelevantLights;
+};
 const FName FLightMapInstanceData::LightMapInstanceDataTypeName(TEXT("LightMapInstanceData"));
 
 void UStaticMeshComponent::GetComponentInstanceData(FComponentInstanceDataCache& Cache) const
@@ -1497,9 +1571,7 @@ void UStaticMeshComponent::GetComponentInstanceData(FComponentInstanceDataCache&
 		{
 			if( !VertexColorData.IsValid() )
 			{
-				VertexColorData = MakeShareable( new FVertexColorInstanceData );
-				VertexColorData->StaticMesh = StaticMesh;
-				VertexColorData->SerializedComponentsIndex = GetSerializedComponentIndex();
+				VertexColorData = MakeShareable( new FVertexColorInstanceData(this) );
 			}
 
 			VertexColorData->AddVertexColorData( LODInfo, LODIndex );
diff --git a/Engine/Source/Runtime/Engine/Public/ComponentInstanceDataCache.h b/Engine/Source/Runtime/Engine/Public/ComponentInstanceDataCache.h
--- a/Engine/Source/Runtime/Engine/Public/ComponentInstanceDataCache.h
+++ b/Engine/Source/Runtime/Engine/Public/ComponentInstanceDataCache.h
@@ -6,11 +6,27 @@
 class FComponentInstanceDataBase
 {
 public:
+	FComponentInstanceDataBase()
+		: SourceComponentTypeSerializedIndex(-1)
+	{}
+
+	FComponentInstanceDataBase(const UActorComponent* SourceComponent);
+
 	virtual ~FComponentInstanceDataBase()
 	{}
 
 	/** Return the type name of this data */
 	virtual FName GetDataTypeName() const = 0;
+
+	virtual bool MatchesComponent(const UActorComponent* Component) const;
+
+protected:
+	/** The name of the source component */
+	FName SourceComponentName;
+
+	/** The index of the source component in its owner's serialized array 
+		when filtered to just that component type */
+	int32 SourceComponentTypeSerializedIndex;
 };
 
 /** 
