Fixed radial damage to construct a fake FHitResult for actors that get damage but miss the visibility trace

diff --git a/Engine/Source/Runtime/Engine/Private/GameplayStatics.cpp b/Engine/Source/Runtime/Engine/Private/GameplayStatics.cpp
--- a/Engine/Source/Runtime/Engine/Private/GameplayStatics.cpp
+++ b/Engine/Source/Runtime/Engine/Private/GameplayStatics.cpp
@@ -118,7 +118,7 @@ bool UGameplayStatics::SetGamePaused(UObject* WorldContextObject, bool bPaused)
 }
 
 /** @RETURN True if weapon trace from Origin hits component VictimComp.  OutHitResult will contain properties of the hit. */
-static bool ComponentIsVisibleFrom(UPrimitiveComponent const* VictimComp, FVector const& Origin, AActor const* IgnoredActor, const TArray<AActor*>& IgnoreActors, FHitResult& OutHitResult)
+static bool ComponentIsVisibleFrom(UPrimitiveComponent* VictimComp, FVector const& Origin, AActor const* IgnoredActor, const TArray<AActor*>& IgnoreActors, FHitResult& OutHitResult)
 {
 	static FName NAME_ComponentIsVisibleFrom = FName(TEXT("ComponentIsVisibleFrom"));
 	FCollisionQueryParams LineParams(NAME_ComponentIsVisibleFrom, true, IgnoredActor);
@@ -155,6 +155,10 @@ static bool ComponentIsVisibleFrom(UPrimitiveComponent const* VictimComp, FVecto
 	}
 		
 	// didn't hit anything, assume nothing blocking the damage and victim is consequently visible
+	// but since we don't have a hit result to pass back, construct a simple one, modeling the damage as having hit a point at the component's center.
+	FVector const FakeHitLoc = VictimComp->GetComponentLocation();
+	FVector const FakeHitNorm = (Origin - FakeHitLoc).SafeNormal();		// normal points back toward the epicenter
+	OutHitResult = FHitResult(VictimComp->GetOwner(), VictimComp, FakeHitLoc, FakeHitNorm);
 	return true;
 }
 
