UE4: Renaming Spline.cpp to SplineComponent.cpp to fix a non-unity build issue.

diff --git a/Engine/Source/Runtime/Engine/Private/Spline.cpp b/Engine/Source/Runtime/Engine/Private/Spline.cpp
deleted file mode 100644
diff --git a/Engine/Source/Runtime/Engine/Private/SplineComponent.cpp b/Engine/Source/Runtime/Engine/Private/SplineComponent.cpp
new file mode 100644
--- /dev/null
+++ b/Engine/Source/Runtime/Engine/Private/SplineComponent.cpp
@@ -0,0 +1,263 @@
+// Copyright 1998-2014 Epic Games, Inc. All Rights Reserved.
+
+/*=============================================================================
+	Spline.cpp
+=============================================================================*/
+
+#include "EnginePrivate.h"
+
+
+// HSplineProxy
+void HSplineProxy::AddReferencedObjects( FReferenceCollector& Collector )
+{
+	Collector.AddReferencedObject( SplineComp );
+}	
+
+
+//////////////////////////////////////////////////////////////////////////
+// SPLINE SCENE PROXY
+
+/** Represents a USplineComponent to the scene manager. */
+class FSplineSceneProxy : public FPrimitiveSceneProxy
+{
+private:
+	FColor					SplineColor;
+	FInterpCurveVector		SplineInfo;
+	float					SplineDrawRes;
+	float					SplineArrowSize;
+	uint32					bSplineDisabled:1;
+
+public:
+
+	FSplineSceneProxy(USplineComponent* Component):
+		FPrimitiveSceneProxy(Component),
+		SplineColor(Component->SplineColor),
+		SplineInfo(Component->SplineInfo),
+		SplineDrawRes(Component->SplineDrawRes),
+		SplineArrowSize(Component->SplineArrowSize),
+		bSplineDisabled(Component->bSplineDisabled)
+	{
+	}
+
+	virtual HHitProxy* CreateHitProxies(UPrimitiveComponent* Component, TArray<TRefCountPtr<HHitProxy> >& OutHitProxies)
+	{
+		HSplineProxy* SplineHitProxy = new HSplineProxy( CastChecked<USplineComponent>(Component) );
+		OutHitProxies.Add(SplineHitProxy);
+		return SplineHitProxy;
+	}
+
+	virtual void DrawDynamicElements(FPrimitiveDrawInterface* PDI, const FSceneView* View)
+	{
+		QUICK_SCOPE_CYCLE_COUNTER( STAT_SplineSceneProxy_DrawDynamicElements );
+
+		// Override specified color with red if spline is 'disabled'
+		FColor UseColor = bSplineDisabled ? FColor(255,0,0) : SplineColor;
+		
+		//DrawDirectionalArrow(PDI, LocalToWorld, SplineColor, 1.0 * 3.0f, 1.0f, GetDepthPriorityGroup(View));		
+		FVector OldKeyPos(0);
+		float OldKeyTime = 0.f;
+		for(int32 i=0; i<SplineInfo.Points.Num(); i++)
+		{
+			float NewKeyTime = SplineInfo.Points[i].InVal;
+			FVector NewKeyPos = SplineInfo.Eval(NewKeyTime, FVector::ZeroVector);
+
+			// Draw the keypoint
+			PDI->DrawPoint(NewKeyPos, UseColor, 6.f, GetDepthPriorityGroup(View));
+
+			// If not the first keypoint, draw a line to the last keypoint.
+			if(i>0)
+			{
+				// For constant interpolation - don't draw ticks - just draw dotted line.
+				if(SplineInfo.Points[i-1].InterpMode == CIM_Constant)
+				{
+					DrawDashedLine(PDI, OldKeyPos, NewKeyPos, UseColor, 20, GetDepthPriorityGroup(View));
+				}
+				else
+				{
+					int32 NumSteps = FMath::CeilToInt( (NewKeyTime - OldKeyTime)/SplineDrawRes );
+					float DrawSubstep = (NewKeyTime - OldKeyTime)/NumSteps;
+
+					// Find position on first keyframe.
+					float OldTime = OldKeyTime;
+					FVector OldPos = OldKeyPos;
+
+					// Find step that we want to draw the arrow head on.
+					int32 ArrowStep = FMath::Max<int32>(0, NumSteps-2);
+					
+					// Then draw a line for each substep.
+					for(int32 j=1; j<NumSteps+1; j++)
+					{
+						float NewTime = OldKeyTime + j*DrawSubstep;
+						FVector NewPos = SplineInfo.Eval(NewTime, FVector::ZeroVector);
+
+						// If this is where we want to draw the arrow, do that
+						// Only draw arrow if SplineArrowSize > 0.0
+						if(j == ArrowStep && SplineArrowSize > KINDA_SMALL_NUMBER)
+						{
+							FMatrix ArrowMatrix;
+							FVector ArrowDir = (NewPos - OldPos);
+							float ArrowLen = ArrowDir.Size();
+							if(ArrowLen > KINDA_SMALL_NUMBER)
+							{
+								ArrowDir /= ArrowLen;
+							}
+
+							ArrowMatrix = FRotationTranslationMatrix(ArrowDir.Rotation(), OldPos);
+
+							DrawDirectionalArrow(PDI, ArrowMatrix, UseColor, ArrowLen, SplineArrowSize, GetDepthPriorityGroup(View));								
+						}
+						// Otherwise normal line
+						else
+						{
+							PDI->DrawLine(OldPos, NewPos, UseColor, GetDepthPriorityGroup(View));														
+						}
+
+						// Don't draw point for last one - its the keypoint drawn above.
+						if(j != NumSteps)
+						{
+							PDI->DrawPoint(NewPos, UseColor, 3.f, GetDepthPriorityGroup(View));
+						}
+
+						OldTime = NewTime;
+						OldPos = NewPos;
+					}
+				}
+			}
+
+			OldKeyTime = NewKeyTime;
+			OldKeyPos = NewKeyPos;
+		}
+	}
+
+	virtual FPrimitiveViewRelevance GetViewRelevance(const FSceneView* View)
+	{
+		FPrimitiveViewRelevance Result;
+
+		Result.bDrawRelevance = IsShown(View) && View->Family->EngineShowFlags.Splines;
+		Result.bDynamicRelevance = true;
+		Result.bEditorPrimitiveRelevance = UseEditorCompositing(View);
+
+		return Result;
+	}
+
+	virtual uint32 GetMemoryFootprint( void ) const { return( sizeof( *this ) + GetAllocatedSize() ); }
+	uint32 GetAllocatedSize( void ) const { return( FPrimitiveSceneProxy::GetAllocatedSize() ); }
+
+
+};
+
+USplineComponent::USplineComponent(const class FPostConstructInitializeProperties& PCIP)
+	: Super(PCIP)
+{
+	SplineDrawRes = 0.1f;
+	SplineArrowSize = 60.0f;
+	SplineColor.R = 255;
+	SplineColor.B = 255;
+	SplineColor.A = 255;
+
+	bUseEditorCompositing = true;
+	BodyInstance.bEnableCollision_DEPRECATED = false;
+
+	SetCollisionProfileName(UCollisionProfile::NoCollision_ProfileName);
+}
+
+
+FPrimitiveSceneProxy* USplineComponent::CreateSceneProxy()
+{
+	return new FSplineSceneProxy(this);
+}
+
+FBoxSphereBounds USplineComponent::CalcBounds(const FTransform & LocalToWorld) const
+{
+	FVector BoxMin, BoxMax;
+
+	SplineInfo.CalcBounds(BoxMin, BoxMax, LocalToWorld.GetLocation());
+
+	return FBoxSphereBounds( FBox(BoxMin, BoxMax) );
+}
+
+
+void USplineComponent::UpdateSplineReparamTable()
+{
+	// Start by clearing it
+	SplineReparamTable.Reset();
+	
+	// Nothing to do if no points
+	if(SplineInfo.Points.Num() < 2)
+	{
+		return;
+	}
+	
+	const int32 NumSteps = 10; // TODO: Make this adaptive...
+	
+	// Find range of input
+	float Param = SplineInfo.Points[0].InVal;
+	const float MaxInput = SplineInfo.Points[SplineInfo.Points.Num()-1].InVal;
+	const float Interval = (MaxInput - Param)/((float)(NumSteps-1)); 
+	
+	// Add first entry, using first point on curve, total distance will be 0
+	FVector OldSplinePos = SplineInfo.Eval(Param, FVector::ZeroVector);
+	float TotalDist = 0.f;
+	SplineReparamTable.AddPoint(TotalDist, Param);
+	Param += Interval;
+
+	// Then work over rest of points	
+	for(int32 i=1; i<NumSteps; i++)
+	{
+		// Iterate along spline at regular param intervals
+		const FVector NewSplinePos = SplineInfo.Eval(Param, FVector::ZeroVector);
+		
+		TotalDist += (NewSplinePos - OldSplinePos).Size();
+		OldSplinePos = NewSplinePos;
+
+		SplineReparamTable.AddPoint(TotalDist, Param);
+
+		// move along
+		Param += Interval;
+	}
+}
+
+
+void USplineComponent::UpdateSplineCurviness()
+{
+	const float SplineLength = GetSplineLength();
+
+	const FVector Start = GetLocationAtDistanceAlongSpline( 0.0f );;
+	const FVector End = GetLocationAtDistanceAlongSpline( SplineLength );
+
+	SplineCurviness =(Start - End).Size() / SplineLength;
+
+	//GetWorld()->PersistentLineBatcher->DrawLine( Start, Start+FVector(0,0,200), FColor(255,0,0), SDPG_World );
+	//GetWorld()->PersistentLineBatcher->DrawLine( End, End+FVector(0,0,200), FColor(255,0,0), SDPG_World );
+
+	//UE_LOG(LogSpline, Warning, TEXT( "hahah %s %f %f %f"), *GetName(), SplineCurviness, (Start - End).Size(), SplineLength );
+}
+
+
+float USplineComponent::GetSplineLength() const
+{
+	// This is given by the input of the last entry in the remap table
+	if(SplineReparamTable.Points.Num() > 0)
+	{
+		return SplineReparamTable.Points[SplineReparamTable.Points.Num()-1].InVal;
+	}
+	
+	return 0.f;
+}
+
+
+FVector USplineComponent::GetLocationAtDistanceAlongSpline(float Distance) const
+{
+	const float Param = SplineReparamTable.Eval(Distance, 0.f);
+	const FVector Location = SplineInfo.Eval(Param, FVector::ZeroVector);
+	return Location;
+}
+
+
+FVector USplineComponent::GetTangentAtDistanceAlongSpline(float Distance) const
+{
+	const float Param = SplineReparamTable.Eval(Distance, 0.f);
+	const FVector Tangent = SplineInfo.EvalDerivative(Param, FVector::ZeroVector);
+	return Tangent;
+}
+
