Remove SkipFixedAndSimulatingBones.
Solves a few issues with changing fixed/unfixed at runtime, as well as simplifying blending code.

diff --git a/Engine/Config/BaseEngine.ini b/Engine/Config/BaseEngine.ini
--- a/Engine/Config/BaseEngine.ini
+++ b/Engine/Config/BaseEngine.ini
@@ -469,6 +469,9 @@ bUseStreamingPause=false
 +EnumRedirects=(EnumName="EAxis",OldEnumEntry="AXIS_NONE",NewEnumEntry="EAxis::None")
 +EnumRedirects=(EnumName="EAxis",OldEnumEntry="AXIS_BLANK",NewEnumEntry="EAxis::None")
 
+;EKinematicBonesUpdateToPhysics rename
++EnumRedirects=(EnumName="EKinematicBonesUpdateToPhysics", OldEnumEntry="EKinematicBonesUpdateToPhysics::SkipFixedAndSimulatingBones",NewEnumEntry="EKinematicBonesUpdateToPhysics::SkipAllBones")
+
 ;Struct Redirects for UE4
 +ActiveStructRedirects=(OldStructName="VimDebugData",NewStructName="AnimBlueprintDebugData")
 
diff --git a/Engine/Source/Editor/PhAT/Private/PhATRender.cpp b/Engine/Source/Editor/PhAT/Private/PhATRender.cpp
--- a/Engine/Source/Editor/PhAT/Private/PhATRender.cpp
+++ b/Engine/Source/Editor/PhAT/Private/PhATRender.cpp
@@ -37,7 +37,7 @@ UPhATEdSkeletalMeshComponent::UPhATEdSkeletalMeshComponent(const class FPostCons
 	check(BoneNoCollisionMaterial);
 
 	// this is because in phat editor, you'd like to see fixed bones to be fixed without animation force update
-	KinematicBonesUpdateType = EKinematicBonesUpdateToPhysics::SkipFixedAndSimulatingBones;
+	KinematicBonesUpdateType = EKinematicBonesUpdateToPhysics::SkipSimulatingBones;
 	bUpdateJointsFromAnimation = true;
 	ForcedLodModel = 1;
 
diff --git a/Engine/Source/Runtime/Engine/Classes/Components/SkeletalMeshComponent.h b/Engine/Source/Runtime/Engine/Classes/Components/SkeletalMeshComponent.h
--- a/Engine/Source/Runtime/Engine/Classes/Components/SkeletalMeshComponent.h
+++ b/Engine/Source/Runtime/Engine/Classes/Components/SkeletalMeshComponent.h
@@ -120,11 +120,9 @@ namespace EKinematicBonesUpdateToPhysics
 {
 	enum Type
 	{
-		/** This only updates anything non-simulating bones - i.e. (fixed) or (default if owner is not simulating)**/
+		/** Update any bones that are not simulating*/
 		SkipSimulatingBones,
-		/** This updates everything except (fixed) or (unfixed) or (default if owner is simulating)**/
-		SkipFixedAndSimulatingBones, 
-		/** Skip any physics update from kinematic changes **/
+		/** Skip physics update from kinematic changes **/
 		SkipAllBones
 	};
 }
diff --git a/Engine/Source/Runtime/Engine/Private/PhysicsEngine/BodyInstance.cpp b/Engine/Source/Runtime/Engine/Private/PhysicsEngine/BodyInstance.cpp
--- a/Engine/Source/Runtime/Engine/Private/PhysicsEngine/BodyInstance.cpp
+++ b/Engine/Source/Runtime/Engine/Private/PhysicsEngine/BodyInstance.cpp
@@ -754,19 +754,27 @@ void FBodyInstance::InitBody(UBodySetup* Setup, const FTransform& Transform, UPr
 		SkelMeshComp = Cast<USkeletalMeshComponent>(OwnerComponent.Get());
 
 		// if not dedicated and if set up is done to use physics, enable flag on
-		if (BodySetup->PhysicsType == PhysType_Unfixed)
+		if (BodySetup->PhysicsType == PhysType_Unfixed || BodySetup->PhysicsType == PhysType_Default)
 		{
 			bool bEnableSimulation = (SkelMeshComp && IsRunningDedicatedServer())? SkelMeshComp->bEnablePhysicsOnDedicatedServer: true;
+			bEnableSimulation &= ((BodySetup->PhysicsType == PhysType_Unfixed) || (SkelMeshComp->BodyInstance.bSimulatePhysics));	//if unfixed enable. If default look at parent
 			if ( bEnableSimulation )
 			{
 				// set simulate to true if using physics
 				bSimulatePhysics = true;
-				PhysicsBlendWeight = 1.f;
+				if (BodySetup->PhysicsType == PhysType_Unfixed)
+				{
+					PhysicsBlendWeight = 1.f;
+				}
+				
 			}
 	 		else
 	 		{
 	 			bSimulatePhysics = false;
-	 			PhysicsBlendWeight = 0.f;
+				if (BodySetup->PhysicsType == PhysType_Unfixed)
+				{
+					PhysicsBlendWeight = 0.f;
+				}
 	 		}
 		}
 	}
diff --git a/Engine/Source/Runtime/Engine/Private/PhysicsEngine/PhysAnim.cpp b/Engine/Source/Runtime/Engine/Private/PhysicsEngine/PhysAnim.cpp
--- a/Engine/Source/Runtime/Engine/Private/PhysicsEngine/PhysAnim.cpp
+++ b/Engine/Source/Runtime/Engine/Private/PhysicsEngine/PhysAnim.cpp
@@ -86,8 +86,7 @@ void USkeletalMeshComponent::BlendPhysicsBones( TArray<FBoneIndexType>& InRequir
 		// need to update back to physX so that physX knows where it was after blending
 		bool bUpdatePhysics = false;
 		FBodyInstance* BodyInstance = NULL;
-		// if this is true fixed bones haven't been updated, so we can't blend between. 
-		bool bShouldSkipFixedBones = KinematicBonesUpdateType == EKinematicBonesUpdateToPhysics::SkipFixedAndSimulatingBones;
+
 		// If so - get its world space matrix and its parents world space matrix and calc relative atom.
 		if(BodyIndex != INDEX_NONE )
 		{	
@@ -106,12 +105,9 @@ void USkeletalMeshComponent::BlendPhysicsBones( TArray<FBoneIndexType>& InRequir
 			}
 #endif
 			BodyInstance = Bodies[BodyIndex];
-			// since we don't copy back to physics, we shouldn't blend fixed bones here if this set up is used
-			// if you'd like to use fixed bones to blend use SkipSimulatingBones
-			// you can simulate fixed bones by setting it to instance, but that case, please use SkipSimulating Bones. 
-			bool bSkipFixedBones = bShouldSkipFixedBones && !BodyInstance->bSimulatePhysics;
 
-			if(!bSkipFixedBones && BodyInstance->IsValidBodyInstance())
+			//if simulated body copy back and blend with animation
+			if(BodyInstance->IsInstanceSimulatingPhysics(true))
 			{
 				FTransform PhysTM = BodyInstance->GetUnrealWorldTransform();
 
@@ -294,9 +290,6 @@ void USkeletalMeshComponent::UpdateKinematicBonesToPhysics(bool bTeleport)
 		}
 #endif
 
-		// see if we should check fixed bones or not
-		bool bShouldSkipFixedBones = KinematicBonesUpdateType == EKinematicBonesUpdateToPhysics::SkipFixedAndSimulatingBones;
-
 		// Iterate over each body
 		for(int32 i = 0; i < Bodies.Num(); i++)
 		{
@@ -304,7 +297,7 @@ void USkeletalMeshComponent::UpdateKinematicBonesToPhysics(bool bTeleport)
 			FBodyInstance* BodyInst = Bodies[i];
 			check(BodyInst);
 
-			if (!bShouldSkipFixedBones && BodyInst->IsValidBodyInstance() && !BodyInst->IsInstanceSimulatingPhysics(true))
+			if (BodyInst->IsValidBodyInstance() && !BodyInst->IsInstanceSimulatingPhysics(true))
 			{
 				// Find the graphics bone index that corresponds to this physics body.
 				FName const BodyName = PhysicsAsset->BodySetup[i]->BoneName;
diff --git a/Engine/Source/Runtime/Engine/Private/Vehicles/WheeledVehicleMovementComponent.cpp b/Engine/Source/Runtime/Engine/Private/Vehicles/WheeledVehicleMovementComponent.cpp
--- a/Engine/Source/Runtime/Engine/Private/Vehicles/WheeledVehicleMovementComponent.cpp
+++ b/Engine/Source/Runtime/Engine/Private/Vehicles/WheeledVehicleMovementComponent.cpp
@@ -1255,7 +1255,7 @@ void UWheeledVehicleMovementComponent::FixupSkeletalMesh()
 			}
 		}
 
-		Mesh->KinematicBonesUpdateType = EKinematicBonesUpdateToPhysics::SkipFixedAndSimulatingBones;	//since the wheels use physx we need to skip fixed bones from updating
+		Mesh->KinematicBonesUpdateType = EKinematicBonesUpdateToPhysics::SkipSimulatingBones;
 		
 	}
 
