Don't register actor components on a rename in to a world that is not initialized
TTP #335305

diff --git a/Engine/Source/Runtime/Engine/Private/Actor.cpp b/Engine/Source/Runtime/Engine/Private/Actor.cpp
--- a/Engine/Source/Runtime/Engine/Private/Actor.cpp
+++ b/Engine/Source/Runtime/Engine/Private/Actor.cpp
@@ -687,7 +687,11 @@ bool AActor::Rename( const TCHAR* InName, UObject* NewOuter, ERenameFlags Flags
 
 	if (NewOuter && NewOuter->IsA<ULevel>())
 	{
-		RegisterAllComponents();
+		UWorld* World = NewOuter->GetWorld();
+		if (World && World->bIsWorldInitialized)
+		{
+			RegisterAllComponents();
+		}
 		RegisterAllActorTickFunctions(true, true); // register all tick functions
 	}
 	return bSuccess;
