#ue4 - Fix StepUp() sometimes allowing movement over obstacles above MaxStepHeight.

- Impact point for downward sweep should be considered even if not on the lower hemisphere.
- Reject StepUp if downward sweep hits something on the edge of the capsule.

UDN: https://udn.unrealengine.com/questions/207371/suggested-change-to-ucharactermovementcomponentste.html

diff --git a/Engine/Source/Runtime/Engine/Private/CharacterMovementComponent.cpp b/Engine/Source/Runtime/Engine/Private/CharacterMovementComponent.cpp
--- a/Engine/Source/Runtime/Engine/Private/CharacterMovementComponent.cpp
+++ b/Engine/Source/Runtime/Engine/Private/CharacterMovementComponent.cpp
@@ -4134,28 +4134,20 @@ bool UCharacterMovementComponent::StepUp(const FVector& GravDir, const FVector&
 
 	FStepDownResult StepDownResult;
 	if (Hit.IsValidBlockingHit())
-	{
-		const FVector HitLocation = Hit.Location;
-
-		// See if the downward move impacts on the lower capsule hemisphere
-		const float CurBaseLocation = (HitLocation.Z - PawnHalfHeight);
-		const float LowerImpactHeight = Hit.ImpactPoint.Z - CurBaseLocation;
-		if (LowerImpactHeight <= PawnRadius)
+	{	
+		// See if this step sequence would have allowed us to travel higher than our max step height allows.
+		const float DeltaZ = Hit.ImpactPoint.Z - PawnFloorPointZ;
+		if (DeltaZ > MaxStepHeight)
 		{
-			// See if this step sequence would have allowed us to travel higher than our max step height allows.
-			const float DeltaZ = Hit.ImpactPoint.Z - PawnFloorPointZ;
-			if (DeltaZ > MaxStepHeight)
-			{
-				//UE_LOG(LogCharacterMovement, VeryVerbose, TEXT("- Reject StepUp (too high Height %.3f) up from floor base %f to %f"), DeltaZ, PawnInitialFloorBaseZ, NewLocation.Z);
-				ScopedStepUpMovement.RevertMove();
-				return false;
-			}
+			//UE_LOG(LogCharacterMovement, VeryVerbose, TEXT("- Reject StepUp (too high Height %.3f) up from floor base %f to %f"), DeltaZ, PawnInitialFloorBaseZ, NewLocation.Z);
+			ScopedStepUpMovement.RevertMove();
+			return false;
 		}
 
 		// Reject unwalkable surface normals here.
 		if (!IsWalkable(Hit))
 		{
-			// Reject if normal is towards movement direction
+			// Reject if normal opposes movement direction
 			const bool bNormalTowardsMe = (Delta | Hit.ImpactNormal) < 0.f;
 			if (bNormalTowardsMe)
 			{
@@ -4166,7 +4158,7 @@ bool UCharacterMovementComponent::StepUp(const FVector& GravDir, const FVector&
 
 			// Also reject if we would end up being higher than our starting location by stepping down.
 			// It's fine to step down onto an unwalkable normal below us, we will just slide off. Rejecting those moves would prevent us from being able to walk off the edge.
-			if (HitLocation.Z > OldLocation.Z)
+			if (Hit.Location.Z > OldLocation.Z)
 			{
 				//UE_LOG(LogCharacterMovement, VeryVerbose, TEXT("- Reject StepUp (unwalkable normal %s above old position)"), *Hit.ImpactNormal.ToString());
 				ScopedStepUpMovement.RevertMove();
@@ -4174,6 +4166,14 @@ bool UCharacterMovementComponent::StepUp(const FVector& GravDir, const FVector&
 			}
 		}
 
+		// Reject moves where the downward sweep hit something very close to the edge of the capsule. This maintains consistency with FindFloor as well.
+		if (!IsWithinEdgeTolerance(Hit.Location, Hit.ImpactPoint, PawnRadius))
+		{
+			//UE_LOG(LogCharacterMovement, VeryVerbose, TEXT("- Reject StepUp (outside edge tolerance)"));
+			ScopedStepUpMovement.RevertMove();
+			return false;
+		}
+
 		// See if we can validate the floor as a result of this step down. In almost all cases this should succeed, and we can avoid computing the floor outside this method.
 		if (OutStepDownResult != NULL)
 		{
@@ -4181,7 +4181,7 @@ bool UCharacterMovementComponent::StepUp(const FVector& GravDir, const FVector&
 
 			// Reject unwalkable normals if we end up higher than our initial height.
 			// It's fine to walk down onto an unwalkable surface, don't reject those moves.
-			if (HitLocation.Z > OldLocation.Z)
+			if (Hit.Location.Z > OldLocation.Z)
 			{
 				// We should reject the floor result if we are trying to step up an actual step where we are not able to perch (this is rare).
 				// In those cases we should instead abort the step up and try to slide along the stair.
