ShowLayer command
* Can be used to show/hide LD-specified layers of actors in game, useful for profiling on mobile where only the final frametime is reliable

diff --git a/Engine/Source/Runtime/Engine/Classes/Engine/GameViewportClient.h b/Engine/Source/Runtime/Engine/Classes/Engine/GameViewportClient.h
--- a/Engine/Source/Runtime/Engine/Classes/Engine/GameViewportClient.h
+++ b/Engine/Source/Runtime/Engine/Classes/Engine/GameViewportClient.h
@@ -287,6 +287,7 @@ public:
 	 */
 	bool HandleForceFullscreenCommand( const TCHAR* Cmd, FOutputDevice& Ar );
 	bool HandleShowCommand( const TCHAR* Cmd, FOutputDevice& Ar, UWorld* InWorld );
+	bool HandleShowLayerCommand( const TCHAR* Cmd, FOutputDevice& Ar, UWorld* InWorld );
 	bool HandleViewModeCommand( const TCHAR* Cmd, FOutputDevice& Ar, UWorld* InWorld );
 	bool HandleNextViewModeCommand( const TCHAR* Cmd, FOutputDevice& Ar, UWorld* InWorld );
 	bool HandlePrevViewModeCommand( const TCHAR* Cmd, FOutputDevice& Ar, UWorld* InWorld );
diff --git a/Engine/Source/Runtime/Engine/Classes/GameFramework/Actor.h b/Engine/Source/Runtime/Engine/Classes/GameFramework/Actor.h
--- a/Engine/Source/Runtime/Engine/Classes/GameFramework/Actor.h
+++ b/Engine/Source/Runtime/Engine/Classes/GameFramework/Actor.h
@@ -385,6 +385,10 @@ public:
 	/** Return the value of bAllowReceiveTickEventOnDedicatedServer, indicating whether the Blueprint ReceiveTick event will occur on dedicated servers */
 	FORCEINLINE bool AllowReceiveTickEventOnDedicatedServer() const { return bAllowReceiveTickEventOnDedicatedServer; }
 
+	/** Layer's the actor belongs to.  This is outside of the editoronly data to allow hiding of LD-specified layers at runtime for profiling. */
+	UPROPERTY()
+	TArray< FName > Layers;
+
 protected:
 
 #if WITH_EDITORONLY_DATA
@@ -439,10 +443,6 @@ public:
 	UPROPERTY()
 	FName Layer_DEPRECATED;
 
-	/** Layer's the actor belongs to */
-	UPROPERTY()
-	TArray< FName > Layers;
-
 	UPROPERTY()
 	FName Group_DEPRECATED;
 
diff --git a/Engine/Source/Runtime/Engine/Private/GameViewportClient.cpp b/Engine/Source/Runtime/Engine/Private/GameViewportClient.cpp
--- a/Engine/Source/Runtime/Engine/Private/GameViewportClient.cpp
+++ b/Engine/Source/Runtime/Engine/Private/GameViewportClient.cpp
@@ -1851,6 +1851,10 @@ bool UGameViewportClient::Exec( UWorld* InWorld, const TCHAR* Cmd,FOutputDevice&
 	{
 		return HandleShowCommand( Cmd, Ar, InWorld );
 	}
+	else if( FParse::Command(&Cmd,TEXT("SHOWLAYER")) )
+	{
+		return HandleShowLayerCommand( Cmd, Ar, InWorld );
+	}
 	else if (FParse::Command(&Cmd,TEXT("VIEWMODE")))
 	{
 		return HandleViewModeCommand( Cmd, Ar, InWorld );
@@ -2199,6 +2203,64 @@ bool UGameViewportClient::HandleShowCommand( const TCHAR* Cmd, FOutputDevice& Ar
 	return true;
 }
 
+bool UGameViewportClient::HandleShowLayerCommand( const TCHAR* Cmd, FOutputDevice& Ar, UWorld* InWorld )
+{
+	FString LayerName = FParse::Token(Cmd, 0);
+	bool bPrintValidEntries = false;
+
+	if (LayerName.IsEmpty())
+	{
+		Ar.Logf(TEXT("Missing layer name."));
+		bPrintValidEntries = true;
+	}
+	else
+	{
+		int32 NumActorsToggled = 0;
+		FName LayerFName = FName(*LayerName);
+
+		for (FActorIterator It(InWorld); It; ++It)
+		{
+			AActor* Actor = *It;
+			
+			if (Actor->Layers.Contains(LayerFName))
+			{
+				NumActorsToggled++;
+				// Note: overriding existing hidden property, ideally this would be something orthogonal
+				Actor->bHidden = !Actor->bHidden;
+
+				Actor->MarkComponentsRenderStateDirty();
+			}
+		}
+
+		Ar.Logf(TEXT("Toggled visibility of %u actors"), NumActorsToggled);
+		bPrintValidEntries = NumActorsToggled == 0;
+	}
+
+	if (bPrintValidEntries)
+	{
+		TArray<FName> LayerNames;
+
+		for (FActorIterator It(InWorld); It; ++It)
+		{
+			AActor* Actor = *It;
+
+			for (int32 LayerIndex = 0; LayerIndex < Actor->Layers.Num(); LayerIndex++)
+			{
+				LayerNames.AddUnique(Actor->Layers[LayerIndex]);
+			}
+		}
+
+		Ar.Logf(TEXT("Valid layer names:"));
+
+		for (int32 LayerIndex = 0; LayerIndex < LayerNames.Num(); LayerIndex++)
+		{
+			Ar.Logf(TEXT("   %s"), *LayerNames[LayerIndex].ToString());
+		}
+	}
+
+	return true;
+}
+
 bool UGameViewportClient::HandleViewModeCommand( const TCHAR* Cmd, FOutputDevice& Ar, UWorld* InWorld )
 {
 #if !UE_BUILD_DEBUG
