Avoid duplicate ChildActors when spawned via ChildActorComponent in standalone games as well

diff --git a/Engine/Source/Runtime/Engine/Classes/Components/ChildActorComponent.h b/Engine/Source/Runtime/Engine/Classes/Components/ChildActorComponent.h
--- a/Engine/Source/Runtime/Engine/Classes/Components/ChildActorComponent.h
+++ b/Engine/Source/Runtime/Engine/Classes/Components/ChildActorComponent.h
@@ -15,7 +15,7 @@ class UChildActorComponent : public USceneComponent
 	TSubclassOf<AActor>	ChildActorClass;
 
 	/** The actor that we spawned and own */
-	UPROPERTY(BlueprintReadOnly, Category = ChildActorComponent, NonPIETransient)
+	UPROPERTY(BlueprintReadOnly, Category = ChildActorComponent, TextExportTransient)
 	AActor*	ChildActor;
 
 	// Begin ActorComponent interface.
