#ue4 - Add protection against NaNs from potential divisions by zero in Character movement.

diff --git a/Engine/Source/Runtime/Engine/Private/CharacterMovementComponent.cpp b/Engine/Source/Runtime/Engine/Private/CharacterMovementComponent.cpp
--- a/Engine/Source/Runtime/Engine/Private/CharacterMovementComponent.cpp
+++ b/Engine/Source/Runtime/Engine/Private/CharacterMovementComponent.cpp
@@ -12,6 +12,7 @@ DEFINE_LOG_CATEGORY_STATIC(LogCharacterMovement, Log, All);
 // MAGIC NUMBERS
 const float MAX_STEP_SIDE_Z = 0.08f;	// maximum z value for the normal on the vertical side of steps
 const float SWIMBOBSPEED = -80.f;
+const float MIN_TICK_TIME = 0.0002f;	// minimum delta time considered when ticking.
 
 const float UCharacterMovementComponent::MIN_FLOOR_DIST = 1.9f;
 const float UCharacterMovementComponent::MAX_FLOOR_DIST = 2.4f;
@@ -1517,7 +1518,7 @@ void UCharacterMovementComponent::UnCrouch(bool bClientSimulation)
 
 void UCharacterMovementComponent::StartNewPhysics(float deltaTime, int32 Iterations)
 {
-	if ( (deltaTime < 0.0003f) || (Iterations > 7) || !HasValidData())
+	if ( (deltaTime < MIN_TICK_TIME) || (Iterations > 7) || !HasValidData())
 		return;
 
 	if (UpdatedComponent->IsSimulatingPhysics())
@@ -2169,6 +2170,11 @@ void UCharacterMovementComponent::ApplyVelocityBraking(float DeltaTime, float Fr
 
 void UCharacterMovementComponent::PhysFlying(float deltaTime, int32 Iterations)
 {
+	if (deltaTime < MIN_TICK_TIME)
+	{
+		return;
+	}
+
 	if( !HasRootMotion() )
 	{
 		if( bCheatFlying && Acceleration.IsZero() )
@@ -2221,6 +2227,11 @@ void UCharacterMovementComponent::PhysFlying(float deltaTime, int32 Iterations)
 
 void UCharacterMovementComponent::PhysSwimming(float deltaTime, int32 Iterations)
 {
+	if (deltaTime < MIN_TICK_TIME)
+	{
+		return;
+	}
+
 	float NetFluidFriction  = 0.f;
 	float Depth = ImmersionDepth();
 	float NetBuoyancy = Buoyancy * Depth;
@@ -2266,7 +2277,7 @@ void UCharacterMovementComponent::PhysSwimming(float deltaTime, int32 Iterations
 		OldLocation.Z = CharacterOwner->GetActorLocation().Z + (OldLocation.Z - stepZ);
 	}
 
-	if( !HasRootMotion() && !bJustTeleported && (remainingTime < deltaTime) && CharacterOwner )
+	if( !HasRootMotion() && !bJustTeleported && ((deltaTime - remainingTime) > KINDA_SMALL_NUMBER) && CharacterOwner )
 	{
 		bool bWaterJump = !GetPhysicsVolume()->bWaterVolume;
 		float velZ = Velocity.Z;
@@ -2292,32 +2303,34 @@ void UCharacterMovementComponent::PhysSwimming(float deltaTime, int32 Iterations
 
 void UCharacterMovementComponent::StartSwimming(FVector OldLocation, FVector OldVelocity, float timeTick, float remainingTime, int32 Iterations)
 {
+	if (remainingTime < MIN_TICK_TIME || timeTick < MIN_TICK_TIME)
+	{
+		return;
+	}
+
 	if( !HasRootMotion() && !bJustTeleported )
 	{
-		if ( timeTick > 0.f )
-		{
-			Velocity = (CharacterOwner->GetActorLocation() - OldLocation)/timeTick; //actual average velocity
-		}
+		Velocity = (CharacterOwner->GetActorLocation() - OldLocation)/timeTick; //actual average velocity
 		Velocity = 2.f*Velocity - OldVelocity; //end velocity has 2* accel of avg
-		if (Velocity.SizeSquared() > FMath::Square(GetPhysicsVolume()->TerminalVelocity))
-		{
-			Velocity = Velocity.SafeNormal();
-			Velocity *= GetPhysicsVolume()->TerminalVelocity;
-		}
+		Velocity.ClampMaxSize(GetPhysicsVolume()->TerminalVelocity);
 	}
-	FVector End = FindWaterLine(CharacterOwner->GetActorLocation(), OldLocation);
+	const FVector End = FindWaterLine(CharacterOwner->GetActorLocation(), OldLocation);
 	float waterTime = 0.f;
 	if (End != CharacterOwner->GetActorLocation())
 	{	
-		waterTime = timeTick * (End - CharacterOwner->GetActorLocation()).Size()/(CharacterOwner->GetActorLocation() - OldLocation).Size();
-		remainingTime += waterTime;
+		const float ActualDist = (CharacterOwner->GetActorLocation() - OldLocation).Size();
+		if (ActualDist > KINDA_SMALL_NUMBER)
+		{
+			waterTime = timeTick * (End - CharacterOwner->GetActorLocation()).Size() / ActualDist;
+			remainingTime += waterTime;
+		}
 		MoveUpdatedComponent(End - CharacterOwner->GetActorLocation(), CharacterOwner->GetActorRotation(), true);
 	}
 	if ( !HasRootMotion() && CharacterOwner && (Velocity.Z > 2.f*SWIMBOBSPEED) && (Velocity.Z < 0.f)) //allow for falling out of water
 	{
 		Velocity.Z = SWIMBOBSPEED - Velocity.Size2D() * 0.7f; //smooth bobbing
 	}
-	if ( (remainingTime > 0.01f) && (Iterations < 8) && CharacterOwner )
+	if ( (remainingTime >= MIN_TICK_TIME) && (Iterations < 8) && CharacterOwner )
 	{
 		PhysSwimming(remainingTime, Iterations);
 	}
@@ -2331,10 +2344,11 @@ float UCharacterMovementComponent::Swim(FVector Delta, FHitResult &Hit)
 
 	if ( !GetPhysicsVolume()->bWaterVolume ) //then left water
 	{
-		FVector End = FindWaterLine(Start,CharacterOwner->GetActorLocation());
-		if (End != CharacterOwner->GetActorLocation())
+		const FVector End = FindWaterLine(Start,CharacterOwner->GetActorLocation());
+		const float DesiredDist = Delta.Size();
+		if (End != CharacterOwner->GetActorLocation() && DesiredDist > KINDA_SMALL_NUMBER)
 		{
-			airTime = (End - CharacterOwner->GetActorLocation()).Size()/Delta.Size();
+			airTime = (End - CharacterOwner->GetActorLocation()).Size() / DesiredDist;
 			if ( ((CharacterOwner->GetActorLocation() - Start) | (End - CharacterOwner->GetActorLocation())) > 0.f )
 			{
 				airTime = 0.f;
@@ -2386,6 +2400,11 @@ void UCharacterMovementComponent::NotifyJumpApex()
 
 void UCharacterMovementComponent::PhysFalling(float deltaTime, int32 Iterations)
 {
+	if (deltaTime < MIN_TICK_TIME)
+	{
+		return;
+	}
+
 	// Bound final 2d portion of velocity
 	const float Speed2d = Velocity.Size2D();
 	const float BoundSpeed = FMath::Max(Speed2d, GetModifiedMaxSpeed() * AnalogInputModifier);
@@ -2450,14 +2469,14 @@ void UCharacterMovementComponent::PhysFalling(float deltaTime, int32 Iterations)
 	float remainingTime = deltaTime;
 	float timeTick = 0.1f;
 
-	while( (remainingTime > 0.f) && (Iterations < 8) )
+	while( (remainingTime >= MIN_TICK_TIME) && (Iterations < 8) )
 	{
 		Iterations++;
-		timeTick = (remainingTime > 0.05f) 
-			? FMath::Min(0.05f, remainingTime * 0.5f)
-			: remainingTime;
-
+		
+		// subdivide moves to be no longer than 0.05 seconds, no less than MIN_TICK_TIME.
+		const float timeTick = FMath::Max(MIN_TICK_TIME, (remainingTime > 0.05f) ? FMath::Min(0.05f, remainingTime * 0.5f) : remainingTime);
 		remainingTime -= timeTick;
+		
 		const FVector OldLocation = CharacterOwner->GetActorLocation();
 		const FRotator PawnRotation = CharacterOwner->GetActorRotation();
 		bJustTeleported = false;
@@ -2501,7 +2520,7 @@ void UCharacterMovementComponent::PhysFalling(float deltaTime, int32 Iterations)
 		
 		if ( IsSwimming() ) //just entered water
 		{
-			remainingTime = remainingTime + timeTick * (1.f - Hit.Time);
+			remainingTime += timeTick * (1.f - Hit.Time);
 			StartSwimming(OldLocation, OldVelocity, timeTick, remainingTime, Iterations);
 			return;
 		}
@@ -2890,6 +2909,11 @@ void UCharacterMovementComponent::MaintainHorizontalGroundVelocity()
 
 void UCharacterMovementComponent::PhysWalking(float deltaTime, int32 Iterations)
 {
+	if (deltaTime < MIN_TICK_TIME)
+	{
+		return;
+	}
+
 	if( (!CharacterOwner || !CharacterOwner->Controller) && !bRunPhysicsWithNoController && !HasRootMotion() )
 	{
 		Acceleration = FVector::ZeroVector;
@@ -2913,6 +2937,8 @@ void UCharacterMovementComponent::PhysWalking(float deltaTime, int32 Iterations)
 		CalcVelocity(deltaTime, GroundFriction, false, BrakingDecelerationWalking);
 	}
 
+	checkf(!Velocity.ContainsNaN(), TEXT("PhysWalking: Velocity contains NaN (%s: %s)\n%s"), *GetPathNameSafe(this), *GetPathNameSafe(GetOuter()), *Velocity.ToString());
+
 	FVector DesiredMove = Velocity;
 	DesiredMove.Z = 0.f;
 
@@ -2923,13 +2949,14 @@ void UCharacterMovementComponent::PhysWalking(float deltaTime, int32 Iterations)
 	const FVector PreviousBaseLocation = (OldBase != NULL) ? OldBase->GetComponentLocation() : FVector::ZeroVector;
 	bJustTeleported = false;
 	bool bCheckedFall = false;
+	bool bTriedLedgeMove = false;
 	float remainingTime = deltaTime;
 
-	while ( (remainingTime > 0.f) && (Iterations < 8) && (CharacterOwner->Controller || bRunPhysicsWithNoController || HasRootMotion()) )
+	while ( (remainingTime >= MIN_TICK_TIME) && (Iterations < 8) && (CharacterOwner->Controller || bRunPhysicsWithNoController || HasRootMotion()) )
 	{
 		Iterations++;
-		// subdivide moves to be no longer than 0.05 seconds
-		const float timeTick = (remainingTime > 0.05f) ? FMath::Min(0.05f, remainingTime * 0.5f) : remainingTime;
+		// subdivide moves to be no longer than 0.05 seconds, no less than MIN_TICK_TIME.
+		const float timeTick = FMath::Max(MIN_TICK_TIME, (remainingTime > 0.05f) ? FMath::Min(0.05f, remainingTime * 0.5f) : remainingTime);
 		remainingTime -= timeTick;
 		const FVector Delta = timeTick * DesiredMove;
 		const FVector subLoc = CharacterOwner->GetActorLocation();
@@ -2952,9 +2979,12 @@ void UCharacterMovementComponent::PhysWalking(float deltaTime, int32 Iterations)
 			if ( IsFalling() )
 			{
 				// pawn decided to jump up
-				const float ActualDist = (CharacterOwner->GetActorLocation() - subLoc).Size2D();
 				const float DesiredDist = Delta.Size();
-				remainingTime += timeTick * (1 - FMath::Min(1.f,ActualDist/DesiredDist));
+				if (DesiredDist > KINDA_SMALL_NUMBER)
+				{
+					const float ActualDist = (CharacterOwner->GetActorLocation() - subLoc).Size2D();
+					remainingTime += timeTick * (1.f - FMath::Min(1.f,ActualDist/DesiredDist));
+				}
 				StartNewPhysics(remainingTime,Iterations);
 				return;
 			}
@@ -2982,17 +3012,19 @@ void UCharacterMovementComponent::PhysWalking(float deltaTime, int32 Iterations)
 		{
 			// calculate possible alternate movement
 			const FVector GravDir = FVector(0.f,0.f,-1.f);
-			FVector NewDelta = GetLedgeMove(OldLocation, Delta, GravDir);
+			const FVector NewDelta = bTriedLedgeMove ? FVector::ZeroVector : GetLedgeMove(OldLocation, Delta, GravDir);
 			if ( !NewDelta.IsZero() )
 			{
 				// first revert this move
 				RevertMove(OldLocation, OldBase, PreviousBaseLocation, OldFloor, false);
 
-				// @todo hunting down NaN TTP 304692
-				check (timeTick!=0.f);
-				// redo move using NewDelta
+				// avoid repeated ledge moves if the first one fails
+				bTriedLedgeMove = true;
+
+				// Try new movement direction
 				DesiredMove = NewDelta/timeTick;
 				remainingTime += timeTick;
+				continue;
 			}
 			else
 			{
@@ -3009,6 +3041,7 @@ void UCharacterMovementComponent::PhysWalking(float deltaTime, int32 Iterations)
 				//UE_LOG(LogCharacterMovement, Log, TEXT("%s REVERT MOVE 1"), *CharacterOwner->GetName());
 				RevertMove(OldLocation, OldBase, PreviousBaseLocation, OldFloor, true);
 				remainingTime = 0.f;
+				break;
 			}
 		}
 		else
@@ -3062,15 +3095,17 @@ void UCharacterMovementComponent::PhysWalking(float deltaTime, int32 Iterations)
 		// If we didn't move at all this iteration then abort (since future iterations will also be stuck).
 		if (CharacterOwner->GetActorLocation() == subLoc)
 		{
+			remainingTime = 0.f;
 			break;
 		}
 	}
 
+
 	// Allow overlap events and such to change physics state and velocity
 	if( IsMovingOnGround() )
 	{
 		// Make velocity reflect actual move
-		if( !bJustTeleported && !HasRootMotion() )
+		if( !bJustTeleported && !HasRootMotion() && deltaTime >= MIN_TICK_TIME)
 		{
 			Velocity = (CharacterOwner->GetActorLocation() - OldLocation) / deltaTime;
 		}
