#ue4 - More NaN checks and divide-by-zero prevention in CharacterMovement

diff --git a/Engine/Source/Runtime/Engine/Private/CharacterMovementComponent.cpp b/Engine/Source/Runtime/Engine/Private/CharacterMovementComponent.cpp
--- a/Engine/Source/Runtime/Engine/Private/CharacterMovementComponent.cpp
+++ b/Engine/Source/Runtime/Engine/Private/CharacterMovementComponent.cpp
@@ -1225,12 +1225,7 @@ void UCharacterMovementComponent::PerformMovement(float DeltaSeconds)
 	}
 
 	// NaN tracking
-	if( Velocity.ContainsNaN() )
-	{
-		UE_LOG(LogCharacterMovement, Error, TEXT("UCharacterMovementComponent::PerformMovement Velocity contains NaN (%s). DeltaSeconds: %f"), 
-			*Velocity.ToString(), DeltaSeconds );
-		check( !Velocity.ContainsNaN() && "UCharacterMovementComponent::PerformMovement");
-	}
+	checkf(!Velocity.ContainsNaN(), TEXT("UCharacterMovementComponent::PerformMovement: Velocity contains NaN (%s: %s)\n%s"), *GetPathNameSafe(this), *GetPathNameSafe(GetOuter()), *Velocity.ToString());
 
 	// change position
 	StartNewPhysics(DeltaSeconds, 0);
@@ -1815,7 +1810,7 @@ bool UCharacterMovementComponent::IsCrouching() const
 void UCharacterMovementComponent::CalcVelocity(float DeltaTime, float Friction, bool bFluid, float BrakingDeceleration)
 {
 	// Do not update velocity when using root motion
-	if (!HasValidData() || HasRootMotion())
+	if (!HasValidData() || HasRootMotion() || DeltaTime < MIN_TICK_TIME)
 	{
 		return;
 	}
@@ -2128,7 +2123,7 @@ FVector UCharacterMovementComponent::GetCurrentAcceleration() const
 
 void UCharacterMovementComponent::ApplyVelocityBraking(float DeltaTime, float Friction, float BrakingDeceleration)
 {
-	if (Velocity.IsZero() || !HasValidData() || HasRootMotion())
+	if (Velocity.IsZero() || !HasValidData() || HasRootMotion() || DeltaTime < MIN_TICK_TIME)
 	{
 		return;
 	}
@@ -2143,7 +2138,7 @@ void UCharacterMovementComponent::ApplyVelocityBraking(float DeltaTime, float Fr
 
 	// Decelerate to brake to a stop
 	const FVector RevAccel = (-1.f * BrakingDeceleration) * SafeNormalPrecise(Velocity);
-	while( RemainingTime > 0.f )
+	while( RemainingTime >= MIN_TICK_TIME )
 	{
 		float dt = ((RemainingTime > TimeStep) ? TimeStep : RemainingTime);
 		RemainingTime -= dt;
@@ -2929,16 +2924,16 @@ void UCharacterMovementComponent::PhysWalking(float deltaTime, int32 Iterations)
 
 	// Ensure velocity is horizontal.
 	MaintainHorizontalGroundVelocity();
+	checkf(!Velocity.ContainsNaN(), TEXT("PhysWalking: Velocity contains NaN before CalcVelocity (%s: %s)\n%s"), *GetPathNameSafe(this), *GetPathNameSafe(GetOuter()), *Velocity.ToString());
 
 	//bound acceleration
 	Acceleration.Z = 0.f;
 	if( !HasRootMotion() )
 	{
 		CalcVelocity(deltaTime, GroundFriction, false, BrakingDecelerationWalking);
+		checkf(!Velocity.ContainsNaN(), TEXT("PhysWalking: Velocity contains NaN after CalcVelocity (%s: %s)\n%s"), *GetPathNameSafe(this), *GetPathNameSafe(GetOuter()), *Velocity.ToString());
 	}
 
-	checkf(!Velocity.ContainsNaN(), TEXT("PhysWalking: Velocity contains NaN (%s: %s)\n%s"), *GetPathNameSafe(this), *GetPathNameSafe(GetOuter()), *Velocity.ToString());
-
 	FVector DesiredMove = Velocity;
 	DesiredMove.Z = 0.f;
 
@@ -2970,7 +2965,7 @@ void UCharacterMovementComponent::PhysWalking(float deltaTime, int32 Iterations)
 		else
 		{
 			// @todo hunting down NaN TTP 304692
-			checkf( !Delta.ContainsNaN(), TEXT("PhysWalking: NewTransform contains NaN (%s: %s)\n%s"), *GetPathNameSafe(this), *GetPathNameSafe(GetOuter()), *Delta.ToString() );
+			checkf( !Delta.ContainsNaN(), TEXT("PhysWalking: Delta contains NaN (%s: %s)\ntimeTick = %d Velocity = %s Delta = %s"), *GetPathNameSafe(this), *GetPathNameSafe(GetOuter()), timeTick, *Velocity.ToString(), *Delta.ToString() );
 			// @todo hunting down NaN TTP 304692
 
 			// try to move forward
