Added a way to have BP-implementable EQS contexts #UE4

diff --git a/Engine/Source/Runtime/Engine/Classes/AI/EnvironmentQuery/Contexts/EnvQueryContext_BlueprintBase.h b/Engine/Source/Runtime/Engine/Classes/AI/EnvironmentQuery/Contexts/EnvQueryContext_BlueprintBase.h
new file mode 100644
--- /dev/null
+++ b/Engine/Source/Runtime/Engine/Classes/AI/EnvironmentQuery/Contexts/EnvQueryContext_BlueprintBase.h
@@ -0,0 +1,36 @@
+// Copyright 1998-2014 Epic Games, Inc. All Rights Reserved.
+
+#pragma once
+
+#include "EnvQueryContext_BlueprintBase.generated.h"
+
+UCLASS(MinimalAPI, Blueprintable)
+class UEnvQueryContext_BlueprintBase : public UEnvQueryContext
+{
+	GENERATED_UCLASS_BODY()
+
+	enum ECallMode
+	{
+		InvalidCallMode,
+		SingleActor,
+		SingleLocation,
+		ActorSet,
+		LocationSet
+	};
+
+	ECallMode CallMode;
+
+	virtual void ProvideContext(struct FEnvQueryInstance& QueryInstance, struct FEnvQueryContextData& ContextData) const OVERRIDE;
+
+	UFUNCTION(BlueprintImplementableEvent)
+	virtual void ProvideSingleActor(AActor* QuerierActor, AActor*& ResultingActor) const;
+
+	UFUNCTION(BlueprintImplementableEvent)
+	virtual void ProvideSingleLocation(AActor* QuerierActor, FVector& ResultingLocation) const;
+
+	UFUNCTION(BlueprintImplementableEvent)
+	virtual void ProvideActorsSet(AActor* QuerierActor, TArray<AActor*>& ResultingActorsSet) const;
+
+	UFUNCTION(BlueprintImplementableEvent)
+	virtual void ProvideLocationsSet(AActor* QuerierActor, TArray<FVector>& ResultingLocationSet) const;
+};
diff --git a/Engine/Source/Runtime/Engine/Private/AI/EnvironmentQuery/Contexts/EnvQueryContext_BlueprintBase.cpp b/Engine/Source/Runtime/Engine/Private/AI/EnvironmentQuery/Contexts/EnvQueryContext_BlueprintBase.cpp
new file mode 100644
--- /dev/null
+++ b/Engine/Source/Runtime/Engine/Private/AI/EnvironmentQuery/Contexts/EnvQueryContext_BlueprintBase.cpp
@@ -0,0 +1,87 @@
+// Copyright 1998-2014 Epic Games, Inc. All Rights Reserved.
+
+#include "EnginePrivate.h"
+
+UEnvQueryContext_BlueprintBase::UEnvQueryContext_BlueprintBase(const class FPostConstructInitializeProperties& PCIP) 
+	: Super(PCIP), CallMode(UEnvQueryContext_BlueprintBase::InvalidCallMode)
+{
+	UClass* StopAtClass = UEnvQueryContext_BlueprintBase::StaticClass();
+	bool bImplementsProvideSingleActor = BlueprintNodeHelpers::HasBlueprintFunction(GET_FUNCTION_NAME_CHECKED(UEnvQueryContext_BlueprintBase, ProvideSingleActor), this, StopAtClass);
+	bool bImplementsProvideSingleLocation = BlueprintNodeHelpers::HasBlueprintFunction(GET_FUNCTION_NAME_CHECKED(UEnvQueryContext_BlueprintBase, ProvideSingleLocation), this, StopAtClass);
+	bool bImplementsProvideActorSet = BlueprintNodeHelpers::HasBlueprintFunction(GET_FUNCTION_NAME_CHECKED(UEnvQueryContext_BlueprintBase, ProvideActorsSet), this, StopAtClass);
+	bool bImplementsProvideLocationsSet = BlueprintNodeHelpers::HasBlueprintFunction(GET_FUNCTION_NAME_CHECKED(UEnvQueryContext_BlueprintBase, ProvideLocationsSet), this, StopAtClass);
+
+	int32 ImplementationsCount = 0;
+	if (bImplementsProvideSingleActor)
+	{
+		++ImplementationsCount;
+		CallMode = SingleActor;
+	}
+
+	if (bImplementsProvideSingleLocation)
+	{
+		++ImplementationsCount;
+		CallMode = SingleLocation;
+	}
+
+	if (bImplementsProvideActorSet)
+	{
+		++ImplementationsCount;
+		CallMode = ActorSet;
+	}
+
+	if (bImplementsProvideLocationsSet)
+	{
+		++ImplementationsCount;
+		CallMode = LocationSet;
+	}
+
+	if (ImplementationsCount != 1)
+	{
+
+	}
+}
+
+void UEnvQueryContext_BlueprintBase::ProvideContext(struct FEnvQueryInstance& QueryInstance, struct FEnvQueryContextData& ContextData) const
+{
+	AActor* QuerierActor = Cast<AActor>(QueryInstance.Owner.Get());
+
+	if (QuerierActor == NULL || CallMode == InvalidCallMode)
+	{
+		return;
+	}
+	
+	switch (CallMode)
+	{
+	case SingleActor:
+		{
+			AActor* ResultingActor = NULL;
+			ProvideSingleActor(QuerierActor, ResultingActor);
+			UEnvQueryItemType_Actor::SetContextHelper(ContextData, ResultingActor);
+		}
+		break;
+	case SingleLocation:
+		{
+			FVector ResultingLocation = FAISystem::InvalidLocation;
+			ProvideSingleLocation(QuerierActor, ResultingLocation);
+			UEnvQueryItemType_Point::SetContextHelper(ContextData, ResultingLocation);
+		}
+		break;
+	case ActorSet:
+		{
+			TArray<AActor*> ActorSet;
+			ProvideActorsSet(QuerierActor, ActorSet);
+			UEnvQueryItemType_Actor::SetContextHelper(ContextData, ActorSet);
+		}
+		break;
+	case LocationSet:
+		{
+			TArray<FVector> LocationSet;
+			ProvideLocationsSet(QuerierActor, LocationSet);
+			UEnvQueryItemType_Point::SetContextHelper(ContextData, LocationSet);
+		}
+		break;
+	default:
+		break;
+	}
+}
