UE4: Initial world asset renaming support.
Renaming a world also renames it's level blueprint.
Moving a world to a new package also transfers over all of it's lightmaps and shadowmaps.
Redirectors are allowed to UObject* (hard) reference assets in another package, but all other hard references are still disallowed.
Referencing a redirector package in a UStreamingLevel now follows the redirector to the world in the new package.
Packages that lose their worlds now also lose the PKG_ContainsMap flag. Packages that gain a world also gain this flag.

diff --git a/Engine/Source/Editor/DetailCustomizations/Private/WorldSettingsDetails.cpp b/Engine/Source/Editor/DetailCustomizations/Private/WorldSettingsDetails.cpp
--- a/Engine/Source/Editor/DetailCustomizations/Private/WorldSettingsDetails.cpp
+++ b/Engine/Source/Editor/DetailCustomizations/Private/WorldSettingsDetails.cpp
@@ -253,14 +253,13 @@ void FLightmapCustomNodeBuilder::RefreshLightmapItems()
 	UWorld* World = Context.World();
 	if ( World )
 	{
-		TArray<UObject*> WorldPackageObjects;
-		const bool bIncludeNestedObjects = false;
-		GetObjectsWithOuter(World->GetOutermost(), WorldPackageObjects, bIncludeNestedObjects);
+		TArray<UTexture2D*> LightMapsAndShadowMaps;
+		World->GetLightMapsAndShadowMaps(LightMapsAndShadowMaps);
 
-		for ( auto ObjIt = WorldPackageObjects.CreateConstIterator(); ObjIt; ++ObjIt )
+		for ( auto ObjIt = LightMapsAndShadowMaps.CreateConstIterator(); ObjIt; ++ObjIt )
 		{
-			UObject* CurrentObject = *ObjIt;
-			if ( CurrentObject->IsA(ULightMapTexture2D::StaticClass()) || CurrentObject->IsA(UShadowMapTexture2D::StaticClass()) )
+			UTexture2D* CurrentObject = *ObjIt;
+			if (CurrentObject)
 			{
 				FAssetData AssetData = FAssetData(CurrentObject);
 				const uint32 ThumbnailResolution = 64;
diff --git a/Engine/Source/Runtime/CoreUObject/Private/UObject/SavePackage.cpp b/Engine/Source/Runtime/CoreUObject/Private/UObject/SavePackage.cpp
--- a/Engine/Source/Runtime/CoreUObject/Private/UObject/SavePackage.cpp
+++ b/Engine/Source/Runtime/CoreUObject/Private/UObject/SavePackage.cpp
@@ -2551,6 +2551,9 @@ bool UPackage::SavePackage( UPackage* InOuter, UObject* Base, EObjectFlags TopLe
 
 				// structure to track what ever export needs to import
 				TMap<UObject*, TArray<UObject*> > ObjectDependencies;
+
+				// and a structure to track non-redirector references
+				TSet<UObject*> DependenciesReferencedByNonRedirectors;
 		
 				/** If true, we are going to compress the package to memory to save a little time */
 				bool bCompressFromMemory = !!(InOuter->PackageFlags & PKG_StoreCompressed);
@@ -2660,6 +2663,18 @@ bool UPackage::SavePackage( UPackage* InOuter, UObject* Base, EObjectFlags TopLe
 
 						// add the list of dependencies to the dependency map
 						ObjectDependencies.Add(Obj, ImportTagger.Dependencies);
+
+						if (Obj->GetClass() != UObjectRedirector::StaticClass())
+						{
+							for ( auto DepIt = ImportTagger.Dependencies.CreateConstIterator(); DepIt; ++DepIt )
+							{
+								UObject* DependencyObject = *DepIt;
+								if ( DependencyObject )
+								{
+									DependenciesReferencedByNonRedirectors.Add(DependencyObject);
+								}
+							}
+						}
 					}
 				}
 
@@ -2720,11 +2735,14 @@ bool UPackage::SavePackage( UPackage* InOuter, UObject* Base, EObjectFlags TopLe
 								{
 									ObjectsInOtherMaps.Add(Obj);
 
-									UE_LOG(LogSavePackage, Warning, TEXT( "Obj in another map: %s"), *Obj->GetFullName() );
-
-									if (!(SaveFlags & SAVE_NoError))
+									if ( DependenciesReferencedByNonRedirectors.Contains(Obj) )
 									{
-										Error->Logf(ELogVerbosity::Warning, *FText::Format( NSLOCTEXT( "Core", "SavePackageObjInAnotherMap", "Object '{0}' is in another map" ), FText::FromString( *Obj->GetFullName() ) ).ToString() );
+										UE_LOG(LogSavePackage, Warning, TEXT( "Obj in another map: %s"), *Obj->GetFullName() );
+
+										if (!(SaveFlags & SAVE_NoError))
+										{
+											Error->Logf(ELogVerbosity::Warning, *FText::Format( NSLOCTEXT( "Core", "SavePackageObjInAnotherMap", "Object '{0}' is in another map" ), FText::FromString( *Obj->GetFullName() ) ).ToString() );
+										}
 									}
 								}
 								else
@@ -2747,8 +2765,19 @@ bool UPackage::SavePackage( UPackage* InOuter, UObject* Base, EObjectFlags TopLe
 					ObjectsInOtherMaps = LevelObjects;
 				}
 
+				// It is allowed for redirectors to reference objects in other maps.
+				// Form the list of objects that erroneously reference another map.
+				TArray<UObject*> IllegalObjectsInOtherMaps;
+				for ( auto ObjIt = ObjectsInOtherMaps.CreateConstIterator(); ObjIt; ++ObjIt )
+				{
+					if ( DependenciesReferencedByNonRedirectors.Contains(*ObjIt) )
+					{
+						IllegalObjectsInOtherMaps.Add(*ObjIt);
+					}
+				}
+
 				// The graph is linked to objects in a different map package!
-				if( ObjectsInOtherMaps.Num() )
+				if( IllegalObjectsInOtherMaps.Num() )
 				{
 					UObject* MostLikelyCulprit = NULL;
 					const UProperty* PropertyRef = NULL;
@@ -2758,15 +2787,15 @@ bool UPackage::SavePackage( UPackage* InOuter, UObject* Base, EObjectFlags TopLe
 					int32 MaxNamesToDisplay = 5;
 					bool DisplayIsLimited = true;
 
-					if (ObjectsInOtherMaps.Num() < MaxNamesToDisplay)
+					if (IllegalObjectsInOtherMaps.Num() < MaxNamesToDisplay)
 					{
-						MaxNamesToDisplay = ObjectsInOtherMaps.Num();
+						MaxNamesToDisplay = IllegalObjectsInOtherMaps.Num();
 						DisplayIsLimited = false;
 					}
 
 					for (int32 Idx = 0; Idx < MaxNamesToDisplay; Idx++)
 					{
-						ObjectNames += ObjectsInOtherMaps[Idx]->GetName() + TEXT("\n");
+						ObjectNames += IllegalObjectsInOtherMaps[Idx]->GetName() + TEXT("\n");
 					}
 					
 					// if there are more than 5 items we indicated this by adding "..." at the end of the list
@@ -2784,7 +2813,7 @@ bool UPackage::SavePackage( UPackage* InOuter, UObject* Base, EObjectFlags TopLe
 					const bool bProceed = FMessageDialog::Open(EAppMsgType::YesNo, Message ) == EAppReturnType::Yes;
 					if ( bProceed )
 					{
-						FindMostLikelyCulprit( ObjectsInOtherMaps, MostLikelyCulprit, PropertyRef );
+						FindMostLikelyCulprit( IllegalObjectsInOtherMaps, MostLikelyCulprit, PropertyRef );
 						if( MostLikelyCulprit != NULL && PropertyRef != NULL )
 						{
 							CulpritString = FString::Printf(TEXT("%s (%s)"), *MostLikelyCulprit->GetFullName(), *PropertyRef->GetName());
diff --git a/Engine/Source/Runtime/Engine/Classes/Engine/World.h b/Engine/Source/Runtime/Engine/Classes/Engine/World.h
--- a/Engine/Source/Runtime/Engine/Classes/Engine/World.h
+++ b/Engine/Source/Runtime/Engine/Classes/Engine/World.h
@@ -1623,6 +1623,10 @@ public:
 	virtual void PostSaveRoot( bool bCleanupIsRequired ) OVERRIDE;
 	virtual UWorld* GetWorld() const OVERRIDE;
 	static void AddReferencedObjects(UObject* InThis, FReferenceCollector& Collector);
+#if WITH_EDITOR
+	virtual bool Rename(const TCHAR* NewName = NULL, UObject* NewOuter = NULL, ERenameFlags Flags = REN_None) OVERRIDE;
+#endif
+
 	// End UObject Interface
 	
 	/**
@@ -2401,6 +2405,9 @@ public:
 	/** Updates all physics constraint actor joint locations.  */
 	virtual void UpdateConstraintActors();
 
+	/** Gets all LightMaps and ShadowMaps associated with this world's persistent level */
+	void GetLightMapsAndShadowMaps(TArray<UTexture2D*>& OutLightMapsAndShadowMaps);
+
 public:
 	static FString ConvertToPIEPackageName(const FString& PackageName, int32 PIEInstanceID);
 	static FString BuildPIEPackagePrefix(int32 PIEInstanceID);
diff --git a/Engine/Source/Runtime/Engine/Private/Blueprint.cpp b/Engine/Source/Runtime/Engine/Private/Blueprint.cpp
--- a/Engine/Source/Runtime/Engine/Private/Blueprint.cpp
+++ b/Engine/Source/Runtime/Engine/Private/Blueprint.cpp
@@ -239,9 +239,10 @@ bool UBlueprint::Rename( const TCHAR* InName, UObject* NewOuter, ERenameFlags Fl
 	FName SkelClassName, GenClassName;
 	GetBlueprintClassNames(GenClassName, SkelClassName, FName(InName));
 
+	UPackage* NewTopLevelObjectOuter = NewOuter ? NewOuter->GetOutermost() : NULL;
 	if(GeneratedClass != NULL)
 	{
-		bool bMovedOK = GeneratedClass->Rename(*GenClassName.ToString(), NewOuter, Flags);
+		bool bMovedOK = GeneratedClass->Rename(*GenClassName.ToString(), NewTopLevelObjectOuter, Flags);
 		if(!bMovedOK)
 		{
 			return false;
@@ -251,7 +252,7 @@ bool UBlueprint::Rename( const TCHAR* InName, UObject* NewOuter, ERenameFlags Fl
 	// Also move skeleton class, if different from generated class, to new package (again, to create redirector)
 	if(SkeletonGeneratedClass != NULL && SkeletonGeneratedClass != GeneratedClass)
 	{
-		bool bMovedOK = SkeletonGeneratedClass->Rename(*SkelClassName.ToString(), NewOuter, Flags);
+		bool bMovedOK = SkeletonGeneratedClass->Rename(*SkelClassName.ToString(), NewTopLevelObjectOuter, Flags);
 		if(!bMovedOK)
 		{
 			return false;
diff --git a/Engine/Source/Runtime/Engine/Private/Level.cpp b/Engine/Source/Runtime/Engine/Private/Level.cpp
--- a/Engine/Source/Runtime/Engine/Private/Level.cpp
+++ b/Engine/Source/Runtime/Engine/Private/Level.cpp
@@ -605,7 +605,15 @@ void ULevel::PostLoad()
 	if (Model == NULL || Model->NumUniqueVertices == 0)
 	{
 		StaticNavigableGeometry.Empty();
-	} 
+	}
+
+	// Rename the LevelScriptBlueprint after the outer world.
+	UWorld* OuterWorld = Cast<UWorld>(GetOuter());
+	if (LevelScriptBlueprint && OuterWorld && LevelScriptBlueprint->GetFName() != OuterWorld->GetFName())
+	{
+		// Use LevelScriptBlueprint->GetOuter() instead of NULL to make sure the generated top level objects are moved appropriately
+		LevelScriptBlueprint->Rename(*OuterWorld->GetName(), LevelScriptBlueprint->GetOuter(), REN_DoNotDirty | REN_DontCreateRedirectors | REN_ForceNoResetLoaders | REN_NonTransactional);
+	}
 }
 
 UWorld* ULevel::GetWorld() const
diff --git a/Engine/Source/Runtime/Engine/Private/LevelStreaming.cpp b/Engine/Source/Runtime/Engine/Private/LevelStreaming.cpp
--- a/Engine/Source/Runtime/Engine/Private/LevelStreaming.cpp
+++ b/Engine/Source/Runtime/Engine/Private/LevelStreaming.cpp
@@ -417,6 +417,8 @@ bool ULevelStreaming::RequestLevel(UWorld* PersistentWorld, bool bAllowLevelLoad
 			bHasLoadRequestPending = true;
 			
 			ULevel::StreamedLevelsOwningWorld.Add(DesiredPackageName, PersistentWorld);
+			UWorld::WorldTypePreLoadMap.FindOrAdd(DesiredPackageName) = PersistentWorld->WorldType;
+
 			// Kick off async load request.
 			LoadPackageAsync(*DesiredPackageName.ToString(), 
 				FLoadPackageAsyncDelegate::CreateUObject(this, &ULevelStreaming::AsyncLevelLoadComplete), 
@@ -451,42 +453,126 @@ void ULevelStreaming::AsyncLevelLoadComplete( const FString& InPackageName, UPac
 	
 	if( InLoadedPackage )
 	{
-		UPackage* LevelPackage = CastChecked<UPackage>(InLoadedPackage);
+		UPackage* LevelPackage = InLoadedPackage;
 		
 		// Try to find a UWorld object in the level package.
 		UWorld* World = UWorld::FindWorldInPackage(LevelPackage);
-		ULevel* Level = World ? World->PersistentLevel : NULL;	
-		if( Level )
-		{			
-			check(PendingUnloadLevel == NULL);
-			SetLoadedLevel(Level);
-
-			// Broadcast level loaded event to blueprints
-			OnLevelLoaded.Broadcast();
 
-			// Make sure this level will start to render only when it will be fully added to the world
-			if (LODPackageNames.Num() > 0)
-			{
-				Level->bRequireFullVisibilityToRender = true;
-			}
-			
-			// In the editor levels must be in the levels array regardless of whether they are visible or not
-			if (Level->OwningWorld->WorldType == EWorldType::Editor)
+		if ( World )
+		{
+			ULevel* Level = World->PersistentLevel;
+			if( Level )
 			{
-				Level->OwningWorld->AddLevel(Level);
-#if WITH_EDITOR
-				// We should also at this point, apply the level's editor transform
-				if (!Level->bAlreadyMovedActors)
+				check(PendingUnloadLevel == NULL);
+				SetLoadedLevel(Level);
+
+				// Broadcast level loaded event to blueprints
+				OnLevelLoaded.Broadcast();
+
+				// Make sure this level will start to render only when it will be fully added to the world
+				if (LODPackageNames.Num() > 0)
 				{
-					FLevelUtils::ApplyEditorTransform(this, false);
-					Level->bAlreadyMovedActors = true;
+					Level->bRequireFullVisibilityToRender = true;
 				}
+			
+				// In the editor levels must be in the levels array regardless of whether they are visible or not
+				if (ensure(Level->OwningWorld) && Level->OwningWorld->WorldType == EWorldType::Editor)
+				{
+					Level->OwningWorld->AddLevel(Level);
+#if WITH_EDITOR
+					// We should also at this point, apply the level's editor transform
+					if (!Level->bAlreadyMovedActors)
+					{
+						FLevelUtils::ApplyEditorTransform(this, false);
+						Level->bAlreadyMovedActors = true;
+					}
 #endif // WITH_EDITOR
+				}
+			}
+			else
+			{
+				UE_LOG(LogLevelStreaming, Warning, TEXT("Couldn't find ULevel object in package '%s'"), *InPackageName );
 			}
 		}
 		else
 		{
-			UE_LOG(LogLevelStreaming, Warning, TEXT("Couldn't find ULevel object in package '%s'"), *InPackageName );
+			// There could have been a redirector in the package. Attempt to follow it.
+			const FString ExpectedRedirectorName = (PackageNameToLoad != NAME_None) ? PackageNameToLoad.ToString() : PackageName.ToString();
+			UObjectRedirector* WorldRedirector = FindObject<UObjectRedirector>(LevelPackage, *FPackageName::GetLongPackageAssetName(*ExpectedRedirectorName));
+			if (WorldRedirector)
+			{
+				// A redirector was found! Make sure the destination world was loaded.
+				UWorld* DestinationWorld = Cast<UWorld>(WorldRedirector->DestinationObject);
+				if (DestinationWorld)
+				{
+					// To follow the world redirector for level streaming...
+					// 1) Update all globals that refer to the redirector package by name
+					// 2) Update the PackageNameToLoad to refer to the new package location
+					// 3) If the package name to load was the same as the destination package name...
+					//         ... update the package name to the new package and let the next RequestLevel try this process again.
+					//    If the package name to load was different...
+					//         ... it means the specified package name was explicit and we will just load from another file.
+
+					FName OldDesiredPackageName = FName(*InPackageName);
+					UWorld** OwningWorldPtr = ULevel::StreamedLevelsOwningWorld.Find(OldDesiredPackageName);
+					UWorld* OwningWorld = OwningWorldPtr ? *OwningWorldPtr : NULL;
+					if ( OwningWorld )
+					{
+						if ( DestinationWorld->PersistentLevel )
+						{
+							DestinationWorld->PersistentLevel->OwningWorld = OwningWorld;
+						}
+
+						// In some cases, BSP render data is not created because the OwningWorld was not set correctly.
+						// Regenerate that render data here
+						DestinationWorld->PersistentLevel->InvalidateModelSurface();
+						DestinationWorld->PersistentLevel->CommitModelSurfaces();
+
+						ULevel::StreamedLevelsOwningWorld.Remove(OldDesiredPackageName);
+					}
+
+					// Try again with the destination package to load.
+					// IMPORTANT: check this BEFORE changing PackageNameToLoad, otherwise you wont know if the package name was supposed to be different.
+					const bool bLoadingIntoDifferentPackage = (PackageName != PackageNameToLoad) && (PackageNameToLoad != NAME_None);
+
+					// ... now set PackageNameToLoad
+					PackageNameToLoad = DestinationWorld->GetOutermost()->GetFName();
+
+					if ( PackageNameToLoad != OldDesiredPackageName )
+					{
+						EWorldType::Type* OldPackageWorldType = UWorld::WorldTypePreLoadMap.Find(OldDesiredPackageName);
+						if ( OldPackageWorldType )
+						{
+							UWorld::WorldTypePreLoadMap.FindOrAdd(PackageNameToLoad) = *OldPackageWorldType;
+							UWorld::WorldTypePreLoadMap.Remove(OldDesiredPackageName);
+						}
+					}
+
+					// Now determine if we are loading into the package explicitly or if it is okay to just load the other package.
+					if ( bLoadingIntoDifferentPackage )
+					{
+						// Loading into a new custom package explicitly. Load the destination world directly into the package.
+						// Detach the linker to load from a new file into the same package.
+						ULinkerLoad* PackageLinker = ULinkerLoad::FindExistingLinkerForPackage(LevelPackage);
+						if (PackageLinker)
+						{
+							PackageLinker->Detach(false);
+						}
+
+						// Make sure the redirector is not in the way of the new world.
+						if (WorldRedirector->GetFName() == DestinationWorld->GetFName())
+						{
+							// Pass NULL as the name to make a new unique name and NULL for the outer to not change the outer.
+							WorldRedirector->Rename(NULL, NULL, REN_DoNotDirty | REN_DontCreateRedirectors | REN_ForceNoResetLoaders | REN_NonTransactional);
+						}
+					}
+					else
+					{
+						// Loading the requested package normally. Update the requested package to the destination.
+						PackageName = PackageNameToLoad;
+					}
+				}
+			}
 		}
 	}
 	else
diff --git a/Engine/Source/Runtime/Engine/Private/World.cpp b/Engine/Source/Runtime/Engine/Private/World.cpp
--- a/Engine/Source/Runtime/Engine/Private/World.cpp
+++ b/Engine/Source/Runtime/Engine/Private/World.cpp
@@ -196,6 +196,61 @@ void UWorld::AddReferencedObjects(UObject* InThis, FReferenceCollector& Collecto
 	Super::AddReferencedObjects( This, Collector );
 }
 
+#if WITH_EDITOR
+bool UWorld::Rename(const TCHAR* InName, UObject* NewOuter, ERenameFlags Flags)
+{
+	// Rename LightMaps and ShadowMaps to the new location. Keep the old name, since they are not named after the world.
+	TArray<UTexture2D*> LightMapsAndShadowMaps;
+	GetLightMapsAndShadowMaps(LightMapsAndShadowMaps);
+	for (auto* Tex : LightMapsAndShadowMaps)
+	{
+		if ( Tex && !Tex->Rename(*Tex->GetName(), NewOuter, Flags) )
+		{
+			return false;
+		}
+	}
+
+	// Rename the world itself
+	UPackage* OldPackage = GetOutermost();
+	if ( !Super::Rename(InName, NewOuter, Flags) )
+	{
+		return false;
+	}
+
+	// Rename the level script blueprint now, unless we are in PostLoad. ULevel::PostLoad should handle renaming this at load time.
+	if ( !GIsRoutingPostLoad )
+	{
+		const bool bDontCreate = true;
+		UBlueprint* LevelScriptBlueprint = PersistentLevel ? PersistentLevel->GetLevelScriptBlueprint(bDontCreate) : NULL;
+		if ( LevelScriptBlueprint )
+		{
+			// Use LevelScriptBlueprint->GetOuter() instead of NULL to make sure the generated top level objects are moved appropriately
+			if ( !LevelScriptBlueprint->Rename(InName, LevelScriptBlueprint->GetOuter(), Flags) )
+			{
+				return false;
+			}
+		}
+	}
+
+	// Update the PKG_ContainsMap package flag
+	UPackage* NewPackage = GetOutermost();
+	const bool bTestRename = (Flags & REN_Test) != 0;
+	if ( !bTestRename && NewPackage != OldPackage )
+	{
+		// If this is the last world removed from a package, clear the PKG_ContainsMap flag
+		if ( UWorld::FindWorldInPackage(OldPackage) == NULL )
+		{
+			OldPackage->PackageFlags &= ~PKG_ContainsMap;
+		}
+
+		// Set the PKG_ContainsMap flag in the new package
+		NewPackage->ThisContainsMap();
+	}
+
+	return true;
+}
+#endif
+
 void UWorld::FinishDestroy()
 {
 	// Avoid cleanup if the world hasn't been initialized, e.g., the default object or a world object that got loaded
@@ -246,6 +301,29 @@ void UWorld::FinishDestroy()
 		delete TimerManager;
 	}
 
+	// Remove the PKG_ContainsMap flag from packages that no longer contain a world
+	{
+		UPackage* WorldPackage = GetOutermost();
+
+		bool bContainsAnotherWorld = false;
+		TArray<UObject*> PotentialWorlds;
+		GetObjectsWithOuter(WorldPackage, PotentialWorlds, false);
+		for (auto ObjIt = PotentialWorlds.CreateConstIterator(); ObjIt; ++ObjIt)
+		{
+			UWorld* World = Cast<UWorld>(*ObjIt);
+			if (World && World != this)
+			{
+				bContainsAnotherWorld = true;
+				break;
+			}
+		}
+
+		if ( !bContainsAnotherWorld )
+		{
+			WorldPackage->PackageFlags &= ~PKG_ContainsMap;
+		}
+	}
+
 	Super::FinishDestroy();
 }
 
@@ -256,6 +334,11 @@ void UWorld::PostLoad()
 	{
 		WorldType = *PreLoadWorldType;
 	}
+	else
+	{
+		// Since we did not specify a world type, assume it was loaded from a package to persist in memory
+		WorldType = EWorldType::Inactive;
+	}
 
 	Super::PostLoad();
 	CurrentLevel = PersistentLevel;
@@ -382,6 +465,22 @@ void UWorld::PostLoad()
 			SetFlags(RF_Public|RF_Standalone);
 		}
 	}
+
+#if WITH_EDITOR
+	if (GIsEditor)
+	{
+		// Ensure the DefaultBrush's model has the same outer as the default brush itself. Older packages erroneously stored this object as a top-level package
+		ABrush* DefaultBrush = PersistentLevel->Actors.Num() < 2 ? NULL : Cast<ABrush>(PersistentLevel->Actors[1]);
+		UModel* Model = DefaultBrush ? DefaultBrush->Brush : NULL;
+		if (Model != NULL)
+		{
+			if (Model->GetOuter() != DefaultBrush->GetOuter())
+			{
+				Model->Rename(TEXT("Brush"), DefaultBrush->GetOuter(), REN_DoNotDirty | REN_DontCreateRedirectors | REN_ForceNoResetLoaders | REN_NonTransactional);
+			}
+		}
+	}
+#endif
 }
 
 
@@ -714,12 +813,6 @@ void UWorld::InitWorld(const InitializationValues IVS)
 		{
 			UModel* Model = DefaultBrush->Brush;
 
-			// Ensure the DefaultBrush's model has the same outer as the default brush itself. Older packages erroneously stored this object as a top-level package
-			if (Model->GetOuter() != DefaultBrush->GetOuter())
-			{
-				Model->Rename(NULL, DefaultBrush->GetOuter(), REN_DoNotDirty | REN_DontCreateRedirectors | REN_ForceNoResetLoaders | REN_NonTransactional);
-			}
-
 			const FLightmassPrimitiveSettings DefaultSettings;
 
 			for (int32 i = 0; i < Model->LightmassSettings.Num(); ++i)
@@ -4816,6 +4909,26 @@ void UWorld::CopyGameState(AGameMode* FromGameMode, AGameState* FromGameState)
 	GameState = FromGameState;
 }
 
+void UWorld::GetLightMapsAndShadowMaps(TArray<UTexture2D*>& OutLightMapsAndShadowMaps)
+{
+	TArray<UObject*> WorldPackageObjects;
+	const bool bIncludeNestedObjects = false;
+	GetObjectsWithOuter(GetOutermost(), WorldPackageObjects, bIncludeNestedObjects);
+
+	for (auto ObjIt = WorldPackageObjects.CreateConstIterator(); ObjIt; ++ObjIt)
+	{
+		UObject* CurrentObject = *ObjIt;
+		if ( ULightMapTexture2D* LightMap = Cast<ULightMapTexture2D>(CurrentObject) )
+		{
+			OutLightMapsAndShadowMaps.Add(LightMap);
+		}
+		else if ( UShadowMapTexture2D* ShadowMap = Cast<UShadowMapTexture2D>(CurrentObject) )
+		{
+			OutLightMapsAndShadowMaps.Add(ShadowMap);
+		}
+	}
+}
+
 /**
 * Dump visible actors in current world.
 */
