UE4: Added #if WITH_EDITOR around some recently added editor-only code

diff --git a/Engine/Source/Runtime/Engine/Private/Level.cpp b/Engine/Source/Runtime/Engine/Private/Level.cpp
--- a/Engine/Source/Runtime/Engine/Private/Level.cpp
+++ b/Engine/Source/Runtime/Engine/Private/Level.cpp
@@ -607,6 +607,7 @@ void ULevel::PostLoad()
 		StaticNavigableGeometry.Empty();
 	}
 
+#if WITH_EDITOR
 	// Rename the LevelScriptBlueprint after the outer world.
 	UWorld* OuterWorld = Cast<UWorld>(GetOuter());
 	if (LevelScriptBlueprint && OuterWorld && LevelScriptBlueprint->GetFName() != OuterWorld->GetFName())
@@ -614,6 +615,7 @@ void ULevel::PostLoad()
 		// Use LevelScriptBlueprint->GetOuter() instead of NULL to make sure the generated top level objects are moved appropriately
 		LevelScriptBlueprint->Rename(*OuterWorld->GetName(), LevelScriptBlueprint->GetOuter(), REN_DoNotDirty | REN_DontCreateRedirectors | REN_ForceNoResetLoaders | REN_NonTransactional);
 	}
+#endif
 }
 
 UWorld* ULevel::GetWorld() const
