UE4: Added a missing #include to fix CIS

diff --git a/Engine/Source/Runtime/Engine/Private/AI/EnvironmentQuery/Contexts/EnvQueryContext_BlueprintBase.cpp b/Engine/Source/Runtime/Engine/Private/AI/EnvironmentQuery/Contexts/EnvQueryContext_BlueprintBase.cpp
--- a/Engine/Source/Runtime/Engine/Private/AI/EnvironmentQuery/Contexts/EnvQueryContext_BlueprintBase.cpp
+++ b/Engine/Source/Runtime/Engine/Private/AI/EnvironmentQuery/Contexts/EnvQueryContext_BlueprintBase.cpp
@@ -1,6 +1,7 @@
 // Copyright 1998-2014 Epic Games, Inc. All Rights Reserved.
 
 #include "EnginePrivate.h"
+#include "BlueprintNodeHelpers.h"
 
 UEnvQueryContext_BlueprintBase::UEnvQueryContext_BlueprintBase(const class FPostConstructInitializeProperties& PCIP) 
 	: Super(PCIP), CallMode(UEnvQueryContext_BlueprintBase::InvalidCallMode)
