Fixed an Enum was being passed as the index due to the method signature having changed.
#ttp 335156 EDITOR: CASCADE: CRASH: Selecting "Remove duplicate modules" on an emitter can cause a crash
reviewer Andrew.Brown

diff --git a/Engine/Source/Runtime/Engine/Private/Particles/ParticleComponents.cpp b/Engine/Source/Runtime/Engine/Private/Particles/ParticleComponents.cpp
--- a/Engine/Source/Runtime/Engine/Private/Particles/ParticleComponents.cpp
+++ b/Engine/Source/Runtime/Engine/Private/Particles/ParticleComponents.cpp
@@ -2663,10 +2663,14 @@ bool UParticleSystem::RemoveAllDuplicateModules(bool bInMarkForCooker, TMap<UObj
 								// Ignore 'Cascade' category, transient, native and EditorOnly properties...
 								if (!bIsTransient && !bIsEditorOnly && !bIsCascade)
 								{
-									bool bIsIdentical = Property->Identical_InContainer(SourceModule, CheckModule, PPF_DeepComparison);
-									if (bIsIdentical == false)
+									for( int32 iProp=0; iProp<Property->ArrayDim; iProp++ )
 									{
-										bIsDifferent = true;
+										bool bIsIdentical = Property->Identical_InContainer(SourceModule, CheckModule, iProp, PPF_DeepComparison);
+										if (bIsIdentical == false)
+										{
+											bIsDifferent = true;
+											break;
+										}
 									}
 								}
 							}
