Added ENGINE_API to FFindFloorResult & extra bConstrainToPlane condition to UCharacterMovementComponent::DoJump()

#github https://github.com/EpicGames/UnrealEngine/pull/145

TTP# 335090 - [GitHub] 145 : Fix for CharacterMovementComponent to avoid linker errors

diff --git a/Engine/Source/Runtime/Engine/Classes/GameFramework/CharacterMovementComponent.h b/Engine/Source/Runtime/Engine/Classes/GameFramework/CharacterMovementComponent.h
--- a/Engine/Source/Runtime/Engine/Classes/GameFramework/CharacterMovementComponent.h
+++ b/Engine/Source/Runtime/Engine/Classes/GameFramework/CharacterMovementComponent.h
@@ -25,7 +25,7 @@ enum EMovementMode
 
 // Data about the floor for walking movement.
 USTRUCT(BlueprintType)
-struct FFindFloorResult
+struct ENGINE_API FFindFloorResult
 {
 	GENERATED_USTRUCT_BODY()
 
diff --git a/Engine/Source/Runtime/Engine/Private/CharacterMovementComponent.cpp b/Engine/Source/Runtime/Engine/Private/CharacterMovementComponent.cpp
--- a/Engine/Source/Runtime/Engine/Private/CharacterMovementComponent.cpp
+++ b/Engine/Source/Runtime/Engine/Private/CharacterMovementComponent.cpp
@@ -349,7 +349,7 @@ bool UCharacterMovementComponent::DoJump()
 	if ( CharacterOwner )
 	{
 		// Don't jump if we can't move up/down.
-		if (FMath::Abs(PlaneConstraintNormal.Z) != 1.f)
+		if (!bConstrainToPlane || FMath::Abs(PlaneConstraintNormal.Z) != 1.f)
 		{
 			Velocity.Z = JumpZVelocity;
 			SetMovementMode(MOVE_Falling);
