Landscape : Added new LOD falloff model to reduce triangles easily with good transitions

diff --git a/Engine/Source/Runtime/Engine/Classes/Landscape/LandscapeProxy.h b/Engine/Source/Runtime/Engine/Classes/Landscape/LandscapeProxy.h
--- a/Engine/Source/Runtime/Engine/Classes/Landscape/LandscapeProxy.h
+++ b/Engine/Source/Runtime/Engine/Classes/Landscape/LandscapeProxy.h
@@ -169,6 +169,18 @@ namespace ELandscapeLayerPaintingRestriction
 	};
 }
 
+UENUM()
+namespace ELandscapeLODFalloff
+{
+	enum Type
+	{
+		// Default mode
+		Linear			UMETA(DisplayName = "Linear"),
+		// Square Root give more natural transition, and also keep the same LOD 
+		SquareRoot		UMETA(DisplayName = "Square Root"),
+	};
+}
+
 UCLASS(dependson=UEngineTypes, NotPlaceable, hidecategories=(Display, Attachment, Physics, Debug, Lighting, LOD), showcategories=(Rendering, "Utilities|Transformation"), MinimalAPI)
 class ALandscapeProxy : public AActor, public INavRelevantActorInterface
 {
@@ -301,6 +313,9 @@ public:
 	UPROPERTY(EditAnywhere, Category=Landscape)
 	uint32 bUsedForNavigation:1;
 
+	UPROPERTY(EditAnywhere, Category=LOD)
+	TEnumAsByte<enum ELandscapeLODFalloff::Type> LODFalloff;
+
 #if WITH_EDITORONLY_DATA
 	UPROPERTY(EditAnywhere, Category=Landscape)
 	int32 MaxPaintedLayersPerComponent; // 0 = disabled
diff --git a/Engine/Source/Runtime/Engine/Private/Landscape/Landscape.cpp b/Engine/Source/Runtime/Engine/Private/Landscape/Landscape.cpp
--- a/Engine/Source/Runtime/Engine/Private/Landscape/Landscape.cpp
+++ b/Engine/Source/Runtime/Engine/Private/Landscape/Landscape.cpp
@@ -460,6 +460,7 @@ ALandscapeProxy::ALandscapeProxy(const class FPostConstructInitializeProperties&
 	bIsMovingToLevel = false;
 #endif // WITH_EDITORONLY_DATA
 	LODDistanceFactor = 1.0f;
+	LODFalloff = ELandscapeLODFalloff::Linear;
 	bCastStaticShadow = true;
 	bUsedForNavigation = true;
 	CollisionThickness = 16;
diff --git a/Engine/Source/Runtime/Engine/Private/Landscape/LandscapeEdit.cpp b/Engine/Source/Runtime/Engine/Private/Landscape/LandscapeEdit.cpp
--- a/Engine/Source/Runtime/Engine/Private/Landscape/LandscapeEdit.cpp
+++ b/Engine/Source/Runtime/Engine/Private/Landscape/LandscapeEdit.cpp
@@ -3668,7 +3668,7 @@ void ALandscape::PostEditChangeProperty(FPropertyChangedEvent& PropertyChangedEv
 	}
 	else if ( PropertyName == FName(TEXT("LODDistanceFactor")) )
 	{
-		LODDistanceFactor = FMath::Clamp<float>(LODDistanceFactor, 0.1f, 10.f); // limit because LOD transition became too popping...
+		LODDistanceFactor = FMath::Clamp<float>(LODDistanceFactor, 0.1f, MAX_LANDSCAPE_LOD_DISTANCE_FACTOR); // limit because LOD transition became too popping...
 		bPropagateToProxies = true;
 	}
 	else if ( PropertyName == FName(TEXT("CollisionMipLevel")) )
@@ -3676,6 +3676,10 @@ void ALandscape::PostEditChangeProperty(FPropertyChangedEvent& PropertyChangedEv
 		CollisionMipLevel = FMath::Clamp<int32>( CollisionMipLevel, 0, FMath::CeilLogTwo(SubsectionSizeQuads+1)-1 );
 		bPropagateToProxies = true;
 	}
+	else if (PropertyName == FName(TEXT("LODFalloff")))
+	{
+		bPropagateToProxies = true;
+	}
 	else if ( GIsEditor && PropertyName == FName(TEXT("StaticLightingResolution")) )
 	{
 		StaticLightingResolution = ::AdjustStaticLightingResolution(StaticLightingResolution, NumSubsections, SubsectionSizeQuads, ComponentSizeQuads);
diff --git a/Engine/Source/Runtime/Engine/Private/Landscape/LandscapeRender.cpp b/Engine/Source/Runtime/Engine/Private/Landscape/LandscapeRender.cpp
--- a/Engine/Source/Runtime/Engine/Private/Landscape/LandscapeRender.cpp
+++ b/Engine/Source/Runtime/Engine/Private/Landscape/LandscapeRender.cpp
@@ -15,6 +15,7 @@ IMPLEMENT_UNIFORM_BUFFER_STRUCT(FLandscapeUniformShaderParameters,TEXT("Landscap
 
 #define LANDSCAPE_LOD_DISTANCE_FACTOR 2.f
 #define LANDSCAPE_MAX_COMPONENT_SIZE 255
+#define LANDSCAPE_LOD_SQUARE_ROOT_FACTOR 1.5f
 
 /*------------------------------------------------------------------------------
 	Forsyth algorithm for cache optimizing index buffers.
@@ -544,6 +545,7 @@ FLandscapeComponentSceneProxy::FLandscapeComponentSceneProxy(ULandscapeComponent
 ,	LevelColor(1.f, 1.f, 1.f)
 ,	ForcedLOD(InComponent->ForcedLOD)
 ,	LODBias(InComponent->LODBias)
+,	LODFalloff(InComponent->GetLandscapeProxy()->LODFalloff)
 {
 	if (GRHIFeatureLevel == ERHIFeatureLevel::ES2)
 	{
@@ -585,9 +587,23 @@ FLandscapeComponentSceneProxy::FLandscapeComponentSceneProxy(ULandscapeComponent
 		MaxLOD = FMath::Min<int32>(MaxLOD, InComponent->GetLandscapeProxy()->MaxLODLevel);
 	}
 
-	LODDistance = FMath::Sqrt(2.f * FMath::Square((float)SubsectionSizeQuads)) * LANDSCAPE_LOD_DISTANCE_FACTOR / InComponent->GetLandscapeProxy()->LODDistanceFactor; // vary in 0...1
+	float LODDistanceFactor;
+	switch (LODFalloff)
+	{
+		case ELandscapeLODFalloff::SquareRoot:
+			LODDistanceFactor = FMath::Square(FMath::Min(LANDSCAPE_LOD_SQUARE_ROOT_FACTOR * InComponent->GetLandscapeProxy()->LODDistanceFactor, MAX_LANDSCAPE_LOD_DISTANCE_FACTOR));
+			break;
+		case ELandscapeLODFalloff::Linear:
+		default:
+			LODDistanceFactor = InComponent->GetLandscapeProxy()->LODDistanceFactor;
+			break;
+	}
+
+	LODDistance = FMath::Sqrt(2.f * FMath::Square((float)SubsectionSizeQuads)) * LANDSCAPE_LOD_DISTANCE_FACTOR / LODDistanceFactor; // vary in 0...1
 	DistDiff = -FMath::Sqrt(2.f * FMath::Square(0.5f*(float)SubsectionSizeQuads));
-	LODDistanceFactor = FMath::Min(InComponent->GetLandscapeProxy()->LODDistanceFactor, 3.f) * 0.33f;
+	// This value is used for shader. I made some modification to LOD transition looks better than plain transition
+	// Used for LODBias.xy values in shader
+	PrecomputedLODFactor = FMath::Min(LODDistanceFactor, 3.f) * 0.33f; 
 
 	if (InComponent->StaticLightingResolution > 0.f)
 	{
@@ -1205,7 +1221,20 @@ int32 FLandscapeComponentSceneProxy::CalcLODForSubsectionNoForced(const class FS
 		FVector2D CurrentCameraLocalPos = CameraLocalPos - FVector2D(SubX * SubsectionSizeQuads,SubY * SubsectionSizeQuads);
 		float ComponentDistance = FVector2D(CurrentCameraLocalPos-ComponentPosition).Size() + DistDiff;
 		// Clamp calculated distance based LOD with LODBiased values
-		float fLOD = FMath::Clamp<float>( ComponentDistance / LODDistance, FMath::Max<int32>(LODBias, MinLOD), FMath::Min<int32>(MaxLOD, MaxLOD+LODBias) );
+		float fLOD;
+		switch (LODFalloff)
+		{
+			case ELandscapeLODFalloff::SquareRoot:
+				fLOD = FMath::Sqrt(ComponentDistance / LODDistance);
+				break;
+			default:
+			case ELandscapeLODFalloff::Linear:
+				fLOD = ComponentDistance / LODDistance;
+				break;
+		}
+		// Clamping based on LODBias
+		fLOD = FMath::Clamp<float>(fLOD, FMath::Max<int32>(LODBias, MinLOD), FMath::Min<int32>(MaxLOD, MaxLOD + LODBias));
+
 		return FMath::FloorToInt( fLOD );
 	}
 	else
@@ -1213,7 +1242,21 @@ int32 FLandscapeComponentSceneProxy::CalcLODForSubsectionNoForced(const class FS
 		float Scale = 1.0f / (View.ViewRect.Width() * View.ViewMatrices.ProjMatrix.M[0][0]);
 
 		// The "/ 5.0f" is totally arbitrary
-		float fLOD = FMath::Clamp<float>(FMath::Sqrt(Scale / 5.0f), FMath::Max<int32>(LODBias, MinLOD), FMath::Min<int32>(MaxLOD, MaxLOD + LODBias));
+		float fLOD;
+		switch (LODFalloff)
+		{
+			case ELandscapeLODFalloff::SquareRoot:
+				fLOD = FMath::Sqrt(FMath::Sqrt(Scale / 5.0f));
+				break;
+			default:
+			case ELandscapeLODFalloff::Linear:
+				fLOD = Scale / 5.0f;
+				break;
+		}
+
+		// Clamping based on LODBias
+		fLOD = FMath::Clamp<float>(fLOD, FMath::Max<int32>(LODBias, MinLOD), FMath::Min<int32>(MaxLOD, MaxLOD + LODBias));
+
 		return FMath::FloorToInt(fLOD);
 	}
 }
@@ -1236,6 +1279,7 @@ void FLandscapeComponentSceneProxy::CalcLODParamsForSubsection(const class FScen
 	FVector2D CurrentCameraLocalPos = CameraLocalPos - FVector2D(SubX * SubsectionSizeQuads,SubY * SubsectionSizeQuads);
 
 	float ComponentDistance = FVector2D(CurrentCameraLocalPos - ComponentPosition).Size() + DistDiff;
+
 	OutDistLOD = ComponentDistance / LODDistance;
 
 	int32 FirstLOD = FMath::Max<int32>(LODBias, 0);
@@ -1246,12 +1290,34 @@ void FLandscapeComponentSceneProxy::CalcLODParamsForSubsection(const class FScen
 	}
 	else
 	{
-		OutfLOD = FMath::Clamp<float>(OutDistLOD, FirstLOD, LastLOD);
+		switch (LODFalloff)
+		{
+		case ELandscapeLODFalloff::SquareRoot:
+			OutfLOD = FMath::Sqrt(OutDistLOD);
+			break;
+		default:
+		case ELandscapeLODFalloff::Linear:
+			OutfLOD = OutDistLOD;
+			break;
+		}
+
+		OutfLOD = FMath::Clamp<float>(OutfLOD, FirstLOD, LastLOD);
 	}
 
 	for (int32 Idx = 0; Idx < LANDSCAPE_NEIGHBOR_NUM; ++Idx)
 	{
 		float ComponentDistance = FVector2D(CurrentCameraLocalPos-NeighborPosition[Idx]).Size() + DistDiff;
+		float fLOD;
+		switch (LODFalloff)
+		{
+			case ELandscapeLODFalloff::SquareRoot:
+				fLOD = FMath::Sqrt(ComponentDistance / LODDistance);
+				break;
+			default:
+			case ELandscapeLODFalloff::Linear:
+				fLOD = ComponentDistance / LODDistance;
+				break;
+		}
 
 		if (NumSubsections > 1 
 			&& ((SubX == 0 && Idx == 2) 
@@ -1260,12 +1326,12 @@ void FLandscapeComponentSceneProxy::CalcLODParamsForSubsection(const class FScen
 			|| (SubY == NumSubsections-1 && Idx == 0)) )
 		{
 			// Clamp calculated distance based LOD with LODBiased values
-			OutNeighborLODs[Idx] = ForcedLOD >= 0 ? ForcedLOD : FMath::Clamp<float>( ComponentDistance / LODDistance, FirstLOD, LastLOD );
+			OutNeighborLODs[Idx] = ForcedLOD >= 0 ? ForcedLOD : FMath::Clamp<float>( fLOD, FirstLOD, LastLOD);
 		}
 		else
 		{
 			// Neighbor LODBias are saved in BYTE, so need to convert to range [-128:127]
-			OutNeighborLODs[Idx] = ForcedNeighborLOD[Idx] != 255 ? ForcedNeighborLOD[Idx] : FMath::Clamp<float>( ComponentDistance / LODDistance, FMath::Max<float>(NeighborLODBias[Idx]-128, 0.f), FMath::Min<float>(MaxLOD, MaxLOD+NeighborLODBias[Idx]-128) );
+			OutNeighborLODs[Idx] = ForcedNeighborLOD[Idx] != 255 ? ForcedNeighborLOD[Idx] : FMath::Clamp<float>(fLOD, FMath::Max<float>(NeighborLODBias[Idx] - 128, 0.f), FMath::Min<float>(MaxLOD, MaxLOD + NeighborLODBias[Idx] - 128));
 		}
 
 		OutNeighborLODs[Idx] = FMath::Max<float>(OutfLOD, OutNeighborLODs[Idx]);
@@ -1942,8 +2008,8 @@ public:
 		if( LodBiasParameter.IsBound() )
 		{
 			FVector4 LodBias(
-				SceneProxy->LODDistanceFactor,
-				1.f / ( 1.f - SceneProxy->LODDistanceFactor ),
+				SceneProxy->PrecomputedLODFactor,
+				1.f / (1.f - SceneProxy->PrecomputedLODFactor),
 				SceneProxy->HeightmapTexture->GetNumMips() - FMath::Min(SceneProxy->HeightmapTexture->ResidentMips, SceneProxy->HeightmapTexture->RequestedMips),
 				0.f // Reserved
 				);
@@ -2046,8 +2112,8 @@ public:
 		if( LodBiasParameter.IsBound() )
 		{
 			FVector4 LodBias(
-				SceneProxy->LODDistanceFactor,
-				1.f / ( 1.f - SceneProxy->LODDistanceFactor ),
+				SceneProxy->PrecomputedLODFactor,
+				1.f / (1.f - SceneProxy->PrecomputedLODFactor),
 				SceneProxy->HeightmapTexture->GetNumMips() - FMath::Min(SceneProxy->HeightmapTexture->ResidentMips, SceneProxy->HeightmapTexture->RequestedMips),
 				SceneProxy->XYOffsetmapTexture->GetNumMips() - FMath::Min(SceneProxy->XYOffsetmapTexture->ResidentMips, SceneProxy->XYOffsetmapTexture->RequestedMips)
 				);
@@ -2362,8 +2428,19 @@ void ULandscapeComponent::GetStreamingTextureInfo(TArray<FStreamingTexturePrimit
 
 void ALandscapeProxy::ChangeLODDistanceFactor(float InLODDistanceFactor)
 {
-	LODDistanceFactor = FMath::Clamp<float>(InLODDistanceFactor, 0.1f, 10.f);
-	
+	LODDistanceFactor = FMath::Clamp<float>(InLODDistanceFactor, 0.1f, MAX_LANDSCAPE_LOD_DISTANCE_FACTOR);
+	float LODFactor;
+	switch (LODFalloff)
+	{
+		case ELandscapeLODFalloff::SquareRoot:
+			LODFactor = FMath::Square(FMath::Min(LANDSCAPE_LOD_SQUARE_ROOT_FACTOR * LODDistanceFactor, MAX_LANDSCAPE_LOD_DISTANCE_FACTOR));
+			break;
+		default:
+		case ELandscapeLODFalloff::Linear:
+			LODFactor = LODDistanceFactor;
+			break;
+	}
+
 	if (LandscapeComponents.Num())
 	{
 		int32 CompNum = LandscapeComponents.Num();
@@ -2377,7 +2454,7 @@ void ALandscapeProxy::ChangeLODDistanceFactor(float InLODDistanceFactor)
 			LandscapeChangeLODDistanceFactorCommand,
 			FLandscapeComponentSceneProxy**, Proxies, Proxies,
 			int32, CompNum, CompNum,
-			FVector2D, InLODDistanceFactors, FVector2D(FMath::Sqrt(2.f * FMath::Square((float)SubsectionSizeQuads)) * LANDSCAPE_LOD_DISTANCE_FACTOR / LODDistanceFactor, FMath::Max(LODDistanceFactor, 3.f) * 0.33f),
+			FVector2D, InLODDistanceFactors, FVector2D(FMath::Sqrt(2.f * FMath::Square((float)SubsectionSizeQuads)) * LANDSCAPE_LOD_DISTANCE_FACTOR / LODFactor, FMath::Min(LODFactor, 3.f) * 0.33f),
 		{
 			for (int32 Idx = 0; Idx < CompNum; ++Idx)
 			{
@@ -2392,5 +2469,5 @@ void ALandscapeProxy::ChangeLODDistanceFactor(float InLODDistanceFactor)
 void FLandscapeComponentSceneProxy::ChangeLODDistanceFactor_RenderThread(FVector2D InLODDistanceFactors)
 {
 	LODDistance = InLODDistanceFactors.X;
-	LODDistanceFactor = InLODDistanceFactors.Y;
+	PrecomputedLODFactor = InLODDistanceFactors.Y;
 }
diff --git a/Engine/Source/Runtime/Engine/Private/Landscape/LandscapeRenderMobile.cpp b/Engine/Source/Runtime/Engine/Private/Landscape/LandscapeRenderMobile.cpp
--- a/Engine/Source/Runtime/Engine/Private/Landscape/LandscapeRenderMobile.cpp
+++ b/Engine/Source/Runtime/Engine/Private/Landscape/LandscapeRenderMobile.cpp
@@ -76,8 +76,8 @@ public:
 		if( LodBiasParameter.IsBound() )
 		{
 			FVector4 LodBias(
-				SceneProxy->LODDistanceFactor,
-				1.f / ( 1.f - SceneProxy->LODDistanceFactor ),
+				SceneProxy->PrecomputedLODFactor,
+				1.f / (1.f - SceneProxy->PrecomputedLODFactor),
 				CameraLocalPos3D.X + SceneProxy->SectionBase.X,
 				CameraLocalPos3D.Y + SceneProxy->SectionBase.Y 
 				);
diff --git a/Engine/Source/Runtime/Engine/Public/Landscape/LandscapeEdit.h b/Engine/Source/Runtime/Engine/Public/Landscape/LandscapeEdit.h
--- a/Engine/Source/Runtime/Engine/Public/Landscape/LandscapeEdit.h
+++ b/Engine/Source/Runtime/Engine/Public/Landscape/LandscapeEdit.h
@@ -7,6 +7,8 @@ LandscapeEdit.h: Classes for the editor to access to Landscape data
 #ifndef _LANDSCAPEEDIT_H
 #define _LANDSCAPEEDIT_H
 
+#define MAX_LANDSCAPE_LOD_DISTANCE_FACTOR 10.f
+
 #if WITH_EDITOR
 
 struct FLandscapeTextureDataInfo
diff --git a/Engine/Source/Runtime/Engine/Public/Landscape/LandscapeRender.h b/Engine/Source/Runtime/Engine/Public/Landscape/LandscapeRender.h
--- a/Engine/Source/Runtime/Engine/Public/Landscape/LandscapeRender.h
+++ b/Engine/Source/Runtime/Engine/Public/Landscape/LandscapeRender.h
@@ -405,9 +405,9 @@ protected:
 	// Storage for static draw list batch params
 	TArray<FLandscapeBatchElementParams> StaticBatchParamArray;
 
-	// Precomputed
+	// Precomputed values
 	float					LODDistance;
-	float					LODDistanceFactor;
+	float					PrecomputedLODFactor;
 	float					DistDiff;
 
 	FVector4 WeightmapScaleBias;
@@ -452,6 +452,8 @@ protected:
 	uint8					ForcedNeighborLOD[LANDSCAPE_NEIGHBOR_NUM];
 	uint8					NeighborLODBias[LANDSCAPE_NEIGHBOR_NUM];
 
+	enum ELandscapeLODFalloff::Type LODFalloff;
+
 	TUniformBuffer<FLandscapeUniformShaderParameters> LandscapeUniformShaderParameters;
 
 	// Cached versions of these
