Fixed server-client streaming levels replication in World Composition
Dedicated server does not use distance streaming and loads all distance dependent levels.
Client does not accept streaming levels state updates from a server for a distance dependent levels

diff --git a/Engine/Source/Runtime/Engine/Classes/Engine/WorldComposition.h b/Engine/Source/Runtime/Engine/Classes/Engine/WorldComposition.h
--- a/Engine/Source/Runtime/Engine/Classes/Engine/WorldComposition.h
+++ b/Engine/Source/Runtime/Engine/Classes/Engine/WorldComposition.h
@@ -77,6 +77,9 @@ class ENGINE_API UWorldComposition : public UObject
 	/** Returns currently visible and hidden levels based on distance based streaming */
 	void GetDistanceVisibleLevels(const FVector& InLocation, TArray<FDistanceVisibleLevel>& OutVisibleLevels, TArray<FDistanceVisibleLevel>& OutHiddenLevels) const;
 
+	/** @returns Whether specified streaming level is distance dependent */
+	bool IsDistanceDependentLevel(FName PackageName) const;
+
 	/** @returns Currently opened world composition root folder (long PackageName)*/
 	FString GetWorldRoot() const;
 	
diff --git a/Engine/Source/Runtime/Engine/Private/PlayerController.cpp b/Engine/Source/Runtime/Engine/Private/PlayerController.cpp
--- a/Engine/Source/Runtime/Engine/Private/PlayerController.cpp
+++ b/Engine/Source/Runtime/Engine/Private/PlayerController.cpp
@@ -119,6 +119,14 @@ void APlayerController::ClientUpdateLevelStreamingStatus_Implementation(FName Pa
 		PackageName = FName(*PackageNameStr);
 	}
 
+	// Distance dependent streaming levels should be controlled by client only
+	if (GetWorld() && GetWorld()->WorldComposition)
+	{
+		if (GetWorld()->WorldComposition->IsDistanceDependentLevel(PackageName))
+		{
+			return;
+		}
+	}
 
 	//GEngine->NetworkRemapPath(GetWorld(), RealName, false);
 
diff --git a/Engine/Source/Runtime/Engine/Private/World.cpp b/Engine/Source/Runtime/Engine/Private/World.cpp
--- a/Engine/Source/Runtime/Engine/Private/World.cpp
+++ b/Engine/Source/Runtime/Engine/Private/World.cpp
@@ -2157,11 +2157,14 @@ void UWorld::UpdateLevelStreaming( FSceneViewFamily* ViewFamily )
 	{
 		return;
 	}
-	
-	if (WorldComposition && ViewFamily)
+
+	if (WorldComposition)
 	{
 		// May issue the world origin shift request depending on view location
-		EvaluateWorldOriginLocation(*ViewFamily);
+		if (ViewFamily)
+		{
+			EvaluateWorldOriginLocation(*ViewFamily);
+		}
 		// May add/remove streaming objects to persistent world depending on view location
 		WorldComposition->UpdateStreamingState(ViewFamily);
 	}
diff --git a/Engine/Source/Runtime/Engine/Private/WorldComposition.cpp b/Engine/Source/Runtime/Engine/Private/WorldComposition.cpp
--- a/Engine/Source/Runtime/Engine/Private/WorldComposition.cpp
+++ b/Engine/Source/Runtime/Engine/Private/WorldComposition.cpp
@@ -442,27 +442,39 @@ void UWorldComposition::GetDistanceVisibleLevels(
 			TilesStreaming[TileIdx],
 			INDEX_NONE
 		};
-				
-		FIntPoint LevelOffset = Tile.Info.AbsolutePosition - WorldOffset;
-		FBox LevelBounds = Tile.Info.Bounds.ShiftBy(FVector(LevelOffset, 0));
-		// We don't care about third dimension yet
-		LevelBounds.Min.Z = -WORLD_MAX;
-		LevelBounds.Max.Z = +WORLD_MAX;
-		
+
 		bool bIsVisible = false;
-		int32 NumAvailableLOD = FMath::Min(Tile.Info.LODList.Num(), Tile.LODPackageNames.Num());
-		// Find highest visible LOD entry
-		// INDEX_NONE for original non-LOD level
-		for (int32 LODIdx = INDEX_NONE; LODIdx < NumAvailableLOD; ++LODIdx)
+
+		if (OwningWorld->GetNetMode() == NM_DedicatedServer) 
 		{
-			int32 TileStreamingDistance = Tile.Info.GetStreamingDistance(LODIdx);
-			FSphere QuerySphere(InLocation, TileStreamingDistance);
-		
-			if (FMath::SphereAABBIntersection(QuerySphere, LevelBounds))
+			// Dedicated server always loads all distance dependent tiles
+			bIsVisible = true;
+		}
+		else
+		{
+			//
+			// Check if tile bounding box intersects with a sphere with origin at provided location and with radius equal to tile layer distance settings
+			//
+			FIntPoint LevelOffset = Tile.Info.AbsolutePosition - WorldOffset;
+			FBox LevelBounds = Tile.Info.Bounds.ShiftBy(FVector(LevelOffset, 0));
+			// We don't care about third dimension yet
+			LevelBounds.Min.Z = -WORLD_MAX;
+			LevelBounds.Max.Z = +WORLD_MAX;
+				
+			int32 NumAvailableLOD = FMath::Min(Tile.Info.LODList.Num(), Tile.LODPackageNames.Num());
+			// Find highest visible LOD entry
+			// INDEX_NONE for original non-LOD level
+			for (int32 LODIdx = INDEX_NONE; LODIdx < NumAvailableLOD; ++LODIdx)
 			{
-				VisibleLevel.LODIndex = LODIdx;
-				bIsVisible = true;
-				break;
+				int32 TileStreamingDistance = Tile.Info.GetStreamingDistance(LODIdx);
+				FSphere QuerySphere(InLocation, TileStreamingDistance);
+		
+				if (FMath::SphereAABBIntersection(QuerySphere, LevelBounds))
+				{
+					VisibleLevel.LODIndex = LODIdx;
+					bIsVisible = true;
+					break;
+				}
 			}
 		}
 
@@ -522,6 +534,12 @@ void UWorldComposition::UpdateStreamingState(FSceneViewFamily* ViewFamily)
 	UpdateStreamingState(ViewLocation);
 }
 
+bool UWorldComposition::IsDistanceDependentLevel(FName PackageName) const
+{
+	FWorldCompositionTile* Tile = FindTileByName(PackageName);
+	return (Tile && Tile->Info.Layer.DistanceStreamingEnabled);
+}
+
 void UWorldComposition::CommitTileStreamingState(UWorld* PersistenWorld, int32 TileIdx, bool bShouldBeLoaded, bool bShouldBeVisible, int32 LODIdx)
 {
 	if (!Tiles.IsValidIndex(TileIdx))
