Non-unity compilation fix #UE4

diff --git a/Engine/Source/Runtime/Engine/Private/AI/EnvironmentQuery/Contexts/EnvQueryContext_BlueprintBase.cpp b/Engine/Source/Runtime/Engine/Private/AI/EnvironmentQuery/Contexts/EnvQueryContext_BlueprintBase.cpp
--- a/Engine/Source/Runtime/Engine/Private/AI/EnvironmentQuery/Contexts/EnvQueryContext_BlueprintBase.cpp
+++ b/Engine/Source/Runtime/Engine/Private/AI/EnvironmentQuery/Contexts/EnvQueryContext_BlueprintBase.cpp
@@ -3,14 +3,22 @@
 #include "EnginePrivate.h"
 #include "BlueprintNodeHelpers.h"
 
+namespace
+{
+	FORCEINLINE bool DoesImplementBPFunction(FName FuncName, const UObject* Ob, const UClass* StopAtClass)
+	{
+		return (Ob->GetClass()->FindFunctionByName(FuncName)->GetOuter() != StopAtClass);
+	}
+}
+
 UEnvQueryContext_BlueprintBase::UEnvQueryContext_BlueprintBase(const class FPostConstructInitializeProperties& PCIP) 
 	: Super(PCIP), CallMode(UEnvQueryContext_BlueprintBase::InvalidCallMode)
 {
-	UClass* StopAtClass = UEnvQueryContext_BlueprintBase::StaticClass();
-	bool bImplementsProvideSingleActor = BlueprintNodeHelpers::HasBlueprintFunction(GET_FUNCTION_NAME_CHECKED(UEnvQueryContext_BlueprintBase, ProvideSingleActor), this, StopAtClass);
-	bool bImplementsProvideSingleLocation = BlueprintNodeHelpers::HasBlueprintFunction(GET_FUNCTION_NAME_CHECKED(UEnvQueryContext_BlueprintBase, ProvideSingleLocation), this, StopAtClass);
-	bool bImplementsProvideActorSet = BlueprintNodeHelpers::HasBlueprintFunction(GET_FUNCTION_NAME_CHECKED(UEnvQueryContext_BlueprintBase, ProvideActorsSet), this, StopAtClass);
-	bool bImplementsProvideLocationsSet = BlueprintNodeHelpers::HasBlueprintFunction(GET_FUNCTION_NAME_CHECKED(UEnvQueryContext_BlueprintBase, ProvideLocationsSet), this, StopAtClass);
+	UClass* StopAtClass = UEnvQueryContext_BlueprintBase::StaticClass();	
+	bool bImplementsProvideSingleActor = DoesImplementBPFunction(GET_FUNCTION_NAME_CHECKED(UEnvQueryContext_BlueprintBase, ProvideSingleActor), this, StopAtClass);
+	bool bImplementsProvideSingleLocation = DoesImplementBPFunction(GET_FUNCTION_NAME_CHECKED(UEnvQueryContext_BlueprintBase, ProvideSingleLocation), this, StopAtClass);
+	bool bImplementsProvideActorSet = DoesImplementBPFunction(GET_FUNCTION_NAME_CHECKED(UEnvQueryContext_BlueprintBase, ProvideActorsSet), this, StopAtClass);
+	bool bImplementsProvideLocationsSet = DoesImplementBPFunction(GET_FUNCTION_NAME_CHECKED(UEnvQueryContext_BlueprintBase, ProvideLocationsSet), this, StopAtClass);
 
 	int32 ImplementationsCount = 0;
 	if (bImplementsProvideSingleActor)
