Fix sphere component to draw min scale to match pxvis collision

diff --git a/Engine/Source/Runtime/Engine/Private/SphereComponent.cpp b/Engine/Source/Runtime/Engine/Private/SphereComponent.cpp
--- a/Engine/Source/Runtime/Engine/Private/SphereComponent.cpp
+++ b/Engine/Source/Runtime/Engine/Private/SphereComponent.cpp
@@ -110,9 +110,19 @@ FPrimitiveSceneProxy* USphereComponent::CreateSceneProxy()
 			else
 			{
 				const FLinearColor DrawSphereColor = GetSelectionColor(SphereColor, IsSelected(), IsHovered(), /*bUseOverlayIntensity=*/false);
-				DrawCircle( PDI, LocalToWorld.GetOrigin(), LocalToWorld.GetScaledAxis( EAxis::X ), LocalToWorld.GetScaledAxis( EAxis::Y ), DrawSphereColor, SphereRadius, SphereSides, SDPG_World );
-				DrawCircle( PDI, LocalToWorld.GetOrigin(), LocalToWorld.GetScaledAxis( EAxis::X ), LocalToWorld.GetScaledAxis( EAxis::Z ), DrawSphereColor, SphereRadius, SphereSides, SDPG_World );
-				DrawCircle( PDI, LocalToWorld.GetOrigin(), LocalToWorld.GetScaledAxis( EAxis::Y ), LocalToWorld.GetScaledAxis( EAxis::Z ), DrawSphereColor, SphereRadius, SphereSides, SDPG_World );
+				
+				float AbsScaleX = LocalToWorld.GetScaledAxis(EAxis::X).Size();
+				float AbsScaleY = LocalToWorld.GetScaledAxis(EAxis::Y).Size();
+				float AbsScaleZ = LocalToWorld.GetScaledAxis(EAxis::Z).Size();
+				float MinAbsScale = FMath::Min3(AbsScaleX, AbsScaleY, AbsScaleZ);
+
+				FVector ScaledX = LocalToWorld.GetUnitAxis(EAxis::X) * MinAbsScale;
+				FVector ScaledY = LocalToWorld.GetUnitAxis(EAxis::Y) * MinAbsScale;
+				FVector ScaledZ = LocalToWorld.GetUnitAxis(EAxis::Z) * MinAbsScale;
+
+				DrawCircle(PDI, LocalToWorld.GetOrigin(), ScaledX, ScaledY, DrawSphereColor, SphereRadius, SphereSides, SDPG_World);
+				DrawCircle(PDI, LocalToWorld.GetOrigin(), ScaledX, ScaledZ, DrawSphereColor, SphereRadius, SphereSides, SDPG_World);
+				DrawCircle(PDI, LocalToWorld.GetOrigin(), ScaledY, ScaledZ, DrawSphereColor, SphereRadius, SphereSides, SDPG_World);
 			}
 		}
 
