Add ENGINE_API to FActorComponentTickFunction functions

diff --git a/Engine/Source/Runtime/Engine/Classes/Engine/EngineBaseTypes.h b/Engine/Source/Runtime/Engine/Classes/Engine/EngineBaseTypes.h
--- a/Engine/Source/Runtime/Engine/Classes/Engine/EngineBaseTypes.h
+++ b/Engine/Source/Runtime/Engine/Classes/Engine/EngineBaseTypes.h
@@ -351,9 +351,9 @@ struct FActorComponentTickFunction : public FTickFunction
 		* @param CurrentThread - thread we are executing on, useful to pass along as new tasks are created
 		* @param MyCompletionGraphEvent - completion event for this task. Useful for holding the completetion of this task until certain child tasks are complete.
 	**/
-	virtual void ExecuteTick(float DeltaTime, enum ELevelTick TickType, ENamedThreads::Type CurrentThread, const FGraphEventRef& MyCompletionGraphEvent) OVERRIDE;
+	ENGINE_API virtual void ExecuteTick(float DeltaTime, enum ELevelTick TickType, ENamedThreads::Type CurrentThread, const FGraphEventRef& MyCompletionGraphEvent) OVERRIDE;
 	/** Abstract function to describe this tick. Used to print messages about illegal cycles in the dependency graph **/
-	virtual FString DiagnosticMessage();
+	ENGINE_API virtual FString DiagnosticMessage();
 };
 
 /** 
