Fix for destructible shapes that share a PxActor not properly updating their transform
TTP#335161

diff --git a/Engine/Source/Runtime/Engine/Private/PhysicsEngine/PhysScene.cpp b/Engine/Source/Runtime/Engine/Private/PhysicsEngine/PhysScene.cpp
--- a/Engine/Source/Runtime/Engine/Private/PhysicsEngine/PhysScene.cpp
+++ b/Engine/Source/Runtime/Engine/Private/PhysicsEngine/PhysScene.cpp
@@ -4,8 +4,8 @@
 #include "EnginePrivate.h"
 
 #if WITH_PHYSX
-	#include "PhysXSupport.h"
-	#include "../Vehicles/PhysXVehicleManager.h"
+#include "PhysXSupport.h"
+#include "../Vehicles/PhysXVehicleManager.h"
 #endif
 
 #include "PhysSubstepTasks.h"	//needed even if not substepping, contains common utility class for PhysX
@@ -29,7 +29,11 @@ FORCEINLINE EPhysicsSceneType SceneType(const FBodyInstance * BodyInstance)
 }
 
 /**
- * Return true if we should lag the async scene a frame
+* Return true if we should be running in single threaded mode, ala dedicated server
+**/
+
+/**
+* Return true if we should lag the async scene a frame
 **/
 FORCEINLINE static bool FrameLagAsync()
 {
@@ -53,11 +57,11 @@ FPhysScene::FPhysScene()
 	// Create dispatcher for tasks
 	if (PhysSingleThreadedMode())
 	{
-		CPUDispatcher =  new FPhysXCPUDispatcherSingleThread();
+		CPUDispatcher = new FPhysXCPUDispatcherSingleThread();
 	}
 	else
 	{
-		CPUDispatcher =  new FPhysXCPUDispatcher();
+		CPUDispatcher = new FPhysXCPUDispatcher();
 	}
 	// Create sim event callback
 	SimEventCallback = new FPhysXSimEventCallback();
@@ -76,7 +80,7 @@ FPhysScene::FPhysScene()
 		FrameTimeSmoothingFactor[PST_Async] = PhysSetting->AsyncSceneSmoothingFactor;
 		bInitializeConsoleVariable = false;
 	}
-		
+
 #if WITH_SUBSTEPPING
 	bSubstepping = PhysSetting->bSubstepping;
 #endif
@@ -86,7 +90,7 @@ FPhysScene::FPhysScene()
 
 
 	// Create scenes of all scene types
-	for(uint32 SceneType = 0; SceneType < NumPhysScenes; ++SceneType)
+	for (uint32 SceneType = 0; SceneType < NumPhysScenes; ++SceneType)
 	{
 		// Create the physics scene
 		InitPhysScene(SceneType);
@@ -95,7 +99,7 @@ FPhysScene::FPhysScene()
 		bPhysXSceneExecuting[SceneType] = false;
 
 		// Initialize to a value which would be acceptable if FrameTimeSmoothingFactor[i] = 1.0f, i.e. constant simulation substeps
-		AveragedFrameTime[SceneType] = 	InitialAverageFrameRate;
+		AveragedFrameTime[SceneType] = InitialAverageFrameRate;
 
 		// gets from console variable, and clamp to [0, 1] - 1 should be fixed time as 30 fps
 		FrameTimeSmoothingFactor[SceneType] = FMath::Clamp<float>(FrameTimeSmoothingFactor[SceneType], 0.f, 1.f);
@@ -125,7 +129,7 @@ FPhysScene::~FPhysScene()
 	// Make sure no scenes are left simulating (no-ops if not simulating)
 	WaitPhysScenes();
 	// Loop through scene types to get all scenes
-	for(uint32 SceneType = 0; SceneType < NumPhysScenes; ++SceneType)
+	for (uint32 SceneType = 0; SceneType < NumPhysScenes; ++SceneType)
 	{
 
 		// Destroy the physics scene
@@ -339,7 +343,7 @@ void FPhysScene::FlushDeferredCollisionDisableTableQueue()
 	for (int32 i = 0; i < DeferredCollisionDisableTableQueue.Num(); ++i)
 	{
 		FPendingCollisionDisableTable & PendingCollisionDisableTable = DeferredCollisionDisableTableQueue[i];
-		
+
 		if (PendingCollisionDisableTable.CollisionDisableTable)
 		{
 			CollisionDisableTableLookup.Add(PendingCollisionDisableTable.SkelMeshCompID, PendingCollisionDisableTable.CollisionDisableTable);
@@ -361,7 +365,7 @@ void FPhysScene::TickPhysScene(uint32 SceneType, FGraphEventRef& InOutCompletion
 
 	check(SceneType < NumPhysScenes);
 
-	if(	bPhysXSceneExecuting[SceneType] != 0 )
+	if (bPhysXSceneExecuting[SceneType] != 0)
 	{
 		// Already executing this scene, must call WaitPhysScene before calling this function again.
 		UE_LOG(LogPhysics, Log, TEXT("TickPhysScene: Already executing scene (%d) - aborting."), SceneType);
@@ -369,24 +373,24 @@ void FPhysScene::TickPhysScene(uint32 SceneType, FGraphEventRef& InOutCompletion
 	}
 
 #if WITH_SUBSTEPPING
-	if(IsSubstepping() && SceneType != PST_Cloth)	//we don't bother sub-stepping cloth
+	if (IsSubstepping() && SceneType != PST_Cloth)	//we don't bother sub-stepping cloth
 	{
 		//We're about to start stepping so swap buffers. Might want to find a better place for this?
 		PhysSubSteppers[SceneType]->SwapBuffers();
 	}
 #endif
 
-	/** 
-	 * clamp down... if this happens we are simming physics slower than real-time, so be careful with it.
-	 * it can improve framerate dramatically (really, it is the same as scaling all velocities down and
-	 * enlarging all timesteps) but at the same time, it will screw with networking (client and server will
-	 * diverge a lot more.)
-	 */
-	
+	/**
+	* clamp down... if this happens we are simming physics slower than real-time, so be careful with it.
+	* it can improve framerate dramatically (really, it is the same as scaling all velocities down and
+	* enlarging all timesteps) but at the same time, it will screw with networking (client and server will
+	* diverge a lot more.)
+	*/
+
 	float UseDelta = FMath::Min(UseSyncTime(SceneType) ? SyncDeltaSeconds : DeltaSeconds, MaxPhysicsDeltaTime);
 
 	// Only simulate a positive time step.
-	if(UseDelta <= 0.f)
+	if (UseDelta <= 0.f)
 	{
 		if (UseDelta < 0.f)
 		{
@@ -397,10 +401,10 @@ void FPhysScene::TickPhysScene(uint32 SceneType, FGraphEventRef& InOutCompletion
 	}
 
 	/**
-	 * Weight frame time according to PhysScene settings.
-	 */
+	* Weight frame time according to PhysScene settings.
+	*/
 	AveragedFrameTime[SceneType] *= FrameTimeSmoothingFactor[SceneType];
-	AveragedFrameTime[SceneType] += (1.0f-FrameTimeSmoothingFactor[SceneType])*UseDelta;
+	AveragedFrameTime[SceneType] += (1.0f - FrameTimeSmoothingFactor[SceneType])*UseDelta;
 
 	// Set execution flag
 	bPhysXSceneExecuting[SceneType] = true;
@@ -431,7 +435,7 @@ void FPhysScene::TickPhysScene(uint32 SceneType, FGraphEventRef& InOutCompletion
 
 #if !WITH_APEX
 	PxScene* PScene = GetPhysXScene(SceneType);
-	if(PScene && (UseDelta > 0.f))
+	if (PScene && (UseDelta > 0.f))
 	{
 		PhysXCompletionTask* Task = new PhysXCompletionTask(InOutCompletionEvent, PScene->getTaskManager());
 		PScene->lockWrite();
@@ -480,7 +484,7 @@ void FPhysScene::WaitPhysScenes()
 	}
 	// Loop through scene types to get all scenes
 	// we just wait on everything, though some of these are redundant
-	for(uint32 SceneType = 0; SceneType < NumPhysScenes; ++SceneType)
+	for (uint32 SceneType = 0; SceneType < NumPhysScenes; ++SceneType)
 	{
 		if (PhysicsSubsceneCompletion[SceneType].GetReference())
 		{
@@ -502,16 +506,16 @@ void FPhysScene::WaitClothScene()
 	FGraphEventArray ThingsToComplete;
 	if (PhysicsSubsceneCompletion[PST_Cloth].GetReference())
 	{
-			ThingsToComplete.Add(PhysicsSubsceneCompletion[PST_Cloth]);
+		ThingsToComplete.Add(PhysicsSubsceneCompletion[PST_Cloth]);
 	}
-	
+
 	if (ThingsToComplete.Num())
 	{
 		FTaskGraphInterface::Get().WaitUntilTasksComplete(ThingsToComplete, ENamedThreads::GameThread);
 	}
 }
 
-void FPhysScene::SceneCompletionTask(ENamedThreads::Type CurrentThread, const FGraphEventRef& MyCompletionGraphEvent,EPhysicsSceneType SceneType)
+void FPhysScene::SceneCompletionTask(ENamedThreads::Type CurrentThread, const FGraphEventRef& MyCompletionGraphEvent, EPhysicsSceneType SceneType)
 {
 	ProcessPhysScene(SceneType);
 }
@@ -522,7 +526,7 @@ void FPhysScene::ProcessPhysScene(uint32 SceneType)
 	SCOPE_CYCLE_COUNTER(STAT_PhysicsFetchDynamicsTime);
 
 	check(SceneType < NumPhysScenes);
-	if(	bPhysXSceneExecuting[SceneType] == 0 )
+	if (bPhysXSceneExecuting[SceneType] == 0)
 	{
 		// Not executing this scene, must call TickPhysScene before calling this function again.
 		UE_LOG(LogPhysics, Log, TEXT("WaitPhysScene`: Not executing this scene (%d) - aborting."), SceneType);
@@ -550,18 +554,18 @@ void FPhysScene::ProcessPhysScene(uint32 SceneType)
 	PxU32 OutErrorCode = 0;
 #if !WITH_APEX
 	PScene->lockWrite();
-	PScene->fetchResults( true, &OutErrorCode );
+	PScene->fetchResults(true, &OutErrorCode);
 	PScene->unlockWrite();
 #else	//	#if !WITH_APEX
 	// The APEX scene calls the fetchResults function for the PhysX scene, so we only call ApexScene->fetchResults().
 	NxApexScene* ApexScene = GetApexScene(SceneType);
 	check(ApexScene);
-	ApexScene->fetchResults( true, &OutErrorCode );
+	ApexScene->fetchResults(true, &OutErrorCode);
 #endif	//	#if !WITH_APEX
 
 	UpdateActiveTransforms(SceneType);
 
-	if(OutErrorCode != 0)
+	if (OutErrorCode != 0)
 	{
 		UE_LOG(LogPhysics, Log, TEXT("PHYSX FETCHRESULTS ERROR: %d"), OutErrorCode);
 	}
@@ -587,7 +591,7 @@ void FPhysScene::UpdateActiveTransforms(uint32 SceneType)
 	PxU32 NumTransforms = 0;
 	const PxActiveTransform* PActiveTransforms = PScene->getActiveTransforms(NumTransforms);
 	SCENE_UNLOCK_READ(PScene);
-	
+
 	ActiveTransforms[SceneType].Empty(NumTransforms);
 
 	for (PxU32 TransformIdx = 0; TransformIdx < NumTransforms; ++TransformIdx)
@@ -603,7 +607,7 @@ void FPhysScene::SyncComponentsToBodies(uint32 SceneType)
 	PxU32 NumTransforms = ActiveTransforms[SceneType].Num();
 	const TArray<const PxActiveTransform*> & PActiveTransforms = ActiveTransforms[SceneType];
 
-	for(PxU32 TransformIdx=0; TransformIdx<NumTransforms; TransformIdx++)
+	for (PxU32 TransformIdx = 0; TransformIdx<NumTransforms; TransformIdx++)
 	{
 		if (PActiveTransforms[TransformIdx] == NULL)	//it's possible to call TermBody on FBodyInstance after fetchResults, but before SyncComponentsToBodies - in this case a NULL is used to represent the stale data
 		{
@@ -619,25 +623,36 @@ void FPhysScene::SyncComponentsToBodies(uint32 SceneType)
 		{
 			if (DestructibleChunkInfo->OwningComponent.IsValid())
 			{
-				const physx::PxTransform & PTransform = PActiveTransform.actor2World;
-				DestructibleChunkInfo->OwningComponent->SetChunkWorldRT(DestructibleChunkInfo->ChunkIndex, P2UQuat(PTransform.q), P2UVector(PTransform.p));
+				TArray<PxShape*> Shapes;
+				Shapes.AddZeroed(RigidActor->getNbShapes());
+				int32 NumShapes = RigidActor->getShapes(Shapes.GetData(), Shapes.Num());
+				for (int32 ShapeIdx = 0; ShapeIdx < Shapes.Num(); ++ShapeIdx)
+				{
+					PxShape* Shape = Shapes[ShapeIdx];
+					int32 ChunkIndex;
+					NxDestructibleActor* DestructibleActor = GApexModuleDestructible->getDestructibleAndChunk(Shape, &ChunkIndex);
+					ensure(DestructibleActor == DestructibleChunkInfo->OwningComponent->ApexDestructibleActor);
+					const physx::PxMat44 ChunkPoseRT = DestructibleActor->getChunkPose(ChunkIndex);
+					const physx::PxTransform Transform(ChunkPoseRT);
+					DestructibleChunkInfo->OwningComponent->SetChunkWorldRT(ChunkIndex, P2UQuat(Transform.q), P2UVector(Transform.p));
+				}
 			}
 		}
 #endif
 
 		FBodyInstance* BodyInst = FPhysxUserData::Get<FBodyInstance>(PActiveTransform.userData);
-		if(	BodyInst != NULL && 
-			BodyInst->InstanceBodyIndex == INDEX_NONE && 
+		if (BodyInst != NULL &&
+			BodyInst->InstanceBodyIndex == INDEX_NONE &&
 			BodyInst->OwnerComponent != NULL &&
-			BodyInst->IsInstanceSimulatingPhysics() )
+			BodyInst->IsInstanceSimulatingPhysics())
 		{
 			check(BodyInst->OwnerComponent->IsRegistered()); // shouldn't have a physics body for a non-registered component!
 
 			AActor* Owner = BodyInst->OwnerComponent->GetOwner();
 
 			// See if the transform is actually different, and if so, move the component to match physics
-			const FTransform NewTransform = BodyInst->GetUnrealWorldTransform();	
-			if(!NewTransform.EqualsNoScale(BodyInst->OwnerComponent->ComponentToWorld))
+			const FTransform NewTransform = BodyInst->GetUnrealWorldTransform();
+			if (!NewTransform.EqualsNoScale(BodyInst->OwnerComponent->ComponentToWorld))
 			{
 				const FVector MoveBy = NewTransform.GetLocation() - BodyInst->OwnerComponent->ComponentToWorld.GetLocation();
 				const FRotator NewRotation = NewTransform.Rotator();
@@ -649,7 +664,7 @@ void FPhysScene::SyncComponentsToBodies(uint32 SceneType)
 			}
 
 			// Check if we didn't fall out of the world
-			if(Owner != NULL && !Owner->IsPendingKill())
+			if (Owner != NULL && !Owner->IsPendingKill())
 			{
 				Owner->CheckStillInWorld();
 			}
@@ -691,7 +706,7 @@ void FPhysScene::DispatchPhysNotifications()
 		}
 		PendingCollisionNotifies.Empty();
 	}
-	
+
 #if WITH_SUBSTEPPING
 #if WITH_APEX
 	//TODO: Queue is always empty for now - waiting for support from NVIDIA
@@ -719,10 +734,10 @@ void FPhysScene::SetUpForFrame(const FVector* NewGrav, float InDeltaSeconds, flo
 	if (NewGrav)
 	{
 		// Loop through scene types to get all scenes
-		for(uint32 SceneType = 0; SceneType < NumPhysScenes; ++SceneType)
+		for (uint32 SceneType = 0; SceneType < NumPhysScenes; ++SceneType)
 		{
 			PxScene* PScene = GetPhysXScene(SceneType);
-			if(PScene != NULL)
+			if (PScene != NULL)
 			{
 				//@todo phys_thread don't do this if gravity changes
 
@@ -730,7 +745,7 @@ void FPhysScene::SetUpForFrame(const FVector* NewGrav, float InDeltaSeconds, flo
 				// Lock scene lock, in case it is required
 				SCENE_LOCK_WRITE(PScene);
 
-				PScene->setGravity( U2PVector(*NewGrav) );
+				PScene->setGravity(U2PVector(*NewGrav));
 
 				// Unlock scene lock, in case it is required
 				SCENE_UNLOCK_WRITE(PScene);
@@ -748,30 +763,30 @@ void FPhysScene::StartFrame()
 	FlushDeferredCollisionDisableTableQueue();
 
 	// Run the sync scene
-	 TickPhysScene(PST_Sync, PhysicsSubsceneCompletion[PST_Sync]);
-	 {
-		 FGraphEventArray MainScenePrerequisites;
-		 if (FrameLagAsync() && bAsyncSceneEnabled)
-		 {
-			 if (FrameLaggedPhysicsSubsceneCompletion[PST_Async].GetReference() && !FrameLaggedPhysicsSubsceneCompletion[PST_Async]->IsComplete())
-			 {
-				 MainScenePrerequisites.Add(FrameLaggedPhysicsSubsceneCompletion[PST_Async]);
-				 FinishPrerequisites.Add(FrameLaggedPhysicsSubsceneCompletion[PST_Async]);
-			 }
-		 }
-		 if (PhysicsSubsceneCompletion[PST_Sync].GetReference())
-		 {
-			 MainScenePrerequisites.Add(PhysicsSubsceneCompletion[PST_Sync]);
-			 new (FinishPrerequisites) FGraphEventRef(FDelegateGraphTask::CreateAndDispatchWhenReady(FDelegateGraphTask::FDelegate::CreateRaw(this, &FPhysScene::SceneCompletionTask, PST_Sync), TEXT("ProcessPhysScene_Sync"), &MainScenePrerequisites, ENamedThreads::GameThread, ENamedThreads::GameThread));
-		 }
-	 }
+	TickPhysScene(PST_Sync, PhysicsSubsceneCompletion[PST_Sync]);
+	{
+		FGraphEventArray MainScenePrerequisites;
+		if (FrameLagAsync() && bAsyncSceneEnabled)
+		{
+			if (FrameLaggedPhysicsSubsceneCompletion[PST_Async].GetReference() && !FrameLaggedPhysicsSubsceneCompletion[PST_Async]->IsComplete())
+			{
+				MainScenePrerequisites.Add(FrameLaggedPhysicsSubsceneCompletion[PST_Async]);
+				FinishPrerequisites.Add(FrameLaggedPhysicsSubsceneCompletion[PST_Async]);
+			}
+		}
+		if (PhysicsSubsceneCompletion[PST_Sync].GetReference())
+		{
+			MainScenePrerequisites.Add(PhysicsSubsceneCompletion[PST_Sync]);
+			new (FinishPrerequisites)FGraphEventRef(FDelegateGraphTask::CreateAndDispatchWhenReady(FDelegateGraphTask::FDelegate::CreateRaw(this, &FPhysScene::SceneCompletionTask, PST_Sync), TEXT("ProcessPhysScene_Sync"), &MainScenePrerequisites, ENamedThreads::GameThread, ENamedThreads::GameThread));
+		}
+	}
 
 	if (!FrameLagAsync() && bAsyncSceneEnabled)
 	{
 		TickPhysScene(PST_Async, PhysicsSubsceneCompletion[PST_Async]);
 		if (PhysicsSubsceneCompletion[PST_Async].GetReference())
 		{
-			new (FinishPrerequisites) FGraphEventRef(FDelegateGraphTask::CreateAndDispatchWhenReady(FDelegateGraphTask::FDelegate::CreateRaw(this, &FPhysScene::SceneCompletionTask, PST_Async), TEXT("ProcessPhysScene_Async"), PhysicsSubsceneCompletion[PST_Async], ENamedThreads::GameThread, ENamedThreads::GameThread));
+			new (FinishPrerequisites)FGraphEventRef(FDelegateGraphTask::CreateAndDispatchWhenReady(FDelegateGraphTask::FDelegate::CreateRaw(this, &FPhysScene::SceneCompletionTask, PST_Async), TEXT("ProcessPhysScene_Async"), PhysicsSubsceneCompletion[PST_Async], ENamedThreads::GameThread, ENamedThreads::GameThread));
 		}
 	}
 
@@ -780,7 +795,7 @@ void FPhysScene::StartFrame()
 	{
 		if (FinishPrerequisites.Num() > 1)  // we don't need to create a new task if we only have one prerequisite
 		{
-			PhysicsSceneCompletion = TGraphTask<FNullGraphTask>::CreateTask(&FinishPrerequisites, ENamedThreads::GameThread).ConstructAndDispatchWhenReady(TEXT("ProcessPhysScene_Join"), 
+			PhysicsSceneCompletion = TGraphTask<FNullGraphTask>::CreateTask(&FinishPrerequisites, ENamedThreads::GameThread).ConstructAndDispatchWhenReady(TEXT("ProcessPhysScene_Join"),
 				PhysSingleThreadedMode() ? ENamedThreads::GameThread : ENamedThreads::AnyThread);
 		}
 		else
@@ -846,12 +861,12 @@ void FPhysScene::EndFrame(ULineBatchComponent* InLineBatcher)
 
 void FPhysScene::SetIsStaticLoading(bool bStaticLoading)
 {
-	#if WITH_PHYSX
+#if WITH_PHYSX
 	// Loop through scene types to get all scenes
-	for(uint32 SceneType = 0; SceneType < NumPhysScenes; ++SceneType)
+	for (uint32 SceneType = 0; SceneType < NumPhysScenes; ++SceneType)
 	{
 		PxScene* PScene = GetPhysXScene(SceneType);
-		if(PScene != NULL)
+		if (PScene != NULL)
 		{
 			// Lock scene lock, in case it is required
 			SCENE_LOCK_WRITE(PScene);
@@ -890,10 +905,10 @@ NxApexScene* FPhysScene::GetApexScene(uint32 SceneType)
 static void BatchPxRenderBufferLines(class ULineBatchComponent& LineBatcherToUse, const PxRenderBuffer& DebugData)
 {
 	int32 NumPoints = DebugData.getNbPoints();
-	if(NumPoints > 0)
+	if (NumPoints > 0)
 	{
 		const PxDebugPoint* Points = DebugData.getPoints();
-		for(int32 i=0; i<NumPoints; i++)
+		for (int32 i = 0; i<NumPoints; i++)
 		{
 			LineBatcherToUse.DrawPoint(P2UVector(Points->pos), FColor((uint32)Points->color), 2, SDPG_World);
 
@@ -909,29 +924,29 @@ static void BatchPxRenderBufferLines(class ULineBatchComponent& LineBatcherToUse
 	if (NumLines > 0)
 	{
 		const PxDebugLine* Lines = DebugData.getLines();
-		for(int32 i = 0; i<NumLines; i++)
+		for (int32 i = 0; i<NumLines; i++)
 		{
-			new(DebugLines) FBatchedLine(P2UVector(Lines->pos0), P2UVector(Lines->pos1), FColor((uint32)Lines->color0), 0.f, 0.0f, SDPG_World);
+			new(DebugLines)FBatchedLine(P2UVector(Lines->pos0), P2UVector(Lines->pos1), FColor((uint32)Lines->color0), 0.f, 0.0f, SDPG_World);
 			Lines++;
 		}
 	}
 
 	// Add all the 'triangles' from PhysX
 	int32 NumTris = DebugData.getNbTriangles();
-	if(NumTris > 0)
+	if (NumTris > 0)
 	{
 		const PxDebugTriangle* Triangles = DebugData.getTriangles();
-		for(int32 i=0; i<NumTris; i++)
+		for (int32 i = 0; i<NumTris; i++)
 		{
-			new(DebugLines) FBatchedLine(P2UVector(Triangles->pos0), P2UVector(Triangles->pos1), FColor((uint32)Triangles->color0), 0.f, 0.0f,  SDPG_World);
-			new(DebugLines) FBatchedLine(P2UVector(Triangles->pos1), P2UVector(Triangles->pos2), FColor((uint32)Triangles->color1), 0.f, 0.0f,  SDPG_World);
-			new(DebugLines) FBatchedLine(P2UVector(Triangles->pos2), P2UVector(Triangles->pos0), FColor((uint32)Triangles->color2), 0.f, 0.0f,  SDPG_World);
+			new(DebugLines)FBatchedLine(P2UVector(Triangles->pos0), P2UVector(Triangles->pos1), FColor((uint32)Triangles->color0), 0.f, 0.0f, SDPG_World);
+			new(DebugLines)FBatchedLine(P2UVector(Triangles->pos1), P2UVector(Triangles->pos2), FColor((uint32)Triangles->color1), 0.f, 0.0f, SDPG_World);
+			new(DebugLines)FBatchedLine(P2UVector(Triangles->pos2), P2UVector(Triangles->pos0), FColor((uint32)Triangles->color2), 0.f, 0.0f, SDPG_World);
 			Triangles++;
 		}
 	}
 
 	// Draw them all in one call.
-	if(DebugLines.Num() > 0)
+	if (DebugLines.Num() > 0)
 	{
 		LineBatcherToUse.DrawLines(DebugLines);
 	}
@@ -944,7 +959,7 @@ void FPhysScene::AddDebugLines(uint32 SceneType, class ULineBatchComponent* Line
 {
 	check(SceneType < NumPhysScenes);
 
-	if(LineBatcherToUse)
+	if (LineBatcherToUse)
 	{
 #if WITH_PHYSX
 		// Render PhysX debug data
@@ -969,14 +984,14 @@ void FPhysScene::ApplyWorldOffset(FVector InOffset)
 {
 #if WITH_PHYSX
 	// Loop through scene types to get all scenes
-	for(uint32 SceneType = 0; SceneType < NumPhysScenes; ++SceneType)
+	for (uint32 SceneType = 0; SceneType < NumPhysScenes; ++SceneType)
 	{
 		PxScene* PScene = GetPhysXScene(SceneType);
-		if(PScene != NULL)
+		if (PScene != NULL)
 		{
 			// Lock scene lock, in case it is required
 			SCENE_LOCK_WRITE(PScene);
-			
+
 			PScene->shiftOrigin(U2PVector(-InOffset));
 
 			// Unlock scene lock, in case it is required
@@ -996,7 +1011,7 @@ void FPhysScene::InitPhysScene(uint32 SceneType)
 	// Include scene descriptor in loop, so that we might vary it with scene type
 	PxSceneDesc PSceneDesc(GPhysXSDK->getTolerancesScale());
 	PSceneDesc.cpuDispatcher = CPUDispatcher;
-	
+
 	FPhysSceneShaderInfo PhysSceneShaderInfo;
 	PhysSceneShaderInfo.PhysScene = this;
 	PSceneDesc.filterShaderData = &PhysSceneShaderInfo;
@@ -1047,7 +1062,7 @@ void FPhysScene::InitPhysScene(uint32 SceneType)
 	PSceneDesc.dynamicTreeRebuildRateHint = PhysXSlowRebuildRate;
 
 	bool bIsValid = PSceneDesc.isValid();
-	if(!bIsValid)
+	if (!bIsValid)
 	{
 		UE_LOG(LogPhysics, Log, TEXT("Invalid PSceneDesc"));
 	}
@@ -1081,7 +1096,7 @@ void FPhysScene::InitPhysScene(uint32 SceneType)
 	SCENE_LOCK_WRITE(PScene);
 
 	// enable CCD at scene level
-	if(bGlobalCCD)
+	if (bGlobalCCD)
 	{
 		PScene->setFlag(PxSceneFlag::eENABLE_CCD, true);
 	}
@@ -1105,7 +1120,7 @@ void FPhysScene::InitPhysScene(uint32 SceneType)
 	PhysXSceneCount++;
 
 	// Only create PhysXVehicleManager in the sync scene
-	if ( SceneType == PST_Sync )
+	if (SceneType == PST_Sync)
 	{
 		check(VehicleManager == NULL);
 		VehicleManager = new FPhysXVehicleManager(PScene);
@@ -1125,7 +1140,7 @@ void FPhysScene::InitPhysScene(uint32 SceneType)
 	{
 		PhysSubSteppers[SceneType]->SetVehicleManager(VehicleManager);
 	}
-	
+
 #endif
 
 #endif // WITH_PHYSX
@@ -1137,17 +1152,17 @@ void FPhysScene::TermPhysScene(uint32 SceneType)
 
 #if WITH_PHYSX
 	PxScene* PScene = GetPhysXScene(SceneType);
-	if(PScene != NULL)
+	if (PScene != NULL)
 	{
 #if WITH_APEX
 		NxApexScene* ApexScene = GetApexScene(SceneType);
-		if(ApexScene != NULL)
+		if (ApexScene != NULL)
 		{
 			GPhysCommandHandler->DeferredRelease(ApexScene);
 		}
 #endif // #if WITH_APEX
 
-		if ( SceneType == PST_Sync && VehicleManager != NULL )
+		if (SceneType == PST_Sync && VehicleManager != NULL)
 		{
 			delete VehicleManager;
 			VehicleManager = NULL;
