Fix incorrect ensure for destructible shapes.

diff --git a/Engine/Source/Runtime/Engine/Private/PhysicsEngine/PhysScene.cpp b/Engine/Source/Runtime/Engine/Private/PhysicsEngine/PhysScene.cpp
--- a/Engine/Source/Runtime/Engine/Private/PhysicsEngine/PhysScene.cpp
+++ b/Engine/Source/Runtime/Engine/Private/PhysicsEngine/PhysScene.cpp
@@ -631,7 +631,7 @@ void FPhysScene::SyncComponentsToBodies(uint32 SceneType)
 					PxShape* Shape = Shapes[ShapeIdx];
 					int32 ChunkIndex;
 					NxDestructibleActor* DestructibleActor = GApexModuleDestructible->getDestructibleAndChunk(Shape, &ChunkIndex);
-					ensure(DestructibleActor == DestructibleChunkInfo->OwningComponent->ApexDestructibleActor);
+					
 					const physx::PxMat44 ChunkPoseRT = DestructibleActor->getChunkPose(ChunkIndex);
 					const physx::PxTransform Transform(ChunkPoseRT);
 					DestructibleChunkInfo->OwningComponent->SetChunkWorldRT(ChunkIndex, P2UQuat(Transform.q), P2UVector(Transform.p));
