Put Destroyed functions back in to Blueprints
EndPlay in Blueprints provides reason Enum

diff --git a/Engine/Config/BaseEngine.ini b/Engine/Config/BaseEngine.ini
--- a/Engine/Config/BaseEngine.ini
+++ b/Engine/Config/BaseEngine.ini
@@ -735,11 +735,6 @@ bUseStreamingPause=false
 +K2FieldRedirects=(OldFieldName="LevelScriptActor.BeginGame",NewFieldName="Actor.ReceiveBeginPlay")
 +K2FieldRedirects=(OldFieldName="LevelScriptActor.BeginGameSignature__DelegateSignature", NewFieldName="Actor.ReceiveBeginPlaySignature__DelegateSignature")
 
-+K2FieldRedirects=(OldFieldName="Actor.ReceiveDestroyed",NewFieldName="Actor.ReceiveEndPlay")
-+K2FieldRedirects=(OldFieldName="Actor.ReceiveDestroyed__DelegateSignature",NewFieldName="Actor.ReceiveEndPlay__DelegateSignature")
-+K2FieldRedirects=(OldFieldName="Actor.OnDestroyed",NewFieldName="Actor.OnEndPlay")
-+K2FieldRedirects=(OldFieldName="Actor.DestroyedSignature__DelegateSignature",NewFieldName="Actor.EndPlaySignature__DelegateSignature")
-
 ; AttachTo change - changed bMaintainWorldPosition to enum
 +K2ParamRedirects=(NodeName="/Script/BlueprintGraph.K2Node_CallFunction", OldParamName="AttachTo.bMaintainWorldPosition", NewParamName="AttachType", OldParamValues="false;true", NewParamValues="KeepRelativeOffset;KeepWorldPosition")
 +K2ParamRedirects=(NodeName="/Script/BlueprintGraph.K2Node_CallFunction", OldParamName="AttachRootComponentTo.bMaintainWorldPosition", NewParamName="AttachLocationType", OldParamValues="false;true", NewParamValues="KeepRelativeOffset;KeepWorldPosition")
diff --git a/Engine/Source/Editor/UnrealEd/Private/PlayLevel.cpp b/Engine/Source/Editor/UnrealEd/Private/PlayLevel.cpp
--- a/Engine/Source/Editor/UnrealEd/Private/PlayLevel.cpp
+++ b/Engine/Source/Editor/UnrealEd/Private/PlayLevel.cpp
@@ -426,7 +426,7 @@ void UEditorEngine::TeardownPlaySession(FWorldContext &PieWorldContext)
 	// Go through and let all the PlayWorld Actor's know they are being destroyed
 	for (FActorIterator ActorIt(PlayWorld); ActorIt; ++ActorIt)
 	{
-		ActorIt->EndPlay(AActor::EEndPlayReason::EndPlayInEditor);
+		ActorIt->EndPlay(EEndPlayReason::EndPlayInEditor);
 	}
 
 	// Move blueprint debugging pointers back to the objects in the editor world
diff --git a/Engine/Source/Runtime/Engine/Classes/GameFramework/Actor.h b/Engine/Source/Runtime/Engine/Classes/GameFramework/Actor.h
--- a/Engine/Source/Runtime/Engine/Classes/GameFramework/Actor.h
+++ b/Engine/Source/Runtime/Engine/Classes/GameFramework/Actor.h
@@ -4,10 +4,25 @@
 
 #include "Runtime/Engine/Classes/Engine/LatentActionManager.h"
 #include "Runtime/Engine/Classes/Camera/CameraTypes.h"
+
+UENUM(BlueprintType)
+namespace EEndPlayReason
+{
+	enum Type
+	{
+		ActorDestroyed,		// When the Actor is explicitly destroyed
+		LevelTransition,	// When the world is being unloaded for a level transition
+		EndPlayInEditor,	// When the world is being unloaded because PIE is ending
+		RemovedFromWorld,	// When the level it is a member of is streamed out
+	};
+
+}
+
 #include "Actor.generated.h"
 
 ENGINE_API DECLARE_LOG_CATEGORY_EXTERN(LogActor, Log, Warning);
  
+
 // Delegate signatures
 DECLARE_DYNAMIC_MULTICAST_DELEGATE_FourParams( FTakeAnyDamageSignature, float, Damage, const class UDamageType*, DamageType, class AController*, InstigatedBy, class AActor*, DamageCauser );
 DECLARE_DYNAMIC_MULTICAST_DELEGATE_EightParams( FTakePointDamageSignature, float, Damage, class AController*, InstigatedBy, FVector, HitLocation, class UPrimitiveComponent*, FHitComponent, FName, BoneName, FVector, ShotFromDirection, const class UDamageType*, DamageType, class AActor*, DamageCauser );
@@ -24,7 +39,8 @@ DECLARE_DYNAMIC_MULTICAST_DELEGATE_OneParam( FActorOnInputTouchEndSignature, ETo
 DECLARE_DYNAMIC_MULTICAST_DELEGATE_OneParam( FActorBeginTouchOverSignature, ETouchIndex::Type, FingerIndex );
 DECLARE_DYNAMIC_MULTICAST_DELEGATE_OneParam( FActorEndTouchOverSignature, ETouchIndex::Type, FingerIndex );
 
-DECLARE_DYNAMIC_MULTICAST_DELEGATE( FActorEndPlaySignature );
+DECLARE_DYNAMIC_MULTICAST_DELEGATE(FActorDestroyedSignature);
+DECLARE_DYNAMIC_MULTICAST_DELEGATE_OneParam(FActorEndPlaySignature, EEndPlayReason::Type, EndPlayReason);
 
 DECLARE_DELEGATE_FourParams(FMakeNoiseDelegate, class AActor*, float, class APawn*, const FVector&);
 
@@ -994,9 +1010,17 @@ public:
 	 */
 	bool Destroy(bool bNetForce = false, bool bShouldModifyLevel = true );
 
+	UFUNCTION(BlueprintImplementableEvent, meta = (Keywords = "delete", FriendlyName = "Destroyed"))
+	virtual void ReceiveDestroyed();
+
+	/** Called when the actor is destroyed. */
+	UPROPERTY(BlueprintAssignable)
+	FActorDestroyedSignature OnDestroyed;
+
+
 	/** Event to notify blueprints this actor is about to be deleted. */
 	UFUNCTION(BlueprintImplementableEvent, meta=(Keywords = "delete", FriendlyName = "End Play"))
-	virtual void ReceiveEndPlay();
+	virtual void ReceiveEndPlay(EEndPlayReason::Type EndPlayReason);
 
 	/** Called when the actor is destroyed. */
 	UPROPERTY(BlueprintAssignable)
@@ -1625,16 +1649,8 @@ public:
 	/**	Do anything needed to clear out cross level references; Called from ULevel::PreSave	 */
 	virtual void ClearCrossLevelReferences();
 	
-	enum EEndPlayReason
-	{
-		ActorDestroyed,		// When the Actor is explicitly destroyed
-		LevelTransition,	// When the world is being unloaded for a level transition
-		EndPlayInEditor,	// When the world is being unloaded because PIE is ending
-		RemovedFromWorld,	// When the level it is a member of is streamed out
-	};
-
 	/** Called whenever this actor is being removed from a level */
-	virtual void EndPlay(const EEndPlayReason EndPlayReason);
+	virtual void EndPlay(const EEndPlayReason::Type EndPlayReason);
 
 	/** iterates up the Base chain to see whether or not this Actor is based on the given Actor
 	 * @param Other the Actor to test for
diff --git a/Engine/Source/Runtime/Engine/Private/Actor.cpp b/Engine/Source/Runtime/Engine/Private/Actor.cpp
--- a/Engine/Source/Runtime/Engine/Private/Actor.cpp
+++ b/Engine/Source/Runtime/Engine/Private/Actor.cpp
@@ -1521,11 +1521,11 @@ void AActor::PrestreamTextures( float Seconds, bool bEnableStreaming, int32 Cine
 void AActor::OnRep_Instigator() {}
 
 
-void AActor::EndPlay(const EEndPlayReason EndPlayReason)
+void AActor::EndPlay(const EEndPlayReason::Type EndPlayReason)
 {
 	// Dispatch the blueprint events
-	ReceiveEndPlay();
-	OnEndPlay.Broadcast();
+	ReceiveEndPlay(EndPlayReason);
+	OnEndPlay.Broadcast(EndPlayReason);
 
 	// Behaviors specific to an actor being unloaded due to a streaming level removal
 	if (EndPlayReason == EEndPlayReason::RemovedFromWorld)
@@ -1598,7 +1598,9 @@ void AActor::Destroyed()
 	{
 		EndPlay(EEndPlayReason::ActorDestroyed);
 	}
-	GetWorld()->RemoveNetworkActor( this );
+	ReceiveDestroyed();
+	OnDestroyed.Broadcast();
+	GetWorld()->RemoveNetworkActor(this);
 }
 
 void AActor::TornOff() {}
diff --git a/Engine/Source/Runtime/Engine/Private/UnrealEngine.cpp b/Engine/Source/Runtime/Engine/Private/UnrealEngine.cpp
--- a/Engine/Source/Runtime/Engine/Private/UnrealEngine.cpp
+++ b/Engine/Source/Runtime/Engine/Private/UnrealEngine.cpp
@@ -8106,7 +8106,7 @@ bool UEngine::LoadMap( FWorldContext& WorldContext, FURL URL, class UPendingNetG
 		{
 			if (ActorIt->bActorInitialized)
 			{
-				ActorIt->EndPlay(AActor::EEndPlayReason::LevelTransition);
+				ActorIt->EndPlay(EEndPlayReason::LevelTransition);
 			}
 		}
 
diff --git a/Engine/Source/Runtime/Engine/Private/World.cpp b/Engine/Source/Runtime/Engine/Private/World.cpp
--- a/Engine/Source/Runtime/Engine/Private/World.cpp
+++ b/Engine/Source/Runtime/Engine/Private/World.cpp
@@ -1660,7 +1660,7 @@ void UWorld::RemoveFromWorld( ULevel* Level )
 			AActor* Actor = Level->Actors[ActorIdx];
 			if (Actor != NULL)
 			{
-				Actor->EndPlay(AActor::EEndPlayReason::RemovedFromWorld);
+				Actor->EndPlay(EEndPlayReason::RemovedFromWorld);
 				if (NetDriver)
 				{
 					NetDriver->NotifyActorLevelUnloaded(Actor);
@@ -4263,7 +4263,7 @@ UWorld* FSeamlessTravelHandler::Tick()
 
 					if (TheActor->bActorInitialized)
 					{
-						TheActor->EndPlay(AActor::EEndPlayReason::LevelTransition);
+						TheActor->EndPlay(EEndPlayReason::LevelTransition);
 					}
 
 					// otherwise, set to be deleted
