UE4: Fixed a cleanup issue that was caused by assigning the OwningWorld on a streaming level in situations where it was not necessary. Also updated the log statement for a missing level when duplicating for PIE to be a log instead of a warning if we are expecting to load the level later.

diff --git a/Engine/Source/Runtime/Engine/Private/LevelStreaming.cpp b/Engine/Source/Runtime/Engine/Private/LevelStreaming.cpp
--- a/Engine/Source/Runtime/Engine/Private/LevelStreaming.cpp
+++ b/Engine/Source/Runtime/Engine/Private/LevelStreaming.cpp
@@ -380,11 +380,18 @@ bool ULevelStreaming::RequestLevel(UWorld* PersistentWorld, bool bAllowLevelLoad
 		}
 		else if (PersistentWorld->WorldComposition == NULL) // In world composition streaming levels are not loaded by default
 		{
-			UE_LOG(LogLevelStreaming, Warning, TEXT("Unable to duplicate PIE World: '%s'"), *NonPrefixedLevelName);
-			for (TObjectIterator<UWorld> It; It; ++It)
+			if ( bAllowLevelLoadRequests )
 			{
-				UWorld *W = *It;
-				UE_LOG(LogLevelStreaming, Warning, TEXT("    Loaded World: %s"), *W->GetPathName() );
+				UE_LOG(LogLevelStreaming, Log, TEXT("World to duplicate for PIE '%s' not found. Attempting load."), *NonPrefixedLevelName);
+			}
+			else
+			{
+				UE_LOG(LogLevelStreaming, Warning, TEXT("Unable to duplicate PIE World: '%s'"), *NonPrefixedLevelName);
+				for (TObjectIterator<UWorld> It; It; ++It)
+				{
+					UWorld *W = *It;
+					UE_LOG(LogLevelStreaming, Warning, TEXT("    Loaded World: %s"), *W->GetPathName() );
+				}
 			}
 		}
 	}
@@ -516,20 +523,7 @@ void ULevelStreaming::AsyncLevelLoadComplete( const FString& InPackageName, UPac
 					FName OldDesiredPackageName = FName(*InPackageName);
 					UWorld** OwningWorldPtr = ULevel::StreamedLevelsOwningWorld.Find(OldDesiredPackageName);
 					UWorld* OwningWorld = OwningWorldPtr ? *OwningWorldPtr : NULL;
-					if ( OwningWorld )
-					{
-						if ( DestinationWorld->PersistentLevel )
-						{
-							DestinationWorld->PersistentLevel->OwningWorld = OwningWorld;
-						}
-
-						// In some cases, BSP render data is not created because the OwningWorld was not set correctly.
-						// Regenerate that render data here
-						DestinationWorld->PersistentLevel->InvalidateModelSurface();
-						DestinationWorld->PersistentLevel->CommitModelSurfaces();
-
-						ULevel::StreamedLevelsOwningWorld.Remove(OldDesiredPackageName);
-					}
+					ULevel::StreamedLevelsOwningWorld.Remove(OldDesiredPackageName);
 
 					// Try again with the destination package to load.
 					// IMPORTANT: check this BEFORE changing PackageNameToLoad, otherwise you wont know if the package name was supposed to be different.
@@ -568,7 +562,20 @@ void ULevelStreaming::AsyncLevelLoadComplete( const FString& InPackageName, UPac
 					}
 					else
 					{
-						// Loading the requested package normally. Update the requested package to the destination.
+						// Loading the requested package normally. Fix up the destination world then update the requested package to the destination.
+						if (OwningWorld)
+						{
+							if (DestinationWorld->PersistentLevel)
+							{
+								DestinationWorld->PersistentLevel->OwningWorld = OwningWorld;
+							}
+
+							// In some cases, BSP render data is not created because the OwningWorld was not set correctly.
+							// Regenerate that render data here
+							DestinationWorld->PersistentLevel->InvalidateModelSurface();
+							DestinationWorld->PersistentLevel->CommitModelSurfaces();
+						}
+						
 						PackageName = PackageNameToLoad;
 					}
 				}
