Remove unnecessary condition.

This was a stop gap measure to stop spam happening in benign cases, but since the unmapped logic exists on the client, we no longer need this.

diff --git a/Engine/Source/Runtime/CoreUObject/Public/UObject/CoreNet.h b/Engine/Source/Runtime/CoreUObject/Public/UObject/CoreNet.h
--- a/Engine/Source/Runtime/CoreUObject/Public/UObject/CoreNet.h
+++ b/Engine/Source/Runtime/CoreUObject/Public/UObject/CoreNet.h
@@ -292,7 +292,6 @@ enum ELifetimeCondition
 	COND_None				= 0,		// This property has no condition, and will send anytime it changes
 	COND_InitialOnly		= 1,		// This property will only attempt to send on the initial bunch
 	COND_OwnerOnly			= 2,		// This property will only send to the actor's owner
-	COND_OwnerOrNotNull		= 3,		// This property will only send to the actor's owner (or anyone if it's not NULL for their copy)
 	COND_SkipOwner			= 4,		// This property send to every connection EXCEPT the owner
 	COND_SimulatedOnly		= 5,		// This property will only send to simulated actors
 	COND_AutonomousOnly		= 6,		// This property will only send to autonomous actors
diff --git a/Engine/Source/Runtime/Engine/Private/ActorReplication.cpp b/Engine/Source/Runtime/Engine/Private/ActorReplication.cpp
--- a/Engine/Source/Runtime/Engine/Private/ActorReplication.cpp
+++ b/Engine/Source/Runtime/Engine/Private/ActorReplication.cpp
@@ -223,6 +223,7 @@ void AActor::GetLifetimeReplicatedProps( TArray< FLifetimeProperty > & OutLifeti
 
 	DOREPLIFETIME( AActor, Role );
 	DOREPLIFETIME( AActor, RemoteRole );
+	DOREPLIFETIME( AActor, Owner );
 	DOREPLIFETIME( AActor, bHidden );
 
 	DOREPLIFETIME( AActor, bTearOff );
@@ -231,7 +232,6 @@ void AActor::GetLifetimeReplicatedProps( TArray< FLifetimeProperty > & OutLifeti
 
 	DOREPLIFETIME( AActor, Instigator );
 
-	DOREPLIFETIME_CONDITION( AActor, Owner, COND_OwnerOrNotNull );
 	DOREPLIFETIME_CONDITION( AActor, ReplicatedMovement, COND_SimulatedOrPhysics );
 }
 
diff --git a/Engine/Source/Runtime/Engine/Private/RepLayout.cpp b/Engine/Source/Runtime/Engine/Private/RepLayout.cpp
--- a/Engine/Source/Runtime/Engine/Private/RepLayout.cpp
+++ b/Engine/Source/Runtime/Engine/Private/RepLayout.cpp
@@ -966,13 +966,8 @@ uint16						Handle ) const
 #if ENABLE_CLIENT_UNMAP_LOGIC == 0
 			// HACK: REMOVE ME. 
 			// Move all this unmapped logic to client.
-			// This is to fix an issue where "Owner" on actor is replicated to all connections, but since only one connection will be mapped
-			// this will trigger a bunch of "unmapped longer than normal" false positives.
-			// We can't use COND_OwnerOnly since we also want to send NULL's to old owners
-			const bool bWorryAboutUnmapped = ( RepFlags.bNetOwner || ParentCmd.Condition != COND_OwnerOrNotNull );
-
 			// If this property is unmapped, force it to resend
-			if ( ( !bMapped || WriterState.Writer.PackageMap->SerializedUnAckedObject() ) && bWorryAboutUnmapped )
+			if ( !bMapped || WriterState.Writer.PackageMap->SerializedUnAckedObject() )
 			{
 				Unmapped.Add( Handle );
 
@@ -2237,7 +2232,7 @@ void FRepLayout::InitFromObjectClass( UClass * InObjectClass )
 		// Fix Lifetime props to have the proper index to the parent
 		for ( int32 i = 0; i < LifetimeProps.Num(); i++ )
 		{
-			// Store the condition on the parent in case we need it (for now we use this for COND_OwnerOrNotNull)
+			// Store the condition on the parent in case we need it
 			Parents[LifetimeProps[i].RepIndex].Condition = LifetimeProps[i].Condition;
 
 			if ( Parents[LifetimeProps[i].RepIndex].Flags & PARENT_IsCustomDelta )
@@ -2526,7 +2521,6 @@ void FRepLayout::RebuildConditionalProperties( FRepState * RESTRICT	RepState, co
 	ConditionMap[COND_InitialOnly]			= bIsInitial;
 
 	ConditionMap[COND_OwnerOnly]			= bIsOwner;
-	ConditionMap[COND_OwnerOrNotNull]		= true;			// We always put these in the list so we can do special logic at send time
 	ConditionMap[COND_SkipOwner]			= !bIsOwner;
 
 	ConditionMap[COND_SimulatedOnly]		= bIsSimulated;
