Remove unmapped logic from server code now that client handles this

diff --git a/Engine/Source/Runtime/Engine/Private/RepLayout.cpp b/Engine/Source/Runtime/Engine/Private/RepLayout.cpp
--- a/Engine/Source/Runtime/Engine/Private/RepLayout.cpp
+++ b/Engine/Source/Runtime/Engine/Private/RepLayout.cpp
@@ -23,10 +23,6 @@ FAutoConsoleVariable CVarDoReplicationContextString( TEXT( "net.ContextDebug" ),
 
 //#define ENABLE_SUPER_CHECKSUMS
 
-static const int32 UNMAPPED_FRAMES_THRESHOLD = 100;
-
-#define ENABLE_CLIENT_UNMAP_LOGIC 1
-
 #ifdef USE_CUSTOM_COMPARE
 static FORCEINLINE bool CompareBool( const FRepLayoutCmd & Cmd, const void * A, const void * B )
 {
@@ -554,7 +550,7 @@ bool FRepLayout::ReplicateProperties(
 	// PreOpenAckHistory are all the properties sent before we got our first open ack
 	const bool bFlushPreOpenAckHistory = RepState->OpenAckedCalled && RepState->PreOpenAckHistory.Num() > 0;
 
-	if ( PropertyChanged || RepState->Unmapped.Num() > 0 || RepState->NumNaks > 0 || bFlushPreOpenAckHistory )
+	if ( PropertyChanged || RepState->NumNaks > 0 || bFlushPreOpenAckHistory )
 	{
 		// Use the first inactive history item to build this change list on
 		check( RepState->HistoryEnd - RepState->HistoryStart < FRepState::MAX_CHANGE_HISTORY );
@@ -596,20 +592,6 @@ bool FRepLayout::ReplicateProperties(
 		// Update the history, and merge in any nak'd change lists
 		UpdateChangelistHistory( RepState, ObjectClass, Data, OwningChannel->Connection->OutAckPacketId, &Changed );
 
-		// Merge in unmapped properties to be resent
-		if ( RepState->Unmapped.Num() > 0 )
-		{
-			TArray< uint16 > Temp = Changed;
-			Changed.Empty();
-			MergeDirtyList( (void*)Data, Temp, RepState->Unmapped, Changed );
-
-#ifdef SANITY_CHECK_MERGES
-			SanityCheckChangeList( Data, Changed );
-#endif
-
-			RepState->Unmapped.Empty();
-		}
-
 		// Merge in the PreOpenAckHistory (unreliable properties sent before the bunch was initially acked)
 		if ( bFlushPreOpenAckHistory )
 		{
@@ -644,21 +626,6 @@ bool FRepLayout::ReplicateProperties(
 		}
 #endif
 		
-		if ( RepState->Unmapped.Num() )
-		{
-			OwningChannel->bActorMustStayDirty = true;
-			RepState->UnmappedFrames++;
-			if ( RepState->UnmappedFrames >= UNMAPPED_FRAMES_THRESHOLD )
-			{
-				UE_LOG( LogNet, Warning, TEXT( "FRepLayout::ReplicateProperties: Properties have been unmapped longer than normal: %s" ), *ObjectClass->GetName() );
-				RepState->UnmappedFrames = 0;
-			}
-		}
-		else
-		{
-			RepState->UnmappedFrames = 0;
-		}
-
 		return true;
 	}
 
@@ -790,12 +757,6 @@ bool FRepLayout::AllAcked( FRepState * RepState ) const
 		return false;
 	}
 
-	if ( RepState->Unmapped.Num() > 0 )
-	{
-		// We have unmapped properties that need to be resent
-		return false;
-	}
-
 	if ( RepState->NumNaks > 0 )
 	{
 		return false;
@@ -860,7 +821,6 @@ void FRepLayout::SendProperties_DynamicArray_r(
 	const int32					CmdIndex, 
 	const uint8 * RESTRICT		StoredData, 
 	const uint8 * RESTRICT		Data, 
-	TArray< uint16 > &			Unmapped,
 	uint16						Handle ) const
 {
 	const FRepLayoutCmd & Cmd = Cmds[ CmdIndex ];
@@ -886,14 +846,12 @@ void FRepLayout::SendProperties_DynamicArray_r(
 	Data = (uint8*)Array->GetData();
 	StoredData = (uint8*)StoredArray->GetData();
 
-	TArray< uint16 > LocalUnmapped;
-
 	uint16 LocalHandle = 0;
 
 	for ( int32 i = 0; i < Array->Num(); i++ )
 	{
 		const int32 ElementOffset = i * Cmd.ElementSize;
-		LocalHandle = SendProperties_r( RepState, RepFlags, WriterState, CmdIndex + 1, Cmd.EndCmd - 1, StoredData + ElementOffset, Data + ElementOffset, LocalUnmapped, LocalHandle );
+		LocalHandle = SendProperties_r( RepState, RepFlags, WriterState, CmdIndex + 1, Cmd.EndCmd - 1, StoredData + ElementOffset, Data + ElementOffset, LocalHandle );
 	}
 
 	check( WriterState.CurrentChanged - OldChangedIndex == ArrayChangedCount );	// Make sure we read correct amount
@@ -902,15 +860,6 @@ void FRepLayout::SendProperties_DynamicArray_r(
 	WriterState.CurrentChanged++;
 
 	WritePropertyHandle( WriterState.Writer, 0, WriterState.bDoChecksum );		// Signify end of dynamic array
-
-	if ( LocalUnmapped.Num() > 0 )
-	{
-		// We have unmapped properties, save them
-		Unmapped.Add( Handle );
-		Unmapped.Add( LocalUnmapped.Num() );	// So we can jump over this array if needed
-		Unmapped.Append( LocalUnmapped );
-		Unmapped.Add( 0 );
-	}
 }
 
 uint16 FRepLayout::SendProperties_r( 
@@ -921,8 +870,7 @@ uint16 FRepLayout::SendProperties_r(
 	const int32					CmdEnd, 
 	const uint8 * RESTRICT		StoredData, 
 	const uint8 * RESTRICT		Data, 
-	TArray< uint16 > &			Unmapped,
-uint16						Handle ) const
+	uint16						Handle ) const
 {
 	for ( int32 CmdIndex = CmdStart; CmdIndex < CmdEnd; CmdIndex++ )
 	{
@@ -936,7 +884,7 @@ uint16						Handle ) const
 		{
 			if ( ShouldSendProperty( WriterState, Handle ) )
 			{
-				SendProperties_DynamicArray_r( RepState, RepFlags, WriterState, CmdIndex, StoredData + Cmd.Offset, Data + Cmd.Offset, Unmapped, Handle );
+				SendProperties_DynamicArray_r( RepState, RepFlags, WriterState, CmdIndex, StoredData + Cmd.Offset, Data + Cmd.Offset, Handle );
 			}
 			CmdIndex = Cmd.EndCmd - 1;	// Jump past children of this array (-1 for the ++ in the for loop)
 			continue;
@@ -963,21 +911,6 @@ uint16						Handle ) const
 
 			NETWORK_PROFILER( GNetworkProfiler.TrackReplicateProperty( ParentCmd.Property, false, false, 0, 0, NumEndBits - NumStartBits ) );
 
-#if ENABLE_CLIENT_UNMAP_LOGIC == 0
-			// HACK: REMOVE ME. 
-			// Move all this unmapped logic to client.
-			// If this property is unmapped, force it to resend
-			if ( !bMapped || WriterState.Writer.PackageMap->SerializedUnAckedObject() )
-			{
-				Unmapped.Add( Handle );
-
-				if ( RepState->UnmappedFrames == UNMAPPED_FRAMES_THRESHOLD - 1 )
-				{
-					UE_LOG( LogNet, Warning, TEXT( "  FRepLayout::SendProperties_r: Property has been unmapped longer than normal: %s" ), *Cmd.Property->GetName() );
-				}
-			}
-#endif
-
 			// Make the shadow state match the actual state at the time of send
 			StoreProperty( Cmd, (void*)( StoredData + Cmd.Offset ), (const void*)( Data + Cmd.Offset ) );
 
@@ -1060,8 +993,6 @@ void FRepLayout::SendProperties(
 	int32 &						LastIndex, 
 	bool &						bContentBlockWritten ) const
 {
-	check( RepState->Unmapped.Num() == 0 );
-
 #ifdef ENABLE_PROPERTY_CHECKSUMS
 	const bool bDoChecksum = CVarDoPropertyChecksum.GetValueOnGameThread() == 1;
 #else
@@ -1074,14 +1005,9 @@ void FRepLayout::SendProperties(
 	Writer.WriteBit( bDoChecksum ? 1 : 0 );
 #endif
 
-	SendProperties_r( RepState, RepFlags, WriterState, 0, Cmds.Num() - 1, RepState->StaticBuffer.GetTypedData(), Data, RepState->Unmapped, 0 );
+	SendProperties_r( RepState, RepFlags, WriterState, 0, Cmds.Num() - 1, RepState->StaticBuffer.GetTypedData(), Data, 0 );
 
 	WritePropertyHandle( Writer, 0, bDoChecksum );
-
-	if ( RepState->Unmapped.Num() > 0 )
-	{
-		RepState->Unmapped.Add( 0 );
-	}
 }
 
 static void ReadNextHandle( FRepReaderState & ReaderState )
@@ -1168,7 +1094,6 @@ bool FRepLayout::ReadProperty(
 		Cmd.Property->NetSerializeItem( ReaderState.Bunch, ReaderState.Bunch.PackageMap, Data + SwappedCmd.Offset );
 	}
 
-#if ENABLE_CLIENT_UNMAP_LOGIC == 1
 	if ( ReaderState.Bunch.PackageMap->GetLoadedUnmappedObject() )
 	{
 		// If we have an unmapped guid, we need to remember it, so we can fix up this object pointer when it finally arrives at a other time
@@ -1178,7 +1103,6 @@ bool FRepLayout::ReadProperty(
 		UE_LOG( LogNet, Verbose, TEXT( "ADDED unmapped property: Offset: %i, Guid: %s, Name: %s"), LocalAbsOffset, *ReaderState.Bunch.PackageMap->GetLastUnmappedNetGUID().ToString(), *Cmd.Property->GetName() );
 		ReaderState.bHasUnmapped = true;
 	}
-#endif
 
 #ifdef ENABLE_PROPERTY_CHECKSUMS
 	if ( ReaderState.bDoChecksum )
@@ -1251,7 +1175,6 @@ bool FRepLayout::ReceiveProperties_DynamicArray_r(
 		FScriptArray * Array = (FScriptArray *)Data;
 		FScriptArray * StoredArray = (FScriptArray *)StoredData;
 
-#if ENABLE_CLIENT_UNMAP_LOGIC == 1
 		// Since we don't know yet if something under us could be unmapped, go ahead and allocate an array container now
 		FUnmappedGuidMgrElement * NewArrayElement = UnmappedGuids->Map.Find( AbsOffset );
 
@@ -1268,7 +1191,6 @@ bool FRepLayout::ReceiveProperties_DynamicArray_r(
 		check( NewArrayElement->CmdIndex == CmdIndex );
 
 		UnmappedGuids = NewArrayElement->Array;
-#endif
 
 		const FRepParentCmd & Parent = Parents[Cmd.ParentIndex];
 
@@ -1473,9 +1395,7 @@ void FRepLayout::UpdateUnmappedObjects( FRepState *	RepState, UPackageMap * Pack
 	bOutSomeObjectsWereMapped	= false;
 	bOutHasMoreUnmapped			= false;
 
-#if ENABLE_CLIENT_UNMAP_LOGIC == 1
 	UpdateUnmappedObjects_r( RepState, &RepState->UnmappedGuids, OriginalObject, PackageMap, (uint8*)OriginalObject, RepState->StaticBuffer.Num(), bOutSomeObjectsWereMapped, bOutHasMoreUnmapped );
-#endif
 }
 
 void FRepLayout::CallRepNotifies( FRepState * RepState, UObject * Object ) const
diff --git a/Engine/Source/Runtime/Engine/Public/Net/RepLayout.h b/Engine/Source/Runtime/Engine/Public/Net/RepLayout.h
--- a/Engine/Source/Runtime/Engine/Public/Net/RepLayout.h
+++ b/Engine/Source/Runtime/Engine/Public/Net/RepLayout.h
@@ -98,7 +98,6 @@ public:
 		HistoryEnd( 0 ),
 		LastReplicationFrame( 0 ),
 		NumNaks( 0 ),
-		UnmappedFrames( 0 ),
 		OpenAckedCalled( false ),
 		AwakeFromDormancy( false ),
 		ActiveStatusChanged( 0 )
@@ -112,8 +111,6 @@ public:
 
 	TSharedPtr< FRepLayout >	RepLayout;
 	
-	TArray< uint16 >			Unmapped;
-
 	TArray< UProperty * >		RepNotifies;
 
 	TSharedPtr<FRepChangedPropertyTracker> RepChangedPropertyTracker;
@@ -127,8 +124,6 @@ public:
 	uint32						LastReplicationFrame;
 	int32						NumNaks;
 
-	int32						UnmappedFrames;		// Continuous frames that contain unmapped objects, used to warn when we have unmapped objects for too long
-
 	TArray< FRepChangedHistory >	PreOpenAckHistory;
 
 	bool							OpenAckedCalled;
@@ -404,7 +399,6 @@ private:
 		const int32					CmdIndex, 
 		const uint8 * RESTRICT		StoredData, 
 		const uint8 * RESTRICT		Data, 
-		TArray< uint16 > &			Unmapped,
 		uint16						Handle ) const;
 
 	uint16 SendProperties_r( 
@@ -415,7 +409,6 @@ private:
 		const int32					CmdEnd, 
 		const uint8 * RESTRICT		StoredData, 
 		const uint8 * RESTRICT		Data, 
-		TArray< uint16 > &			Unmapped,
 		uint16						Handle ) const;
 
 	bool ReadProperty( 
