#ue4 - Defer blocking hit notifications that are within a deferred FScopedMovementUpdate.

- Added UPrimitiveComponent::DispatchBlockingHit() to avoid duplicate code when dispatching hits.

UDN: https://udn.unrealengine.com/questions/206046/fscopedmovementupdate.html

diff --git a/Engine/Source/Runtime/Engine/Classes/Components/PrimitiveComponent.h b/Engine/Source/Runtime/Engine/Classes/Components/PrimitiveComponent.h
--- a/Engine/Source/Runtime/Engine/Classes/Components/PrimitiveComponent.h
+++ b/Engine/Source/Runtime/Engine/Classes/Components/PrimitiveComponent.h
@@ -1176,6 +1176,14 @@ public:
 	virtual FVector GetComponentVelocity() const OVERRIDE;
 	//End USceneComponent Interface
 
+	/**
+	 * Dispatch notifications for the given HitResult.
+	 *
+	 * @param Owner: AActor that owns this component
+	 * @param BlockingHit: FHitResult that generated the blocking hit.
+	 */
+	void DispatchBlockingHit(AActor& Owner, FHitResult const& BlockingHit);
+
 	/**
 	 * Set collision params on OutParams (such as CollisionResponse, bTraceAsyncScene) to match the settings on this PrimitiveComponent.
 	 */
diff --git a/Engine/Source/Runtime/Engine/Classes/Components/SceneComponent.h b/Engine/Source/Runtime/Engine/Classes/Components/SceneComponent.h
--- a/Engine/Source/Runtime/Engine/Classes/Components/SceneComponent.h
+++ b/Engine/Source/Runtime/Engine/Classes/Components/SceneComponent.h
@@ -735,6 +735,8 @@ class ENGINE_API FScopedMovementUpdate : private FNoncopyable
 {
 public:
 	
+	typedef TArray<struct FHitResult, TInlineAllocator<2>> TBlockingHitArray;
+
 	FScopedMovementUpdate( class USceneComponent* Component, EScopedUpdate::Type ScopeBehavior = EScopedUpdate::DeferredUpdates );
 	~FScopedMovementUpdate();
 
@@ -759,6 +761,12 @@ public:
 	/** Returns the list of overlaps at the end location, or null if the list is invalid. */
 	const TArray<struct FOverlapInfo>* GetOverlapsAtEnd() const;
 
+	/** Add blocking hit that will get processed once the move is committed. This is intended for use only by SceneComponent and its derived classes. */
+	void AppendBlockingHit(const FHitResult& Hit);
+	
+	/** Returns the list of pending blocking hits, which will be used for notifications once the move is committed. */
+	const TBlockingHitArray& GetPendingBlockingHits() const;
+
 private:
 
 	/** Notify this scope that the given inner scope completed its update (ie is going out of scope). Only occurs for deferred updates. */
@@ -779,6 +787,7 @@ private:
 	FTransform InitialTransform;
 	TArray<struct FOverlapInfo> PendingOverlaps;
 	TArray<struct FOverlapInfo> OverlapsAtEnd;
+	TBlockingHitArray BlockingHits;
 
 	friend class USceneComponent;
 };
@@ -805,3 +814,13 @@ FORCEINLINE const TArray<struct FOverlapInfo>* FScopedMovementUpdate::GetOverlap
 {
 	return bHasValidOverlapsAtEnd ? &OverlapsAtEnd : NULL;
 }
+
+FORCEINLINE const FScopedMovementUpdate::TBlockingHitArray& FScopedMovementUpdate::GetPendingBlockingHits() const
+{
+	return BlockingHits;
+}
+
+FORCEINLINE void FScopedMovementUpdate::AppendBlockingHit(const FHitResult& Hit)
+{
+	BlockingHits.Add(Hit);
+}
diff --git a/Engine/Source/Runtime/Engine/Private/PrimitiveComponent.cpp b/Engine/Source/Runtime/Engine/Private/PrimitiveComponent.cpp
--- a/Engine/Source/Runtime/Engine/Private/PrimitiveComponent.cpp
+++ b/Engine/Source/Runtime/Engine/Private/PrimitiveComponent.cpp
@@ -1464,15 +1464,15 @@ bool UPrimitiveComponent::MoveComponent( const FVector& Delta, const FRotator& N
 		// Handle hit notifications.
 		if (BlockingHit.bBlockingHit)
 		{
-			if (Actor != NULL)
+			if (IsDeferringMovementUpdates())
 			{
-				Actor->DispatchBlockingHit(this, BlockingHit.Component.Get(), true, BlockingHit);
+				FScopedMovementUpdate* ScopedUpdate = GetCurrentScopedMovement();
+				checkSlow(ScopedUpdate != NULL);
+				ScopedUpdate->AppendBlockingHit(BlockingHit);
 			}
-
-			// BlockingHit.GetActor() could be marked for deletion in DispatchBlockingHit(), which would make the weak pointer return NULL.
-			if ( BlockingHit.GetActor() != NULL )
+			else
 			{
-				BlockingHit.GetActor()->DispatchBlockingHit(BlockingHit.Component.Get(), this, false, BlockingHit);
+				DispatchBlockingHit(*Actor, BlockingHit);
 			}
 		}
 	}
@@ -1484,7 +1484,7 @@ bool UPrimitiveComponent::MoveComponent( const FVector& Delta, const FRotator& N
 		{
 			// Defer UpdateOverlaps until the scoped move ends.
 			FScopedMovementUpdate* ScopedUpdate = GetCurrentScopedMovement();
-			check(ScopedUpdate != NULL);
+			checkSlow(ScopedUpdate != NULL);
 			ScopedUpdate->AppendOverlaps(PendingOverlaps, OverlapsAtEndLocationPtr);
 		}
 		else
@@ -1528,6 +1528,19 @@ bool UPrimitiveComponent::MoveComponent( const FVector& Delta, const FRotator& N
 	return bMoved;
 }
 
+
+void UPrimitiveComponent::DispatchBlockingHit(AActor& Owner, FHitResult const& BlockingHit)
+{
+	Owner.DispatchBlockingHit(this, BlockingHit.Component.Get(), true, BlockingHit);
+
+	// BlockingHit.GetActor() could be marked for deletion in DispatchBlockingHit(), which would make the weak pointer return NULL.
+	if ( BlockingHit.GetActor() != NULL )
+	{
+		BlockingHit.GetActor()->DispatchBlockingHit(BlockingHit.Component.Get(), this, false, BlockingHit);
+	}
+}
+
+
 bool UPrimitiveComponent::IsNavigationRelevant(bool bSkipCollisionEnabledCheck) const 
 { 
 	if (HasCustomNavigableGeometry() >= EHasCustomNavigableGeometry::EvenIfNotCollidable)
diff --git a/Engine/Source/Runtime/Engine/Private/SceneComponent.cpp b/Engine/Source/Runtime/Engine/Private/SceneComponent.cpp
--- a/Engine/Source/Runtime/Engine/Private/SceneComponent.cpp
+++ b/Engine/Source/Runtime/Engine/Private/SceneComponent.cpp
@@ -190,6 +190,21 @@ void USceneComponent::EndScopedMovementUpdate(class FScopedMovementUpdate& Compl
 				PropagateTransformUpdate(true);
 			}
 
+			// Dispatch all deferred blocking hits
+			if (CompletedScope.BlockingHits.Num() > 0)
+			{
+				AActor* const Owner = GetOwner();
+				if (Owner)
+				{
+					// If we have blocking hits, we must be a primitive component.
+					UPrimitiveComponent* PrimitiveThis = CastChecked<UPrimitiveComponent>(this);
+					for (const FHitResult& Hit : CompletedScope.BlockingHits)
+					{
+						PrimitiveThis->DispatchBlockingHit(*Owner, Hit);
+					}
+				}				
+			}
+
 			// We may have moved somewhere and then moved back to the start, we still need to update overlaps if we touched things along the way.
 			if (bMoved || CurrentScopedUpdate->HasPendingOverlaps() || CurrentScopedUpdate->GetOverlapsAtEnd())
 			{
@@ -1584,6 +1599,7 @@ void FScopedMovementUpdate::RevertMove()
 	{
 		bHasValidOverlapsAtEnd = false;
 		PendingOverlaps.Reset();
+		BlockingHits.Reset();
 		
 		if (IsTransformDirty())
 		{
@@ -1624,6 +1640,7 @@ void FScopedMovementUpdate::OnInnerScopeComplete(const FScopedMovementUpdate& In
 
 		// Combine with the next item on the stack
 		AppendOverlaps(InnerScope.GetPendingOverlaps(), InnerScope.GetOverlapsAtEnd());
+		BlockingHits.Append(InnerScope.GetPendingBlockingHits());
 	}	
 }
 
