Landscape : Added CollisionMipLevel to be propagated to sublevel

diff --git a/Engine/Source/Runtime/Engine/Private/Landscape/Landscape.cpp b/Engine/Source/Runtime/Engine/Private/Landscape/Landscape.cpp
--- a/Engine/Source/Runtime/Engine/Private/Landscape/Landscape.cpp
+++ b/Engine/Source/Runtime/Engine/Private/Landscape/Landscape.cpp
@@ -1257,6 +1257,7 @@ void ALandscapeProxy::GetSharedProperties(ALandscapeProxy* Landscape)
 		SubsectionSizeQuads = Landscape->SubsectionSizeQuads;
 		MaxLODLevel = Landscape->MaxLODLevel;
 		LODDistanceFactor = Landscape->LODDistanceFactor;
+		CollisionMipLevel = Landscape->CollisionMipLevel;
 		if (!LandscapeMaterial)
 		{
 			LandscapeMaterial = Landscape->LandscapeMaterial;
