Move timeline event track execution after all properties are set
https://github.com/EpicGames/UnrealEngine/pull/156 (tsky1971)

diff --git a/Engine/Source/Runtime/Engine/Private/Timeline.cpp b/Engine/Source/Runtime/Engine/Private/Timeline.cpp
--- a/Engine/Source/Runtime/Engine/Private/Timeline.cpp
+++ b/Engine/Source/Runtime/Engine/Private/Timeline.cpp
@@ -131,64 +131,6 @@ void FTimeline::SetPlaybackPosition(float NewPosition, bool bFireEvents)
 	float OldPosition = Position;
 	Position = NewPosition;
 
-	// If playing sequence forwards.
-	float MinTime, MaxTime;
-	if(!bReversePlayback)
-	{
-		MinTime = OldPosition;
-		MaxTime = Position;
-
-		// Slight hack here.. if playing forwards and reaching the end of the sequence, force it over a little to ensure we fire events actually on the end of the sequence.
-		if(MaxTime == GetTimelineLength())
-		{
-			MaxTime += (float)KINDA_SMALL_NUMBER;
-		}
-	}
-	// If playing sequence backwards.
-	else
-	{
-		MinTime = Position;
-		MaxTime = OldPosition;
-
-		// Same small hack as above for backwards case.
-		if(MinTime == 0.f)
-		{
-			MinTime -= (float)KINDA_SMALL_NUMBER;
-		}
-	}
-
-	// If we should be firing events for this track...
-	if(bFireEvents)
-	{
-		// See which events fall into traversed region.
-		for(int32 i=0; i<Events.Num(); i++)
-		{
-			float EventTime = Events[i].Time;
-
-			// Need to be slightly careful here and make behavior for firing events symmetric when playing forwards of backwards.
-			bool bFireThisEvent = false;
-			if(!bReversePlayback)
-			{
-				if( EventTime >= MinTime && EventTime < MaxTime )
-				{
-					bFireThisEvent = true;
-				}
-			}
-			else
-			{
-				if( EventTime > MinTime && EventTime <= MaxTime )
-				{
-					bFireThisEvent = true;
-				}
-			}
-
-			if( bFireThisEvent )
-			{
-				Events[i].EventFunc.ExecuteIfBound();
-			}
-		}
-	}
-
 	UObject* PropSetObject = PropertySetObject.Get();
 
 	// Iterate over each vector interpolation
@@ -314,6 +256,65 @@ void FTimeline::SetPlaybackPosition(float NewPosition, bool bFireEvents)
 		}
 	}
 
+
+	// If we should be firing events for this track...
+	if (bFireEvents)
+	{
+		// If playing sequence forwards.
+		float MinTime, MaxTime;
+		if (!bReversePlayback)
+		{
+			MinTime = OldPosition;
+			MaxTime = Position;
+
+			// Slight hack here.. if playing forwards and reaching the end of the sequence, force it over a little to ensure we fire events actually on the end of the sequence.
+			if (MaxTime == GetTimelineLength())
+			{
+				MaxTime += (float)KINDA_SMALL_NUMBER;
+			}
+		}
+		// If playing sequence backwards.
+		else
+		{
+			MinTime = Position;
+			MaxTime = OldPosition;
+
+			// Same small hack as above for backwards case.
+			if (MinTime == 0.f)
+			{
+				MinTime -= (float)KINDA_SMALL_NUMBER;
+			}
+		}
+
+		// See which events fall into traversed region.
+		for (int32 i = 0; i < Events.Num(); i++)
+		{
+			float EventTime = Events[i].Time;
+
+			// Need to be slightly careful here and make behavior for firing events symmetric when playing forwards of backwards.
+			bool bFireThisEvent = false;
+			if (!bReversePlayback)
+			{
+				if (EventTime >= MinTime && EventTime < MaxTime)
+				{
+					bFireThisEvent = true;
+				}
+			}
+			else
+			{
+				if (EventTime > MinTime && EventTime <= MaxTime)
+				{
+					bFireThisEvent = true;
+				}
+			}
+
+			if (bFireThisEvent)
+			{
+				Events[i].EventFunc.ExecuteIfBound();
+			}
+		}
+	}
+
 	// Execute the delegate to say that all properties are updated
 	TimelinePostUpdateFunc.ExecuteIfBound();
 }
