UCanvas Blueprint Extension
https://github.com/EpicGames/UnrealEngine/pull/108 (DiCoJamesTan)

diff --git a/Engine/Source/Runtime/Engine/Classes/Engine/Canvas.h b/Engine/Source/Runtime/Engine/Classes/Engine/Canvas.h
--- a/Engine/Source/Runtime/Engine/Classes/Engine/Canvas.h
+++ b/Engine/Source/Runtime/Engine/Classes/Engine/Canvas.h
@@ -578,7 +578,157 @@ public:
 	 * @param Y				Y Position to draw item
 	 */
 	void DrawItem( class FCanvasItem& Item, float X, float Y );
-};
 
+	/**
+	 * Draws a line on the Canvas.
+	 *
+	 * @param ScreenPositionA		Starting position of the line in screen space.
+	 * @param ScreenPositionB		Ending position of the line in screen space.
+	 * @param Thickness				How many pixels thick this line should be.
+	 * @param RenderColor			Color to render the line.
+	 */
+	UFUNCTION(BlueprintCallable, Category=Canvas, meta=(FriendlyName="Draw Line"))
+	void K2_DrawLine(FVector2D ScreenPositionA=FVector2D::ZeroVector, FVector2D ScreenPositionB=FVector2D::ZeroVector, float Thickness=1.0f, FLinearColor RenderColor=FLinearColor::White);
+
+	/**
+	 * Draws a texture on the Canvas.
+	 *
+	 * @param RenderTexture				Texture to use when rendering. If no texture is set then this will use the default white texture.
+	 * @param ScreenPosition			Screen space position to render the texture.
+	 * @param ScreenSize				Screen space size to render the texture.
+	 * @param CoordinatePosition		Normalized UV starting coordinate to use when rendering the texture.
+	 * @param CoordinateSize			Normalized UV size coordinate to use when rendering the texture.
+	 * @param RenderColor				Color to use when rendering the texture.
+	 * @param BlendMode					Blending mode to use when rendering the texture.
+	 * @param Rotation					Rotation, in degrees, to render the texture.
+	 * @param PivotPoint				Normalized pivot point to use when rotating the texture.
+	 */
+	UFUNCTION(BlueprintCallable, Category=Canvas, meta=(FriendlyName="Draw Texture"))
+	void K2_DrawTexture(UTexture* RenderTexture, FVector2D ScreenPosition, FVector2D ScreenSize, FVector2D CoordinatePosition, FVector2D CoordinateSize=FVector2D::UnitVector, FLinearColor RenderColor=FLinearColor::White, EBlendMode BlendMode=BLEND_Translucent, float Rotation=0.f, FVector2D PivotPoint=FVector2D(0.5f,0.5f));
 
+	/**
+	 * Draws a material on the Canvas.
+	 *
+	 * @param RenderMaterial			Material to use when rendering. Remember that only the emissive channel is able to be rendered as no lighting is performed when rendering to the Canvas.
+	 * @param ScreenPosition			Screen space position to render the texture.
+	 * @param ScreenSize				Screen space size to render the texture.
+	 * @param CoordinatePosition		Normalized UV starting coordinate to use when rendering the texture.
+	 * @param CoordinateSize			Normalized UV size coordinate to use when rendering the texture.
+	 * @param Rotation					Rotation, in degrees, to render the texture.
+	 * @param PivotPoint				Normalized pivot point to use when rotating the texture.
+	 */
+	UFUNCTION(BlueprintCallable, Category=Canvas, meta=(FriendlyName="Draw Material"))
+	void K2_DrawMaterial(UMaterialInterface* RenderMaterial, FVector2D ScreenPosition, FVector2D ScreenSize, FVector2D CoordinatePosition, FVector2D CoordinateSize=FVector2D::UnitVector, float Rotation=0.f, FVector2D PivotPoint=FVector2D(0.5f,0.5f));
 
+	/**
+	 * Draws text on the Canvas.
+	 *
+	 * @param RenderFont				Font to use when rendering the text. If this is null, then a default engine font is used.
+	 * @param RenderText				Text to render on the Canvas.
+	 * @param ScreenPosition			Screen space position to render the text.
+	 * @param RenderColor				Color to render the text.
+	 * @param Kerning					Horizontal spacing adjustment to modify the spacing between each letter.
+	 * @param ShadowColor				Color to render the shadow of the text.
+	 * @param ShadowOffset				Pixel offset relative to the screen space position to render the shadow of the text.
+	 * @param bCentreX					If true, then interpret the screen space position X coordinate as the center of the rendered text.
+	 * @param bCentreY					If true, then interpret the screen space position Y coordinate as the center of the rendered text.
+	 * @param bOutlined					If true, then the text should be rendered with an outline.
+	 * @param OutlineColor				Color to render the outline for the text.
+	 */	
+	UFUNCTION(BlueprintCallable, Category=Canvas, meta=(FriendlyName="Draw Text"))
+	void K2_DrawText(UFont* RenderFont, const FString& RenderText, FVector2D ScreenPosition, FLinearColor RenderColor=FLinearColor::White, float Kerning=0.0f, FLinearColor ShadowColor=FLinearColor::Black, FVector2D ShadowOffset=FVector2D::UnitVector, bool bCentreX=false, bool bCentreY=false, bool bOutlined=false, FLinearColor OutlineColor=FLinearColor::Black);
+	
+	/**
+	 * Draws a 3x3 grid border with tiled frame and tiled interior on the Canvas.
+	 *
+	 * @param BorderTexture				Texture to use for border.
+	 * @param BackgroundTexture			Texture to use for border background.
+	 * @param LeftBorderTexture			Texture to use for the tiling left border.
+	 * @param RightBorderTexture		Texture to use for the tiling right border.
+	 * @param TopBorderTexture			Texture to use for the tiling top border.
+	 * @param BottomBorderTexture		Texture to use for the tiling bottom border.
+	 * @param ScreenPosition			Screen space position to render the texture.
+	 * @param ScreenSize				Screen space size to render the texture.
+	 * @param CoordinatePosition		Normalized UV starting coordinate to use when rendering the border texture.
+	 * @param CoordinateSize			Normalized UV size coordinate to use when rendering the border texture.
+	 * @param RenderColor				Color to tint the border.
+	 * @param BorderScale				Scale of the border.
+	 * @param BackgroundScale			Scale of the background.
+	 * @param Rotation					Rotation, in degrees, to render the texture.
+	 * @param PivotPoint				Normalized pivot point to use when rotating the texture.
+	 * @param CornerSize				Frame corner size in percent of frame texture (should be < 0.5f).
+	 */
+	UFUNCTION(BlueprintCallable, Category=Canvas, meta=(FriendlyName="Draw Border"))
+	void K2_DrawBorder(UTexture* BorderTexture, UTexture* BackgroundTexture, UTexture* LeftBorderTexture, UTexture* RightBorderTexture, UTexture* TopBorderTexture, UTexture* BottomBorderTexture, FVector2D ScreenPosition, FVector2D ScreenSize, FVector2D CoordinatePosition, FVector2D CoordinateSize=FVector2D::UnitVector, FLinearColor RenderColor=FLinearColor::White, FVector2D BorderScale=FVector2D(0.1f,0.1f), FVector2D BackgroundScale=FVector2D(0.1f,0.1f), float Rotation=0.0f, FVector2D PivotPoint=FVector2D(0.5f,0.5f), FVector2D CornerSize=FVector2D::ZeroVector);
+
+	/**
+	 * Draws an unfilled box on the Canvas.
+	 *
+	 * @param ScreenPosition			Screen space position to render the text.
+	 * @param ScreenSize				Screen space size to render the texture.
+	 * @param Thickness					How many pixels thick the box lines should be.
+	 */
+	UFUNCTION(BlueprintCallable, Category=Canvas, meta=(FriendlyName="Draw Box"))
+	void K2_DrawBox(FVector2D ScreenPosition, FVector2D ScreenSize, float Thickness=1.0f);
+
+	/**
+	 * Draws a set of triangles on the Canvas.
+	 *
+	 * @param RenderTexture				Texture to use when rendering the triangles. If no texture is set, then the default white texture is used.
+	 * @param Triangles					Triangles to render.
+	 */
+	UFUNCTION(BlueprintCallable, Category=Canvas, meta=(FriendlyName="Draw Triangles"))
+	void K2_DrawTriangle(UTexture* RenderTexture, TArray<FCanvasUVTri> Triangles);
+
+	/**
+	 * Draws a polygon on the Canvas.
+	 *
+	 * @param RenderTexture				Texture to use when rendering the triangles. If no texture is set, then the default white texture is used.
+	 * @param ScreenPosition			Screen space position to render the text.
+	 * @param Radius					How large in pixels this polygon should be.
+	 * @param NumberOfSides				How many sides this polygon should have. This should be above or equal to three.
+	 * @param RenderColor				Color to tint the polygon.
+	 */
+	UFUNCTION(BlueprintCallable, Category=Canvas, meta=(FriendlyName="Draw Polygon"))
+	void K2_DrawPolygon(UTexture* RenderTexture, FVector2D ScreenPosition, FVector2D Radius=FVector2D::UnitVector, int32 NumberOfSides=3, FLinearColor RenderColor=FLinearColor::White);
+
+	/**
+	 * Performs a projection of a world space coordinates using the projection matrix set up for the Canvas.
+	 *
+	 * @param WorldLocation				World space location to project onto the Canvas rendering plane.
+	 * @return							Returns a vector where X, Y defines a screen space position representing the world space location.
+	 */
+	UFUNCTION(BlueprintCallable, Category=Canvas, meta=(FriendlyName="Project"))
+	FVector K2_Project(FVector WorldLocation);
+
+	/**
+	 * Performs a deprojection of a screen space coordinate using the projection matrix set up for the Canvas.
+	 *
+	 * @param ScreenPosition			Screen space position to deproject to the World.
+	 * @param WorldOrigin				Vector which is the world position of the screen space position.
+	 * @param WorldDirection			Vector which can be used in a trace to determine what is "behind" the screen space position. Useful for object picking.
+	 */
+	UFUNCTION(BlueprintCallable, Category=Canvas, meta=(FriendlyName="Deproject"))
+	void K2_Deproject(FVector2D ScreenPosition, FVector& WorldOrigin, FVector& WorldDirection);
+
+	/**
+	 * Returns the wrapped text size in screen space coordinates.
+	 *
+	 * @param RenderFont				Font to use when determining the size of the text. If this is null, then a default engine font is used.
+	 * @param RenderText				Text to determine the size of.
+	 * @return							Returns the screen space size of the text.
+	 */
+	UFUNCTION(BlueprintCallable, Category=Canvas, meta=(FriendlyName="Wrapped Text Size"))
+	FVector2D K2_StrLen(UFont* RenderFont, const FString& RenderText);
+
+	/**
+	 * Returns the clipped text size in screen space coordinates.
+	 *
+	 * @param RenderFont				Font to use when determining the size of the text. If this is null, then a default engine font is used.
+	 * @param RenderText				Text to determine the size of.
+	 * @param Scale						Scale of the font to use when determining the size of the text.
+	 * @return							Returns the screen space size of the text.
+	 */
+	UFUNCTION(BlueprintCallable, Category=Canvas, meta=(FriendlyName="Clipped Text Size"))
+	FVector2D K2_TextSize(UFont* RenderFont, const FString& RenderText, FVector2D Scale=FVector2D::UnitVector);
+};
\ No newline at end of file
diff --git a/Engine/Source/Runtime/Engine/Private/UserInterface/Canvas.cpp b/Engine/Source/Runtime/Engine/Private/UserInterface/Canvas.cpp
--- a/Engine/Source/Runtime/Engine/Private/UserInterface/Canvas.cpp
+++ b/Engine/Source/Runtime/Engine/Private/UserInterface/Canvas.cpp
@@ -1241,18 +1241,18 @@ void VARARGS UCanvas::WrappedStrLenf( UFont* Font, float ScaleX, float ScaleY, i
 float UCanvas::DrawText(UFont* InFont, const FText& InText, float X, float Y, float XScale, float YScale, const FFontRenderInfo& RenderInfo)
 {
 	ensure(InFont);
-	int32		XL = 0;
-	int32		YL = 0;
+	int32		XL		= 0;
+	int32		YL		= 0; 
 	// need this call in any case to update YL and XL - one of them will be needed anyway
 	WrappedPrint(RenderInfo.bClipText == false, X, Y, XL, YL, InFont, XScale, YScale, bCenterX, bCenterY, *InText.ToString(), RenderInfo);
 
 	if (RenderInfo.bClipText)
 	{
 		FCanvasTextItem TextItem(FVector2D(FMath::TruncToFloat(OrgX + X), FMath::TruncToFloat(OrgY + Y)), InText, InFont, DrawColor);
-		TextItem.Scale = FVector2D(XScale, YScale),
-			TextItem.BlendMode = SE_BLEND_Translucent;
+		TextItem.Scale = FVector2D( XScale, YScale ), 
+		TextItem.BlendMode = SE_BLEND_Translucent;
 		TextItem.FontRenderInfo = RenderInfo;
-		Canvas->DrawItem(TextItem);
+		Canvas->DrawItem( TextItem );	
 	}
 
 	return (float)YL;
@@ -1616,3 +1616,133 @@ void UCanvas::SetView(FSceneView* InView)
 		//ViewRotation = FRotator(0,0,0);
 	}
 }
+
+void UCanvas::K2_DrawLine(FVector2D ScreenPositionA, FVector2D ScreenPositionB, float Thickness, FLinearColor RenderColor)
+{
+	if (FMath::Square(ScreenPositionB.X - ScreenPositionA.X) + FMath::Square(ScreenPositionB.Y - ScreenPositionA.Y))
+	{
+		FCanvasLineItem LineItem(ScreenPositionA, ScreenPositionB);
+		LineItem.LineThickness = Thickness;
+		LineItem.SetColor(RenderColor);
+		DrawItem(LineItem);
+	}
+}
+
+void UCanvas::K2_DrawTexture(UTexture* RenderTexture, FVector2D ScreenPosition, FVector2D ScreenSize, FVector2D CoordinatePosition, FVector2D CoordinateSize, FLinearColor RenderColor, EBlendMode BlendMode, float Rotation, FVector2D PivotPoint)
+{
+	if (ScreenSize.X > 0.0f && ScreenSize.Y > 0.0f)
+	{
+		FTexture* RenderTextureResource = (RenderTexture) ? RenderTexture->Resource : GWhiteTexture;
+		FCanvasTileItem TileItem(ScreenPosition, RenderTextureResource, ScreenSize, CoordinatePosition, CoordinatePosition + CoordinateSize, RenderColor);
+		TileItem.Rotation = FRotator(0, Rotation, 0);
+		TileItem.PivotPoint = PivotPoint;
+		TileItem.BlendMode = FCanvas::BlendToSimpleElementBlend(BlendMode);
+		DrawItem(TileItem);
+	}
+}
+
+void UCanvas::K2_DrawMaterial(UMaterialInterface* RenderMaterial, FVector2D ScreenPosition, FVector2D ScreenSize, FVector2D CoordinatePosition, FVector2D CoordinateSize, float Rotation, FVector2D PivotPoint)
+{
+	if (RenderMaterial && ScreenSize.X > 0.0f && ScreenSize.Y > 0.0f)
+	{
+		FCanvasTileItem TileItem(ScreenPosition, RenderMaterial->GetRenderProxy(0), ScreenSize, CoordinatePosition, CoordinatePosition + CoordinateSize);
+		TileItem.Rotation = FRotator(0, Rotation, 0);
+		TileItem.PivotPoint = PivotPoint;
+		DrawItem(TileItem);
+	}
+}
+
+void UCanvas::K2_DrawText(UFont* RenderFont, const FString& RenderText, FVector2D ScreenPosition, FLinearColor RenderColor, float Kerning, FLinearColor ShadowColor, FVector2D ShadowOffset, bool bCentreX, bool bCentreY, bool bOutlined, FLinearColor OutlineColor)
+{
+	if (!RenderText.IsEmpty())
+	{
+		FCanvasTextItem TextItem(ScreenPosition, FText::FromString(RenderText), RenderFont, RenderColor);
+		TextItem.HorizSpacingAdjust = Kerning;
+		TextItem.ShadowColor = ShadowColor;
+		TextItem.ShadowOffset = ShadowOffset;
+		TextItem.bCentreX = bCentreX;
+		TextItem.bCentreY = bCentreY;
+		TextItem.bOutlined = bOutlined;
+		TextItem.OutlineColor = OutlineColor;
+		DrawItem(TextItem);
+	}
+}
+
+void UCanvas::K2_DrawBorder(UTexture* BorderTexture, UTexture* BackgroundTexture, UTexture* LeftBorderTexture, UTexture* RightBorderTexture, UTexture* TopBorderTexture, UTexture* BottomBorderTexture, FVector2D ScreenPosition, FVector2D ScreenSize, FVector2D CoordinatePosition, FVector2D CoordinateSize, FLinearColor RenderColor, FVector2D BorderScale, FVector2D BackgroundScale, float Rotation, FVector2D PivotPoint, FVector2D CornerSize)
+{
+	if (ScreenSize.X > 0.0f && ScreenSize.Y > 0.0f && BorderTexture && BackgroundTexture && LeftBorderTexture && RightBorderTexture && TopBorderTexture && BottomBorderTexture)
+	{
+		FCanvasBorderItem BorderItem(ScreenPosition, BorderTexture->Resource, BackgroundTexture->Resource, LeftBorderTexture->Resource, RightBorderTexture->Resource, TopBorderTexture->Resource, BottomBorderTexture->Resource, ScreenSize, RenderColor);
+		BorderItem.BorderScale = BorderScale;
+		BorderItem.BackgroundScale = BackgroundScale;
+		BorderItem.BorderUV0 = CoordinatePosition;
+		BorderItem.BorderUV1 = CoordinatePosition + CoordinateSize;
+		BorderItem.Rotation = FRotator(0, Rotation, 0);
+		BorderItem.PivotPoint = PivotPoint;
+		BorderItem.CornerSize = CornerSize;
+		DrawItem(BorderItem);
+	}
+}
+
+void UCanvas::K2_DrawBox(FVector2D ScreenPosition, FVector2D ScreenSize, float Thickness)
+{
+	if (ScreenSize.X > 0.0f && ScreenSize.Y > 0.0f)
+	{
+		FCanvasBoxItem BoxItem(ScreenPosition, ScreenSize);
+		BoxItem.LineThickness = Thickness;
+		DrawItem(BoxItem);
+	}
+}
+
+void UCanvas::K2_DrawTriangle(UTexture* RenderTexture, TArray<FCanvasUVTri> Triangles)
+{
+	if (Triangles.Num() > 0)
+	{
+		FCanvasTriangleItem TriangleItem(FVector2D::ZeroVector, FVector2D::ZeroVector, FVector2D::ZeroVector, (RenderTexture) ? RenderTexture->Resource : GWhiteTexture);
+		TriangleItem.TriangleList = Triangles;
+		DrawItem(TriangleItem);
+	}
+}
+
+void UCanvas::K2_DrawPolygon(UTexture* RenderTexture, FVector2D ScreenPosition, FVector2D Radius, int32 NumberOfSides, FLinearColor RenderColor)
+{
+	if (Radius.X > 0.0f && Radius.Y > 0.0f && NumberOfSides >= 3)
+	{
+		FCanvasNGonItem NGonItem(ScreenPosition, Radius, NumberOfSides, (RenderTexture) ? RenderTexture->Resource : GWhiteTexture, RenderColor);
+		DrawItem(NGonItem);
+	}
+}
+
+FVector UCanvas::K2_Project(FVector WorldLocation)
+{
+	return Project(WorldLocation);
+}
+
+void UCanvas::K2_Deproject(FVector2D ScreenPosition, FVector& WorldOrigin, FVector& WorldDirection)
+{
+	Deproject(ScreenPosition, WorldOrigin, WorldDirection);
+}
+
+FVector2D UCanvas::K2_StrLen(UFont* RenderFont, const FString& RenderText)
+{
+	if (!RenderText.IsEmpty())
+	{
+		FVector2D OutTextSize;
+		StrLen(RenderFont, RenderText, OutTextSize.X, OutTextSize.Y);
+		return OutTextSize;
+	}
+
+	return FVector2D::ZeroVector;
+}
+
+FVector2D UCanvas::K2_TextSize(UFont* RenderFont, const FString& RenderText, FVector2D Scale)
+{
+	if (!RenderText.IsEmpty())
+	{
+		FVector2D OutTextSize;
+		TextSize(RenderFont, RenderText, OutTextSize.X, OutTextSize.Y, Scale.X, Scale.Y);
+		return OutTextSize;
+	}
+	
+	return FVector2D::ZeroVector;
+}
\ No newline at end of file
