Added user-editable array of object query types to Radial Force components.

Users can now specify the object types that radial forces affect.
Radial forces now default to affecting Pawns, Physics, Vehicles & Destructables.

TTP# 335072 - CRITICAL: Regression: CONTENTEXAMPLES: PHYSICS: force applied by joint actors on test 1.5 knocks pawn back

diff --git a/Engine/Source/Runtime/Engine/Classes/Engine/CollisionProfile.h b/Engine/Source/Runtime/Engine/Classes/Engine/CollisionProfile.h
--- a/Engine/Source/Runtime/Engine/Classes/Engine/CollisionProfile.h
+++ b/Engine/Source/Runtime/Engine/Classes/Engine/CollisionProfile.h
@@ -186,6 +186,18 @@ public:
 	/** Convert ObjectType or TraceType to CollisionChannel */
 	ECollisionChannel ConvertToCollisionChannel(bool TraceType, int32 Index) const;
 
+	/** 
+	 * Convert collision channel to ObjectTypeQuery. Note: performs a search of object types.
+	 * @return ObjectTypeQuery_MAX if the conversion was not possible 
+	 */
+	EObjectTypeQuery ConvertToObjectType(ECollisionChannel CollisionChannel) const;
+
+	/** 
+	 * Convert collision channel to TraceTypeQuery. Note: performs a search of object types.
+	 * @return TraceTypeQuery_MAX if the conversion was not possible 
+	 */
+	ETraceTypeQuery ConvertToTraceType(ECollisionChannel CollisionChannel) const;
+
 	/* custom collision profile name that you can modify what you'd like */
 	ENGINE_API static FName CustomCollisionProfileName;
 
diff --git a/Engine/Source/Runtime/Engine/Classes/Engine/EngineTypes.h b/Engine/Source/Runtime/Engine/Classes/Engine/EngineTypes.h
--- a/Engine/Source/Runtime/Engine/Classes/Engine/EngineTypes.h
+++ b/Engine/Source/Runtime/Engine/Classes/Engine/EngineTypes.h
@@ -308,7 +308,9 @@ enum EObjectTypeQuery
 	ObjectTypeQuery29 UMETA(Hidden), 
 	ObjectTypeQuery30 UMETA(Hidden), 
 	ObjectTypeQuery31 UMETA(Hidden), 
-	ObjectTypeQuery32 UMETA(Hidden)
+	ObjectTypeQuery32 UMETA(Hidden),
+
+	ObjectTypeQuery_MAX	UMETA(Hidden)
 };
 
 UENUM(BlueprintType)
@@ -345,7 +347,9 @@ enum ETraceTypeQuery
 	TraceTypeQuery29 UMETA(Hidden), 
 	TraceTypeQuery30 UMETA(Hidden), 
 	TraceTypeQuery31 UMETA(Hidden), 
-	TraceTypeQuery32 UMETA(Hidden)
+	TraceTypeQuery32 UMETA(Hidden),
+
+	TraceTypeQuery_MAX	UMETA(Hidden)
 };
 
 /** Enum indicating which physics scene to use. */
@@ -2349,8 +2353,17 @@ class UEngineTypes : public UObject
 	GENERATED_UCLASS_BODY()
 
 public:
+	/** Convert a trace type to a collision channel */
 	static ECollisionChannel ConvertToCollisionChannel(ETraceTypeQuery TraceType);
+
+	/** Convert an object type to a collision channel */
 	static ECollisionChannel ConvertToCollisionChannel(EObjectTypeQuery ObjectType);
+
+	/** Convert a collision channel to an object type. Note: performs a search of object types */
+	static EObjectTypeQuery ConvertToObjectType(ECollisionChannel CollisionChannel);
+
+	/** Convert a collision channel to a trace type. Note: performs a search of trace types */
+	static ETraceTypeQuery ConvertToTraceType(ECollisionChannel CollisionChannel);
 };
 
 
diff --git a/Engine/Source/Runtime/Engine/Classes/PhysicsEngine/RadialForceComponent.h b/Engine/Source/Runtime/Engine/Classes/PhysicsEngine/RadialForceComponent.h
--- a/Engine/Source/Runtime/Engine/Classes/PhysicsEngine/RadialForceComponent.h
+++ b/Engine/Source/Runtime/Engine/Classes/PhysicsEngine/RadialForceComponent.h
@@ -44,6 +44,25 @@ class URadialForceComponent : public USceneComponent
 	UFUNCTION(BlueprintCallable, Category="Physics|Components|RadialForce")
 	virtual void FireImpulse();
 
+	/** Add an object type for this radial force to affect */
+	UFUNCTION(BlueprintCallable, Category="Physics|Components|RadialForce")
+	virtual void AddObjectTypeToAffect(TEnumAsByte<enum EObjectTypeQuery> ObjectType);
+
+	/** Remove an object type that is affected by this radial force */
+	UFUNCTION(BlueprintCallable, Category="Physics|Components|RadialForce")
+	virtual void RemoveObjectTypeToAffect(TEnumAsByte<enum EObjectTypeQuery> ObjectType);
+
+	/** Add a collision channel for this radial force to affect */
+	void AddCollisionChannelToAffect(enum ECollisionChannel CollisionChannel);
+
+protected:
+	/** The object types that are affected by this radial force */
+	UPROPERTY(EditAnywhere, Category=RadialForceComponent)
+	TArray<TEnumAsByte<enum EObjectTypeQuery> > ObjectTypesToAffect;
+
+	/** Cached object query params derived from ObjectTypesToAffect */
+	FCollisionObjectQueryParams CollisionObjectQueryParams;
+
 protected:
 	// Begin UActorComponent interface.
 	virtual void TickComponent(float DeltaTime, enum ELevelTick TickType, FActorComponentTickFunction *ThisTickFunction) OVERRIDE;
@@ -51,7 +70,11 @@ protected:
 
 	// Begin UObject interface.
 	virtual void PostLoad() OVERRIDE;
+	virtual void PostEditChangeProperty(FPropertyChangedEvent& PropertyChangedEvent) OVERRIDE;
 	// End UObject interface.
+
+	/** Update CollisionObjectQueryParams from ObjectTypesToAffect */
+	void UpdateCollisionObjectQueryParams();
 };
 
 
diff --git a/Engine/Source/Runtime/Engine/Private/Collision/CollisionProfile.cpp b/Engine/Source/Runtime/Engine/Private/Collision/CollisionProfile.cpp
--- a/Engine/Source/Runtime/Engine/Private/Collision/CollisionProfile.cpp
+++ b/Engine/Source/Runtime/Engine/Private/Collision/CollisionProfile.cpp
@@ -646,3 +646,41 @@ ECollisionChannel UCollisionProfile::ConvertToCollisionChannel(bool TraceType, i
 	// invalid
 	return ECC_MAX;
 }
+
+EObjectTypeQuery UCollisionProfile::ConvertToObjectType(ECollisionChannel CollisionChannel) const
+{
+	if (CollisionChannel < ECC_MAX)
+	{
+		int32 ObjectTypeIndex = 0;
+		for(const auto& MappedCollisionChannel : ObjectTypeMapping)
+		{
+			if(MappedCollisionChannel == CollisionChannel)
+			{
+				return (EObjectTypeQuery)ObjectTypeIndex;
+			}
+
+			ObjectTypeIndex++;
+		}
+	}
+
+	return ObjectTypeQuery_MAX;
+}
+
+ETraceTypeQuery UCollisionProfile::ConvertToTraceType(ECollisionChannel CollisionChannel) const
+{
+	if (CollisionChannel < ECC_MAX)
+	{
+		int32 TraceTypeIndex = 0;
+		for(const auto& MappedCollisionChannel : TraceTypeMapping)
+		{
+			if(MappedCollisionChannel == CollisionChannel)
+			{
+				return (ETraceTypeQuery)TraceTypeIndex;
+			}
+
+			TraceTypeIndex++;
+		}
+	}
+
+	return TraceTypeQuery_MAX;
+}
\ No newline at end of file
diff --git a/Engine/Source/Runtime/Engine/Private/EngineTypes.cpp b/Engine/Source/Runtime/Engine/Private/EngineTypes.cpp
--- a/Engine/Source/Runtime/Engine/Private/EngineTypes.cpp
+++ b/Engine/Source/Runtime/Engine/Private/EngineTypes.cpp
@@ -22,6 +22,16 @@ ECollisionChannel UEngineTypes::ConvertToCollisionChannel(EObjectTypeQuery Objec
 	return UCollisionProfile::Get()->ConvertToCollisionChannel(false, (int32)ObjectType);
 }
 
+EObjectTypeQuery UEngineTypes::ConvertToObjectType(ECollisionChannel CollisionChannel)
+{
+	return UCollisionProfile::Get()->ConvertToObjectType(CollisionChannel);
+}
+
+ETraceTypeQuery UEngineTypes::ConvertToTraceType(ECollisionChannel CollisionChannel)
+{
+	return UCollisionProfile::Get()->ConvertToTraceType(CollisionChannel);
+}
+
 void FDamageEvent::GetBestHitInfo(AActor const* HitActor, AActor const* HitInstigator, FHitResult& OutHitInfo, FVector& OutImpulseDir) const
 {
 	ensure(HitActor);
diff --git a/Engine/Source/Runtime/Engine/Private/PhysicsEngine/RadialForceComponent.cpp b/Engine/Source/Runtime/Engine/Private/PhysicsEngine/RadialForceComponent.cpp
--- a/Engine/Source/Runtime/Engine/Private/PhysicsEngine/RadialForceComponent.cpp
+++ b/Engine/Source/Runtime/Engine/Private/PhysicsEngine/RadialForceComponent.cpp
@@ -15,6 +15,14 @@ URadialForceComponent::URadialForceComponent(const class FPostConstructInitializ
 	ImpulseStrength = 1000.0f;
 	ForceStrength = 10.0f;
 	bAutoActivate = true;
+
+	// by default we affect all 'dynamic' objects that can currently be affected by forces
+	AddCollisionChannelToAffect(ECC_Pawn);
+	AddCollisionChannelToAffect(ECC_PhysicsBody);
+	AddCollisionChannelToAffect(ECC_Vehicle);
+	AddCollisionChannelToAffect(ECC_Destructible);
+
+	UpdateCollisionObjectQueryParams();
 }
 
 void URadialForceComponent::TickComponent(float DeltaTime, enum ELevelTick TickType, FActorComponentTickFunction *ThisTickFunction)
@@ -31,7 +39,7 @@ void URadialForceComponent::TickComponent(float DeltaTime, enum ELevelTick TickT
 
 		FCollisionQueryParams Params(AddForceOverlapName, false);
 		Params.bTraceAsyncScene = true; // want to hurt stuff in async scene
-		GetWorld()->OverlapMulti(Overlaps, Origin, FQuat::Identity, FCollisionShape::MakeSphere(Radius), Params, FCollisionObjectQueryParams(FCollisionObjectQueryParams::InitType::AllDynamicObjects));
+		GetWorld()->OverlapMulti(Overlaps, Origin, FQuat::Identity, FCollisionShape::MakeSphere(Radius), Params, CollisionObjectQueryParams);
 
 		// Iterate over each and apply force
 		for ( int32 OverlapIdx=0; OverlapIdx<Overlaps.Num(); OverlapIdx++ )
@@ -65,6 +73,8 @@ void URadialForceComponent::PostLoad()
 	{
 		bAutoActivate = bForceEnabled_DEPRECATED;
 	}
+
+	UpdateCollisionObjectQueryParams();
 }
 
 void URadialForceComponent::FireImpulse()
@@ -77,7 +87,7 @@ void URadialForceComponent::FireImpulse()
 
 	FCollisionQueryParams Params(FireImpulseOverlapName, false);
 	Params.bTraceAsyncScene = true; // want to hurt stuff in async scene
-	GetWorld()->OverlapMulti(Overlaps, Origin, FQuat::Identity, FCollisionShape::MakeSphere(Radius), Params, FCollisionObjectQueryParams(FCollisionObjectQueryParams::InitType::AllObjects));
+	GetWorld()->OverlapMulti(Overlaps, Origin, FQuat::Identity, FCollisionShape::MakeSphere(Radius), Params, CollisionObjectQueryParams);
 
 	// Iterate over each and apply an impulse
 	for ( int32 OverlapIdx=0; OverlapIdx<Overlaps.Num(); OverlapIdx++ )
@@ -113,6 +123,43 @@ void URadialForceComponent::FireImpulse()
 	}
 }
 
+void URadialForceComponent::AddCollisionChannelToAffect(enum ECollisionChannel CollisionChannel)
+{
+	EObjectTypeQuery ObjectType = UEngineTypes::ConvertToObjectType(CollisionChannel);
+	if(ObjectType != ObjectTypeQuery_MAX)
+	{
+		AddObjectTypeToAffect(ObjectType);
+	}
+}
+
+void URadialForceComponent::AddObjectTypeToAffect(TEnumAsByte<enum EObjectTypeQuery> ObjectType)
+{
+	ObjectTypesToAffect.AddUnique(ObjectType);
+	UpdateCollisionObjectQueryParams();
+}
+
+void URadialForceComponent::RemoveObjectTypeToAffect(TEnumAsByte<enum EObjectTypeQuery> ObjectType)
+{
+	ObjectTypesToAffect.Remove(ObjectType);
+	UpdateCollisionObjectQueryParams();
+}
+
+void URadialForceComponent::PostEditChangeProperty(FPropertyChangedEvent& PropertyChangedEvent)
+{
+	Super::PostEditChangeProperty(PropertyChangedEvent);
+
+	// If we have edited the object types to effect, update our bitfield.
+	if(PropertyChangedEvent.Property && PropertyChangedEvent.Property->GetFName() == TEXT("ObjectTypesToAffect"))
+	{
+		UpdateCollisionObjectQueryParams();
+	}
+}
+
+void URadialForceComponent::UpdateCollisionObjectQueryParams()
+{
+	CollisionObjectQueryParams = FCollisionObjectQueryParams(ObjectTypesToAffect);
+}
+
 
 //////////////////////////////////////////////////////////////////////////
 // ARB_RADIALFORCEACTOR
