CIS fix for non-editor compilation of RadialForceComponent

diff --git a/Engine/Source/Runtime/Engine/Classes/PhysicsEngine/RadialForceComponent.h b/Engine/Source/Runtime/Engine/Classes/PhysicsEngine/RadialForceComponent.h
--- a/Engine/Source/Runtime/Engine/Classes/PhysicsEngine/RadialForceComponent.h
+++ b/Engine/Source/Runtime/Engine/Classes/PhysicsEngine/RadialForceComponent.h
@@ -70,7 +70,9 @@ protected:
 
 	// Begin UObject interface.
 	virtual void PostLoad() OVERRIDE;
+#if WITH_EDITOR
 	virtual void PostEditChangeProperty(FPropertyChangedEvent& PropertyChangedEvent) OVERRIDE;
+#endif
 	// End UObject interface.
 
 	/** Update CollisionObjectQueryParams from ObjectTypesToAffect */
diff --git a/Engine/Source/Runtime/Engine/Private/PhysicsEngine/RadialForceComponent.cpp b/Engine/Source/Runtime/Engine/Private/PhysicsEngine/RadialForceComponent.cpp
--- a/Engine/Source/Runtime/Engine/Private/PhysicsEngine/RadialForceComponent.cpp
+++ b/Engine/Source/Runtime/Engine/Private/PhysicsEngine/RadialForceComponent.cpp
@@ -144,6 +144,8 @@ void URadialForceComponent::RemoveObjectTypeToAffect(TEnumAsByte<enum EObjectTyp
 	UpdateCollisionObjectQueryParams();
 }
 
+#if WITH_EDITOR
+
 void URadialForceComponent::PostEditChangeProperty(FPropertyChangedEvent& PropertyChangedEvent)
 {
 	Super::PostEditChangeProperty(PropertyChangedEvent);
@@ -155,6 +157,8 @@ void URadialForceComponent::PostEditChangeProperty(FPropertyChangedEvent& Proper
 	}
 }
 
+#endif
+
 void URadialForceComponent::UpdateCollisionObjectQueryParams()
 {
 	CollisionObjectQueryParams = FCollisionObjectQueryParams(ObjectTypesToAffect);
