Fix rendering of landscape components with holes when they're moved to a streaming level
TTP#320718

diff --git a/Engine/Source/Runtime/Engine/Private/Landscape/Landscape.cpp b/Engine/Source/Runtime/Engine/Private/Landscape/Landscape.cpp
--- a/Engine/Source/Runtime/Engine/Private/Landscape/Landscape.cpp
+++ b/Engine/Source/Runtime/Engine/Private/Landscape/Landscape.cpp
@@ -1357,9 +1357,9 @@ UMaterialInterface* ALandscapeProxy::GetLandscapeMaterial() const
 
 UMaterialInterface* ALandscapeProxy::GetLandscapeHoleMaterial() const
 {
-	if (LandscapeMaterial)
+	if (LandscapeHoleMaterial)
 	{
-		return LandscapeMaterial;
+		return LandscapeHoleMaterial;
 	}
 	else if (LandscapeActor)
 	{
