Fix TTP 335424 - CRITICAL: LANDSCAPE: Ramp tool ramps do not have collision
Also affected the "deform landscape to splines" feature.
Also fixes navmesh for these two tools.

diff --git a/Engine/Source/Editor/LandscapeEditor/Private/LandscapeEdModeRampTool.cpp b/Engine/Source/Editor/LandscapeEditor/Private/LandscapeEdModeRampTool.cpp
--- a/Engine/Source/Editor/LandscapeEditor/Private/LandscapeEdModeRampTool.cpp
+++ b/Engine/Source/Editor/LandscapeEditor/Private/LandscapeEdModeRampTool.cpp
@@ -498,6 +498,26 @@ public:
 			Rasterizer.DrawTriangle(FVector2D(0, Heights[0]), FVector2D(1, Heights[1]), FVector2D(0, Heights[1]), OuterVerts[0][1], InnerVerts[1][1], OuterVerts[1][1], false);
 
 			LandscapeEdit.SetHeightData(MinX, MinY, MaxX, MaxY, Data.GetTypedData(), 0, true);
+			LandscapeEdit.Flush();
+
+			TSet<ULandscapeComponent*> Components;
+			if (LandscapeEdit.GetComponentsInRegion(MinX, MinY, MaxX, MaxY, &Components))
+			{
+				for (ULandscapeComponent* Component : Components)
+				{
+					// Recreate collision for modified components and update the navmesh
+					ULandscapeHeightfieldCollisionComponent* CollisionComponent = Component->CollisionComponent.Get();
+					if (CollisionComponent)
+					{
+						CollisionComponent->RecreateCollision(false);
+						UNavigationSystem* NavSys = UNavigationSystem::GetCurrent(Component);
+						if (NavSys)
+						{
+							NavSys->UpdateNavOctree(CollisionComponent);
+						}
+					}
+				}
+			}
 		}
 	}
 
diff --git a/Engine/Source/Runtime/Engine/Private/Landscape/LandscapeEdit.cpp b/Engine/Source/Runtime/Engine/Private/Landscape/LandscapeEdit.cpp
--- a/Engine/Source/Runtime/Engine/Private/Landscape/LandscapeEdit.cpp
+++ b/Engine/Source/Runtime/Engine/Private/Landscape/LandscapeEdit.cpp
@@ -2805,6 +2805,9 @@ bool ULandscapeInfo::ApplySplinesInternal(bool bOnlySelected, ALandscapeProxy* L
 
 	const FTransform SplineToLandscape = Landscape->SplineComponent->ComponentToWorld.GetRelativeTransform(Landscape->LandscapeActorToWorld());
 
+	FLandscapeEditDataInterface LandscapeEdit(this);
+	TSet<ULandscapeComponent*> ModifiedComponents;
+
 	for (auto It = Landscape->SplineComponent->ControlPoints.CreateConstIterator(); It; ++It)
 	{
 		const ULandscapeSplineControlPoint* ControlPoint = *It;
@@ -2827,7 +2830,6 @@ bool ULandscapeInfo::ApplySplinesInternal(bool bOnlySelected, ALandscapeProxy* L
 		int32 MaxX = FMath::FloorToInt(LandscapeBounds.Max.X);
 		int32 MaxY = FMath::FloorToInt(LandscapeBounds.Max.Y);
 
-		FLandscapeEditDataInterface LandscapeEdit(this);
 		TArray<FLandscapeSplineInterpPoint> Points = ControlPoint->GetPoints();
 		for (int32 j = 0; j < Points.Num(); j++)
 		{
@@ -2900,6 +2902,7 @@ bool ULandscapeInfo::ApplySplinesInternal(bool bOnlySelected, ALandscapeProxy* L
 			}
 
 			LandscapeEdit.SetHeightData(MinX, MinY, MaxX, MaxY, Data.GetTypedData(), 0, true);
+			LandscapeEdit.GetComponentsInRegion(MinX, MinY, MaxX, MaxY, &ModifiedComponents);
 		}
 
 		// Blend layer raster
@@ -2960,6 +2963,7 @@ bool ULandscapeInfo::ApplySplinesInternal(bool bOnlySelected, ALandscapeProxy* L
 			}
 
 			LandscapeEdit.SetAlphaData(LayerInfo, MinX, MinY, MaxX, MaxY, Data.GetTypedData(), 0, ELandscapeLayerPaintingRestriction::None, true, false);
+			LandscapeEdit.GetComponentsInRegion(MinX, MinY, MaxX, MaxY, &ModifiedComponents);
 		}
 	}
 
@@ -2980,7 +2984,6 @@ bool ULandscapeInfo::ApplySplinesInternal(bool bOnlySelected, ALandscapeProxy* L
 		int32 MaxX = FMath::FloorToInt(LandscapeBounds.Max.X);
 		int32 MaxY = FMath::FloorToInt(LandscapeBounds.Max.Y);
 
-		FLandscapeEditDataInterface LandscapeEdit(this);
 		TArray<FLandscapeSplineInterpPoint> Points = Segment->GetPoints();
 		for (int32 j = 0; j < Points.Num(); j++)
 		{
@@ -3058,6 +3061,7 @@ bool ULandscapeInfo::ApplySplinesInternal(bool bOnlySelected, ALandscapeProxy* L
 			}
 
 			LandscapeEdit.SetHeightData(MinX, MinY, MaxX, MaxY, Data.GetTypedData(), 0, true);
+			LandscapeEdit.GetComponentsInRegion(MinX, MinY, MaxX, MaxY, &ModifiedComponents);
 		}
 
 		// Blend layer raster
@@ -3123,6 +3127,24 @@ bool ULandscapeInfo::ApplySplinesInternal(bool bOnlySelected, ALandscapeProxy* L
 			}
 
 			LandscapeEdit.SetAlphaData(LayerInfo, MinX, MinY, MaxX, MaxY, Data.GetTypedData(), 0, ELandscapeLayerPaintingRestriction::None, true, false);
+			LandscapeEdit.GetComponentsInRegion(MinX, MinY, MaxX, MaxY, &ModifiedComponents);
+		}
+	}
+
+	LandscapeEdit.Flush();
+
+	for (ULandscapeComponent* Component : ModifiedComponents)
+	{
+		// Recreate collision for modified components and update the navmesh
+		ULandscapeHeightfieldCollisionComponent* CollisionComponent = Component->CollisionComponent.Get();
+		if (CollisionComponent)
+		{
+			CollisionComponent->RecreateCollision(false);
+			UNavigationSystem* NavSys = UNavigationSystem::GetCurrent(Component);
+			if (NavSys)
+			{
+				NavSys->UpdateNavOctree(CollisionComponent);
+			}
 		}
 	}
 
