Restore the pre-10.9.2 workaround for GPU skinning that uses a 2D texture rather than a texture buffer to avoid an AMD Dx00/7xx0 driver regression in 10.9.3.
AMD have already fixed this in their driver upstream, but it missed 10.9.3. This workaround is used on all Macs for consistency & will be removed when the fix is released.
Note that this workaround has a small but noticeable performance hit.

diff --git a/Engine/Shaders/GpuSkinVertexFactory.usf b/Engine/Shaders/GpuSkinVertexFactory.usf
--- a/Engine/Shaders/GpuSkinVertexFactory.usf
+++ b/Engine/Shaders/GpuSkinVertexFactory.usf
@@ -51,10 +51,17 @@ float3 MeshExtension;
 bool PerBoneMotionBlur;
 
 #if GPUSKIN_APEX_CLOTH
+#if GPUSKIN_USE_DATA_BUFFERS == 1
 /** Vertex buffer from which to read simulated positions of clothing. */
 Buffer<float4> ClothSimulVertsPositions;
 /** Vertex buffer from which to read simulated normals of clothing. */
 Buffer<float4> ClothSimulVertsNormals;
+#else
+/** Texture from which to read simulated position of clothing. */
+Texture2D ClothSimulVertsPositions;
+/** Texture from which to read simulated normal of clothing. */
+Texture2D ClothSimulVertsNormals;
+#endif // GPUSKIN_USE_DATA_BUFFERS
 #endif// #if GPUSKIN_APEX_CLOTH
 
 struct FVertexFactoryInput
@@ -159,10 +166,17 @@ half2 GetShadowMapCoordinate(FVertexFactoryInterpolantsVSToPS Interpolants)
 
 #if FEATURE_LEVEL >= FEATURE_LEVEL_SM4
 
+#if GPUSKIN_USE_DATA_BUFFERS == 1
 // The bone matrix buffer stored as 4x3 (3 float4 texels behind each other), all chunks of a skeletal mesh in one
 Buffer<float4> BoneMatrices;
 // buffer with all old bone matrices stored as 4x3 (3 float4 texels behind each other), all chunks of a skeletal mesh in one
 Buffer<float4> PreviousBoneMatrices;
+#else
+// The bone matrix buffer stored as 4x3 (3 float4 texels behind each other), all chunks of a skeletal mesh in one
+Texture2D BoneMatrices;
+// buffer with all old bone matrices stored as 4x3 (3 float4 texels behind each other), all chunks of a skeletal mesh in one
+Texture2D PreviousBoneMatrices;
+#endif
 
 // .xyz to offset the lookup in the buffer .w unused
 uint4 BoneIndexOffset;
@@ -223,7 +237,11 @@ bool IsSimulatedVertex( FVertexFactoryInput Input )
 float3 GetClothSimulPosition(int Index)
 {
 #if FEATURE_LEVEL >= FEATURE_LEVEL_SM4
-	return ClothSimulVertsPositions[Index].xyz;
+	#if GPUSKIN_USE_DATA_BUFFERS == 1
+		return ClothSimulVertsPositions[Index].xyz;
+	#else
+		return ClothSimulVertsPositions.Load(int3(Index % GPUSKIN_CLOTH_TEXTURE_STRIDE, Index / GPUSKIN_CLOTH_TEXTURE_STRIDE, 0));
+	#endif // GPUSKIN_USE_DATA_BUFFERS
 #else
 	return APEXClothParam.Positions[Index];
 #endif //FEATURE_LEVEL >= FEATURE_LEVEL_SM4
@@ -232,7 +250,11 @@ float3 GetClothSimulPosition(int Index)
 float3 GetClothSimulNormal(int Index)
 {
 #if FEATURE_LEVEL >= FEATURE_LEVEL_SM4
-	return ClothSimulVertsNormals[Index].xyz;
+	#if GPUSKIN_USE_DATA_BUFFERS == 1
+		return ClothSimulVertsNormals[Index].xyz;
+	#else
+		return ClothSimulVertsNormals.Load(int3(Index % GPUSKIN_CLOTH_TEXTURE_STRIDE, Index / GPUSKIN_CLOTH_TEXTURE_STRIDE, 0)).xyz;
+	#endif // GPUSKIN_USE_DATA_BUFFERS
 #else
 	return APEXClothParam.Normals[Index];
 #endif //FEATURE_LEVEL >= FEATURE_LEVEL_SM4
@@ -280,9 +302,15 @@ float3 MorphPosition( FVertexFactoryInput Input, FVertexFactoryIntermediates Int
 FBoneMatrix GetBoneMatrix(int Index)
 {
 #if FEATURE_LEVEL >= FEATURE_LEVEL_SM4
+#if GPUSKIN_USE_DATA_BUFFERS == 1
 	float4 A = BoneMatrices[Index * 3];
 	float4 B = BoneMatrices[Index * 3 + 1];
 	float4 C = BoneMatrices[Index * 3 + 2];
+#else
+	float4 A = BoneMatrices.Load(int3((Index * 3), 0, 0));
+	float4 B = BoneMatrices.Load(int3((Index * 3) + 1, 0, 0));
+	float4 C = BoneMatrices.Load(int3((Index * 3) + 2, 0, 0));
+#endif
 	return FBoneMatrix(A,B,C);
 #else
 	return Bones.BoneMatrices[Index];
@@ -309,9 +337,16 @@ FBoneMatrix CalcBoneMatrix( FVertexFactoryInput Input )
 FBoneMatrix GetPreviousBoneMatrix(int Index)
 {
 #if FEATURE_LEVEL >= FEATURE_LEVEL_SM4
+#if GPUSKIN_USE_DATA_BUFFERS == 1
 	float4 A = PreviousBoneMatrices[Index * 3 + BoneIndexOffset.x];
 	float4 B = PreviousBoneMatrices[Index * 3 + BoneIndexOffset.y];
 	float4 C = PreviousBoneMatrices[Index * 3 + BoneIndexOffset.z];
+#else
+	int3 Indices = int3(Index * 3 + BoneIndexOffset.x, Index * 3 + BoneIndexOffset.y, Index * 3 + BoneIndexOffset.z);
+	float4 A = PreviousBoneMatrices.Load(int3(Indices.x % GPUSKIN_TEXTURE_STRIDE, Indices.x / GPUSKIN_TEXTURE_STRIDE, 0));
+	float4 B = PreviousBoneMatrices.Load(int3(Indices.y % GPUSKIN_TEXTURE_STRIDE, Indices.y / GPUSKIN_TEXTURE_STRIDE, 0));
+	float4 C = PreviousBoneMatrices.Load(int3(Indices.z % GPUSKIN_TEXTURE_STRIDE, Indices.z / GPUSKIN_TEXTURE_STRIDE, 0));
+#endif
 	return FBoneMatrix(A,B,C);
 #else
 	return Bones.BoneMatrices[Index];
diff --git a/Engine/Source/Runtime/Engine/Private/GPUSkinCache.cpp b/Engine/Source/Runtime/Engine/Private/GPUSkinCache.cpp
--- a/Engine/Source/Runtime/Engine/Private/GPUSkinCache.cpp
+++ b/Engine/Source/Runtime/Engine/Private/GPUSkinCache.cpp
@@ -51,6 +51,8 @@ public:
 	static void ModifyCompilationEnvironment(EShaderPlatform Platform, FShaderCompilerEnvironment& OutEnvironment)
 	{
 		FGlobalShader::ModifyCompilationEnvironment(Platform,OutEnvironment);
+		const uint32 UseDataBuffers = GPUSKIN_USE_DATA_BUFFERS;
+		OutEnvironment.SetDefine(TEXT("GPUSKIN_USE_DATA_BUFFERS"), UseDataBuffers);
 		const uint32 UseExtraBoneInfluences = bUseExtraBoneInfluencesT;
 		OutEnvironment.SetDefine(TEXT("GPUSKIN_USE_EXTRA_INFLUENCES"), UseExtraBoneInfluences);
 	}
@@ -92,7 +94,9 @@ public:
 		}
 		else
 		{
+#if GPUSKIN_USE_DATA_BUFFERS
 			RHISetShaderResourceViewParameter(ComputeShaderRHI, BoneMatrices.GetBaseIndex(), BoneBuffer.VertexBufferSRV);
+#endif
 		}
 
 		RHISetShaderResourceViewParameter(ComputeShaderRHI, SkinInputStream.GetBaseIndex(), VertexBufferSRV);
@@ -108,7 +112,9 @@ public:
 	{
 		FComputeShaderRHIParamRef ComputeShaderRHI = GetComputeShader();
 		FShaderResourceViewRHIParamRef NullSRV = FShaderResourceViewRHIParamRef();
+#if GPUSKIN_USE_DATA_BUFFERS
 		RHISetShaderResourceViewParameter(ComputeShaderRHI, BoneMatrices.GetBaseIndex(), NullSRV);
+#endif
 
 		RHISetShaderResourceViewParameter(ComputeShaderRHI, SkinInputStream.GetBaseIndex(), NullSRV);
 
diff --git a/Engine/Source/Runtime/Engine/Private/GPUSkinVertexFactory.cpp b/Engine/Source/Runtime/Engine/Private/GPUSkinVertexFactory.cpp
--- a/Engine/Source/Runtime/Engine/Private/GPUSkinVertexFactory.cpp
+++ b/Engine/Source/Runtime/Engine/Private/GPUSkinVertexFactory.cpp
@@ -78,6 +78,7 @@ uint32 FSharedPoolPolicyData::BucketSizes[NumPoolBucketSizes] = {
 	65536, 131072, 262144 // these 3 numbers are added for large cloth simulation vertices, supports up to 16,384 verts
 };
 
+#if GPUSKIN_USE_DATA_BUFFERS
 /*-----------------------------------------------------------------------------
  FBoneBufferPoolPolicy
  -----------------------------------------------------------------------------*/
@@ -107,6 +108,36 @@ TStatId FBoneBufferPool::GetStatId() const
 	RETURN_QUICK_DECLARE_CYCLE_STAT(FBoneBufferPool, STATGROUP_Tickables);
 }
 
+#else
+/*-----------------------------------------------------------------------------
+ FBoneTexturePoolPolicy
+ -----------------------------------------------------------------------------*/
+FTexture2DRHIRef FBoneTexturePoolPolicy::CreateResource(FSharedPoolPolicyData::CreationArguments Args)
+{
+	FSharedPoolPolicyData::CreationArguments ActualSize = GetPoolBucketSize(GetPoolBucketIndex(Args)) / sizeof(FVector4);
+	FTexture2DRHIRef Texture = RHICreateTexture2D(ActualSize, 1, PF_A32B32G32R32F, 1, 1, (TexCreate_ShaderResource|TexCreate_NoMipTail|TexCreate_Dynamic), NULL);
+	return Texture;
+}
+
+FSharedPoolPolicyData::CreationArguments FBoneTexturePoolPolicy::GetCreationArguments(FTexture2DRHIRef Resource)
+{
+	return (Resource->GetSizeX() * Resource->GetSizeY());
+}
+
+/*-----------------------------------------------------------------------------
+ FBoneTexturePool
+ -----------------------------------------------------------------------------*/
+FBoneTexturePool::~FBoneTexturePool()
+{
+}
+
+TStatId FBoneTexturePool::GetStatId() const
+{
+	RETURN_QUICK_DECLARE_CYCLE_STAT(FBoneTexturePool, STATGROUP_Tickables);
+}
+
+#endif
+
 TConsoleVariableData<int32>* FGPUBaseSkinVertexFactory::ShaderDataType::MaxBonesVar = NULL;
 uint32 FGPUBaseSkinVertexFactory::ShaderDataType::MaxGPUSkinBones = 0;
 
@@ -121,6 +152,7 @@ void FGPUBaseSkinVertexFactory::ShaderDataType::UpdateBoneData()
 		check(NumVectors <= (MaxGPUSkinBones*3));
 		uint32 VectorArraySize = NumVectors * sizeof(FVector4);
 		uint32 PooledArraySize = BoneBufferPool.PooledSizeForCreationArguments(VectorArraySize);
+#if GPUSKIN_USE_DATA_BUFFERS
 		if(!IsValidRef(BoneBuffer) || PooledArraySize != BoneBuffer.VertexBufferRHI->GetSize())
 		{
 			if(IsValidRef(BoneBuffer))
@@ -137,6 +169,26 @@ void FGPUBaseSkinVertexFactory::ShaderDataType::UpdateBoneData()
 			FMemory::Memcpy(Data, BoneMatrices.GetTypedData(), NumBones * sizeof(BoneMatrices[0]));
 			RHIUnlockVertexBuffer(BoneBuffer.VertexBufferRHI);
 		}
+#else
+		check((uint32)GMaxTextureDimensions >= (MaxGPUSkinBones*3));
+		
+		if(!IsValidRef(BoneBuffer) || (PooledArraySize != (BoneBuffer->GetSizeX() * BoneBuffer->GetSizeY()) * sizeof(FVector4)))
+		{
+			if(IsValidRef(BoneBuffer))
+			{
+				BoneBufferPool.ReleasePooledResource(BoneBuffer);
+			}
+			BoneBuffer = BoneBufferPool.CreatePooledResource(VectorArraySize);
+			check(IsValidRef(BoneBuffer));
+		}
+		if(NumBones)
+		{
+			uint32 Width = BoneBuffer->GetSizeX();
+			uint32 Height = BoneBuffer->GetSizeY();
+			FUpdateTextureRegion2D UpdateRegion(0, 0, 0, 0, NumVectors, Height);
+			RHIUpdateTexture2D(BoneBuffer, 0, UpdateRegion, NumVectors * sizeof(FVector4), (uint8 const*)BoneMatrices.GetTypedData());
+		}
+#endif
 	}
 	else
 	{
@@ -162,7 +214,17 @@ float* FBoneDataVertexBuffer::LockData()
 	checkSlow(GetSizeX());
 	checkSlow(IsValidRef(BoneBuffer));
 
+#if GPUSKIN_USE_DATA_BUFFERS
 	float* Data = (float*)RHILockVertexBuffer(BoneBuffer.VertexBufferRHI, 0, ComputeMemorySize(), RLM_WriteOnly);
+#else
+	uint32 DestStride = 0;
+	float* Data = (float*)RHILockTexture2D(
+										   BoneBuffer,
+										   0,
+										   RLM_WriteOnly,
+										   DestStride,
+										   false);
+#endif
 	checkSlow(Data);
 
 	return Data;
@@ -171,7 +233,11 @@ float* FBoneDataVertexBuffer::LockData()
 void FBoneDataVertexBuffer::UnlockData()
 {
 	checkSlow(IsValidRef(BoneBuffer));
+#if GPUSKIN_USE_DATA_BUFFERS
 	RHIUnlockVertexBuffer(BoneBuffer.VertexBufferRHI);
+#else
+	RHIUnlockTexture2D(BoneBuffer, 0, false);
+#endif
 }
 
 uint32 FBoneDataVertexBuffer::GetSizeX() const
@@ -209,6 +275,10 @@ void TGPUSkinVertexFactory<bExtraBoneInfluencesT>::ModifyCompilationEnvironment(
 	FVertexFactory::ModifyCompilationEnvironment(Platform, Material, OutEnvironment);
 	const int32 MaxGPUSkinBones = GetFeatureLevelMaxNumberOfBones(GetMaxSupportedFeatureLevel(Platform));
 	OutEnvironment.SetDefine(TEXT("MAX_SHADER_BONES"), MaxGPUSkinBones);
+	const uint32 UseDataBuffers = GPUSKIN_USE_DATA_BUFFERS;
+	OutEnvironment.SetDefine(TEXT("GPUSKIN_USE_DATA_BUFFERS"), UseDataBuffers);
+	OutEnvironment.SetDefine(TEXT("GPUSKIN_TEXTURE_STRIDE"), GPUSKIN_TEXTURE_STRIDE);
+	OutEnvironment.SetDefine(TEXT("GPUSKIN_CLOTH_TEXTURE_STRIDE"), GPUSKIN_CLOTH_TEXTURE_STRIDE);
 	const uint32 UseExtraBoneInfluences = bExtraBoneInfluencesT;
 	OutEnvironment.SetDefine(TEXT("GPUSKIN_USE_EXTRA_INFLUENCES"), UseExtraBoneInfluences);
 }
@@ -374,7 +444,11 @@ public:
 			{
 				if(BoneMatrices.IsBound())
 				{
+#if GPUSKIN_USE_DATA_BUFFERS
 					RHISetShaderResourceViewParameter(Shader->GetVertexShader(), BoneMatrices.GetBaseIndex(), ShaderData.GetBoneBuffer().VertexBufferSRV);
+#else
+					RHISetShaderTexture(Shader->GetVertexShader(), BoneMatrices.GetBaseIndex(), ShaderData.GetBoneBuffer());
+#endif
 				}
 			}
 			else
@@ -397,7 +471,11 @@ public:
 				// we tell the shader where to pickup the data (always, even if we don't have bone data, to avoid false binding)
 				if(PreviousBoneMatrices.IsBound())
 				{
+#if GPUSKIN_USE_DATA_BUFFERS
 					RHISetShaderResourceViewParameter(Shader->GetVertexShader(), PreviousBoneMatrices.GetBaseIndex(), GPrevPerBoneMotionBlur.GetReadData()->BoneBuffer.VertexBufferSRV);
+#else
+					RHISetShaderTexture(Shader->GetVertexShader(), PreviousBoneMatrices.GetBaseIndex(), GPrevPerBoneMotionBlur.GetReadData()->BoneBuffer);
+#endif
 				}
 
 				if(bLocalPerBoneMotionBlur)
@@ -640,12 +718,20 @@ public:
 			// we tell the shader where to pickup the data
 			if(ClothSimulPositionsParameter.IsBound())
 			{
+#if GPUSKIN_USE_DATA_BUFFERS
 				RHISetShaderResourceViewParameter(Shader->GetVertexShader(), ClothSimulPositionsParameter.GetBaseIndex(), ClothShaderData.GetClothSimulPositionBuffer().VertexBufferSRV);
+#else
+				RHISetShaderTexture(Shader->GetVertexShader(), ClothSimulPositionsParameter.GetBaseIndex(), ClothShaderData.GetClothSimulPositionBuffer());
+#endif
 			}
-
+			
 			if(ClothSimulNormalsParameter.IsBound())
 			{
+#if GPUSKIN_USE_DATA_BUFFERS
 				RHISetShaderResourceViewParameter(Shader->GetVertexShader(), ClothSimulNormalsParameter.GetBaseIndex(), ClothShaderData.GetClothSimulNormalBuffer().VertexBufferSRV);
+#else
+				RHISetShaderTexture(Shader->GetVertexShader(), ClothSimulNormalsParameter.GetBaseIndex(), ClothShaderData.GetClothSimulNormalBuffer());
+#endif
 			}
 		}
 	}
@@ -668,7 +754,7 @@ void FGPUBaseSkinAPEXClothVertexFactory::ClothShaderType::UpdateClothSimulData(c
 
 		uint32 VectorArraySize = NumSimulVerts * sizeof(FVector4);
 		uint32 PooledArraySize = ClothSimulDataBufferPool.PooledSizeForCreationArguments(VectorArraySize);
-
+#if GPUSKIN_USE_DATA_BUFFERS
 		if(!IsValidRef(ClothSimulPositionBuffer) || PooledArraySize != ClothSimulPositionBuffer.VertexBufferRHI->GetSize())
 		{
 			if(IsValidRef(ClothSimulPositionBuffer))
@@ -701,6 +787,38 @@ void FGPUBaseSkinAPEXClothVertexFactory::ClothShaderType::UpdateClothSimulData(c
 			FMemory::Memcpy(Data, InSimulNormals.GetTypedData(), NumSimulVerts * sizeof(FVector4));
 			RHIUnlockVertexBuffer(ClothSimulNormalBuffer.VertexBufferRHI);
 		}
+#else
+		
+		if(!IsValidRef(ClothSimulPositionBuffer) || (PooledArraySize != (ClothSimulPositionBuffer->GetSizeX() * ClothSimulPositionBuffer->GetSizeY()) * sizeof(FVector4)))
+		{
+			if(IsValidRef(ClothSimulPositionBuffer))
+			{
+				ClothSimulDataBufferPool.ReleasePooledResource(ClothSimulPositionBuffer);
+			}
+			ClothSimulPositionBuffer = ClothSimulDataBufferPool.CreatePooledResource(VectorArraySize);
+			check(IsValidRef(ClothSimulPositionBuffer));
+		}
+		
+		if(!IsValidRef(ClothSimulNormalBuffer) || (PooledArraySize != (ClothSimulNormalBuffer->GetSizeX() * ClothSimulNormalBuffer->GetSizeY()) * sizeof(FVector4)))
+		{
+			if(IsValidRef(ClothSimulNormalBuffer))
+			{
+				ClothSimulDataBufferPool.ReleasePooledResource(ClothSimulNormalBuffer);
+			}
+			ClothSimulNormalBuffer = ClothSimulDataBufferPool.CreatePooledResource(VectorArraySize);
+			check(IsValidRef(ClothSimulNormalBuffer));
+		}
+		
+		if(NumSimulVerts)
+		{
+			uint32 Width = ClothSimulPositionBuffer->GetSizeX();
+			uint32 Height = ClothSimulPositionBuffer->GetSizeY();
+			FUpdateTextureRegion2D UpdateRegion(0, 0, 0, 0, NumSimulVerts, Height);
+			RHIUpdateTexture2D(ClothSimulPositionBuffer, 0, UpdateRegion, NumSimulVerts * sizeof(FVector4), (uint8 const*)InSimulPositions.GetTypedData());
+			RHIUpdateTexture2D(ClothSimulNormalBuffer, 0, UpdateRegion, NumSimulVerts * sizeof(FVector4), (uint8 const*)InSimulNormals.GetTypedData());
+		}
+		
+#endif// #if GPUSKIN_USE_DATA_BUFFERS
 	}
 	else
 	{
diff --git a/Engine/Source/Runtime/Engine/Public/GPUSkinVertexFactory.h b/Engine/Source/Runtime/Engine/Public/GPUSkinVertexFactory.h
--- a/Engine/Source/Runtime/Engine/Public/GPUSkinVertexFactory.h
+++ b/Engine/Source/Runtime/Engine/Public/GPUSkinVertexFactory.h
@@ -108,6 +108,15 @@ typedef FSkinMatrix3x4 FBoneSkinning;
 
 #define SET_BONE_DATA(B, X) B.SetMatrixTranspose(X)
 
+// For most platforms we want to use a uniform buffer + vertex buffer/SRV pair, but:
+// - OS X <= 10.8.5 on Intel HD4000 falls back to software rendering when using the large uniform buffer for SM4+.
+// - OS X <= 10.8.5 crashes when using the GL_TEXTURE_BUFFER binding for the vertex buffer/SRV pair.
+// - OS X == 10.9.3 will render poly-soup on AMD Dx00/7xx0 if the texture unit sRGB state changes when binding the GL_TEXTURE_BUFFER (already fixed by AMD, but missed 10.9.3).
+// When not using a VB/SRV pair we'll fallback to 2D textures for both.
+#define GPUSKIN_USE_DATA_BUFFERS !PLATFORM_MAC
+#define GPUSKIN_TEXTURE_STRIDE 1024u
+#define GPUSKIN_CLOTH_TEXTURE_STRIDE 16384u
+
 /** Shared data & implementation for the different types of pool */
 class FSharedPoolPolicyData
 {
@@ -139,6 +148,7 @@ private:
 	static uint32 BucketSizes[NumPoolBucketSizes];
 };
 
+#if GPUSKIN_USE_DATA_BUFFERS
 /** Struct to pool the vertex buffer & SRV together */
 struct FBoneBuffer
 {
@@ -202,6 +212,50 @@ public: // From FTickableObjectRenderThread
 /** The type for the buffer pool */
 typedef FBoneBufferPool FBoneBufferPool;
 
+#else
+
+/** The policy for pooling bone texture buffers */
+class FBoneTexturePoolPolicy : public FSharedPoolPolicyData
+{
+public:
+	enum
+	{
+		NumSafeFrames = FSharedPoolPolicyData::NumSafeFrames,
+		NumPoolBuckets = FSharedPoolPolicyData::NumPoolBucketSizes,
+		NumToDrainPerFrame = FSharedPoolPolicyData::NumToDrainPerFrame,
+		CullAfterFramesNum = FSharedPoolPolicyData::CullAfterFramesNum
+	};
+	/** Creates the resource
+	 * @param Args The buffer size in bytes.
+	 * @returns A suitably sized buffer or NULL on failure.
+	 */
+	FTexture2DRHIRef CreateResource(FSharedPoolPolicyData::CreationArguments Args);
+	
+	/** Gets the arguments used to create resource
+	 * @param Resource The buffer to get data for.
+	 * @returns The arguments used to create the buffer.
+	 */
+	FSharedPoolPolicyData::CreationArguments GetCreationArguments(FTexture2DRHIRef Resource);
+};
+
+
+/** A pool for bone textures with consistent usage, bucketed for efficiency. */
+class FBoneTexturePool : public TRenderResourcePool<FTexture2DRHIRef, FBoneTexturePoolPolicy, FSharedPoolPolicyData::CreationArguments>
+{
+public:
+	/** Destructor */
+	virtual ~FBoneTexturePool();
+	
+public: // From FTickableObjectRenderThread
+	virtual TStatId GetStatId() const OVERRIDE;
+};
+
+/** The type for bone buffers */
+typedef FTexture2DRHIRef FBoneBufferTypeRef;
+/** The type for the buffer pool */
+typedef FBoneTexturePool FBoneBufferPool;
+#endif
+
 /** for motion blur skinning */
 class FBoneDataVertexBuffer : public FRenderResource
 {
@@ -241,9 +295,21 @@ public:
 		if(SizeX)
 		{
 			INC_DWORD_STAT_BY( STAT_SkeletalMeshMotionBlurSkinningMemory, ComputeMemorySize());
-			const int32 TileBufferSize = ComputeMemorySize();
-			BoneBuffer.VertexBufferRHI = RHICreateVertexBuffer( TileBufferSize, NULL, BUF_Volatile | BUF_ShaderResource );
-			BoneBuffer.VertexBufferSRV = RHICreateShaderResourceView( BoneBuffer.VertexBufferRHI, sizeof(FVector4), PF_A32B32G32R32F );
+#if GPUSKIN_USE_DATA_BUFFERS
+			{
+				const int32 TileBufferSize = ComputeMemorySize();
+				BoneBuffer.VertexBufferRHI = RHICreateVertexBuffer( TileBufferSize, NULL, BUF_Volatile | BUF_ShaderResource );
+				BoneBuffer.VertexBufferSRV = RHICreateShaderResourceView( BoneBuffer.VertexBufferRHI, sizeof(FVector4), PF_A32B32G32R32F );
+			}
+#else
+			{
+				uint32 X = GPUSKIN_TEXTURE_STRIDE;
+				uint32 Y = (SizeX / GPUSKIN_TEXTURE_STRIDE)+1;
+				check((uint32)GMaxTextureDimensions >= X);
+				check((uint32)GMaxTextureDimensions >= Y);
+				BoneBuffer = RHICreateTexture2D(X, Y, PF_A32B32G32R32F, 1, 1, (TexCreate_ShaderResource|TexCreate_NoMipTail|TexCreate_Dynamic), NULL);
+			}
+#endif
 		}
 	}
 	
