Initial checkin for editor functionality in ComponentVisualizers
Simplified SplineComponent, now a SceneComponent not a PrimitiveComponent
Added support for persisting SplineComp state across CS re-run
SplineComponentVisualizer editing started

diff --git a/Engine/Source/Editor/ComponentVisualizers/ComponentVisualizers.Build.cs b/Engine/Source/Editor/ComponentVisualizers/ComponentVisualizers.Build.cs
--- a/Engine/Source/Editor/ComponentVisualizers/ComponentVisualizers.Build.cs
+++ b/Engine/Source/Editor/ComponentVisualizers/ComponentVisualizers.Build.cs
@@ -12,6 +12,7 @@ public class ComponentVisualizers : ModuleRules
 			new string[] {
 				"Core",
 				"CoreUObject",
+                "InputCore",
 				"Engine",
 				"Slate",
 				"UnrealEd",
diff --git a/Engine/Source/Editor/ComponentVisualizers/Private/SplineComponentVisualizer.cpp b/Engine/Source/Editor/ComponentVisualizers/Private/SplineComponentVisualizer.cpp
--- a/Engine/Source/Editor/ComponentVisualizers/Private/SplineComponentVisualizer.cpp
+++ b/Engine/Source/Editor/ComponentVisualizers/Private/SplineComponentVisualizer.cpp
@@ -3,10 +3,195 @@
 #include "ComponentVisualizersPrivatePCH.h"
 #include "SplineComponentVisualizer.h"
 
+IMPLEMENT_HIT_PROXY(HSplineVisProxy, HComponentVisProxy);
+IMPLEMENT_HIT_PROXY(HSplineKeyProxy, HSplineVisProxy);
+
+
 void FSplineComponentVisualizer::DrawVisualization(const UActorComponent* Component, const FSceneView* View, FPrimitiveDrawInterface* PDI)
 {
 	if (const USplineComponent* SplineComp = Cast<const USplineComponent>(Component))
 	{
-		
+		FColor DrawColor = SplineComp->SplineColor;
+
+		const FInterpCurveVector& SplineInfo = SplineComp->SplineInfo;
+
+		FVector OldKeyPos(0);
+		float OldKeyTime = 0.f;
+		for (int32 KeyIdx = 0; KeyIdx<SplineInfo.Points.Num(); KeyIdx++)
+		{
+			float NewKeyTime = SplineInfo.Points[KeyIdx].InVal;
+			FVector NewKeyPos = SplineComp->ComponentToWorld.TransformPosition( SplineInfo.Eval(NewKeyTime, FVector::ZeroVector) );
+
+			USplineComponent* EditedSplineComp = GetEditedSplineComponent();
+			FColor KeyColor = (SplineComp == EditedSplineComp && KeyIdx == SelectedKeyIndex) ? FColor(255,0,0) : DrawColor;
+
+			// Draw the keypoint
+			PDI->SetHitProxy(new HSplineKeyProxy(Component, KeyIdx));
+
+			PDI->DrawPoint(NewKeyPos, KeyColor, 6.f, SDPG_World);
+
+			PDI->SetHitProxy(NULL);
+
+
+			// If not the first keypoint, draw a line to the last keypoint.
+			if (KeyIdx>0)
+			{
+				// For constant interpolation - don't draw ticks - just draw dotted line.
+				if (SplineInfo.Points[KeyIdx - 1].InterpMode == CIM_Constant)
+				{
+					DrawDashedLine(PDI, OldKeyPos, NewKeyPos, DrawColor, 20, SDPG_World);
+				}
+				else
+				{
+					int32 NumSteps = FMath::CeilToInt((NewKeyTime - OldKeyTime) / 0.1f);
+					float DrawSubstep = (NewKeyTime - OldKeyTime) / NumSteps;
+
+					// Find position on first keyframe.
+					float OldTime = OldKeyTime;
+					FVector OldPos = OldKeyPos;
+
+					// Then draw a line for each substep.
+					for (int32 StepIdx = 1; StepIdx<NumSteps + 1; StepIdx++)
+					{
+						float NewTime = OldKeyTime + StepIdx*DrawSubstep;
+						FVector NewPos = SplineComp->ComponentToWorld.TransformPosition( SplineInfo.Eval(NewTime, FVector::ZeroVector) );
+
+						PDI->DrawLine(OldPos, NewPos, DrawColor, SDPG_World);
+
+						// Don't draw point for last one - its the keypoint drawn above.
+						if (StepIdx != NumSteps)
+						{
+							PDI->DrawPoint(NewPos, DrawColor, 3.f, SDPG_World);
+						}
+
+						OldTime = NewTime;
+						OldPos = NewPos;
+					}
+				}
+			}
+
+			OldKeyTime = NewKeyTime;
+			OldKeyPos = NewKeyPos;
+		}
+	}
+}
+
+bool FSplineComponentVisualizer::VisProxyHandleClick(HComponentVisProxy* VisProxy)
+{
+	bool bEditing = false;
+
+	if(VisProxy && VisProxy->Component.IsValid())
+	{
+		const USplineComponent* SplineComp = CastChecked<const USplineComponent>(VisProxy->Component.Get());
+
+		SplineCompPropName = GetComponentPropertyName(SplineComp);
+		if(SplineCompPropName != NAME_None)
+		{
+			SplineOwningActor = SplineComp->GetOwner();
+
+			if (VisProxy->IsA(HSplineKeyProxy::StaticGetType()))
+			{
+				HSplineKeyProxy* KeyProxy = (HSplineKeyProxy*)VisProxy;
+				SelectedKeyIndex = KeyProxy->KeyIndex;
+				bEditing = true;
+			}
+		}
+		else
+		{
+			SplineOwningActor = NULL;
+		}
 	}
+
+	return bEditing;
 }
+
+USplineComponent* FSplineComponentVisualizer::GetEditedSplineComponent()
+{
+	return Cast<USplineComponent>(GetComponentFromPropertyName(SplineOwningActor.Get(), SplineCompPropName));
+}
+
+
+bool FSplineComponentVisualizer::GetWidgetLocation(FVector& OutLocation)
+{
+	USplineComponent* SplineComp = GetEditedSplineComponent();
+	if (SplineComp && SelectedKeyIndex != INDEX_NONE)
+	{
+		if (SelectedKeyIndex < SplineComp->SplineInfo.Points.Num())
+		{
+			OutLocation = SplineComp->ComponentToWorld.TransformPosition( SplineComp->SplineInfo.Points[SelectedKeyIndex].OutVal );
+			return true;
+		}
+	}
+
+	return false;
+}
+
+
+bool FSplineComponentVisualizer::HandleInputDelta(FLevelEditorViewportClient* ViewportClient, FViewport* Viewport, FVector& DeltaTranslate, FRotator& DeltalRotate, FVector& DeltaScale)
+{
+	USplineComponent* SplineComp = GetEditedSplineComponent();
+	if( SplineComp && 
+		SelectedKeyIndex != INDEX_NONE &&
+		SelectedKeyIndex < SplineComp->SplineInfo.Points.Num() )
+	{
+		if (ViewportClient->IsAltPressed() && bAllowDuplication)
+		{
+			// Duplicate key
+			FInterpCurvePoint<FVector> KeyToDupe = SplineComp->SplineInfo.Points[SelectedKeyIndex];
+			SplineComp->SplineInfo.Points.Insert(KeyToDupe, SelectedKeyIndex);
+			// move selection to 'next' key
+			SelectedKeyIndex++;
+
+			// Update Input value for all keys
+			SplineComp->RefreshSplineInputs();
+
+			// Don't duplicate again until we release LMB
+			bAllowDuplication = false;
+		}
+
+		// Find key position in world space
+		FVector CurrentWorldPos = SplineComp->ComponentToWorld.TransformPosition(SplineComp->SplineInfo.Points[SelectedKeyIndex].OutVal);
+		// Move in world space
+		FVector NewWorldPos = CurrentWorldPos + DeltaTranslate;
+		// Convert back to local space
+		SplineComp->SplineInfo.Points[SelectedKeyIndex].OutVal = SplineComp->ComponentToWorld.InverseTransformPosition(NewWorldPos);
+
+		return true;
+	}
+
+	return false;
+}
+
+bool FSplineComponentVisualizer::HandleInputKey(FLevelEditorViewportClient* ViewportClient, FViewport* Viewport, FKey Key, EInputEvent Event)
+{
+	bool bHandled = false;
+
+	if(Key == EKeys::LeftMouseButton && Event == IE_Released)
+	{
+		// Reset duplication flag on LMB release
+		bAllowDuplication = true;
+	}
+	else if(Key == EKeys::Delete && Event == IE_Pressed)
+	{
+		// Delete selected key
+		USplineComponent* SplineComp = GetEditedSplineComponent();
+		if (SplineComp &&
+			SelectedKeyIndex != INDEX_NONE &&
+			SelectedKeyIndex < SplineComp->SplineInfo.Points.Num())
+		{
+			SplineComp->SplineInfo.Points.RemoveAt(SelectedKeyIndex);
+
+			bHandled = true;
+		}
+	}
+
+	return bHandled;
+}
+
+
+void FSplineComponentVisualizer::EndEditing()
+{
+	SplineOwningActor = NULL;
+	SplineCompPropName = NAME_None;
+	SelectedKeyIndex = INDEX_NONE;
+}
\ No newline at end of file
diff --git a/Engine/Source/Editor/ComponentVisualizers/Private/SplineComponentVisualizer.h b/Engine/Source/Editor/ComponentVisualizers/Private/SplineComponentVisualizer.h
--- a/Engine/Source/Editor/ComponentVisualizers/Private/SplineComponentVisualizer.h
+++ b/Engine/Source/Editor/ComponentVisualizers/Private/SplineComponentVisualizer.h
@@ -4,10 +4,53 @@
 
 #include "ComponentVisualizer.h"
 
+struct HSplineVisProxy : public HComponentVisProxy
+{
+	DECLARE_HIT_PROXY();
+
+	HSplineVisProxy(const UActorComponent* InComponent)
+	: HComponentVisProxy(InComponent, HPP_Wireframe)
+	{}
+};
+
+struct HSplineKeyProxy : public HSplineVisProxy
+{
+	DECLARE_HIT_PROXY();
+
+	HSplineKeyProxy(const UActorComponent* InComponent, int32 InKeyIndex) 
+	: HSplineVisProxy(InComponent)
+	, KeyIndex(InKeyIndex)
+	{}
+
+
+	int32 KeyIndex;
+};
+
 class FSplineComponentVisualizer : public FComponentVisualizer
 {
 public:
+	FSplineComponentVisualizer()
+	: FComponentVisualizer()
+	, SelectedKeyIndex(INDEX_NONE)
+	{}
+
 	// Begin FComponentVisualizer interface
 	virtual void DrawVisualization(const UActorComponent* Component, const FSceneView* View, FPrimitiveDrawInterface* PDI) OVERRIDE;
+	virtual bool VisProxyHandleClick(HComponentVisProxy* VisProxy) OVERRIDE;
+	virtual void EndEditing() OVERRIDE;
+	virtual bool GetWidgetLocation(FVector& OutLocation) OVERRIDE;
+	virtual bool HandleInputDelta(FLevelEditorViewportClient* ViewportClient, FViewport* Viewport, FVector& DeltaTranslate, FRotator& DeltalRotate, FVector& DeltaScale) OVERRIDE;
+	virtual bool HandleInputKey(FLevelEditorViewportClient* ViewportClient, FViewport* Viewport, FKey Key, EInputEvent Event) OVERRIDE;
 	// End FComponentVisualizer interface
+
+	/** */
+	USplineComponent* GetEditedSplineComponent();
+
+private:
+
+	TWeakObjectPtr<AActor> SplineOwningActor;
+	FName SplineCompPropName;
+	int32 SelectedKeyIndex;
+
+	bool bAllowDuplication;
 };
diff --git a/Engine/Source/Editor/Kismet/Private/SCSEditorViewportClient.cpp b/Engine/Source/Editor/Kismet/Private/SCSEditorViewportClient.cpp
--- a/Engine/Source/Editor/Kismet/Private/SCSEditorViewportClient.cpp
+++ b/Engine/Source/Editor/Kismet/Private/SCSEditorViewportClient.cpp
@@ -234,10 +234,10 @@ void FSCSEditorViewportClient::Draw(const FSceneView* View, FPrimitiveDrawInterf
 				if(Comp != NULL && Comp->IsRegistered())
 				{
 					// Try and find a visualizer
-					TSharedPtr<FComponentVisualizer>* VisualizerPtr = GUnrealEd->ComponentVisualizerMap.Find(Comp->GetClass()->GetFName());
-					if (VisualizerPtr != NULL && (*VisualizerPtr).IsValid())
+					TSharedPtr<FComponentVisualizer> Visualizer = GUnrealEd->FindComponentVisualizer(Comp->GetClass()->GetFName());
+					if (Visualizer.IsValid())
 					{
-						(*VisualizerPtr)->DrawVisualization(Comp, View, PDI);
+						Visualizer->DrawVisualization(Comp, View, PDI);
 					}
 				}
 			}
diff --git a/Engine/Source/Editor/UnrealEd/Classes/Editor/UnrealEdEngine.h b/Engine/Source/Editor/UnrealEd/Classes/Editor/UnrealEdEngine.h
--- a/Engine/Source/Editor/UnrealEd/Classes/Editor/UnrealEdEngine.h
+++ b/Engine/Source/Editor/UnrealEd/Classes/Editor/UnrealEdEngine.h
@@ -218,9 +218,10 @@ class UNREALED_API UUnrealEdEngine : public UEditorEngine, public FNotifyHook
 
 	/** Register a function to draw extra information when a particular component is selected */
 	void RegisterComponentVisualizer(FName ComponentClassName, TSharedPtr<class FComponentVisualizer> Visualizer);
-
 	/** Unregister component visualizer function */
 	void UnregisterComponentVisualizer(FName ComponentClassName);
+	/** Find a component visualizer for the given component class name */
+	TSharedPtr<class FComponentVisualizer> FindComponentVisualizer(FName ComponentClassName);
 
 	/** Draw component visualizers for components for selected actors */
 	void DrawComponentVisualizers(const FSceneView* View, FPrimitiveDrawInterface* PDI);
diff --git a/Engine/Source/Editor/UnrealEd/Private/ComponentVisualizer.cpp b/Engine/Source/Editor/UnrealEd/Private/ComponentVisualizer.cpp
new file mode 100644
--- /dev/null
+++ b/Engine/Source/Editor/UnrealEd/Private/ComponentVisualizer.cpp
@@ -0,0 +1,50 @@
+// Copyright 1998-2014 Epic Games, Inc. All Rights Reserved.
+
+#include "UnrealEd.h"
+#include "ComponentVisualizer.h"
+
+IMPLEMENT_HIT_PROXY(HComponentVisProxy, HHitProxy);
+
+
+FName FComponentVisualizer::GetComponentPropertyName(const UActorComponent* Component)
+{
+	FName PropertyName = NAME_None;
+	if (Component)
+	{
+		const AActor* CompOwner = Component->GetOwner();
+		if(CompOwner)
+		{
+			// Iterate over fields of this actor
+			UClass* SplineActorClass = CompOwner->GetClass();
+			for (TFieldIterator<UObjectProperty> It(SplineActorClass); It; ++It)
+			{
+				// See if this property points to the component in question
+				UObjectProperty* ObjectProp = *It;
+				UObject* Object = ObjectProp->GetObjectPropertyValue(ObjectProp->ContainerPtrToValuePtr<void>(CompOwner));
+				if (Object == Component)
+				{
+					// It does! Remember this name
+					PropertyName = ObjectProp->GetFName();
+				}
+			}
+		}
+	}
+	return PropertyName;
+}
+
+UActorComponent* FComponentVisualizer::GetComponentFromPropertyName(const AActor* CompOwner, FName PropertyName)
+{
+	UActorComponent* ResultComp = NULL;
+	if(CompOwner && PropertyName != NAME_None)
+	{
+		UClass* ActorClass = CompOwner->GetClass();
+		UObjectProperty* ObjectProp = FindField<UObjectProperty>(ActorClass, PropertyName);
+		if(ObjectProp)
+		{
+			UObject* Object = ObjectProp->GetObjectPropertyValue(ObjectProp->ContainerPtrToValuePtr<void>(CompOwner));
+			ResultComp = Cast<UActorComponent>(Object);
+		}
+
+	}
+	return ResultComp;
+}
diff --git a/Engine/Source/Editor/UnrealEd/Private/EditorModes.cpp b/Engine/Source/Editor/UnrealEd/Private/EditorModes.cpp
--- a/Engine/Source/Editor/UnrealEd/Private/EditorModes.cpp
+++ b/Engine/Source/Editor/UnrealEd/Private/EditorModes.cpp
@@ -19,7 +19,7 @@
 #include "Editor/FoliageEdit/Public/FoliageEditModule.h"
 #include "ActorEditorUtils.h"
 #include "EditorStyle.h"
-
+#include "ComponentVisualizer.h"
 
 DEFINE_LOG_CATEGORY_STATIC(LogEditorModes, Log, All);
 
@@ -278,6 +278,14 @@ bool FEdMode::CapturedMouseMove( FLevelEditorViewportClient* InViewportClient, F
 
 bool FEdMode::InputKey(FLevelEditorViewportClient* ViewportClient, FViewport* Viewport, FKey Key, EInputEvent Event)
 {
+	if (EditedVisualizer.IsValid())
+	{
+		if (EditedVisualizer.Pin()->HandleInputKey(ViewportClient, Viewport, Key, Event))
+		{
+			return true;
+		}
+	}
+
 	if( GetCurrentTool() && GetCurrentTool()->InputKey( ViewportClient, Viewport, Key, Event ) )
 	{
 		return true;
@@ -311,6 +319,14 @@ bool FEdMode::InputAxis(FLevelEditorViewportClient* InViewportClient, FViewport*
 
 bool FEdMode::InputDelta(FLevelEditorViewportClient* InViewportClient, FViewport* InViewport, FVector& InDrag, FRotator& InRot, FVector& InScale)
 {
+	if (EditedVisualizer.IsValid() && InViewportClient->GetCurrentWidgetAxis() != EAxisList::None)
+	{
+		if(EditedVisualizer.Pin()->HandleInputDelta(InViewportClient, InViewport, InDrag, InRot, InScale))
+		{
+			return true;
+		}
+	}
+	
 	if(UsesPropertyWidgets())
 	{
 		AActor* SelectedActor = GetFirstSelectedActorInstance();
@@ -414,6 +430,16 @@ bool FEdMode::UsesTransformWidget(FWidget::EWidgetMode CheckMode) const
 
 FVector FEdMode::GetWidgetLocation() const
 {
+	if(EditedVisualizer.IsValid())
+	{
+		FVector VisualizerWidgetPos;
+		bool bSetWidgetLocation = EditedVisualizer.Pin()->GetWidgetLocation(VisualizerWidgetPos);
+		if(bSetWidgetLocation)
+		{
+			return VisualizerWidgetPos;
+		}
+	}
+
 	if(UsesPropertyWidgets())
 	{
 		AActor* SelectedActor = GetFirstSelectedActorInstance();
@@ -519,6 +545,13 @@ void FEdMode::ActorSelectionChangeNotify()
 	EditedPropertyName = TEXT("");
 	EditedPropertyIndex = INDEX_NONE;
 	bEditedPropertyIsTransform = false;
+
+	// Clear active editing visualizer on selection change
+	if (EditedVisualizer.IsValid())
+	{
+		EditedVisualizer.Pin()->EndEditing();
+		EditedVisualizer = NULL;
+	}
 }
 
 bool FEdMode::HandleClick(FLevelEditorViewportClient* InViewportClient, HHitProxy *HitProxy, const FViewportClick &Click)
@@ -533,12 +566,41 @@ bool FEdMode::HandleClick(FLevelEditorViewportClient* InViewportClient, HHitProx
 			bEditedPropertyIsTransform = PropertyProxy->bPropertyIsTransform;
 			return true;
 		}
+		else if(HitProxy->IsA(HComponentVisProxy::StaticGetType()))
+		{
+			HComponentVisProxy* VisProxy = (HComponentVisProxy*)HitProxy;
+			const UActorComponent* ClickedComponent = VisProxy->Component.Get();
+			if(ClickedComponent != NULL)
+			{
+				TSharedPtr<FComponentVisualizer> Visualizer = GUnrealEd->FindComponentVisualizer(ClickedComponent->GetClass()->GetFName());
+				if(Visualizer.IsValid())
+				{
+					bool bIsActive = Visualizer->VisProxyHandleClick(VisProxy);
+					if(bIsActive)
+					{
+						// call EndEditing on any currently edited visualizer, if we are going to change it
+						if (EditedVisualizer.IsValid() && Visualizer.Get() != EditedVisualizer.Pin().Get())
+						{
+							EditedVisualizer.Pin()->EndEditing();
+						}
+
+						EditedVisualizer = Visualizer;
+					}
+				}
+			}
+		}
 		// Left clicking on an actor, stop editing a property
 		else if( HitProxy->IsA(HActor::StaticGetType()) )
 		{
 			EditedPropertyName = TEXT("");
 			EditedPropertyIndex = INDEX_NONE;
 			bEditedPropertyIsTransform = false;
+
+			if(EditedVisualizer.IsValid())
+			{
+				EditedVisualizer.Pin()->EndEditing();
+				EditedVisualizer = NULL;
+			}
 		}
 	}
 
diff --git a/Engine/Source/Editor/UnrealEd/Private/UnrealEdEngine.cpp b/Engine/Source/Editor/UnrealEd/Private/UnrealEdEngine.cpp
--- a/Engine/Source/Editor/UnrealEd/Private/UnrealEdEngine.cpp
+++ b/Engine/Source/Editor/UnrealEd/Private/UnrealEdEngine.cpp
@@ -1016,6 +1016,19 @@ void UUnrealEdEngine::UnregisterComponentVisualizer(FName ComponentClassName)
 	ComponentVisualizerMap.Remove(ComponentClassName);
 }
 
+TSharedPtr<FComponentVisualizer> UUnrealEdEngine::FindComponentVisualizer(FName ComponentClassName)
+{
+	TSharedPtr<FComponentVisualizer> Visualizer = NULL;
+
+	TSharedPtr<FComponentVisualizer>* VisualizerPtr = ComponentVisualizerMap.Find(ComponentClassName);
+	if(VisualizerPtr != NULL)
+	{
+		Visualizer = *VisualizerPtr;
+	}
+
+	return Visualizer;
+}
+
 
 void UUnrealEdEngine::DrawComponentVisualizers(const FSceneView* View, FPrimitiveDrawInterface* PDI)
 {
@@ -1035,10 +1048,10 @@ void UUnrealEdEngine::DrawComponentVisualizers(const FSceneView* View, FPrimitiv
 				if(Comp->IsRegistered())
 				{
 					// Try and find a visualizer
-					TSharedPtr<FComponentVisualizer>* VisualizerPtr = ComponentVisualizerMap.Find(Comp->GetClass()->GetFName());
-					if(VisualizerPtr != NULL && (*VisualizerPtr).IsValid())
+					TSharedPtr<FComponentVisualizer> Visualizer = FindComponentVisualizer(Comp->GetClass()->GetFName());
+					if(Visualizer.IsValid())
 					{
-						(*VisualizerPtr)->DrawVisualization(Comp, View, PDI);
+						Visualizer->DrawVisualization(Comp, View, PDI);
 					}
 				}
 			}
diff --git a/Engine/Source/Editor/UnrealEd/Public/ComponentVisualizer.h b/Engine/Source/Editor/UnrealEd/Public/ComponentVisualizer.h
--- a/Engine/Source/Editor/UnrealEd/Public/ComponentVisualizer.h
+++ b/Engine/Source/Editor/UnrealEd/Public/ComponentVisualizer.h
@@ -2,10 +2,42 @@
 
 #pragma once
 
+struct HComponentVisProxy : public HHitProxy
+{
+	DECLARE_HIT_PROXY(UNREALED_API);
+
+	HComponentVisProxy(const UActorComponent* InComponent, EHitProxyPriority InPriority = HPP_Wireframe) 
+	: HHitProxy(InPriority)
+	, Component(InComponent)
+	{}
+
+	virtual EMouseCursor::Type GetMouseCursor()
+	{
+		return EMouseCursor::Crosshairs;
+	}
+
+	TWeakObjectPtr<const UActorComponent> Component;
+};
+
 /** Base class for a component visualizer, that draw editor information for a particular component class */
 class UNREALED_API FComponentVisualizer : public TSharedFromThis<FComponentVisualizer>
 {
 public:
 	/** Draw visualization for the supplied component */
-	virtual void DrawVisualization(const UActorComponent* Component, const FSceneView* View, FPrimitiveDrawInterface* PDI) = 0;
+	virtual void DrawVisualization(const UActorComponent* Component, const FSceneView* View, FPrimitiveDrawInterface* PDI) {}
+	/** */
+	virtual bool VisProxyHandleClick(HComponentVisProxy* VisProxy) { return false; }
+	/** */
+	virtual void EndEditing() {}
+	/** */
+	virtual bool GetWidgetLocation(FVector& OutLocation) { return false; }
+	/** */
+	virtual bool HandleInputDelta(FLevelEditorViewportClient* ViewportClient, FViewport* Viewport, FVector& DeltaTranslate, FRotator& DeltalRotate, FVector& DeltaScale) { return false; }
+	/** */
+	virtual bool HandleInputKey(FLevelEditorViewportClient* ViewportClient,FViewport* Viewport,FKey Key,EInputEvent Event) { return false; }
+
+	/** Find the name of the property that points to this component */
+	static FName GetComponentPropertyName(const UActorComponent* Component);
+	/** Get a component pointer from the property name */
+	static UActorComponent* GetComponentFromPropertyName(const AActor* CompOwner, FName PropertyName);
 };
\ No newline at end of file
diff --git a/Engine/Source/Editor/UnrealEd/Public/EditorModes.h b/Engine/Source/Editor/UnrealEd/Public/EditorModes.h
--- a/Engine/Source/Editor/UnrealEd/Public/EditorModes.h
+++ b/Engine/Source/Editor/UnrealEd/Public/EditorModes.h
@@ -398,6 +398,9 @@ protected:
 	 */
 	void GetPropertyWidgetInfos(const UStruct* InStruct, const void* InContainer, TArray<FPropertyWidgetInfo>& OutInfos, FString PropertyNamePrefix = TEXT("")) const;
 
+	/** Currently 'active' visualizer that we should pass input to etc */
+	TWeakPtr<FComponentVisualizer> EditedVisualizer;
+
 	/** Name of the property currently being edited */
 	FString EditedPropertyName;
 
diff --git a/Engine/Source/Runtime/Engine/Classes/Components/SplineComponent.h b/Engine/Source/Runtime/Engine/Classes/Components/SplineComponent.h
--- a/Engine/Source/Runtime/Engine/Classes/Components/SplineComponent.h
+++ b/Engine/Source/Runtime/Engine/Classes/Components/SplineComponent.h
@@ -4,8 +4,31 @@
 #pragma once
 #include "SplineComponent.generated.h"
 
-UCLASS(MinimalAPI)
-class USplineComponent : public UPrimitiveComponent
+/** Used to store spline data during RerunConstructionScripts */
+class FSplineInstanceData : public FComponentInstanceDataBase
+{
+public:
+	static const FName SplineInstanceDataTypeName;
+
+	virtual ~FSplineInstanceData()
+	{}
+
+	// Begin FComponentInstanceDataBase interface
+	virtual FName GetDataTypeName() const OVERRIDE
+	{
+		return SplineInstanceDataTypeName;
+	}
+	// End FComponentInstanceDataBase interface
+
+	/** Name of component */
+	FName ComponentName;
+
+	/** Lightmaps from LODData */
+	FInterpCurveVector SplineInfo;
+};
+
+UCLASS()
+class ENGINE_API USplineComponent : public USceneComponent
 {
 	GENERATED_UCLASS_BODY()
 
@@ -13,53 +36,33 @@ class USplineComponent : public UPrimitiveComponent
 	UPROPERTY(EditAnywhere, BlueprintReadWrite, Category=SplineComponent)
 	FInterpCurveVector SplineInfo;
 
-	/**
-	 * This is how curvy this spline is.  1.0f is straight and anything below that is curvy!
-	 * We are doing a simplistic calculate of:  vsize(points) / Length Of Spline 
-	 **/
-	UPROPERTY(Category=SplineComponent, VisibleAnywhere, BlueprintReadWrite)
-	float SplineCurviness;
-
 	/** Color of spline */
 	UPROPERTY(EditAnywhere, BlueprintReadWrite, Category=SplineComponent)
 	FColor SplineColor;
 
-	/** Resolution to draw spline at */
-	UPROPERTY(EditAnywhere, BlueprintReadWrite, Category=SplineComponent)
-	float SplineDrawRes;
-
-	/** Size of arrow on end of spline. If zero, no arrow drawn */
-	UPROPERTY(EditAnywhere, BlueprintReadWrite, Category=SplineComponent)
-	float SplineArrowSize;
-
-	/** If true, this spline is for whatever reason disabled, and will be drawn in red. */
-	UPROPERTY(EditAnywhere, BlueprintReadWrite, Category=SplineComponent)
-	uint32 bSplineDisabled:1;
-
 	/** Input, distance along curve, output, parameter that puts you there. */
-	UPROPERTY(EditAnywhere, BlueprintReadWrite, Category=SplineComponent)
+	UPROPERTY()
 	FInterpCurveFloat SplineReparamTable;
 
-
-	// Begin UPrimitiveComponent interface.
-	virtual FPrimitiveSceneProxy* CreateSceneProxy() OVERRIDE;
-	virtual FBoxSphereBounds CalcBounds(const FTransform & LocalToWorld) const OVERRIDE;
-	// End UPrimitiveComponent interface.
-	
-	/** This will update the spline curviness value **/
-	virtual void UpdateSplineCurviness();
+	// Begin UActorComponent interface.
+	virtual void GetComponentInstanceData(FComponentInstanceDataCache& Cache) const OVERRIDE;
+	virtual void ApplyComponentInstanceData(const FComponentInstanceDataCache& Cache) OVERRIDE;
+	// End UActorComponent interface.
 
 	/** Update the SplineReparamTable */
-	virtual void UpdateSplineReparamTable();
+	void UpdateSplineReparamTable();
 
 	/** Returns total length along this spline */
-	virtual float GetSplineLength() const;
+	float GetSplineLength() const;
 	
 	/** Given a distance along the length of this spline, return the point in space where this puts you */
-	virtual FVector GetLocationAtDistanceAlongSpline(float Distance) const;
+	FVector GetLocationAtDistanceAlongSpline(float Distance) const;
 	
 	/** Given a distance along the length of this spline, return the direction of the spline there. Note, result is non-unit length. */
-	virtual FVector GetTangentAtDistanceAlongSpline(float Distance) const;
+	FVector GetTangentAtDistanceAlongSpline(float Distance) const;
+
+	/** Walk through keys and set time for each one */
+	void RefreshSplineInputs();
 };
 
 
diff --git a/Engine/Source/Runtime/Engine/Private/EngineUtils.cpp b/Engine/Source/Runtime/Engine/Private/EngineUtils.cpp
--- a/Engine/Source/Runtime/Engine/Private/EngineUtils.cpp
+++ b/Engine/Source/Runtime/Engine/Private/EngineUtils.cpp
@@ -14,7 +14,6 @@ DEFINE_LOG_CATEGORY_STATIC(LogEngineUtils, Log, All);
 IMPLEMENT_HIT_PROXY(HActor,HHitProxy)
 IMPLEMENT_HIT_PROXY(HBSPBrushVert,HHitProxy);
 IMPLEMENT_HIT_PROXY(HStaticMeshVert,HHitProxy);
-IMPLEMENT_HIT_PROXY(HSplineProxy,HHitProxy);
 IMPLEMENT_HIT_PROXY(HTranslucentActor,HActor)
 
 
diff --git a/Engine/Source/Runtime/Engine/Private/SplineComponent.cpp b/Engine/Source/Runtime/Engine/Private/SplineComponent.cpp
--- a/Engine/Source/Runtime/Engine/Private/SplineComponent.cpp
+++ b/Engine/Source/Runtime/Engine/Private/SplineComponent.cpp
@@ -7,176 +7,21 @@
 #include "EnginePrivate.h"
 
 
-// HSplineProxy
-void HSplineProxy::AddReferencedObjects( FReferenceCollector& Collector )
-{
-	Collector.AddReferencedObject( SplineComp );
-}	
-
-
-//////////////////////////////////////////////////////////////////////////
-// SPLINE SCENE PROXY
-
-/** Represents a USplineComponent to the scene manager. */
-class FSplineSceneProxy : public FPrimitiveSceneProxy
-{
-private:
-	FColor					SplineColor;
-	FInterpCurveVector		SplineInfo;
-	float					SplineDrawRes;
-	float					SplineArrowSize;
-	uint32					bSplineDisabled:1;
-
-public:
-
-	FSplineSceneProxy(USplineComponent* Component):
-		FPrimitiveSceneProxy(Component),
-		SplineColor(Component->SplineColor),
-		SplineInfo(Component->SplineInfo),
-		SplineDrawRes(Component->SplineDrawRes),
-		SplineArrowSize(Component->SplineArrowSize),
-		bSplineDisabled(Component->bSplineDisabled)
-	{
-	}
-
-	virtual HHitProxy* CreateHitProxies(UPrimitiveComponent* Component, TArray<TRefCountPtr<HHitProxy> >& OutHitProxies)
-	{
-		HSplineProxy* SplineHitProxy = new HSplineProxy( CastChecked<USplineComponent>(Component) );
-		OutHitProxies.Add(SplineHitProxy);
-		return SplineHitProxy;
-	}
-
-	virtual void DrawDynamicElements(FPrimitiveDrawInterface* PDI, const FSceneView* View)
-	{
-		QUICK_SCOPE_CYCLE_COUNTER( STAT_SplineSceneProxy_DrawDynamicElements );
-
-		// Override specified color with red if spline is 'disabled'
-		FColor UseColor = bSplineDisabled ? FColor(255,0,0) : SplineColor;
-		
-		//DrawDirectionalArrow(PDI, LocalToWorld, SplineColor, 1.0 * 3.0f, 1.0f, GetDepthPriorityGroup(View));		
-		FVector OldKeyPos(0);
-		float OldKeyTime = 0.f;
-		for(int32 i=0; i<SplineInfo.Points.Num(); i++)
-		{
-			float NewKeyTime = SplineInfo.Points[i].InVal;
-			FVector NewKeyPos = SplineInfo.Eval(NewKeyTime, FVector::ZeroVector);
-
-			// Draw the keypoint
-			PDI->DrawPoint(NewKeyPos, UseColor, 6.f, GetDepthPriorityGroup(View));
-
-			// If not the first keypoint, draw a line to the last keypoint.
-			if(i>0)
-			{
-				// For constant interpolation - don't draw ticks - just draw dotted line.
-				if(SplineInfo.Points[i-1].InterpMode == CIM_Constant)
-				{
-					DrawDashedLine(PDI, OldKeyPos, NewKeyPos, UseColor, 20, GetDepthPriorityGroup(View));
-				}
-				else
-				{
-					int32 NumSteps = FMath::CeilToInt( (NewKeyTime - OldKeyTime)/SplineDrawRes );
-					float DrawSubstep = (NewKeyTime - OldKeyTime)/NumSteps;
-
-					// Find position on first keyframe.
-					float OldTime = OldKeyTime;
-					FVector OldPos = OldKeyPos;
-
-					// Find step that we want to draw the arrow head on.
-					int32 ArrowStep = FMath::Max<int32>(0, NumSteps-2);
-					
-					// Then draw a line for each substep.
-					for(int32 j=1; j<NumSteps+1; j++)
-					{
-						float NewTime = OldKeyTime + j*DrawSubstep;
-						FVector NewPos = SplineInfo.Eval(NewTime, FVector::ZeroVector);
-
-						// If this is where we want to draw the arrow, do that
-						// Only draw arrow if SplineArrowSize > 0.0
-						if(j == ArrowStep && SplineArrowSize > KINDA_SMALL_NUMBER)
-						{
-							FMatrix ArrowMatrix;
-							FVector ArrowDir = (NewPos - OldPos);
-							float ArrowLen = ArrowDir.Size();
-							if(ArrowLen > KINDA_SMALL_NUMBER)
-							{
-								ArrowDir /= ArrowLen;
-							}
-
-							ArrowMatrix = FRotationTranslationMatrix(ArrowDir.Rotation(), OldPos);
-
-							DrawDirectionalArrow(PDI, ArrowMatrix, UseColor, ArrowLen, SplineArrowSize, GetDepthPriorityGroup(View));								
-						}
-						// Otherwise normal line
-						else
-						{
-							PDI->DrawLine(OldPos, NewPos, UseColor, GetDepthPriorityGroup(View));														
-						}
-
-						// Don't draw point for last one - its the keypoint drawn above.
-						if(j != NumSteps)
-						{
-							PDI->DrawPoint(NewPos, UseColor, 3.f, GetDepthPriorityGroup(View));
-						}
-
-						OldTime = NewTime;
-						OldPos = NewPos;
-					}
-				}
-			}
-
-			OldKeyTime = NewKeyTime;
-			OldKeyPos = NewKeyPos;
-		}
-	}
-
-	virtual FPrimitiveViewRelevance GetViewRelevance(const FSceneView* View)
-	{
-		FPrimitiveViewRelevance Result;
-
-		Result.bDrawRelevance = IsShown(View) && View->Family->EngineShowFlags.Splines;
-		Result.bDynamicRelevance = true;
-		Result.bEditorPrimitiveRelevance = UseEditorCompositing(View);
-
-		return Result;
-	}
-
-	virtual uint32 GetMemoryFootprint( void ) const { return( sizeof( *this ) + GetAllocatedSize() ); }
-	uint32 GetAllocatedSize( void ) const { return( FPrimitiveSceneProxy::GetAllocatedSize() ); }
-
-
-};
-
 USplineComponent::USplineComponent(const class FPostConstructInitializeProperties& PCIP)
 	: Super(PCIP)
 {
-	SplineDrawRes = 0.1f;
-	SplineArrowSize = 60.0f;
-	SplineColor.R = 255;
-	SplineColor.B = 255;
-	SplineColor.A = 255;
-
-	bUseEditorCompositing = true;
-	BodyInstance.bEnableCollision_DEPRECATED = false;
-
-	SetCollisionProfileName(UCollisionProfile::NoCollision_ProfileName);
-}
-
-
-FPrimitiveSceneProxy* USplineComponent::CreateSceneProxy()
-{
-	return new FSplineSceneProxy(this);
-}
-
-FBoxSphereBounds USplineComponent::CalcBounds(const FTransform & LocalToWorld) const
-{
-	FVector BoxMin, BoxMax;
+	SplineColor = FColor(255,255,255);
 
-	SplineInfo.CalcBounds(BoxMin, BoxMax, LocalToWorld.GetLocation());
+	// Add 2 keys by default
+	int32 PointIndex = SplineInfo.AddPoint(0.f, FVector(0,0,0));
+	FInterpCurvePoint<FVector>& Point = SplineInfo.Points[PointIndex];
+	Point.LeaveTangent = FVector(1,0,0);
 
-	return FBoxSphereBounds( FBox(BoxMin, BoxMax) );
+	PointIndex = SplineInfo.AddPoint(1.f, FVector(100,0,0));
+	Point = SplineInfo.Points[PointIndex];
+	Point.ArriveTangent = FVector(1, 0, 0);
 }
 
-
 void USplineComponent::UpdateSplineReparamTable()
 {
 	// Start by clearing it
@@ -217,23 +62,6 @@ void USplineComponent::UpdateSplineReparamTable()
 	}
 }
 
-
-void USplineComponent::UpdateSplineCurviness()
-{
-	const float SplineLength = GetSplineLength();
-
-	const FVector Start = GetLocationAtDistanceAlongSpline( 0.0f );;
-	const FVector End = GetLocationAtDistanceAlongSpline( SplineLength );
-
-	SplineCurviness =(Start - End).Size() / SplineLength;
-
-	//GetWorld()->PersistentLineBatcher->DrawLine( Start, Start+FVector(0,0,200), FColor(255,0,0), SDPG_World );
-	//GetWorld()->PersistentLineBatcher->DrawLine( End, End+FVector(0,0,200), FColor(255,0,0), SDPG_World );
-
-	//UE_LOG(LogSpline, Warning, TEXT( "hahah %s %f %f %f"), *GetName(), SplineCurviness, (Start - End).Size(), SplineLength );
-}
-
-
 float USplineComponent::GetSplineLength() const
 {
 	// This is given by the input of the last entry in the remap table
@@ -261,3 +89,43 @@ FVector USplineComponent::GetTangentAtDistanceAlongSpline(float Distance) const
 	return Tangent;
 }
 
+void USplineComponent::RefreshSplineInputs()
+{
+	for(int32 KeyIdx=0; KeyIdx<SplineInfo.Points.Num(); KeyIdx++)
+	{
+		SplineInfo.Points[KeyIdx].InVal = KeyIdx;
+	}
+}
+
+// Init type name static
+const FName FSplineInstanceData::SplineInstanceDataTypeName(TEXT("SplineInstanceData"));
+
+void USplineComponent::GetComponentInstanceData(FComponentInstanceDataCache& Cache) const
+{
+	TSharedRef<FSplineInstanceData> SplineData = MakeShareable(new FSplineInstanceData());
+	SplineData->ComponentName = GetFName();
+	SplineData->SplineInfo = SplineInfo;
+
+	// Add to cache
+	Cache.AddInstanceData(SplineData);
+}
+
+void USplineComponent::ApplyComponentInstanceData(const FComponentInstanceDataCache& Cache)
+{
+	TArray< TSharedPtr<FComponentInstanceDataBase> > CachedData;
+	Cache.GetInstanceDataOfType(FSplineInstanceData::SplineInstanceDataTypeName, CachedData);
+
+	for (int32 DataIdx = 0; DataIdx < CachedData.Num(); DataIdx++)
+	{
+		check(CachedData[DataIdx].IsValid());
+		check(CachedData[DataIdx]->GetDataTypeName() == FSplineInstanceData::SplineInstanceDataTypeName);
+		TSharedPtr<FSplineInstanceData> SplineData = StaticCastSharedPtr<FSplineInstanceData>(CachedData[DataIdx]);
+
+		// See if data matches current state
+		if (SplineData->ComponentName == GetFName())
+		{
+			SplineInfo = SplineData->SplineInfo;
+			UpdateSplineReparamTable();			
+		}
+	}
+}
\ No newline at end of file
diff --git a/Engine/Source/Runtime/Engine/Public/EngineUtils.h b/Engine/Source/Runtime/Engine/Public/EngineUtils.h
--- a/Engine/Source/Runtime/Engine/Public/EngineUtils.h
+++ b/Engine/Source/Runtime/Engine/Public/EngineUtils.h
@@ -102,27 +102,6 @@ struct HStaticMeshVert : public HHitProxy
 	}
 };
 
-//
-//	HSplineProxy
-//
-
-struct HSplineProxy : public HHitProxy
-{
-	DECLARE_HIT_PROXY( ENGINE_API );
-
-	class USplineComponent* SplineComp;
-
-	HSplineProxy(class USplineComponent* InSplineComp):
-		HHitProxy(HPP_World),
-		SplineComp(InSplineComp)
-	{}
-	virtual EMouseCursor::Type GetMouseCursor()
-	{
-		return EMouseCursor::Crosshairs;
-	}	
-	virtual void AddReferencedObjects( FReferenceCollector& Collector ) OVERRIDE;
-};
-
 // Hit an actor even with translucency
 struct HTranslucentActor : public HActor
 {
