Fix skeletal mesh setSimulated not properly updating bodies
TTP#335496

diff --git a/Engine/Source/Runtime/Engine/Private/PhysicsEngine/BodyInstance.cpp b/Engine/Source/Runtime/Engine/Private/PhysicsEngine/BodyInstance.cpp
--- a/Engine/Source/Runtime/Engine/Private/PhysicsEngine/BodyInstance.cpp
+++ b/Engine/Source/Runtime/Engine/Private/PhysicsEngine/BodyInstance.cpp
@@ -1443,6 +1443,8 @@ void FBodyInstance::UpdateInstanceSimulatePhysics(bool bIgnoreParent)
 		{
 			PhysicsBlendWeight = 0.f;
 		}
+
+		bSimulatePhysics = bUseSimulate;
 	}
 #endif
 }
