Fix TTP 335495 - LANDSCAPE: QAGAME: Camera is passing through spline geometry
It was defaulting to the ôno collisionö profile, which ignores visibility and camera traces. If you enabled collision, it wouldn't change the profile, leaving the responses list the same (so still ignoring camera/visibility traces), instead of changing to the ôblock allö profile (which correctly blocks camera and visibility traces).

diff --git a/Engine/Source/Runtime/Engine/Private/Landscape/LandscapeSplines.cpp b/Engine/Source/Runtime/Engine/Private/Landscape/LandscapeSplines.cpp
--- a/Engine/Source/Runtime/Engine/Private/Landscape/LandscapeSplines.cpp
+++ b/Engine/Source/Runtime/Engine/Private/Landscape/LandscapeSplines.cpp
@@ -493,6 +493,13 @@ void ULandscapeSplineControlPoint::PostLoad()
 #if WITH_EDITOR
 	if (MeshComponent != NULL)
 	{
+		// Fix collision profile
+		const FName CollisionProfile = bEnableCollision ? UCollisionProfile::BlockAll_ProfileName : UCollisionProfile::NoCollision_ProfileName;
+		if (MeshComponent->GetCollisionProfileName() != CollisionProfile)
+		{
+			MeshComponent->SetCollisionProfileName(CollisionProfile);
+		}
+
 		MeshComponent->SetFlags(RF_TextExportTransient);
 	}
 #endif
@@ -788,11 +795,11 @@ void ULandscapeSplineControlPoint::UpdateSplinePoints(bool bUpdateCollision)
 			MeshComponent->SetCastShadow(bCastShadow);
 		}
 
-		ECollisionEnabled::Type CollisionType = bEnableCollision ? ECollisionEnabled::QueryAndPhysics : ECollisionEnabled::NoCollision;
-		if (MeshComponent->BodyInstance.GetCollisionEnabled() != CollisionType)
+		const FName CollisionProfile = bEnableCollision ? UCollisionProfile::BlockAll_ProfileName : UCollisionProfile::NoCollision_ProfileName;
+		if (MeshComponent->BodyInstance.GetCollisionProfileName() != CollisionProfile)
 		{
 			MeshComponent->Modify();
-			MeshComponent->BodyInstance.SetCollisionEnabled(CollisionType);
+			MeshComponent->BodyInstance.SetCollisionProfileName(CollisionProfile);
 			bNavDirty = true;
 		}
 	}
@@ -1034,6 +1041,14 @@ void ULandscapeSplineSegment::PostLoad()
 
 	for (auto* MeshComponent : MeshComponents)
 	{
+		// Fix collision profile
+		const bool bUsingEditorMesh = MeshComponent->bHiddenInGame;
+		const FName CollisionProfile = (bEnableCollision && !bUsingEditorMesh) ? UCollisionProfile::BlockAll_ProfileName : UCollisionProfile::NoCollision_ProfileName;
+		if (MeshComponent->GetCollisionProfileName() != CollisionProfile)
+		{
+			MeshComponent->SetCollisionProfileName(CollisionProfile);
+		}
+
 		MeshComponent->SetFlags(RF_TextExportTransient);
 	}
 #endif
@@ -1566,8 +1581,7 @@ void ULandscapeSplineSegment::UpdateSplinePoints(bool bUpdateCollision)
 			MeshComponent->SetCastShadow(bCastShadow);
 			MeshComponent->InvalidateLightingCache();
 
-			MeshComponent->BodyInstance.bEnableCollision_DEPRECATED = bEnableCollision && !bUsingEditorMesh;
-			MeshComponent->BodyInstance.SetCollisionEnabled((bEnableCollision && !bUsingEditorMesh) ? ECollisionEnabled::QueryAndPhysics : ECollisionEnabled::NoCollision);
+			MeshComponent->BodyInstance.SetCollisionProfileName((bEnableCollision && !bUsingEditorMesh) ? UCollisionProfile::BlockAll_ProfileName : UCollisionProfile::NoCollision_ProfileName);
 
 #if WITH_EDITOR
 			if (bUpdateCollision)
