Fixed compile issue with 32 bit

diff --git a/Engine/Source/Runtime/Engine/Classes/Animation/Skeleton.h b/Engine/Source/Runtime/Engine/Classes/Animation/Skeleton.h
--- a/Engine/Source/Runtime/Engine/Classes/Animation/Skeleton.h
+++ b/Engine/Source/Runtime/Engine/Classes/Animation/Skeleton.h
@@ -549,8 +549,12 @@ public:
 
 	static const FName AnimNotifyTag;
 	static const TCHAR AnimNotifyTagDeliminator;
-	ENGINE_API static const FName DefaultSlotGroupName;
 #endif
+
+public:
+	// this should be outside of editor because slot node is initialize to it
+	ENGINE_API static const FName DefaultSlotGroupName;
+
 private:
 	void RegenerateGuid();
 };
diff --git a/Engine/Source/Runtime/Engine/Private/Animation/Skeleton.cpp b/Engine/Source/Runtime/Engine/Private/Animation/Skeleton.cpp
--- a/Engine/Source/Runtime/Engine/Private/Animation/Skeleton.cpp
+++ b/Engine/Source/Runtime/Engine/Private/Animation/Skeleton.cpp
@@ -14,9 +14,10 @@
 #if WITH_EDITOR
 const FName USkeleton::AnimNotifyTag = FName(TEXT("AnimNotifyList"));
 const TCHAR USkeleton::AnimNotifyTagDeliminator = TEXT(';');
-const FName USkeleton::DefaultSlotGroupName = FName(TEXT("Default"));
 #endif 
 
+const FName USkeleton::DefaultSlotGroupName = FName(TEXT("Default"));
+
 FArchive& operator<<(FArchive& Ar, FReferencePose & P)
 {
 	Ar << P.PoseName;
@@ -34,8 +35,10 @@ FArchive& operator<<(FArchive& Ar, FReferencePose & P)
 USkeleton::USkeleton(const class FPostConstructInitializeProperties& PCIP)
 	: Super(PCIP)
 {
+#if WITH_EDITORONLY_DATA
 	SlotGroupNames.Empty();
 	SlotGroupNames.Add(DefaultSlotGroupName);
+#endif
 }
 
 void USkeleton::PostInitProperties()
