#ue4 - Fix issue with FScopedMovementUpdate not restoring relative location/rotation/scale after RevertMove.

UDN: https://udn.unrealengine.com/questions/207599/fscopedmovementupdaterevertmove-doesnt-revert-rela.html

diff --git a/Engine/Source/Runtime/Engine/Classes/Components/SceneComponent.h b/Engine/Source/Runtime/Engine/Classes/Components/SceneComponent.h
--- a/Engine/Source/Runtime/Engine/Classes/Components/SceneComponent.h
+++ b/Engine/Source/Runtime/Engine/Classes/Components/SceneComponent.h
@@ -785,6 +785,10 @@ private:
 	uint32 bDeferUpdates:1;
 	uint32 bHasValidOverlapsAtEnd:1;
 	FTransform InitialTransform;
+	FVector InitialRelativeLocation;
+	FRotator InitialRelativeRotation;
+	FVector InitialRelativeScale;
+
 	TArray<struct FOverlapInfo> PendingOverlaps;
 	TArray<struct FOverlapInfo> OverlapsAtEnd;
 	TBlockingHitArray BlockingHits;
diff --git a/Engine/Source/Runtime/Engine/Private/SceneComponent.cpp b/Engine/Source/Runtime/Engine/Private/SceneComponent.cpp
--- a/Engine/Source/Runtime/Engine/Private/SceneComponent.cpp
+++ b/Engine/Source/Runtime/Engine/Private/SceneComponent.cpp
@@ -1542,6 +1542,9 @@ FScopedMovementUpdate::FScopedMovementUpdate( class USceneComponent* Component,
 	if (IsValid(Component))
 	{
 		InitialTransform = Component->GetComponentToWorld();
+		InitialRelativeLocation = Component->RelativeLocation;
+		InitialRelativeRotation = Component->RelativeRotation;
+		InitialRelativeScale = Component->RelativeScale3D;
 
 		if (ScopeBehavior == EScopedUpdate::ImmediateUpdates)
 		{
@@ -1605,6 +1608,9 @@ void FScopedMovementUpdate::RevertMove()
 		{
 			// Teleport to start
 			Component->ComponentToWorld = InitialTransform;
+			Component->RelativeLocation = InitialRelativeLocation;
+			Component->RelativeRotation = InitialRelativeRotation;
+			Component->RelativeScale3D = InitialRelativeScale;
 
 			if (!IsDeferringUpdates())
 			{
