Fix Group Name

diff --git a/Engine/Source/Runtime/Engine/Private/Animation/AnimNode_Slot.cpp b/Engine/Source/Runtime/Engine/Private/Animation/AnimNode_Slot.cpp
--- a/Engine/Source/Runtime/Engine/Private/Animation/AnimNode_Slot.cpp
+++ b/Engine/Source/Runtime/Engine/Private/Animation/AnimNode_Slot.cpp
@@ -86,8 +86,8 @@ void FAnimNode_Slot::GatherDebugData(FNodeDebugData& DebugData)
 }
 
 FAnimNode_Slot::FAnimNode_Slot()
-: SourceWeight(0.f)
+: GroupName(USkeleton::DefaultSlotGroupName)
+, SourceWeight(0.f)
 , SlotNodeWeight(0.f)
-, GroupName(USkeleton::DefaultSlotGroupName)
 {
 }
