Made non-string-pulled path draw its corridor for VLog #UE4

- and other minor improvements to debugging code

diff --git a/Engine/Source/Runtime/Engine/Private/Debug/GameplayDebuggingControllerComponent.cpp b/Engine/Source/Runtime/Engine/Private/Debug/GameplayDebuggingControllerComponent.cpp
--- a/Engine/Source/Runtime/Engine/Private/Debug/GameplayDebuggingControllerComponent.cpp
+++ b/Engine/Source/Runtime/Engine/Private/Debug/GameplayDebuggingControllerComponent.cpp
@@ -136,7 +136,9 @@ void UGameplayDebuggingControllerComponent::CycleActiveView()
 			break;
 	}
 	if (++Index >= EAIDebugDrawDataView::EditorDebugAIFlag)
+	{
 		Index = EAIDebugDrawDataView::Basic;
+	}
 
 	ActiveViews = (1 << EAIDebugDrawDataView::OverHead) | (1 << Index);
 }
