Fix for destructible chunks that share the same PxActor but have different Destructible Components

diff --git a/Engine/Source/Runtime/Engine/Private/PhysicsEngine/PhysScene.cpp b/Engine/Source/Runtime/Engine/Private/PhysicsEngine/PhysScene.cpp
--- a/Engine/Source/Runtime/Engine/Private/PhysicsEngine/PhysScene.cpp
+++ b/Engine/Source/Runtime/Engine/Private/PhysicsEngine/PhysScene.cpp
@@ -623,6 +623,7 @@ void FPhysScene::SyncComponentsToBodies(uint32 SceneType)
 		{
 			if (DestructibleChunkInfo->OwningComponent.IsValid())
 			{
+				//TODO: waiting on new API to avoid duplicate updates per shape.
 				TArray<PxShape*> Shapes;
 				Shapes.AddZeroed(RigidActor->getNbShapes());
 				int32 NumShapes = RigidActor->getShapes(Shapes.GetData(), Shapes.Num());
@@ -631,10 +632,12 @@ void FPhysScene::SyncComponentsToBodies(uint32 SceneType)
 					PxShape* Shape = Shapes[ShapeIdx];
 					int32 ChunkIndex;
 					NxDestructibleActor* DestructibleActor = GApexModuleDestructible->getDestructibleAndChunk(Shape, &ChunkIndex);
-					
 					const physx::PxMat44 ChunkPoseRT = DestructibleActor->getChunkPose(ChunkIndex);
 					const physx::PxTransform Transform(ChunkPoseRT);
-					DestructibleChunkInfo->OwningComponent->SetChunkWorldRT(ChunkIndex, P2UQuat(Transform.q), P2UVector(Transform.p));
+					if (UDestructibleComponent * DestructibleComponent = Cast<UDestructibleComponent>(FPhysxUserData::Get<UPrimitiveComponent>(Shape->userData)))
+					{
+						DestructibleComponent->SetChunkWorldRT(ChunkIndex, P2UQuat(Transform.q), P2UVector(Transform.p));
+					}
 				}
 			}
 		}
