AllowsFramerateSmoothing is now a platform proprty. Enabled for all platforms except XB1 and PS4

diff --git a/Engine/Source/Runtime/Core/Public/Android/AndroidProperties.h b/Engine/Source/Runtime/Core/Public/Android/AndroidProperties.h
--- a/Engine/Source/Runtime/Core/Public/Android/AndroidProperties.h
+++ b/Engine/Source/Runtime/Core/Public/Android/AndroidProperties.h
@@ -92,6 +92,11 @@ struct FAndroidPlatformProperties
 	{
 		return true;
 	}
+
+	static FORCEINLINE bool AllowsFramerateSmoothing()
+	{
+		return true;
+	}
 };
 
 struct FAndroid_PVRTCPlatformProperties : public FAndroidPlatformProperties
diff --git a/Engine/Source/Runtime/Core/Public/GenericPlatform/GenericPlatformProperties.h b/Engine/Source/Runtime/Core/Public/GenericPlatform/GenericPlatformProperties.h
--- a/Engine/Source/Runtime/Core/Public/GenericPlatform/GenericPlatformProperties.h
+++ b/Engine/Source/Runtime/Core/Public/GenericPlatform/GenericPlatformProperties.h
@@ -202,6 +202,14 @@ struct FGenericPlatformProperties
 		return false;
 	}
 
+	/**
+	 * Whether this platform wants to allow framerate smoothing or not.
+	 */
+	static FORCEINLINE bool AllowsFramerateSmoothing()
+	{
+		return true;
+	}
+
 	static FORCEINLINE bool SupportsHighQualityLightmaps()
 	{
 		return true;
diff --git a/Engine/Source/Runtime/Core/Public/HTML5/HTML5PlatformProperties.h b/Engine/Source/Runtime/Core/Public/HTML5/HTML5PlatformProperties.h
--- a/Engine/Source/Runtime/Core/Public/HTML5/HTML5PlatformProperties.h
+++ b/Engine/Source/Runtime/Core/Public/HTML5/HTML5PlatformProperties.h
@@ -77,4 +77,9 @@ struct FHTML5PlatformProperties
 	{
 		return false;
 	}
+
+	static FORCEINLINE bool AllowsFramerateSmoothing()
+	{
+		return true;
+	}
 };
diff --git a/Engine/Source/Runtime/Core/Public/IOS/IOSPlatformProperties.h b/Engine/Source/Runtime/Core/Public/IOS/IOSPlatformProperties.h
--- a/Engine/Source/Runtime/Core/Public/IOS/IOSPlatformProperties.h
+++ b/Engine/Source/Runtime/Core/Public/IOS/IOSPlatformProperties.h
@@ -72,4 +72,9 @@ struct FIOSPlatformProperties
 	{
 		return true;
 	}
+
+	static FORCEINLINE bool AllowsFramerateSmoothing()
+	{
+		return true;
+	}
 };
diff --git a/Engine/Source/Runtime/Core/Public/Linux/LinuxPlatformProperties.h b/Engine/Source/Runtime/Core/Public/Linux/LinuxPlatformProperties.h
--- a/Engine/Source/Runtime/Core/Public/Linux/LinuxPlatformProperties.h
+++ b/Engine/Source/Runtime/Core/Public/Linux/LinuxPlatformProperties.h
@@ -110,6 +110,11 @@ struct FLinuxPlatformProperties
 		return true;
 	}
 
+	static FORCEINLINE bool AllowsFramerateSmoothing()
+	{
+		return true;
+	}
+
 	static FORCEINLINE bool SupportsQuit()
 	{
 		return true;
diff --git a/Engine/Source/Runtime/Core/Public/Mac/MacPlatformProperties.h b/Engine/Source/Runtime/Core/Public/Mac/MacPlatformProperties.h
--- a/Engine/Source/Runtime/Core/Public/Mac/MacPlatformProperties.h
+++ b/Engine/Source/Runtime/Core/Public/Mac/MacPlatformProperties.h
@@ -98,6 +98,11 @@ struct FMacPlatformProperties
 		return true;
 	}
 
+	static FORCEINLINE bool AllowsFramerateSmoothing()
+	{
+		return true;
+	}
+
 	static FORCEINLINE bool HasFixedResolution()
 	{
 		return false;
diff --git a/Engine/Source/Runtime/Engine/Private/UnrealEngine.cpp b/Engine/Source/Runtime/Engine/Private/UnrealEngine.cpp
--- a/Engine/Source/Runtime/Engine/Private/UnrealEngine.cpp
+++ b/Engine/Source/Runtime/Engine/Private/UnrealEngine.cpp
@@ -5463,7 +5463,7 @@ float UEngine::GetMaxTickRate( float DeltaTime, bool bAllowFrameRateSmoothing )
 {
 	float MaxTickRate = 0;
 
-	if (FPlatformProperties::SupportsWindowedMode())
+	if (FPlatformProperties::AllowsFramerateSmoothing())
 	{
 		// Smooth the framerate if wanted. The code uses a simplistic running average. Other approaches, like reserving
 		// a percentage of time, ended up creating negative feedback loops in conjunction with GPU load and were abandonend.
