Removed old xenon bug check

diff --git a/Engine/Source/Runtime/Engine/Private/Materials/HLSLMaterialTranslator.h b/Engine/Source/Runtime/Engine/Private/Materials/HLSLMaterialTranslator.h
--- a/Engine/Source/Runtime/Engine/Private/Materials/HLSLMaterialTranslator.h
+++ b/Engine/Source/Runtime/Engine/Private/Materials/HLSLMaterialTranslator.h
@@ -2644,22 +2644,8 @@ protected:
 			return INDEX_NONE;
 		}
 
-		FString ExponentCode = CoerceParameter(Exponent,MCT_Float);
-		if (CodeChunks[MaterialProperty][ShaderFrequency][Exponent].UniformExpression && CodeChunks[MaterialProperty][ShaderFrequency][Exponent].UniformExpression->IsConstant())
-		{
-			//chop off the parenthesis
-			FString NumericPortion = ExponentCode.Mid(1, ExponentCode.Len() - 2);
-			float ExponentValue = FCString::Atof(*NumericPortion); 
-			//check if the power was 1.0f to work around a xenon HLSL compiler bug in the Feb XDK
-			//which incorrectly optimizes pow(x, 1.0f) as if it were pow(x, 0.0f).
-			if (fabs(ExponentValue - 1.0f) < (float)KINDA_SMALL_NUMBER)
-			{
-				return Base;
-			}
-		}
-
 		// use ClampedPow so artist are prevented to cause NAN creeping into the math
-		return AddCodeChunk(GetParameterType(Base),TEXT("ClampedPow(%s,%s)"),*GetParameterCode(Base),*ExponentCode);
+		return AddCodeChunk(GetParameterType(Base),TEXT("ClampedPow(%s,%s)"),*GetParameterCode(Base),*CoerceParameter(Exponent,MCT_Float));
 	}
 
 	virtual int32 SquareRoot(int32 X) OVERRIDE
