Quick fix for terrible dependency bug caused by the way AIModule has been plugged in #UE4

diff --git a/Engine/Source/Runtime/Engine/Engine.Build.cs b/Engine/Source/Runtime/Engine/Engine.Build.cs
--- a/Engine/Source/Runtime/Engine/Engine.Build.cs
+++ b/Engine/Source/Runtime/Engine/Engine.Build.cs
@@ -61,7 +61,6 @@ public class Engine : ModuleRules
 				"SlateCore",
 				"SlateReflector",
 				"VectorVM",
-                "AIModule",
 			}
 		);
 
