AIModule stub plugged in properly #UE4

diff --git a/Engine/Source/Runtime/Engine/Engine.Build.cs b/Engine/Source/Runtime/Engine/Engine.Build.cs
--- a/Engine/Source/Runtime/Engine/Engine.Build.cs
+++ b/Engine/Source/Runtime/Engine/Engine.Build.cs
@@ -61,9 +61,12 @@ public class Engine : ModuleRules
 				"SlateCore",
 				"SlateReflector",
 				"VectorVM",
+                "AIModule"
 			}
 		);
 
+        CircularlyReferencedDependentModules.Add("AIModule");
+
 		DynamicallyLoadedModuleNames.AddRange(
 			new string[]
 			{
@@ -89,7 +92,6 @@ public class Engine : ModuleRules
 			PrivateIncludePathModuleNames.Add("MeshUtilities");
 			DynamicallyLoadedModuleNames.Add("MeshUtilities");
 
-
 			if (Target.Configuration != UnrealTargetConfiguration.Shipping && Target.Configuration != UnrealTargetConfiguration.Test)
 			{
 				PrivateDependencyModuleNames.Add("CollisionAnalyzer");
