Landscape : Added LODFalloff to be propagated to sublevels

diff --git a/Engine/Source/Runtime/Engine/Private/Landscape/Landscape.cpp b/Engine/Source/Runtime/Engine/Private/Landscape/Landscape.cpp
--- a/Engine/Source/Runtime/Engine/Private/Landscape/Landscape.cpp
+++ b/Engine/Source/Runtime/Engine/Private/Landscape/Landscape.cpp
@@ -1257,6 +1257,7 @@ void ALandscapeProxy::GetSharedProperties(ALandscapeProxy* Landscape)
 		SubsectionSizeQuads = Landscape->SubsectionSizeQuads;
 		MaxLODLevel = Landscape->MaxLODLevel;
 		LODDistanceFactor = Landscape->LODDistanceFactor;
+		LODFalloff = Landscape->LODFalloff;
 		CollisionMipLevel = Landscape->CollisionMipLevel;
 		if (!LandscapeMaterial)
 		{
