Landscape : RecreateLandscapeCollision command now remove owned componeent to make sure there will be no orphans

diff --git a/Engine/Source/Runtime/Engine/Private/Landscape/LandscapeEdit.cpp b/Engine/Source/Runtime/Engine/Private/Landscape/LandscapeEdit.cpp
--- a/Engine/Source/Runtime/Engine/Private/Landscape/LandscapeEdit.cpp
+++ b/Engine/Source/Runtime/Engine/Private/Landscape/LandscapeEdit.cpp
@@ -3428,6 +3428,15 @@ void ALandscapeProxy::RecreateCollisionComponents()
 {
 	// Clear old CollisionComponent containers
 	CollisionComponents.Empty();
+
+	// Clear any Owned Collision Components
+	TArray<ULandscapeHeightfieldCollisionComponent*> CollisionComps;
+	GetComponents(CollisionComps);
+	for (int32 i = 0; i < CollisionComps.Num(); ++i)
+	{
+		RemoveOwnedComponent(CollisionComps[i]);
+	}
+
 	CollisionMipLevel = FMath::Clamp<int32>( CollisionMipLevel, 0, FMath::CeilLogTwo(SubsectionSizeQuads+1)-1 );
 	for (int32 i = 0; i < LandscapeComponents.Num(); ++i)
 	{
