Corrected bRebuildAtRuntime tooltip.
#ttp 335236 - LIVE: EDITOR: Tooltip under RecastNavMesh's Rebuild at Runtime option is incorrect
#branch UE4
#proj Runtime.Engine
#change Corrected comment text from which tooltip is autogenerated.

diff --git a/Engine/Source/Runtime/Engine/Classes/AI/Navigation/NavigationData.h b/Engine/Source/Runtime/Engine/Classes/AI/Navigation/NavigationData.h
--- a/Engine/Source/Runtime/Engine/Classes/AI/Navigation/NavigationData.h
+++ b/Engine/Source/Runtime/Engine/Classes/AI/Navigation/NavigationData.h
@@ -41,7 +41,7 @@ class ENGINE_API ANavigationData : public AActor
 	// game-time config
 	//----------------------------------------------------------------------//
 
-	/** Size of the tallest agent that will path with this navmesh. */
+	/** If true, the NavMesh can be dynamically rebuilt at runtime. */
 	UPROPERTY(EditAnywhere, Category=Runtime, config)
 	uint32 bRebuildAtRuntime:1;
 
