Removed EngineClasses.h #includes.

This is a necessary first step in moving to a include-what-you-use model.

diff --git a/Engine/Source/Runtime/Engine/Private/AI/BrainComponent.cpp b/Engine/Source/Runtime/Engine/Private/AI/BrainComponent.cpp
--- a/Engine/Source/Runtime/Engine/Private/AI/BrainComponent.cpp
+++ b/Engine/Source/Runtime/Engine/Private/AI/BrainComponent.cpp
@@ -1,7 +1,6 @@
 // Copyright 1998-2014 Epic Games, Inc. All Rights Reserved.
 
 #include "EnginePrivate.h"
-#include "EngineClasses.h"
 #include "VisualLog.h"
 
 const FName UBrainComponent::AIMessage_MoveFinished = TEXT("MoveFinished");
diff --git a/Engine/Source/Runtime/Engine/Private/Debug/GameplayDebuggingComponent.cpp b/Engine/Source/Runtime/Engine/Private/Debug/GameplayDebuggingComponent.cpp
--- a/Engine/Source/Runtime/Engine/Private/Debug/GameplayDebuggingComponent.cpp
+++ b/Engine/Source/Runtime/Engine/Private/Debug/GameplayDebuggingComponent.cpp
@@ -1,7 +1,6 @@
 // Copyright 1998-2014 Epic Games, Inc. All Rights Reserved.
 
 #include "EnginePrivate.h"
-#include "EngineClasses.h"
 #include "Net/UnrealNetwork.h"
 #include "AI/BehaviorTreeDelegates.h"
 #include "../AI/Navigation/RecastNavMeshGenerator.h"
diff --git a/Engine/Source/Runtime/Engine/Private/Debug/GameplayDebuggingControllerComponent.cpp b/Engine/Source/Runtime/Engine/Private/Debug/GameplayDebuggingControllerComponent.cpp
--- a/Engine/Source/Runtime/Engine/Private/Debug/GameplayDebuggingControllerComponent.cpp
+++ b/Engine/Source/Runtime/Engine/Private/Debug/GameplayDebuggingControllerComponent.cpp
@@ -1,7 +1,6 @@
 // Copyright 1998-2014 Epic Games, Inc. All Rights Reserved.
 
 #include "EnginePrivate.h"
-#include "EngineClasses.h"
 #include "Net/UnrealNetwork.h"
 #include "AI/BehaviorTreeDelegates.h"
 
diff --git a/Engine/Source/Runtime/Engine/Private/EdGraph/EdGraph.cpp b/Engine/Source/Runtime/Engine/Private/EdGraph/EdGraph.cpp
--- a/Engine/Source/Runtime/Engine/Private/EdGraph/EdGraph.cpp
+++ b/Engine/Source/Runtime/Engine/Private/EdGraph/EdGraph.cpp
@@ -1,7 +1,6 @@
 // Copyright 1998-2014 Epic Games, Inc. All Rights Reserved.
 
 #include "EnginePrivate.h"
-#include "EngineClasses.h"
 #include "BlueprintUtilities.h"
 #if WITH_EDITOR
 #include "Editor/UnrealEd/Public/Kismet2/BlueprintEditorUtils.h"
diff --git a/Engine/Source/Runtime/Engine/Private/EdGraph/EdGraphNode.cpp b/Engine/Source/Runtime/Engine/Private/EdGraph/EdGraphNode.cpp
--- a/Engine/Source/Runtime/Engine/Private/EdGraph/EdGraphNode.cpp
+++ b/Engine/Source/Runtime/Engine/Private/EdGraph/EdGraphNode.cpp
@@ -1,7 +1,6 @@
 // Copyright 1998-2014 Epic Games, Inc. All Rights Reserved.
 
 #include "EnginePrivate.h"
-#include "EngineClasses.h"
 #include "BlueprintUtilities.h"
 #if WITH_EDITOR
 #include "Editor/UnrealEd/Public/Kismet2/BlueprintEditorUtils.h"
diff --git a/Engine/Source/Runtime/Engine/Private/EdGraph/EdGraphNode_Comment.cpp b/Engine/Source/Runtime/Engine/Private/EdGraph/EdGraphNode_Comment.cpp
--- a/Engine/Source/Runtime/Engine/Private/EdGraph/EdGraphNode_Comment.cpp
+++ b/Engine/Source/Runtime/Engine/Private/EdGraph/EdGraphNode_Comment.cpp
@@ -1,7 +1,6 @@
 // Copyright 1998-2014 Epic Games, Inc. All Rights Reserved.
 
 #include "EnginePrivate.h"
-#include "EngineClasses.h"
 #include "BlueprintUtilities.h"
 #if WITH_EDITOR
 #include "Editor/UnrealEd/Public/Kismet2/BlueprintEditorUtils.h"
diff --git a/Engine/Source/Runtime/Engine/Private/EdGraph/EdGraphPin.cpp b/Engine/Source/Runtime/Engine/Private/EdGraph/EdGraphPin.cpp
--- a/Engine/Source/Runtime/Engine/Private/EdGraph/EdGraphPin.cpp
+++ b/Engine/Source/Runtime/Engine/Private/EdGraph/EdGraphPin.cpp
@@ -1,7 +1,6 @@
 // Copyright 1998-2014 Epic Games, Inc. All Rights Reserved.
 
 #include "EnginePrivate.h"
-#include "EngineClasses.h"
 #include "BlueprintUtilities.h"
 #if WITH_EDITOR
 #include "Editor/UnrealEd/Public/Kismet2/BlueprintEditorUtils.h"
diff --git a/Engine/Source/Runtime/Engine/Private/EdGraph/EdGraphSchema.cpp b/Engine/Source/Runtime/Engine/Private/EdGraph/EdGraphSchema.cpp
--- a/Engine/Source/Runtime/Engine/Private/EdGraph/EdGraphSchema.cpp
+++ b/Engine/Source/Runtime/Engine/Private/EdGraph/EdGraphSchema.cpp
@@ -1,7 +1,6 @@
 // Copyright 1998-2014 Epic Games, Inc. All Rights Reserved.
 
 #include "EnginePrivate.h"
-#include "EngineClasses.h"
 #include "BlueprintUtilities.h"
 #if WITH_EDITOR
 #include "Editor/UnrealEd/Public/Kismet2/BlueprintEditorUtils.h"
diff --git a/Engine/Source/Runtime/Engine/Private/Interpolation.cpp b/Engine/Source/Runtime/Engine/Private/Interpolation.cpp
--- a/Engine/Source/Runtime/Engine/Private/Interpolation.cpp
+++ b/Engine/Source/Runtime/Engine/Private/Interpolation.cpp
@@ -5,7 +5,6 @@
 =============================================================================*/
 
 #include "EnginePrivate.h"
-#include "EngineClasses.h"
 #include "ParticleDefinitions.h"
 #include "SoundDefinitions.h"
 #include "InterpolationHitProxy.h"
diff --git a/Engine/Source/Runtime/Engine/Private/LocalMessage.cpp b/Engine/Source/Runtime/Engine/Private/LocalMessage.cpp
--- a/Engine/Source/Runtime/Engine/Private/LocalMessage.cpp
+++ b/Engine/Source/Runtime/Engine/Private/LocalMessage.cpp
@@ -1,7 +1,6 @@
 // Copyright 1998-2014 Epic Games, Inc. All Rights Reserved.
 
 #include "EnginePrivate.h"
-#include "EngineClasses.h"
 
 FClientReceiveData::FClientReceiveData()
 	: LocalPC(NULL)
diff --git a/Engine/Source/Runtime/Engine/Private/PhysicsEngine/PhysLevel.cpp b/Engine/Source/Runtime/Engine/Private/PhysicsEngine/PhysLevel.cpp
--- a/Engine/Source/Runtime/Engine/Private/PhysicsEngine/PhysLevel.cpp
+++ b/Engine/Source/Runtime/Engine/Private/PhysicsEngine/PhysLevel.cpp
@@ -1,7 +1,6 @@
 // Copyright 1998-2014 Epic Games, Inc. All Rights Reserved.
 
 #include "EnginePrivate.h"
-#include "EngineClasses.h"
 #include "ParticleDefinitions.h"
 #include "PrecomputedLightVolume.h"
 
diff --git a/Engine/Source/Runtime/Engine/Private/StreamableManager.cpp b/Engine/Source/Runtime/Engine/Private/StreamableManager.cpp
--- a/Engine/Source/Runtime/Engine/Private/StreamableManager.cpp
+++ b/Engine/Source/Runtime/Engine/Private/StreamableManager.cpp
@@ -1,7 +1,6 @@
 // Copyright 1998-2014 Epic Games, Inc. All Rights Reserved.
 
 #include "EnginePrivate.h"
-#include "EngineClasses.h"
 
 DEFINE_LOG_CATEGORY_STATIC(LogStreamableManager, Log, All);
 
diff --git a/Engine/Source/Runtime/Engine/Private/Tests/ExternalToolTest.cpp b/Engine/Source/Runtime/Engine/Private/Tests/ExternalToolTest.cpp
--- a/Engine/Source/Runtime/Engine/Private/Tests/ExternalToolTest.cpp
+++ b/Engine/Source/Runtime/Engine/Private/Tests/ExternalToolTest.cpp
@@ -1,7 +1,6 @@
 // Copyright 1998-2014 Epic Games, Inc. All Rights Reserved.
 
 #include "EnginePrivate.h"
-#include "EngineClasses.h"
 #include "AutomationTestCommon.h"
 
 DEFINE_LOG_CATEGORY_STATIC(LogBaseAutomationTests, Log, All);
diff --git a/Engine/Source/Runtime/Engine/Public/Engine.h b/Engine/Source/Runtime/Engine/Public/Engine.h
--- a/Engine/Source/Runtime/Engine/Public/Engine.h
+++ b/Engine/Source/Runtime/Engine/Public/Engine.h
@@ -601,7 +601,934 @@ struct FURL;
 #include "RenderUtils.h"					// Render utility classes.
 #include "Distributions.h"					// Distributions
 #include "PhysxUserData.h"
-#include "EngineClasses.h"					// All actor classes.
+
+// EngineClasses.h
+#include "Engine/LatentActionManager.h"
+#include "Engine/Scene.h"
+#include "Camera/CameraTypes.h"
+#include "Engine/EngineBaseTypes.h"
+#include "Engine/NetSerialization.h"
+#include "Engine/EngineTypes.h"
+#include "Interfaces/Interface_AssetUserData.h"
+#include "Components/ActorComponent.h"
+#include "GameFramework/Actor.h"
+#include "Sound/AmbientSound.h"
+#include "Engine/Brush.h"
+#include "Engine/BrushShape.h"
+#include "GameFramework/Volume.h"
+#include "Engine/BlockingVolume.h"
+#include "GameFramework/CameraBlockingVolume.h"
+#include "Engine/CullDistanceVolume.h"
+#include "Engine/LevelStreamingVolume.h"
+#include "Lightmass/LightmassCharacterIndirectDetailVolume.h"
+#include "Lightmass/LightmassImportanceVolume.h"
+#include "AI/Navigation/NavMeshBoundsVolume.h"
+#include "AI/Navigation/NavRelevantActorInterface.h"
+#include "AI/Navigation/NavModifierVolume.h"
+#include "GameFramework/PhysicsVolume.h"
+#include "GameFramework/DefaultPhysicsVolume.h"
+#include "GameFramework/KillZVolume.h"
+#include "GameFramework/PainCausingVolume.h"
+#include "Interfaces/Interface_PostProcessVolume.h"
+#include "Engine/PostProcessVolume.h"
+#include "Lightmass/PrecomputedVisibilityOverrideVolume.h"
+#include "Lightmass/PrecomputedVisibilityVolume.h"
+#include "Sound/ReverbVolume.h"
+#include "Engine/TriggerVolume.h"
+#include "GameFramework/BugRepro.h"
+#include "Camera/CameraActor.h"
+#include "AI/Navigation/NavAgentInterface.h"
+#include "GameFramework/Controller.h"
+#include "AI/AITypes.h"
+#include "AI/Navigation/NavLinkDefinition.h"
+#include "AI/Navigation/NavigationTypes.h"
+#include "Kismet/BlueprintFunctionLibrary.h"
+#include "AI/Navigation/NavigationSystem.h"
+#include "AI/AIResourceInterface.h"
+#include "AI/Navigation/PathFollowingComponent.h"
+#include "GameFramework/AIController.h"
+#include "GameFramework/PlayerMuteList.h"
+#include "Camera/PlayerCameraManager.h"
+#include "GameFramework/Info.h"
+#include "GameFramework/GameMode.h"
+#include "Components/InputComponent.h"
+#include "Engine/ScriptViewportClient.h"
+#include "Engine/GameViewportClient.h"
+#include "Curves/CurveBase.h"
+#include "Curves/CurveFloat.h"
+#include "GameFramework/ForceFeedbackEffect.h"
+#include "GameFramework/OnlineReplStructs.h"
+#include "GameFramework/PlayerController.h"
+#include "Engine/DebugCameraController.h"
+#include "Debug/LogVisualizerCameraController.h"
+#include "Engine/DecalActor.h"
+#include "PhysicsEngine/DestructibleActor.h"
+#include "Particles/Emitter.h"
+#include "Particles/EmitterCameraLensEffectBase.h"
+#include "Particles/EmitterSpawnable.h"
+#include "Engine/Canvas.h"
+#include "GameFramework/HUD.h"
+#include "Engine/DebugCameraHUD.h"
+#include "Debug/LogVisualizerHUD.h"
+#include "Debug/GameplayDebuggingHUDComponent.h"
+#include "Atmosphere/AtmosphericFog.h"
+#include "Engine/ExponentialHeightFog.h"
+#include "GameFramework/GameNetworkManager.h"
+#include "GameFramework/GameSession.h"
+#include "GameFramework/GameState.h"
+#include "GameFramework/PlayerState.h"
+#include "Engine/SkyLight.h"
+#include "Engine/WindDirectionalSource.h"
+#include "GameFramework/MusicTrackDataStructures.h"
+#include "GameFramework/WorldSettings.h"
+#include "Foliage/InstancedFoliageActor.h"
+#include "Landscape/LandscapeGizmoActor.h"
+#include "Landscape/LandscapeGizmoActiveActor.h"
+#include "Landscape/LandscapeLayerInfoObject.h"
+#include "Landscape/LandscapeInfo.h"
+#include "PhysicsEngine/BodyInstance.h"
+#include "Landscape/LandscapeProxy.h"
+#include "Components/SceneComponent.h"
+#include "Components/LightComponentBase.h"
+#include "Components/LightComponent.h"
+#include "Landscape/Landscape.h"
+#include "Engine/LevelBounds.h"
+#include "Engine/LevelScriptActor.h"
+#include "Engine/Light.h"
+#include "Engine/DirectionalLight.h"
+#include "Engine/PointLight.h"
+#include "Engine/SpotLight.h"
+#include "Engine/GeneratedMeshAreaLight.h"
+#include "Materials/MaterialInstanceActor.h"
+#include "Matinee/MatineeActor.h"
+#include "Matinee/MatineeActorCameraAnim.h"
+#include "AI/Navigation/NavigationData.h"
+#include "AI/Navigation/NavigationGraph.h"
+#include "AI/Navigation/NavAreas/NavArea.h"
+#include "AI/Navigation/RecastNavMesh.h"
+#include "AI/Navigation/NavigationGraphNode.h"
+#include "Engine/NavigationObjectBase.h"
+#include "GameFramework/PlayerStart.h"
+#include "Engine/PlayerStartPIE.h"
+#include "AI/Navigation/NavPathObserverInterface.h"
+#include "AI/Navigation/NavigationTestingActor.h"
+#include "AI/Navigation/NavLinkHostInterface.h"
+#include "AI/Navigation/NavLinkProxy.h"
+#include "Engine/NiagaraActor.h"
+#include "Engine/Note.h"
+#include "Components/PrimitiveComponent.h"
+#include "Particles/ParticleSystem.h"
+#include "Particles/ParticleSystemComponent.h"
+#include "Particles/ParticleEventManager.h"
+#include "GameFramework/MovementComponent.h"
+#include "GameFramework/NavMovementComponent.h"
+#include "GameFramework/PawnMovementComponent.h"
+#include "GameFramework/Pawn.h"
+#include "Interfaces/NetworkPredictionInterface.h"
+#include "GameFramework/CharacterMovementComponent.h"
+#include "GameFramework/Character.h"
+#include "AI/EnvironmentQuery/EQSQueryResultSourceInterface.h"
+#include "AI/EnvironmentQuery/EnvQueryTypes.h"
+#include "AI/EnvironmentQuery/EQSTestingPawn.h"
+#include "GameFramework/FloatingPawnMovement.h"
+#include "GameFramework/DefaultPawn.h"
+#include "GameFramework/SpectatorPawn.h"
+#include "GameFramework/WheeledVehicle.h"
+#include "Engine/ReflectionCapture.h"
+#include "Engine/BoxReflectionCapture.h"
+#include "Engine/PlaneReflectionCapture.h"
+#include "Engine/SphereReflectionCapture.h"
+#include "PhysicsEngine/RigidBodyBase.h"
+#include "PhysicsEngine/PhysicsConstraintActor.h"
+#include "PhysicsEngine/PhysicsThruster.h"
+#include "PhysicsEngine/RadialForceActor.h"
+#include "Engine/SceneCapture.h"
+#include "Engine/SceneCapture2D.h"
+#include "Engine/SceneCaptureCube.h"
+#include "Components/MeshComponent.h"
+#include "Components/SkinnedMeshComponent.h"
+#include "PhysicsEngine/ConstraintInstance.h"
+#include "PhysicsEngine/PhysicsConstraintComponent.h"
+#include "Components/SkeletalMeshComponent.h"
+#include "Matinee/MatineeAnimInterface.h"
+#include "Animation/SkeletalMeshActor.h"
+#include "Engine/StaticMeshActor.h"
+#include "Engine/InteractiveFoliageActor.h"
+#include "Engine/TargetPoint.h"
+#include "Engine/TextRenderActor.h"
+#include "Engine/TriggerBase.h"
+#include "Engine/TriggerBox.h"
+#include "Engine/TriggerCapsule.h"
+#include "Engine/TriggerSphere.h"
+#include "VectorField/VectorFieldVolume.h"
+#include "Components/ApplicationLifecycleComponent.h"
+#include "AI/BehaviorTree/BehaviorTreeTypes.h"
+#include "AI/BehaviorTree/Blackboard/BlackboardKeyType.h"
+#include "Engine/DataAsset.h"
+#include "AI/BehaviorTree/BlackboardData.h"
+#include "AI/BehaviorTree/BlackboardComponent.h"
+#include "AI/BrainComponent.h"
+#include "AI/BehaviorTree/BehaviorTreeComponent.h"
+#include "Debug/GameplayDebuggingControllerComponent.h"
+#include "GameFramework/SpectatorPawnMovement.h"
+#include "Vehicles/VehicleWheel.h"
+#include "Vehicles/WheeledVehicleMovementComponent.h"
+#include "Vehicles/WheeledVehicleMovementComponent4W.h"
+#include "GameFramework/ProjectileMovementComponent.h"
+#include "GameFramework/RotatingMovementComponent.h"
+#include "AI/Navigation/NavigationPathGenerator.h"
+#include "AI/Navigation/NavigationComponent.h"
+#include "AI/Navigation/NavRelevantComponent.h"
+#include "AI/Navigation/SmartNavLinkComponent.h"
+#include "Components/PawnNoiseEmitterComponent.h"
+#include "Components/PawnSensingComponent.h"
+#include "PhysicsEngine/PhysicsHandleComponent.h"
+#include "Atmosphere/AtmosphericFogComponent.h"
+#include "Sound/SoundAttenuation.h"
+#include "Components/AudioComponent.h"
+#include "Camera/CameraComponent.h"
+#include "Components/ChildActorComponent.h"
+#include "Components/DecalComponent.h"
+#include "Components/ExponentialHeightFogComponent.h"
+#include "Components/DirectionalLightComponent.h"
+#include "Components/PointLightComponent.h"
+#include "Components/SpotLightComponent.h"
+#include "Components/SkyLightComponent.h"
+#include "AI/Navigation/NavigationGraphNodeComponent.h"
+#include "PhysicsEngine/PhysicsThrusterComponent.h"
+#include "Components/PostProcessComponent.h"
+#include "Components/ArrowComponent.h"
+#include "Components/BillboardComponent.h"
+#include "Components/BrushComponent.h"
+#include "Components/DrawFrustumComponent.h"
+#include "Debug/DebugDrawService.h"
+#include "AI/EnvironmentQuery/EQSRenderingComponent.h"
+#include "Debug/GameplayDebuggingComponent.h"
+#include "Landscape/LandscapeComponent.h"
+#include "Landscape/LandscapeGizmoRenderComponent.h"
+#include "Landscape/LandscapeHeightfieldCollisionComponent.h"
+#include "Landscape/LandscapeMeshCollisionComponent.h"
+#include "Landscape/LandscapeSplinesComponent.h"
+#include "Components/LineBatchComponent.h"
+#include "Components/MaterialBillboardComponent.h"
+#include "Components/DestructibleComponent.h"
+#include "Animation/AnimationAsset.h"
+#include "Components/PoseableMeshComponent.h"
+#include "Lightmass/LightmassPrimitiveSettingsObject.h"
+#include "Components/StaticMeshComponent.h"
+#include "Landscape/ControlPointMeshComponent.h"
+#include "Components/InstancedStaticMeshComponent.h"
+#include "Components/InteractiveFoliageComponent.h"
+#include "Interfaces/Interface_CollisionDataProvider.h"
+#include "Components/SplineMeshComponent.h"
+#include "Components/ModelComponent.h"
+#include "AI/Navigation/NavLinkRenderingComponent.h"
+#include "AI/Navigation/NavMeshRenderingComponent.h"
+#include "AI/Navigation/NavTestRenderingComponent.h"
+#include "Components/NiagaraComponent.h"
+#include "Components/ShapeComponent.h"
+#include "Components/BoxComponent.h"
+#include "Components/CapsuleComponent.h"
+#include "Components/SphereComponent.h"
+#include "Components/DrawSphereComponent.h"
+#include "Components/SplineComponent.h"
+#include "Components/TextRenderComponent.h"
+#include "Components/VectorFieldComponent.h"
+#include "PhysicsEngine/RadialForceComponent.h"
+#include "Components/ReflectionCaptureComponent.h"
+#include "Components/BoxReflectionCaptureComponent.h"
+#include "Components/PlaneReflectionCaptureComponent.h"
+#include "Components/SphereReflectionCaptureComponent.h"
+#include "Components/SceneCaptureComponent.h"
+#include "Components/SceneCaptureComponent2D.h"
+#include "Components/SceneCaptureComponentCube.h"
+#include "GameFramework/SpringArmComponent.h"
+#include "Components/WindDirectionalSourceComponent.h"
+#include "Components/TimelineComponent.h"
+#include "Animation/AnimSequenceBase.h"
+#include "Animation/AnimCompositeBase.h"
+#include "Animation/AnimComposite.h"
+#include "Animation/AnimMontage.h"
+#include "Animation/AnimSequence.h"
+#include "Animation/BlendSpaceBase.h"
+#include "Animation/AnimStateMachineTypes.h"
+#include "Animation/AnimInstance.h"
+#include "Animation/BlendSpace.h"
+#include "Animation/AimOffsetBlendSpace.h"
+#include "Animation/BlendSpace1D.h"
+#include "Animation/AimOffsetBlendSpace1D.h"
+#include "Animation/AnimationTypes.h"
+#include "Animation/AnimCompress.h"
+#include "Animation/AnimCompress_Automatic.h"
+#include "Animation/AnimCompress_BitwiseCompressOnly.h"
+#include "Animation/AnimCompress_LeastDestructive.h"
+#include "Animation/AnimCompress_RemoveEverySecondKey.h"
+#include "Animation/AnimCompress_RemoveLinearKeys.h"
+#include "Animation/AnimCompress_PerTrackCompression.h"
+#include "Animation/AnimCompress_RemoveTrivialKeys.h"
+#include "Animation/AnimCompress_RevertToRaw.h"
+#include "Animation/AnimSingleNodeInstance.h"
+#include "Vehicles/VehicleAnimInstance.h"
+#include "Animation/AnimNotifies/AnimNotify.h"
+#include "Animation/AnimNotifies/AnimNotifyState.h"
+#include "Animation/AnimNotifies/AnimNotifyState_TimedParticleEffect.h"
+#include "Animation/AnimNotifies/AnimNotifyState_Trail.h"
+#include "Animation/AnimSet.h"
+#include "EditorFramework/AssetImportData.h"
+#include "Engine/AssetUserData.h"
+#include "Sound/AudioSettings.h"
+#include "Tests/AutomationTestSettings.h"
+#include "AI/Navigation/AvoidanceManager.h"
+#include "AI/BehaviorTree/BTNode.h"
+#include "AI/BehaviorTree/BTCompositeNode.h"
+#include "AI/BehaviorTree/BTTaskNode.h"
+#include "AI/BehaviorTree/BehaviorTree.h"
+#include "AI/BehaviorTree/BehaviorTreeManager.h"
+#include "AI/BehaviorTree/Blackboard/BlackboardKeyType_Bool.h"
+#include "AI/BehaviorTree/Blackboard/BlackboardKeyType_Class.h"
+#include "AI/BehaviorTree/Blackboard/BlackboardKeyType_Enum.h"
+#include "AI/BehaviorTree/Blackboard/BlackboardKeyType_Float.h"
+#include "AI/BehaviorTree/Blackboard/BlackboardKeyType_Int.h"
+#include "AI/BehaviorTree/Blackboard/BlackboardKeyType_Name.h"
+#include "AI/BehaviorTree/Blackboard/BlackboardKeyType_NativeEnum.h"
+#include "AI/BehaviorTree/Blackboard/BlackboardKeyType_Object.h"
+#include "AI/BehaviorTree/Blackboard/BlackboardKeyType_String.h"
+#include "AI/BehaviorTree/Blackboard/BlackboardKeyType_Vector.h"
+#include "Engine/BlendableInterface.h"
+#include "Engine/BlueprintCore.h"
+#include "EdGraph/EdGraphPin.h"
+#include "EdGraph/EdGraphNode.h"
+#include "EdGraph/EdGraph.h"
+#include "Engine/Blueprint.h"
+#include "Animation/AnimBlueprint.h"
+#include "Engine/LevelScriptBlueprint.h"
+#include "AI/AISystem.h"
+#include "AI/BehaviorTree/BTFunctionLibrary.h"
+#include "Sound/DialogueTypes.h"
+#include "Kismet/GameplayStatics.h"
+#include "Kismet/HeadMountedDisplayFunctionLibrary.h"
+#include "Kismet/KismetAIHelperLibrary.h"
+#include "Kismet/KismetArrayLibrary.h"
+#include "Kismet/KismetMaterialLibrary.h"
+#include "Kismet/KismetMathLibrary.h"
+#include "Kismet/KismetNodeHelperLibrary.h"
+#include "Kismet/KismetStringLibrary.h"
+#include "Kismet/KismetSystemLibrary.h"
+#include "Kismet/KismetTextLibrary.h"
+#include "Engine/BlueprintGeneratedClass.h"
+#include "Animation/AnimBlueprintGeneratedClass.h"
+#include "PhysicsEngine/BodySetup.h"
+#include "Animation/AnimData/BoneMaskFilter.h"
+#include "Engine/BookMark.h"
+#include "Engine/BookMark2D.h"
+#include "Engine/Breakpoint.h"
+#include "Engine/BrushBuilder.h"
+#include "AI/BehaviorTree/BTAuxiliaryNode.h"
+#include "AI/BehaviorTree/BTDecorator.h"
+#include "AI/BehaviorTree/Decorators/BTDecorator_BlackboardBase.h"
+#include "AI/BehaviorTree/Decorators/BTDecorator_Blackboard.h"
+#include "AI/BehaviorTree/Decorators/BTDecorator_BlueprintBase.h"
+#include "AI/BehaviorTree/Decorators/BTDecorator_CompareBBEntries.h"
+#include "AI/BehaviorTree/Decorators/BTDecorator_ConeCheck.h"
+#include "AI/BehaviorTree/Decorators/BTDecorator_Cooldown.h"
+#include "AI/BehaviorTree/Decorators/BTDecorator_DoesPathExist.h"
+#include "AI/BehaviorTree/Decorators/BTDecorator_ForceSuccess.h"
+#include "AI/BehaviorTree/Decorators/BTDecorator_KeepInCone.h"
+#include "AI/BehaviorTree/Decorators/BTDecorator_Loop.h"
+#include "AI/BehaviorTree/Decorators/BTDecorator_Optional.h"
+#include "AI/BehaviorTree/Decorators/BTDecorator_ReachedMoveGoal.h"
+#include "AI/BehaviorTree/Decorators/BTDecorator_TimeLimit.h"
+#include "AI/BehaviorTree/BTService.h"
+#include "AI/BehaviorTree/Services/BTService_BlackboardBase.h"
+#include "AI/BehaviorTree/Services/BTService_DefaultFocus.h"
+#include "AI/BehaviorTree/Services/BTService_BlueprintBase.h"
+#include "AI/BehaviorTree/Composites/BTComposite_Selector.h"
+#include "AI/BehaviorTree/Composites/BTComposite_Sequence.h"
+#include "AI/BehaviorTree/Composites/BTComposite_SimpleParallel.h"
+#include "AI/BehaviorTree/Tasks/BTTask_BlackboardBase.h"
+#include "AI/BehaviorTree/Tasks/BTTask_MoveDirectlyToward.h"
+#include "AI/BehaviorTree/Tasks/BTTask_MoveTo.h"
+#include "AI/BehaviorTree/Tasks/BTTask_RunEQSQuery.h"
+#include "AI/BehaviorTree/Tasks/BTTask_BlueprintBase.h"
+#include "AI/BehaviorTree/Tasks/BTTask_MakeNoise.h"
+#include "AI/BehaviorTree/Tasks/BTTask_PlaySound.h"
+#include "AI/BehaviorTree/Tasks/BTTask_RunBehavior.h"
+#include "AI/BehaviorTree/Tasks/BTTask_Wait.h"
+#include "Slate/ButtonStyleAsset.h"
+#include "Camera/CameraAnim.h"
+#include "Camera/CameraAnimInst.h"
+#include "Camera/CameraModifier.h"
+#include "Camera/CameraShake.h"
+#include "Camera/CameraModifier_CameraShake.h"
+#include "Engine/Channel.h"
+#include "Engine/ActorChannel.h"
+#include "Engine/ControlChannel.h"
+#include "Engine/VoiceChannel.h"
+#include "GameFramework/CheatManager.h"
+#include "Slate/CheckboxStyleAsset.h"
+#include "Engine/ClipPadEntry.h"
+#include "Engine/CodecMovie.h"
+#include "Engine/CodecMovieFallback.h"
+#include "Engine/CollisionProfile.h"
+#include "Commandlets/Commandlet.h"
+#include "Commandlets/SmokeTestCommandlet.h"
+#include "Engine/Console.h"
+#include "Curves/CurveLinearColor.h"
+#include "Curves/CurveVector.h"
+#include "Curves/CurveEdPresetCurve.h"
+#include "Engine/CurveTable.h"
+#include "GameFramework/DamageType.h"
+#include "Vehicles/TireType.h"
+#include "Engine/DataTable.h"
+#include "Engine/DestructibleFractureSettings.h"
+#include "DeviceProfiles/DeviceProfile.h"
+#include "DeviceProfiles/DeviceProfileManager.h"
+#include "Sound/DialogueVoice.h"
+#include "Sound/DialogueWave.h"
+#include "Distributions/Distribution.h"
+#include "Distributions/DistributionFloat.h"
+#include "Distributions/DistributionFloatConstant.h"
+#include "Distributions/DistributionFloatParameterBase.h"
+#include "Distributions/DistributionFloatParticleParameter.h"
+#include "Distributions/DistributionFloatSoundParameter.h"
+#include "Distributions/DistributionFloatConstantCurve.h"
+#include "Distributions/DistributionFloatUniform.h"
+#include "Distributions/DistributionFloatUniformCurve.h"
+#include "Distributions/DistributionVector.h"
+#include "Distributions/DistributionVectorConstant.h"
+#include "Distributions/DistributionVectorParameterBase.h"
+#include "Distributions/DistributionVectorParticleParameter.h"
+#include "Distributions/DistributionVectorConstantCurve.h"
+#include "Distributions/DistributionVectorUniform.h"
+#include "Distributions/DistributionVectorUniformCurve.h"
+#include "Engine/DynamicBlueprintBinding.h"
+#include "Engine/ComponentDelegateBinding.h"
+#include "Engine/InputDelegateBinding.h"
+#include "Engine/InputActionDelegateBinding.h"
+#include "Engine/InputAxisDelegateBinding.h"
+#include "Engine/InputAxisKeyDelegateBinding.h"
+#include "Engine/InputVectorAxisDelegateBinding.h"
+#include "Engine/InputKeyDelegateBinding.h"
+#include "Engine/InputTouchDelegateBinding.h"
+#include "EdGraph/EdGraphNode_Comment.h"
+#include "EdGraph/EdGraphSchema.h"
+#include "Engine/Engine.h"
+#include "Engine/GameEngine.h"
+#include "AI/EnvironmentQuery/EnvQuery.h"
+#include "AI/EnvironmentQuery/EnvQueryContext.h"
+#include "AI/EnvironmentQuery/Contexts/EnvQueryContext_BlueprintBase.h"
+#include "AI/EnvironmentQuery/Contexts/EnvQueryContext_Item.h"
+#include "AI/EnvironmentQuery/Contexts/EnvQueryContext_Querier.h"
+#include "AI/EnvironmentQuery/EnvQueryGenerator.h"
+#include "AI/EnvironmentQuery/Generators/EnvQueryGenerator_Composite.h"
+#include "AI/EnvironmentQuery/Generators/EnvQueryGenerator_PathingGrid.h"
+#include "AI/EnvironmentQuery/Generators/EnvQueryGenerator_ProjectedPoints.h"
+#include "AI/EnvironmentQuery/Generators/EnvQueryGenerator_OnCircle.h"
+#include "AI/EnvironmentQuery/Generators/EnvQueryGenerator_SimpleGrid.h"
+#include "AI/EnvironmentQuery/Items/EnvQueryItemType.h"
+#include "AI/EnvironmentQuery/Items/EnvQueryItemType_VectorBase.h"
+#include "AI/EnvironmentQuery/Items/EnvQueryItemType_ActorBase.h"
+#include "AI/EnvironmentQuery/Items/EnvQueryItemType_Actor.h"
+#include "AI/EnvironmentQuery/Items/EnvQueryItemType_Direction.h"
+#include "AI/EnvironmentQuery/Items/EnvQueryItemType_Point.h"
+#include "AI/EnvironmentQuery/EnvQueryManager.h"
+#include "AI/EnvironmentQuery/EnvQueryOption.h"
+#include "AI/EnvironmentQuery/EnvQueryTest.h"
+#include "AI/EnvironmentQuery/Tests/EnvQueryTest_Distance.h"
+#include "AI/EnvironmentQuery/Tests/EnvQueryTest_Dot.h"
+#include "AI/EnvironmentQuery/Tests/EnvQueryTest_Pathfinding.h"
+#include "AI/EnvironmentQuery/Tests/EnvQueryTest_Trace.h"
+#include "Exporters/Exporter.h"
+#include "Engine/FontImportOptions.h"
+#include "Engine/Font.h"
+#include "GameFramework/GameUserSettings.h"
+#include "GameFramework/PlayerInput.h"
+#include "GameFramework/InputSettings.h"
+#include "Foliage/InstancedFoliageSettings.h"
+#include "Engine/InterpCurveEdSetup.h"
+#include "Matinee/InterpData.h"
+#include "Matinee/InterpFilter.h"
+#include "Matinee/InterpFilter_Classes.h"
+#include "Matinee/InterpFilter_Custom.h"
+#include "Matinee/InterpGroup.h"
+#include "Matinee/InterpGroupCamera.h"
+#include "Matinee/InterpGroupDirector.h"
+#include "Matinee/InterpGroupInst.h"
+#include "Matinee/InterpGroupInstCamera.h"
+#include "Matinee/InterpGroupInstDirector.h"
+#include "Matinee/InterpTrack.h"
+#include "Matinee/InterpTrackBoolProp.h"
+#include "Matinee/InterpTrackDirector.h"
+#include "Matinee/InterpTrackEvent.h"
+#include "Matinee/InterpTrackFloatBase.h"
+#include "Matinee/InterpTrackAnimControl.h"
+#include "Matinee/InterpTrackFade.h"
+#include "Matinee/InterpTrackFloatMaterialParam.h"
+#include "Matinee/InterpTrackFloatParticleParam.h"
+#include "Matinee/InterpTrackFloatProp.h"
+#include "Matinee/InterpTrackMove.h"
+#include "Matinee/InterpTrackMoveAxis.h"
+#include "Matinee/InterpTrackSlomo.h"
+#include "Matinee/InterpTrackLinearColorBase.h"
+#include "Matinee/InterpTrackLinearColorProp.h"
+#include "Matinee/InterpTrackParticleReplay.h"
+#include "Matinee/InterpTrackToggle.h"
+#include "Matinee/InterpTrackVectorBase.h"
+#include "Matinee/InterpTrackAudioMaster.h"
+#include "Matinee/InterpTrackColorProp.h"
+#include "Matinee/InterpTrackColorScale.h"
+#include "Matinee/InterpTrackSound.h"
+#include "Matinee/InterpTrackVectorMaterialParam.h"
+#include "Matinee/InterpTrackVectorProp.h"
+#include "Matinee/InterpTrackVisibility.h"
+#include "Matinee/InterpTrackInst.h"
+#include "Matinee/InterpTrackInstAnimControl.h"
+#include "Matinee/InterpTrackInstAudioMaster.h"
+#include "Matinee/InterpTrackInstColorScale.h"
+#include "Matinee/InterpTrackInstDirector.h"
+#include "Matinee/InterpTrackInstEvent.h"
+#include "Matinee/InterpTrackInstFade.h"
+#include "Matinee/InterpTrackInstFloatMaterialParam.h"
+#include "Matinee/InterpTrackInstFloatParticleParam.h"
+#include "Matinee/InterpTrackInstMove.h"
+#include "Matinee/InterpTrackInstParticleReplay.h"
+#include "Matinee/InterpTrackInstProperty.h"
+#include "Matinee/InterpTrackInstBoolProp.h"
+#include "Matinee/InterpTrackInstColorProp.h"
+#include "Matinee/InterpTrackInstFloatProp.h"
+#include "Matinee/InterpTrackInstLinearColorProp.h"
+#include "Matinee/InterpTrackInstVectorProp.h"
+#include "Matinee/InterpTrackInstSlomo.h"
+#include "Matinee/InterpTrackInstSound.h"
+#include "Matinee/InterpTrackInstToggle.h"
+#include "Matinee/InterpTrackInstVectorMaterialParam.h"
+#include "Matinee/InterpTrackInstVisibility.h"
+#include "Engine/IntSerialization.h"
+#include "AI/KismetAIAsyncTaskProxy.h"
+#include "Landscape/LandscapeSplineSegment.h"
+#include "Landscape/LandscapeSplineControlPoint.h"
+#include "Layers/Layer.h"
+#include "Engine/LevelBase.h"
+#include "Engine/Level.h"
+#include "Engine/LevelStreaming.h"
+#include "Engine/LevelStreamingAlwaysLoaded.h"
+#include "Engine/LevelStreamingBounds.h"
+#include "Engine/LevelStreamingKismet.h"
+#include "Engine/LevelStreamingPersistent.h"
+#include "Lightmass/LightmappedSurfaceCollection.h"
+#include "GameFramework/LocalMessage.h"
+#include "GameFramework/EngineMessage.h"
+#include "Materials/MaterialExpression.h"
+#include "Materials/MaterialExpressionAbs.h"
+#include "Materials/MaterialExpressionActorPositionWS.h"
+#include "Materials/MaterialExpressionAdd.h"
+#include "Materials/MaterialExpressionAppendVector.h"
+#include "Materials/MaterialExpressionAtmosphericFogColor.h"
+#include "Materials/MaterialExpressionBlackBody.h"
+#include "Materials/MaterialExpressionBreakMaterialAttributes.h"
+#include "Materials/MaterialExpressionBumpOffset.h"
+#include "Materials/MaterialExpressionCameraPositionWS.h"
+#include "Materials/MaterialExpressionCameraVectorWS.h"
+#include "Materials/MaterialExpressionCeil.h"
+#include "Materials/MaterialExpressionClamp.h"
+#include "Materials/MaterialExpressionCollectionParameter.h"
+#include "Materials/MaterialExpressionComment.h"
+#include "Materials/MaterialExpressionComponentMask.h"
+#include "Materials/MaterialExpressionConstant.h"
+#include "Materials/MaterialExpressionConstant2Vector.h"
+#include "Materials/MaterialExpressionConstant3Vector.h"
+#include "Materials/MaterialExpressionConstant4Vector.h"
+#include "Materials/MaterialExpressionConstantBiasScale.h"
+#include "Materials/MaterialExpressionCosine.h"
+#include "Materials/MaterialExpressionCrossProduct.h"
+#include "Materials/MaterialExpressionCustom.h"
+#include "Materials/MaterialExpressionCustomTexture.h"
+#include "Materials/MaterialExpressionDDX.h"
+#include "Materials/MaterialExpressionDDY.h"
+#include "Materials/MaterialExpressionDepthFade.h"
+#include "Materials/MaterialExpressionDepthOfFieldFunction.h"
+#include "Materials/MaterialExpressionDeriveNormalZ.h"
+#include "Materials/MaterialExpressionDesaturation.h"
+#include "Materials/MaterialExpressionDistance.h"
+#include "Materials/MaterialExpressionDistanceCullFade.h"
+#include "Materials/MaterialExpressionDivide.h"
+#include "Materials/MaterialExpressionDotProduct.h"
+#include "Materials/MaterialExpressionDynamicParameter.h"
+#include "Materials/MaterialExpressionEyeAdaptation.h"
+#include "Materials/MaterialExpressionFeatureLevelSwitch.h"
+#include "Materials/MaterialExpressionFloor.h"
+#include "Materials/MaterialExpressionFmod.h"
+#include "Materials/MaterialExpressionFontSample.h"
+#include "Materials/MaterialExpressionFontSampleParameter.h"
+#include "Materials/MaterialExpressionFrac.h"
+#include "Materials/MaterialExpressionFresnel.h"
+#include "Materials/MaterialExpressionFunctionInput.h"
+#include "Materials/MaterialExpressionFunctionOutput.h"
+#include "Materials/MaterialExpressionGIReplace.h"
+#include "Materials/MaterialExpressionIf.h"
+#include "Materials/MaterialExpressionLandscapeLayerBlend.h"
+#include "Materials/MaterialExpressionLandscapeLayerCoords.h"
+#include "Materials/MaterialExpressionLandscapeLayerSwitch.h"
+#include "Materials/MaterialExpressionLandscapeLayerWeight.h"
+#include "Materials/MaterialExpressionLandscapeVisibilityMask.h"
+#include "Materials/MaterialExpressionLightmapUVs.h"
+#include "Materials/MaterialExpressionLightmassReplace.h"
+#include "Materials/MaterialExpressionLightVector.h"
+#include "Materials/MaterialExpressionLinearInterpolate.h"
+#include "Materials/MaterialExpressionMakeMaterialAttributes.h"
+#include "Materials/MaterialExpressionMaterialFunctionCall.h"
+#include "Materials/MaterialExpressionMax.h"
+#include "Materials/MaterialExpressionMin.h"
+#include "Materials/MaterialExpressionMultiply.h"
+#include "Materials/MaterialExpressionNoise.h"
+#include "Materials/MaterialExpressionNormalize.h"
+#include "Materials/MaterialExpressionObjectBounds.h"
+#include "Materials/MaterialExpressionObjectOrientation.h"
+#include "Materials/MaterialExpressionObjectPositionWS.h"
+#include "Materials/MaterialExpressionObjectRadius.h"
+#include "Materials/MaterialExpressionOneMinus.h"
+#include "Materials/MaterialExpressionPanner.h"
+#include "Materials/MaterialExpressionParameter.h"
+#include "Materials/MaterialExpressionScalarParameter.h"
+#include "Materials/MaterialExpressionStaticBoolParameter.h"
+#include "Materials/MaterialExpressionStaticSwitchParameter.h"
+#include "Materials/MaterialExpressionStaticComponentMaskParameter.h"
+#include "Materials/MaterialExpressionVectorParameter.h"
+#include "Materials/MaterialExpressionParticleColor.h"
+#include "Materials/MaterialExpressionParticleDirection.h"
+#include "Materials/MaterialExpressionParticleMacroUV.h"
+#include "Materials/MaterialExpressionParticleMotionBlurFade.h"
+#include "Materials/MaterialExpressionParticlePositionWS.h"
+#include "Materials/MaterialExpressionParticleRadius.h"
+#include "Materials/MaterialExpressionParticleRelativeTime.h"
+#include "Materials/MaterialExpressionParticleSize.h"
+#include "Materials/MaterialExpressionParticleSpeed.h"
+#include "Materials/MaterialExpressionPerInstanceFadeAmount.h"
+#include "Materials/MaterialExpressionPerInstanceRandom.h"
+#include "Materials/MaterialExpressionPixelDepth.h"
+#include "Materials/MaterialExpressionPixelNormalWS.h"
+#include "Materials/MaterialExpressionPower.h"
+#include "Materials/MaterialExpressionQualitySwitch.h"
+#include "Materials/MaterialExpressionReflectionVectorWS.h"
+#include "Materials/MaterialExpressionRotateAboutAxis.h"
+#include "Materials/MaterialExpressionRotator.h"
+#include "Materials/MaterialExpressionSceneColor.h"
+#include "Materials/MaterialExpressionSceneDepth.h"
+#include "Materials/MaterialExpressionSceneTexelSize.h"
+#include "Materials/MaterialExpressionSceneTexture.h"
+#include "Materials/MaterialExpressionScreenPosition.h"
+#include "Materials/MaterialExpressionSine.h"
+#include "Materials/MaterialExpressionSpeedTree.h"
+#include "Materials/MaterialExpressionSpeedTreeColorVariation.h"
+#include "Materials/MaterialExpressionSphereMask.h"
+#include "Materials/MaterialExpressionSphericalParticleOpacity.h"
+#include "Materials/MaterialExpressionSquareRoot.h"
+#include "Materials/MaterialExpressionStaticBool.h"
+#include "Materials/MaterialExpressionStaticSwitch.h"
+#include "Materials/MaterialExpressionSubtract.h"
+#include "Materials/MaterialExpressionTextureBase.h"
+#include "Materials/MaterialExpressionTextureObject.h"
+#include "Materials/MaterialExpressionTextureSample.h"
+#include "Materials/MaterialExpressionParticleSubUV.h"
+#include "Materials/MaterialExpressionTextureSampleParameter.h"
+#include "Materials/MaterialExpressionTextureObjectParameter.h"
+#include "Materials/MaterialExpressionTextureSampleParameter2D.h"
+#include "Materials/MaterialExpressionAntialiasedTextureMask.h"
+#include "Materials/MaterialExpressionTextureSampleParameterSubUV.h"
+#include "Materials/MaterialExpressionTextureSampleParameterCube.h"
+#include "Materials/MaterialExpressionTextureSampleParameterMovie.h"
+#include "Materials/MaterialExpressionTextureCoordinate.h"
+#include "Materials/MaterialExpressionTime.h"
+#include "Materials/MaterialExpressionTransform.h"
+#include "Materials/MaterialExpressionTransformPosition.h"
+#include "Materials/MaterialExpressionTwoSidedSign.h"
+#include "Materials/MaterialExpressionVertexColor.h"
+#include "Materials/MaterialExpressionVertexNormalWS.h"
+#include "Materials/MaterialExpressionViewSize.h"
+#include "Materials/MaterialExpressionWorldPosition.h"
+#include "Materials/MaterialFunction.h"
+#include "Materials/MaterialInterface.h"
+#include "Materials/Material.h"
+#include "Materials/MaterialInstance.h"
+#include "Materials/MaterialInstanceConstant.h"
+#include "Landscape/LandscapeMaterialInstanceConstant.h"
+#include "Materials/MaterialInstanceDynamic.h"
+#include "Materials/MaterialParameterCollection.h"
+#include "Materials/MaterialParameterCollectionInstance.h"
+#include "Matinee/MatineeInterface.h"
+#include "AI/Navigation/NavAreas/NavArea_Default.h"
+#include "AI/Navigation/NavAreas/NavArea_Null.h"
+#include "AI/Navigation/NavAreas/NavAreaMeta.h"
+#include "AI/Navigation/NavAreas/NavAreaMeta_SwitchByAgent.h"
+#include "AI/Navigation/NavCollision.h"
+#include "AI/Navigation/NavigationProxy.h"
+#include "AI/Navigation/NavFilters/NavigationQueryFilter.h"
+#include "AI/Navigation/NavFilters/RecastFilter_UseDefaultArea.h"
+#include "AI/Navigation/NavLinkTrivial.h"
+#include "AI/Navigation/NavNodeInterface.h"
+#include "Engine/NetDriver.h"
+#include "Engine/NiagaraScript.h"
+#include "Engine/NiagaraScriptSourceBase.h"
+#include "Engine/ObjectLibrary.h"
+#include "Engine/ObjectReferencer.h"
+#include "GameFramework/OnlineSession.h"
+#include "Engine/PackageMapClient.h"
+#include "Particles/ParticleLODLevel.h"
+#include "Particles/ParticleEmitter.h"
+#include "Particles/ParticleSpriteEmitter.h"
+#include "Particles/ParticleModule.h"
+#include "Particles/Acceleration/ParticleModuleAccelerationBase.h"
+#include "Particles/Acceleration/ParticleModuleAcceleration.h"
+#include "Particles/Acceleration/ParticleModuleAccelerationConstant.h"
+#include "Particles/Acceleration/ParticleModuleAccelerationDrag.h"
+#include "Particles/Acceleration/ParticleModuleAccelerationDragScaleOverLife.h"
+#include "Particles/Acceleration/ParticleModuleAccelerationOverLifetime.h"
+#include "Particles/Attractor/ParticleModuleAttractorBase.h"
+#include "Particles/Attractor/ParticleModuleAttractorLine.h"
+#include "Particles/Attractor/ParticleModuleAttractorParticle.h"
+#include "Particles/Attractor/ParticleModuleAttractorPoint.h"
+#include "Particles/Attractor/ParticleModuleAttractorPointGravity.h"
+#include "Particles/Beam/ParticleModuleBeamBase.h"
+#include "Particles/Beam/ParticleModuleBeamModifier.h"
+#include "Particles/Beam/ParticleModuleBeamNoise.h"
+#include "Particles/Beam/ParticleModuleBeamSource.h"
+#include "Particles/Beam/ParticleModuleBeamTarget.h"
+#include "Particles/Camera/ParticleModuleCameraBase.h"
+#include "Particles/Camera/ParticleModuleCameraOffset.h"
+#include "Particles/Collision/ParticleModuleCollisionBase.h"
+#include "Particles/Collision/ParticleModuleCollision.h"
+#include "Particles/Collision/ParticleModuleCollisionGPU.h"
+#include "Particles/Color/ParticleModuleColorBase.h"
+#include "Particles/Color/ParticleModuleColor.h"
+#include "Particles/Color/ParticleModuleColor_Seeded.h"
+#include "Particles/Color/ParticleModuleColorOverLife.h"
+#include "Particles/Color/ParticleModuleColorScaleOverLife.h"
+#include "Particles/Event/ParticleModuleEventBase.h"
+#include "Particles/Event/ParticleModuleEventGenerator.h"
+#include "Particles/Event/ParticleModuleEventReceiverBase.h"
+#include "Particles/Event/ParticleModuleEventReceiverKillParticles.h"
+#include "Particles/Event/ParticleModuleEventReceiverSpawn.h"
+#include "Particles/Kill/ParticleModuleKillBase.h"
+#include "Particles/Kill/ParticleModuleKillBox.h"
+#include "Particles/Kill/ParticleModuleKillHeight.h"
+#include "Particles/Lifetime/ParticleModuleLifetimeBase.h"
+#include "Particles/Lifetime/ParticleModuleLifetime.h"
+#include "Particles/Lifetime/ParticleModuleLifetime_Seeded.h"
+#include "Particles/Light/ParticleModuleLightBase.h"
+#include "Particles/Light/ParticleModuleLight.h"
+#include "Particles/Light/ParticleModuleLight_Seeded.h"
+#include "Particles/Location/ParticleModuleLocationBase.h"
+#include "Particles/Location/ParticleModuleLocation.h"
+#include "Particles/Location/ParticleModuleLocation_Seeded.h"
+#include "Particles/Location/ParticleModuleLocationWorldOffset.h"
+#include "Particles/Location/ParticleModuleLocationWorldOffset_Seeded.h"
+#include "Particles/Location/ParticleModuleLocationBoneSocket.h"
+#include "Particles/Location/ParticleModuleLocationDirect.h"
+#include "Particles/Location/ParticleModuleLocationEmitter.h"
+#include "Particles/Location/ParticleModuleLocationEmitterDirect.h"
+#include "Particles/Location/ParticleModuleLocationPrimitiveBase.h"
+#include "Particles/Location/ParticleModuleLocationPrimitiveCylinder.h"
+#include "Particles/Location/ParticleModuleLocationPrimitiveCylinder_Seeded.h"
+#include "Particles/Location/ParticleModuleLocationPrimitiveSphere.h"
+#include "Particles/Location/ParticleModuleLocationPrimitiveSphere_Seeded.h"
+#include "Particles/Location/ParticleModuleLocationPrimitiveTriangle.h"
+#include "Particles/Location/ParticleModuleLocationSkelVertSurface.h"
+#include "Particles/Modules/Location/ParticleModulePivotOffset.h"
+#include "Particles/Location/ParticleModuleSourceMovement.h"
+#include "Particles/Material/ParticleModuleMaterialBase.h"
+#include "Particles/Material/ParticleModuleMeshMaterial.h"
+#include "Particles/Orbit/ParticleModuleOrbitBase.h"
+#include "Particles/Orbit/ParticleModuleOrbit.h"
+#include "Particles/Orientation/ParticleModuleOrientationBase.h"
+#include "Particles/Orientation/ParticleModuleOrientationAxisLock.h"
+#include "Particles/Parameter/ParticleModuleParameterBase.h"
+#include "Particles/Parameter/ParticleModuleParameterDynamic.h"
+#include "Particles/Parameter/ParticleModuleParameterDynamic_Seeded.h"
+#include "Particles/ParticleModuleRequired.h"
+#include "Particles/Rotation/ParticleModuleRotationBase.h"
+#include "Particles/Rotation/ParticleModuleMeshRotation.h"
+#include "Particles/Rotation/ParticleModuleMeshRotation_Seeded.h"
+#include "Particles/Rotation/ParticleModuleRotation.h"
+#include "Particles/Rotation/ParticleModuleRotation_Seeded.h"
+#include "Particles/Rotation/ParticleModuleRotationOverLifetime.h"
+#include "Particles/RotationRate/ParticleModuleRotationRateBase.h"
+#include "Particles/RotationRate/ParticleModuleMeshRotationRate.h"
+#include "Particles/RotationRate/ParticleModuleMeshRotationRate_Seeded.h"
+#include "Particles/RotationRate/ParticleModuleMeshRotationRateMultiplyLife.h"
+#include "Particles/RotationRate/ParticleModuleMeshRotationRateOverLife.h"
+#include "Particles/RotationRate/ParticleModuleRotationRate.h"
+#include "Particles/RotationRate/ParticleModuleRotationRate_Seeded.h"
+#include "Particles/RotationRate/ParticleModuleRotationRateMultiplyLife.h"
+#include "Particles/Size/ParticleModuleSizeBase.h"
+#include "Particles/Size/ParticleModuleSize.h"
+#include "Particles/Size/ParticleModuleSize_Seeded.h"
+#include "Particles/Size/ParticleModuleSizeMultiplyLife.h"
+#include "Particles/Size/ParticleModuleSizeScale.h"
+#include "Particles/Size/ParticleModuleSizeScaleBySpeed.h"
+#include "Particles/Spawn/ParticleModuleSpawnBase.h"
+#include "Particles/Spawn/ParticleModuleSpawn.h"
+#include "Particles/Spawn/ParticleModuleSpawnPerUnit.h"
+#include "Particles/SubUV/ParticleModuleSubUVBase.h"
+#include "Particles/SubUV/ParticleModuleSubUV.h"
+#include "Particles/SubUV/ParticleModuleSubUVMovie.h"
+#include "Particles/Trail/ParticleModuleTrailBase.h"
+#include "Particles/Trail/ParticleModuleTrailSource.h"
+#include "Particles/TypeData/ParticleModuleTypeDataBase.h"
+#include "Particles/TypeData/ParticleModuleTypeDataAnimTrail.h"
+#include "Particles/TypeData/ParticleModuleTypeDataBeam2.h"
+#include "Particles/TypeData/ParticleModuleTypeDataGpu.h"
+#include "Particles/TypeData/ParticleModuleTypeDataMesh.h"
+#include "Particles/TypeData/ParticleModuleTypeDataRibbon.h"
+#include "Particles/VectorField/ParticleModuleVectorFieldBase.h"
+#include "Particles/VectorField/ParticleModuleVectorFieldGlobal.h"
+#include "Particles/VectorField/ParticleModuleVectorFieldLocal.h"
+#include "Particles/VectorField/ParticleModuleVectorFieldRotation.h"
+#include "Particles/VectorField/ParticleModuleVectorFieldRotationRate.h"
+#include "Particles/VectorField/ParticleModuleVectorFieldScale.h"
+#include "Particles/VectorField/ParticleModuleVectorFieldScaleOverLife.h"
+#include "Particles/Velocity/ParticleModuleVelocityBase.h"
+#include "Particles/Velocity/ParticleModuleVelocity.h"
+#include "Particles/Velocity/ParticleModuleVelocity_Seeded.h"
+#include "Particles/Velocity/ParticleModuleVelocityCone.h"
+#include "Particles/Velocity/ParticleModuleVelocityInheritParent.h"
+#include "Particles/Velocity/ParticleModuleVelocityOverLifetime.h"
+#include "Particles/Event/ParticleModuleEventSendToGame.h"
+#include "Particles/ParticleSystemReplay.h"
+#include "Engine/PendingNetGame.h"
+#include "PhysicalMaterials/PhysicalMaterial.h"
+#include "PhysicalMaterials/PhysicalMaterialPropertyBase.h"
+#include "PhysicsEngine/PhysicsAsset.h"
+#include "PhysicsEngine/PhysicsCollisionHandler.h"
+#include "PhysicsEngine/PhysicsConstraintTemplate.h"
+#include "PhysicsEngine/PhysicsSettings.h"
+#include "Engine/PlatformInterfaceBase.h"
+#include "Engine/CloudStorageBase.h"
+#include "Engine/InGameAdManager.h"
+#include "Engine/MicroTransactionBase.h"
+#include "Engine/TwitterIntegrationBase.h"
+#include "Engine/PlatformInterfaceWebResponse.h"
+#include "Engine/Player.h"
+#include "Engine/LocalPlayer.h"
+#include "Engine/NetConnection.h"
+#include "Engine/ChildConnection.h"
+#include "Engine/Polys.h"
+#include "Engine/RendererSettings.h"
+#include "Sound/ReverbEffect.h"
+#include "MovieScene/RuntimeMovieScenePlayerInterface.h"
+#include "GameFramework/SaveGame.h"
+#include "Engine/SaveGameSummary.h"
+#include "Engine/SCS_Node.h"
+#include "Engine/Selection.h"
+#include "Engine/SimpleConstructionScript.h"
+#include "Animation/PreviewAssetAttachComponent.h"
+#include "Engine/SkeletalMesh.h"
+#include "Engine/DestructibleMesh.h"
+#include "Engine/SkeletalMeshSocket.h"
+#include "Animation/Skeleton.h"
+#include "Slate/SlateBrushAsset.h"
+#include "Sound/SoundBase.h"
+#include "Sound/DialogueSoundWaveProxy.h"
+#include "Sound/SoundClass.h"
+#include "Sound/SoundCue.h"
+#include "Sound/SoundGroups.h"
+#include "Sound/SoundWave.h"
+#include "Sound/SoundWaveStreaming.h"
+#include "Sound/SoundMix.h"
+#include "Sound/SoundNode.h"
+#include "Sound/SoundNodeAttenuation.h"
+#include "Sound/SoundNodeBranch.h"
+#include "Sound/SoundNodeConcatenator.h"
+#include "Sound/SoundNodeDelay.h"
+#include "Sound/SoundNodeDeprecated.h"
+#include "Sound/SoundNodeAmbient.h"
+#include "Sound/SoundNodeAmbientNonLoop.h"
+#include "Sound/SoundNodeAmbientNonLoopToggle.h"
+#include "Sound/SoundNodeConcatenatorRadio.h"
+#include "Sound/SoundNodeWave.h"
+#include "Sound/SoundNodeDialoguePlayer.h"
+#include "Sound/SoundNodeDistanceCrossFade.h"
+#include "Sound/SoundNodeParamCrossFade.h"
+#include "Sound/SoundNodeDoppler.h"
+#include "Sound/SoundNodeEnveloper.h"
+#include "Sound/SoundNodeGroupControl.h"
+#include "Sound/SoundNodeLooping.h"
+#include "Sound/SoundNodeMature.h"
+#include "Sound/SoundNodeMixer.h"
+#include "Sound/SoundNodeModulator.h"
+#include "Sound/SoundNodeModulatorContinuous.h"
+#include "Sound/SoundNodeOscillator.h"
+#include "Sound/SoundNodeRandom.h"
+#include "Sound/SoundNodeSoundClass.h"
+#include "Sound/SoundNodeSwitch.h"
+#include "Sound/SoundNodeWaveParam.h"
+#include "Sound/SoundNodeWavePlayer.h"
+#include "Spline/Spline.h"
+#include "Engine/StaticMesh.h"
+#include "Engine/StaticMeshSocket.h"
+#include "Animation/AnimNodeBase.h"
+#include "Animation/InputScaleBias.h"
+#include "Animation/AnimNode_ApplyAdditive.h"
+#include "Animation/AnimNode_BlendListBase.h"
+#include "Animation/AnimNode_BlendListByBool.h"
+#include "Animation/AnimNode_BlendListByEnum.h"
+#include "Animation/AnimNode_BlendListByInt.h"
+#include "Animation/AnimNode_BlendSpacePlayer.h"
+#include "Animation/AnimNode_BlendSpaceEvaluator.h"
+#include "Animation/BoneControllers/AnimNode_SkeletalControlBase.h"
+#include "Animation/BoneControllers/AnimNode_CopyBone.h"
+#include "Animation/BoneControllers/AnimNode_Fabrik.h"
+#include "Animation/AnimNode_LayeredBoneBlend.h"
+#include "Animation/BoneControllers/AnimNode_LookAt.h"
+#include "Animation/BoneControllers/AnimNode_ModifyBone.h"
+#include "Animation/AnimNode_RefPose.h"
+#include "Animation/AnimNode_Root.h"
+#include "Animation/AnimNode_RotateRootBone.h"
+#include "Animation/BoneControllers/AnimNode_RotationMultiplier.h"
+#include "Animation/AnimNode_RotationOffsetBlendSpace.h"
+#include "Animation/AnimNode_SaveCachedPose.h"
+#include "Animation/AnimNode_SequenceEvaluator.h"
+#include "Animation/AnimNode_SequencePlayer.h"
+#include "Animation/AnimNode_Slot.h"
+#include "Animation/BoneControllers/AnimNode_SpringBone.h"
+#include "Animation/AnimNode_StateMachine.h"
+#include "Animation/AnimNode_TransitionPoseEvaluator.h"
+#include "Animation/AnimNode_TransitionResult.h"
+#include "Animation/BoneControllers/AnimNode_TwoBoneIK.h"
+#include "Animation/AnimNode_TwoWayBlend.h"
+#include "Animation/AnimNode_UseCachedPose.h"
+#include "Animation/BoneControllers/AnimNode_WheelHandler.h"
+#include "Animation/AnimNodeSpaceConversions.h"
+#include "Camera/CameraStackTypes.h"
+#include "Materials/MaterialInstanceBasePropertyOverrides.h"
+#include "Engine/StreamableManager.h"
+#include "Engine/TextureDefines.h"
+#include "Tests/TextPropertyTestObject.h"
+#include "Engine/Texture.h"
+#include "Engine/Texture2D.h"
+#include "Engine/LightMapTexture2D.h"
+#include "Engine/ShadowMapTexture2D.h"
+#include "Engine/TextureLightProfile.h"
+#include "Engine/Texture2DDynamic.h"
+#include "Engine/TextureCube.h"
+#include "Engine/TextureMovie.h"
+#include "Engine/TextureRenderTarget.h"
+#include "Engine/TextureRenderTarget2D.h"
+#include "Engine/CanvasRenderTarget2D.h"
+#include "Engine/TextureRenderTargetCube.h"
+#include "EditorFramework/ThumbnailInfo.h"
+#include "Engine/TimelineTemplate.h"
+#include "GameFramework/TouchInterface.h"
+#include "Engine/UserDefinedEnum.h"
+#include "Engine/UserDefinedStruct.h"
+#include "VectorField/VectorField.h"
+#include "VectorField/VectorFieldAnimated.h"
+#include "VectorField/VectorFieldStatic.h"
+#include "Animation/VertexAnim/VertexAnimBase.h"
+#include "Animation/VertexAnim/MorphTarget.h"
+#include "Animation/VertexAnim/VertexAnimation.h"
+#include "Engine/World.h"
+#include "Engine/WorldComposition.h"
+
 #include "VisualLog.h"
 #include "MaterialShared.h"					// Shared material definitions.
 #include "Components.h"						// Forward declarations of object components of actors
