CIS fix.

diff --git a/Engine/Source/Runtime/Engine/Public/Engine.h b/Engine/Source/Runtime/Engine/Public/Engine.h
--- a/Engine/Source/Runtime/Engine/Public/Engine.h
+++ b/Engine/Source/Runtime/Engine/Public/Engine.h
@@ -635,7 +635,6 @@ struct FURL;
 #include "Lightmass/PrecomputedVisibilityVolume.h"
 #include "Sound/ReverbVolume.h"
 #include "Engine/TriggerVolume.h"
-#include "GameFramework/BugRepro.h"
 #include "Camera/CameraActor.h"
 #include "AI/Navigation/NavAgentInterface.h"
 #include "GameFramework/Controller.h"
@@ -1462,7 +1461,6 @@ struct FURL;
 #include "Sound/SoundNodeSwitch.h"
 #include "Sound/SoundNodeWaveParam.h"
 #include "Sound/SoundNodeWavePlayer.h"
-#include "Spline/Spline.h"
 #include "Engine/StaticMesh.h"
 #include "Engine/StaticMeshSocket.h"
 #include "Animation/AnimNodeBase.h"
