Added "Randomly Select Seed Array" option to random seed functionality

#ttp 311394 - UE4: CASCADE: "Randomly Select Seed Array" missing from settings in Random Seed functionality (was working in UE3)

diff --git a/Engine/Source/Runtime/Engine/Classes/Particles/ParticleModule.h b/Engine/Source/Runtime/Engine/Classes/Particles/ParticleModule.h
--- a/Engine/Source/Runtime/Engine/Classes/Particles/ParticleModule.h
+++ b/Engine/Source/Runtime/Engine/Classes/Particles/ParticleModule.h
@@ -97,6 +97,12 @@ struct FParticleRandomSeedInfo
 	UPROPERTY(EditAnywhere, Category=ParticleRandomSeedInfo)
 	uint32 bResetSeedOnEmitterLooping:1;
 
+	/**
+	*	If true, then randomly select a seed entry from the RandomSeeds array
+	*/
+	UPROPERTY(EditAnywhere, Category = ParticleRandomSeedInfo)
+	uint32 bRandomlySelectSeedArray:1;
+
 	/** 
 	 *	The random seed values to utilize for the module. 
 	 *	More than 1 means the instance will randomly select one.
@@ -110,6 +116,7 @@ struct FParticleRandomSeedInfo
 		: bGetSeedFromInstance(false)
 		, bInstanceSeedIsIndex(false)
 		, bResetSeedOnEmitterLooping(true)
+		, bRandomlySelectSeedArray(false)
 		{
 		}
 		FORCEINLINE int32 GetInstancePayloadSize() const
diff --git a/Engine/Source/Runtime/Engine/Private/Particles/ParticleModules.cpp b/Engine/Source/Runtime/Engine/Private/Particles/ParticleModules.cpp
--- a/Engine/Source/Runtime/Engine/Private/Particles/ParticleModules.cpp
+++ b/Engine/Source/Runtime/Engine/Private/Particles/ParticleModules.cpp
@@ -660,7 +660,15 @@ uint32 UParticleModule::PrepRandomSeedInstancePayload(FParticleEmitterInstance*
 		// Pick a seed to use and initialize it!!!!
 		if (InRandSeedInfo.RandomSeeds.Num() > 0)
 		{
-			InRandSeedPayload->RandomStream.Initialize(InRandSeedInfo.RandomSeeds[0]);
+			if (InRandSeedInfo.bRandomlySelectSeedArray)
+			{
+				int32 Index = FMath::RandRange(0, InRandSeedInfo.RandomSeeds.Num() - 1);
+				InRandSeedPayload->RandomStream.Initialize(InRandSeedInfo.RandomSeeds[Index]);
+			}
+			else
+			{
+				InRandSeedPayload->RandomStream.Initialize(InRandSeedInfo.RandomSeeds[0]);
+			}
 			return 0;
 		}
 	}
