Pass FComponentInstanceDataCache into OnConstruction, allowing us to apply data before running UCS
Remove ComponentInstanceDataCache.h from ActorComponent.h
Remove SplineComponent.h from Engine.h
Change splinecomp funcs to get get location/direction at point on spline to be world space (and normalized direction vector)
Update reparam table when editing spline

diff --git a/Engine/Source/Editor/ComponentVisualizers/Private/SplineComponentVisualizer.cpp b/Engine/Source/Editor/ComponentVisualizers/Private/SplineComponentVisualizer.cpp
--- a/Engine/Source/Editor/ComponentVisualizers/Private/SplineComponentVisualizer.cpp
+++ b/Engine/Source/Editor/ComponentVisualizers/Private/SplineComponentVisualizer.cpp
@@ -157,8 +157,9 @@ bool FSplineComponentVisualizer::HandleInputDelta(FLevelEditorViewportClient* Vi
 		// Convert back to local space
 		SplineComp->SplineInfo.Points[SelectedKeyIndex].OutVal = SplineComp->ComponentToWorld.InverseTransformPosition(NewWorldPos);
 
-		// Update tangents
+		// Update tangents and reparam table
 		SplineComp->SplineInfo.AutoSetTangents();
+		SplineComp->UpdateSplineReparamTable();
 
 		return true;
 	}
@@ -186,9 +187,8 @@ bool FSplineComponentVisualizer::HandleInputKey(FLevelEditorViewportClient* View
 			SplineComp->SplineInfo.Points.RemoveAt(SelectedKeyIndex);
 
 			SplineComp->RefreshSplineInputs(); // update input value for each key
-
 			SplineComp->SplineInfo.AutoSetTangents(); // update tangents
-
+			SplineComp->UpdateSplineReparamTable(); // update reparam table
 			SelectedKeyIndex = INDEX_NONE; // deselect any keys
 
 			bHandled = true; // consume key input
diff --git a/Engine/Source/Editor/ComponentVisualizers/Private/SplineComponentVisualizer.h b/Engine/Source/Editor/ComponentVisualizers/Private/SplineComponentVisualizer.h
--- a/Engine/Source/Editor/ComponentVisualizers/Private/SplineComponentVisualizer.h
+++ b/Engine/Source/Editor/ComponentVisualizers/Private/SplineComponentVisualizer.h
@@ -3,6 +3,7 @@
 #pragma once
 
 #include "ComponentVisualizer.h"
+#include "Components/SplineComponent.h"
 
 struct HSplineVisProxy : public HComponentVisProxy
 {
diff --git a/Engine/Source/Editor/UnrealEd/Private/Kismet2/KismetReinstanceUtilities.cpp b/Engine/Source/Editor/UnrealEd/Private/Kismet2/KismetReinstanceUtilities.cpp
--- a/Engine/Source/Editor/UnrealEd/Private/Kismet2/KismetReinstanceUtilities.cpp
+++ b/Engine/Source/Editor/UnrealEd/Private/Kismet2/KismetReinstanceUtilities.cpp
@@ -5,6 +5,7 @@
 #include "Kismet2/KismetEditorUtilities.h"
 #include "BlueprintEditorUtils.h"
 #include "Layers/ILayers.h"
+#include "ComponentInstanceDataCache.h"
 
 /////////////////////////////////////////////////////////////////////////////////
 // FBlueprintCompileReinstancer
@@ -337,7 +338,7 @@ void FBlueprintCompileReinstancer::ReplaceInstancesOfClass(UClass* OldClass, UCl
 				NewActor->RegisterAllComponents(); // Register native components
 
 				// Run the construction script, which will use the properties we just copied over
-				NewActor->OnConstruction(WorldTransform);
+				NewActor->OnConstruction(WorldTransform, &InstanceDataCache);
 
 				//Attach the new instance to original parent
 				if(AttachParentActor)
@@ -351,9 +352,6 @@ void FBlueprintCompileReinstancer::ReplaceInstancesOfClass(UClass* OldClass, UCl
 					GEditor->ParentActors(NewActor, *AttachIt, FName());
 				}
 
-				// Apply cached data to new instance
-				InstanceDataCache.ApplyToActor(NewActor);
-
 				// Determine if the actor used to be selected.
 				if (OldActor->IsSelected())
 				{
diff --git a/Engine/Source/Runtime/Engine/Classes/Atmosphere/AtmosphericFogComponent.h b/Engine/Source/Runtime/Engine/Classes/Atmosphere/AtmosphericFogComponent.h
--- a/Engine/Source/Runtime/Engine/Classes/Atmosphere/AtmosphericFogComponent.h
+++ b/Engine/Source/Runtime/Engine/Classes/Atmosphere/AtmosphericFogComponent.h
@@ -83,28 +83,7 @@ struct FAtmospherePrecomputeParameters
 	}
 };
 
-/** Used to store lightmap data during RerunConstructionScripts */
-class FAtmospherePrecomputeInstanceData : public FComponentInstanceDataBase
-{
-public:
-	static const FName InstanceDataTypeName;
-
-	virtual ~FAtmospherePrecomputeInstanceData()
-	{}
-
-	// Begin FComponentInstanceDataBase interface
-	virtual FName GetDataTypeName() const OVERRIDE
-	{
-		return InstanceDataTypeName;
-	}
-	// End FComponentInstanceDataBase interface
 
-	struct FAtmospherePrecomputeParameters PrecomputeParameter;
-
-	FByteBulkData TransmittanceData;
-	FByteBulkData IrradianceData;
-	FByteBulkData InscatterData;
-};
 
 /**
  *	Used to create fogging effects such as clouds.
diff --git a/Engine/Source/Runtime/Engine/Classes/Components/ActorComponent.h b/Engine/Source/Runtime/Engine/Classes/Components/ActorComponent.h
--- a/Engine/Source/Runtime/Engine/Classes/Components/ActorComponent.h
+++ b/Engine/Source/Runtime/Engine/Classes/Components/ActorComponent.h
@@ -2,7 +2,6 @@
 
 #pragma once
 
-#include "ComponentInstanceDataCache.h"
 #include "Engine/EngineTypes.h"
 #include "Interfaces/Interface_AssetUserData.h"
 
@@ -417,9 +416,9 @@ public:
 	virtual void PostNetReceive() { }
 
 	/** Called before we throw away components during RerunConstructionScripts, to cache any data we wish to persist across that operation */
-	virtual void GetComponentInstanceData(FComponentInstanceDataCache& Cache) const {}
+	virtual void GetComponentInstanceData(class FComponentInstanceDataCache& Cache) const {}
 	/** Called after we create new components during RerunConstructionScripts, to optionally apply any data backed up during GetComponentInstanceData */
-	virtual void ApplyComponentInstanceData(const FComponentInstanceDataCache& Cache) {}
+	virtual void ApplyComponentInstanceData(const class FComponentInstanceDataCache& Cache) {}
 
 	// Begin UObject interface.
 	virtual void BeginDestroy() OVERRIDE;
diff --git a/Engine/Source/Runtime/Engine/Classes/Components/LightComponent.h b/Engine/Source/Runtime/Engine/Classes/Components/LightComponent.h
--- a/Engine/Source/Runtime/Engine/Classes/Components/LightComponent.h
+++ b/Engine/Source/Runtime/Engine/Classes/Components/LightComponent.h
@@ -4,28 +4,6 @@
 #pragma once
 #include "LightComponent.generated.h"
 
-/** Used to store lightmap data during RerunConstructionScripts */
-class FPrecomputedLightInstanceData : public FComponentInstanceDataBase
-{
-public:
-	static const FName PrecomputedLightInstanceDataTypeName;
-
-	virtual ~FPrecomputedLightInstanceData()
-	{}
-
-	// Begin FComponentInstanceDataBase interface
-	virtual FName GetDataTypeName() const OVERRIDE
-	{
-		return PrecomputedLightInstanceDataTypeName;
-	}
-	// End FComponentInstanceDataBase interface
-
-	FTransform Transform;
-	FGuid LightGuid;
-	int32 ShadowMapChannel;
-	int32 PreviewShadowMapChannel;
-	bool bPrecomputedLightingIsValid;
-};
 
 UCLASS(abstract, HideCategories=(Trigger,Activation,"Components|Activation",Physics), ShowCategories=(Mobility))
 class ENGINE_API ULightComponent : public ULightComponentBase
diff --git a/Engine/Source/Runtime/Engine/Classes/Components/SplineComponent.h b/Engine/Source/Runtime/Engine/Classes/Components/SplineComponent.h
--- a/Engine/Source/Runtime/Engine/Classes/Components/SplineComponent.h
+++ b/Engine/Source/Runtime/Engine/Classes/Components/SplineComponent.h
@@ -4,28 +4,7 @@
 #pragma once
 #include "SplineComponent.generated.h"
 
-/** Used to store spline data during RerunConstructionScripts */
-class FSplineInstanceData : public FComponentInstanceDataBase
-{
-public:
-	static const FName SplineInstanceDataTypeName;
-
-	virtual ~FSplineInstanceData()
-	{}
-
-	// Begin FComponentInstanceDataBase interface
-	virtual FName GetDataTypeName() const OVERRIDE
-	{
-		return SplineInstanceDataTypeName;
-	}
-	// End FComponentInstanceDataBase interface
 
-	/** Name of component */
-	FName ComponentName;
-
-	/** Lightmaps from LODData */
-	FInterpCurveVector SplineInfo;
-};
 
 UCLASS(ClassGroup=Shapes, meta=(BlueprintSpawnableComponent))
 class ENGINE_API USplineComponent : public USceneComponent
@@ -54,11 +33,11 @@ class ENGINE_API USplineComponent : public USceneComponent
 	
 	/** Given a distance along the length of this spline, return the point in space where this puts you */
 	UFUNCTION(BlueprintCallable, Category=Spline)
-	FVector GetLocationAtDistanceAlongSpline(float Distance) const;
+	FVector GetWorldLocationAtDistanceAlongSpline(float Distance) const;
 	
-	/** Given a distance along the length of this spline, return the direction of the spline there. Note, result is non-unit length. */
+	/** Given a distance along the length of this spline, return a unit direction vector of the spline there. */
 	UFUNCTION(BlueprintCallable, Category=Spline)
-	FVector GetTangentAtDistanceAlongSpline(float Distance) const;
+	FVector GetWorldDirectionAtDistanceAlongSpline(float Distance) const;
 
 	/** Walk through keys and set time for each one */
 	void RefreshSplineInputs();
diff --git a/Engine/Source/Runtime/Engine/Classes/GameFramework/Actor.h b/Engine/Source/Runtime/Engine/Classes/GameFramework/Actor.h
--- a/Engine/Source/Runtime/Engine/Classes/GameFramework/Actor.h
+++ b/Engine/Source/Runtime/Engine/Classes/GameFramework/Actor.h
@@ -1695,9 +1695,10 @@ public:
 
 	/**
 	 * Called when an instance of this class is placed (in editor) or spawned.
-	 * @param	Transform	The transform to construct the actor at.
+	 * @param	Transform			The transform to construct the actor at.
+	 * @param	InstanceDataCache	Optional cache of state to apply to newly created components (e.g. precomputed lighting)
 	 */
-	virtual void OnConstruction(const FTransform& Transform);
+	virtual void OnConstruction(const FTransform& Transform, const FComponentInstanceDataCache* InstanceDataCache = NULL);
 
 	/**
 	 * Helper function to regoster tje specified component, and add it to the serialized components array
diff --git a/Engine/Source/Runtime/Engine/Private/ActorConstruction.cpp b/Engine/Source/Runtime/Engine/Private/ActorConstruction.cpp
--- a/Engine/Source/Runtime/Engine/Private/ActorConstruction.cpp
+++ b/Engine/Source/Runtime/Engine/Private/ActorConstruction.cpp
@@ -4,6 +4,7 @@
 #include "BlueprintUtilities.h"
 #include "LatentActions.h"
 #include "DeferRegisterComponents.h"
+#include "ComponentInstanceDataCache.h"
 
 #if WITH_EDITOR
 #include "Editor/UnrealEd/Public/Kismet2/BlueprintEditorUtils.h"
@@ -209,7 +210,7 @@ void AActor::RerunConstructionScripts()
 		}
 
 		// Run the construction scripts
-		OnConstruction(OldTransform);
+		OnConstruction(OldTransform, &InstanceDataCache);
 
 		if(Parent)
 		{
@@ -221,9 +222,6 @@ void AActor::RerunConstructionScripts()
 			}
 		}
 
-		// Apply per-instance data.
-		InstanceDataCache.ApplyToActor(this);
-
 		// If we had attached children reattach them now - unless they are already attached
 		for(FAttachedActorInfo& Info : AttachedActorInfos)
 		{
@@ -244,7 +242,7 @@ void AActor::RerunConstructionScripts()
 	}
 }
 
-void AActor::OnConstruction(const FTransform& Transform)
+void AActor::OnConstruction(const FTransform& Transform, const FComponentInstanceDataCache* InstanceDataCache)
 {
 	check(!IsPendingKill());
 	check(!HasAnyFlags(RF_BeginDestroyed|RF_FinishDestroyed));
@@ -274,6 +272,12 @@ void AActor::OnConstruction(const FTransform& Transform)
 				CurrentBPGClass->CreateComponentsForActor(this);
 			}
 
+			// If we passed in cached data, we apply it now, so that the UserConstructionScript can use the updated values
+			if(InstanceDataCache)
+			{
+				InstanceDataCache->ApplyToActor(this);
+			}
+
 #if WITH_EDITOR
 			bool bDoUserConstructionScript;
 			GConfig->GetBool(TEXT("Kismet"), TEXT("bTurnOffEditorConstructionScript"), bDoUserConstructionScript, GEngineIni);
@@ -286,6 +290,13 @@ void AActor::OnConstruction(const FTransform& Transform)
 
 			// Bind any delegates on components			
 			((UBlueprintGeneratedClass*)GetClass())->BindDynamicDelegates(this); // We have a BP stack, we must have a UBlueprintGeneratedClass...
+
+			// Apply any cached data procedural components
+			// @TODO Don't re-apply to components we already applied to above
+			if (InstanceDataCache)
+			{
+				InstanceDataCache->ApplyToActor(this);
+			}
 		}
 		else
 		{
diff --git a/Engine/Source/Runtime/Engine/Private/Atmosphere/Atmosphere.cpp b/Engine/Source/Runtime/Engine/Private/Atmosphere/Atmosphere.cpp
--- a/Engine/Source/Runtime/Engine/Private/Atmosphere/Atmosphere.cpp
+++ b/Engine/Source/Runtime/Engine/Private/Atmosphere/Atmosphere.cpp
@@ -3,6 +3,7 @@
 #include "EnginePrivate.h"
 #include "../../../Renderer/Private/ScenePrivate.h"
 #include "Atmosphere.h"
+#include "ComponentInstanceDataCache.h"
 
 #if WITH_EDITOR
 #include "ObjectEditorUtils.h"
@@ -681,6 +682,29 @@ void UAtmosphericFogComponent::Serialize(FArchive& Ar)
 	}
 }
 
+/** Used to store lightmap data during RerunConstructionScripts */
+class FAtmospherePrecomputeInstanceData : public FComponentInstanceDataBase
+{
+public:
+	static const FName InstanceDataTypeName;
+
+	virtual ~FAtmospherePrecomputeInstanceData()
+	{}
+
+	// Begin FComponentInstanceDataBase interface
+	virtual FName GetDataTypeName() const OVERRIDE
+	{
+		return InstanceDataTypeName;
+	}
+	// End FComponentInstanceDataBase interface
+
+	struct FAtmospherePrecomputeParameters PrecomputeParameter;
+
+	FByteBulkData TransmittanceData;
+	FByteBulkData IrradianceData;
+	FByteBulkData InscatterData;
+};
+
 const FName FAtmospherePrecomputeInstanceData::InstanceDataTypeName(TEXT("AtmospherePrecomputedInstanceData"));
 
 // Backup the precomputed data before re-running Blueprint construction script
diff --git a/Engine/Source/Runtime/Engine/Private/ComponentInstanceDataCache.cpp b/Engine/Source/Runtime/Engine/Private/ComponentInstanceDataCache.cpp
--- a/Engine/Source/Runtime/Engine/Private/ComponentInstanceDataCache.cpp
+++ b/Engine/Source/Runtime/Engine/Private/ComponentInstanceDataCache.cpp
@@ -1,6 +1,7 @@
 // Copyright 1998-2014 Epic Games, Inc. All Rights Reserved.
 
 #include "EnginePrivate.h"
+#include "ComponentInstanceDataCache.h"
 
 FComponentInstanceDataBase::FComponentInstanceDataBase(const UActorComponent* SourceComponent)
 {
@@ -94,7 +95,7 @@ void FComponentInstanceDataCache::GetFromActor(AActor* Actor, FComponentInstance
 	}
 }
 
-void FComponentInstanceDataCache::ApplyToActor(AActor* Actor)
+void FComponentInstanceDataCache::ApplyToActor(AActor* Actor) const
 {
 	if(Actor != NULL)
 	{
diff --git a/Engine/Source/Runtime/Engine/Private/InstancedStaticMesh.cpp b/Engine/Source/Runtime/Engine/Private/InstancedStaticMesh.cpp
--- a/Engine/Source/Runtime/Engine/Private/InstancedStaticMesh.cpp
+++ b/Engine/Source/Runtime/Engine/Private/InstancedStaticMesh.cpp
@@ -7,80 +7,10 @@
 #include "EnginePrivate.h"
 #include "ShaderParameters.h"
 
-#include "Components/ApplicationLifecycleComponent.h"
-#include "AI/BehaviorTree/BlackboardComponent.h"
-#include "AI/BrainComponent.h"
-#include "AI/BehaviorTree/BehaviorTreeComponent.h"
-#include "Debug/GameplayDebuggingControllerComponent.h"
-#include "Components/InputComponent.h"
-#include "GameFramework/MovementComponent.h"
-#include "Components/SceneComponent.h"
-#include "Components/LightComponentBase.h"
-#include "Components/LightComponent.h"
-#include "Components/PrimitiveComponent.h"
-#include "GameFramework/NavMovementComponent.h"
-#include "GameFramework/PawnMovementComponent.h"
-#include "GameFramework/CharacterMovementComponent.h"
-#include "GameFramework/FloatingPawnMovement.h"
-#include "GameFramework/SpectatorPawnMovement.h"
-#include "Vehicles/WheeledVehicleMovementComponent.h"
-#include "Vehicles/WheeledVehicleMovementComponent4W.h"
-#include "GameFramework/ProjectileMovementComponent.h"
-#include "GameFramework/RotatingMovementComponent.h"
-#include "AI/Navigation/NavigationComponent.h"
-#include "AI/Navigation/PathFollowingComponent.h"
-#include "Components/PawnNoiseEmitterComponent.h"
-#include "Components/PawnSensingComponent.h"
-#include "PhysicsEngine/PhysicsHandleComponent.h"
-#include "Atmosphere/AtmosphericFogComponent.h"
-#include "Components/AudioComponent.h"
-#include "Camera/CameraComponent.h"
-#include "Components/ChildActorComponent.h"
-#include "Components/ExponentialHeightFogComponent.h"
-#include "Components/SkyLightComponent.h"
-#include "AI/Navigation/NavigationGraphNodeComponent.h"
-#include "PhysicsEngine/PhysicsConstraintComponent.h"
-#include "PhysicsEngine/PhysicsThrusterComponent.h"
-#include "Components/PostProcessComponent.h"
-#include "Components/ArrowComponent.h"
-#include "Components/BillboardComponent.h"
-#include "Components/BrushComponent.h"
-#include "Components/DrawFrustumComponent.h"
-#include "AI/EnvironmentQuery/EQSRenderingComponent.h"
-#include "Debug/GameplayDebuggingComponent.h"
-#include "Components/LineBatchComponent.h"
-#include "Components/MaterialBillboardComponent.h"
-#include "Components/MeshComponent.h"
-#include "Components/SkinnedMeshComponent.h"
-#include "Components/DestructibleComponent.h"
-#include "Components/PoseableMeshComponent.h"
-#include "Components/SkeletalMeshComponent.h"
-#include "Components/StaticMeshComponent.h"
 #include "Components/InstancedStaticMeshComponent.h"
-#include "Components/InteractiveFoliageComponent.h"
-#include "Components/SplineMeshComponent.h"
-#include "Components/ModelComponent.h"
-#include "AI/Navigation/NavLinkRenderingComponent.h"
-#include "AI/Navigation/NavMeshRenderingComponent.h"
-#include "AI/Navigation/NavTestRenderingComponent.h"
-#include "Components/NiagaraComponent.h"
-#include "Components/ShapeComponent.h"
-#include "Components/BoxComponent.h"
-#include "Components/CapsuleComponent.h"
-#include "Components/SphereComponent.h"
-#include "Components/DrawSphereComponent.h"
-#include "Components/TextRenderComponent.h"
-#include "Components/VectorFieldComponent.h"
-#include "PhysicsEngine/RadialForceComponent.h"
-#include "GameFramework/SpringArmComponent.h"
-#include "Components/WindDirectionalSourceComponent.h"
-#include "Components/TimelineComponent.h"
-#include "Slate.h"
-#include "AI/NavDataGenerator.h"
-#include "OnlineSubsystemUtils.h"
-
 #include "StaticMeshLight.h"
 #include "SpeedTreeWind.h"
+#include "ComponentInstanceDataCache.h"
 
 #if WITH_PHYSX
 #include "PhysicsEngine/PhysXSupport.h"
diff --git a/Engine/Source/Runtime/Engine/Private/LightComponent.cpp b/Engine/Source/Runtime/Engine/Private/LightComponent.cpp
--- a/Engine/Source/Runtime/Engine/Private/LightComponent.cpp
+++ b/Engine/Source/Runtime/Engine/Private/LightComponent.cpp
@@ -10,6 +10,7 @@
 #endif
 #include "MessageLog.h"
 #include "UObjectToken.h"
+#include "ComponentInstanceDataCache.h"
 
 void ULightComponentBase::Serialize(FArchive& Ar)
 {
@@ -750,6 +751,29 @@ void ULightComponent::InvalidateLightingCacheInner(bool bRecreateLightGuids)
 	MarkRenderStateDirty();
 }
 
+/** Used to store lightmap data during RerunConstructionScripts */
+class FPrecomputedLightInstanceData : public FComponentInstanceDataBase
+{
+public:
+	static const FName PrecomputedLightInstanceDataTypeName;
+
+	virtual ~FPrecomputedLightInstanceData()
+	{}
+
+	// Begin FComponentInstanceDataBase interface
+	virtual FName GetDataTypeName() const OVERRIDE
+	{
+		return PrecomputedLightInstanceDataTypeName;
+	}
+	// End FComponentInstanceDataBase interface
+
+	FTransform Transform;
+	FGuid LightGuid;
+	int32 ShadowMapChannel;
+	int32 PreviewShadowMapChannel;
+	bool bPrecomputedLightingIsValid;
+};
+
 // Init type name static
 const FName FPrecomputedLightInstanceData::PrecomputedLightInstanceDataTypeName(TEXT("PrecomputedLightInstanceData"));
 
diff --git a/Engine/Source/Runtime/Engine/Private/SkyLightComponent.cpp b/Engine/Source/Runtime/Engine/Private/SkyLightComponent.cpp
--- a/Engine/Source/Runtime/Engine/Private/SkyLightComponent.cpp
+++ b/Engine/Source/Runtime/Engine/Private/SkyLightComponent.cpp
@@ -12,6 +12,7 @@
 #include "UObjectToken.h"
 #include "Net/UnrealNetwork.h"
 #include "MapErrors.h"
+#include "ComponentInstanceDataCache.h"
 
 #define LOCTEXT_NAMESPACE "SkyLightComponent"
 
diff --git a/Engine/Source/Runtime/Engine/Private/SplineComponent.cpp b/Engine/Source/Runtime/Engine/Private/SplineComponent.cpp
--- a/Engine/Source/Runtime/Engine/Private/SplineComponent.cpp
+++ b/Engine/Source/Runtime/Engine/Private/SplineComponent.cpp
@@ -5,6 +5,8 @@
 =============================================================================*/
 
 #include "EnginePrivate.h"
+#include "Components/SplineComponent.h"
+#include "ComponentInstanceDataCache.h"
 
 
 USplineComponent::USplineComponent(const class FPostConstructInitializeProperties& PCIP)
@@ -18,6 +20,7 @@ USplineComponent::USplineComponent(const class FPostConstructInitializePropertie
 	SplineInfo.Points[PointIndex].InterpMode = CIM_CurveAuto;
 
 	SplineInfo.AutoSetTangents();
+	UpdateSplineReparamTable();
 }
 
 void USplineComponent::UpdateSplineReparamTable()
@@ -72,19 +75,19 @@ float USplineComponent::GetSplineLength() const
 }
 
 
-FVector USplineComponent::GetLocationAtDistanceAlongSpline(float Distance) const
+FVector USplineComponent::GetWorldLocationAtDistanceAlongSpline(float Distance) const
 {
 	const float Param = SplineReparamTable.Eval(Distance, 0.f);
 	const FVector Location = SplineInfo.Eval(Param, FVector::ZeroVector);
-	return Location;
+	return ComponentToWorld.TransformPosition(Location);
 }
 
 
-FVector USplineComponent::GetTangentAtDistanceAlongSpline(float Distance) const
+FVector USplineComponent::GetWorldDirectionAtDistanceAlongSpline(float Distance) const
 {
 	const float Param = SplineReparamTable.Eval(Distance, 0.f);
-	const FVector Tangent = SplineInfo.EvalDerivative(Param, FVector::ZeroVector);
-	return Tangent;
+	const FVector Tangent = SplineInfo.EvalDerivative(Param, FVector::ZeroVector).SafeNormal();
+	return ComponentToWorld.TransformVectorNoScale(Tangent);
 }
 
 void USplineComponent::RefreshSplineInputs()
@@ -95,6 +98,29 @@ void USplineComponent::RefreshSplineInputs()
 	}
 }
 
+/** Used to store spline data during RerunConstructionScripts */
+class FSplineInstanceData : public FComponentInstanceDataBase
+{
+public:
+	static const FName SplineInstanceDataTypeName;
+
+	virtual ~FSplineInstanceData()
+	{}
+
+	// Begin FComponentInstanceDataBase interface
+	virtual FName GetDataTypeName() const OVERRIDE
+	{
+		return SplineInstanceDataTypeName;
+	}
+	// End FComponentInstanceDataBase interface
+
+	/** Name of component */
+	FName ComponentName;
+
+	/** Lightmaps from LODData */
+	FInterpCurveVector SplineInfo;
+};
+
 // Init type name static
 const FName FSplineInstanceData::SplineInstanceDataTypeName(TEXT("SplineInstanceData"));
 
@@ -123,7 +149,7 @@ void USplineComponent::ApplyComponentInstanceData(const FComponentInstanceDataCa
 		if (SplineData->ComponentName == GetFName())
 		{
 			SplineInfo = SplineData->SplineInfo;
-			UpdateSplineReparamTable();			
+			UpdateSplineReparamTable();
 		}
 	}
 }
\ No newline at end of file
diff --git a/Engine/Source/Runtime/Engine/Private/StaticMeshComponent.cpp b/Engine/Source/Runtime/Engine/Private/StaticMeshComponent.cpp
--- a/Engine/Source/Runtime/Engine/Private/StaticMeshComponent.cpp
+++ b/Engine/Source/Runtime/Engine/Private/StaticMeshComponent.cpp
@@ -5,6 +5,7 @@
 #include "MessageLog.h"
 #include "UObjectToken.h"
 #include "MapErrors.h"
+#include "ComponentInstanceDataCache.h"
 
 #define LOCTEXT_NAMESPACE "StaticMeshComponent"
 
diff --git a/Engine/Source/Runtime/Engine/Public/ComponentInstanceDataCache.h b/Engine/Source/Runtime/Engine/Public/ComponentInstanceDataCache.h
--- a/Engine/Source/Runtime/Engine/Public/ComponentInstanceDataCache.h
+++ b/Engine/Source/Runtime/Engine/Public/ComponentInstanceDataCache.h
@@ -54,7 +54,7 @@ public:
 	static void GetFromActor(AActor* Actor, FComponentInstanceDataCache& Cache);
 
 	/** Util to iterate over components and apply data to each */
-	void ApplyToActor(AActor* Actor);
+	void ApplyToActor(AActor* Actor) const;
 
 private:
 	/** Map of data type name to data of that type */
diff --git a/Engine/Source/Runtime/Engine/Public/Engine.h b/Engine/Source/Runtime/Engine/Public/Engine.h
--- a/Engine/Source/Runtime/Engine/Public/Engine.h
+++ b/Engine/Source/Runtime/Engine/Public/Engine.h
@@ -830,7 +830,6 @@ struct FURL;
 #include "Components/CapsuleComponent.h"
 #include "Components/SphereComponent.h"
 #include "Components/DrawSphereComponent.h"
-#include "Components/SplineComponent.h"
 #include "Components/TextRenderComponent.h"
 #include "Components/VectorFieldComponent.h"
 #include "PhysicsEngine/RadialForceComponent.h"
