Fixed problem where saving a level with a duplicated matinee actor will result in an unloadable file.
#ttp 335604 - LICENSEE: UDN: EDITOR: CRASH: Map with duplicated MatineeActor crashes on load
#branch UE4
#proj Runtime.Engine
#add Added RF_TextExportTransient flag on the UInterpFilter objects which are created post load/initialize, thus ensuring that they are not copied when duplicated.

diff --git a/Engine/Source/Runtime/Engine/Private/Interpolation.cpp b/Engine/Source/Runtime/Engine/Private/Interpolation.cpp
--- a/Engine/Source/Runtime/Engine/Private/Interpolation.cpp
+++ b/Engine/Source/Runtime/Engine/Private/Interpolation.cpp
@@ -1698,36 +1698,36 @@ UInterpData::UInterpData(const class FPostConstructInitializeProperties& PCIP)
 #if WITH_EDITORONLY_DATA
 void UInterpData::CreateDefaultFilters()
 {
-	UInterpFilter* FilterAll = NewNamedObject<UInterpFilter>(this, TEXT("FilterAll"), RF_Transient);
+	UInterpFilter* FilterAll = NewNamedObject<UInterpFilter>(this, TEXT("FilterAll"), RF_Transient | RF_TextExportTransient);
 	FilterAll->Caption = TEXT("All");
 	DefaultFilters.Add(FilterAll);
 
-	UInterpFilter_Classes* FilterCameras = NewNamedObject<UInterpFilter_Classes>(this, TEXT("FilterCameras"), RF_Transient);
+	UInterpFilter_Classes* FilterCameras = NewNamedObject<UInterpFilter_Classes>(this, TEXT("FilterCameras"), RF_Transient | RF_TextExportTransient);
 	FilterCameras->Caption = TEXT("Cameras");
 	FilterCameras->ClassToFilterBy = ACameraActor::StaticClass();
 	DefaultFilters.Add(FilterCameras);
 
-	UInterpFilter_Classes* FilterSkeletalMeshes = NewNamedObject<UInterpFilter_Classes>(this, TEXT("FilterSkeletalMeshes"), RF_Transient);
+	UInterpFilter_Classes* FilterSkeletalMeshes = NewNamedObject<UInterpFilter_Classes>(this, TEXT("FilterSkeletalMeshes"), RF_Transient | RF_TextExportTransient);
 	FilterSkeletalMeshes->Caption = TEXT("Skeletal Meshes");
 	FilterSkeletalMeshes->ClassToFilterBy = ASkeletalMeshActor::StaticClass();
 	DefaultFilters.Add(FilterSkeletalMeshes);
 
-	UInterpFilter_Classes* FilterLighting = NewNamedObject<UInterpFilter_Classes>(this, TEXT("FilterLighting"), RF_Transient);
+	UInterpFilter_Classes* FilterLighting = NewNamedObject<UInterpFilter_Classes>(this, TEXT("FilterLighting"), RF_Transient | RF_TextExportTransient);
 	FilterLighting->Caption = TEXT("Lights");
 	FilterLighting->ClassToFilterBy = ALight::StaticClass();
 	DefaultFilters.Add(FilterLighting);
 
-	UInterpFilter_Classes* FilterEmitters = NewNamedObject<UInterpFilter_Classes>(this, TEXT("FilterEmitters"), RF_Transient);
+	UInterpFilter_Classes* FilterEmitters = NewNamedObject<UInterpFilter_Classes>(this, TEXT("FilterEmitters"), RF_Transient | RF_TextExportTransient);
 	FilterEmitters->Caption = TEXT("Particles");
 	FilterEmitters->ClassToFilterBy = AEmitter::StaticClass();
 	DefaultFilters.Add(FilterEmitters);
 
-	UInterpFilter_Classes* FilterSounds = NewNamedObject<UInterpFilter_Classes>(this, TEXT("FilterSounds"), RF_Transient);
+	UInterpFilter_Classes* FilterSounds = NewNamedObject<UInterpFilter_Classes>(this, TEXT("FilterSounds"), RF_Transient | RF_TextExportTransient);
 	FilterSounds->Caption = TEXT("Sounds");
 	FilterSounds->TrackClasses.Add(UInterpTrackSound::StaticClass());
 	DefaultFilters.Add(FilterSounds);
 
-	UInterpFilter_Classes* FilterEvents = NewNamedObject<UInterpFilter_Classes>(this, TEXT("FilterEvents"), RF_Transient);
+	UInterpFilter_Classes* FilterEvents = NewNamedObject<UInterpFilter_Classes>(this, TEXT("FilterEvents"), RF_Transient | RF_TextExportTransient);
 	FilterEvents->Caption = TEXT("Events");
 	FilterEvents->TrackClasses.Add(UInterpTrackEvent::StaticClass());
 	DefaultFilters.Add(FilterEvents);
