Refactor code so that notification virtual OnConstruct does not change signature, added new function ExecuteConstruction that takes the FComponentInstanceDataCache

diff --git a/Engine/Source/Editor/UnrealEd/Private/Kismet2/KismetReinstanceUtilities.cpp b/Engine/Source/Editor/UnrealEd/Private/Kismet2/KismetReinstanceUtilities.cpp
--- a/Engine/Source/Editor/UnrealEd/Private/Kismet2/KismetReinstanceUtilities.cpp
+++ b/Engine/Source/Editor/UnrealEd/Private/Kismet2/KismetReinstanceUtilities.cpp
@@ -338,7 +338,7 @@ void FBlueprintCompileReinstancer::ReplaceInstancesOfClass(UClass* OldClass, UCl
 				NewActor->RegisterAllComponents(); // Register native components
 
 				// Run the construction script, which will use the properties we just copied over
-				NewActor->OnConstruction(WorldTransform, &InstanceDataCache);
+				NewActor->ExecuteConstruction(WorldTransform, &InstanceDataCache);
 
 				//Attach the new instance to original parent
 				if(AttachParentActor)
diff --git a/Engine/Source/Runtime/Engine/Classes/GameFramework/Actor.h b/Engine/Source/Runtime/Engine/Classes/GameFramework/Actor.h
--- a/Engine/Source/Runtime/Engine/Classes/GameFramework/Actor.h
+++ b/Engine/Source/Runtime/Engine/Classes/GameFramework/Actor.h
@@ -1694,11 +1694,17 @@ public:
 	void DebugShowOneComponentHierarchy( USceneComponent* SceneComp, int32& NestLevel, bool bShowPosition );
 
 	/**
-	 * Called when an instance of this class is placed (in editor) or spawned.
+	 * Run any construction script for this Actor. Will call OnConstruction.
 	 * @param	Transform			The transform to construct the actor at.
 	 * @param	InstanceDataCache	Optional cache of state to apply to newly created components (e.g. precomputed lighting)
 	 */
-	virtual void OnConstruction(const FTransform& Transform, const FComponentInstanceDataCache* InstanceDataCache = NULL);
+	void ExecuteConstruction(const FTransform& Transform, const FComponentInstanceDataCache* InstanceDataCache);
+
+	/**
+	 * Called when an instance of this class is placed (in editor) or spawned.
+	 * @param	Transform			The transform the actor was constructed at.
+	 */
+	virtual void OnConstruction(const FTransform& Transform) {}
 
 	/**
 	 * Helper function to regoster tje specified component, and add it to the serialized components array
diff --git a/Engine/Source/Runtime/Engine/Private/ActorConstruction.cpp b/Engine/Source/Runtime/Engine/Private/ActorConstruction.cpp
--- a/Engine/Source/Runtime/Engine/Private/ActorConstruction.cpp
+++ b/Engine/Source/Runtime/Engine/Private/ActorConstruction.cpp
@@ -210,7 +210,7 @@ void AActor::RerunConstructionScripts()
 		}
 
 		// Run the construction scripts
-		OnConstruction(OldTransform, &InstanceDataCache);
+		ExecuteConstruction(OldTransform, &InstanceDataCache);
 
 		if(Parent)
 		{
@@ -242,7 +242,7 @@ void AActor::RerunConstructionScripts()
 	}
 }
 
-void AActor::OnConstruction(const FTransform& Transform, const FComponentInstanceDataCache* InstanceDataCache)
+void AActor::ExecuteConstruction(const FTransform& Transform, const FComponentInstanceDataCache* InstanceDataCache)
 {
 	check(!IsPendingKill());
 	check(!HasAnyFlags(RF_BeginDestroyed|RF_FinishDestroyed));
@@ -317,6 +317,9 @@ void AActor::OnConstruction(const FTransform& Transform, const FComponentInstanc
 			}
 		}
 	}
+
+	// Now run virtual notification
+	OnConstruction(Transform);
 }
 
 
