FRichCurveKey is now serialized natively in order to save space.  Previously each key was 288 bytes due to uproperty tagging so for curves with a lot of keys (e.g, animations), the saved asset becomes huge.

diff --git a/Engine/Source/Runtime/Core/Public/UObject/ObjectVersion.h b/Engine/Source/Runtime/Core/Public/UObject/ObjectVersion.h
--- a/Engine/Source/Runtime/Core/Public/UObject/ObjectVersion.h
+++ b/Engine/Source/Runtime/Core/Public/UObject/ObjectVersion.h
@@ -544,6 +544,8 @@ enum EUnrealEngineObjectUE4Version
 	VER_UE4_SPEEDTREE_WIND_V7,
 	// NeedsLoadForEditorGame added
 	VER_UE4_LOAD_FOR_EDITOR_GAME,
+	// Manual serialization of FRichCurveKey to save space
+	VER_UE4_SERIALIZE_RICH_CURVE_KEY,
 
 	// -----<new versions can be added before this line>-------------------------------------------------
 	// - this needs to be the last line (see note below)
diff --git a/Engine/Source/Runtime/Engine/Classes/Curves/CurveBase.h b/Engine/Source/Runtime/Engine/Classes/Curves/CurveBase.h
--- a/Engine/Source/Runtime/Engine/Classes/Curves/CurveBase.h
+++ b/Engine/Source/Runtime/Engine/Classes/Curves/CurveBase.h
@@ -238,8 +238,30 @@ struct ENGINE_API FRichCurveKey
 	FRichCurveKey(const FInterpCurvePoint<float>& InPoint);
 	FRichCurveKey(const FInterpCurvePoint<FVector>& InPoint, int32 ComponentIndex);
 
-	/** Determine if two RichCurvesKeys are the same */
-	bool operator == (const FRichCurveKey& Curve) const;
+	/** ICPPStructOps interface */
+	bool Serialize(FArchive& Ar);
+	bool operator==(const FRichCurveKey& Other) const;
+	bool operator!=(const FRichCurveKey& Other) const;
+
+	friend FArchive& operator<<(FArchive& Ar, FRichCurveKey& P)
+	{
+		P.Serialize(Ar);
+		return Ar;
+	}
+
+};
+
+template<> struct TIsPODType<FRichCurveKey> { enum { Value = true }; };
+
+template<>
+struct TStructOpsTypeTraits< FRichCurveKey > : public TStructOpsTypeTraitsBase
+{
+	enum
+	{
+		WithSerializer = true,
+		WithCopy = false,
+		WithIdenticalViaEquality = true,
+	};
 };
 
 /** A rich, editable float curve */
diff --git a/Engine/Source/Runtime/Engine/Private/CurveBase.cpp b/Engine/Source/Runtime/Engine/Private/CurveBase.cpp
--- a/Engine/Source/Runtime/Engine/Private/CurveBase.cpp
+++ b/Engine/Source/Runtime/Engine/Private/CurveBase.cpp
@@ -208,6 +208,30 @@ FRichCurveKey::FRichCurveKey(const FInterpCurvePoint<FVector>& InPoint, int32 Co
 	LeaveTangentWeight = 0.f;
 }
 
+bool FRichCurveKey::Serialize(FArchive& Ar)
+{
+	if( Ar.UE4Ver() >= VER_UE4_SERIALIZE_RICH_CURVE_KEY )
+	{
+		// Serialization is handled manually to avoid the extra size overhead of UProperty tagging.
+		// Otherwise with many keys in a rich curve the size can become quite large.
+		Ar << InterpMode;
+		Ar << TangentMode;
+		Ar << TangentWeightMode;
+		Ar << Time;
+		Ar << Value;
+		Ar << ArriveTangent;
+		Ar << ArriveTangentWeight;
+		Ar << LeaveTangent;
+		Ar << LeaveTangentWeight;
+
+		return true;
+	}
+	else
+	{
+		return false;
+	}
+}
+
 bool FRichCurveKey::operator==( const FRichCurveKey& Curve ) const
 {
 	return (Time == Curve.Time) && (Value == Curve.Value) && (InterpMode == Curve.InterpMode) &&
@@ -217,6 +241,11 @@ bool FRichCurveKey::operator==( const FRichCurveKey& Curve ) const
 
 }
 
+bool FRichCurveKey::operator!=(const FRichCurveKey& Other) const
+{
+	return !(*this == Other);
+}
+
 //////////////////////////////////////////////////////////////////////////
 // FRichCurve
 
