Fix case of non uniform scaling where capsules keep growing.
Refactor update scale for body instance
Fix some more cases where show collision and pxvis collision don't agree
TTP#334723

diff --git a/Engine/Source/Editor/PhAT/Private/PhATPreviewViewportClient.cpp b/Engine/Source/Editor/PhAT/Private/PhATPreviewViewportClient.cpp
--- a/Engine/Source/Editor/PhAT/Private/PhATPreviewViewportClient.cpp
+++ b/Engine/Source/Editor/PhAT/Private/PhATPreviewViewportClient.cpp
@@ -950,6 +950,7 @@ void FPhATEdPreviewViewportClient::ModifyPrimitiveSize(int32 BodyIndex, EKCollis
 	}
 	else if (PrimType == KPT_Sphyl)
 	{
+		FKSphylElem & SphylElem = AggGeom->SphylElems[PrimIndex];
 		float DeltaRadius = DeltaSize.X;
 		if (FMath::Abs(DeltaSize.Y) > FMath::Abs(DeltaRadius))
 		{
@@ -957,9 +958,14 @@ void FPhATEdPreviewViewportClient::ModifyPrimitiveSize(int32 BodyIndex, EKCollis
 		}
 
 		float DeltaHeight = DeltaSize.Z;
+		float Radius = FMath::Max(SphylElem.Radius + DeltaRadius, MinPrimSize);
+		float Length = SphylElem.Length + DeltaHeight;
+		
+		Length += SphylElem.Radius - Radius;
+		Length = FMath::Max(0.f, Length);
 
-		AggGeom->SphylElems[PrimIndex].Radius = FMath::Max(AggGeom->SphylElems[PrimIndex].Radius + DeltaRadius, MinPrimSize);
-		AggGeom->SphylElems[PrimIndex].Length = FMath::Max(AggGeom->SphylElems[PrimIndex].Length + 2*DeltaHeight, MinPrimSize);
+		SphylElem.Radius = Radius;
+		SphylElem.Length = Length;
 	}
 	else if (PrimType == KPT_Convex)
 	{
diff --git a/Engine/Source/Runtime/Engine/Private/PhysicsEngine/BodyInstance.cpp b/Engine/Source/Runtime/Engine/Private/PhysicsEngine/BodyInstance.cpp
--- a/Engine/Source/Runtime/Engine/Private/PhysicsEngine/BodyInstance.cpp
+++ b/Engine/Source/Runtime/Engine/Private/PhysicsEngine/BodyInstance.cpp
@@ -282,6 +282,7 @@ TArray<PxShape*> FBodyInstance::GetAllShapes(int32& OutNumSyncShapes) const
 	TArray<PxShape*> PSyncShapes;
 	if( RigidActorSync != NULL )
 	{
+		SCOPED_SCENE_READ_LOCK(RigidActorSync->getScene());
 		PSyncShapes.AddZeroed( RigidActorSync->getNbShapes() );
 		RigidActorSync->getShapes(PSyncShapes.GetData(), PSyncShapes.Num());
 	}
@@ -1131,9 +1132,87 @@ float AdjustForSmallThreshold(float NewVal, float OldVal)
 	return NewVal;
 }
 
-bool FBodyInstance::UpdateBodyScale(const FVector & inScale3D)
+//Non uniform scaling depends on the primitive that has the least non uniform scaling capability. So for example, a capsule's x and y axes scale are locked.
+//So if a capsule exists in this body we must use locked x and y scaling for all shapes.
+namespace EScaleMode
+{
+	enum Type
+	{
+		Free,
+		LockedXY,
+		LockedXYZ
+	};
+}
+
+//computes the relative scaling vectors based on scale mode used
+void ComputeScalingVectors(EScaleMode::Type ScaleMode, const FVector & NewScale3D, const FVector & OldScale3D, FVector & RelativeScale3D, FVector & RelativeScale3DAbs, FVector & OutScale3D)
+{
+	FVector NewScale3DAbs = NewScale3D.GetAbs();
+	FVector OldScale3DAbs = OldScale3D.GetAbs();
+
+	switch (ScaleMode)
+	{
+	case EScaleMode::Free:
+	{
+		OutScale3D = NewScale3D;
+		RelativeScale3DAbs.X = NewScale3DAbs.X / OldScale3DAbs.X;
+		RelativeScale3DAbs.Y = NewScale3DAbs.Y / OldScale3DAbs.Y;
+		RelativeScale3DAbs.Z = NewScale3DAbs.Z / OldScale3DAbs.Z;
+
+		RelativeScale3D.X = NewScale3D.X / OldScale3D.X;
+		RelativeScale3D.Y = NewScale3D.Y / OldScale3D.Y;
+		RelativeScale3D.Z = NewScale3D.Z / OldScale3D.Z;
+		break;
+	}
+	case EScaleMode::LockedXY:
+	{
+		float XYScaleAbs = FMath::Max(NewScale3DAbs.X, NewScale3DAbs.Y);
+		float XYScale = FMath::Max(NewScale3D.X, NewScale3D.Y) < 0.f ? -XYScaleAbs : XYScaleAbs;	//if both xy are negative we should make the xy scale negative
+
+		float OldXYScaleAbs = FMath::Max(OldScale3DAbs.X, OldScale3DAbs.Y);
+		float OldScaleXY = FMath::Max(OldScale3D.X, OldScale3D.Y) < 0.f ? -OldXYScaleAbs : OldXYScaleAbs;
+
+		float RelativeScaleAbs = XYScaleAbs / OldXYScaleAbs;
+		float RelativeScale = XYScale / OldScaleXY;
+
+		RelativeScale3DAbs.X = RelativeScale3DAbs.Y = RelativeScaleAbs;
+		RelativeScale3DAbs.Z = NewScale3DAbs.Z / OldScale3DAbs.Z;
+
+		RelativeScale3D.X = RelativeScale3D.Y = RelativeScale;
+		RelativeScale3D.Z = NewScale3D.Z / OldScale3D.Z;
+		
+		OutScale3D = NewScale3D;
+		OutScale3D.X = OutScale3D.Y = XYScale;
+
+		break;
+	}
+	case EScaleMode::LockedXYZ:
+	{
+		float UniformScaleAbs = NewScale3DAbs.GetMin();	//uniform scale uses the smallest magnitude
+		float UniformScale = FMath::Max3(NewScale3D.X, NewScale3D.Y, NewScale3D.Z) < 0.f ? -UniformScaleAbs : UniformScaleAbs;	//if all three values are negative we should make uniform scale negative
+
+		float OldUniformScaleAbs = OldScale3D.GetAbs().GetMin();
+		float OldUniformScale = FMath::Max3(OldScale3D.X, OldScale3D.Y, OldScale3D.Z) < 0.f ? -OldUniformScaleAbs : OldUniformScaleAbs;
+
+		float RelativeScale = UniformScale / OldUniformScale;
+		float RelativeScaleAbs = UniformScaleAbs / OldUniformScaleAbs;
+
+		RelativeScale3DAbs = FVector(RelativeScaleAbs);
+		RelativeScale3D = FVector(RelativeScale);
+
+		OutScale3D = FVector(UniformScale);
+		break;
+	}
+	default:
+	{
+		check(false);	//invalid scale mode
+	}
+	}
+}
+
+bool FBodyInstance::UpdateBodyScale(const FVector & InScale3D)
 {
-	FVector InScale3DAdjusted = inScale3D;
+	FVector InScale3DAdjusted = InScale3D;
 #if WITH_PHYSX
 	if (BodySetup == NULL || GetPxRigidActor() == NULL)
 	{
@@ -1142,7 +1221,7 @@ bool FBodyInstance::UpdateBodyScale(const FVector & inScale3D)
 	}
 
 	// if same, return
-	if (Scale3D.Equals(inScale3D))
+	if (Scale3D.Equals(InScale3D))
 	{
 		return false;
 	}
@@ -1150,260 +1229,226 @@ bool FBodyInstance::UpdateBodyScale(const FVector & inScale3D)
 	bool bSuccess=false;
 
 #if !(UE_BUILD_SHIPPING || UE_BUILD_TEST)
-	ensure ( !Scale3D.ContainsNaN() && !inScale3D.ContainsNaN() );
+	ensure ( !Scale3D.ContainsNaN() && !InScale3D.ContainsNaN() );
 #endif
 	FVector OldScale3D = Scale3D;
 	
-	//We don't want to go smaller than KINDA_SMALL_NUMBER
-	//But we still want to be able to cross the negative threshold to flip the scale
-	InScale3DAdjusted.X = AdjustForSmallThreshold(inScale3D.X, OldScale3D.X);
-	InScale3DAdjusted.Y = AdjustForSmallThreshold(inScale3D.Y, OldScale3D.Y);
-	InScale3DAdjusted.Z = AdjustForSmallThreshold(inScale3D.Z, OldScale3D.Z);
-	
-	if (FMath::Abs(OldScale3D.X) < 0.1f)
-	{
-		OldScale3D.X = 0.1f;
-	}
-
-	if (FMath::Abs(OldScale3D.Y) < 0.1f)
-	{
-		OldScale3D.Y = 0.1f;
-	}
+	//we never want to hit a scale of 0
+	//But we still want to be able to cross from positive to negative
+	InScale3DAdjusted.X = AdjustForSmallThreshold(InScale3D.X, OldScale3D.X);
+	InScale3DAdjusted.Y = AdjustForSmallThreshold(InScale3D.Y, OldScale3D.Y);
+	InScale3DAdjusted.Z = AdjustForSmallThreshold(InScale3D.Z, OldScale3D.Z);
 	
-	if (FMath::Abs(OldScale3D.Z) < 0.1f)
-	{
-		OldScale3D.Z = 0.1f;
-	}
+	//Make sure OldScale3D is not too small or NaNs can happen
+	OldScale3D.X = OldScale3D.X < 0.1f ? 0.1f : OldScale3D.X;
+	OldScale3D.Y = OldScale3D.Y < 0.1f ? 0.1f : OldScale3D.Y;
+	OldScale3D.Z = OldScale3D.Z < 0.1f ? 0.1f : OldScale3D.Z;
 
-	// Determine if applied scaling is uniform. If it isn't, only convex geometry will be copied over
-	float	NewUniScale = InScale3DAdjusted.GetMin() / OldScale3D.GetMin(); // total scale includes inversing old scale and applying new scale
-	float	NewUniScaleAbs = FMath::Abs(NewUniScale);
-	FVector NewScale3D = InScale3DAdjusted;
-	NewScale3D.X /= OldScale3D.X;
-	NewScale3D.Y /= OldScale3D.Y;
-	NewScale3D.Z /= OldScale3D.Z;
-	FVector NewScale3DAbs = NewScale3D.GetAbs();
+	//Get all shapes
+	int32 TotalShapeCount = 0;
+	TArray<PxShape *> PShapes = GetAllShapes(TotalShapeCount);
 
-	// Is the target a static actor
-	// Not used
-//	bool bDestStatic = (RigidActor->isRigidStatic() != NULL);
-
-	// Add up shapes so we don't have to allocate twice
-	uint32 TotalShapeCount = 0;
-
-	if(RigidActorSync != NULL)
-	{
-		SCOPED_SCENE_READ_LOCK(RigidActorSync->getScene());
-		TotalShapeCount += RigidActorSync->getNbShapes();
-	}
+	EScaleMode::Type ScaleMode = EScaleMode::Free;
 
-	if(RigidActorAsync != NULL)
+	for (int32 ShapeIdx = 0; ShapeIdx < PShapes.Num(); ++ShapeIdx)
 	{
-		SCOPED_SCENE_READ_LOCK(RigidActorAsync->getScene());
-		TotalShapeCount += RigidActorAsync->getNbShapes();
+		PxShape * PShape = PShapes[ShapeIdx];
+		PxGeometryType::Enum GeomType = PShape->getGeometryType();
+		
+		if (GeomType == PxGeometryType::eSPHERE)
+		{
+			ScaleMode = EScaleMode::LockedXYZ;	//sphere is most restrictive so we can stop
+			break;
+		}
+		else if (GeomType == PxGeometryType::eCAPSULE)
+		{
+			ScaleMode = EScaleMode::LockedXY;
+		}
 	}
 
-	TArray<PxShape*> PShapes;
-	PShapes.AddZeroed(TotalShapeCount);
+	FVector RelativeScale3D;
+	FVector RelativeScale3DAbs;
+	FVector NewScale3D;
+	ComputeScalingVectors(ScaleMode, InScale3DAdjusted, OldScale3D, RelativeScale3D, RelativeScale3DAbs, NewScale3D);
 
-	// Record the sync actor shapes
-	int32 NumShapesSync = 0;
-	if( RigidActorSync != NULL )
-	{
-		SCOPED_SCENE_READ_LOCK(RigidActorSync->getScene());
-		NumShapesSync = RigidActorSync->getShapes(PShapes.GetData(), RigidActorSync->getNbShapes());
-	}
 
-	// Record the async actor shapes
-	int32 NumShapesAsync = 0;
-	if( RigidActorAsync != NULL )
-	{
-		SCOPED_SCENE_READ_LOCK(RigidActorAsync->getScene());
-		// Note array start is offset by NumShapesSync
-		NumShapesAsync = RigidActorAsync->getShapes(PShapes.GetData() + NumShapesSync, RigidActorAsync->getNbShapes());
-	}
+	//we need to allocate all of these here because PhysX insists on using the stack. This is wasteful, but reduces a lot of code duplication
+	PxSphereGeometry PSphereGeom;
+	PxBoxGeometry PBoxGeom;
+	PxCapsuleGeometry PCapsuleGeom;
+	PxConvexMeshGeometry PConvexGeom;
+	PxTriangleMeshGeometry PTriMeshGeom;
 
 	for(int32 ShapeIdx=0; ShapeIdx<PShapes.Num(); ShapeIdx++)
 	{
+		PxGeometry * UpdatedGeometry = NULL;
 		PxShape* PShape = PShapes[ShapeIdx];
-		check(PShape);
-
 		PxScene* PScene = PShape->getActor()->getScene();
 		SCENE_LOCK_READ(PScene);
 
 		PxTransform PLocalPose = PShape->getLocalPose();
 		PLocalPose.q.normalize();
 		PxGeometryType::Enum GeomType = PShape->getGeometryType();
-		
-		if (GeomType == PxGeometryType::eSPHERE)
+
+		switch (GeomType)
 		{
-			PxSphereGeometry PSphereGeom;
-			PShape->getSphereGeometry(PSphereGeom);
-			SCENE_UNLOCK_READ(PScene);
+			case PxGeometryType::eSPHERE:
+			{
+				ensure(ScaleMode == EScaleMode::LockedXYZ);
 
-			PSphereGeom.radius	*= NewUniScaleAbs;
-			PLocalPose.p		*= NewUniScale;
+				PShape->getSphereGeometry(PSphereGeom);
+				SCENE_UNLOCK_READ(PScene);
 
-			if (PSphereGeom.isValid())
-			{
-				SCOPED_SCENE_WRITE_LOCK(PScene);
-				PShape->setLocalPose(PLocalPose);
-				PShape->setGeometry(PSphereGeom);
-				bSuccess = true;
+				PSphereGeom.radius *= RelativeScale3DAbs.X;
+				PLocalPose.p *= RelativeScale3D.X;
+				
+				if (PSphereGeom.isValid())
+				{
+					UpdatedGeometry = &PSphereGeom;
+					bSuccess = true;
+				}
+				break;
 			}
-			else
+			case PxGeometryType::eBOX:
 			{
-				FMessageLog("PIE").Warning(FText::Format(LOCTEXT("PhysicsInvalidScale", "Scale ''{0}'' is not valid on object '{1}'."), FText::FromString(InScale3DAdjusted.ToString()), FText::FromString(GetBodyDebugName())));
-			}
-		}
-		else if (GeomType == PxGeometryType::eBOX)
-		{
-			PxBoxGeometry PBoxGeom;
-			PShape->getBoxGeometry(PBoxGeom);
-			SCENE_UNLOCK_READ(PScene);
+				PShape->getBoxGeometry(PBoxGeom);
+				SCENE_UNLOCK_READ(PScene);
 
-			PBoxGeom.halfExtents.x *= NewScale3DAbs.X;
-			PBoxGeom.halfExtents.y *= NewScale3DAbs.Y;
-			PBoxGeom.halfExtents.z *= NewScale3DAbs.Z;
-			PLocalPose.p.x *= NewScale3D.X;
-			PLocalPose.p.y *= NewScale3D.Y;
-			PLocalPose.p.z *= NewScale3D.Z;
+				PBoxGeom.halfExtents.x *= RelativeScale3DAbs.X;
+				PBoxGeom.halfExtents.y *= RelativeScale3DAbs.Y;
+				PBoxGeom.halfExtents.z *= RelativeScale3DAbs.Z;
+				PLocalPose.p.x *= RelativeScale3D.X;
+				PLocalPose.p.y *= RelativeScale3D.Y;
+				PLocalPose.p.z *= RelativeScale3D.Z;
 
-			if (PBoxGeom.isValid())
-			{
-				SCOPED_SCENE_WRITE_LOCK(PScene);
-				PShape->setGeometry(PBoxGeom);
-				PShape->setLocalPose(PLocalPose);
-				bSuccess = true;
+				if (PBoxGeom.isValid())
+				{
+					UpdatedGeometry = &PBoxGeom;
+					bSuccess = true;
+				}
+				break;
 			}
-			else
+			case PxGeometryType::eCAPSULE:
 			{
-				FMessageLog("PIE").Warning(FText::Format(LOCTEXT("PhysicsInvalidScale", "Scale ''{0}'' is not valid on object '{1}'."), FText::FromString(InScale3DAdjusted.ToString()), FText::FromString(GetBodyDebugName())));
-			}
-		}
-		else if (GeomType == PxGeometryType::eCAPSULE)
-		{
-			PxCapsuleGeometry PCapsuleGeom;
-			PShape->getCapsuleGeometry(PCapsuleGeom);
-			SCENE_UNLOCK_READ(PScene);
+				ensure(ScaleMode == EScaleMode::LockedXY || ScaleMode == EScaleMode::LockedXYZ);
 
-			PCapsuleGeom.halfHeight *= NewScale3DAbs.Z;
-			PCapsuleGeom.radius *= FMath::Max(NewScale3DAbs.X, NewScale3DAbs.Y);
+				PShape->getCapsuleGeometry(PCapsuleGeom);
+				SCENE_UNLOCK_READ(PScene);
 
-			PLocalPose.p.x *= NewScale3D.X;
-			PLocalPose.p.y *= NewScale3D.Y;
-			PLocalPose.p.z *= NewScale3D.Z;
+				PCapsuleGeom.halfHeight *= RelativeScale3DAbs.Z;
+				PCapsuleGeom.radius *= RelativeScale3DAbs.X;
 
-			if(PCapsuleGeom.isValid())
-			{
-				SCOPED_SCENE_WRITE_LOCK(PScene);
-				PShape->setGeometry(PCapsuleGeom);
-				PShape->setLocalPose(PLocalPose);
-				bSuccess = true;
-			}
-			else
-			{
-				FMessageLog("PIE").Warning(FText::Format(LOCTEXT("PhysicsInvalidScale", "Scale ''{0}'' is not valid on object '{1}'."), FText::FromString(InScale3DAdjusted.ToString()), FText::FromString(GetBodyDebugName())));
-			}
-		}
-		else if (GeomType == PxGeometryType::eCONVEXMESH)
-		{
-			PxConvexMeshGeometry PConvexGeom;
-			PShape->getConvexMeshGeometry(PConvexGeom);
-			SCENE_UNLOCK_READ(PScene);
+				PLocalPose.p.x *= RelativeScale3D.X;
+				PLocalPose.p.y *= RelativeScale3D.Y;
+				PLocalPose.p.z *= RelativeScale3D.Z;
 
-			// find which convex elems it is
-			// it would be nice to know if the order of PShapes array index is in the order of createShape
-			// Create convex shapes
-			for(int32 i=0; i<BodySetup->AggGeom.ConvexElems.Num(); i++)
+				if (PCapsuleGeom.isValid())
+				{
+					UpdatedGeometry = &PCapsuleGeom;
+					bSuccess = true;
+				}
+
+				break;
+			}
+			case PxGeometryType::eCONVEXMESH:
 			{
-				FKConvexElem* ConvexElem = &(BodySetup->AggGeom.ConvexElems[i]);
+				PShape->getConvexMeshGeometry(PConvexGeom);
+				SCENE_UNLOCK_READ(PScene);
 
-				// found it
-				if(ConvexElem->ConvexMesh == PConvexGeom.convexMesh)
+				// find which convex elems it is
+				// it would be nice to know if the order of PShapes array index is in the order of createShape
+				// Create convex shapes
+				for (int32 i = 0; i < BodySetup->AggGeom.ConvexElems.Num(); i++)
 				{
-					// Please note that this one we don't inverse old scale, but just set new one
-					FVector NewScale3D = InScale3DAdjusted;
-					FVector Scale3DAbs(FMath::Abs(NewScale3D.X), FMath::Abs(NewScale3D.Y), FMath::Abs(NewScale3D.Z)); // magnitude of scale (sign removed)
+					FKConvexElem* ConvexElem = &(BodySetup->AggGeom.ConvexElems[i]);
 
-					PxTransform PNewLocalPose;
-					bool bUseNegX = CalcMeshNegScaleCompensation(NewScale3D, PNewLocalPose);
+					// found it
+					if (ConvexElem->ConvexMesh == PConvexGeom.convexMesh)
+					{
+						// Please note that this one we don't inverse old scale, but just set new one (but we still follow scale mode restriction)
+						FVector NewScale3D = RelativeScale3D * OldScale3D;
+						FVector Scale3DAbs(FMath::Abs(NewScale3D.X), FMath::Abs(NewScale3D.Y), FMath::Abs(NewScale3D.Z)); // magnitude of scale (sign removed)
 
-					PxTransform PElementTransform = U2PTransform(ConvexElem->GetTransform());
-					PNewLocalPose.q *= PElementTransform.q;
-					PNewLocalPose.p += PElementTransform.p;
+						PxTransform PNewLocalPose;
+						bool bUseNegX = CalcMeshNegScaleCompensation(NewScale3D, PNewLocalPose);
 
-					PConvexGeom.convexMesh = bUseNegX ? ConvexElem->ConvexMeshNegX : ConvexElem->ConvexMesh;
-					PConvexGeom.scale.scale = U2PVector(Scale3DAbs);
+						PxTransform PElementTransform = U2PTransform(ConvexElem->GetTransform());
+						PNewLocalPose.q *= PElementTransform.q;
+						PNewLocalPose.p += PElementTransform.p;
 
-					if(PConvexGeom.isValid())
-					{
-						SCOPED_SCENE_WRITE_LOCK(PScene);
-						PShape->setGeometry(PConvexGeom);
-						PShape->setLocalPose(PNewLocalPose);
-						bSuccess = true;
-					}
-					else
-					{
-						FMessageLog("PIE").Warning(FText::Format(LOCTEXT("PhysicsInvalidScale", "Scale ''{0}'' is not valid on object '{1}'."), FText::FromString(InScale3DAdjusted.ToString()), FText::FromString(GetBodyDebugName())));
+						PConvexGeom.convexMesh = bUseNegX ? ConvexElem->ConvexMeshNegX : ConvexElem->ConvexMesh;
+						PConvexGeom.scale.scale = U2PVector(Scale3DAbs);
+
+						if (PConvexGeom.isValid())
+						{
+							UpdatedGeometry = &PConvexGeom;
+							bSuccess = true;
+						}
+						break;
 					}
-					break;
 				}
+				break;
 			}
-		}
-		else if (GeomType == PxGeometryType::eTRIANGLEMESH)
-		{
-			PxTriangleMeshGeometry PTriMeshGeom;
-			PShape->getTriangleMeshGeometry(PTriMeshGeom);
-			SCENE_UNLOCK_READ(PScene);
-
-			// Create tri-mesh shape
-			if(BodySetup->TriMesh != NULL || BodySetup->TriMeshNegX != NULL)
+			case PxGeometryType::eTRIANGLEMESH:
 			{
-				// Please note that this one we don't inverse old scale, but just set new one
-				FVector NewScale3D = InScale3DAdjusted;
-				FVector Scale3DAbs(FMath::Abs(NewScale3D.X), FMath::Abs(NewScale3D.Y), FMath::Abs(NewScale3D.Z)); // magnitude of scale (sign removed)
-
-				PxTransform PNewLocalPose;
-				bool bUseNegX = CalcMeshNegScaleCompensation(NewScale3D, PNewLocalPose);
+				PShape->getTriangleMeshGeometry(PTriMeshGeom);
+				SCENE_UNLOCK_READ(PScene);
 
-				// Only case where TriMeshNegX should be null is BSP, which should not require negX version
-				if(bUseNegX && BodySetup->TriMeshNegX == NULL)
+				// Create tri-mesh shape
+				if (BodySetup->TriMesh != NULL || BodySetup->TriMeshNegX != NULL)
 				{
-					UE_LOG(LogPhysics, Warning,  TEXT("FBodyInstance::UpdateBodyScale: Want to use NegX but it doesn't exist! %s"), *BodySetup->GetPathName() );
-				}
+					// Please note that this one we don't inverse old scale, but just set new one (but still adjust for scale mode)
+					FVector NewScale3D = RelativeScale3D * OldScale3D;
+					FVector Scale3DAbs(FMath::Abs(NewScale3D.X), FMath::Abs(NewScale3D.Y), FMath::Abs(NewScale3D.Z)); // magnitude of scale (sign removed)
 
-				PxTriangleMesh* UseTriMesh = bUseNegX ? BodySetup->TriMeshNegX : BodySetup->TriMesh;
-				if(UseTriMesh != NULL)
-				{
-					PTriMeshGeom.triangleMesh = bUseNegX ? BodySetup->TriMeshNegX : BodySetup->TriMesh;
-					PTriMeshGeom.scale.scale = U2PVector(Scale3DAbs);
+					PxTransform PNewLocalPose;
+					bool bUseNegX = CalcMeshNegScaleCompensation(NewScale3D, PNewLocalPose);
 
-					if(PTriMeshGeom.isValid())
+					// Only case where TriMeshNegX should be null is BSP, which should not require negX version
+					if (bUseNegX && BodySetup->TriMeshNegX == NULL)
 					{
-						SCOPED_SCENE_WRITE_LOCK(PScene);
-						PShape->setGeometry(PTriMeshGeom);
-						PShape->setLocalPose(PNewLocalPose);
-						bSuccess = true;
+						UE_LOG(LogPhysics, Warning, TEXT("FBodyInstance::UpdateBodyScale: Want to use NegX but it doesn't exist! %s"), *BodySetup->GetPathName());
 					}
-					else
+
+					PxTriangleMesh* UseTriMesh = bUseNegX ? BodySetup->TriMeshNegX : BodySetup->TriMesh;
+					if (UseTriMesh != NULL)
 					{
-						FMessageLog("PIE").Warning(FText::Format(LOCTEXT("PhysicsInvalidScale", "Scale ''{0}'' is not valid on object '{1}'."), FText::FromString(InScale3DAdjusted.ToString()), FText::FromString(GetBodyDebugName())));
+						PTriMeshGeom.triangleMesh = bUseNegX ? BodySetup->TriMeshNegX : BodySetup->TriMesh;
+						PTriMeshGeom.scale.scale = U2PVector(Scale3DAbs);
+
+						if (PTriMeshGeom.isValid())
+						{
+							UpdatedGeometry = &PTriMeshGeom;
+							bSuccess = true;
+
+						}
 					}
 				}
+				break;
+			}
+			default:
+			{
+					   SCENE_UNLOCK_READ(PScene);
+					   UE_LOG(LogPhysics, Error, TEXT("Unknown geom type."));
+			}
+
+			if (UpdatedGeometry)
+			{
+				SCOPED_SCENE_WRITE_LOCK(PScene);
+				PShape->setLocalPose(PLocalPose);
+				PShape->setGeometry(*UpdatedGeometry);
+			}
+			else
+			{
+				FMessageLog("PIE").Warning(FText::Format(LOCTEXT("PhysicsInvalidScale", "Scale ''{0}'' is not valid on object '{1}'."), FText::FromString(InScale3DAdjusted.ToString()), FText::FromString(GetBodyDebugName())));
 			}
-		}
-		else
-		{
-			UE_LOG(LogPhysics, Error, TEXT("Unknown grom type."));
-			SCENE_UNLOCK_READ(PScene);
 		}
 	}
 
 	// if success, overwrite old Scale3D, otherwise, just don't do it. It will have invalid scale next time
 	if ( bSuccess )
 	{
-		Scale3D = InScale3DAdjusted;
+		Scale3D = NewScale3D;
 
 		// update mass if required
 		if (bUpdateMassWhenScaleChanges)
diff --git a/Engine/Source/Runtime/Engine/Private/PhysicsEngine/BodySetup.cpp b/Engine/Source/Runtime/Engine/Private/PhysicsEngine/BodySetup.cpp
--- a/Engine/Source/Runtime/Engine/Private/PhysicsEngine/BodySetup.cpp
+++ b/Engine/Source/Runtime/Engine/Private/PhysicsEngine/BodySetup.cpp
@@ -241,17 +241,18 @@ void SetupNonUniformHelper(FVector & Scale3D, float & MinScale, float & MinScale
 		// set min scale
 		Scale3D = FVector(0.1f);
 	}
-	// Determine if applied scaling is uniform. If it isn't, only convex geometry will be copied over
-	MinScale = Scale3D.GetMin(); // Uniform scaling factor
-	MinScaleAbs = FMath::Abs(MinScale);
+
+	Scale3DAbs = Scale3D.GetAbs();
+	MinScaleAbs = Scale3DAbs.GetMin();
+
+	MinScale = FMath::Max3(Scale3D.X, Scale3D.Y, Scale3D.Z) < 0.f ? -MinScaleAbs : MinScaleAbs;	//if all three values are negative make minScale negative
+	
 	if (FMath::IsNearlyZero(MinScale))
 	{
 		// only one of them can be 0, we make sure they have mini set up correctly
 		MinScale = 0.1f;
 		MinScaleAbs = 0.1f;
 	}
-
-	Scale3DAbs = FVector(FMath::Abs(Scale3D.X), FMath::Abs(Scale3D.Y), FMath::Abs(Scale3D.Z)); // magnitude of scale (sign removed)
 }
 
 #if WITH_BODY_WELDING
@@ -373,7 +374,11 @@ void UBodySetup::AddShapesToRigidActor(PxRigidActor* PDestActor, FVector& Scale3
 				// this is a bit confusing since radius and height is scaled
 				// first apply the scale first 
 				float Radius = FMath::Max(SphylElem->Radius * ScaleRadius, 0.1f);
-				float HalfHeight = SphylElem->Length * ScaleLength * 0.5f;
+				float Length = SphylElem->Length + SphylElem->Radius * 2.f;
+				float HalfLength = Length * ScaleLength * 0.5f;
+				Radius = FMath::Clamp(Radius, 0.1f, HalfLength);	//radius is capped by half length
+				float HalfHeight = HalfLength - Radius;
+				HalfHeight = FMath::Max(0.1f, HalfHeight);
 
 				PxCapsuleGeometry PCapsuleGeom;
 				PCapsuleGeom.halfHeight = HalfHeight;
diff --git a/Engine/Source/Runtime/Engine/Private/PrimitiveDrawingUtils.cpp b/Engine/Source/Runtime/Engine/Private/PrimitiveDrawingUtils.cpp
--- a/Engine/Source/Runtime/Engine/Private/PrimitiveDrawingUtils.cpp
+++ b/Engine/Source/Runtime/Engine/Private/PrimitiveDrawingUtils.cpp
@@ -782,8 +782,9 @@ void DrawWireCapsule(FPrimitiveDrawInterface* PDI,const FVector& Base,const FVec
 	const float XScale = X.Size();
 	const float YScale = Y.Size();
 	const float ZScale = Z.Size();
-	float CapsuleRadius = Radius * FMath::Max( FMath::Max(XScale,YScale), ZScale );
+	float CapsuleRadius = Radius * FMath::Max(XScale,YScale);
 	HalfHeight *= ZScale;
+	CapsuleRadius = FMath::Clamp(CapsuleRadius, 0.f, HalfHeight);	//cap radius based on total height
 	HalfHeight -= CapsuleRadius;
 	HalfHeight = FMath::Max( 0.0f, HalfHeight );
 
