Prevent vehicle from getting negative mass.
Fix default inertia tensor to be identity.
TTP#335996

diff --git a/Engine/Source/Runtime/Engine/Classes/Vehicles/WheeledVehicleMovementComponent.h b/Engine/Source/Runtime/Engine/Classes/Vehicles/WheeledVehicleMovementComponent.h
--- a/Engine/Source/Runtime/Engine/Classes/Vehicles/WheeledVehicleMovementComponent.h
+++ b/Engine/Source/Runtime/Engine/Classes/Vehicles/WheeledVehicleMovementComponent.h
@@ -171,7 +171,7 @@ class ENGINE_API UWheeledVehicleMovementComponent : public UPawnMovementComponen
 
 	/** Mass to set the vehicle chassis to. It's much easier to tweak vehicle settings when
 	 * the mass doesn't change due to tweaks with the physics asset. [kg] */
-	UPROPERTY(EditAnywhere, Category=VehicleSetup)
+	UPROPERTY(EditAnywhere, Category = VehicleSetup, meta = (ClampMin = "0.01", UIMin = "0.01"))
 	float Mass;
 
 	/** DragCoefficient of the vehicle chassis. */
@@ -179,11 +179,11 @@ class ENGINE_API UWheeledVehicleMovementComponent : public UPawnMovementComponen
 	float DragCoefficient;
 
 	/** Chassis width used for drag force computation (cm)*/
-	UPROPERTY(EditAnywhere, Category = VehicleSetup)
+	UPROPERTY(EditAnywhere, Category = VehicleSetup, meta = (ClampMin = "0.01", UIMin = "0.01"))
 	float ChassisWidth;
 
 	/** Chassis height used for drag force computation (cm)*/
-	UPROPERTY(EditAnywhere, Category = VehicleSetup)
+	UPROPERTY(EditAnywhere, Category = VehicleSetup, meta = (ClampMin = "0.01", UIMin = "0.01"))
 	float ChassisHeight;
 
 	// Drag area in cm^2
diff --git a/Engine/Source/Runtime/Engine/Private/Vehicles/WheeledVehicleMovementComponent.cpp b/Engine/Source/Runtime/Engine/Private/Vehicles/WheeledVehicleMovementComponent.cpp
--- a/Engine/Source/Runtime/Engine/Private/Vehicles/WheeledVehicleMovementComponent.cpp
+++ b/Engine/Source/Runtime/Engine/Private/Vehicles/WheeledVehicleMovementComponent.cpp
@@ -84,7 +84,7 @@ UWheeledVehicleMovementComponent::UWheeledVehicleMovementComponent(const class F
 	DragCoefficient = 0.3f;
 	ChassisWidth = 180.f;
 	ChassisHeight = 140.f;
-	InertiaTensorScale = FVector( 1.0f, 0.80f, 1.0f );
+	InertiaTensorScale = FVector( 1.0f, 1.0f, 1.0f );
 	AngErrorAccumulator = 0.0f;
 	MinNormalizedTireLoad = 0.0f;
 	MaxNormalizedTireLoad = 10.0f;
