Fix wheel mass and other wheel values being negative.
TTP#335967

diff --git a/Engine/Source/Runtime/Engine/Classes/Vehicles/VehicleWheel.h b/Engine/Source/Runtime/Engine/Classes/Vehicles/VehicleWheel.h
--- a/Engine/Source/Runtime/Engine/Classes/Vehicles/VehicleWheel.h
+++ b/Engine/Source/Runtime/Engine/Classes/Vehicles/VehicleWheel.h
@@ -38,23 +38,23 @@ class ENGINE_API UVehicleWheel : public UObject
 	FVector											Offset;
 
 	/** Radius of the wheel */
-	UPROPERTY(EditAnywhere, Category=Wheel)
+	UPROPERTY(EditAnywhere, Category = Wheel, meta = (ClampMin = "0.01", UIMin = "0.01"))
 	float											ShapeRadius;
 
 	/** Width of the wheel */
-	UPROPERTY(EditAnywhere, Category=Wheel)
+	UPROPERTY(EditAnywhere, Category = Wheel, meta = (ClampMin = "0.01", UIMin = "0.01"))
 	float											ShapeWidth;
 
 	/** Mass of this wheel */
-	UPROPERTY(EditAnywhere, Category=Wheel)
+	UPROPERTY(EditAnywhere, Category = Wheel, meta = (ClampMin = "0.01", UIMin = "0.01"))
 	float											Mass;
 
 	/** Damping rate for this wheel (Kgm^2/s) */
-	UPROPERTY(EditAnywhere, Category=Wheel)
+	UPROPERTY(EditAnywhere, Category = Wheel, meta = (ClampMin = "0.01", UIMin = "0.01"))
 	float											DampingRate;
 
 	// steer angle in degrees for this wheel
-	UPROPERTY(EditAnywhere, Category = WheelsSetup)
+	UPROPERTY(EditAnywhere, Category = WheelsSetup, meta = (ClampMin = "0", UIMin = "0"))
 	float SteerAngle;
 
 	/** Whether handbrake should affect this wheel */
