Change UTextRenderComponent to use 16-bit index buffers to fix Tappy Chicken scoreboard text on Mali.

diff --git a/Engine/Source/Runtime/Engine/Private/TextRenderComponent.cpp b/Engine/Source/Runtime/Engine/Private/TextRenderComponent.cpp
--- a/Engine/Source/Runtime/Engine/Private/TextRenderComponent.cpp
+++ b/Engine/Source/Runtime/Engine/Private/TextRenderComponent.cpp
@@ -117,15 +117,15 @@ public:
 class FTextRenderIndexBuffer : public FIndexBuffer 
 {
 public:
-	TArray<int32> Indices;
+	TArray<uint16> Indices;
 
 	void InitRHI()
 	{
-		IndexBufferRHI = RHICreateIndexBuffer(sizeof(int32),Indices.Num() * sizeof(int32),NULL,BUF_Static);
+		IndexBufferRHI = RHICreateIndexBuffer(sizeof(uint16), Indices.Num() * sizeof(uint16), NULL, BUF_Static);
 
 		// Copy the index data into the index buffer.
-		void* Buffer = RHILockIndexBuffer(IndexBufferRHI,0,Indices.Num() * sizeof(int32),RLM_WriteOnly);
-		FMemory::Memcpy(Buffer,Indices.GetTypedData(),Indices.Num() * sizeof(int32));
+		void* Buffer = RHILockIndexBuffer(IndexBufferRHI, 0, Indices.Num() * sizeof(uint16), RLM_WriteOnly);
+		FMemory::Memcpy(Buffer, Indices.GetTypedData(), Indices.Num() * sizeof(uint16));
 		RHIUnlockIndexBuffer(IndexBufferRHI);
 	}
 };
@@ -365,7 +365,7 @@ private:
 	// End FPrimitiveSceneProxy interface
 
 	void ReleaseRenderThreadResources();
-	bool BuildStringMesh( TArray<FDynamicMeshVertex>& OutVertices, TArray<int32>& OutIndices );
+	bool BuildStringMesh( TArray<FDynamicMeshVertex>& OutVertices, TArray<uint16>& OutIndices );
 
 private:
 	FMaterialRelevance MaterialRelevance;
@@ -552,7 +552,7 @@ uint32 FTextRenderSceneProxy::GetAllocatedSize() const
 /**
 * For the given text, constructs a mesh to be used by the vertex factory for rendering.
 */
-bool  FTextRenderSceneProxy::BuildStringMesh( TArray<FDynamicMeshVertex>& OutVertices, TArray<int32>& OutIndices )
+bool  FTextRenderSceneProxy::BuildStringMesh( TArray<FDynamicMeshVertex>& OutVertices, TArray<uint16>& OutIndices )
 {
 	if(!Font || Text.IsEmpty())
 	{
@@ -636,6 +636,11 @@ bool  FTextRenderSceneProxy::BuildStringMesh( TArray<FDynamicMeshVertex>& OutVer
 				int32 V01 = OutVertices.Add( FDynamicMeshVertex( V2, TangentX, TangentZ, FVector2D(U,	V + SizeV), ActualColor));
 				int32 V11 = OutVertices.Add( FDynamicMeshVertex( V3, TangentX, TangentZ, FVector2D(U + SizeU,	V + SizeV), ActualColor));
 
+				check(V00 < 65536);
+				check(V10 < 65536);
+				check(V01 < 65536);
+				check(V11 < 65536);
+
 				OutIndices.Add(V00);
 				OutIndices.Add(V11);
 				OutIndices.Add(V10);
