Added a kismet node for Calibrating motion. Implemented for IOS devices

diff --git a/Engine/Source/Runtime/Core/Public/IOS/IOSInputInterface.h b/Engine/Source/Runtime/Core/Public/IOS/IOSInputInterface.h
--- a/Engine/Source/Runtime/Core/Public/IOS/IOSInputInterface.h
+++ b/Engine/Source/Runtime/Core/Public/IOS/IOSInputInterface.h
@@ -22,7 +22,7 @@ struct TouchInput
 /**
  * Interface class for IOS input devices
  */
-class FIOSInputInterface : public FExec
+class FIOSInputInterface : private FSelfRegisteringExec
 {
 public:
 
diff --git a/Engine/Source/Runtime/Engine/Classes/Kismet/KismetInputLibrary.h b/Engine/Source/Runtime/Engine/Classes/Kismet/KismetInputLibrary.h
new file mode 100644
--- /dev/null
+++ b/Engine/Source/Runtime/Engine/Classes/Kismet/KismetInputLibrary.h
@@ -0,0 +1,15 @@
+// Copyright 1998-2014 Epic Games, Inc. All Rights Reserved.
+
+#pragma once
+
+#include "KismetInputLibrary.generated.h"
+
+UCLASS(MinimalAPI)
+class UKismetInputLibrary : public UBlueprintFunctionLibrary
+{
+	GENERATED_UCLASS_BODY()
+
+	/** Calibrate the tilt for the input device */
+	UFUNCTION(BlueprintCallable, Category="Input")
+	static void CalibrateTilt();
+};
diff --git a/Engine/Source/Runtime/Engine/Private/KismetInputLibrary.cpp b/Engine/Source/Runtime/Engine/Private/KismetInputLibrary.cpp
new file mode 100644
--- /dev/null
+++ b/Engine/Source/Runtime/Engine/Private/KismetInputLibrary.cpp
@@ -0,0 +1,25 @@
+// Copyright 1998-2014 Epic Games, Inc. All Rights Reserved.
+
+#include "EnginePrivate.h"
+#include "MessageLog.h"
+#include "UObjectToken.h"
+
+
+//////////////////////////////////////////////////////////////////////////
+// UKismetInputLibrary
+
+#define LOCTEXT_NAMESPACE "KismetInputLibrary"
+
+
+UKismetInputLibrary::UKismetInputLibrary(const class FPostConstructInitializeProperties& PCIP)
+	: Super(PCIP)
+{
+}
+
+
+void UKismetInputLibrary::CalibrateTilt()
+{
+	GEngine->Exec(NULL, TEXT("CALIBRATEMOTION"));
+}
+
+#undef LOCTEXT_NAMESPACE
\ No newline at end of file
diff --git a/Engine/Source/Runtime/Engine/Public/Engine.h b/Engine/Source/Runtime/Engine/Public/Engine.h
--- a/Engine/Source/Runtime/Engine/Public/Engine.h
+++ b/Engine/Source/Runtime/Engine/Public/Engine.h
@@ -907,6 +907,7 @@ struct FURL;
 #include "Kismet/HeadMountedDisplayFunctionLibrary.h"
 #include "Kismet/KismetAIHelperLibrary.h"
 #include "Kismet/KismetArrayLibrary.h"
+#include "Kismet/KismetInputLibrary.h"
 #include "Kismet/KismetMaterialLibrary.h"
 #include "Kismet/KismetMathLibrary.h"
 #include "Kismet/KismetNodeHelperLibrary.h"
