Move EnginePrivate.h to Private folder, stop including it outside Engine module
Add Runtime/Engine/Private as and Engine private include path
Rename Engine/Private/Materials/MaterialInstance.h to MaterialInstanceSupport.h because it was causing ambiguities
Move cpp files from Classes and Public into Private

diff --git a/Engine/Source/Editor/DeviceProfileServices/Private/DeviceProfileServicesPCH.h b/Engine/Source/Editor/DeviceProfileServices/Private/DeviceProfileServicesPCH.h
--- a/Engine/Source/Editor/DeviceProfileServices/Private/DeviceProfileServicesPCH.h
+++ b/Engine/Source/Editor/DeviceProfileServices/Private/DeviceProfileServicesPCH.h
@@ -12,7 +12,7 @@
 /* Dependencies
  *****************************************************************************/
 
-#include "EnginePrivate.h"
+#include "Engine.h"
 #include "ModuleManager.h"
 #include "TargetDeviceServices.h"
 
diff --git a/Engine/Source/Programs/UnrealBuildTool/System/XcodeProject.cs b/Engine/Source/Programs/UnrealBuildTool/System/XcodeProject.cs
--- a/Engine/Source/Programs/UnrealBuildTool/System/XcodeProject.cs
+++ b/Engine/Source/Programs/UnrealBuildTool/System/XcodeProject.cs
@@ -262,7 +262,7 @@ namespace UnrealBuildTool
 			{
 				return null;
 			}
-			else if (InitFilePath.EndsWith("Classes/Engine/Engine.h") || InitFilePath.EndsWith("Classes/VectorField/VectorField.h") || InitFilePath.EndsWith("Classes/Materials/MaterialInstance.h")
+			else if (InitFilePath.EndsWith("Classes/Engine/Engine.h") || InitFilePath.EndsWith("Classes/VectorField/VectorField.h") 
 				|| InitFilePath.EndsWith("Classes/Engine/Canvas.h") || InitFilePath.EndsWith("Classes/Intrinsic/Model.h") || InitFilePath.EndsWith("Classes/Engine/Texture.h")
 				|| InitFilePath.EndsWith("Windows/OnlineFriendsFacebook.h") || InitFilePath.EndsWith("Windows/OnlineIdentityFacebook.h") || InitFilePath.EndsWith("HTML5/Simulator/SocketSubsystem.h")
 				|| InitFilePath.EndsWith("HTML5/Device/SocketSubsystem.h") || InitFilePath.EndsWith("ContentBrowser/Private/HistoryManager.h") || InitFilePath.EndsWith("Private/Menus/TranslationEditorMenu.h"))
diff --git a/Engine/Source/Runtime/Engine/Classes/Animation/PreviewAssetAttachComponent.cpp b/Engine/Source/Runtime/Engine/Classes/Animation/PreviewAssetAttachComponent.cpp
deleted file mode 100644
diff --git a/Engine/Source/Runtime/Engine/Engine.Build.cs b/Engine/Source/Runtime/Engine/Engine.Build.cs
--- a/Engine/Source/Runtime/Engine/Engine.Build.cs
+++ b/Engine/Source/Runtime/Engine/Engine.Build.cs
@@ -17,6 +17,7 @@ public class Engine : ModuleRules
 				"Runtime/Online/OnlineSubsystem/Public",
 				"Runtime/Online/OnlineSubsystemUtils/Public",
                 "Developer/SynthBenchmark/Public",
+                "Runtime/Engine/Private",
 			}
 		);
 
diff --git a/Engine/Source/Runtime/Engine/Private/Animation/PreviewAssetAttachComponent.cpp b/Engine/Source/Runtime/Engine/Private/Animation/PreviewAssetAttachComponent.cpp
new file mode 100644
--- /dev/null
+++ b/Engine/Source/Runtime/Engine/Private/Animation/PreviewAssetAttachComponent.cpp
@@ -0,0 +1,94 @@
+// Copyright 1998-2014 Epic Games, Inc. All Rights Reserved.
+
+#include "EnginePrivate.h"
+#include "Animation/PreviewAssetAttachComponent.h"
+
+void FPreviewAssetAttachContainer::AddAttachedObject( UObject* AttachObject, FName AttachPointName )
+{
+	FPreviewAttachedObjectPair Pair;
+	Pair.AttachedTo = AttachPointName;
+	Pair.SetAttachedObject(AttachObject);
+	AttachedObjects.Add( Pair );
+}
+
+void FPreviewAssetAttachContainer::RemoveAttachedObject( UObject* ObjectToRemove, FName AttachName )
+{
+	for( int i = AttachedObjects.Num() - 1; i >= 0; --i )
+	{
+		FPreviewAttachedObjectPair& Pair = AttachedObjects[i];
+
+		if( Pair.GetAttachedObject() == ObjectToRemove && Pair.AttachedTo == AttachName )
+		{
+			AttachedObjects.RemoveAtSwap( i, 1, false );
+			break;
+		}
+	}
+}
+
+UObject* FPreviewAssetAttachContainer::GetAttachedObjectByAttachName( FName AttachName ) const
+{
+	for( int i = 0; i < AttachedObjects.Num(); ++i )
+	{
+		if( AttachedObjects[i].AttachedTo == AttachName )
+		{
+			return AttachedObjects[i].GetAttachedObject();
+		}
+	}
+	return NULL;
+}
+
+void FPreviewAssetAttachContainer::ClearAllAttachedObjects()
+{
+	AttachedObjects.Empty();
+}
+
+int32 FPreviewAssetAttachContainer::Num() const
+{
+	return AttachedObjects.Num();
+}
+
+FPreviewAttachedObjectPair& FPreviewAssetAttachContainer::operator []( int32 i )
+{
+	return AttachedObjects[i];
+}
+
+TConstIterator FPreviewAssetAttachContainer::CreateConstIterator() const
+{
+	return AttachedObjects.CreateConstIterator();
+}
+
+TIterator FPreviewAssetAttachContainer::CreateIterator()
+{
+	return AttachedObjects.CreateIterator();
+}
+
+void FPreviewAssetAttachContainer::RemoveAtSwap( int32 Index, int32 Count /* = 1 */, bool bAllowShrinking /*=true */ )
+{
+	AttachedObjects.RemoveAtSwap( Index, Count, bAllowShrinking );
+}
+
+void FPreviewAssetAttachContainer::SaveAttachedObjectsFromDeprecatedProperties()
+{
+	for (FPreviewAttachedObjectPair& Pair : AttachedObjects)
+	{
+		Pair.SaveAttachedObjectFromDeprecatedProperty();
+	}
+}
+
+int32 FPreviewAssetAttachContainer::ValidatePreviewAttachedObjects()
+{
+	int32 NumBrokenAssets = 0;
+
+	for (int32 i = AttachedObjects.Num() - 1; i >= 0; --i)
+	{
+		FPreviewAttachedObjectPair& PreviewAttachedObject = AttachedObjects[i];
+
+		if (!PreviewAttachedObject.GetAttachedObject())
+		{
+			AttachedObjects.RemoveAtSwap(i);
+			++NumBrokenAssets;
+		}
+	}
+
+	return NumBrokenAssets;
+}
\ No newline at end of file
diff --git a/Engine/Source/Runtime/Engine/Private/EnginePrivate.h b/Engine/Source/Runtime/Engine/Private/EnginePrivate.h
new file mode 100644
--- /dev/null
+++ b/Engine/Source/Runtime/Engine/Private/EnginePrivate.h
@@ -0,0 +1,16 @@
+// Copyright 1998-2014 Epic Games, Inc. All Rights Reserved.
+
+/*=============================================================================
+	EnginePrivate.h: Unreal engine private header file.
+=============================================================================*/
+
+/*----------------------------------------------------------------------------
+	Engine public includes.
+----------------------------------------------------------------------------*/
+
+#include "Engine.h"
+
+/*-----------------------------------------------------------------------------
+	Engine private includes.
+-----------------------------------------------------------------------------*/
+
diff --git a/Engine/Source/Runtime/Engine/Private/Materials/Material.cpp b/Engine/Source/Runtime/Engine/Private/Materials/Material.cpp
--- a/Engine/Source/Runtime/Engine/Private/Materials/Material.cpp
+++ b/Engine/Source/Runtime/Engine/Private/Materials/Material.cpp
@@ -6,7 +6,7 @@
 
 #include "EnginePrivate.h"
 #include "MaterialShader.h"
-#include "MaterialInstance.h"
+#include "MaterialInstanceSupport.h"
 #include "UObjectAnnotation.h"
 #include "MaterialCompiler.h"
 #include "TargetPlatform.h"
diff --git a/Engine/Source/Runtime/Engine/Private/Materials/MaterialInstance.cpp b/Engine/Source/Runtime/Engine/Private/Materials/MaterialInstance.cpp
--- a/Engine/Source/Runtime/Engine/Private/Materials/MaterialInstance.cpp
+++ b/Engine/Source/Runtime/Engine/Private/Materials/MaterialInstance.cpp
@@ -1,7 +1,7 @@
 // Copyright 1998-2014 Epic Games, Inc. All Rights Reserved.
 
 #include "EnginePrivate.h"
-#include "MaterialInstance.h"
+#include "MaterialInstanceSupport.h"
 #include "MaterialShader.h"
 #include "TargetPlatform.h"
 
diff --git a/Engine/Source/Runtime/Engine/Private/Materials/MaterialInstance.h b/Engine/Source/Runtime/Engine/Private/Materials/MaterialInstance.h
deleted file mode 100644
diff --git a/Engine/Source/Runtime/Engine/Private/Materials/MaterialInstanceConstant.cpp b/Engine/Source/Runtime/Engine/Private/Materials/MaterialInstanceConstant.cpp
--- a/Engine/Source/Runtime/Engine/Private/Materials/MaterialInstanceConstant.cpp
+++ b/Engine/Source/Runtime/Engine/Private/Materials/MaterialInstanceConstant.cpp
@@ -5,7 +5,7 @@
 =============================================================================*/
 
 #include "EnginePrivate.h"
-#include "MaterialInstance.h"
+#include "MaterialInstanceSupport.h"
 
 
 UMaterialInstanceConstant::UMaterialInstanceConstant(const class FPostConstructInitializeProperties& PCIP)
diff --git a/Engine/Source/Runtime/Engine/Private/Materials/MaterialInstanceDynamic.cpp b/Engine/Source/Runtime/Engine/Private/Materials/MaterialInstanceDynamic.cpp
--- a/Engine/Source/Runtime/Engine/Private/Materials/MaterialInstanceDynamic.cpp
+++ b/Engine/Source/Runtime/Engine/Private/Materials/MaterialInstanceDynamic.cpp
@@ -5,7 +5,7 @@
 ==============================================================================*/
 
 #include "EnginePrivate.h"
-#include "MaterialInstance.h"
+#include "MaterialInstanceSupport.h"
 
 UMaterialInstanceDynamic::UMaterialInstanceDynamic(const class FPostConstructInitializeProperties& PCIP)
 	: Super(PCIP)
diff --git a/Engine/Source/Runtime/Engine/Private/Materials/MaterialInstanceSupport.h b/Engine/Source/Runtime/Engine/Private/Materials/MaterialInstanceSupport.h
new file mode 100644
--- /dev/null
+++ b/Engine/Source/Runtime/Engine/Private/Materials/MaterialInstanceSupport.h
@@ -0,0 +1,215 @@
+// Copyright 1998-2014 Epic Games, Inc. All Rights Reserved.
+
+/*=============================================================================
+	MaterialInstance.h: MaterialInstance definitions.
+=============================================================================*/
+
+#pragma once
+
+/**
+ * Cache uniform expressions for the given material instance.
+ * @param MaterialInstance - The material instance for which to cache uniform expressions.
+ */
+void CacheMaterialInstanceUniformExpressions(const UMaterialInstance* MaterialInstance);
+
+/**
+ * Recaches uniform expressions for all material instances with a given parent.
+ * WARNING: This function is a noop outside of the Editor!
+ * @param ParentMaterial - The parent material to look for.
+ */
+void RecacheMaterialInstanceUniformExpressions(const UMaterialInterface* ParentMaterial);
+
+/** Protects the members of a UMaterialInstanceConstant from re-entrance. */
+class FMICReentranceGuard
+{
+public:
+
+	FMICReentranceGuard(const UMaterialInstance* InMaterial)
+	: Material(InMaterial)
+	{
+		check(IsInGameThread());
+		// NOTE:  we are switching the check to be a log so we can find out which material is causing the assert
+		//check(!Material->ReentrantFlag);
+		if( Material->ReentrantFlag == 1 )
+		{
+			UE_LOG(LogMaterial, Warning, TEXT("InMaterial: %s GameThread: %d RenderThread: %d"), *InMaterial->GetFullName(), IsInGameThread(), IsInRenderingThread() );
+			check(!Material->ReentrantFlag);
+		}
+		const_cast<UMaterialInstance*>(Material)->ReentrantFlag = 1;
+	}
+
+	~FMICReentranceGuard()
+	{
+		check(IsInGameThread());
+		const_cast<UMaterialInstance*>(Material)->ReentrantFlag = 0;
+	}
+
+private:
+	const UMaterialInstance* Material;
+};
+
+/**
+* The resource used to render a UMaterialInstance.
+*/
+class FMaterialInstanceResource: public FMaterialRenderProxy
+{
+public:
+
+	/** Material instances store pairs of names and values in arrays to look up parameters at run time. */
+	template <typename ValueType>
+	struct TNamedParameter
+	{
+		FName Name;
+		ValueType Value;
+	};
+
+	/** Initialization constructor. */
+	FMaterialInstanceResource(UMaterialInstance* InOwner,bool bInSelected,bool bInHovered);
+
+	/**
+	 * Called from the game thread to destroy the material instance on the rendering thread.
+	 */
+	void GameThread_Destroy()
+	{
+		ENQUEUE_UNIQUE_RENDER_COMMAND_ONEPARAMETER(
+			FDestroyMaterialInstanceResourceCommand,
+			FMaterialInstanceResource*,Resource,this,
+		{
+			delete Resource;
+		});
+	}
+
+	// FRenderResource interface.
+	virtual FString GetFriendlyName() const OVERRIDE { return Owner->GetName(); }
+
+	// FMaterialRenderProxy interface.
+	/** Get the FMaterial to use for rendering.  Must return a valid FMaterial, even if it had to fall back to the default material. */
+	virtual const FMaterial* GetMaterial(ERHIFeatureLevel::Type FeatureLevel) const OVERRIDE;
+	/** Get the FMaterial that should be used for rendering, but might not be in a valid state to actually use.  Can return NULL. */
+	virtual FMaterial* GetMaterialNoFallback(ERHIFeatureLevel::Type FeatureLevel) const OVERRIDE;
+	virtual bool GetVectorValue(const FName ParameterName, FLinearColor* OutValue, const FMaterialRenderContext& Context) const OVERRIDE;
+	virtual bool GetScalarValue(const FName ParameterName,float* OutValue, const FMaterialRenderContext& Context) const OVERRIDE;
+	virtual bool GetTextureValue(const FName ParameterName,const UTexture** OutValue, const FMaterialRenderContext& Context) const OVERRIDE;
+	virtual float GetDistanceFieldPenumbraScale() const OVERRIDE { return DistanceFieldPenumbraScale; }
+
+	/** Sets the material instance's parent. */
+	void GameThread_SetParent(UMaterialInterface* InParent);
+
+	/** Called from the game thread to update DistanceFieldPenumbraScale. */
+	void GameThread_UpdateDistanceFieldPenumbraScale(float NewDistanceFieldPenumbraScale);
+
+	/**
+	 * Clears all parameters set on this material instance.
+	 */
+	void RenderThread_ClearParameters()
+	{
+		VectorParameterArray.Empty();
+		ScalarParameterArray.Empty();
+		TextureParameterArray.Empty();
+		InvalidateUniformExpressionCache();
+	}
+
+	/**
+	 * Updates a named parameter on the render thread.
+	 */
+	template <typename ValueType>
+	void RenderThread_UpdateParameter(const FName Name, const ValueType& Value )
+	{
+		InvalidateUniformExpressionCache();
+		TArray<TNamedParameter<ValueType> >& ValueArray = GetValueArray<ValueType>();
+		const int32 ParameterCount = ValueArray.Num();
+		for (int32 ParameterIndex = 0; ParameterIndex < ParameterCount; ++ParameterIndex)
+		{
+			TNamedParameter<ValueType>& Parameter = ValueArray[ParameterIndex];
+			if (Parameter.Name == Name)
+			{
+				Parameter.Value = Value;
+				return;
+			}
+		}
+		TNamedParameter<ValueType> NewParameter;
+		NewParameter.Name = Name;
+		NewParameter.Value = Value;
+		ValueArray.Add(NewParameter);
+	}
+
+	/**
+	 * Retrieves a parameter by name.
+	 */
+	template <typename ValueType>
+	const ValueType* RenderThread_FindParameterByName(FName ParameterName) const
+	{
+		const TArray<TNamedParameter<ValueType> >& ValueArray = GetValueArray<ValueType>();
+		const int32 ParameterCount = ValueArray.Num();
+		for (int32 ParameterIndex = 0; ParameterIndex < ParameterCount; ++ParameterIndex)
+		{
+			const TNamedParameter<ValueType>& Parameter = ValueArray[ParameterIndex];
+			if (Parameter.Name == ParameterName)
+			{
+				return &Parameter.Value;
+			}
+		}
+		return NULL;
+	}
+
+private:
+	/**
+	 * Retrieves the array of values for a given type.
+	 */
+	template <typename ValueType> TArray<TNamedParameter<ValueType> >& GetValueArray() { return ScalarParameterArray; }
+	template <typename ValueType> const TArray<TNamedParameter<ValueType> >& GetValueArray() const { return ScalarParameterArray; }
+
+	/** The parent of the material instance. */
+	UMaterialInterface* Parent;
+
+	/** The UMaterialInstance which owns this resource. */
+	UMaterialInstance* Owner;
+
+	/** Scales the penumbra size of distance field shadows. */
+	float DistanceFieldPenumbraScale;
+
+	/** The game thread accessible parent of the material instance. */
+	UMaterialInterface* GameThreadParent;
+
+	/** Vector parameters for this material instance. */
+	TArray<TNamedParameter<FLinearColor> > VectorParameterArray;
+	/** Scalar parameters for this material instance. */
+	TArray<TNamedParameter<float> > ScalarParameterArray;
+	/** Texture parameters for this material instance. */
+	TArray<TNamedParameter<const UTexture*> > TextureParameterArray;
+};
+
+template <> FORCEINLINE TArray<FMaterialInstanceResource::TNamedParameter<float> >& FMaterialInstanceResource::GetValueArray() { return ScalarParameterArray; }
+template <> FORCEINLINE TArray<FMaterialInstanceResource::TNamedParameter<FLinearColor> >& FMaterialInstanceResource::GetValueArray() { return VectorParameterArray; }
+template <> FORCEINLINE TArray<FMaterialInstanceResource::TNamedParameter<const UTexture*> >& FMaterialInstanceResource::GetValueArray() { return TextureParameterArray; }
+template <> FORCEINLINE const TArray<FMaterialInstanceResource::TNamedParameter<float> >& FMaterialInstanceResource::GetValueArray() const { return ScalarParameterArray; }
+template <> FORCEINLINE const TArray<FMaterialInstanceResource::TNamedParameter<FLinearColor> >& FMaterialInstanceResource::GetValueArray() const { return VectorParameterArray; }
+template <> FORCEINLINE const TArray<FMaterialInstanceResource::TNamedParameter<const UTexture*> >& FMaterialInstanceResource::GetValueArray() const { return TextureParameterArray; }
+
+/** Finds a parameter by name from the game thread. */
+template <typename ParameterType>
+ParameterType* GameThread_FindParameterByName(TArray<ParameterType>& Parameters, FName Name)
+{
+	for (int32 ParameterIndex = 0; ParameterIndex < Parameters.Num(); ParameterIndex++)
+	{
+		ParameterType* Parameter = &Parameters[ParameterIndex];
+		if (Parameter->ParameterName == Name)
+		{
+			return Parameter;
+		}
+	}
+	return NULL;
+}
+template <typename ParameterType>
+const ParameterType* GameThread_FindParameterByName(const TArray<ParameterType>& Parameters, FName Name)
+{
+	for (int32 ParameterIndex = 0; ParameterIndex < Parameters.Num(); ParameterIndex++)
+	{
+		const ParameterType* Parameter = &Parameters[ParameterIndex];
+		if (Parameter->ParameterName == Name)
+		{
+			return Parameter;
+		}
+	}
+	return NULL;
+}
diff --git a/Engine/Source/Runtime/Engine/Private/SceneView.cpp b/Engine/Source/Runtime/Engine/Private/SceneView.cpp
--- a/Engine/Source/Runtime/Engine/Private/SceneView.cpp
+++ b/Engine/Source/Runtime/Engine/Private/SceneView.cpp
@@ -6,7 +6,7 @@
 
 #include "EnginePrivate.h"
 #include "SceneManagement.h"
-#include "Engine.h"
+#include "EnginePrivate.h"
 #include "StereoRendering.h"
 #include "HighResScreenshot.h"
 #include "Slate/SceneViewport.h"
diff --git a/Engine/Source/Runtime/Engine/Private/Tests/AutomationCommon.cpp b/Engine/Source/Runtime/Engine/Private/Tests/AutomationCommon.cpp
new file mode 100644
--- /dev/null
+++ b/Engine/Source/Runtime/Engine/Private/Tests/AutomationCommon.cpp
@@ -0,0 +1,45 @@
+// Copyright 1998-2014 Epic Games, Inc. All Rights Reserved.
+
+#include "EnginePrivate.h"
+#include "AutomationCommon.h"
+#include "ImageUtils.h"
+
+///////////////////////////////////////////////////////////////////////
+// Common Latent commands
+
+namespace AutomationCommon
+{
+	/** These save a PNG and get sent over the network */
+	static void SaveWindowAsScreenshot(TSharedRef<SWindow> Window, const FString& FileName)
+	{
+		TSharedRef<SWidget> WindowRef = Window;
+
+		TArray<FColor> OutImageData;
+		FIntVector OutImageSize;
+		FSlateApplication::Get().TakeScreenshot(WindowRef,OutImageData,OutImageSize);
+
+		FAutomationTestFramework::GetInstance().OnScreenshotCaptured().ExecuteIfBound(OutImageSize.X, OutImageSize.Y, OutImageData, FileName);
+	}
+}
+
+bool FWaitLatentCommand::Update()
+{
+	float NewTime = FPlatformTime::Seconds();
+	if (NewTime - StartTime >= Duration)
+	{
+		return true;
+	}
+	return false;
+}
+
+bool FTakeActiveEditorScreenshotCommand::Update()
+{
+	AutomationCommon::SaveWindowAsScreenshot(FSlateApplication::Get().GetActiveTopLevelWindow().ToSharedRef(),ScreenshotName);
+	return true;
+}
+
+bool FTakeEditorScreenshotCommand::Update()
+{
+	AutomationCommon::SaveWindowAsScreenshot(ScreenshotParameters.CurrentWindow.ToSharedRef(), ScreenshotParameters.ScreenshotName);
+	return true;
+}
\ No newline at end of file
diff --git a/Engine/Source/Runtime/Engine/Private/UserInterface/CanvasItem.cpp b/Engine/Source/Runtime/Engine/Private/UserInterface/CanvasItem.cpp
--- a/Engine/Source/Runtime/Engine/Private/UserInterface/CanvasItem.cpp
+++ b/Engine/Source/Runtime/Engine/Private/UserInterface/CanvasItem.cpp
@@ -8,7 +8,7 @@
 #include "CanvasItem.h"
 #include "TileRendering.h"
 #include "RHIStaticStates.h"
-#include "Engine.h"
+#include "EnginePrivate.h"
 #include "Internationalization.h"
 
 
diff --git a/Engine/Source/Runtime/Engine/Public/EnginePrivate.h b/Engine/Source/Runtime/Engine/Public/EnginePrivate.h
deleted file mode 100644
diff --git a/Engine/Source/Runtime/Engine/Public/Tests/AutomationCommon.cpp b/Engine/Source/Runtime/Engine/Public/Tests/AutomationCommon.cpp
deleted file mode 100644
diff --git a/Engine/Source/Runtime/Engine/Public/Tests/AutomationCommon.h b/Engine/Source/Runtime/Engine/Public/Tests/AutomationCommon.h
--- a/Engine/Source/Runtime/Engine/Public/Tests/AutomationCommon.h
+++ b/Engine/Source/Runtime/Engine/Public/Tests/AutomationCommon.h
@@ -2,7 +2,7 @@
 
 #pragma once
 
-#include "EnginePrivate.h"
+#include "Engine.h"
 
 ///////////////////////////////////////////////////////////////////////
 // Common Latent commands which are used across test type. I.e. Engine, Network, etc...
diff --git a/Engine/Source/Runtime/Renderer/Private/CustomDepthRendering.cpp b/Engine/Source/Runtime/Renderer/Private/CustomDepthRendering.cpp
--- a/Engine/Source/Runtime/Renderer/Private/CustomDepthRendering.cpp
+++ b/Engine/Source/Runtime/Renderer/Private/CustomDepthRendering.cpp
@@ -5,7 +5,7 @@
 =============================================================================*/
 
 #include "RendererPrivate.h"
-#include "EnginePrivate.h"
+#include "Engine.h"
 #include "ScenePrivate.h"
 #include "ScreenRendering.h"
 #include "PostProcessing.h"
diff --git a/Engine/Source/Runtime/Renderer/Private/DeferredShadingRenderer.cpp b/Engine/Source/Runtime/Renderer/Private/DeferredShadingRenderer.cpp
--- a/Engine/Source/Runtime/Renderer/Private/DeferredShadingRenderer.cpp
+++ b/Engine/Source/Runtime/Renderer/Private/DeferredShadingRenderer.cpp
@@ -5,7 +5,7 @@
 =============================================================================*/
 
 #include "RendererPrivate.h"
-#include "EnginePrivate.h"
+#include "Engine.h"
 #include "ScenePrivate.h"
 #include "ScreenRendering.h"
 #include "SceneFilterRendering.h"
diff --git a/Engine/Source/Runtime/Renderer/Private/DistortionRendering.cpp b/Engine/Source/Runtime/Renderer/Private/DistortionRendering.cpp
--- a/Engine/Source/Runtime/Renderer/Private/DistortionRendering.cpp
+++ b/Engine/Source/Runtime/Renderer/Private/DistortionRendering.cpp
@@ -5,7 +5,7 @@
 =============================================================================*/
 
 #include "RendererPrivate.h"
-#include "EnginePrivate.h"
+#include "Engine.h"
 #include "ScenePrivate.h"
 #include "ScreenRendering.h"
 #include "PostProcessing.h"
diff --git a/Engine/Source/Runtime/Renderer/Private/ForwardShadingRenderer.cpp b/Engine/Source/Runtime/Renderer/Private/ForwardShadingRenderer.cpp
--- a/Engine/Source/Runtime/Renderer/Private/ForwardShadingRenderer.cpp
+++ b/Engine/Source/Runtime/Renderer/Private/ForwardShadingRenderer.cpp
@@ -5,7 +5,7 @@
 =============================================================================*/
 
 #include "RendererPrivate.h"
-#include "EnginePrivate.h"
+#include "Engine.h"
 #include "ScenePrivate.h"
 #include "FXSystem.h"
 #include "PostProcessing.h"
diff --git a/Engine/Source/Runtime/Renderer/Private/RendererPrivate.h b/Engine/Source/Runtime/Renderer/Private/RendererPrivate.h
--- a/Engine/Source/Runtime/Renderer/Private/RendererPrivate.h
+++ b/Engine/Source/Runtime/Renderer/Private/RendererPrivate.h
@@ -7,7 +7,7 @@
 #ifndef __RendererPrivate_h__
 #define __RendererPrivate_h__
 
-#include "EnginePrivate.h"
+#include "Engine.h"
 #include "ShaderCore.h"
 #include "RHI.h"
 #include "ScenePrivate.h"
diff --git a/Engine/Source/Runtime/Renderer/Private/SceneRendering.cpp b/Engine/Source/Runtime/Renderer/Private/SceneRendering.cpp
--- a/Engine/Source/Runtime/Renderer/Private/SceneRendering.cpp
+++ b/Engine/Source/Runtime/Renderer/Private/SceneRendering.cpp
@@ -5,7 +5,7 @@
 =============================================================================*/
 
 #include "RendererPrivate.h"
-#include "EnginePrivate.h"
+#include "Engine.h"
 #include "ScenePrivate.h"
 #include "ScreenRendering.h"
 #include "SceneFilterRendering.h"
diff --git a/Engine/Source/Runtime/Renderer/Private/SceneVisibility.cpp b/Engine/Source/Runtime/Renderer/Private/SceneVisibility.cpp
--- a/Engine/Source/Runtime/Renderer/Private/SceneVisibility.cpp
+++ b/Engine/Source/Runtime/Renderer/Private/SceneVisibility.cpp
@@ -5,7 +5,7 @@
 =============================================================================*/
 
 #include "RendererPrivate.h"
-#include "EnginePrivate.h"
+#include "Engine.h"
 #include "ScenePrivate.h"
 #include "FXSystem.h"
 #include "../../Engine/Private/SkeletalRenderGPUSkin.h"		// GPrevPerBoneMotionBlur
diff --git a/Engine/Source/Runtime/StreamingPauseRendering/Private/StreamingPauseRendering.cpp b/Engine/Source/Runtime/StreamingPauseRendering/Private/StreamingPauseRendering.cpp
--- a/Engine/Source/Runtime/StreamingPauseRendering/Private/StreamingPauseRendering.cpp
+++ b/Engine/Source/Runtime/StreamingPauseRendering/Private/StreamingPauseRendering.cpp
@@ -4,7 +4,7 @@
 StreamingPauseRendering.cpp: Streaming pause implementation.
 =============================================================================*/
 
-#include "EnginePrivate.h"
+#include "Engine.h"
 #include "StreamingPauseRendering.h"
 #include "SceneViewport.h"
 #include "MoviePlayer.h"
