Fixed transform of Canvas Material rendering for ES2 on Device

diff --git a/Engine/Source/Runtime/Engine/Private/TileRendering.cpp b/Engine/Source/Runtime/Engine/Private/TileRendering.cpp
--- a/Engine/Source/Runtime/Engine/Private/TileRendering.cpp
+++ b/Engine/Source/Runtime/Engine/Private/TileRendering.cpp
@@ -134,10 +134,20 @@ void FTileRenderer::DrawTile(const class FSceneView& View, const FMaterialRender
 	FMaterialTileVertex DestVertex[4];
 
 	// create verts
-	DestVertex[0].Initialize(X + SizeX, Y, U + SizeU, V);
-	DestVertex[1].Initialize(X, Y, U, V);
-	DestVertex[2].Initialize(X + SizeX, Y + SizeY, U + SizeU, V + SizeV);
-	DestVertex[3].Initialize(X, Y + SizeY, U, V + SizeV);
+	if ((IsES2Platform(GRHIShaderPlatform) && !IsPCPlatform(GRHIShaderPlatform)))
+	{
+		DestVertex[0].Initialize(X + SizeX, View.ViewRect.Height() - (Y + SizeY), U + SizeU, V + SizeV);
+		DestVertex[1].Initialize(X, View.ViewRect.Height() - (Y + SizeY), U, V + SizeV);
+		DestVertex[2].Initialize(X + SizeX, View.ViewRect.Height() - Y, U + SizeU, V);
+		DestVertex[3].Initialize(X, View.ViewRect.Height() - Y, U, V);
+	}
+	else
+	{
+		DestVertex[0].Initialize(X + SizeX, Y, U + SizeU, V);
+		DestVertex[1].Initialize(X, Y, U, V);
+		DestVertex[2].Initialize(X + SizeX, Y + SizeY, U + SizeU, V + SizeV);
+		DestVertex[3].Initialize(X, Y + SizeY, U, V + SizeV);
+	}
 
 	// update the FMeshBatch
 	FMeshBatch& Mesh = GTileMesh.MeshElement;
