Fix for AStaticMeshActors replicating position/movement.

diff --git a/Engine/Source/Runtime/Engine/Classes/Engine/StaticMeshActor.h b/Engine/Source/Runtime/Engine/Classes/Engine/StaticMeshActor.h
--- a/Engine/Source/Runtime/Engine/Classes/Engine/StaticMeshActor.h
+++ b/Engine/Source/Runtime/Engine/Classes/Engine/StaticMeshActor.h
@@ -14,6 +14,8 @@ class ENGINE_API AStaticMeshActor : public AActor, public INavRelevantActorInter
 
 	UPROPERTY(Category=StaticMeshActor, VisibleAnywhere, BlueprintReadOnly, meta=(ExposeFunctionCategories="Mesh,Rendering,Physics,Components|StaticMesh"))
 	TSubobjectPtr<class UStaticMeshComponent> StaticMeshComponent;
+	
+	virtual void BeginPlay() OVERRIDE;
 
 	/** This static mesh should replicate movement. Automatically sets the RemoteRole and bReplicateMovement flags. Meant to be edited on placed actors (those other two proeprties are not) */
 	UPROPERTY(Category=Actor, EditAnywhere, AdvancedDisplay)
diff --git a/Engine/Source/Runtime/Engine/Private/StaticMeshActor.cpp b/Engine/Source/Runtime/Engine/Private/StaticMeshActor.cpp
--- a/Engine/Source/Runtime/Engine/Private/StaticMeshActor.cpp
+++ b/Engine/Source/Runtime/Engine/Private/StaticMeshActor.cpp
@@ -25,6 +25,24 @@ AStaticMeshActor::AStaticMeshActor(const class FPostConstructInitializePropertie
 	RootComponent = StaticMeshComponent;
 }
 
+void AStaticMeshActor::BeginPlay()
+{
+	// Since we allow AStaticMeshActor to specify whether it replicates via bStaticMeshReplicateMovement - per placed instance
+	// and we do the normal SetReplicates call in PostInitProperties, before instanced properties are serialized in, 
+	// we need to do this here. 
+	//
+	// This is a short term fix until we find a better play for SetReplicates to be called in AActor.
+
+	if (Role == ROLE_Authority && bStaticMeshReplicateMovement)
+	{
+		bReplicates = false;
+		SetRemoteRoleForBackwardsCompat(ROLE_SimulatedProxy);
+		SetReplicates(true);
+	}	
+
+	Super::BeginPlay();
+}
+
 FString AStaticMeshActor::GetDetailedInfoInternal() const
 {
 	return StaticMeshComponent ? StaticMeshComponent->GetDetailedInfoInternal() : TEXT("No_StaticMeshComponent");
