Move engine stat declarations into EngineStats.h
Remove unused GPropertyWindowDataCache
Move MAX/MIN_ORTHOZOOM to EditorViewportClient.cpp
Move ELevelTick,  ENetworkFailure and ETravelFailure to EngineBaseTypes
Move FBoneIndexType declaration to SkeletalMeshTypes.h
Remove unecessary FURL forward decl

diff --git a/Engine/Source/Editor/UnrealEd/Private/EditorViewportClient.cpp b/Engine/Source/Editor/UnrealEd/Private/EditorViewportClient.cpp
--- a/Engine/Source/Editor/UnrealEd/Private/EditorViewportClient.cpp
+++ b/Engine/Source/Editor/UnrealEd/Private/EditorViewportClient.cpp
@@ -48,6 +48,9 @@ namespace {
 	static const float LightRotSpeed = 0.22f;
 }
 
+#define MIN_ORTHOZOOM				250.0					/* Limit of 2D viewport zoom in */
+#define MAX_ORTHOZOOM				MAX_FLT					/* Limit of 2D viewport zoom out */
+
 namespace OrbitConstants
 {
 	const float OrbitPanSpeed = 1.0f;
diff --git a/Engine/Source/Runtime/Engine/Classes/Engine/EngineBaseTypes.h b/Engine/Source/Runtime/Engine/Classes/Engine/EngineBaseTypes.h
--- a/Engine/Source/Runtime/Engine/Classes/Engine/EngineBaseTypes.h
+++ b/Engine/Source/Runtime/Engine/Classes/Engine/EngineBaseTypes.h
@@ -379,6 +379,130 @@ struct FPrimitiveComponentPostPhysicsTickFunction : public FTickFunction
 	virtual FString DiagnosticMessage();
 };
 
+enum ELevelTick
+{
+	LEVELTICK_TimeOnly = 0,	// Update the level time only.
+	LEVELTICK_ViewportsOnly = 1,	// Update time and viewports.
+	LEVELTICK_All = 2,	// Update all.
+	LEVELTICK_PauseTick = 3,	// Delta time is zero, we are paused. Components don't tick.
+};
+
+
+
+/** Types of network failures broadcast from the engine */
+namespace ENetworkFailure
+{
+	enum Type
+	{
+		/** A relevant net driver has already been created for this service */
+		NetDriverAlreadyExists,
+		/** The net driver creation failed */
+		NetDriverCreateFailure,
+		/** The net driver failed its Listen() call */
+		NetDriverListenFailure,
+		/** A connection to the net driver has been lost */
+		ConnectionLost,
+		/** A connection to the net driver has timed out */
+		ConnectionTimeout,
+		/** The net driver received an NMT_Failure message */
+		FailureReceived,
+		/** The client needs to upgrade their game */
+		OutdatedClient,
+		/** The server needs to upgrade their game */
+		OutdatedServer,
+		/** There was an error during connection to the game */
+		PendingConnectionFailure
+	};
+
+	inline const TCHAR* ToString(ENetworkFailure::Type FailureType)
+	{
+		switch (FailureType)
+		{
+		case NetDriverAlreadyExists:
+			return TEXT("NetDriverAlreadyExists");
+		case NetDriverCreateFailure:
+			return TEXT("NetDriverCreateFailure");
+		case NetDriverListenFailure:
+			return TEXT("NetDriverListenFailure");
+		case ConnectionLost:
+			return TEXT("ConnectionLost");
+		case ConnectionTimeout:
+			return TEXT("ConnectionTimeout");
+		case FailureReceived:
+			return TEXT("FailureReceived");
+		case OutdatedClient:
+			return TEXT("OutdatedClient");
+		case OutdatedServer:
+			return TEXT("OutdatedServer");
+		case PendingConnectionFailure:
+			return TEXT("PendingConnectionFailure");
+		}
+		return TEXT("Unknown ENetworkFailure error occurred.");
+	}
+}
+
+/** Types of server travel failures broadcast by the engine */
+namespace ETravelFailure
+{
+	enum Type
+	{
+		/** No level found in the loaded package */
+		NoLevel,
+		/** LoadMap failed on travel (about to Browse to default map) */
+		LoadMapFailure,
+		/** Invalid URL specified */
+		InvalidURL,
+		/** A package is missing on the client */
+		PackageMissing,
+		/** A package version mismatch has occurred between client and server */
+		PackageVersion,
+		/** A package is missing and the client is unable to download the file */
+		NoDownload,
+		/** General travel failure */
+		TravelFailure,
+		/** Cheat commands have been used disabling travel */
+		CheatCommands,
+		/** Failed to create the pending net game for travel */
+		PendingNetGameCreateFailure,
+		/** Failed to save before travel */
+		CloudSaveFailure,
+		/** There was an error during a server travel to a new map */
+		ServerTravelFailure,
+		/** There was an error during a client travel to a new map */
+		ClientTravelFailure,
+	};
+
+	inline const TCHAR* ToString(ETravelFailure::Type FailureType)
+	{
+		switch (FailureType)
+		{
+		case NoLevel:
+			return TEXT("NoLevel");
+		case LoadMapFailure:
+			return TEXT("LoadMapFailure");
+		case InvalidURL:
+			return TEXT("InvalidURL");
+		case PackageMissing:
+			return TEXT("PackageMissing");
+		case PackageVersion:
+			return TEXT("PackageVersion");
+		case NoDownload:
+			return TEXT("NoDownload");
+		case TravelFailure:
+			return TEXT("TravelFailure");
+		case CheatCommands:
+			return TEXT("CheatCommands");
+		case PendingNetGameCreateFailure:
+			return TEXT("PendingNetGameCreateFailure");
+		case ServerTravelFailure:
+			return TEXT("ServerTravelFailure");
+		case ClientTravelFailure:
+			return TEXT("ClientTravelFailure");
+		}
+		return TEXT("Unknown ETravelFailure error occurred.");
+	}
+}
+
 // Traveling from server to server.
 UENUM()
 enum ETravelType
diff --git a/Engine/Source/Runtime/Engine/Private/UnrealEngine.cpp b/Engine/Source/Runtime/Engine/Private/UnrealEngine.cpp
--- a/Engine/Source/Runtime/Engine/Private/UnrealEngine.cpp
+++ b/Engine/Source/Runtime/Engine/Private/UnrealEngine.cpp
@@ -69,11 +69,6 @@ void FEngineModule::StartupModule()
  */
 ENGINE_API UEngine*	GEngine = NULL;
 
-/**
- * Cache for property window queries
- */
-FPropertyWindowDataCache* GPropertyWindowDataCache = NULL;
-
 /**
  * Whether to visualize the light map selected by the Debug Camera.
  */
diff --git a/Engine/Source/Runtime/Engine/Public/Engine.h b/Engine/Source/Runtime/Engine/Public/Engine.h
--- a/Engine/Source/Runtime/Engine/Public/Engine.h
+++ b/Engine/Source/Runtime/Engine/Public/Engine.h
@@ -136,14 +136,11 @@ ENGINE_API DECLARE_LOG_CATEGORY_EXTERN(LogPlayerManagement, Error, All);
  */
 extern ENGINE_API class UEngine*			GEngine;
 
-/** Cache for property window queries*/
-extern ENGINE_API class FPropertyWindowDataCache* GPropertyWindowDataCache;
-
 /** when set, disallows all network travel (pending level rejects all client travel attempts) */
-extern ENGINE_API bool GDisallowNetworkTravel;
+extern ENGINE_API bool						GDisallowNetworkTravel;
 
 /** The GPU time taken to render the last frame. Same metric as FPlatformTime::Cycles(). */
-extern ENGINE_API uint32				GGPUFrameTime;
+extern ENGINE_API uint32					GGPUFrameTime;
 
 #if WITH_EDITOR
 
@@ -194,354 +191,13 @@ struct ENGINE_API FEditorSupportDelegates
 #define WORLD_MAX					524288.0				/* Maximum size of the world */
 #define HALF_WORLD_MAX				(WORLD_MAX * 0.5)		/* Half the maximum size of the world */
 #define HALF_WORLD_MAX1				(HALF_WORLD_MAX - 1.0)	/* Half the maximum size of the world minus one */
-#define MIN_ORTHOZOOM				250.0					/* Limit of 2D viewport zoom in */
-#define MAX_ORTHOZOOM				MAX_FLT					/* Limit of 2D viewport zoom out */
 #define DEFAULT_ORTHOZOOM			10000.0					/* Default 2D viewport zoom */
 
 /*-----------------------------------------------------------------------------
-	Level updating. Moved out here from Level.h so that AActor.h can use it (needed for gcc)
+	Stats.
 -----------------------------------------------------------------------------*/
 
-enum ELevelTick
-{
-	LEVELTICK_TimeOnly		= 0,	// Update the level time only.
-	LEVELTICK_ViewportsOnly	= 1,	// Update time and viewports.
-	LEVELTICK_All			= 2,	// Update all.
-	LEVELTICK_PauseTick		= 3,	// Delta time is zero, we are paused. Components don't tick.
-};
-
-
-/** Types of network failures broadcast from the engine */
-namespace ENetworkFailure
-{
-	enum Type
-	{
-		/** A relevant net driver has already been created for this service */
-		NetDriverAlreadyExists,
-		/** The net driver creation failed */
-		NetDriverCreateFailure,
-		/** The net driver failed its Listen() call */
-		NetDriverListenFailure,
-		/** A connection to the net driver has been lost */
-		ConnectionLost,
-		/** A connection to the net driver has timed out */
-		ConnectionTimeout,
-		/** The net driver received an NMT_Failure message */
-		FailureReceived,
-		/** The client needs to upgrade their game */
-		OutdatedClient,
-		/** The server needs to upgrade their game */
-		OutdatedServer,
-		/** There was an error during connection to the game */
-		PendingConnectionFailure
-	};
-
-	inline const TCHAR* ToString(ENetworkFailure::Type FailureType)
-	{
-		switch (FailureType)
-		{
-		case NetDriverAlreadyExists:
-			return TEXT("NetDriverAlreadyExists");
-		case NetDriverCreateFailure:
-			return TEXT("NetDriverCreateFailure");
-		case NetDriverListenFailure:
-			return TEXT("NetDriverListenFailure");
-		case ConnectionLost:
-			return TEXT("ConnectionLost");
-		case ConnectionTimeout:
-			return TEXT("ConnectionTimeout");
-		case FailureReceived:
-			return TEXT("FailureReceived");
-		case OutdatedClient:
-			return TEXT("OutdatedClient");
-		case OutdatedServer:
-			return TEXT("OutdatedServer");
-		case PendingConnectionFailure:
-			return TEXT("PendingConnectionFailure");
-		}
-		return TEXT("Unknown ENetworkFailure error occurred.");	
-	}
-}
-
-/** Types of server travel failures broadcast by the engine */
-namespace ETravelFailure
-{
-	enum Type
-	{	
-		/** No level found in the loaded package */
-		NoLevel,
-		/** LoadMap failed on travel (about to Browse to default map) */
-		LoadMapFailure,
-		/** Invalid URL specified */
-		InvalidURL,
-		/** A package is missing on the client */
-		PackageMissing,
-		/** A package version mismatch has occurred between client and server */
-		PackageVersion,
-		/** A package is missing and the client is unable to download the file */
-		NoDownload,
-		/** General travel failure */
-		TravelFailure,
-		/** Cheat commands have been used disabling travel */
-		CheatCommands,
-		/** Failed to create the pending net game for travel */
-		PendingNetGameCreateFailure,
-		/** Failed to save before travel */
-		CloudSaveFailure,
-		/** There was an error during a server travel to a new map */
-		ServerTravelFailure,
-		/** There was an error during a client travel to a new map */
-		ClientTravelFailure,
-	};
-
-	inline const TCHAR* ToString(ETravelFailure::Type FailureType)
-	{
-		switch (FailureType)
-		{
-		case NoLevel:
-			return TEXT("NoLevel");
-		case LoadMapFailure:
-			return TEXT("LoadMapFailure");
-		case InvalidURL:
-			return TEXT("InvalidURL");
-		case PackageMissing:
-			return TEXT("PackageMissing");
-		case PackageVersion:
-			return TEXT("PackageVersion");
-		case NoDownload:
-			return TEXT("NoDownload");
-		case TravelFailure:
-			return TEXT("TravelFailure");
-		case CheatCommands:
-			return TEXT("CheatCommands");
-		case PendingNetGameCreateFailure:
-			return TEXT("PendingNetGameCreateFailure");
-		case ServerTravelFailure:
-			return TEXT("ServerTravelFailure");
-		case ClientTravelFailure:
-			return TEXT("ClientTravelFailure");
-		}
-		return TEXT("Unknown ETravelFailure error occurred.");		
-	}
-}
-
-/*-----------------------------------------------------------------------------
-	Animation support.
------------------------------------------------------------------------------*/
-typedef uint16 FBoneIndexType;
-
-/**
- * Engine stats
- */
-/** Input stat */
-DECLARE_CYCLE_STAT_EXTERN(TEXT("Input Time"),STAT_InputTime,STATGROUP_Engine, );
-DECLARE_FLOAT_COUNTER_STAT_EXTERN(TEXT("Input Latency"),STAT_InputLatencyTime,STATGROUP_Engine, );
-/** HUD stat */
-DECLARE_CYCLE_STAT_EXTERN(TEXT("HUD Time"),STAT_HudTime,STATGROUP_Engine, );
-/** Static mesh tris rendered */
-DECLARE_DWORD_COUNTER_STAT_EXTERN(TEXT("Static Mesh Tris"),STAT_StaticMeshTriangles,STATGROUP_Engine, ENGINE_API);
-DECLARE_CYCLE_STAT_EXTERN(TEXT("Skel Skin Time"),STAT_SkinningTime,STATGROUP_Engine, );
-DECLARE_CYCLE_STAT_EXTERN(TEXT("Update Cloth Verts Time"),STAT_UpdateClothVertsTime,STATGROUP_Engine, );
-DECLARE_CYCLE_STAT_EXTERN(TEXT("Update SoftBody Verts Time"),STAT_UpdateSoftBodyVertsTime,STATGROUP_Engine, );
-DECLARE_DWORD_COUNTER_STAT_EXTERN(TEXT("Skel Mesh Tris"),STAT_SkelMeshTriangles,STATGROUP_Engine, );
-DECLARE_DWORD_COUNTER_STAT_EXTERN(TEXT("Skel Mesh Draw Calls"),STAT_SkelMeshDrawCalls,STATGROUP_Engine, );
-DECLARE_DWORD_COUNTER_STAT_EXTERN(TEXT("Skel Verts CPU Skin"),STAT_CPUSkinVertices,STATGROUP_Engine, );
-DECLARE_DWORD_COUNTER_STAT_EXTERN(TEXT("Skel Verts GPU Skin"),STAT_GPUSkinVertices,STATGROUP_Engine, );
-DECLARE_CYCLE_STAT_EXTERN(TEXT("GameEngine Tick"),STAT_GameEngineTick,STATGROUP_Engine, );
-DECLARE_CYCLE_STAT_EXTERN(TEXT("GameViewport Tick"),STAT_GameViewportTick,STATGROUP_Engine, );
-DECLARE_CYCLE_STAT_EXTERN(TEXT("RedrawViewports"),STAT_RedrawViewports,STATGROUP_Engine, );
-DECLARE_CYCLE_STAT_EXTERN(TEXT("Update Level Streaming"),STAT_UpdateLevelStreaming,STATGROUP_Engine, ENGINE_API);
-DECLARE_CYCLE_STAT_EXTERN(TEXT("RHI Game Tick"),STAT_RHITickTime,STATGROUP_Engine, ENGINE_API);
-DECLARE_CYCLE_STAT_EXTERN(TEXT("Debug Hitch"),STAT_IntentionalHitch,STATGROUP_Engine, );
-DECLARE_CYCLE_STAT_EXTERN(TEXT("Platform Message Time"),STAT_PlatformMessageTime,STATGROUP_Engine, ENGINE_API);
-DECLARE_CYCLE_STAT_EXTERN(TEXT("Frame Sync Time"),STAT_FrameSyncTime,STATGROUP_Engine, ENGINE_API);
-DECLARE_CYCLE_STAT_EXTERN(TEXT("Deferred Tick Time"),STAT_DeferredTickTime,STATGROUP_Engine, ENGINE_API);
-
-/**
- * Game stats
- */
-DECLARE_CYCLE_STAT_EXTERN(TEXT("Async Physics Time"),STAT_PhysicsTime,STATGROUP_Game, );
-DECLARE_CYCLE_STAT_EXTERN(TEXT("Spawn Actor Time"),STAT_SpawnActorTime,STATGROUP_Game, );
-DECLARE_CYCLE_STAT_EXTERN(TEXT("MoveComponent Time"),STAT_MoveComponentTime,STATGROUP_Game, );
-DECLARE_CYCLE_STAT_EXTERN(TEXT("TeleportTo Time"),STAT_TeleportToTime,STATGROUP_Game, );
-DECLARE_CYCLE_STAT_EXTERN(TEXT("GC Mark Time"),STAT_GCMarkTime,STATGROUP_Game, );
-DECLARE_CYCLE_STAT_EXTERN(TEXT("GC Sweep Time"),STAT_GCSweepTime,STATGROUP_Game, );
-DECLARE_CYCLE_STAT_EXTERN(TEXT("Tick Time"),STAT_TickTime,STATGROUP_Game, );
-DECLARE_CYCLE_STAT_EXTERN(TEXT("World Tick Time"),STAT_WorldTickTime,STATGROUP_Game, );
-DECLARE_CYCLE_STAT_EXTERN(TEXT("Update Camera Time"),STAT_UpdateCameraTime,STATGROUP_Game, );
-DECLARE_CYCLE_STAT_EXTERN(TEXT("Post Tick Component Update"),STAT_PostTickComponentUpdate,STATGROUP_Game, );
-DECLARE_CYCLE_STAT_EXTERN(TEXT("Post Tick Component Wait"),STAT_PostTickComponentUpdateWait,STATGROUP_Game, );
-DECLARE_CYCLE_STAT_EXTERN(TEXT("  Recreate"),STAT_PostTickComponentRecreate,STATGROUP_Game, );
-DECLARE_CYCLE_STAT_EXTERN(TEXT("  Transform or RenderData"),STAT_PostTickComponentLW,STATGROUP_Game, );
-DECLARE_CYCLE_STAT_EXTERN(TEXT("Update Particle Data"),STAT_ParticleManagerUpdateData,STATGROUP_Game, );
-DECLARE_CYCLE_STAT_EXTERN(TEXT("Async Work Wait"),STAT_AsyncWorkWaitTime,STATGROUP_Game, );
-DECLARE_CYCLE_STAT_EXTERN(TEXT("Char Movement Authority Time"),STAT_CharacterMovementAuthority,STATGROUP_Game, );
-DECLARE_CYCLE_STAT_EXTERN(TEXT("Char Movement Simulated Time"),STAT_CharacterMovementSimulated,STATGROUP_Game, );
-DECLARE_CYCLE_STAT_EXTERN(TEXT("Visual Logging time"), STAT_VisualLog, STATGROUP_Game, ENGINE_API);
-DECLARE_CYCLE_STAT_EXTERN(TEXT("Net Tick Time"),STAT_NetWorldTickTime,STATGROUP_Game, );
-DECLARE_CYCLE_STAT_EXTERN(TEXT("Nav Tick Time"),STAT_NavWorldTickTime,STATGROUP_Game, );
-DECLARE_CYCLE_STAT_EXTERN(TEXT("Reset Async Trace Time"),STAT_ResetAsyncTraceTickTime,STATGROUP_Game, );
-DECLARE_CYCLE_STAT_EXTERN(TEXT("GT Tickable Time"),STAT_TickableTickTime,STATGROUP_Game, );
-DECLARE_CYCLE_STAT_EXTERN(TEXT("Runtime Movie Tick Time"),STAT_RuntimeMovieSceneTickTime,STATGROUP_Game, );
-DECLARE_CYCLE_STAT_EXTERN(TEXT("Finish Async Trace Time"),STAT_FinishAsyncTraceTickTime,STATGROUP_Game, );
-DECLARE_CYCLE_STAT_EXTERN(TEXT("Net Broadcast Tick Time"),STAT_NetBroadcastTickTime,STATGROUP_Game, );
-DECLARE_CYCLE_STAT_EXTERN(TEXT("  ServerReplicateActors Time"),STAT_NetServerRepActorsTime,STATGROUP_Game, );
-DECLARE_CYCLE_STAT_EXTERN(TEXT("  Consider Actors Time"),STAT_NetConsiderActorsTime,STATGROUP_Game, );
-DECLARE_CYCLE_STAT_EXTERN(TEXT("  Inital Dormant Time"),STAT_NetInitialDormantCheckTime,STATGROUP_Game, );
-DECLARE_CYCLE_STAT_EXTERN(TEXT("  Prioritize Actors Time"),STAT_NetPrioritizeActorsTime,STATGROUP_Game, );
-DECLARE_CYCLE_STAT_EXTERN(TEXT("  Replicate Actors Time"),STAT_NetReplicateActorsTime,STATGROUP_Game, );
-DECLARE_CYCLE_STAT_EXTERN(TEXT("  Dynamic Property Rep Time"),STAT_NetReplicateDynamicPropTime,STATGROUP_Game, );
-DECLARE_DWORD_COUNTER_STAT_EXTERN(TEXT("  Skipped Dynamic Props"),STAT_NetSkippedDynamicProps,STATGROUP_Game, );
-DECLARE_CYCLE_STAT_EXTERN(TEXT("  NetSerializeItemDelta Time"),STAT_NetSerializeItemDeltaTime,STATGROUP_Game, );
-DECLARE_CYCLE_STAT_EXTERN(TEXT("  Static Property Rep Time"),STAT_NetReplicateStaticPropTime,STATGROUP_Game, );
-DECLARE_CYCLE_STAT_EXTERN(TEXT("  Rebuild Conditionals"),STAT_NetRebuildConditionalTime,STATGROUP_Game, );
-DECLARE_CYCLE_STAT_EXTERN(TEXT("Net Post BC Tick Time"),STAT_NetBroadcastPostTickTime,STATGROUP_Game, );
-DECLARE_CYCLE_STAT_EXTERN(TEXT("Net PackageMap SerializeObject"),STAT_PackageMap_SerializeObjectTime,STATGROUP_Game, );
-
-/**
- * FPS chart stats
- */
-DECLARE_FLOAT_ACCUMULATOR_STAT_EXTERN(TEXT("00 - 05 FPS"), STAT_FPSChart_0_5,		STATGROUP_FPSChart, );
-DECLARE_FLOAT_ACCUMULATOR_STAT_EXTERN(TEXT("05 - 10 FPS"), STAT_FPSChart_5_10,		STATGROUP_FPSChart, );
-DECLARE_FLOAT_ACCUMULATOR_STAT_EXTERN(TEXT("10 - 15 FPS"), STAT_FPSChart_10_15,	STATGROUP_FPSChart, );
-DECLARE_FLOAT_ACCUMULATOR_STAT_EXTERN(TEXT("15 - 20 FPS"), STAT_FPSChart_15_20,	STATGROUP_FPSChart, );
-DECLARE_FLOAT_ACCUMULATOR_STAT_EXTERN(TEXT("20 - 25 FPS"), STAT_FPSChart_20_25,	STATGROUP_FPSChart, );
-DECLARE_FLOAT_ACCUMULATOR_STAT_EXTERN(TEXT("25 - 30 FPS"), STAT_FPSChart_25_30,	STATGROUP_FPSChart, );
-DECLARE_FLOAT_ACCUMULATOR_STAT_EXTERN(TEXT("30 - 35 FPS"), STAT_FPSChart_30_35,	STATGROUP_FPSChart, );
-DECLARE_FLOAT_ACCUMULATOR_STAT_EXTERN(TEXT("35 - 40 FPS"), STAT_FPSChart_35_40,	STATGROUP_FPSChart, );
-DECLARE_FLOAT_ACCUMULATOR_STAT_EXTERN(TEXT("40 - 45 FPS"), STAT_FPSChart_40_45,	STATGROUP_FPSChart, );
-DECLARE_FLOAT_ACCUMULATOR_STAT_EXTERN(TEXT("45 - 50 FPS"), STAT_FPSChart_45_50,	STATGROUP_FPSChart, );
-DECLARE_FLOAT_ACCUMULATOR_STAT_EXTERN(TEXT("50 - 55 FPS"), STAT_FPSChart_50_55,	STATGROUP_FPSChart, );
-DECLARE_FLOAT_ACCUMULATOR_STAT_EXTERN(TEXT("55 - 60 FPS"), STAT_FPSChart_55_60,	STATGROUP_FPSChart, );
-DECLARE_FLOAT_ACCUMULATOR_STAT_EXTERN(TEXT("60 - .. FPS"), STAT_FPSChart_60_INF,	STATGROUP_FPSChart, );
-DECLARE_FLOAT_ACCUMULATOR_STAT_EXTERN(TEXT("30+ FPS"), STAT_FPSChart_30Plus,	STATGROUP_FPSChart, );
-DECLARE_FLOAT_ACCUMULATOR_STAT_EXTERN(TEXT("Unaccounted time"), STAT_FPSChart_UnaccountedTime,	STATGROUP_FPSChart, );
-DECLARE_DWORD_ACCUMULATOR_STAT_EXTERN(TEXT("Frame count"),	STAT_FPSChart_FrameCount, STATGROUP_FPSChart, );
-
-/** Hitch stats */
-DECLARE_DWORD_ACCUMULATOR_STAT_EXTERN(TEXT("5.0s  - ....  hitches"), STAT_FPSChart_Hitch_5000_Plus, STATGROUP_FPSChart, );
-DECLARE_DWORD_ACCUMULATOR_STAT_EXTERN(TEXT("2.5s  - 5.0s  hitches"), STAT_FPSChart_Hitch_2500_5000, STATGROUP_FPSChart, );
-DECLARE_DWORD_ACCUMULATOR_STAT_EXTERN(TEXT("2.0s  - 2.5s  hitches"), STAT_FPSChart_Hitch_2000_2500, STATGROUP_FPSChart, );
-DECLARE_DWORD_ACCUMULATOR_STAT_EXTERN(TEXT("1.5s  - 2.0s  hitches"), STAT_FPSChart_Hitch_1500_2000, STATGROUP_FPSChart, );
-DECLARE_DWORD_ACCUMULATOR_STAT_EXTERN(TEXT("1.0s  - 1.5s  hitches"), STAT_FPSChart_Hitch_1000_1500, STATGROUP_FPSChart, );
-DECLARE_DWORD_ACCUMULATOR_STAT_EXTERN(TEXT("0.75s - 1.0s  hitches"), STAT_FPSChart_Hitch_750_1000, STATGROUP_FPSChart, );
-DECLARE_DWORD_ACCUMULATOR_STAT_EXTERN(TEXT("0.5s  - 0.75s hitches"), STAT_FPSChart_Hitch_500_750, STATGROUP_FPSChart, );
-DECLARE_DWORD_ACCUMULATOR_STAT_EXTERN(TEXT("0.3s  - 0.5s  hitches"), STAT_FPSChart_Hitch_300_500, STATGROUP_FPSChart, );
-DECLARE_DWORD_ACCUMULATOR_STAT_EXTERN(TEXT("0.2s  - 0.3s  hitches"), STAT_FPSChart_Hitch_200_300, STATGROUP_FPSChart, );
-DECLARE_DWORD_ACCUMULATOR_STAT_EXTERN(TEXT("0.15s - 0.2s  hitches"), STAT_FPSChart_Hitch_150_200, STATGROUP_FPSChart, );
-DECLARE_DWORD_ACCUMULATOR_STAT_EXTERN(TEXT("0.1s  - 0.15s hitches"), STAT_FPSChart_Hitch_100_150, STATGROUP_FPSChart, );
-DECLARE_DWORD_ACCUMULATOR_STAT_EXTERN(TEXT("0.06s  - 0.1s hitches"), STAT_FPSChart_Hitch_60_100, STATGROUP_FPSChart, );
-DECLARE_DWORD_ACCUMULATOR_STAT_EXTERN(TEXT("Total hitches"), STAT_FPSChart_TotalHitchCount, STATGROUP_FPSChart, );
-
-/** Unit time stats */
-DECLARE_FLOAT_COUNTER_STAT_EXTERN(TEXT("StatUnit FrameTime"), STAT_FPSChart_UnitFrame, STATGROUP_FPSChart, );
-DECLARE_FLOAT_COUNTER_STAT_EXTERN(TEXT("StatUnit RenderThreadTime"), STAT_FPSChart_UnitRender, STATGROUP_FPSChart, );
-DECLARE_FLOAT_COUNTER_STAT_EXTERN(TEXT("StatUnit GameThreadTime"), STAT_FPSChart_UnitGame, STATGROUP_FPSChart, );
-DECLARE_FLOAT_COUNTER_STAT_EXTERN(TEXT("StatUnit GPUTime"), STAT_FPSChart_UnitGPU, STATGROUP_FPSChart, );
-
-
-/**
- * Path finding stats
- */
-DECLARE_CYCLE_STAT_EXTERN(TEXT("Async navmesh tile building"),STAT_Navigation_TileBuildAsync,STATGROUP_Navigation, );
-DECLARE_CYCLE_STAT_EXTERN(TEXT("Sync tiles building preparation"),STAT_Navigation_TileBuildPreparationSync,STATGROUP_Navigation, );
-DECLARE_CYCLE_STAT_EXTERN(TEXT("Processing world for building navmesh"),STAT_Navigation_ProcessingActorsForNavMeshBuilding,STATGROUP_Navigation, );
-DECLARE_CYCLE_STAT_EXTERN(TEXT("    Sync actors' geometry export"),STAT_Navigation_ActorsGeometryExportSync,STATGROUP_Navigation, );
-DECLARE_CYCLE_STAT_EXTERN(TEXT("    Sync BSP geometry export"),STAT_Navigation_BSPExportSync,STATGROUP_Navigation, );
-DECLARE_CYCLE_STAT_EXTERN(TEXT("    Sync gathering navigation modifiers"),STAT_Navigation_GatheringNavigationModifiersSync,STATGROUP_Navigation, );
-DECLARE_CYCLE_STAT_EXTERN(TEXT("    Async debug OBJ file export"),STAT_Navigation_TileGeometryExportToObjAsync,STATGROUP_Navigation, );
-DECLARE_CYCLE_STAT_EXTERN(TEXT("    Async voxel filtering"),STAT_Navigation_TileVoxelFilteringAsync,STATGROUP_Navigation, );
-DECLARE_CYCLE_STAT_EXTERN(TEXT("Sync pathfinding"),STAT_Navigation_PathfindingSync,STATGROUP_Navigation, );
-DECLARE_CYCLE_STAT_EXTERN(TEXT("Sync requests for async pathfinding"),STAT_Navigation_RequestingAsyncPathfinding,STATGROUP_Navigation, );
-DECLARE_CYCLE_STAT_EXTERN(TEXT("Async pathfinding"),STAT_Navigation_PathfindingAsync,STATGROUP_Navigation, );
-DECLARE_CYCLE_STAT_EXTERN(TEXT("Offset from corners"), STAT_Navigation_OffsetFromCorners, STATGROUP_Navigation, );
-DECLARE_CYCLE_STAT_EXTERN(TEXT("Visibility test for path optimisation"), STAT_Navigation_PathVisibilityOptimisation, STATGROUP_Navigation, );
-DECLARE_CYCLE_STAT_EXTERN(TEXT("Sync queries"),STAT_Navigation_QueriesTimeSync,STATGROUP_Navigation, );
-DECLARE_CYCLE_STAT_EXTERN(TEXT("    Sync meta nav area preparation"),STAT_Navigation_MetaAreaTranslation,STATGROUP_Navigation, ); 
-DECLARE_CYCLE_STAT_EXTERN(TEXT("    Async nav areas sorting"),STAT_Navigation_TileNavAreaSorting,STATGROUP_Navigation, );
-DECLARE_DWORD_COUNTER_STAT_EXTERN(TEXT("Num Destructible shapes exported"),STAT_Navigation_DestructiblesShapesExported,STATGROUP_Navigation, );
-DECLARE_DWORD_COUNTER_STAT_EXTERN(TEXT("Num update nav octree"),STAT_Navigation_UpdateNavOctree,STATGROUP_Navigation, );
-DECLARE_CYCLE_STAT_EXTERN(TEXT("Adding actors to navoctree"),STAT_Navigation_AddingActorsToNavOctree,STATGROUP_Navigation, );
-DECLARE_CYCLE_STAT_EXTERN(TEXT("    Adjusting nav links"),STAT_Navigation_AdjustingNavLinks,STATGROUP_Navigation, );
-DECLARE_CYCLE_STAT_EXTERN(TEXT("    Async AddTile"),STAT_Navigation_AddTile,STATGROUP_Navigation, );
-DECLARE_CYCLE_STAT_EXTERN(TEXT("Recast tick"),STAT_Navigation_RecastTick,STATGROUP_Navigation, );
-DECLARE_CYCLE_STAT_EXTERN(TEXT("Recast: build compressed layers"),STAT_Navigation_RecastBuildCompressedLayers,STATGROUP_Navigation, );
-DECLARE_CYCLE_STAT_EXTERN(TEXT("Recast: build navmesh"),STAT_Navigation_RecastBuildNavigation,STATGROUP_Navigation, );
-DECLARE_MEMORY_STAT_EXTERN(TEXT("Nav tree memory"),STAT_Navigation_CollisionTreeMemory,STATGROUP_Navigation, );
-DECLARE_MEMORY_STAT_EXTERN(TEXT("Nav data memory"),STAT_Navigation_NavDataMemory,STATGROUP_Navigation, );
-DECLARE_MEMORY_STAT_EXTERN(TEXT("Tile cache memory"),STAT_Navigation_TileCacheMemory,STATGROUP_Navigation, );
-DECLARE_DWORD_COUNTER_STAT_EXTERN(TEXT("Out of nodes path"),STAT_Navigation_OutOfNodesPath,STATGROUP_Navigation, );
-DECLARE_DWORD_COUNTER_STAT_EXTERN(TEXT("Partial path"),STAT_Navigation_PartialPath,STATGROUP_Navigation, );
-DECLARE_FLOAT_ACCUMULATOR_STAT_EXTERN(TEXT("Navmesh cumulative build Time"),STAT_Navigation_CumulativeBuildTime,STATGROUP_Navigation, );
-DECLARE_CYCLE_STAT_EXTERN(TEXT("Navmesh build time"),STAT_Navigation_BuildTime,STATGROUP_Navigation, );
-
-/**
- * UI Stats
- */
-DECLARE_CYCLE_STAT_EXTERN(TEXT("UI Drawing Time"),STAT_UIDrawingTime,STATGROUP_UI, );
-
-/**
-* Canvas Stats
-*/
-DECLARE_CYCLE_STAT_EXTERN(TEXT("Flush Time"),STAT_Canvas_FlushTime,STATGROUP_Canvas, );	
-DECLARE_CYCLE_STAT_EXTERN(TEXT("Draw Texture Tile Time"),STAT_Canvas_DrawTextureTileTime,STATGROUP_Canvas, );
-DECLARE_CYCLE_STAT_EXTERN(TEXT("Draw Material Tile Time"),STAT_Canvas_DrawMaterialTileTime,STATGROUP_Canvas, );
-DECLARE_CYCLE_STAT_EXTERN(TEXT("Draw String Time"),STAT_Canvas_DrawStringTime,STATGROUP_Canvas, );
-DECLARE_CYCLE_STAT_EXTERN(TEXT("Get Batched Element Time"),STAT_Canvas_GetBatchElementsTime,STATGROUP_Canvas, );
-DECLARE_CYCLE_STAT_EXTERN(TEXT("Add Material Tile Time"),STAT_Canvas_AddTileRenderTime,STATGROUP_Canvas, );
-DECLARE_DWORD_COUNTER_STAT_EXTERN(TEXT("Num Batches Created"),STAT_Canvas_NumBatchesCreated,STATGROUP_Canvas, );	
-
-/**
- * Landscape stats
- */
-DECLARE_CYCLE_STAT_EXTERN(TEXT("Landscape Dynamic Draw Time"),STAT_LandscapeDynamicDrawTime,STATGROUP_Landscape, );
-DECLARE_CYCLE_STAT_EXTERN(TEXT("Landscape Static Draw LOD Time"),STAT_LandscapeStaticDrawLODTime,STATGROUP_Landscape, );
-DECLARE_CYCLE_STAT_EXTERN(TEXT("LandscapeVF Draw Time"),STAT_LandscapeVFDrawTime,STATGROUP_Landscape, );
-DECLARE_DWORD_COUNTER_STAT_EXTERN(TEXT("Landscape Triangles"),STAT_LandscapeTriangles,STATGROUP_Landscape, );
-DECLARE_DWORD_COUNTER_STAT_EXTERN(TEXT("Landscape Render Passes"),STAT_LandscapeComponents,STATGROUP_Landscape, );
-DECLARE_DWORD_COUNTER_STAT_EXTERN(TEXT("Landscape DrawCalls"),STAT_LandscapeDrawCalls,STATGROUP_Landscape, );
-DECLARE_MEMORY_STAT_EXTERN(TEXT("Landscape Vertex Mem"),STAT_LandscapeVertexMem,STATGROUP_Landscape, );
-DECLARE_DWORD_COUNTER_STAT_EXTERN(TEXT("Component Mem"),STAT_LandscapeComponentMem,STATGROUP_Landscape, );
-
-/**
- * Network stats counters
- */
-
-DECLARE_DWORD_ACCUMULATOR_STAT_EXTERN(TEXT("Ping"),STAT_Ping,STATGROUP_Net, );
-DECLARE_DWORD_ACCUMULATOR_STAT_EXTERN(TEXT("Channels"),STAT_Channels,STATGROUP_Net, );
-DECLARE_DWORD_ACCUMULATOR_STAT_EXTERN(TEXT("In Rate (bytes)"),STAT_InRate,STATGROUP_Net, );
-DECLARE_DWORD_ACCUMULATOR_STAT_EXTERN(TEXT("Out Rate (bytes)"),STAT_OutRate,STATGROUP_Net, );
-DECLARE_DWORD_ACCUMULATOR_STAT_EXTERN(TEXT("In Packets"),STAT_InPackets,STATGROUP_Net, );
-DECLARE_DWORD_ACCUMULATOR_STAT_EXTERN(TEXT("Out Packets"),STAT_OutPackets,STATGROUP_Net, );
-DECLARE_DWORD_ACCUMULATOR_STAT_EXTERN(TEXT("In Bunches"),STAT_InBunches,STATGROUP_Net, );
-DECLARE_DWORD_ACCUMULATOR_STAT_EXTERN(TEXT("Out Bunches"),STAT_OutBunches,STATGROUP_Net, );
-DECLARE_DWORD_ACCUMULATOR_STAT_EXTERN(TEXT("Out Loss"),STAT_OutLoss,STATGROUP_Net, );
-DECLARE_DWORD_ACCUMULATOR_STAT_EXTERN(TEXT("In Loss"),STAT_InLoss,STATGROUP_Net, );
-DECLARE_DWORD_ACCUMULATOR_STAT_EXTERN(TEXT("Voice bytes sent"),STAT_VoiceBytesSent,STATGROUP_Net, );
-DECLARE_DWORD_ACCUMULATOR_STAT_EXTERN(TEXT("Voice bytes recv"),STAT_VoiceBytesRecv,STATGROUP_Net, );
-DECLARE_DWORD_ACCUMULATOR_STAT_EXTERN(TEXT("Voice packets sent"),STAT_VoicePacketsSent,STATGROUP_Net, );
-DECLARE_DWORD_ACCUMULATOR_STAT_EXTERN(TEXT("Voice packets recv"),STAT_VoicePacketsRecv,STATGROUP_Net, );
-DECLARE_DWORD_ACCUMULATOR_STAT_EXTERN(TEXT("In % Voice"),STAT_PercentInVoice,STATGROUP_Net, );
-DECLARE_DWORD_ACCUMULATOR_STAT_EXTERN(TEXT("Out % Voice"),STAT_PercentOutVoice,STATGROUP_Net, );
-DECLARE_DWORD_ACCUMULATOR_STAT_EXTERN(TEXT("Num Actor Channels"),STAT_NumActorChannels,STATGROUP_Net, );
-DECLARE_DWORD_ACCUMULATOR_STAT_EXTERN(TEXT("Num Considered Actors"),STAT_NumConsideredActors,STATGROUP_Net, );
-DECLARE_DWORD_ACCUMULATOR_STAT_EXTERN(TEXT("Num Prioritized Actors"),STAT_PrioritizedActors,STATGROUP_Net, );
-DECLARE_DWORD_ACCUMULATOR_STAT_EXTERN(TEXT("Num Relevant Actors"),STAT_NumRelevantActors,STATGROUP_Net, );
-DECLARE_DWORD_ACCUMULATOR_STAT_EXTERN(TEXT("Num Relevant Deleted Actors"),STAT_NumRelevantDeletedActors,STATGROUP_Net, );
-DECLARE_DWORD_ACCUMULATOR_STAT_EXTERN(TEXT("Num Replicated Actor Attempts"),STAT_NumReplicatedActorAttempts,STATGROUP_Net, );
-DECLARE_DWORD_ACCUMULATOR_STAT_EXTERN(TEXT("Num Replicated Actors Sent"),STAT_NumReplicatedActors,STATGROUP_Net, );
-DECLARE_DWORD_ACCUMULATOR_STAT_EXTERN(TEXT("Num Actors"),STAT_NumActors,STATGROUP_Net, );
-DECLARE_DWORD_ACCUMULATOR_STAT_EXTERN(TEXT("Num Network Actors"),STAT_NumNetActors,STATGROUP_Net, );
-DECLARE_DWORD_ACCUMULATOR_STAT_EXTERN(TEXT("Num Dormant Actors"),STAT_NumDormantActors,STATGROUP_Net, );
-DECLARE_DWORD_ACCUMULATOR_STAT_EXTERN(TEXT("Num Initially Dormant Actors"),STAT_NumInitiallyDormantActors,STATGROUP_Net, );
-DECLARE_DWORD_ACCUMULATOR_STAT_EXTERN(TEXT("Num Chan ready for dormancy"),STAT_NumActorChannelsReadyDormant,STATGROUP_Net, );
-DECLARE_DWORD_ACCUMULATOR_STAT_EXTERN(TEXT("Num ACKd NetGUIDs"),STAT_NumNetGUIDsAckd,STATGROUP_Net, );
-DECLARE_DWORD_ACCUMULATOR_STAT_EXTERN(TEXT("Num Pending NetGUIDs"),STAT_NumNetGUIDsPending,STATGROUP_Net, );
-DECLARE_DWORD_ACCUMULATOR_STAT_EXTERN(TEXT("Num UnACKd NetGUIDs"),STAT_NumNetGUIDsUnAckd,STATGROUP_Net, );
-DECLARE_DWORD_ACCUMULATOR_STAT_EXTERN(TEXT("Object path (bytes)"),STAT_ObjPathBytes,STATGROUP_Net, );
-DECLARE_DWORD_ACCUMULATOR_STAT_EXTERN(TEXT("NetGUID Out Rate (bytes)"),STAT_NetGUIDOutRate,STATGROUP_Net, );
-DECLARE_DWORD_ACCUMULATOR_STAT_EXTERN(TEXT("NetGUID In Rate (bytes)"),STAT_NetGUIDInRate,STATGROUP_Net, );
-DECLARE_DWORD_ACCUMULATOR_STAT_EXTERN(TEXT("Saturated"),STAT_NetSaturated,STATGROUP_Net, );
+#include "EngineStats.h"
 
 /*-----------------------------------------------------------------------------
 	Forward declarations.
@@ -587,7 +243,6 @@ namespace physx
 /*-----------------------------------------------------------------------------
 	Engine public includes.
 -----------------------------------------------------------------------------*/
-struct FURL;
 
 #include "BlueprintUtilities.h"
 #include "Tickable.h"						// FTickableGameObject interface.
diff --git a/Engine/Source/Runtime/Engine/Public/EngineStats.h b/Engine/Source/Runtime/Engine/Public/EngineStats.h
new file mode 100644
--- /dev/null
+++ b/Engine/Source/Runtime/Engine/Public/EngineStats.h
@@ -0,0 +1,217 @@
+// Copyright 1998-2014 Epic Games, Inc. All Rights Reserved.
+
+#pragma once
+
+
+/**
+ * Engine stats
+ */
+/** Input stat */
+DECLARE_CYCLE_STAT_EXTERN(TEXT("Input Time"),STAT_InputTime,STATGROUP_Engine, );
+DECLARE_FLOAT_COUNTER_STAT_EXTERN(TEXT("Input Latency"),STAT_InputLatencyTime,STATGROUP_Engine, );
+/** HUD stat */
+DECLARE_CYCLE_STAT_EXTERN(TEXT("HUD Time"),STAT_HudTime,STATGROUP_Engine, );
+/** Static mesh tris rendered */
+DECLARE_DWORD_COUNTER_STAT_EXTERN(TEXT("Static Mesh Tris"),STAT_StaticMeshTriangles,STATGROUP_Engine, ENGINE_API);
+DECLARE_CYCLE_STAT_EXTERN(TEXT("Skel Skin Time"),STAT_SkinningTime,STATGROUP_Engine, );
+DECLARE_CYCLE_STAT_EXTERN(TEXT("Update Cloth Verts Time"),STAT_UpdateClothVertsTime,STATGROUP_Engine, );
+DECLARE_CYCLE_STAT_EXTERN(TEXT("Update SoftBody Verts Time"),STAT_UpdateSoftBodyVertsTime,STATGROUP_Engine, );
+DECLARE_DWORD_COUNTER_STAT_EXTERN(TEXT("Skel Mesh Tris"),STAT_SkelMeshTriangles,STATGROUP_Engine, );
+DECLARE_DWORD_COUNTER_STAT_EXTERN(TEXT("Skel Mesh Draw Calls"),STAT_SkelMeshDrawCalls,STATGROUP_Engine, );
+DECLARE_DWORD_COUNTER_STAT_EXTERN(TEXT("Skel Verts CPU Skin"),STAT_CPUSkinVertices,STATGROUP_Engine, );
+DECLARE_DWORD_COUNTER_STAT_EXTERN(TEXT("Skel Verts GPU Skin"),STAT_GPUSkinVertices,STATGROUP_Engine, );
+DECLARE_CYCLE_STAT_EXTERN(TEXT("GameEngine Tick"),STAT_GameEngineTick,STATGROUP_Engine, );
+DECLARE_CYCLE_STAT_EXTERN(TEXT("GameViewport Tick"),STAT_GameViewportTick,STATGROUP_Engine, );
+DECLARE_CYCLE_STAT_EXTERN(TEXT("RedrawViewports"),STAT_RedrawViewports,STATGROUP_Engine, );
+DECLARE_CYCLE_STAT_EXTERN(TEXT("Update Level Streaming"),STAT_UpdateLevelStreaming,STATGROUP_Engine, ENGINE_API);
+DECLARE_CYCLE_STAT_EXTERN(TEXT("RHI Game Tick"),STAT_RHITickTime,STATGROUP_Engine, ENGINE_API);
+DECLARE_CYCLE_STAT_EXTERN(TEXT("Debug Hitch"),STAT_IntentionalHitch,STATGROUP_Engine, );
+DECLARE_CYCLE_STAT_EXTERN(TEXT("Platform Message Time"),STAT_PlatformMessageTime,STATGROUP_Engine, ENGINE_API);
+DECLARE_CYCLE_STAT_EXTERN(TEXT("Frame Sync Time"),STAT_FrameSyncTime,STATGROUP_Engine, ENGINE_API);
+DECLARE_CYCLE_STAT_EXTERN(TEXT("Deferred Tick Time"),STAT_DeferredTickTime,STATGROUP_Engine, ENGINE_API);
+
+/**
+ * Game stats
+ */
+DECLARE_CYCLE_STAT_EXTERN(TEXT("Async Physics Time"),STAT_PhysicsTime,STATGROUP_Game, );
+DECLARE_CYCLE_STAT_EXTERN(TEXT("Spawn Actor Time"),STAT_SpawnActorTime,STATGROUP_Game, );
+DECLARE_CYCLE_STAT_EXTERN(TEXT("MoveComponent Time"),STAT_MoveComponentTime,STATGROUP_Game, );
+DECLARE_CYCLE_STAT_EXTERN(TEXT("TeleportTo Time"),STAT_TeleportToTime,STATGROUP_Game, );
+DECLARE_CYCLE_STAT_EXTERN(TEXT("GC Mark Time"),STAT_GCMarkTime,STATGROUP_Game, );
+DECLARE_CYCLE_STAT_EXTERN(TEXT("GC Sweep Time"),STAT_GCSweepTime,STATGROUP_Game, );
+DECLARE_CYCLE_STAT_EXTERN(TEXT("Tick Time"),STAT_TickTime,STATGROUP_Game, );
+DECLARE_CYCLE_STAT_EXTERN(TEXT("World Tick Time"),STAT_WorldTickTime,STATGROUP_Game, );
+DECLARE_CYCLE_STAT_EXTERN(TEXT("Update Camera Time"),STAT_UpdateCameraTime,STATGROUP_Game, );
+DECLARE_CYCLE_STAT_EXTERN(TEXT("Post Tick Component Update"),STAT_PostTickComponentUpdate,STATGROUP_Game, );
+DECLARE_CYCLE_STAT_EXTERN(TEXT("Post Tick Component Wait"),STAT_PostTickComponentUpdateWait,STATGROUP_Game, );
+DECLARE_CYCLE_STAT_EXTERN(TEXT("  Recreate"),STAT_PostTickComponentRecreate,STATGROUP_Game, );
+DECLARE_CYCLE_STAT_EXTERN(TEXT("  Transform or RenderData"),STAT_PostTickComponentLW,STATGROUP_Game, );
+DECLARE_CYCLE_STAT_EXTERN(TEXT("Update Particle Data"),STAT_ParticleManagerUpdateData,STATGROUP_Game, );
+DECLARE_CYCLE_STAT_EXTERN(TEXT("Async Work Wait"),STAT_AsyncWorkWaitTime,STATGROUP_Game, );
+DECLARE_CYCLE_STAT_EXTERN(TEXT("Char Movement Authority Time"),STAT_CharacterMovementAuthority,STATGROUP_Game, );
+DECLARE_CYCLE_STAT_EXTERN(TEXT("Char Movement Simulated Time"),STAT_CharacterMovementSimulated,STATGROUP_Game, );
+DECLARE_CYCLE_STAT_EXTERN(TEXT("Visual Logging time"), STAT_VisualLog, STATGROUP_Game, ENGINE_API);
+DECLARE_CYCLE_STAT_EXTERN(TEXT("Net Tick Time"),STAT_NetWorldTickTime,STATGROUP_Game, );
+DECLARE_CYCLE_STAT_EXTERN(TEXT("Nav Tick Time"),STAT_NavWorldTickTime,STATGROUP_Game, );
+DECLARE_CYCLE_STAT_EXTERN(TEXT("Reset Async Trace Time"),STAT_ResetAsyncTraceTickTime,STATGROUP_Game, );
+DECLARE_CYCLE_STAT_EXTERN(TEXT("GT Tickable Time"),STAT_TickableTickTime,STATGROUP_Game, );
+DECLARE_CYCLE_STAT_EXTERN(TEXT("Runtime Movie Tick Time"),STAT_RuntimeMovieSceneTickTime,STATGROUP_Game, );
+DECLARE_CYCLE_STAT_EXTERN(TEXT("Finish Async Trace Time"),STAT_FinishAsyncTraceTickTime,STATGROUP_Game, );
+DECLARE_CYCLE_STAT_EXTERN(TEXT("Net Broadcast Tick Time"),STAT_NetBroadcastTickTime,STATGROUP_Game, );
+DECLARE_CYCLE_STAT_EXTERN(TEXT("  ServerReplicateActors Time"),STAT_NetServerRepActorsTime,STATGROUP_Game, );
+DECLARE_CYCLE_STAT_EXTERN(TEXT("  Consider Actors Time"),STAT_NetConsiderActorsTime,STATGROUP_Game, );
+DECLARE_CYCLE_STAT_EXTERN(TEXT("  Inital Dormant Time"),STAT_NetInitialDormantCheckTime,STATGROUP_Game, );
+DECLARE_CYCLE_STAT_EXTERN(TEXT("  Prioritize Actors Time"),STAT_NetPrioritizeActorsTime,STATGROUP_Game, );
+DECLARE_CYCLE_STAT_EXTERN(TEXT("  Replicate Actors Time"),STAT_NetReplicateActorsTime,STATGROUP_Game, );
+DECLARE_CYCLE_STAT_EXTERN(TEXT("  Dynamic Property Rep Time"),STAT_NetReplicateDynamicPropTime,STATGROUP_Game, );
+DECLARE_DWORD_COUNTER_STAT_EXTERN(TEXT("  Skipped Dynamic Props"),STAT_NetSkippedDynamicProps,STATGROUP_Game, );
+DECLARE_CYCLE_STAT_EXTERN(TEXT("  NetSerializeItemDelta Time"),STAT_NetSerializeItemDeltaTime,STATGROUP_Game, );
+DECLARE_CYCLE_STAT_EXTERN(TEXT("  Static Property Rep Time"),STAT_NetReplicateStaticPropTime,STATGROUP_Game, );
+DECLARE_CYCLE_STAT_EXTERN(TEXT("  Rebuild Conditionals"),STAT_NetRebuildConditionalTime,STATGROUP_Game, );
+DECLARE_CYCLE_STAT_EXTERN(TEXT("Net Post BC Tick Time"),STAT_NetBroadcastPostTickTime,STATGROUP_Game, );
+DECLARE_CYCLE_STAT_EXTERN(TEXT("Net PackageMap SerializeObject"),STAT_PackageMap_SerializeObjectTime,STATGROUP_Game, );
+
+/**
+ * FPS chart stats
+ */
+DECLARE_FLOAT_ACCUMULATOR_STAT_EXTERN(TEXT("00 - 05 FPS"), STAT_FPSChart_0_5,		STATGROUP_FPSChart, );
+DECLARE_FLOAT_ACCUMULATOR_STAT_EXTERN(TEXT("05 - 10 FPS"), STAT_FPSChart_5_10,		STATGROUP_FPSChart, );
+DECLARE_FLOAT_ACCUMULATOR_STAT_EXTERN(TEXT("10 - 15 FPS"), STAT_FPSChart_10_15,	STATGROUP_FPSChart, );
+DECLARE_FLOAT_ACCUMULATOR_STAT_EXTERN(TEXT("15 - 20 FPS"), STAT_FPSChart_15_20,	STATGROUP_FPSChart, );
+DECLARE_FLOAT_ACCUMULATOR_STAT_EXTERN(TEXT("20 - 25 FPS"), STAT_FPSChart_20_25,	STATGROUP_FPSChart, );
+DECLARE_FLOAT_ACCUMULATOR_STAT_EXTERN(TEXT("25 - 30 FPS"), STAT_FPSChart_25_30,	STATGROUP_FPSChart, );
+DECLARE_FLOAT_ACCUMULATOR_STAT_EXTERN(TEXT("30 - 35 FPS"), STAT_FPSChart_30_35,	STATGROUP_FPSChart, );
+DECLARE_FLOAT_ACCUMULATOR_STAT_EXTERN(TEXT("35 - 40 FPS"), STAT_FPSChart_35_40,	STATGROUP_FPSChart, );
+DECLARE_FLOAT_ACCUMULATOR_STAT_EXTERN(TEXT("40 - 45 FPS"), STAT_FPSChart_40_45,	STATGROUP_FPSChart, );
+DECLARE_FLOAT_ACCUMULATOR_STAT_EXTERN(TEXT("45 - 50 FPS"), STAT_FPSChart_45_50,	STATGROUP_FPSChart, );
+DECLARE_FLOAT_ACCUMULATOR_STAT_EXTERN(TEXT("50 - 55 FPS"), STAT_FPSChart_50_55,	STATGROUP_FPSChart, );
+DECLARE_FLOAT_ACCUMULATOR_STAT_EXTERN(TEXT("55 - 60 FPS"), STAT_FPSChart_55_60,	STATGROUP_FPSChart, );
+DECLARE_FLOAT_ACCUMULATOR_STAT_EXTERN(TEXT("60 - .. FPS"), STAT_FPSChart_60_INF,	STATGROUP_FPSChart, );
+DECLARE_FLOAT_ACCUMULATOR_STAT_EXTERN(TEXT("30+ FPS"), STAT_FPSChart_30Plus,	STATGROUP_FPSChart, );
+DECLARE_FLOAT_ACCUMULATOR_STAT_EXTERN(TEXT("Unaccounted time"), STAT_FPSChart_UnaccountedTime,	STATGROUP_FPSChart, );
+DECLARE_DWORD_ACCUMULATOR_STAT_EXTERN(TEXT("Frame count"),	STAT_FPSChart_FrameCount, STATGROUP_FPSChart, );
+
+/** Hitch stats */
+DECLARE_DWORD_ACCUMULATOR_STAT_EXTERN(TEXT("5.0s  - ....  hitches"), STAT_FPSChart_Hitch_5000_Plus, STATGROUP_FPSChart, );
+DECLARE_DWORD_ACCUMULATOR_STAT_EXTERN(TEXT("2.5s  - 5.0s  hitches"), STAT_FPSChart_Hitch_2500_5000, STATGROUP_FPSChart, );
+DECLARE_DWORD_ACCUMULATOR_STAT_EXTERN(TEXT("2.0s  - 2.5s  hitches"), STAT_FPSChart_Hitch_2000_2500, STATGROUP_FPSChart, );
+DECLARE_DWORD_ACCUMULATOR_STAT_EXTERN(TEXT("1.5s  - 2.0s  hitches"), STAT_FPSChart_Hitch_1500_2000, STATGROUP_FPSChart, );
+DECLARE_DWORD_ACCUMULATOR_STAT_EXTERN(TEXT("1.0s  - 1.5s  hitches"), STAT_FPSChart_Hitch_1000_1500, STATGROUP_FPSChart, );
+DECLARE_DWORD_ACCUMULATOR_STAT_EXTERN(TEXT("0.75s - 1.0s  hitches"), STAT_FPSChart_Hitch_750_1000, STATGROUP_FPSChart, );
+DECLARE_DWORD_ACCUMULATOR_STAT_EXTERN(TEXT("0.5s  - 0.75s hitches"), STAT_FPSChart_Hitch_500_750, STATGROUP_FPSChart, );
+DECLARE_DWORD_ACCUMULATOR_STAT_EXTERN(TEXT("0.3s  - 0.5s  hitches"), STAT_FPSChart_Hitch_300_500, STATGROUP_FPSChart, );
+DECLARE_DWORD_ACCUMULATOR_STAT_EXTERN(TEXT("0.2s  - 0.3s  hitches"), STAT_FPSChart_Hitch_200_300, STATGROUP_FPSChart, );
+DECLARE_DWORD_ACCUMULATOR_STAT_EXTERN(TEXT("0.15s - 0.2s  hitches"), STAT_FPSChart_Hitch_150_200, STATGROUP_FPSChart, );
+DECLARE_DWORD_ACCUMULATOR_STAT_EXTERN(TEXT("0.1s  - 0.15s hitches"), STAT_FPSChart_Hitch_100_150, STATGROUP_FPSChart, );
+DECLARE_DWORD_ACCUMULATOR_STAT_EXTERN(TEXT("0.06s  - 0.1s hitches"), STAT_FPSChart_Hitch_60_100, STATGROUP_FPSChart, );
+DECLARE_DWORD_ACCUMULATOR_STAT_EXTERN(TEXT("Total hitches"), STAT_FPSChart_TotalHitchCount, STATGROUP_FPSChart, );
+
+/** Unit time stats */
+DECLARE_FLOAT_COUNTER_STAT_EXTERN(TEXT("StatUnit FrameTime"), STAT_FPSChart_UnitFrame, STATGROUP_FPSChart, );
+DECLARE_FLOAT_COUNTER_STAT_EXTERN(TEXT("StatUnit RenderThreadTime"), STAT_FPSChart_UnitRender, STATGROUP_FPSChart, );
+DECLARE_FLOAT_COUNTER_STAT_EXTERN(TEXT("StatUnit GameThreadTime"), STAT_FPSChart_UnitGame, STATGROUP_FPSChart, );
+DECLARE_FLOAT_COUNTER_STAT_EXTERN(TEXT("StatUnit GPUTime"), STAT_FPSChart_UnitGPU, STATGROUP_FPSChart, );
+
+
+/**
+ * Path finding stats
+ */
+DECLARE_CYCLE_STAT_EXTERN(TEXT("Async navmesh tile building"),STAT_Navigation_TileBuildAsync,STATGROUP_Navigation, );
+DECLARE_CYCLE_STAT_EXTERN(TEXT("Sync tiles building preparation"),STAT_Navigation_TileBuildPreparationSync,STATGROUP_Navigation, );
+DECLARE_CYCLE_STAT_EXTERN(TEXT("Processing world for building navmesh"),STAT_Navigation_ProcessingActorsForNavMeshBuilding,STATGROUP_Navigation, );
+DECLARE_CYCLE_STAT_EXTERN(TEXT("    Sync actors' geometry export"),STAT_Navigation_ActorsGeometryExportSync,STATGROUP_Navigation, );
+DECLARE_CYCLE_STAT_EXTERN(TEXT("    Sync BSP geometry export"),STAT_Navigation_BSPExportSync,STATGROUP_Navigation, );
+DECLARE_CYCLE_STAT_EXTERN(TEXT("    Sync gathering navigation modifiers"),STAT_Navigation_GatheringNavigationModifiersSync,STATGROUP_Navigation, );
+DECLARE_CYCLE_STAT_EXTERN(TEXT("    Async debug OBJ file export"),STAT_Navigation_TileGeometryExportToObjAsync,STATGROUP_Navigation, );
+DECLARE_CYCLE_STAT_EXTERN(TEXT("    Async voxel filtering"),STAT_Navigation_TileVoxelFilteringAsync,STATGROUP_Navigation, );
+DECLARE_CYCLE_STAT_EXTERN(TEXT("Sync pathfinding"),STAT_Navigation_PathfindingSync,STATGROUP_Navigation, );
+DECLARE_CYCLE_STAT_EXTERN(TEXT("Sync requests for async pathfinding"),STAT_Navigation_RequestingAsyncPathfinding,STATGROUP_Navigation, );
+DECLARE_CYCLE_STAT_EXTERN(TEXT("Async pathfinding"),STAT_Navigation_PathfindingAsync,STATGROUP_Navigation, );
+DECLARE_CYCLE_STAT_EXTERN(TEXT("Offset from corners"), STAT_Navigation_OffsetFromCorners, STATGROUP_Navigation, );
+DECLARE_CYCLE_STAT_EXTERN(TEXT("Visibility test for path optimisation"), STAT_Navigation_PathVisibilityOptimisation, STATGROUP_Navigation, );
+DECLARE_CYCLE_STAT_EXTERN(TEXT("Sync queries"),STAT_Navigation_QueriesTimeSync,STATGROUP_Navigation, );
+DECLARE_CYCLE_STAT_EXTERN(TEXT("    Sync meta nav area preparation"),STAT_Navigation_MetaAreaTranslation,STATGROUP_Navigation, ); 
+DECLARE_CYCLE_STAT_EXTERN(TEXT("    Async nav areas sorting"),STAT_Navigation_TileNavAreaSorting,STATGROUP_Navigation, );
+DECLARE_DWORD_COUNTER_STAT_EXTERN(TEXT("Num Destructible shapes exported"),STAT_Navigation_DestructiblesShapesExported,STATGROUP_Navigation, );
+DECLARE_DWORD_COUNTER_STAT_EXTERN(TEXT("Num update nav octree"),STAT_Navigation_UpdateNavOctree,STATGROUP_Navigation, );
+DECLARE_CYCLE_STAT_EXTERN(TEXT("Adding actors to navoctree"),STAT_Navigation_AddingActorsToNavOctree,STATGROUP_Navigation, );
+DECLARE_CYCLE_STAT_EXTERN(TEXT("    Adjusting nav links"),STAT_Navigation_AdjustingNavLinks,STATGROUP_Navigation, );
+DECLARE_CYCLE_STAT_EXTERN(TEXT("    Async AddTile"),STAT_Navigation_AddTile,STATGROUP_Navigation, );
+DECLARE_CYCLE_STAT_EXTERN(TEXT("Recast tick"),STAT_Navigation_RecastTick,STATGROUP_Navigation, );
+DECLARE_CYCLE_STAT_EXTERN(TEXT("Recast: build compressed layers"),STAT_Navigation_RecastBuildCompressedLayers,STATGROUP_Navigation, );
+DECLARE_CYCLE_STAT_EXTERN(TEXT("Recast: build navmesh"),STAT_Navigation_RecastBuildNavigation,STATGROUP_Navigation, );
+DECLARE_MEMORY_STAT_EXTERN(TEXT("Nav tree memory"),STAT_Navigation_CollisionTreeMemory,STATGROUP_Navigation, );
+DECLARE_MEMORY_STAT_EXTERN(TEXT("Nav data memory"),STAT_Navigation_NavDataMemory,STATGROUP_Navigation, );
+DECLARE_MEMORY_STAT_EXTERN(TEXT("Tile cache memory"),STAT_Navigation_TileCacheMemory,STATGROUP_Navigation, );
+DECLARE_DWORD_COUNTER_STAT_EXTERN(TEXT("Out of nodes path"),STAT_Navigation_OutOfNodesPath,STATGROUP_Navigation, );
+DECLARE_DWORD_COUNTER_STAT_EXTERN(TEXT("Partial path"),STAT_Navigation_PartialPath,STATGROUP_Navigation, );
+DECLARE_FLOAT_ACCUMULATOR_STAT_EXTERN(TEXT("Navmesh cumulative build Time"),STAT_Navigation_CumulativeBuildTime,STATGROUP_Navigation, );
+DECLARE_CYCLE_STAT_EXTERN(TEXT("Navmesh build time"),STAT_Navigation_BuildTime,STATGROUP_Navigation, );
+
+/**
+ * UI Stats
+ */
+DECLARE_CYCLE_STAT_EXTERN(TEXT("UI Drawing Time"),STAT_UIDrawingTime,STATGROUP_UI, );
+
+/**
+* Canvas Stats
+*/
+DECLARE_CYCLE_STAT_EXTERN(TEXT("Flush Time"),STAT_Canvas_FlushTime,STATGROUP_Canvas, );	
+DECLARE_CYCLE_STAT_EXTERN(TEXT("Draw Texture Tile Time"),STAT_Canvas_DrawTextureTileTime,STATGROUP_Canvas, );
+DECLARE_CYCLE_STAT_EXTERN(TEXT("Draw Material Tile Time"),STAT_Canvas_DrawMaterialTileTime,STATGROUP_Canvas, );
+DECLARE_CYCLE_STAT_EXTERN(TEXT("Draw String Time"),STAT_Canvas_DrawStringTime,STATGROUP_Canvas, );
+DECLARE_CYCLE_STAT_EXTERN(TEXT("Get Batched Element Time"),STAT_Canvas_GetBatchElementsTime,STATGROUP_Canvas, );
+DECLARE_CYCLE_STAT_EXTERN(TEXT("Add Material Tile Time"),STAT_Canvas_AddTileRenderTime,STATGROUP_Canvas, );
+DECLARE_DWORD_COUNTER_STAT_EXTERN(TEXT("Num Batches Created"),STAT_Canvas_NumBatchesCreated,STATGROUP_Canvas, );	
+
+/**
+ * Landscape stats
+ */
+DECLARE_CYCLE_STAT_EXTERN(TEXT("Landscape Dynamic Draw Time"),STAT_LandscapeDynamicDrawTime,STATGROUP_Landscape, );
+DECLARE_CYCLE_STAT_EXTERN(TEXT("Landscape Static Draw LOD Time"),STAT_LandscapeStaticDrawLODTime,STATGROUP_Landscape, );
+DECLARE_CYCLE_STAT_EXTERN(TEXT("LandscapeVF Draw Time"),STAT_LandscapeVFDrawTime,STATGROUP_Landscape, );
+DECLARE_DWORD_COUNTER_STAT_EXTERN(TEXT("Landscape Triangles"),STAT_LandscapeTriangles,STATGROUP_Landscape, );
+DECLARE_DWORD_COUNTER_STAT_EXTERN(TEXT("Landscape Render Passes"),STAT_LandscapeComponents,STATGROUP_Landscape, );
+DECLARE_DWORD_COUNTER_STAT_EXTERN(TEXT("Landscape DrawCalls"),STAT_LandscapeDrawCalls,STATGROUP_Landscape, );
+DECLARE_MEMORY_STAT_EXTERN(TEXT("Landscape Vertex Mem"),STAT_LandscapeVertexMem,STATGROUP_Landscape, );
+DECLARE_DWORD_COUNTER_STAT_EXTERN(TEXT("Component Mem"),STAT_LandscapeComponentMem,STATGROUP_Landscape, );
+
+/**
+ * Network stats counters
+ */
+
+DECLARE_DWORD_ACCUMULATOR_STAT_EXTERN(TEXT("Ping"),STAT_Ping,STATGROUP_Net, );
+DECLARE_DWORD_ACCUMULATOR_STAT_EXTERN(TEXT("Channels"),STAT_Channels,STATGROUP_Net, );
+DECLARE_DWORD_ACCUMULATOR_STAT_EXTERN(TEXT("In Rate (bytes)"),STAT_InRate,STATGROUP_Net, );
+DECLARE_DWORD_ACCUMULATOR_STAT_EXTERN(TEXT("Out Rate (bytes)"),STAT_OutRate,STATGROUP_Net, );
+DECLARE_DWORD_ACCUMULATOR_STAT_EXTERN(TEXT("In Packets"),STAT_InPackets,STATGROUP_Net, );
+DECLARE_DWORD_ACCUMULATOR_STAT_EXTERN(TEXT("Out Packets"),STAT_OutPackets,STATGROUP_Net, );
+DECLARE_DWORD_ACCUMULATOR_STAT_EXTERN(TEXT("In Bunches"),STAT_InBunches,STATGROUP_Net, );
+DECLARE_DWORD_ACCUMULATOR_STAT_EXTERN(TEXT("Out Bunches"),STAT_OutBunches,STATGROUP_Net, );
+DECLARE_DWORD_ACCUMULATOR_STAT_EXTERN(TEXT("Out Loss"),STAT_OutLoss,STATGROUP_Net, );
+DECLARE_DWORD_ACCUMULATOR_STAT_EXTERN(TEXT("In Loss"),STAT_InLoss,STATGROUP_Net, );
+DECLARE_DWORD_ACCUMULATOR_STAT_EXTERN(TEXT("Voice bytes sent"),STAT_VoiceBytesSent,STATGROUP_Net, );
+DECLARE_DWORD_ACCUMULATOR_STAT_EXTERN(TEXT("Voice bytes recv"),STAT_VoiceBytesRecv,STATGROUP_Net, );
+DECLARE_DWORD_ACCUMULATOR_STAT_EXTERN(TEXT("Voice packets sent"),STAT_VoicePacketsSent,STATGROUP_Net, );
+DECLARE_DWORD_ACCUMULATOR_STAT_EXTERN(TEXT("Voice packets recv"),STAT_VoicePacketsRecv,STATGROUP_Net, );
+DECLARE_DWORD_ACCUMULATOR_STAT_EXTERN(TEXT("In % Voice"),STAT_PercentInVoice,STATGROUP_Net, );
+DECLARE_DWORD_ACCUMULATOR_STAT_EXTERN(TEXT("Out % Voice"),STAT_PercentOutVoice,STATGROUP_Net, );
+DECLARE_DWORD_ACCUMULATOR_STAT_EXTERN(TEXT("Num Actor Channels"),STAT_NumActorChannels,STATGROUP_Net, );
+DECLARE_DWORD_ACCUMULATOR_STAT_EXTERN(TEXT("Num Considered Actors"),STAT_NumConsideredActors,STATGROUP_Net, );
+DECLARE_DWORD_ACCUMULATOR_STAT_EXTERN(TEXT("Num Prioritized Actors"),STAT_PrioritizedActors,STATGROUP_Net, );
+DECLARE_DWORD_ACCUMULATOR_STAT_EXTERN(TEXT("Num Relevant Actors"),STAT_NumRelevantActors,STATGROUP_Net, );
+DECLARE_DWORD_ACCUMULATOR_STAT_EXTERN(TEXT("Num Relevant Deleted Actors"),STAT_NumRelevantDeletedActors,STATGROUP_Net, );
+DECLARE_DWORD_ACCUMULATOR_STAT_EXTERN(TEXT("Num Replicated Actor Attempts"),STAT_NumReplicatedActorAttempts,STATGROUP_Net, );
+DECLARE_DWORD_ACCUMULATOR_STAT_EXTERN(TEXT("Num Replicated Actors Sent"),STAT_NumReplicatedActors,STATGROUP_Net, );
+DECLARE_DWORD_ACCUMULATOR_STAT_EXTERN(TEXT("Num Actors"),STAT_NumActors,STATGROUP_Net, );
+DECLARE_DWORD_ACCUMULATOR_STAT_EXTERN(TEXT("Num Network Actors"),STAT_NumNetActors,STATGROUP_Net, );
+DECLARE_DWORD_ACCUMULATOR_STAT_EXTERN(TEXT("Num Dormant Actors"),STAT_NumDormantActors,STATGROUP_Net, );
+DECLARE_DWORD_ACCUMULATOR_STAT_EXTERN(TEXT("Num Initially Dormant Actors"),STAT_NumInitiallyDormantActors,STATGROUP_Net, );
+DECLARE_DWORD_ACCUMULATOR_STAT_EXTERN(TEXT("Num Chan ready for dormancy"),STAT_NumActorChannelsReadyDormant,STATGROUP_Net, );
+DECLARE_DWORD_ACCUMULATOR_STAT_EXTERN(TEXT("Num ACKd NetGUIDs"),STAT_NumNetGUIDsAckd,STATGROUP_Net, );
+DECLARE_DWORD_ACCUMULATOR_STAT_EXTERN(TEXT("Num Pending NetGUIDs"),STAT_NumNetGUIDsPending,STATGROUP_Net, );
+DECLARE_DWORD_ACCUMULATOR_STAT_EXTERN(TEXT("Num UnACKd NetGUIDs"),STAT_NumNetGUIDsUnAckd,STATGROUP_Net, );
+DECLARE_DWORD_ACCUMULATOR_STAT_EXTERN(TEXT("Object path (bytes)"),STAT_ObjPathBytes,STATGROUP_Net, );
+DECLARE_DWORD_ACCUMULATOR_STAT_EXTERN(TEXT("NetGUID Out Rate (bytes)"),STAT_NetGUIDOutRate,STATGROUP_Net, );
+DECLARE_DWORD_ACCUMULATOR_STAT_EXTERN(TEXT("NetGUID In Rate (bytes)"),STAT_NetGUIDInRate,STATGROUP_Net, );
+DECLARE_DWORD_ACCUMULATOR_STAT_EXTERN(TEXT("Saturated"),STAT_NetSaturated,STATGROUP_Net, );
\ No newline at end of file
diff --git a/Engine/Source/Runtime/Engine/Public/SkeletalMeshTypes.h b/Engine/Source/Runtime/Engine/Public/SkeletalMeshTypes.h
--- a/Engine/Source/Runtime/Engine/Public/SkeletalMeshTypes.h
+++ b/Engine/Source/Runtime/Engine/Public/SkeletalMeshTypes.h
@@ -18,6 +18,8 @@
 // Define that controls showing chart of distance factors for skel meshes during entire run of the game on exit.
 #define CHART_DISTANCE_FACTORS 0
 
+typedef uint16 FBoneIndexType;
+
 class FRawStaticIndexBuffer16or32Interface;
 
 /** 
