Ensure that networked PIE packages are correctly flagged as a PIE package
TTP #335632

diff --git a/Engine/Source/Runtime/Engine/Private/UnrealEngine.cpp b/Engine/Source/Runtime/Engine/Private/UnrealEngine.cpp
--- a/Engine/Source/Runtime/Engine/Private/UnrealEngine.cpp
+++ b/Engine/Source/Runtime/Engine/Private/UnrealEngine.cpp
@@ -8195,6 +8195,10 @@ bool UEngine::LoadMap( FWorldContext& WorldContext, FURL URL, class UPendingNetG
 		// make the package, and use this for the new linker (and to load the map from)
 		MapOuter = CreatePackage(NULL, *Pending->URL.Map);
 #if WITH_EDITOR
+		if (WorldContext.WorldType == EWorldType::PIE)
+		{
+			MapOuter->PackageFlags |= PKG_PlayInEditor;
+		}
 		MapOuter->PIEInstanceID = WorldContext.PIEInstance;
 #endif
 		// create the linker with the map name, and use the Guid so we find the downloaded version
