Expose friction combination mode

diff --git a/Engine/Source/Runtime/Engine/Classes/PhysicalMaterials/PhysicalMaterial.h b/Engine/Source/Runtime/Engine/Classes/PhysicalMaterials/PhysicalMaterial.h
--- a/Engine/Source/Runtime/Engine/Classes/PhysicalMaterials/PhysicalMaterial.h
+++ b/Engine/Source/Runtime/Engine/Classes/PhysicalMaterials/PhysicalMaterial.h
@@ -2,6 +2,7 @@
 
 #pragma once
 
+#include "PhysicsEngine/PhysicsSettings.h"
 #include "PhysicalMaterial.generated.h"
 
 namespace physx
@@ -46,6 +47,14 @@ class ENGINE_API UPhysicalMaterial : public UObject
 	UPROPERTY(EditAnywhere, BlueprintReadWrite, Category=PhysicalMaterial)
 	float Friction;
 
+	/** Friction combine mode, controls how friction is computed for multiple materials. */
+	UPROPERTY(EditAnywhere, BlueprintReadWrite, Category = PhysicalMaterial, meta = (editcondition = "bOverrideFrictionCombineMode"))
+	TEnumAsByte<EFrictionCombineMode::Type> FrictionCombineMode;
+
+	/** If set we will use the FrictionCombineMode of this material, instead of the FrictionCombineMode found in the project settings. */
+	UPROPERTY(EditAnywhere, BlueprintReadWrite, Category = PhysicalMaterial)
+	bool bOverrideFrictionCombineMode;
+
 	/** Resitution or 'bouncyness' of this surface */
 	UPROPERTY(EditAnywhere, BlueprintReadWrite, Category=PhysicalMaterial)
 	float Restitution;
@@ -98,6 +107,7 @@ public:
 	// Begin UObject interface
 #if WITH_EDITOR
 	virtual void PostEditChangeProperty(FPropertyChangedEvent& PropertyChangedEvent) OVERRIDE;
+	static void RebuildPhysicalMaterials();
 #endif // WITH_EDITOR
 	virtual void PostLoad() OVERRIDE;
 	virtual void FinishDestroy() OVERRIDE;
diff --git a/Engine/Source/Runtime/Engine/Classes/PhysicsEngine/PhysicsSettings.h b/Engine/Source/Runtime/Engine/Classes/PhysicsEngine/PhysicsSettings.h
--- a/Engine/Source/Runtime/Engine/Classes/PhysicsEngine/PhysicsSettings.h
+++ b/Engine/Source/Runtime/Engine/Classes/PhysicsEngine/PhysicsSettings.h
@@ -32,6 +32,21 @@ struct FPhysicalSurfaceName
 	{}
 };
 
+UENUM()
+namespace EFrictionCombineMode
+{
+	enum Type
+	{
+		//Uses the average friction of materials touching: (a+b) /2
+		Average = 0,	
+		//Uses the minimum friction of materials touching: min(a,b)
+		Min = 1,		
+		//Uses the product of the friction of materials touching: a*b
+		Multiply = 2,	
+		//Uses the maximum friction of materials touching: max(a,b)
+		Max = 3
+	};
+}
 
 /**
  * Implements project settings for the physics sub-system.
@@ -50,6 +65,10 @@ class ENGINE_API UPhysicsSettings
 	UPROPERTY(config, EditAnywhere, AdvancedDisplay, Category=Simulation)
 	bool bEnableAsyncScene;
 
+	/** Friction combine mode, controls how friction is computed for multiple materials. */
+	UPROPERTY(config, EditAnywhere, Category=Simulation)
+	TEnumAsByte<EFrictionCombineMode::Type> FrictionCombineMode;
+
 	/**Max Physics Delta Time to be clamped. */
 	UPROPERTY(config, EditAnywhere, meta=(ClampMin="0.0013", UIMin = "0.0013", ClampMax="1.0", UIMax="1.0"), Category=Framerate)
 	float MaxPhysicsDeltaTime;
@@ -91,6 +110,7 @@ public:
 #if WITH_EDITOR
 
 	virtual bool CanEditChange( const UProperty* Property ) const OVERRIDE;
+	virtual void PostEditChangeProperty(struct FPropertyChangedEvent& PropertyChangedEvent) OVERRIDE;
 
 	/** Load Material Type data from INI file **/
 	/** this changes displayname meta data. That means we won't need it outside of editor*/
diff --git a/Engine/Source/Runtime/Engine/Private/PhysicsEngine/PhysicalMaterial.cpp b/Engine/Source/Runtime/Engine/Private/PhysicsEngine/PhysicalMaterial.cpp
--- a/Engine/Source/Runtime/Engine/Private/PhysicsEngine/PhysicalMaterial.cpp
+++ b/Engine/Source/Runtime/Engine/Private/PhysicsEngine/PhysicalMaterial.cpp
@@ -26,6 +26,7 @@ UPhysicalMaterial::UPhysicalMaterial(const class FPostConstructInitializePropert
 	Density = 1.0f;
 	DestructibleDamageThresholdScale = 1.0f;
 	TireFrictionScale = 1.0f;
+	bOverrideFrictionCombineMode = false;
 #if WITH_PHYSX
 	PhysxUserData = FPhysxUserData(this);
 #endif
@@ -39,6 +40,19 @@ void UPhysicalMaterial::PostEditChangeProperty(FPropertyChangedEvent& PropertyCh
 
 	Super::PostEditChangeProperty(PropertyChangedEvent);
 }
+
+
+void UPhysicalMaterial::RebuildPhysicalMaterials()
+{
+	for (FObjectIterator Iter(UPhysicalMaterial::StaticClass()); Iter; ++Iter)
+	{
+		if (UPhysicalMaterial * PhysicalMaterial = Cast<UPhysicalMaterial>(*Iter))
+		{
+			PhysicalMaterial->UpdatePhysXMaterial();
+		}
+	}
+}
+
 #endif // WITH_EDITOR
 
 void UPhysicalMaterial::PostLoad()
@@ -69,6 +83,7 @@ void UPhysicalMaterial::FinishDestroy()
 	Super::FinishDestroy();
 }
 
+
 void UPhysicalMaterial::UpdatePhysXMaterial()
 {
 #if WITH_PHYSX
@@ -77,6 +92,8 @@ void UPhysicalMaterial::UpdatePhysXMaterial()
 		PMaterial->setStaticFriction(Friction);
 		PMaterial->setDynamicFriction(Friction);
 		PMaterial->setRestitution(Restitution);
+		uint32 CombineMode = (bOverrideFrictionCombineMode ? FrictionCombineMode.GetValue() : UPhysicsSettings::Get()->FrictionCombineMode.GetValue());
+		PMaterial->setFrictionCombineMode(static_cast<physx::PxCombineMode::Enum>(CombineMode));
 		FPhysXVehicleManager::UpdateTireFrictionTable();
 	}
 #endif	// WITH_PHYSX
@@ -89,6 +106,8 @@ PxMaterial* UPhysicalMaterial::GetPhysXMaterial()
 	if((PMaterial == NULL) && GPhysXSDK)
 	{
 		PMaterial = GPhysXSDK->createMaterial(Friction, Friction, Restitution);
+		uint32 CombineMode = (bOverrideFrictionCombineMode ? FrictionCombineMode.GetValue() : UPhysicsSettings::Get()->FrictionCombineMode.GetValue());
+		PMaterial->setFrictionCombineMode(static_cast<physx::PxCombineMode::Enum>(CombineMode));
 		PMaterial->userData = &PhysxUserData;
 
 		UpdatePhysXMaterial();
diff --git a/Engine/Source/Runtime/Engine/Private/PhysicsEngine/PhysicsSettings.cpp b/Engine/Source/Runtime/Engine/Private/PhysicsEngine/PhysicsSettings.cpp
--- a/Engine/Source/Runtime/Engine/Private/PhysicsEngine/PhysicsSettings.cpp
+++ b/Engine/Source/Runtime/Engine/Private/PhysicsEngine/PhysicsSettings.cpp
@@ -40,6 +40,17 @@ bool UPhysicsSettings::CanEditChange(const UProperty* Property) const
 	return bIsEditable;
 }
 
+void UPhysicsSettings::PostEditChangeProperty(struct FPropertyChangedEvent& PropertyChangedEvent)
+{
+	Super::PostEditChangeProperty(PropertyChangedEvent);
+
+	const FName PropertyName = PropertyChangedEvent.Property ? PropertyChangedEvent.Property->GetFName() : NAME_None;
+	if (PropertyName == FName(TEXT("FrictionCombineMode")))
+	{
+		UPhysicalMaterial::RebuildPhysicalMaterials();
+	}
+}
+
 void UPhysicsSettings::LoadSurfaceType()
 {
 	// read "SurfaceType" defines and set meta data for the enum
