Corrected FieldOfView tooltip.
#ttp 335560 - EDITOR: Tooltip over Field of View property in Camera Details tab indicates it is to set the vertical field of view.
#branch UE4
#proj Runtime.Engine
#change Corrected coment for UCameraComponent::FieldOfView from which tooltip is generated.

diff --git a/Engine/Source/Runtime/Engine/Classes/Camera/CameraComponent.h b/Engine/Source/Runtime/Engine/Classes/Camera/CameraComponent.h
--- a/Engine/Source/Runtime/Engine/Classes/Camera/CameraComponent.h
+++ b/Engine/Source/Runtime/Engine/Classes/Camera/CameraComponent.h
@@ -12,7 +12,7 @@ class UCameraComponent : public USceneComponent
 {
 	GENERATED_UCLASS_BODY()
 
-	/** The vertical field of view (in degrees) in perspective mode (ignored in Orthographic mode) */
+	/** The horizontal field of view (in degrees) in perspective mode (ignored in Orthographic mode) */
 	UPROPERTY(Interp, EditAnywhere, BlueprintReadWrite, Category=CameraSettings, meta=(UIMin = "5.0", UIMax = "170", ClampMin = "0.001", ClampMax = "360.0"))
 	float FieldOfView;
 
