Expose Get Physical Material to blueprint

diff --git a/Engine/Source/Runtime/Engine/Classes/Materials/MaterialInterface.h b/Engine/Source/Runtime/Engine/Classes/Materials/MaterialInterface.h
--- a/Engine/Source/Runtime/Engine/Classes/Materials/MaterialInterface.h
+++ b/Engine/Source/Runtime/Engine/Classes/Materials/MaterialInterface.h
@@ -181,6 +181,7 @@ public:
 	* Return a pointer to the physical material used by this material instance.
 	* @return The physical material.
 	*/
+	UFUNCTION(BlueprintCallable, Category = "Physics|Material")
 	virtual UPhysicalMaterial* GetPhysicalMaterial() const PURE_VIRTUAL(UMaterialInterface::GetPhysicalMaterial,return NULL;);
 
 	/** Return the textures used to render this material. */
