[335890] Fix for a crash when selecting instanced static mesh components in the Blueprint editor with instances added via the construction script.

- Removed what I think was unintentionally pasted code in FSCSEditorViewportClient::ProcessClick() (source of the crash).
- Modified UInstancedStaticMeshComponent's constructor to always default 'bHasPerInstanceHitProxies' to true.
- Modified UInstancedStaticMeshComponent::Serialize() to ignore the serialized value of 'bHasPerInstanceHitProxies' on load.
- Removed explicit assignment of 'bHasPerInstanceHitProxies' (no longer necessary) in UInstancedStaticMeshComponent::PostEditChangeChainProperty().
- Removed explicit assignment of 'bHasPerInstanceHitProxies' (no longer necessary) to cluster component instances in FFoliageMeshInfo::AddInstance().

diff --git a/Engine/Source/Editor/Kismet/Private/SCSEditorViewportClient.cpp b/Engine/Source/Editor/Kismet/Private/SCSEditorViewportClient.cpp
--- a/Engine/Source/Editor/Kismet/Private/SCSEditorViewportClient.cpp
+++ b/Engine/Source/Editor/Kismet/Private/SCSEditorViewportClient.cpp
@@ -313,11 +313,7 @@ void FSCSEditorViewportClient::ProcessClick(class FSceneView& View, class HHitPr
 				{
 					USceneComponent* CompInstance = *CompIt;
 					TSharedPtr<ISCSEditorCustomization> Customization = BlueprintEditorPtr.Pin()->CustomizeSCSEditor(CompInstance);
-					if(Customization.IsValid() && Customization->HandleViewportClick(AsShared(), View, HitProxy, Key, Event, HitX, HitY))
-					{
-						Invalidate();
-					}
-					else if(CompInstance == ActorProxy->PrimComponent)
+					if(CompInstance == ActorProxy->PrimComponent)
 					{
 						const bool bIsCtrlKeyDown = Viewport->KeyState(EKeys::LeftControl) || Viewport->KeyState(EKeys::RightControl);
 						if(BlueprintEditorPtr.IsValid())
diff --git a/Engine/Source/Runtime/Engine/Private/InstancedFoliage.cpp b/Engine/Source/Runtime/Engine/Private/InstancedFoliage.cpp
--- a/Engine/Source/Runtime/Engine/Private/InstancedFoliage.cpp
+++ b/Engine/Source/Runtime/Engine/Private/InstancedFoliage.cpp
@@ -299,7 +299,6 @@ void FFoliageMeshInfo::AddInstance( AInstancedFoliageActor* InIFA, UStaticMesh*
 
 		BestCluster->ClusterComponent->StaticMesh = InMesh;
 		BestCluster->ClusterComponent->bSelectable = true;
-		BestCluster->ClusterComponent->bHasPerInstanceHitProxies = true;
 		BestCluster->ClusterComponent->InstancingRandomSeed = FMath::Rand();
 		BestCluster->ClusterComponent->InstanceStartCullDistance = Settings->StartCullDistance;
 		BestCluster->ClusterComponent->InstanceEndCullDistance = Settings->EndCullDistance;
diff --git a/Engine/Source/Runtime/Engine/Private/InstancedStaticMesh.cpp b/Engine/Source/Runtime/Engine/Private/InstancedStaticMesh.cpp
--- a/Engine/Source/Runtime/Engine/Private/InstancedStaticMesh.cpp
+++ b/Engine/Source/Runtime/Engine/Private/InstancedStaticMesh.cpp
@@ -1119,6 +1119,8 @@ UInstancedStaticMeshComponent::UInstancedStaticMeshComponent(const class FPostCo
 	// Static lighting is not currently supported by instanced static meshes, so disable it
 	// See UInstancedStaticMeshComponent::GetStaticLightingInfo (it's empty, but it's what would add data to the Lightmass generation)
 	bCastStaticShadow = false;
+
+	bHasPerInstanceHitProxies = true;
 }
 
 #if WITH_EDITOR
@@ -1402,6 +1404,12 @@ void UInstancedStaticMeshComponent::Serialize(FArchive& Ar)
 		Ar << SelectedInstances;
 	}
 #endif
+
+	if( Ar.IsLoading() )
+	{
+		// Ignore the serialized value of this and always default to 'true' for this component type
+		bHasPerInstanceHitProxies = true;
+	}
 }
 
 void UInstancedStaticMeshComponent::AddInstance(const FTransform& InstanceTransform)
@@ -1530,9 +1538,6 @@ void UInstancedStaticMeshComponent::PostEditChangeChainProperty(FPropertyChanged
 			int32 AddedAtIndex = PropertyChangedEvent.GetArrayIndex(PropertyChangedEvent.Property->GetFName().ToString());
 			check(AddedAtIndex != INDEX_NONE);
 			SetupNewInstanceData(PerInstanceSMData[AddedAtIndex], AddedAtIndex, FTransform::Identity);
-
-			// added via the property editor, so we will want to interactively work with instances
-			bHasPerInstanceHitProxies = true;
 		}
 
 		MarkRenderStateDirty();
