Make SetBind visible to the exec system

diff --git a/Engine/Source/Runtime/Engine/Classes/GameFramework/PlayerInput.h b/Engine/Source/Runtime/Engine/Classes/GameFramework/PlayerInput.h
--- a/Engine/Source/Runtime/Engine/Classes/GameFramework/PlayerInput.h
+++ b/Engine/Source/Runtime/Engine/Classes/GameFramework/PlayerInput.h
@@ -273,10 +273,8 @@ public:
 	UFUNCTION(exec)
 	void SetMouseSensitivityToDefault();
 
-#if !UE_BUILD_SHIPPING
 	UFUNCTION(exec)
 	void SetBind(FName BindName, const FString& Command);
-#endif
 
 	/** Returns the mouse sensitivity along the X-axis, or the Y-axis, or 1.0 if none are known. */
 	float GetMouseSensitivity();
diff --git a/Engine/Source/Runtime/Engine/Private/UserInterface/PlayerInput.cpp b/Engine/Source/Runtime/Engine/Private/UserInterface/PlayerInput.cpp
--- a/Engine/Source/Runtime/Engine/Private/UserInterface/PlayerInput.cpp
+++ b/Engine/Source/Runtime/Engine/Private/UserInterface/PlayerInput.cpp
@@ -1574,9 +1574,11 @@ FKeyBind UPlayerInput::GetExecBind(FString const& ExecCommand)
 	}
 	return Binding;
 }
+#endif
 
 void UPlayerInput::SetBind(FName BindName, const FString& Command)
 {
+#if !UE_BUILD_SHIPPING
 	FKey const BindKey(BindName);
 	if (BindKey.IsValid())
 	{
@@ -1604,8 +1606,8 @@ void UPlayerInput::SetBind(FName BindName, const FString& Command)
 		DebugExecBindings.Add(NewBind);
 		SaveConfig();
 	}
-}
 #endif
+}
 
 class UWorld* UPlayerInput::GetWorld() const
 {
