Lots of little cleanup of FNetGUID code

* Added support functions on object: IsNameStableForNetworking, IsFullNameStableForNetworking, IsSupportedForNetworking
* Refer to these newly added functions in the various SupportsObject/IsDynamic support functions on UPackageMapClient
* Merged in code from Ratti that cleans up some subobject replication code
* Other various cleanups/name/etc

diff --git a/Engine/Source/Runtime/CoreUObject/Private/UObject/CoreNet.cpp b/Engine/Source/Runtime/CoreUObject/Private/UObject/CoreNet.cpp
--- a/Engine/Source/Runtime/CoreUObject/Private/UObject/CoreNet.cpp
+++ b/Engine/Source/Runtime/CoreUObject/Private/UObject/CoreNet.cpp
@@ -127,7 +127,7 @@ bool UPackageMap::WriteObject( FArchive& Ar, UObject* Outer, FNetworkGUID NetGUI
 FClassNetCache* UPackageMap::GetClassNetCache( UClass* Class )
 {
 	FClassNetCache* Result = ClassFieldIndices.FindRef(Class);
-	if( !Result && SupportsObject(Class) )
+	if( !Result )// && SupportsObject(Class) )
 	{
 		Result                       = ClassFieldIndices.Add( Class, new FClassNetCache(Class) );
 		Result->Super                = NULL;
@@ -143,7 +143,7 @@ FClassNetCache* UPackageMap::GetClassNetCache( UClass* Class )
 		{
 			// Add sandboxed items to net cache.  
 			UField* Field = Class->NetFields[i];
-			if( SupportsObject(Field ) )
+			//if( SupportsObject(Field ) )
 			{
 				int32 ThisIndex      = Result->GetMaxIndex();
 				new(Result->Fields)FFieldNetCache( Field, ThisIndex );
diff --git a/Engine/Source/Runtime/CoreUObject/Private/UObject/Obj.cpp b/Engine/Source/Runtime/CoreUObject/Private/UObject/Obj.cpp
--- a/Engine/Source/Runtime/CoreUObject/Private/UObject/Obj.cpp
+++ b/Engine/Source/Runtime/CoreUObject/Private/UObject/Obj.cpp
@@ -3775,3 +3775,26 @@ void UObject::PostNetReceive()
 {
 
 }
+
+/** IsNameStableForNetworking means an object can be referred to its path name (relative to outer) over the network */
+bool UObject::IsNameStableForNetworking() const
+{
+	return IsDefaultSubobject() || HasAnyFlags( RF_WasLoaded | RF_DefaultSubObject );
+}
+
+/** IsFullNameStableForNetworking means an object can be referred to its full path name over the network */
+bool UObject::IsFullNameStableForNetworking() const
+{
+	if ( GetOuter() != NULL && !GetOuter()->IsNameStableForNetworking() )
+	{
+		return false;	// If any outer isn't stable, we can't consider the full name stable
+	}
+
+	return IsNameStableForNetworking();
+}
+
+/** IsSupportedForNetworking means an object can be referenced over the network */
+bool UObject::IsSupportedForNetworking() const
+{
+	return IsFullNameStableForNetworking();
+}
diff --git a/Engine/Source/Runtime/CoreUObject/Public/UObject/Class.h b/Engine/Source/Runtime/CoreUObject/Public/UObject/Class.h
--- a/Engine/Source/Runtime/CoreUObject/Public/UObject/Class.h
+++ b/Engine/Source/Runtime/CoreUObject/Public/UObject/Class.h
@@ -1648,6 +1648,8 @@ public:
 	// List of network relevant fields (properties and functions)
 	TArray<UField*> NetFields;
 
+	virtual bool IsNameStableForNetworking() const OVERRIDE { return true; }		// For now, assume all classes have stable net names
+
 #if WITH_EDITOR || HACK_HEADER_GENERATOR 
 	// Editor only properties
 	void GetHideCategories(TArray<FString>& OutHideCategories) const;
diff --git a/Engine/Source/Runtime/CoreUObject/Public/UObject/CoreObject.h b/Engine/Source/Runtime/CoreUObject/Public/UObject/CoreObject.h
--- a/Engine/Source/Runtime/CoreUObject/Public/UObject/CoreObject.h
+++ b/Engine/Source/Runtime/CoreUObject/Public/UObject/CoreObject.h
@@ -64,6 +64,8 @@ private:
 	FName ForcedExportBasePackageName;
 public:
 
+	virtual bool IsNameStableForNetworking() const OVERRIDE { return true; }		// For now, assume all packages have stable net names
+
 	/** Package flags, serialized.*/
 	uint32	PackageFlags;
 
diff --git a/Engine/Source/Runtime/CoreUObject/Public/UObject/UObject.h b/Engine/Source/Runtime/CoreUObject/Public/UObject/UObject.h
--- a/Engine/Source/Runtime/CoreUObject/Public/UObject/UObject.h
+++ b/Engine/Source/Runtime/CoreUObject/Public/UObject/UObject.h
@@ -451,6 +451,18 @@ public:
 	/** Returns properties that are replicated for the lifetime of the actor channel */
 	virtual void GetLifetimeReplicatedProps( TArray< class FLifetimeProperty > & OutLifetimeProps ) const;
 
+	/** IsNameStableForNetworking means an object can be referred to by its path name (relative to outer) over the network */
+	virtual bool IsNameStableForNetworking() const;
+
+	/** IsFullNameStableForNetworking means an object can be referred by to its full path name over the network */
+	virtual bool IsFullNameStableForNetworking() const;
+
+	/** IsSupportedForNetworking means an object can be referenced over the network */
+	virtual bool IsSupportedForNetworking() const;
+
+	/** Returns a list of sub-objects that have stable names for networking */
+	virtual void GetSubobjectsWithStableNamesForNetworking(TArray<UObject*> &ObjList) {}
+
 	/** Called right before receiving a bunch */
 	virtual void PreNetReceive();
 
diff --git a/Engine/Source/Runtime/Engine/Classes/Components/ActorComponent.h b/Engine/Source/Runtime/Engine/Classes/Components/ActorComponent.h
--- a/Engine/Source/Runtime/Engine/Classes/Components/ActorComponent.h
+++ b/Engine/Source/Runtime/Engine/Classes/Components/ActorComponent.h
@@ -132,8 +132,11 @@ public:
 	/** This signifies the component can be ID'd by name over the network. This only needs to be called by engine code when constructing blueprint components. */
 	void SetNetAddressable();
 
-	/** Returns whether the component is addressable over a networked game by its path name*/
-	bool GetNetAddressable() const;
+	/** IsNameStableForNetworking means an object can be referred to by its path name (relative to outer) over the network */
+	virtual bool IsNameStableForNetworking() const OVERRIDE;
+
+	/** IsSupportedForNetworking means an object can be referenced over the network */
+	virtual bool IsSupportedForNetworking() const OVERRIDE;
 
 	/** Enable or disable replication. This is the equivelent of RemoteRole for actors (only a bool is required for components) */
 	UFUNCTION(BlueprintCallable, Category="Components")
diff --git a/Engine/Source/Runtime/Engine/Classes/Engine/Level.h b/Engine/Source/Runtime/Engine/Classes/Engine/Level.h
--- a/Engine/Source/Runtime/Engine/Classes/Engine/Level.h
+++ b/Engine/Source/Runtime/Engine/Classes/Engine/Level.h
@@ -549,6 +549,8 @@ public:
 	 */
 	ENGINE_API void SortActorList();
 
+	virtual bool IsNameStableForNetworking() const OVERRIDE { return true; }		// For now, assume all levels have stable net names
+
 	/** Handles network initialization for actors in this level */
 	void InitializeNetworkActors();
 
diff --git a/Engine/Source/Runtime/Engine/Classes/GameFramework/Actor.h b/Engine/Source/Runtime/Engine/Classes/GameFramework/Actor.h
--- a/Engine/Source/Runtime/Engine/Classes/GameFramework/Actor.h
+++ b/Engine/Source/Runtime/Engine/Classes/GameFramework/Actor.h
@@ -1396,6 +1396,15 @@ public:
 	// Always called immediately after properties are received from the remote.
 	virtual void PostNetReceive() OVERRIDE;
 
+	/** IsNameStableForNetworking means an object can be referred to by its path name (relative to outer) over the network */
+	virtual bool IsNameStableForNetworking() const OVERRIDE;
+
+	/** IsSupportedForNetworking means an object can be referenced over the network */
+	virtual bool IsSupportedForNetworking() const OVERRIDE;
+
+	/** Returns a list of sub-objects that have stable names for networking */
+	virtual void GetSubobjectsWithStableNamesForNetworking(TArray<UObject*> &ObjList) OVERRIDE;
+
 	// Always called immediately after spawning and reading in replicated properties
 	virtual void PostNetInit();
 
diff --git a/Engine/Source/Runtime/Engine/Private/ActorComponent.cpp b/Engine/Source/Runtime/Engine/Private/ActorComponent.cpp
--- a/Engine/Source/Runtime/Engine/Private/ActorComponent.cpp
+++ b/Engine/Source/Runtime/Engine/Private/ActorComponent.cpp
@@ -1077,10 +1077,10 @@ void UActorComponent::SetNetAddressable()
 	bNetAddressable = true;
 }
 
-bool UActorComponent::GetNetAddressable() const
+bool UActorComponent::IsNameStableForNetworking() const
 {
 	/** 
-	 * NetAddressable means a component can be referred to its path name (relative to owning AActor*) over the network
+	 * IsNameStableForNetworking means a component can be referred to its path name (relative to owning AActor*) over the network
 	 *
 	 * Components are net addressable if:
 	 *	-They are Default Subobjects (created in C++ constructor)
@@ -1088,7 +1088,12 @@ bool UActorComponent::GetNetAddressable() const
 	 *	-They were explicitly set to bNetAddressable (blueprint components created by SCS)
 	 */
 
-	return  bNetAddressable || IsDefaultSubobject() || HasAnyFlags(RF_WasLoaded);
+	return bNetAddressable || IsDefaultSubobject() || HasAnyFlags( RF_WasLoaded | RF_DefaultSubObject );
+}
+
+bool UActorComponent::IsSupportedForNetworking() const
+{
+	return IsNameStableForNetworking() || GetIsReplicated();
 }
 
 void UActorComponent::SetIsReplicated(bool ShouldReplicate)
diff --git a/Engine/Source/Runtime/Engine/Private/ActorReplication.cpp b/Engine/Source/Runtime/Engine/Private/ActorReplication.cpp
--- a/Engine/Source/Runtime/Engine/Private/ActorReplication.cpp
+++ b/Engine/Source/Runtime/Engine/Private/ActorReplication.cpp
@@ -248,13 +248,38 @@ bool AActor::ReplicateSubobjects(UActorChannel *Channel, FOutBunch *Bunch, FRepl
 		UActorComponent * ActorComp = ReplicatedComponents[CompIdx].Get();
 		if (ActorComp && ActorComp->GetIsReplicated())
 		{
-			WroteSomething |= Channel->ReplicateSubobject(ActorComp, *Bunch, *RepFlags);
-			WroteSomething |= ActorComp->ReplicateSubobjects(Channel, Bunch, RepFlags);
+			WroteSomething |= ActorComp->ReplicateSubobjects(Channel, Bunch, RepFlags);		// Lets the component add subobjects before replicating its own properties.
+			WroteSomething |= Channel->ReplicateSubobject(ActorComp, *Bunch, *RepFlags);	// (this makes those subobjects 'supported', and from here on those objects may have reference replicated)		
 		}
 	}
 	return WroteSomething;
 }
 
+void AActor::GetSubobjectsWithStableNamesForNetworking(TArray<UObject*> &ObjList)
+{	
+#if 0
+	// For experimenting with replicating ALL stably named components initially
+	for (UActorComponent* Component : OwnedComponents)
+	{
+		if (Component && Component->IsNameStableForNetworking())
+		{
+			ObjList.Add(Component);
+			Component->GetSubobjectsWithStableNamesForNetworking(ObjList);
+		}
+	}
+#else
+	for (TWeakObjectPtr<class UActorComponent> WeakPtr : ReplicatedComponents)
+	{
+		UActorComponent * Component = WeakPtr.Get();
+		if (Component && Component->IsNameStableForNetworking())
+		{
+			ObjList.Add(Component);
+			Component->GetSubobjectsWithStableNamesForNetworking(ObjList);
+		}
+	}
+#endif
+}
+
 void AActor::OnSubobjectCreatedFromReplication(UObject *NewSubobject)
 {
 	check(NewSubobject);
@@ -273,4 +298,14 @@ void AActor::OnSubobjectDestroyFromReplication(UObject *NewSubobject)
 	{
 		Component->DestroyComponent();
 	}
+}
+
+bool AActor::IsNameStableForNetworking() const
+{
+	return IsNetStartupActor() || HasAnyFlags( RF_ClassDefaultObject | RF_ArchetypeObject );
+}
+
+bool AActor::IsSupportedForNetworking() const
+{
+	return true;		// All actors are supported for networking
 }
\ No newline at end of file
diff --git a/Engine/Source/Runtime/Engine/Private/PackageMapClient.cpp b/Engine/Source/Runtime/Engine/Private/PackageMapClient.cpp
--- a/Engine/Source/Runtime/Engine/Private/PackageMapClient.cpp
+++ b/Engine/Source/Runtime/Engine/Private/PackageMapClient.cpp
@@ -12,6 +12,8 @@
 static const int GUID_PACKET_NOT_ACKED	= -2;		
 static const int GUID_PACKET_ACKED		= -1;		
 
+#define USE_NEW_DYNAMIC_AND_SUPPORT_CHECKS 1
+
 /**
  * Don't allow infinite recursion of InternalLoadObject - an attacker could
  * send malicious packets that cause a stack overflow on the server.
@@ -289,56 +291,42 @@ bool UPackageMapClient::SerializeNewActor(FArchive& Ar, class UActorChannel *Cha
 		// Serialize Initially replicating components here
 		if (Ar.IsSaving())
 		{
-			int32 NumReplicatedComponents = 0;
-			
-			// Count how many net addressable components we have, then serialize
-			// that number then each component netguid/string pair.
-			for (auto It = Actor->ReplicatedComponents.CreateIterator(); It; ++It)
-			{
-				UActorComponent * Comp = It->Get();
-				if (Comp && Comp->GetNetAddressable())
-				{
-					NumReplicatedComponents++;
-				}
-			}
-			
-			Actor->ReplicatedComponents.Num();
+			TArray<UObject*> Subobjs;
+			Actor->GetSubobjectsWithStableNamesForNetworking(Subobjs);
+
+			int32 NumReplicatedComponents = Subobjs.Num();
+
 			Ar << NumReplicatedComponents;
 
-			for (auto It = Actor->ReplicatedComponents.CreateIterator(); It; ++It)
+			for (UObject *Obj : Subobjs)
 			{
-				UActorComponent * Comp = It->Get();
-				if (Comp && Comp->GetNetAddressable())
-				{
+				FNetworkGUID SubobjNetGUID;
+				InternalGetOrAssignNetGUID(Obj, SubobjNetGUID);
+				Ar << SubobjNetGUID;
 
-					FNetworkGUID ComponentNetGUID;
-					InternalGetOrAssignNetGUID(Comp, ComponentNetGUID);
-					Ar << ComponentNetGUID;
-				
-					FString CompName = Comp->GetName();
-					Ar << CompName;
-					
-					NetGUIDAckStatus.Add( ComponentNetGUID, GUID_PACKET_ACKED );		// We know this GUID is acked since the export is in this bunch, and this bunch is reliable
-					UE_LOG(LogNetPackageMap, Log, TEXT("   Saving: Assigned NetGUID <%s> to component %s"), *ComponentNetGUID.ToString(), *Comp->GetPathName());
-				}
+				FString ObjName = Obj->GetName();
+				Ar << ObjName;
+
+				NetGUIDAckStatus.Add(SubobjNetGUID, GUID_PACKET_ACKED);		// We know this GUID is acked since the export is in this bunch, and this bunch is reliable
+				UE_LOG(LogNetPackageMap, Log, TEXT("   Saving: Assigned NetGUID <%s> to subobject %s"), *SubobjNetGUID.ToString(), *Obj->GetPathName());
 			}
 		}
 		else
 		{
 			int32 NumReplicatedComponents = 0;
 			Ar << NumReplicatedComponents;
-			
+
 			for (int32 idx=0; idx < NumReplicatedComponents && !Ar.IsError(); ++idx)
 			{
-				FNetworkGUID ComponentNetGUID;
-				Ar << ComponentNetGUID;
+				FNetworkGUID SubobjNetGUID;
+				Ar << SubobjNetGUID;
 
-				FString CompName;
-				Ar << CompName;
+				FString ObjName;
+				Ar << ObjName;
 
-				UObject * CompObject = NetGUIDAssign(ComponentNetGUID, CompName, Actor);
+				UObject * Subobj = NetGUIDAssign(SubobjNetGUID, ObjName, Actor);
 
-				UE_LOG(LogNetPackageMap, Log, TEXT("   Loading: Assigned NetGUID <%s> to component %s"), *ComponentNetGUID.ToString(), *CompObject->GetPathName());
+				UE_LOG(LogNetPackageMap, Log, TEXT("   Loading: Assigned NetGUID <%s> to subobject %s"), *SubobjNetGUID.ToString(), *Subobj->GetPathName());
 			}
 		}
 	}
@@ -1194,7 +1182,7 @@ bool UPackageMapClient::ShouldSendFullPath(const UObject* Object, const FNetwork
 
 	// If it is a dynamic actor component that hasn't been Ack'd, then send its path
 	const UActorComponent *Component = Cast<const UActorComponent>(Object);
-	if (IsDynamic && (Component && Component->GetNetAddressable()) && !NetGUIDHasBeenAckd(NetGUID) )
+	if (IsDynamic && (Component && Component->IsNameStableForNetworking()) && !NetGUIDHasBeenAckd(NetGUID) )
 	{
 		return true;
 	}
@@ -1208,6 +1196,13 @@ bool UPackageMapClient::ShouldSendFullPath(const UObject* Object, const FNetwork
  */
 bool UPackageMapClient::IsDynamicObject(const UObject* Object)
 {
+#if USE_NEW_DYNAMIC_AND_SUPPORT_CHECKS == 1
+	check( Object != NULL );
+	check( Object->IsSupportedForNetworking() );
+
+	// Any non net addressable object is dynamic
+	return !Object->IsFullNameStableForNetworking();
+#else
 	const AActor * Actor = Cast<const AActor>(Object);
 	bool NonAddressableSubobject = false;
 	if (!Actor)
@@ -1218,8 +1213,8 @@ bool UPackageMapClient::IsDynamicObject(const UObject* Object)
 			const UActorComponent* ActorComponent = Cast<const UActorComponent>(Object);
 			if (ActorComponent)
 			{
-				// Actor compnents define if they are NetAddressable via UActorComponent::GetNetAddressable()
-				if (!ActorComponent->GetNetAddressable())
+				// Actor compnents define if they are NetAddressable via UActorComponent::IsNameStableForNetworking()
+				if (!ActorComponent->IsNameStableForNetworking())
 				{
 					NonAddressableSubobject  = true;
 				}
@@ -1238,6 +1233,7 @@ bool UPackageMapClient::IsDynamicObject(const UObject* Object)
 	}
 
 	return true;
+#endif
 }
 
 /**
@@ -1313,6 +1309,38 @@ bool UPackageMapClient::IsNetGUIDAuthority()
 
 bool UPackageMapClient::SupportsObject( const UObject* Object )
 {
+#if USE_NEW_DYNAMIC_AND_SUPPORT_CHECKS == 1
+	// NULL is always supported
+	if ( !Object )
+	{
+		return true;
+	}
+
+	// If we already gave it a NetGUID, its supported.
+	// This should happen for dynamic subobjects.
+	FNetworkGUID NetGUID = Cache->NetGUIDLookup.FindRef(Object);
+	if (NetGUID.IsValid())
+	{
+		return true;
+	}
+
+	if ( Object->IsFullNameStableForNetworking() )
+	{
+		// If object is fully net addressable, it's definitely supported
+		return true;
+	}
+
+	if ( Object->IsSupportedForNetworking() )
+	{
+		// This means the server will explicitly tell the client to spawn and assign the id for this object
+		return true;
+	}
+
+	UE_LOG( LogNetPackageMap, Warning, TEXT( "UPackageMapClient::SupportsObject: %s NOT Supported." ), *Object->GetFullName() );
+	UE_LOG( LogNetPackageMap, Warning, TEXT( "   %s"), *DebugContextString );
+
+	return false;
+#else
 	// NULL is always supported
 	if (!Object)
 	{
@@ -1336,7 +1364,7 @@ bool UPackageMapClient::SupportsObject( const UObject* Object )
 		{
 			// NetAddressable actor components are supported even before they have a NetGUID
 			// So things like an actor's static mesh component can be referenced even if the SMC doesn't have replicated properties itself
-			if (ActorComponent->GetNetAddressable())
+			if (ActorComponent->IsNameStableForNetworking())
 			{
 				return true;
 			}
@@ -1373,6 +1401,7 @@ bool UPackageMapClient::SupportsObject( const UObject* Object )
 
 	// Everything else is now supported
 	return true;
+#endif
 }
 
 /** Resets packagemap state for server travel */
diff --git a/Engine/Source/Runtime/Engine/Private/RepLayout.cpp b/Engine/Source/Runtime/Engine/Private/RepLayout.cpp
--- a/Engine/Source/Runtime/Engine/Private/RepLayout.cpp
+++ b/Engine/Source/Runtime/Engine/Private/RepLayout.cpp
@@ -2334,10 +2334,12 @@ void FRepLayout::SendPropertiesForRPC( UObject * Object, UFunction * Function, U
 
 	for ( int32 i = 0; i < Parents.Num(); i++ )
 	{
+		/*
 		if ( !Writer.PackageMap->SupportsObject( Parents[i].Property ) )
 		{
 			continue;
 		}
+		*/
 
 		bool Send = true;
 
@@ -2378,10 +2380,12 @@ void FRepLayout::ReceivePropertiesForRPC( UObject * Object, UFunction * Function
 
 	for ( int32 i = 0; i < Parents.Num(); i++ )
 	{
+		/*
 		if ( !Reader.PackageMap->SupportsObject( Parents[i].Property ) )
 		{
 			continue;
 		}
+		*/
 
 		if ( Cast<UBoolProperty>( Parents[i].Property ) || Reader.ReadBit() ) 
 		{
@@ -2409,10 +2413,12 @@ void FRepLayout::SerializePropertiesForStruct( UStruct * Struct, FArchive & Ar,
 
 	for ( int32 i = 0; i < Parents.Num(); i++ )
 	{
+		/*
 		if ( Map != NULL && !Map->SupportsObject( Parents[i].Property ) )
 		{
 			continue;
 		}
+		*/
 
 		SerializeProperties_r( Ar, Map, Parents[i].CmdStart, Parents[i].CmdEnd, Data, bHasUnmapped );
 
