Change log to warning

diff --git a/Engine/Source/Runtime/Engine/Private/DataChannel.cpp b/Engine/Source/Runtime/Engine/Private/DataChannel.cpp
--- a/Engine/Source/Runtime/Engine/Private/DataChannel.cpp
+++ b/Engine/Source/Runtime/Engine/Private/DataChannel.cpp
@@ -1438,7 +1438,7 @@ void UActorChannel::ReceivedBunch( FInBunch& Bunch )
 		// We are unsynchronized. Instead of crashing, let's try to recover.
 		if (!NewChannelActor)
 		{
-			UE_LOG(LogNet, Log, TEXT("Received invalid actor class on channel %i"), ChIndex);
+			UE_LOG(LogNet, Warning, TEXT("Received invalid actor class on channel %i"), ChIndex);
 			Broken = 1;
 			FNetControlMessage<NMT_ActorChannelFailure>::Send(Connection, ChIndex);
 			return;
