#ue4 - Optimization: calls to PhysX multi-sweep, multi-trace, and multi-overlap tests don't need to initialize the hit result buffer.

- Also add more checks to validate that we only read position/distance/normal values from hit results when the appropriate flags are set on the result.

diff --git a/Engine/Source/Runtime/Engine/Private/Collision/CollisionConversions.cpp b/Engine/Source/Runtime/Engine/Private/Collision/CollisionConversions.cpp
--- a/Engine/Source/Runtime/Engine/Private/Collision/CollisionConversions.cpp
+++ b/Engine/Source/Runtime/Engine/Private/Collision/CollisionConversions.cpp
@@ -330,6 +330,7 @@ static void SetHitResultFromShapeAndFaceIndex(const PxShape* PShape,  const PxRi
 
 void ConvertQueryImpactHit(const PxLocationHit& PHit, FHitResult& OutResult, float CheckLength, const PxFilterData& QueryFilter, const FVector& StartLoc, const FVector& EndLoc, const PxGeometry* const Geom, const PxTransform& QueryTM, bool bReturnFaceIndex, bool bReturnPhysMat)
 {
+	check(PHit.flags & PxHitFlag::eDISTANCE);
 	if (Geom != NULL && PHit.hadInitialOverlap())
 	{
 		ConvertOverlappedShapeToImpactHit(PHit, StartLoc, EndLoc, OutResult, *Geom, QueryTM, QueryFilter, bReturnPhysMat);
@@ -343,13 +344,12 @@ void ConvertQueryImpactHit(const PxLocationHit& PHit, FHitResult& OutResult, flo
 
 	// figure out where the the "safe" location for this shape is by moving from the startLoc toward the ImpactPoint
 	const FVector TraceDir = EndLoc - StartLoc;
-
 	const FVector SafeLocationToFitShape = StartLoc + (HitTime * TraceDir);
 
 	// Other info
 	OutResult.Location = SafeLocationToFitShape;
-	OutResult.ImpactPoint = P2UVector(PHit.position);
-	OutResult.Normal = P2UVector(PHit.normal);
+	OutResult.ImpactPoint = (PHit.flags & PxHitFlag::ePOSITION) ? P2UVector(PHit.position) : StartLoc;
+	OutResult.Normal = (PHit.flags & PxHitFlag::eNORMAL) ? P2UVector(PHit.normal) : -TraceDir;
 	OutResult.ImpactNormal = OutResult.Normal;
 
 	OutResult.TraceStart = StartLoc;
@@ -360,7 +360,6 @@ void ConvertQueryImpactHit(const PxLocationHit& PHit, FHitResult& OutResult, flo
 	PxFilterData PShapeFilter = PHit.shape->getQueryFilterData();
 	PxSceneQueryHitType::Enum HitType = FPxQueryFilterCallback::CalcQueryHitType(QueryFilter, PShapeFilter);
 	OutResult.bBlockingHit = (HitType == PxSceneQueryHitType::eBLOCK); 
-
 	OutResult.bStartPenetrating = (PxU32)(PHit.hadInitialOverlap());
 
 
@@ -484,6 +483,7 @@ bool AddSweepResults(int32 NumHits, PxSweepHit* Hits, float CheckLength, const P
 	for(int32 i=0; i<NumHits; i++)
 	{
 		const PxSweepHit& PHit = Hits[i];
+		check(PHit.flags & PxHitFlag::eDISTANCE);
 		if(PHit.distance <= MaxDistance)
 		{
 			FHitResult NewResult(ForceInit);
diff --git a/Engine/Source/Runtime/Engine/Private/Collision/PhysXCollision.cpp b/Engine/Source/Runtime/Engine/Private/Collision/PhysXCollision.cpp
--- a/Engine/Source/Runtime/Engine/Private/Collision/PhysXCollision.cpp
+++ b/Engine/Source/Runtime/Engine/Private/Collision/PhysXCollision.cpp
@@ -462,7 +462,10 @@ bool RaycastMulti(const UWorld * World, TArray<struct FHitResult>& OutHits, cons
 		PxSceneQueryFlags POutputFlags = PxSceneQueryFlag::ePOSITION | PxSceneQueryFlag::eNORMAL | PxSceneQueryFlag::eDISTANCE;
 		FPxQueryFilterCallback PQueryCallback(Params.IgnoreActors);
 
-		PxRaycastHit PHits[HIT_BUFFER_SIZE];
+		// Create buffer for hits. Note: memory is not initialized (for perf reasons), since API does not require it.
+		TTypeCompatibleBytes<PxRaycastHit> RawPHits[HIT_BUFFER_SIZE];
+		PxRaycastHit* PHits = (PxRaycastHit*)RawPHits;
+
 		bool bBlockingHit = false;
 		const PxVec3 PDir = U2PVector(Delta/DeltaMag);
 
@@ -492,7 +495,7 @@ bool RaycastMulti(const UWorld * World, TArray<struct FHitResult>& OutHits, cons
 
 			// Write into the same PHits buffer
 			PxRaycastHit* PHitsAsync = PHits+NumHits;
-			const PxI32 AsyncBufferSize = (ARRAY_COUNT(PHits) - NumHits);
+			const PxI32 AsyncBufferSize = (ARRAY_COUNT(RawPHits) - NumHits);
 			check(AsyncBufferSize > 0);
 			bool bBlockingHitAsync = false;
 
@@ -778,7 +781,10 @@ bool GeomSweepMulti(const UWorld * World, const PxGeometry& PGeom, const PxQuat&
 		// Keep track of closest blocking hit distance.
 		float MinBlockDistance = DeltaMag;
 
-		PxSweepHit PHits[HIT_BUFFER_SIZE];
+		// Create buffer for hits. Note: memory is not initialized (for perf reasons), since API does not require it.
+		TTypeCompatibleBytes<PxSweepHit> RawPHits[HIT_BUFFER_SIZE];
+		PxSweepHit* PHits = (PxSweepHit*)RawPHits;
+
 		bool bBlockingHitSync;
 		PxI32 NumHits = SyncScene->sweepMultiple(PGeom, PStartTM, PDir, DeltaMag, POutputFlags, PHits, HIT_BUFFER_MAX_SYNC_QUERIES, bBlockingHitSync, PQueryFilterData, &PQueryCallbackSweep);
 		if(NumHits == -1)
@@ -844,7 +850,7 @@ bool GeomSweepMulti(const UWorld * World, const PxGeometry& PGeom, const PxQuat&
 
 			// Write into the same PHits buffer
 			PxSweepHit* PHitsAsync = PHits + NumHits;
-			const PxI32 AsyncBufferSize = (ARRAY_COUNT(PHits) - NumHits);
+			const PxI32 AsyncBufferSize = (ARRAY_COUNT(RawPHits) - NumHits);
 			check(AsyncBufferSize > 0);
 			bool bBlockingHitAsync = false;
 			PxI32 NumAsyncHits = AsyncScene->sweepMultiple(PGeom, PStartTM, PDir, MinBlockDistance, POutputFlags, PHitsAsync, AsyncBufferSize, bBlockingHitAsync, PQueryFilterData, &PQueryCallbackSweep);
@@ -1044,7 +1050,10 @@ bool GeomOverlapMulti(const UWorld * World, const PxGeometry& PGeom, const PxTra
 
 		SCOPED_SCENE_READ_LOCK(SyncScene);
 
-		PxOverlapHit POverlapArray[OVERLAP_BUFFER_SIZE]; 
+		// Create buffer for hits. Note: memory is not initialized (for perf reasons), since API does not require it.
+		TTypeCompatibleBytes<PxOverlapHit> RawPOverlapArray[OVERLAP_BUFFER_SIZE];
+		PxOverlapHit* POverlapArray = (PxOverlapHit*)RawPOverlapArray;
+
 		PxI32 NumHits = SyncScene->overlapMultiple(PGeom, PGeomPose, POverlapArray,  OVERLAP_BUFFER_SIZE_MAX_SYNC_QUERIES, PQueryFilterData, &PQueryCallback);
 		
 		if(NumHits == -1)
@@ -1065,7 +1074,7 @@ bool GeomOverlapMulti(const UWorld * World, const PxGeometry& PGeom, const PxTra
 
 			// Write into the same PHits buffer
 			PxOverlapHit* PAsyncOverlapArray = POverlapArray + NumHits;
-			const PxI32 AsyncBufferSize = (ARRAY_COUNT(POverlapArray) - NumHits);
+			const PxI32 AsyncBufferSize = (ARRAY_COUNT(RawPOverlapArray) - NumHits);
 			check(AsyncBufferSize > 0);
 			PxI32 NumAsyncHits = AsyncScene->overlapMultiple(PGeom, PGeomPose, PAsyncOverlapArray, AsyncBufferSize, PQueryFilterData, &PQueryCallback);
 
