Partial revert of CL# 2079444.

- No longer default 'bHasPerInstanceHitProxies' to 'true' because it breaks actor selection in the level editor viewport.

diff --git a/Engine/Source/Runtime/Engine/Private/InstancedFoliage.cpp b/Engine/Source/Runtime/Engine/Private/InstancedFoliage.cpp
--- a/Engine/Source/Runtime/Engine/Private/InstancedFoliage.cpp
+++ b/Engine/Source/Runtime/Engine/Private/InstancedFoliage.cpp
@@ -299,6 +299,7 @@ void FFoliageMeshInfo::AddInstance( AInstancedFoliageActor* InIFA, UStaticMesh*
 
 		BestCluster->ClusterComponent->StaticMesh = InMesh;
 		BestCluster->ClusterComponent->bSelectable = true;
+		BestCluster->ClusterComponent->bHasPerInstanceHitProxies = true;
 		BestCluster->ClusterComponent->InstancingRandomSeed = FMath::Rand();
 		BestCluster->ClusterComponent->InstanceStartCullDistance = Settings->StartCullDistance;
 		BestCluster->ClusterComponent->InstanceEndCullDistance = Settings->EndCullDistance;
diff --git a/Engine/Source/Runtime/Engine/Private/InstancedStaticMesh.cpp b/Engine/Source/Runtime/Engine/Private/InstancedStaticMesh.cpp
--- a/Engine/Source/Runtime/Engine/Private/InstancedStaticMesh.cpp
+++ b/Engine/Source/Runtime/Engine/Private/InstancedStaticMesh.cpp
@@ -1120,8 +1120,6 @@ UInstancedStaticMeshComponent::UInstancedStaticMeshComponent(const class FPostCo
 	// Static lighting is not currently supported by instanced static meshes, so disable it
 	// See UInstancedStaticMeshComponent::GetStaticLightingInfo (it's empty, but it's what would add data to the Lightmass generation)
 	bCastStaticShadow = false;
-
-	bHasPerInstanceHitProxies = true;
 }
 
 #if WITH_EDITOR
@@ -1405,12 +1403,6 @@ void UInstancedStaticMeshComponent::Serialize(FArchive& Ar)
 		Ar << SelectedInstances;
 	}
 #endif
-
-	if( Ar.IsLoading() )
-	{
-		// Ignore the serialized value of this and always default to 'true' for this component type
-		bHasPerInstanceHitProxies = true;
-	}
 }
 
 void UInstancedStaticMeshComponent::AddInstance(const FTransform& InstanceTransform)
@@ -1539,6 +1531,9 @@ void UInstancedStaticMeshComponent::PostEditChangeChainProperty(FPropertyChanged
 			int32 AddedAtIndex = PropertyChangedEvent.GetArrayIndex(PropertyChangedEvent.Property->GetFName().ToString());
 			check(AddedAtIndex != INDEX_NONE);
 			SetupNewInstanceData(PerInstanceSMData[AddedAtIndex], AddedAtIndex, FTransform::Identity);
+
+			// added via the property editor, so we will want to interactively work with instances
+			bHasPerInstanceHitProxies = true;
 		}
 
 		MarkRenderStateDirty();
