Fix build error on Mac by moving ELevelTick declaration up
Move FEditorSupportDelegates decl into its own header
Move FlushResourceStreaming into ContentStreaming.h
Move appGetStartupMap and appGetAllPotentialStartupPackageNames into UnrealEngine.h
Use #pragma once instead of #ifndef _INC_ENGINE

diff --git a/Engine/Source/Runtime/Engine/Classes/Engine/EngineBaseTypes.h b/Engine/Source/Runtime/Engine/Classes/Engine/EngineBaseTypes.h
--- a/Engine/Source/Runtime/Engine/Classes/Engine/EngineBaseTypes.h
+++ b/Engine/Source/Runtime/Engine/Classes/Engine/EngineBaseTypes.h
@@ -23,6 +23,15 @@ enum EInputEvent
 	IE_MAX                  =5,
 };
 
+/** Type of tick we wish to perform on the level */
+enum ELevelTick
+{
+	LEVELTICK_TimeOnly = 0,	// Update the level time only.
+	LEVELTICK_ViewportsOnly = 1,	// Update time and viewports.
+	LEVELTICK_All = 2,	// Update all.
+	LEVELTICK_PauseTick = 3,	// Delta time is zero, we are paused. Components don't tick.
+};
+
 /**
  * Determines which ticking group an Actor/Component belongs to
  */
@@ -294,7 +303,7 @@ private:
 	 * @param CurrentThread - thread we are executing on, useful to pass along as new tasks are created
 	 * @param MyCompletionGraphEvent - completion event for this task. Useful for holding the completetion of this task until certain child tasks are complete.
 	 **/
-	virtual void ExecuteTick(float DeltaTime, enum ELevelTick TickType, ENamedThreads::Type CurrentThread, const FGraphEventRef& MyCompletionGraphEvent)
+	virtual void ExecuteTick(float DeltaTime, ELevelTick TickType, ENamedThreads::Type CurrentThread, const FGraphEventRef& MyCompletionGraphEvent)
 	{
 		check(0); // you cannot make this pure virtual in script because it wants to create constructors.
 	}
@@ -328,7 +337,7 @@ struct FActorTickFunction : public FTickFunction
 		* @param CurrentThread - thread we are executing on, useful to pass along as new tasks are created
 		* @param MyCompletionGraphEvent - completion event for this task. Useful for holding the completetion of this task until certain child tasks are complete.
 	**/
-	ENGINE_API virtual void ExecuteTick(float DeltaTime, enum ELevelTick TickType, ENamedThreads::Type CurrentThread, const FGraphEventRef& MyCompletionGraphEvent) OVERRIDE;
+	ENGINE_API virtual void ExecuteTick(float DeltaTime, ELevelTick TickType, ENamedThreads::Type CurrentThread, const FGraphEventRef& MyCompletionGraphEvent) OVERRIDE;
 	/** Abstract function to describe this tick. Used to print messages about illegal cycles in the dependency graph **/
 	ENGINE_API virtual FString DiagnosticMessage();
 };
@@ -351,7 +360,7 @@ struct FActorComponentTickFunction : public FTickFunction
 		* @param CurrentThread - thread we are executing on, useful to pass along as new tasks are created
 		* @param MyCompletionGraphEvent - completion event for this task. Useful for holding the completetion of this task until certain child tasks are complete.
 	**/
-	ENGINE_API virtual void ExecuteTick(float DeltaTime, enum ELevelTick TickType, ENamedThreads::Type CurrentThread, const FGraphEventRef& MyCompletionGraphEvent) OVERRIDE;
+	ENGINE_API virtual void ExecuteTick(float DeltaTime, ELevelTick TickType, ENamedThreads::Type CurrentThread, const FGraphEventRef& MyCompletionGraphEvent) OVERRIDE;
 	/** Abstract function to describe this tick. Used to print messages about illegal cycles in the dependency graph **/
 	ENGINE_API virtual FString DiagnosticMessage();
 };
@@ -374,19 +383,11 @@ struct FPrimitiveComponentPostPhysicsTickFunction : public FTickFunction
 		* @param CurrentThread - thread we are executing on, useful to pass along as new tasks are created
 		* @param MyCompletionGraphEvent - completion event for this task. Useful for holding the completetion of this task until certain child tasks are complete.
 	**/
-	virtual void ExecuteTick(float DeltaTime, enum ELevelTick TickType, ENamedThreads::Type CurrentThread, const FGraphEventRef& MyCompletionGraphEvent) OVERRIDE;
+	virtual void ExecuteTick(float DeltaTime, ELevelTick TickType, ENamedThreads::Type CurrentThread, const FGraphEventRef& MyCompletionGraphEvent) OVERRIDE;
 	/** Abstract function to describe this tick. Used to print messages about illegal cycles in the dependency graph **/
 	virtual FString DiagnosticMessage();
 };
 
-enum ELevelTick
-{
-	LEVELTICK_TimeOnly = 0,	// Update the level time only.
-	LEVELTICK_ViewportsOnly = 1,	// Update time and viewports.
-	LEVELTICK_All = 2,	// Update all.
-	LEVELTICK_PauseTick = 3,	// Delta time is zero, we are paused. Components don't tick.
-};
-
 
 
 /** Types of network failures broadcast from the engine */
diff --git a/Engine/Source/Runtime/Engine/Public/ContentStreaming.h b/Engine/Source/Runtime/Engine/Public/ContentStreaming.h
--- a/Engine/Source/Runtime/Engine/Public/ContentStreaming.h
+++ b/Engine/Source/Runtime/Engine/Public/ContentStreaming.h
@@ -66,6 +66,9 @@ template<typename T>
 class FAsyncTask;
 struct FStreamingManagerTexture;
 
+/** Helper function to flush resource streaming. */
+void FlushResourceStreaming();
+
 /*-----------------------------------------------------------------------------
 	Base streaming classes.
 -----------------------------------------------------------------------------*/
diff --git a/Engine/Source/Runtime/Engine/Public/EditorSupportDelegates.h b/Engine/Source/Runtime/Engine/Public/EditorSupportDelegates.h
new file mode 100644
--- /dev/null
+++ b/Engine/Source/Runtime/Engine/Public/EditorSupportDelegates.h
@@ -0,0 +1,45 @@
+// Copyright 1998-2014 Epic Games, Inc. All Rights Reserved.
+
+#pragma once
+
+#if WITH_EDITOR
+
+/** 
+ * FEditorSupportDelegates
+ * Delegates that are needed for proper editor functionality, but are accessed or triggered in engine code.
+ **/
+struct ENGINE_API FEditorSupportDelegates
+{
+	/** delegate type for force property window rebuild events ( Params: UObject* Object ) */
+	DECLARE_MULTICAST_DELEGATE_OneParam(FOnForcePropertyWindowRebuild, UObject*); 
+	/** delegate type for material texture setting change events ( Params: UMaterialIterface* Material ) */
+	DECLARE_MULTICAST_DELEGATE_OneParam(FOnMaterialTextureSettingsChanged, class UMaterialInterface*);	
+	/** delegate type for windows messageing events ( Params: FViewport* Viewport, uint32 Message )*/
+	DECLARE_MULTICAST_DELEGATE_TwoParams(FOnWindowsMessage, class FViewport*, uint32);
+	/** delegate type for material usage flags change events ( Params: UMaterial* material, int32 FlagThatChanged ) */
+	DECLARE_MULTICAST_DELEGATE_TwoParams(FOnMaterialUsageFlagsChanged, class UMaterial*, int32); 
+
+	/** Called when all viewports need to be redrawn */
+	static FSimpleMulticastDelegate RedrawAllViewports;
+	/** Called when the editor is cleansing of transient references before a map change event */
+	static FSimpleMulticastDelegate CleanseEditor;
+	/** Called when the world is modified */
+	static FSimpleMulticastDelegate WorldChange;
+	/** Sent to force a property window rebuild */
+	static FOnForcePropertyWindowRebuild ForcePropertyWindowRebuild;
+	/** Sent when events happen that affect how the editors UI looks (mode changes, grid size changes, etc) */
+	static FSimpleMulticastDelegate UpdateUI;
+	/** Called for a material after the user has change a texture's compression settings.
+		Needed to notify the material editors that the need to reattach their preview objects */
+	static FOnMaterialTextureSettingsChanged MaterialTextureSettingsChanged;
+	/** Refresh property windows w/o creating/destroying controls */
+	static FSimpleMulticastDelegate RefreshPropertyWindows;
+	/** Sent before the given windows message is handled in the given viewport */
+	static FOnWindowsMessage PreWindowsMessage;
+	/** Sent after the given windows message is handled in the given viewport */
+	static FOnWindowsMessage PostWindowsMessage;
+	/** Sent after the usages flags on a material have changed*/
+	static FOnMaterialUsageFlagsChanged MaterialUsageFlagsChanged;
+};
+
+#endif // WITH_EDITOR
diff --git a/Engine/Source/Runtime/Engine/Public/Engine.h b/Engine/Source/Runtime/Engine/Public/Engine.h
--- a/Engine/Source/Runtime/Engine/Public/Engine.h
+++ b/Engine/Source/Runtime/Engine/Public/Engine.h
@@ -4,8 +4,7 @@
 	Engine.h: Unreal engine public header file.
 =============================================================================*/
 
-#ifndef _INC_ENGINE
-#define _INC_ENGINE
+#pragma once
 
 /*-----------------------------------------------------------------------------
 	Configuration defines
@@ -68,6 +67,15 @@
 	#define RECAST_ASYNC_REBUILDING	1
 #endif
 
+/*-----------------------------------------------------------------------------
+	Size of the world.
+-----------------------------------------------------------------------------*/
+
+#define WORLD_MAX					524288.0				/* Maximum size of the world */
+#define HALF_WORLD_MAX				(WORLD_MAX * 0.5)		/* Half the maximum size of the world */
+#define HALF_WORLD_MAX1				(HALF_WORLD_MAX - 1.0)	/* Half the maximum size of the world minus one */
+#define DEFAULT_ORTHOZOOM			10000.0					/* Default 2D viewport zoom */
+
 /*-----------------------------------------------------------------------------
 	Dependencies.
 -----------------------------------------------------------------------------*/
@@ -81,35 +89,9 @@
 #include "InputCore.h"
 #include "EngineMessages.h"
 #include "EngineSettings.h"
+#include "EditorSupportDelegates.h"
+#include "EngineStats.h"
 
-/** Helper function to flush resource streaming. */
-extern void FlushResourceStreaming();
-
-/**
- * This function will look at the given command line and see if we have passed in a map to load.
- * If we have then use that.
- * If we have not then use the DefaultMap which is stored in the Engine.ini
- * 
- * @see UGameEngine::Init() for this method of getting the correct start up map
- *
- * @param CommandLine Commandline to use to get startup map (NULL or "" will return default startup map)
- *
- * @return Name of the startup map without an extension (so can be used as a package name)
- */
-ENGINE_API FString appGetStartupMap(const TCHAR* CommandLine);
-
-/**
- * Get a list of all packages that may be needed at startup, and could be
- * loaded async in the background when doing seek free loading
- *
- * @param PackageNames The output list of package names
- * @param EngineConfigFilename Optional engine config filename to use to lookup the package settings
- */
-ENGINE_API void appGetAllPotentialStartupPackageNames(TArray<FString>& PackageNames, const FString& EngineConfigFilename, bool bIsCreatingHashes);
-
-/*-----------------------------------------------------------------------------
-	Global variables.
------------------------------------------------------------------------------*/
 
 ENGINE_API DECLARE_LOG_CATEGORY_EXTERN(LogMCP, Warning, All);
 ENGINE_API DECLARE_LOG_CATEGORY_EXTERN(LogPath, Warning, All);
@@ -131,9 +113,7 @@ ENGINE_API DECLARE_LOG_CATEGORY_EXTERN(LogTexture, Log, All);
 ENGINE_API DECLARE_LOG_CATEGORY_EXTERN(LogLandscape, Warning, All);
 ENGINE_API DECLARE_LOG_CATEGORY_EXTERN(LogPlayerManagement, Error, All);
 
-/**
- * Global engine pointer. Can be 0 so don't use without checking.
- */
+/** Global engine pointer. Can be 0 so don't use without checking. */
 extern ENGINE_API class UEngine*			GEngine;
 
 /** when set, disallows all network travel (pending level rejects all client travel attempts) */
@@ -142,62 +122,6 @@ extern ENGINE_API bool						GDisallowNetworkTravel;
 /** The GPU time taken to render the last frame. Same metric as FPlatformTime::Cycles(). */
 extern ENGINE_API uint32					GGPUFrameTime;
 
-#if WITH_EDITOR
-
-/** 
- * FEditorSupportDelegates
- * Delegates that are needed for proper editor functionality, but are accessed or triggered in engine code.
- **/
-struct ENGINE_API FEditorSupportDelegates
-{
-	/** delegate type for force property window rebuild events ( Params: UObject* Object ) */
-	DECLARE_MULTICAST_DELEGATE_OneParam(FOnForcePropertyWindowRebuild, UObject*); 
-	/** delegate type for material texture setting change events ( Params: UMaterialIterface* Material ) */
-	DECLARE_MULTICAST_DELEGATE_OneParam(FOnMaterialTextureSettingsChanged, class UMaterialInterface*);	
-	/** delegate type for windows messageing events ( Params: FViewport* Viewport, uint32 Message )*/
-	DECLARE_MULTICAST_DELEGATE_TwoParams(FOnWindowsMessage, class FViewport*, uint32);
-	/** delegate type for material usage flags change events ( Params: UMaterial* material, int32 FlagThatChanged ) */
-	DECLARE_MULTICAST_DELEGATE_TwoParams(FOnMaterialUsageFlagsChanged, class UMaterial*, int32); 
-
-	/** Called when all viewports need to be redrawn */
-	static FSimpleMulticastDelegate RedrawAllViewports;
-	/** Called when the editor is cleansing of transient references before a map change event */
-	static FSimpleMulticastDelegate CleanseEditor;
-	/** Called when the world is modified */
-	static FSimpleMulticastDelegate WorldChange;
-	/** Sent to force a property window rebuild */
-	static FOnForcePropertyWindowRebuild ForcePropertyWindowRebuild;
-	/** Sent when events happen that affect how the editors UI looks (mode changes, grid size changes, etc) */
-	static FSimpleMulticastDelegate UpdateUI;
-	/** Called for a material after the user has change a texture's compression settings.
-		Needed to notify the material editors that the need to reattach their preview objects */
-	static FOnMaterialTextureSettingsChanged MaterialTextureSettingsChanged;
-	/** Refresh property windows w/o creating/destroying controls */
-	static FSimpleMulticastDelegate RefreshPropertyWindows;
-	/** Sent before the given windows message is handled in the given viewport */
-	static FOnWindowsMessage PreWindowsMessage;
-	/** Sent after the given windows message is handled in the given viewport */
-	static FOnWindowsMessage PostWindowsMessage;
-	/** Sent after the usages flags on a material have changed*/
-	static FOnMaterialUsageFlagsChanged MaterialUsageFlagsChanged;
-};
-
-#endif // WITH_EDITOR
-
-/*-----------------------------------------------------------------------------
-	Size of the world.
------------------------------------------------------------------------------*/
-
-#define WORLD_MAX					524288.0				/* Maximum size of the world */
-#define HALF_WORLD_MAX				(WORLD_MAX * 0.5)		/* Half the maximum size of the world */
-#define HALF_WORLD_MAX1				(HALF_WORLD_MAX - 1.0)	/* Half the maximum size of the world minus one */
-#define DEFAULT_ORTHOZOOM			10000.0					/* Default 2D viewport zoom */
-
-/*-----------------------------------------------------------------------------
-	Stats.
------------------------------------------------------------------------------*/
-
-#include "EngineStats.h"
 
 /*-----------------------------------------------------------------------------
 	Forward declarations.
@@ -1234,6 +1158,5 @@ extern ENGINE_API IRendererModule& GetRendererModule();
 /** Clears the cached renderer module reference. */
 extern ENGINE_API void ResetCachedRendererModule();
 
-#endif // _INC_ENGINE
 
 
diff --git a/Engine/Source/Runtime/Engine/Public/UnrealEngine.h b/Engine/Source/Runtime/Engine/Public/UnrealEngine.h
--- a/Engine/Source/Runtime/Engine/Public/UnrealEngine.h
+++ b/Engine/Source/Runtime/Engine/Public/UnrealEngine.h
@@ -315,6 +315,28 @@ public:
 
 #endif
 
+/**
+ * This function will look at the given command line and see if we have passed in a map to load.
+ * If we have then use that.
+ * If we have not then use the DefaultMap which is stored in the Engine.ini
+ * 
+ * @see UGameEngine::Init() for this method of getting the correct start up map
+ *
+ * @param CommandLine Commandline to use to get startup map (NULL or "" will return default startup map)
+ *
+ * @return Name of the startup map without an extension (so can be used as a package name)
+ */
+ENGINE_API FString appGetStartupMap(const TCHAR* CommandLine);
+
+/**
+ * Get a list of all packages that may be needed at startup, and could be
+ * loaded async in the background when doing seek free loading
+ *
+ * @param PackageNames The output list of package names
+ * @param EngineConfigFilename Optional engine config filename to use to lookup the package settings
+ */
+ENGINE_API void appGetAllPotentialStartupPackageNames(TArray<FString>& PackageNames, const FString& EngineConfigFilename, bool bIsCreatingHashes);
+
 // Calculate the average frame time by using the stats system.
 inline void CalculateFPSTimings()
 {
