Fix ExecuteConstruction not being called in all places

diff --git a/Engine/Source/Runtime/Engine/Private/Actor.cpp b/Engine/Source/Runtime/Engine/Private/Actor.cpp
--- a/Engine/Source/Runtime/Engine/Private/Actor.cpp
+++ b/Engine/Source/Runtime/Engine/Private/Actor.cpp
@@ -2199,7 +2199,7 @@ void AActor::PostSpawnInitialize(FVector const& SpawnLocation, FRotator const& S
 	// After this, we can assume all components are created and assembled.
 	if (!bDeferConstruction)
 	{
-		OnConstruction( FTransform(SpawnRotation, SpawnLocation) );
+		ExecuteConstruction( FTransform(SpawnRotation, SpawnLocation), NULL );
 		PostActorConstruction();
 	}
 }
diff --git a/Engine/Source/Runtime/Engine/Private/ChildActorComponent.cpp b/Engine/Source/Runtime/Engine/Private/ChildActorComponent.cpp
--- a/Engine/Source/Runtime/Engine/Private/ChildActorComponent.cpp
+++ b/Engine/Source/Runtime/Engine/Private/ChildActorComponent.cpp
@@ -52,7 +52,7 @@ void UChildActorComponent::CreateChildActor()
 #endif
 
 				// Parts that we deferred from SpawnActor
-				ChildActor->OnConstruction(ComponentToWorld);
+				ChildActor->ExecuteConstruction(ComponentToWorld, NULL);
 				ChildActor->PostActorConstruction();
 
 				// attach new actor to this component
diff --git a/Engine/Source/Runtime/Engine/Private/GameplayStatics.cpp b/Engine/Source/Runtime/Engine/Private/GameplayStatics.cpp
--- a/Engine/Source/Runtime/Engine/Private/GameplayStatics.cpp
+++ b/Engine/Source/Runtime/Engine/Private/GameplayStatics.cpp
@@ -289,7 +289,7 @@ AActor* UGameplayStatics::FinishSpawningActor(AActor* Actor, const FTransform& S
 {
 	if (Actor)
 	{
-		Actor->OnConstruction(SpawnTransform);
+		Actor->ExecuteConstruction(SpawnTransform, NULL);
 		Actor->PostActorConstruction();
 	}
 
