Movement input & external forces can now be applied during path following

Tweaked CalcVelocity to apply path following as pseudo-input acceleration, rather than overriding velocity as it does now.

TTP# 333551 - GF: CharacterMovement: allow movement input and external forces to be applied during path following

diff --git a/Engine/Source/Runtime/Engine/Classes/GameFramework/CharacterMovementComponent.h b/Engine/Source/Runtime/Engine/Classes/GameFramework/CharacterMovementComponent.h
--- a/Engine/Source/Runtime/Engine/Classes/GameFramework/CharacterMovementComponent.h
+++ b/Engine/Source/Runtime/Engine/Classes/GameFramework/CharacterMovementComponent.h
@@ -745,7 +745,7 @@ public:
 	virtual FRotator ComputeOrientToMovementRotation(const FRotator& CurrentRotation, float DeltaTime, FRotator& DeltaRotation);
 
 	/** use velocity requested by path following */
-	virtual bool ApplyRequestedMove(FVector& NewVelocity, FVector& NewAcceleration, float DeltaTime, float MaxAccel, float MaxSpeed);
+	virtual bool ApplyRequestedMove(FVector& NewAcceleration, float DeltaTime, float MaxAccel, float MaxSpeed, float& InOutRequestedSpeed);
 
 	/** Called if bNotifyApex is true and character has just passed the apex of its jump. */
 	virtual void NotifyJumpApex();
diff --git a/Engine/Source/Runtime/Engine/Private/CharacterMovementComponent.cpp b/Engine/Source/Runtime/Engine/Private/CharacterMovementComponent.cpp
--- a/Engine/Source/Runtime/Engine/Private/CharacterMovementComponent.cpp
+++ b/Engine/Source/Runtime/Engine/Private/CharacterMovementComponent.cpp
@@ -1820,82 +1820,88 @@ void UCharacterMovementComponent::CalcVelocity(float DeltaTime, float Friction,
 	float MaxSpeed = GetModifiedMaxSpeed();
 	
 	// Check if path following requested movement
-	if ( !ApplyRequestedMove(Velocity, Acceleration, DeltaTime, MaxAccel, MaxSpeed) )
+	FVector RequestedAcceleration = FVector::ZeroVector;
+	float RequestedSpeed = 0.0f;
+	if(ApplyRequestedMove(RequestedAcceleration, DeltaTime, MaxAccel, MaxSpeed, RequestedSpeed))
 	{
-		// Update Velocity
-		if (bForceMaxAccel)
-		{
-			// Force acceleration at full speed.
-			// In consideration order for direction: Acceleration, then Velocity, then Pawn's rotation.
-			if (Acceleration.SizeSquared() > SMALL_NUMBER)
-			{
-				Acceleration = SafeNormalPrecise(Acceleration) * MaxAccel;
-			}
-			else 
-			{
-				Acceleration = MaxAccel * (Velocity.SizeSquared() < SMALL_NUMBER ? CharacterOwner->GetActorRotation().Vector() : SafeNormalPrecise(Velocity));
-			}
-
-			AnalogInputModifier = 1.f;
-		}
-
-		// Path following above didn't care about the analog modifier, but we do for everything else below, so get the fully modified value.
-		MaxSpeed *= AnalogInputModifier;
+		Acceleration += RequestedAcceleration;
+		Acceleration = Acceleration.ClampMaxSize(MaxAccel);
+	}
 
-		// Apply braking or deceleration
-		const bool bVelocityOverMax = IsExceedingMaxSpeed(MaxSpeed);
-		if (Acceleration.IsZero() || bVelocityOverMax)
+	// Update Velocity
+	if (bForceMaxAccel)
+	{
+		// Force acceleration at full speed.
+		// In consideration order for direction: Acceleration, then Velocity, then Pawn's rotation.
+		if (Acceleration.SizeSquared() > SMALL_NUMBER)
 		{
-			const FVector OldVelocity = Velocity;
-			ApplyVelocityBraking(DeltaTime, Friction, BrakingDeceleration);
-	
-			// Don't allow braking to lower us below max speed if we started above it.
-			if (FVector::DotProduct(Acceleration, OldVelocity) > 0.0f && bVelocityOverMax && Velocity.SizeSquared() < FMath::Square(MaxSpeed))
-			{
-				Velocity = SafeNormalPrecise(OldVelocity) * MaxSpeed;
-			}
+			Acceleration = SafeNormalPrecise(Acceleration) * MaxAccel;
 		}
-		else
+		else 
 		{
-			const FVector AccelDir = SafeNormalPrecise(Acceleration);
-			const float VelSize = Velocity.Size();
-			Velocity = Velocity - (Velocity - AccelDir * VelSize) * DeltaTime * Friction;
+			Acceleration = MaxAccel * (Velocity.SizeSquared() < SMALL_NUMBER ? CharacterOwner->GetActorRotation().Vector() : SafeNormalPrecise(Velocity));
 		}
 
-		// Apply fluid friction
-		if (bFluid)
+		AnalogInputModifier = 1.f;
+	}
+
+	// Path following above didn't care about the analog modifier, but we do for everything else below, so get the fully modified value.
+	// Use max of requested speed and max speed if we modified the speed in ApplyRequestedMove above
+	MaxSpeed = FMath::Max(RequestedSpeed, MaxSpeed * AnalogInputModifier);
+
+	// Apply braking or deceleration
+	const bool bVelocityOverMax = IsExceedingMaxSpeed(MaxSpeed);
+	if (Acceleration.IsZero() || bVelocityOverMax)
+	{
+		const FVector OldVelocity = Velocity;
+		ApplyVelocityBraking(DeltaTime, Friction, BrakingDeceleration);
+	
+		// Don't allow braking to lower us below max speed if we started above it.
+		if (FVector::DotProduct(Acceleration, OldVelocity) > 0.0f && bVelocityOverMax && Velocity.SizeSquared() < FMath::Square(MaxSpeed))
 		{
-			Velocity = Velocity * (1.f - Friction * DeltaTime);
+			Velocity = SafeNormalPrecise(OldVelocity) * MaxSpeed;
 		}
+	}
+	else
+	{
+		const FVector AccelDir = SafeNormalPrecise(Acceleration);
+		const float VelSize = Velocity.Size();
+		Velocity = Velocity - (Velocity - AccelDir * VelSize) * DeltaTime * Friction;
+	}
 
-		// Apply input acceleration
-		const float NewMaxSpeed = (IsExceedingMaxSpeed(MaxSpeed)) ? Velocity.Size() : MaxSpeed;
-		Velocity += Acceleration * DeltaTime;
-		Velocity = ClampMaxSizePrecise(Velocity, NewMaxSpeed);
+	// Apply fluid friction
+	if (bFluid)
+	{
+		Velocity = Velocity * (1.f - Friction * DeltaTime);
 	}
 
+	// Apply input acceleration
+	const float NewMaxSpeed = (IsExceedingMaxSpeed(MaxSpeed)) ? Velocity.Size() : MaxSpeed;
+	Velocity += Acceleration * DeltaTime;
+	Velocity = ClampMaxSizePrecise(Velocity, NewMaxSpeed);
+
 	if (bUseRVOAvoidance)
 	{
 		CalcAvoidanceVelocity(DeltaTime);
 	}
 }
 
-bool UCharacterMovementComponent::ApplyRequestedMove(FVector& NewVelocity, FVector& NewAcceleration, float DeltaTime, float MaxAccel, float MaxSpeed)
+bool UCharacterMovementComponent::ApplyRequestedMove(FVector& NewAcceleration, float DeltaTime, float MaxAccel, float MaxSpeed, float& InOutRequestedSpeed)
 {
 	if (bHasRequestedVelocity)
 	{
 		bHasRequestedVelocity = false;
 
-		float RequestedSpeed = RequestedVelocity.Size();
-		if (RequestedSpeed < KINDA_SMALL_NUMBER)
+		InOutRequestedSpeed = RequestedVelocity.Size();
+		if (InOutRequestedSpeed < KINDA_SMALL_NUMBER)
 		{
 			return true;
 		}
 
-		const FVector RequestedMoveDir = RequestedVelocity / RequestedSpeed;
+		const FVector RequestedMoveDir = RequestedVelocity / InOutRequestedSpeed;
 
-		RequestedSpeed = bRequestedMoveWithMaxSpeed ? MaxSpeed : FMath::Min(MaxSpeed, RequestedSpeed);
-		const FVector MoveVelocity = RequestedMoveDir * RequestedSpeed;
+		InOutRequestedSpeed = bRequestedMoveWithMaxSpeed ? MaxSpeed : FMath::Min(MaxSpeed, InOutRequestedSpeed);
+		const FVector MoveVelocity = RequestedMoveDir * InOutRequestedSpeed;
 
 		NewAcceleration = MoveVelocity / DeltaTime;
 		if (NewAcceleration.SizeSquared() > FMath::Square(MaxAccel) || bForceMaxAccel)
@@ -1903,19 +1909,6 @@ bool UCharacterMovementComponent::ApplyRequestedMove(FVector& NewVelocity, FVect
 			NewAcceleration = RequestedMoveDir * MaxAccel;
 		}
 
-		// use MaxAcceleration to limit speed increase, 2% buffer
-		const float CurrentSpeedSq = Velocity.SizeSquared();
-		if (bRequestedMoveUseAcceleration && CurrentSpeedSq < FMath::Square(RequestedSpeed * 0.98f))
-		{
-			const float ActualSpeed = FMath::Min(RequestedSpeed, FMath::Sqrt(CurrentSpeedSq) + DeltaTime * MaxAccel);
-			NewVelocity = RequestedMoveDir * ActualSpeed;
-		}
-		else
-		{
-			// don't affect deceleration, agent needs to reach end of path without sliding though
-			NewVelocity = MoveVelocity;
-		}
-
 		return true;
 	}
 
