Relax restriction that root component must be registered in AActor::ActorToWorld
Means it can be used in UCS with a static root component (whose registration is deferred). This is also consistent with other functions like GetActorLocation.

diff --git a/Engine/Source/Runtime/Engine/Private/Actor.cpp b/Engine/Source/Runtime/Engine/Private/Actor.cpp
--- a/Engine/Source/Runtime/Engine/Private/Actor.cpp
+++ b/Engine/Source/Runtime/Engine/Private/Actor.cpp
@@ -1566,14 +1566,7 @@ FTransform AActor::ActorToWorld() const
 	FTransform Result = FTransform::Identity;
 	if( RootComponent != NULL )
 	{
-		if( !GetRootComponent()->IsRegistered() )
-		{
-			UE_LOG(LogActor, Log, TEXT("AActor::ActorToWorld (%s) called on RootComponent that is not registered"), *GetPathName());
-		}
-		else
-		{
-			Result = GetRootComponent()->ComponentToWorld;
-		}
+		Result = RootComponent->ComponentToWorld;
 	}
 	else
 	{
