Fixed bug with water raycasts using incorrect bounds dimensions.

This could have advese effects if a characters bounds were quite 'anisotropic'.

Found as part of TTP# 331733 - UE4: Buoyancy: Pawn freezes after approx 3 seconds without input

diff --git a/Engine/Source/Runtime/Engine/Private/CharacterMovementComponent.cpp b/Engine/Source/Runtime/Engine/Private/CharacterMovementComponent.cpp
--- a/Engine/Source/Runtime/Engine/Private/CharacterMovementComponent.cpp
+++ b/Engine/Source/Runtime/Engine/Private/CharacterMovementComponent.cpp
@@ -1730,10 +1730,9 @@ float UCharacterMovementComponent::ImmersionDepth()
 
 	if ( CharacterOwner && GetPhysicsVolume()->bWaterVolume )
 	{
-		const float CollisionHeight = CharacterOwner->GetSimpleCollisionRadius();
-		const float CollisionRadius = CharacterOwner->GetSimpleCollisionHalfHeight();
+		const float CollisionHalfHeight = CharacterOwner->GetSimpleCollisionHalfHeight();
 
-		if ( (CollisionHeight == 0.f) || (Buoyancy == 0.f) )
+		if ( (CollisionHalfHeight == 0.f) || (Buoyancy == 0.f) )
 		{
 			depth = 1.f;
 		}
@@ -1743,8 +1742,8 @@ float UCharacterMovementComponent::ImmersionDepth()
 			FHitResult Hit(1.f);
 			if ( VolumeBrushComp )
 			{
-				const FVector TraceStart = CharacterOwner->GetActorLocation() + FVector(0.f,0.f,CollisionHeight);
-				const FVector TraceEnd = CharacterOwner->GetActorLocation() - FVector(0.f,0.f,CollisionHeight);
+				const FVector TraceStart = CharacterOwner->GetActorLocation() + FVector(0.f,0.f,CollisionHalfHeight);
+				const FVector TraceEnd = CharacterOwner->GetActorLocation() - FVector(0.f,0.f,CollisionHalfHeight);
 
 				const static FName MovementComp_Character_ImmersionDepthName(TEXT("MovementComp_Character_ImmersionDepth"));
 				FCollisionQueryParams NewTraceParams( MovementComp_Character_ImmersionDepthName, true );
