Fixed glitches with based movement on physics objects

Added a new pre-cloth (but post physics) tick function to UCharacterMovementComponent. This makes sure that physics-enabled bases are ticked before the characters that are standing on them.
Added a new function to update attached physics objects' transforms so they correctly align with characters that are ALSO based on physics objects.
Made ALL component's primary tick functoons dependent on the base tick function so that things like bases driven by timelines work correctly (i.e. without a frame lag).

TTP# 320677 - SHOOTERGAME: PlayerÆs camera gets turned around when walking into a trashcan that has been knocked down.

diff --git a/Engine/Source/Runtime/Engine/Classes/Components/SkeletalMeshComponent.h b/Engine/Source/Runtime/Engine/Classes/Components/SkeletalMeshComponent.h
--- a/Engine/Source/Runtime/Engine/Classes/Components/SkeletalMeshComponent.h
+++ b/Engine/Source/Runtime/Engine/Classes/Components/SkeletalMeshComponent.h
@@ -919,6 +919,9 @@ public:
 
 	/** Calls needed cloth updates */
 	void PreClothTick(float DeltaTime);
+	
+	/** Set physics transforms for all bodies */
+	void ApplyDeltaToAllPhysicsTransforms(const FVector& DeltaLocation, const FQuat& DeltaRotation);
 
 #if WITH_APEX_CLOTHING
 	/** 
diff --git a/Engine/Source/Runtime/Engine/Classes/GameFramework/CharacterMovementComponent.h b/Engine/Source/Runtime/Engine/Classes/GameFramework/CharacterMovementComponent.h
--- a/Engine/Source/Runtime/Engine/Classes/GameFramework/CharacterMovementComponent.h
+++ b/Engine/Source/Runtime/Engine/Classes/GameFramework/CharacterMovementComponent.h
@@ -136,12 +136,39 @@ struct FRootMotionMovementParams
 	}
 };
 
+/** 
+ * Tick function that calls UCharacterMovementComponent::PreClothTick
+ **/
+USTRUCT()
+struct FCharacterMovementComponentPreClothTickFunction : public FTickFunction
+{
+	GENERATED_USTRUCT_BODY()
+
+	/** CharacterMovementComponent that is the target of this tick **/
+	class UCharacterMovementComponent* Target;
+
+	/** 
+	 * Abstract function actually execute the tick. 
+	 * @param DeltaTime - frame time to advance, in seconds
+	 * @param TickType - kind of tick for this frame
+	 * @param CurrentThread - thread we are executing on, useful to pass along as new tasks are created
+	 * @param MyCompletionGraphEvent - completion event for this task. Useful for holding the completion of this task until certain child tasks are complete.
+	 **/
+	virtual void ExecuteTick(float DeltaTime, enum ELevelTick TickType, ENamedThreads::Type CurrentThread, const FGraphEventRef& MyCompletionGraphEvent) OVERRIDE;
+
+	/** Abstract function to describe this tick. Used to print messages about illegal cycles in the dependency graph **/
+	virtual FString DiagnosticMessage() OVERRIDE;
+};
 
 UCLASS(dependson=(UNetworkPredictionInterface))
 class ENGINE_API UCharacterMovementComponent : public UPawnMovementComponent, public INetworkPredictionInterface
 {
 	GENERATED_UCLASS_BODY()
 
+	/** Post-physics tick function for this character */
+	UPROPERTY()
+	struct FCharacterMovementComponentPreClothTickFunction PreClothComponentTick;
+
 protected:
 
 	/** Character movement component belongs to */
@@ -558,6 +585,9 @@ protected:
 	/** if set, PostProcessAvoidanceVelocity will be called */
 	uint32 bUseRVOPostProcess : 1;
 
+	/** Flag set in pre-physics update to indicate that based movement should be updated post-physics */
+	uint32 bDeferUpdateBasedMovement : 1;
+
 	/** forced avoidance velocity, used when AvoidanceLockTimer is > 0 */
 	FVector AvoidanceLockVelocity;
 
@@ -632,6 +662,7 @@ public:
 	virtual void TickComponent(float DeltaTime, enum ELevelTick TickType, FActorComponentTickFunction *ThisTickFunction) OVERRIDE;
 	virtual void OnUnregister() OVERRIDE;
 	virtual void PostLoad() OVERRIDE;
+	virtual void RegisterComponentTickFunctions(bool bRegister) OVERRIDE;
 	//End UActorComponent Interface
 
 	//BEGIN UMovementComponent Interface
@@ -685,12 +716,18 @@ public:
 	/** Adjust distance from floor, trying to maintain a slight offset from the floor when walking (based on CurrentFloor). */
 	virtual void AdjustFloorHeight();
 
+	/** Update or defer updating of position based on Base movement */
+	virtual void MaybeUpdateBasedMovement(float DeltaSeconds);
+
 	/** Update position based on Base movement */
 	virtual void UpdateBasedMovement(float DeltaSeconds);
 
 	/** Update controller's view rotation as pawn's base rotates */
 	virtual void UpdateBasedRotation(FRotator &FinalRotation, const FRotator& ReducedRotation);
 
+	/** Update (or defer updating) OldBaseLocation and OldBaseRotation if there is a valid movement base. */
+	virtual void MaybeSaveBaseLocation();
+
 	/** Update OldBaseLocation and OldBaseRotation if there is a valid movement base. */
 	virtual void SaveBaseLocation();
 
@@ -984,6 +1021,8 @@ public:
 	UFUNCTION(BlueprintCallable, Category="Pawn|Components|CharacterMovement")
 	void SetWalkableFloorZ(float InWalkableFloorZ);
 
+	/** Tick function called after physics (sync scene) has finished simulation, before cloth */
+	virtual void PreClothTick(float DeltaTime, FCharacterMovementComponentPreClothTickFunction& ThisTickFunction);
 
 protected:
 	/** @note Movement update functions should only be called through StartNewPhysics()*/
diff --git a/Engine/Source/Runtime/Engine/Private/Character.cpp b/Engine/Source/Runtime/Engine/Private/Character.cpp
--- a/Engine/Source/Runtime/Engine/Private/Character.cpp
+++ b/Engine/Source/Runtime/Engine/Private/Character.cpp
@@ -385,6 +385,13 @@ namespace MovementBaseUtility
 			if (NewBaseOwner)
 			{
 				BasedObjectTick.AddPrerequisite(NewBaseOwner, NewBaseOwner->PrimaryActorTick);
+
+				TArray<UActorComponent*> Components;
+				NewBaseOwner->GetComponents(Components);
+				for (auto& Component : Components)
+				{
+					BasedObjectTick.AddPrerequisite(Component, Component->PrimaryComponentTick);
+				}
 			}
 		}
 	}
@@ -398,6 +405,13 @@ namespace MovementBaseUtility
 			if (OldBaseOwner)
 			{
 				BasedObjectTick.RemovePrerequisite(OldBaseOwner, OldBaseOwner->PrimaryActorTick);
+
+				TArray<UActorComponent*> Components;
+				OldBaseOwner->GetComponents(Components);
+				for (auto& Component : Components)
+				{
+					BasedObjectTick.RemovePrerequisite(Component, Component->PrimaryComponentTick);
+				}
 			}
 		}
 	}
@@ -497,18 +511,27 @@ void ACharacter::SetBase( UPrimitiveComponent* NewBaseComponent, bool bNotifyPaw
 		if (CharacterMovement)
 		{
 			MovementBaseUtility::RemoveTickDependency(CharacterMovement->PrimaryComponentTick, OldBase);
-			MovementBaseUtility::AddTickDependency(CharacterMovement->PrimaryComponentTick, MovementBase);
+			MovementBaseUtility::AddTickDependency(CharacterMovement->PrimaryComponentTick, MovementBase);			
 
 			if (MovementBase)
 			{
 				// Update OldBaseLocation/Rotation as those were referring to a different base
 				CharacterMovement->OldBaseLocation = MovementBase->GetComponentLocation();
 				CharacterMovement->OldBaseQuat = MovementBase->GetComponentQuat();
+
+				// Enable pre-cloth tick if we are standing on a physics object, as we need to to use 
+				// post-physics transforms
+				if(MovementBase->IsSimulatingPhysics())
+				{
+					CharacterMovement->PreClothComponentTick.SetTickFunctionEnable(true);
+				}
 			}
 			else
 			{
 				// Base is NULL, invalidate floor.
 				CharacterMovement->CurrentFloor.Clear();
+
+				CharacterMovement->PreClothComponentTick.SetTickFunctionEnable(false);
 			}
 		}
 	}
diff --git a/Engine/Source/Runtime/Engine/Private/CharacterMovementComponent.cpp b/Engine/Source/Runtime/Engine/Private/CharacterMovementComponent.cpp
--- a/Engine/Source/Runtime/Engine/Private/CharacterMovementComponent.cpp
+++ b/Engine/Source/Runtime/Engine/Private/CharacterMovementComponent.cpp
@@ -87,10 +87,28 @@ void FFindFloorResult::SetFromLineTrace(const FHitResult& InHit, const float InS
 	}
 }
 
+void FCharacterMovementComponentPreClothTickFunction::ExecuteTick(float DeltaTime, enum ELevelTick TickType, ENamedThreads::Type CurrentThread, const FGraphEventRef& MyCompletionGraphEvent)
+{
+	if ( (TickType == LEVELTICK_All) && Target && !Target->HasAnyFlags(RF_PendingKill | RF_Unreachable))
+	{
+		FScopeCycleCounterUObject ComponentScope(Target);
+		FScopeCycleCounterUObject AdditionalScope(Target->AdditionalStatObject());
+		Target->PreClothTick(DeltaTime, *this);	
+	}
+}
+
+FString FCharacterMovementComponentPreClothTickFunction::DiagnosticMessage()
+{
+	return Target->GetFullName() + TEXT("[UCharacterMovementComponent::PreClothTick]");
+}
 
 UCharacterMovementComponent::UCharacterMovementComponent(const class FPostConstructInitializeProperties& PCIP)
 	: Super(PCIP)
 {
+	PreClothComponentTick.bCanEverTick = true;
+	PreClothComponentTick.bStartWithTickEnabled = false;
+	PreClothComponentTick.TickGroup = TG_PreCloth;
+
 	GravityScale = 1.f;
 	GroundFriction = 8.0f;
 	JumpZVelocity = 420.0f;
@@ -717,8 +735,8 @@ void UCharacterMovementComponent::TickComponent(float DeltaTime, enum ELevelTick
 			// Server ticking for remote client.
 			// Between net updates from the client we need to update position if based on another object,
 			// otherwise the object will move on intermediate frames and we won't follow it.
-			UpdateBasedMovement(DeltaTime);
-			SaveBaseLocation();
+			MaybeUpdateBasedMovement(DeltaTime);
+			MaybeSaveBaseLocation();
 		}
 	}
 	else if (CharacterOwner->Role == ROLE_SimulatedProxy)
@@ -756,6 +774,16 @@ void UCharacterMovementComponent::TickComponent(float DeltaTime, enum ELevelTick
 	}
 }
 
+void UCharacterMovementComponent::PreClothTick(float DeltaTime, FCharacterMovementComponentPreClothTickFunction& ThisTickFunction)
+{
+	if(bDeferUpdateBasedMovement)
+	{
+		UpdateBasedMovement(DeltaTime);
+		SaveBaseLocation();
+
+		bDeferUpdateBasedMovement = false;
+	}
+}
 
 void UCharacterMovementComponent::AdjustProxyCapsuleSize()
 {
@@ -960,7 +988,7 @@ void UCharacterMovementComponent::SimulateMovement(float DeltaSeconds)
 	AnalogInputModifier = 1.0f;				// Not currently used for simulated movement
 
 	// simulated pawns predict location
-	UpdateBasedMovement(DeltaSeconds);
+	MaybeUpdateBasedMovement(DeltaSeconds);
 	FStepDownResult StepDownResult;
 	MoveSmooth(Velocity, DeltaSeconds, &StepDownResult);
 
@@ -1012,7 +1040,8 @@ void UCharacterMovementComponent::SimulateMovement(float DeltaSeconds)
 		}
 	}
 
-	SaveBaseLocation();
+	MaybeSaveBaseLocation();
+	
 	UpdateComponentVelocity();
 	bJustTeleported = false;
 }
@@ -1025,6 +1054,32 @@ void UCharacterMovementComponent::SetBase( UPrimitiveComponent* NewBase, bool bN
 	}
 }
 
+void UCharacterMovementComponent::MaybeUpdateBasedMovement(float DeltaSeconds)
+{
+	bDeferUpdateBasedMovement = false;
+
+	UPrimitiveComponent* MovementBase = CharacterOwner->GetMovementBase();
+	if(MovementBaseUtility::UseRelativePosition(MovementBase))
+	{
+		if(!MovementBase->IsSimulatingPhysics())
+		{
+			UpdateBasedMovement(DeltaSeconds);
+		}
+		else
+		{
+			// defer movement base update until after physics
+			bDeferUpdateBasedMovement = true;
+		}
+	}
+}
+
+void UCharacterMovementComponent::MaybeSaveBaseLocation()
+{
+	if(!bDeferUpdateBasedMovement)
+	{
+		SaveBaseLocation();
+	}
+}
 
 // @todo UE4 - handle lift moving up and down through encroachment
 void UCharacterMovementComponent::UpdateBasedMovement(float DeltaSeconds)
@@ -1050,6 +1105,7 @@ void UCharacterMovementComponent::UpdateBasedMovement(float DeltaSeconds)
 	TGuardValue<EMoveComponentFlags> ScopedFlagRestore(MoveComponentFlags, MoveComponentFlags | MOVECOMP_IgnoreBases);
 
 	FQuat DeltaQuat = FQuat::Identity;
+	FVector DeltaPosition = FVector::ZeroVector;
 	const bool bRotationChanged = !OldBaseQuat.Equals(MovementBase->GetComponentQuat());
 	if( bRotationChanged )
 	{
@@ -1093,7 +1149,7 @@ void UCharacterMovementComponent::UpdateBasedMovement(float DeltaSeconds)
 			FVector const BaseOffset(0.0f, 0.0f, HalfHeight);
 			FVector const LocalBasePos = OldLocalToWorld.InverseTransformPosition(CharacterOwner->GetActorLocation() - BaseOffset);
 			FVector const NewWorldPos = ConstrainLocationToPlane(NewLocalToWorld.TransformPosition(LocalBasePos) + BaseOffset);
-			FVector const Delta = ConstrainDirectionToPlane(NewWorldPos - CharacterOwner->GetActorLocation());
+			DeltaPosition = ConstrainDirectionToPlane(NewWorldPos - CharacterOwner->GetActorLocation());
 
 			// move attached actor
 			if (bFastAttachedMove)
@@ -1103,9 +1159,14 @@ void UCharacterMovementComponent::UpdateBasedMovement(float DeltaSeconds)
 			}
 			else
 			{
-				MoveUpdatedComponent( Delta, FinalQuat.Rotator(), true );
+				MoveUpdatedComponent( DeltaPosition, FinalQuat.Rotator(), true );
 			}
 		}
+
+		if(MovementBase->IsSimulatingPhysics())
+		{
+			CharacterOwner->Mesh->ApplyDeltaToAllPhysicsTransforms(DeltaPosition, DeltaQuat);
+		}
 	}
 }
 
@@ -1162,8 +1223,8 @@ void UCharacterMovementComponent::PerformMovement(float DeltaSeconds)
 
 	FScopedMovementUpdate ScopedMovementUpdate(UpdatedComponent, bEnableScopedMovementUpdates ? EScopedUpdate::DeferredUpdates : EScopedUpdate::ImmediateUpdates);
 
-	UpdateBasedMovement(DeltaSeconds);
-
+	MaybeUpdateBasedMovement(DeltaSeconds);
+	
 	ApplyAccumulatedForces(DeltaSeconds);
 
 	FVector OldVelocity = Velocity;
@@ -1289,7 +1350,7 @@ void UCharacterMovementComponent::PerformMovement(float DeltaSeconds)
 	}
 
 	// update relative location/rotation
-	SaveBaseLocation();
+	MaybeSaveBaseLocation();
 
 	// update component velocity
 	UpdateComponentVelocity();
@@ -5660,6 +5721,32 @@ void UCharacterMovementComponent::AddRadialImpulse(const FVector& Origin, float
 	AddImpulse(Delta * ImpulseMagnitude, bVelChange);
 }
 
+void UCharacterMovementComponent::RegisterComponentTickFunctions(bool bRegister)
+{
+	Super::RegisterComponentTickFunctions(bRegister);
+
+	if (bRegister)
+	{
+		if (SetupActorComponentTickFunction(&PreClothComponentTick))
+		{
+			PreClothComponentTick.Target = this;
+
+			// If primary tick is registered, add a prerequisite to it
+			if(PrimaryComponentTick.bCanEverTick)
+			{
+				PreClothComponentTick.AddPrerequisite(this, PrimaryComponentTick); 
+			}
+		}
+	}
+	else
+	{
+		if(PreClothComponentTick.IsTickFunctionRegistered())
+		{
+			PreClothComponentTick.UnRegisterTickFunction();
+		}
+	}
+}
+
 FNetworkPredictionData_Client_Character::FNetworkPredictionData_Client_Character()
 	: ClientUpdateTime(0.f)
 	, CurrentTimeStamp(0.f)
diff --git a/Engine/Source/Runtime/Engine/Private/SkeletalMeshComponentPhysics.cpp b/Engine/Source/Runtime/Engine/Private/SkeletalMeshComponentPhysics.cpp
--- a/Engine/Source/Runtime/Engine/Private/SkeletalMeshComponentPhysics.cpp
+++ b/Engine/Source/Runtime/Engine/Private/SkeletalMeshComponentPhysics.cpp
@@ -551,6 +551,42 @@ void USkeletalMeshComponent::SetAllPhysicsRotation(FRotator NewRot)
 	}
 }
 
+void USkeletalMeshComponent::ApplyDeltaToAllPhysicsTransforms(const FVector& DeltaLocation, const FQuat& DeltaRotation)
+{
+	if(RootBodyIndex < Bodies.Num())
+	{
+		// calculate the deltas to get the root body to NewRot
+		FBodyInstance* RootBI = Bodies[RootBodyIndex];
+		check(RootBI);
+		if(RootBI->IsValidBodyInstance())
+		{
+			// move the root body
+			FTransform RootBodyTM = RootBI->GetUnrealWorldTransform();
+			RootBodyTM.SetRotation(RootBodyTM.GetRotation() * DeltaRotation);
+			RootBodyTM.SetTranslation(RootBodyTM.GetTranslation() + DeltaLocation);
+			RootBI->SetBodyTransform(RootBodyTM, true);
+
+			// apply the delta to all the other bodies
+			for (int32 i = 0; i < Bodies.Num(); i++)
+			{
+				if (i != RootBodyIndex)
+				{
+					FBodyInstance* BI = Bodies[i];
+					check(BI);
+
+					FTransform BodyTM = BI->GetUnrealWorldTransform();
+					BodyTM.SetRotation(BodyTM.GetRotation() * DeltaRotation);
+					BodyTM.SetTranslation(BodyTM.GetTranslation() + DeltaLocation);
+					BI->SetBodyTransform( BodyTM, true );
+				}
+			}
+
+			// Move component to new physics location
+			SyncComponentToRBPhysics();
+		}
+	}
+}
+
 void USkeletalMeshComponent::SetPhysMaterialOverride(UPhysicalMaterial* NewPhysMaterial)
 {
 	// Single-body case - just use PrimComp code.
