Add EngineDefines.h, EngineGlobals.h, EngineLogs.h, EngineModule.h to break up Engine.h more

diff --git a/Engine/Source/Runtime/Engine/Public/Engine.h b/Engine/Source/Runtime/Engine/Public/Engine.h
--- a/Engine/Source/Runtime/Engine/Public/Engine.h
+++ b/Engine/Source/Runtime/Engine/Public/Engine.h
@@ -6,79 +6,7 @@
 
 #pragma once
 
-/*-----------------------------------------------------------------------------
-	Configuration defines
------------------------------------------------------------------------------*/
-
-/** 
- *   Whether or not compiling with PhysX
- */
-#ifndef WITH_PHYSX
-	#define WITH_PHYSX 1
-#endif
-
-/** 
- *   Whether or not compiling with APEX extensions to PhysX
- */
-#ifndef WITH_APEX
-	#define WITH_APEX (1 && WITH_PHYSX)
-#endif
-
-#ifndef WITH_PHYSICS_COOKING
-	#define WITH_PHYSICS_COOKING (WITH_EDITOR || WITH_APEX)		//APEX currently relies on cooking even at runtime
-#endif
-
-#if WITH_APEX
-#ifndef WITH_SUBSTEPPING
-	#define WITH_SUBSTEPPING 1
-#endif
-
-#ifndef WITH_APEX_CLOTHING
-	#define WITH_APEX_CLOTHING	(1 && WITH_APEX)
-#endif // WITH_APEX_CLOTHING
-
-#ifndef WITH_APEX_LEGACY
-	#define WITH_APEX_LEGACY	1
-#endif // WITH_APEX_LEGACY
-
-#endif // WITH_APEX
-
-#if WITH_APEX_CLOTHING
-#ifndef WITH_CLOTH_COLLISION_DETECTION
-	#define WITH_CLOTH_COLLISION_DETECTION (1 && WITH_APEX_CLOTHING)
-#endif//WITH_CLOTH_COLLISION_DETECTION
-#endif //WITH_APEX_CLOTHING
-
-#ifndef WITH_BODY_WELDING
-	//#define WITH_BODY_WELDING 1
-#endif
-
-#ifndef ENABLE_VISUAL_LOG
-	#define ENABLE_VISUAL_LOG (1 && !NO_LOGGING && !USING_CODE_ANALYSIS && !(UE_BUILD_SHIPPING || UE_BUILD_TEST))
-#endif
-
-#ifndef WITH_FIXED_AREA_ENTERING_COST
-	#define WITH_FIXED_AREA_ENTERING_COST 1
-#endif // WITH_FIXED_AREA_ENTERING_COST
-
-// If set, recast will use async workers for rebuilding tiles in runtime
-// All access to tile data must be guarded with critical sections
-#ifndef RECAST_ASYNC_REBUILDING
-	#define RECAST_ASYNC_REBUILDING	1
-#endif
-
-/*-----------------------------------------------------------------------------
-	Size of the world.
------------------------------------------------------------------------------*/
-
-#define WORLD_MAX					524288.0				/* Maximum size of the world */
-#define HALF_WORLD_MAX				(WORLD_MAX * 0.5)		/* Half the maximum size of the world */
-#define HALF_WORLD_MAX1				(HALF_WORLD_MAX - 1.0)	/* Half the maximum size of the world minus one */
-#define DEFAULT_ORTHOZOOM			10000.0					/* Default 2D viewport zoom */
-
-/*-----------------------------------------------------------------------------
-	Dependencies.
------------------------------------------------------------------------------*/
+#include "EngineDefines.h"
 
 #include "Core.h"
 #include "CoreUObject.h"
@@ -91,36 +19,9 @@
 #include "EngineSettings.h"
 #include "EditorSupportDelegates.h"
 #include "EngineStats.h"
+#include "EngineLogs.h"
 
-
-ENGINE_API DECLARE_LOG_CATEGORY_EXTERN(LogMCP, Warning, All);
-ENGINE_API DECLARE_LOG_CATEGORY_EXTERN(LogPath, Warning, All);
-ENGINE_API DECLARE_LOG_CATEGORY_EXTERN(LogPhysics, Warning, All);
-ENGINE_API DECLARE_LOG_CATEGORY_EXTERN(LogBlueprint, Warning, All);
-ENGINE_API DECLARE_LOG_CATEGORY_EXTERN(LogDestructible, Warning, All);
-ENGINE_API DECLARE_LOG_CATEGORY_EXTERN(LogBlueprintUserMessages, Log, All);
-ENGINE_API DECLARE_LOG_CATEGORY_EXTERN(LogAnimation, Warning, All);
-ENGINE_API DECLARE_LOG_CATEGORY_EXTERN(LogRootMotion, Warning, All);
-ENGINE_API DECLARE_LOG_CATEGORY_EXTERN(LogLevel, Log, All);
-ENGINE_API DECLARE_LOG_CATEGORY_EXTERN(LogSkeletalMesh, Log, All);
-ENGINE_API DECLARE_LOG_CATEGORY_EXTERN(LogStaticMesh, Log, All);
-ENGINE_API DECLARE_LOG_CATEGORY_EXTERN(LogNet, Log, All);
-ENGINE_API DECLARE_LOG_CATEGORY_EXTERN(LogNetPlayerMovement, Warning, All);
-ENGINE_API DECLARE_LOG_CATEGORY_EXTERN(LogNetTraffic, Warning, All);
-ENGINE_API DECLARE_LOG_CATEGORY_EXTERN(LogNetDormancy, Warning, All);
-ENGINE_API DECLARE_LOG_CATEGORY_EXTERN(LogSubtitle, Log, All);
-ENGINE_API DECLARE_LOG_CATEGORY_EXTERN(LogTexture, Log, All);
-ENGINE_API DECLARE_LOG_CATEGORY_EXTERN(LogLandscape, Warning, All);
-ENGINE_API DECLARE_LOG_CATEGORY_EXTERN(LogPlayerManagement, Error, All);
-
-/** Global engine pointer. Can be 0 so don't use without checking. */
-extern ENGINE_API class UEngine*			GEngine;
-
-/** when set, disallows all network travel (pending level rejects all client travel attempts) */
-extern ENGINE_API bool						GDisallowNetworkTravel;
-
-/** The GPU time taken to render the last frame. Same metric as FPlatformTime::Cycles(). */
-extern ENGINE_API uint32					GGPUFrameTime;
+#include "EngineGlobals.h"
 
 
 /*-----------------------------------------------------------------------------
@@ -1143,20 +1044,7 @@ namespace physx
 #include "AI/NavigationSystemHelpers.h"
 #include "HardwareInfo.h"
 
-/** Implements the engine module. */
-class FEngineModule : public FDefaultModuleImpl
-{
-public:
-	
-	// IModuleInterface
-	virtual void StartupModule();
-};
-
-/** Accessor that gets the renderer module and caches the result. */
-extern ENGINE_API IRendererModule& GetRendererModule();
-
-/** Clears the cached renderer module reference. */
-extern ENGINE_API void ResetCachedRendererModule();
+#include "EngineModule.h"
 
 
 
diff --git a/Engine/Source/Runtime/Engine/Public/EngineDefines.h b/Engine/Source/Runtime/Engine/Public/EngineDefines.h
new file mode 100644
--- /dev/null
+++ b/Engine/Source/Runtime/Engine/Public/EngineDefines.h
@@ -0,0 +1,76 @@
+// Copyright 1998-2014 Epic Games, Inc. All Rights Reserved.
+// 
+// Public defines form the Engine
+
+#pragma once
+
+/*-----------------------------------------------------------------------------
+	Configuration defines
+-----------------------------------------------------------------------------*/
+
+/** 
+ *   Whether or not compiling with PhysX
+ */
+#ifndef WITH_PHYSX
+	#define WITH_PHYSX 1
+#endif
+
+/** 
+ *   Whether or not compiling with APEX extensions to PhysX
+ */
+#ifndef WITH_APEX
+	#define WITH_APEX (1 && WITH_PHYSX)
+#endif
+
+#ifndef WITH_PHYSICS_COOKING
+	#define WITH_PHYSICS_COOKING (WITH_EDITOR || WITH_APEX)		//APEX currently relies on cooking even at runtime
+#endif
+
+#if WITH_APEX
+#ifndef WITH_SUBSTEPPING
+	#define WITH_SUBSTEPPING 1
+#endif
+
+#ifndef WITH_APEX_CLOTHING
+	#define WITH_APEX_CLOTHING	(1 && WITH_APEX)
+#endif // WITH_APEX_CLOTHING
+
+#ifndef WITH_APEX_LEGACY
+	#define WITH_APEX_LEGACY	1
+#endif // WITH_APEX_LEGACY
+
+#endif // WITH_APEX
+
+#if WITH_APEX_CLOTHING
+#ifndef WITH_CLOTH_COLLISION_DETECTION
+	#define WITH_CLOTH_COLLISION_DETECTION (1 && WITH_APEX_CLOTHING)
+#endif//WITH_CLOTH_COLLISION_DETECTION
+#endif //WITH_APEX_CLOTHING
+
+#ifndef WITH_BODY_WELDING
+	//#define WITH_BODY_WELDING 1
+#endif
+
+#ifndef ENABLE_VISUAL_LOG
+	#define ENABLE_VISUAL_LOG (1 && !NO_LOGGING && !USING_CODE_ANALYSIS && !(UE_BUILD_SHIPPING || UE_BUILD_TEST))
+#endif
+
+#ifndef WITH_FIXED_AREA_ENTERING_COST
+	#define WITH_FIXED_AREA_ENTERING_COST 1
+#endif // WITH_FIXED_AREA_ENTERING_COST
+
+// If set, recast will use async workers for rebuilding tiles in runtime
+// All access to tile data must be guarded with critical sections
+#ifndef RECAST_ASYNC_REBUILDING
+	#define RECAST_ASYNC_REBUILDING	1
+#endif
+
+/*-----------------------------------------------------------------------------
+	Size of the world.
+-----------------------------------------------------------------------------*/
+
+#define WORLD_MAX					524288.0				/* Maximum size of the world */
+#define HALF_WORLD_MAX				(WORLD_MAX * 0.5)		/* Half the maximum size of the world */
+#define HALF_WORLD_MAX1				(HALF_WORLD_MAX - 1.0)	/* Half the maximum size of the world minus one */
+
+#define DEFAULT_ORTHOZOOM			10000.0					/* Default 2D viewport zoom */
\ No newline at end of file
diff --git a/Engine/Source/Runtime/Engine/Public/EngineGlobals.h b/Engine/Source/Runtime/Engine/Public/EngineGlobals.h
new file mode 100644
--- /dev/null
+++ b/Engine/Source/Runtime/Engine/Public/EngineGlobals.h
@@ -0,0 +1,14 @@
+// Copyright 1998-2014 Epic Games, Inc. All Rights Reserved.
+// 
+// Engine global defines
+
+#pragma once
+
+/** Global engine pointer. Can be 0 so don't use without checking. */
+extern ENGINE_API class UEngine*			GEngine;
+
+/** when set, disallows all network travel (pending level rejects all client travel attempts) */
+extern ENGINE_API bool						GDisallowNetworkTravel;
+
+/** The GPU time taken to render the last frame. Same metric as FPlatformTime::Cycles(). */
+extern ENGINE_API uint32					GGPUFrameTime;
\ No newline at end of file
diff --git a/Engine/Source/Runtime/Engine/Public/EngineLogs.h b/Engine/Source/Runtime/Engine/Public/EngineLogs.h
new file mode 100644
--- /dev/null
+++ b/Engine/Source/Runtime/Engine/Public/EngineLogs.h
@@ -0,0 +1,26 @@
+// Copyright 1998-2014 Epic Games, Inc. All Rights Reserved.
+// 
+// Log categories used by the Engine
+
+#pragma once
+
+ 
+ENGINE_API DECLARE_LOG_CATEGORY_EXTERN(LogMCP, Warning, All);
+ENGINE_API DECLARE_LOG_CATEGORY_EXTERN(LogPath, Warning, All);
+ENGINE_API DECLARE_LOG_CATEGORY_EXTERN(LogPhysics, Warning, All);
+ENGINE_API DECLARE_LOG_CATEGORY_EXTERN(LogBlueprint, Warning, All);
+ENGINE_API DECLARE_LOG_CATEGORY_EXTERN(LogDestructible, Warning, All);
+ENGINE_API DECLARE_LOG_CATEGORY_EXTERN(LogBlueprintUserMessages, Log, All);
+ENGINE_API DECLARE_LOG_CATEGORY_EXTERN(LogAnimation, Warning, All);
+ENGINE_API DECLARE_LOG_CATEGORY_EXTERN(LogRootMotion, Warning, All);
+ENGINE_API DECLARE_LOG_CATEGORY_EXTERN(LogLevel, Log, All);
+ENGINE_API DECLARE_LOG_CATEGORY_EXTERN(LogSkeletalMesh, Log, All);
+ENGINE_API DECLARE_LOG_CATEGORY_EXTERN(LogStaticMesh, Log, All);
+ENGINE_API DECLARE_LOG_CATEGORY_EXTERN(LogNet, Log, All);
+ENGINE_API DECLARE_LOG_CATEGORY_EXTERN(LogNetPlayerMovement, Warning, All);
+ENGINE_API DECLARE_LOG_CATEGORY_EXTERN(LogNetTraffic, Warning, All);
+ENGINE_API DECLARE_LOG_CATEGORY_EXTERN(LogNetDormancy, Warning, All);
+ENGINE_API DECLARE_LOG_CATEGORY_EXTERN(LogSubtitle, Log, All);
+ENGINE_API DECLARE_LOG_CATEGORY_EXTERN(LogTexture, Log, All);
+ENGINE_API DECLARE_LOG_CATEGORY_EXTERN(LogLandscape, Warning, All);
+ENGINE_API DECLARE_LOG_CATEGORY_EXTERN(LogPlayerManagement, Error, All);
\ No newline at end of file
diff --git a/Engine/Source/Runtime/Engine/Public/EngineModule.h b/Engine/Source/Runtime/Engine/Public/EngineModule.h
new file mode 100644
--- /dev/null
+++ b/Engine/Source/Runtime/Engine/Public/EngineModule.h
@@ -0,0 +1,18 @@
+// Copyright 1998-2014 Epic Games, Inc. All Rights Reserved.
+// 
+// Engine module class
+
+/** Implements the engine module. */
+class FEngineModule : public FDefaultModuleImpl
+{
+public:
+
+	// IModuleInterface
+	virtual void StartupModule();
+};
+
+/** Accessor that gets the renderer module and caches the result. */
+extern ENGINE_API IRendererModule& GetRendererModule();
+
+/** Clears the cached renderer module reference. */
+extern ENGINE_API void ResetCachedRendererModule();
\ No newline at end of file
