Landscape Gizmo actor can no longer be attached to other actors

diff --git a/Engine/Source/Editor/UnrealEd/Private/Editor.cpp b/Engine/Source/Editor/UnrealEd/Private/Editor.cpp
--- a/Engine/Source/Editor/UnrealEd/Private/Editor.cpp
+++ b/Engine/Source/Editor/UnrealEd/Private/Editor.cpp
@@ -4038,7 +4038,26 @@ bool UEditorEngine::CanParentActors( const AActor* ParentActor, const AActor* Ch
 		return false;
 	}
 
-	if(ChildRoot->Mobility == EComponentMobility::Static && ParentRoot->Mobility != EComponentMobility::Static )
+	{
+		FText Reason;
+		if (!ChildActor->EditorCanAttachTo(ParentActor, Reason))
+		{
+			if (ReasonText)
+			{
+				if (Reason.IsEmpty())
+				{
+					*ReasonText = FText::Format(NSLOCTEXT("ActorAttachmentError", "CannotBeAttached_ActorAttachmentError", "{0} cannot be attached to {1}"), FText::FromString(ChildActor->GetActorLabel()), FText::FromString(ParentActor->GetActorLabel()));
+				}
+				else
+				{
+					*ReasonText = MoveTemp(Reason);
+				}
+			}
+			return false;
+		}
+	}
+
+	if (ChildRoot->Mobility == EComponentMobility::Static && ParentRoot->Mobility != EComponentMobility::Static)
 	{
 		if (ReasonText)
 		{
diff --git a/Engine/Source/Runtime/Engine/Classes/GameFramework/Actor.h b/Engine/Source/Runtime/Engine/Classes/GameFramework/Actor.h
--- a/Engine/Source/Runtime/Engine/Classes/GameFramework/Actor.h
+++ b/Engine/Source/Runtime/Engine/Classes/GameFramework/Actor.h
@@ -424,36 +424,36 @@ private:
 public:
 	/** Is hidden within the editor at its startup. */
 	UPROPERTY()
-	uint32 bHiddenEd:1;    
+	uint32 bHiddenEd:1;
 
 protected:
 	/** Whether the actor can be manipulated by editor operations. */
 	UPROPERTY()
-	uint32 bEditable:1;    
+	uint32 bEditable:1;
 
 	/** True if this actor should be listed in the scene outliner */
 	UPROPERTY()
-	uint32 bListedInSceneOutliner:1;    
+	uint32 bListedInSceneOutliner:1;
 
 public:
 	/** Is hidden by the layer browser. */
 	UPROPERTY()
-	uint32 bHiddenEdLayer:1;    
+	uint32 bHiddenEdLayer:1;
 
 private:
 	/** Is temporarily hidden within the editor; used for show/hide/etc. functionality w/o dirtying the actor */
 	UPROPERTY(transient)
-	uint32 bHiddenEdTemporary:1;    
+	uint32 bHiddenEdTemporary:1;
 
 public:
 
 	/** Is hidden by the level browser. */
 	UPROPERTY(transient)
-	uint32 bHiddenEdLevel:1;    
+	uint32 bHiddenEdLevel:1;
 
 	/** Prevent the actor from being moved in the editor. */
 	UPROPERTY()
-	uint32 bLockLocation:1;    
+	uint32 bLockLocation:1;
 
 	// Actor's layer name.
 	UPROPERTY()
@@ -1247,6 +1247,9 @@ public:
 	/** @return	Returns true if this actor should be shown in the scene outliner */
 	bool IsListedInSceneOutliner() const;
 
+	/** @return	Returns true if this actor is allowed to be attached to the given actor */
+	virtual bool EditorCanAttachTo(const AActor* InParent, FText& OutReason) const;
+
 	// Called before editor copy, true allow export
 	virtual bool ShouldExport() { return true; }
 
diff --git a/Engine/Source/Runtime/Engine/Classes/Landscape/LandscapeGizmoActor.h b/Engine/Source/Runtime/Engine/Classes/Landscape/LandscapeGizmoActor.h
--- a/Engine/Source/Runtime/Engine/Classes/Landscape/LandscapeGizmoActor.h
+++ b/Engine/Source/Runtime/Engine/Classes/Landscape/LandscapeGizmoActor.h
@@ -39,6 +39,11 @@ class ALandscapeGizmoActor : public AActor
 #if WITH_EDITOR
 	virtual void Duplicate(ALandscapeGizmoActor* Gizmo); 
 	//virtual void EditorApplyTranslation(const FVector& DeltaTranslation, bool bAltDown, bool bShiftDown, bool bCtrlDown);
+
+	bool EditorCanAttachTo(const AActor* InParent, FText& OutReason) const OVERRIDE
+	{
+		return false;
+	}
 #endif
 
 	/** 
diff --git a/Engine/Source/Runtime/Engine/Private/ActorEditor.cpp b/Engine/Source/Runtime/Engine/Private/ActorEditor.cpp
--- a/Engine/Source/Runtime/Engine/Private/ActorEditor.cpp
+++ b/Engine/Source/Runtime/Engine/Private/ActorEditor.cpp
@@ -360,6 +360,11 @@ bool AActor::IsListedInSceneOutliner() const
 	return bListedInSceneOutliner;
 }
 
+bool AActor::EditorCanAttachTo(const AActor* InParent, FText& OutReason) const
+{
+	return true;
+}
+
 const FString& AActor::GetActorLabel() const
 {
 	// If the label string is empty then we'll use the default actor label (usually the actor's class name.)
