Move forward-declarations out of Engine.h into places they are needed

diff --git a/Engine/Source/Runtime/Engine/Classes/Components/ActorComponent.h b/Engine/Source/Runtime/Engine/Classes/Components/ActorComponent.h
--- a/Engine/Source/Runtime/Engine/Classes/Components/ActorComponent.h
+++ b/Engine/Source/Runtime/Engine/Classes/Components/ActorComponent.h
@@ -386,7 +386,7 @@ public:
 	}
 
 	/** @name Accessors. */
-	FSceneInterface* GetScene() const;
+	class FSceneInterface* GetScene() const;
 
 	/** Return the ULevel that this Component is part of. */
 	ULevel* GetComponentLevel() const;
diff --git a/Engine/Source/Runtime/Engine/Classes/Components/SkeletalMeshComponent.h b/Engine/Source/Runtime/Engine/Classes/Components/SkeletalMeshComponent.h
--- a/Engine/Source/Runtime/Engine/Classes/Components/SkeletalMeshComponent.h
+++ b/Engine/Source/Runtime/Engine/Classes/Components/SkeletalMeshComponent.h
@@ -5,22 +5,16 @@
 #include "PhysicsEngine/PhysicsConstraintComponent.h"
 #include "SkeletalMeshComponent.generated.h"
 
-namespace physx{ 
-
-	class PxPlane;
-
-	namespace apex {
-	class NxClothingAsset;
-	class NxClothingActor;
-
-  //for clothing collision
-	class NxClothingCollision;
-	class NxClothingPlane;
-	class NxClothingSphere;
-	class NxClothingConvex;
-	class NxClothingCollision;
-	class NxClothingCapsule;
-}}
+#if WITH_APEX
+namespace physx
+{ 
+	namespace apex 
+	{
+		class NxClothingActor;
+		class NxClothingCollision;
+	}
+}
+#endif // WITH_APEX
 
 class FPhysScene;
 
diff --git a/Engine/Source/Runtime/Engine/Classes/Engine/World.h b/Engine/Source/Runtime/Engine/Classes/Engine/World.h
--- a/Engine/Source/Runtime/Engine/Classes/Engine/World.h
+++ b/Engine/Source/Runtime/Engine/Classes/Engine/World.h
@@ -13,6 +13,7 @@
 #include "World.generated.h"
 
 class FPhysScene;
+class FSceneViewFamily;
 
 template<typename,typename> class TOctree;
 
@@ -599,7 +600,7 @@ public:
 	void SetNavigationSystem( UNavigationSystem* InNavigationSystem);
 
 	/** The interface to the scene manager for this world. */
-	FSceneInterface*							Scene;
+	class FSceneInterface*						Scene;
 
 	/** Saved editor viewport states - one for each view type. Indexed using ELevelViewportType above.							*/
 	FLevelViewportInfo							EditorViews[4];
diff --git a/Engine/Source/Runtime/Engine/Classes/Landscape/LandscapeHeightfieldCollisionComponent.h b/Engine/Source/Runtime/Engine/Classes/Landscape/LandscapeHeightfieldCollisionComponent.h
--- a/Engine/Source/Runtime/Engine/Classes/Landscape/LandscapeHeightfieldCollisionComponent.h
+++ b/Engine/Source/Runtime/Engine/Classes/Landscape/LandscapeHeightfieldCollisionComponent.h
@@ -4,6 +4,14 @@
 #pragma once
 #include "LandscapeHeightfieldCollisionComponent.generated.h"
 
+#if WITH_PHYSX
+namespace physx
+{
+	class PxMaterial;
+	class PxHeightField;
+}
+#endif // WITH_PHYSX
+
 UCLASS(MinimalAPI)
 class ULandscapeHeightfieldCollisionComponent : public UPrimitiveComponent
 {
@@ -58,7 +66,7 @@ class ULandscapeHeightfieldCollisionComponent : public UPrimitiveComponent
 
 #if WITH_PHYSX
 		/** List of PxMaterials used on this landscape */
-		TArray<class physx::PxMaterial*>	UsedPhysicalMaterialArray;
+		TArray<physx::PxMaterial*>	UsedPhysicalMaterialArray;
 		class physx::PxHeightField*			RBHeightfield;
 #if WITH_EDITOR
 		class physx::PxHeightField*			RBHeightfieldEd; // Used only by landscape editor, does not have holes in it
diff --git a/Engine/Source/Runtime/Engine/Classes/Landscape/LandscapeMeshCollisionComponent.h b/Engine/Source/Runtime/Engine/Classes/Landscape/LandscapeMeshCollisionComponent.h
--- a/Engine/Source/Runtime/Engine/Classes/Landscape/LandscapeMeshCollisionComponent.h
+++ b/Engine/Source/Runtime/Engine/Classes/Landscape/LandscapeMeshCollisionComponent.h
@@ -4,6 +4,14 @@
 #pragma once
 #include "LandscapeMeshCollisionComponent.generated.h"
 
+#if WITH_PHYSX
+namespace physx
+{
+	class PxMaterial;
+	class PxTriangleMesh;
+}
+#endif // WITH_PHYSX
+
 UCLASS()
 class ULandscapeMeshCollisionComponent : public ULandscapeHeightfieldCollisionComponent
 {
@@ -21,10 +29,10 @@ class ULandscapeMeshCollisionComponent : public ULandscapeHeightfieldCollisionCo
 
 #if WITH_PHYSX
 		/** List of PxMaterials used on this landscape */
-		TArray<class physx::PxMaterial*>	UsedPhysicalMaterialArray;
-		class physx::PxTriangleMesh*		RBTriangleMesh;
+		TArray<physx::PxMaterial*>	UsedPhysicalMaterialArray;
+		physx::PxTriangleMesh*		RBTriangleMesh;
 #if WITH_EDITOR
-		class physx::PxTriangleMesh*		RBTriangleMeshEd; // Used only by landscape editor, does not have holes in it
+		physx::PxTriangleMesh*		RBTriangleMeshEd; // Used only by landscape editor, does not have holes in it
 #endif	//WITH_EDITOR
 #endif	//WITH_PHYSX
 
diff --git a/Engine/Source/Runtime/Engine/Classes/PhysicsEngine/BodyInstance.h b/Engine/Source/Runtime/Engine/Classes/PhysicsEngine/BodyInstance.h
--- a/Engine/Source/Runtime/Engine/Classes/PhysicsEngine/BodyInstance.h
+++ b/Engine/Source/Runtime/Engine/Classes/PhysicsEngine/BodyInstance.h
@@ -4,6 +4,18 @@
 
 #include "BodyInstance.generated.h"
 
+#if WITH_PHYSX
+namespace physx
+{
+	class PxRigidActor;
+	class PxAggregate;
+	class PxRigidDynamic;
+	class PxGeometry;
+	class PxShape;
+	class PxTransform;
+}
+#endif // WITH_PHYSX
+
 USTRUCT()
 struct ENGINE_API FCollisionResponse
 {
@@ -202,15 +214,15 @@ public:
 public:
 #if WITH_PHYSX
 		/** Internal use. Physics-engine representation of this body in the synchronous scene. */
-		class physx::PxRigidActor*			RigidActorSync;
+		physx::PxRigidActor*			RigidActorSync;
 
 		/** Internal use. Physics-engine representation of this body in the asynchronous scene. */
-		class physx::PxRigidActor*			RigidActorAsync;
+		physx::PxRigidActor*			RigidActorAsync;
 
 		/** Internal use. Physics-engine representation of a PxAggregate for this body, in case it has alot of shapes. */
-		class physx::PxAggregate*			BodyAggregate;
+		physx::PxAggregate*				BodyAggregate;
 
-		TSharedPtr<TArray<ANSICHAR>>		CharDebugName;
+		TSharedPtr<TArray<ANSICHAR>>	CharDebugName;
 #endif	//WITH_PHYSX
 
 		/** PrimitiveComponent containing this body.   */
@@ -237,7 +249,7 @@ public:
 		void LoadProfileData(bool bVerifyProfile);
 
 #if WITH_PHYSX
-		void InitBody(class UBodySetup* Setup, const FTransform& Transform, class UPrimitiveComponent* PrimComp,class FPhysScene* InRBScene, class physx::PxAggregate* InAggregate = NULL);
+		void InitBody(class UBodySetup* Setup, const FTransform& Transform, class UPrimitiveComponent* PrimComp,class FPhysScene* InRBScene, physx::PxAggregate* InAggregate = NULL);
 #endif	//WITH_PHYSX
 
 		void TermBody();
@@ -281,15 +293,15 @@ public:
 		 * If SceneType < 0, the PST_Sync actor is returned if it is not NULL, otherwise the PST_Async actor is returned.
 		 * Invalid scene types will cause NULL to be returned.
 		 */
-		class physx::PxRigidActor* GetPxRigidActor(int32 SceneType = -1) const;
+		physx::PxRigidActor* GetPxRigidActor(int32 SceneType = -1) const;
 		/** Return the PxRigidDynamic if it exists in one of the scenes (NULL otherwise).  Currently a PxRigidDynamic can exist in only one of the two scenes. */
-		class physx::PxRigidDynamic* GetPxRigidDynamic() const;
+		physx::PxRigidDynamic* GetPxRigidDynamic() const;
 
 		/** 
 		 *	Utility to get all the shapes from a FBodyInstance 
 		 *	Shapes belonging to sync actor are first, then async. Number of shapes belonging to sync actor is returned.
 		 */
-		TArray<class physx::PxShape*> GetAllShapes(int32& OutNumSyncShapes) const;
+		TArray<physx::PxShape*> GetAllShapes(int32& OutNumSyncShapes) const;
 #endif	//WITH_PHYSX
 
 		/** Returns the body's mass */
diff --git a/Engine/Source/Runtime/Engine/Classes/PhysicsEngine/ConstraintInstance.h b/Engine/Source/Runtime/Engine/Classes/PhysicsEngine/ConstraintInstance.h
--- a/Engine/Source/Runtime/Engine/Classes/PhysicsEngine/ConstraintInstance.h
+++ b/Engine/Source/Runtime/Engine/Classes/PhysicsEngine/ConstraintInstance.h
@@ -3,6 +3,13 @@
 
 #include "ConstraintInstance.generated.h"
 
+#if WITH_PHYSX
+namespace physx
+{
+	class PxD6Joint;
+}
+#endif // WITH_PHYSX
+
 // LINEAR DOF
 UENUM()
 enum ELinearConstraintMotion
@@ -46,7 +53,7 @@ struct ENGINE_API FConstraintInstance
 
 #if WITH_PHYSX
 	/** Internal use. Physics-engine representation of this constraint. */
-	class physx::PxD6Joint*		ConstraintData;
+	physx::PxD6Joint*		ConstraintData;
 #endif	//WITH_PHYSX
 
 	/** Physics scene index. */
diff --git a/Engine/Source/Runtime/Engine/Public/Engine.h b/Engine/Source/Runtime/Engine/Public/Engine.h
--- a/Engine/Source/Runtime/Engine/Public/Engine.h
+++ b/Engine/Source/Runtime/Engine/Public/Engine.h
@@ -23,48 +23,6 @@
 
 #include "EngineGlobals.h"
 
-
-/*-----------------------------------------------------------------------------
-	Forward declarations.
------------------------------------------------------------------------------*/
-class UTexture;
-class UTexture2D;
-class FLightMap2D;
-class FShadowMap2D;
-class FSceneInterface;
-class FPrimitiveSceneInfo;
-class FPrimitiveSceneProxy;
-class UMaterialExpression;
-class FMaterialRenderProxy;
-class UMaterial;
-class FSceneView;
-class FSceneViewFamily;
-class FViewportClient;
-class FCanvas;
-class UActorChannel;
-class FAudioDevice;
-
-#if WITH_PHYSX
-namespace physx
-{
-	class PxRigidActor;
-	class PxRigidDynamic;
-	class PxAggregate;
-	class PxD6Joint;
-	class PxGeometry;
-	class PxShape;
-	class PxMaterial;
-	class PxHeightField;
-	class PxTransform;
-	class PxTriangleMesh;
-	class PxVehicleWheels;
-	class PxVehicleDrive;
-	class PxVehicleNoDrive;
-	class PxVehicleDriveSimData;
-	class PxVehicleWheelsSimData;
-}
-#endif // WITH_PHYSX
-
 /*-----------------------------------------------------------------------------
 	Engine public includes.
 -----------------------------------------------------------------------------*/
diff --git a/Engine/Source/Runtime/Engine/Public/MaterialShared.h b/Engine/Source/Runtime/Engine/Public/MaterialShared.h
--- a/Engine/Source/Runtime/Engine/Public/MaterialShared.h
+++ b/Engine/Source/Runtime/Engine/Public/MaterialShared.h
@@ -14,11 +14,14 @@
 #include "RenderResource.h"
 #include "ShaderCompiler.h"
 #include "UniformBuffer.h"
+
 class FMeshMaterialShaderMap;
 class FMaterialShaderMap;
 class FMaterialShaderType;
 class FMaterial;
+class FMaterialRenderProxy;
 class UMaterialInstance;
+class UMaterialExpression;
 
 #define ME_CAPTION_HEIGHT		18
 #define ME_STD_VPADDING			16
diff --git a/Engine/Source/Runtime/Engine/Public/Texture.h b/Engine/Source/Runtime/Engine/Public/Texture.h
--- a/Engine/Source/Runtime/Engine/Public/Texture.h
+++ b/Engine/Source/Runtime/Engine/Public/Texture.h
@@ -8,6 +8,8 @@
 
 #include "Engine/TextureDefines.h"
 
+class UTexture2D;
+
 /** Maximum number of slices in texture source art. */
 #define MAX_TEXTURE_SOURCE_SLICES 6
 
diff --git a/Engine/Source/Runtime/Engine/Public/UnrealClient.h b/Engine/Source/Runtime/Engine/Public/UnrealClient.h
--- a/Engine/Source/Runtime/Engine/Public/UnrealClient.h
+++ b/Engine/Source/Runtime/Engine/Public/UnrealClient.h
@@ -6,6 +6,9 @@
 
 #pragma once
 
+class FCanvas;
+class FViewportClient;
+
 /**
  * A render target.
  */
