Raw and Unit Time Stats now display correctly if Unit Graph isn't already enabled.

#ttp 335852 - EDITOR: VIEWPORT: certain stats do not appear to display anything

#branch UE4

#change
LevelMap - Temporarily removed until convexelem's are working again
Raw - Needs UnitGraph to be enabled first, so now automatically does this (if not already)
SlateBatches = Removed as the render function was commented out
Unit Time - Needs UnitGraph to be enabled first, so now automatically does this (if not already)

diff --git a/Engine/Source/Runtime/Engine/Classes/Engine/Engine.h b/Engine/Source/Runtime/Engine/Classes/Engine/Engine.h
--- a/Engine/Source/Runtime/Engine/Classes/Engine/Engine.h
+++ b/Engine/Source/Runtime/Engine/Classes/Engine/Engine.h
@@ -2653,6 +2653,7 @@ private:
 #if !UE_BUILD_SHIPPING
 	bool ToggleStatUnitMax(UWorld* World, FCommonViewportClient* ViewportClient, const TCHAR* Stream = NULL);
 	bool ToggleStatUnitGraph(UWorld* World, FCommonViewportClient* ViewportClient, const TCHAR* Stream = NULL);
+	bool ToggleStatUnitTime(UWorld* World, FCommonViewportClient* ViewportClient, const TCHAR* Stream = NULL);
 	bool ToggleStatRaw(UWorld* World, FCommonViewportClient* ViewportClient, const TCHAR* Stream = NULL);
 #endif
 	bool ToggleStatSounds(UWorld* World, FCommonViewportClient* ViewportClient, const TCHAR* Stream = NULL);
diff --git a/Engine/Source/Runtime/Engine/Private/UnrealEngine.cpp b/Engine/Source/Runtime/Engine/Private/UnrealEngine.cpp
--- a/Engine/Source/Runtime/Engine/Private/UnrealEngine.cpp
+++ b/Engine/Source/Runtime/Engine/Private/UnrealEngine.cpp
@@ -741,9 +741,11 @@ void UEngine::Init(IEngineLoop* InEngineLoop)
 	EngineStats.Add(FEngineStatFuncs(TEXT("STAT_FPS"), TEXT("STATCAT_Engine"), FText::GetEmpty(), &UEngine::RenderStatFPS, &UEngine::ToggleStatFPS, true));
 	EngineStats.Add(FEngineStatFuncs(TEXT("STAT_Summary"), TEXT("STATCAT_Engine"), FText::GetEmpty(), &UEngine::RenderStatSummary, NULL, true));
 	EngineStats.Add(FEngineStatFuncs(TEXT("STAT_Unit"), TEXT("STATCAT_Engine"), FText::GetEmpty(), &UEngine::RenderStatUnit, &UEngine::ToggleStatUnit, true));
+	/* @todo Slate Rendering
 #if STATS
 	EngineStats.Add(FEngineStatFuncs(TEXT("STAT_SlateBatches"), TEXT("STATCAT_Engine"), FText::GetEmpty(), &UEngine::RenderStatSlateBatches, NULL, true));
 #endif
+	*/
 	EngineStats.Add(FEngineStatFuncs(TEXT("STAT_Hitches"), TEXT("STATCAT_Engine"), FText::GetEmpty(), &UEngine::RenderStatHitches, &UEngine::ToggleStatHitches, true));
 	EngineStats.Add(FEngineStatFuncs(TEXT("STAT_AI"), TEXT("STATCAT_Engine"), FText::GetEmpty(), &UEngine::RenderStatAI, NULL, true));
 
@@ -756,13 +758,14 @@ void UEngine::Init(IEngineLoop* InEngineLoop)
 	EngineStats.Add(FEngineStatFuncs(TEXT("STAT_SoundCues"), TEXT("STATCAT_Engine"), FText::GetEmpty(), &UEngine::RenderStatSoundCues, NULL));
 #endif
 	EngineStats.Add(FEngineStatFuncs(TEXT("STAT_Sounds"), TEXT("STATCAT_Engine"), FText::GetEmpty(), &UEngine::RenderStatSounds, &UEngine::ToggleStatSounds));
+/* @todo UE4 physx fix this once we have convexelem drawing again
 	EngineStats.Add(FEngineStatFuncs(TEXT("STAT_LevelMap"), TEXT("STATCAT_Engine"), FText::GetEmpty(), &UEngine::RenderStatLevelMap, NULL));
-
+*/
 	EngineStats.Add(FEngineStatFuncs(TEXT("STAT_Detailed"), TEXT("STATCAT_Engine"), FText::GetEmpty(), NULL, &UEngine::ToggleStatDetailed));
 #if !UE_BUILD_SHIPPING
 	EngineStats.Add(FEngineStatFuncs(TEXT("STAT_UnitMax"), TEXT("STATCAT_Engine"), FText::GetEmpty(), NULL, &UEngine::ToggleStatUnitMax));
 	EngineStats.Add(FEngineStatFuncs(TEXT("STAT_UnitGraph"), TEXT("STATCAT_Engine"), FText::GetEmpty(), NULL, &UEngine::ToggleStatUnitGraph));
-	EngineStats.Add(FEngineStatFuncs(TEXT("STAT_UnitTime"), TEXT("STATCAT_Engine"), FText::GetEmpty(), NULL, NULL));
+	EngineStats.Add(FEngineStatFuncs(TEXT("STAT_UnitTime"), TEXT("STATCAT_Engine"), FText::GetEmpty(), NULL, &UEngine::ToggleStatUnitTime));
 	EngineStats.Add(FEngineStatFuncs(TEXT("STAT_Raw"), TEXT("STATCAT_Engine"), FText::GetEmpty(), NULL, &UEngine::ToggleStatRaw));
 #endif
 
@@ -10263,8 +10266,8 @@ bool UEngine::ToggleStatUnitMax(UWorld* World, FCommonViewportClient* ViewportCl
 bool UEngine::ToggleStatUnitGraph(UWorld* World, FCommonViewportClient* ViewportClient, const TCHAR* Stream)
 {
 	check(ViewportClient);
-	const bool bShowUnitTimeGraph = ViewportClient->IsStatEnabled(TEXT("UnitGraph"));
-	if (bShowUnitTimeGraph)
+	const bool bShowUnitGraph = ViewportClient->IsStatEnabled(TEXT("UnitGraph"));
+	if (bShowUnitGraph)
 	{
 		// Force Unit to Active
 		SetEngineStat(World, ViewportClient, TEXT("Unit"), true);
@@ -10284,17 +10287,34 @@ bool UEngine::ToggleStatUnitGraph(UWorld* World, FCommonViewportClient* Viewport
 	return true;
 }
 
+// UNITTIME
+bool UEngine::ToggleStatUnitTime(UWorld* World, FCommonViewportClient* ViewportClient, const TCHAR* Stream)
+{
+	check(ViewportClient);
+	const bool bShowUnitTime = ViewportClient->IsStatEnabled(TEXT("UnitTime"));
+	if (bShowUnitTime)
+	{
+		// Force UnitGraph to Active
+		SetEngineStat(World, ViewportClient, TEXT("UnitGraph"), true);
+	}
+	return true;
+}
+
 // RAW
 bool UEngine::ToggleStatRaw(UWorld* World, FCommonViewportClient* ViewportClient, const TCHAR* Stream)
 {
 	const bool bShowRaw = ViewportClient->IsStatEnabled(TEXT("Raw"));
 	const bool bShowDetailed = ViewportClient->IsStatEnabled(TEXT("Detailed"));
-	if (!bShowRaw && bShowDetailed)
+	if (bShowRaw)
+	{
+		// Force UnitGraph to Active
+		SetEngineStat(World, ViewportClient, TEXT("UnitGraph"), true);
+	}
+	else if (bShowDetailed)
 	{
 		// Since we're turning this off, we also need to toggle off detailed too
 		ExecEngineStat(World, ViewportClient, TEXT("Detailed -Skip"));
 	}
-
 	return true;
 }
 #endif
