CIS fix - adding forward declares

diff --git a/Engine/Source/Runtime/Engine/Public/BufferVisualizationData.h b/Engine/Source/Runtime/Engine/Public/BufferVisualizationData.h
--- a/Engine/Source/Runtime/Engine/Public/BufferVisualizationData.h
+++ b/Engine/Source/Runtime/Engine/Public/BufferVisualizationData.h
@@ -1,6 +1,8 @@
 // Copyright 1998-2014 Epic Games, Inc. All Rights Reserved.
 #pragma once
 
+class UMaterial;
+
 class FBufferVisualizationData
 {
 public:
diff --git a/Engine/Source/Runtime/Engine/Public/MaterialShared.h b/Engine/Source/Runtime/Engine/Public/MaterialShared.h
--- a/Engine/Source/Runtime/Engine/Public/MaterialShared.h
+++ b/Engine/Source/Runtime/Engine/Public/MaterialShared.h
@@ -22,6 +22,7 @@ class FMaterial;
 class FMaterialRenderProxy;
 class UMaterialInstance;
 class UMaterialExpression;
+class UMaterialInterface;
 
 #define ME_CAPTION_HEIGHT		18
 #define ME_STD_VPADDING			16
