Particle System arrow components no longer scale relative to their actor scale

This addresses TTP#334571 - LIVE: FEATURE REQUEST: EDITOR: The ability to resize the green arrow on a particle emitter

diff --git a/Engine/Source/Runtime/Engine/Private/Particles/ParticleComponents.cpp b/Engine/Source/Runtime/Engine/Private/Particles/ParticleComponents.cpp
--- a/Engine/Source/Runtime/Engine/Private/Particles/ParticleComponents.cpp
+++ b/Engine/Source/Runtime/Engine/Private/Particles/ParticleComponents.cpp
@@ -193,6 +193,7 @@ AEmitter::AEmitter(const class FPostConstructInitializeProperties& PCIP)
 			ArrowComponent->SpriteInfo.Category = ConstructorStatics.ID_Effects;
 			ArrowComponent->SpriteInfo.DisplayName = ConstructorStatics.NAME_Effects;
 			ArrowComponent->AttachParent = ParticleSystemComponent;
+			ArrowComponent->bAbsoluteScale = true;
 		}
 	}
 #endif // WITH_EDITORONLY_DATA
