First pass removing several headers from Engine.h

diff --git a/Engine/Source/Editor/GeometryMode/Private/GeometryModifiers.cpp b/Engine/Source/Editor/GeometryMode/Private/GeometryModifiers.cpp
--- a/Engine/Source/Editor/GeometryMode/Private/GeometryModifiers.cpp
+++ b/Engine/Source/Editor/GeometryMode/Private/GeometryModifiers.cpp
@@ -1,6 +1,7 @@
 // Copyright 1998-2014 Epic Games, Inc. All Rights Reserved.
 
 #include "GeometryModePrivatePCH.h"
+#include "EditorSupportDelegates.h"
 #include "ScopedTransaction.h"
 #include "BSPOps.h"
 #include "GeomTools.h"
diff --git a/Engine/Source/Editor/LandscapeEditor/Private/LandscapeEdMode.cpp b/Engine/Source/Editor/LandscapeEditor/Private/LandscapeEdMode.cpp
--- a/Engine/Source/Editor/LandscapeEditor/Private/LandscapeEdMode.cpp
+++ b/Engine/Source/Editor/LandscapeEditor/Private/LandscapeEdMode.cpp
@@ -2,6 +2,7 @@
 
 #include "LandscapeEditorPrivatePCH.h"
 
+#include "EditorSupportDelegates.h"
 #include "ObjectTools.h"
 #include "LandscapeEdMode.h"
 #include "ScopedTransaction.h"
diff --git a/Engine/Source/Editor/LevelEditor/Private/LevelEditorActions.cpp b/Engine/Source/Editor/LevelEditor/Private/LevelEditorActions.cpp
--- a/Engine/Source/Editor/LevelEditor/Private/LevelEditorActions.cpp
+++ b/Engine/Source/Editor/LevelEditor/Private/LevelEditorActions.cpp
@@ -2,6 +2,7 @@
 
 
 #include "LevelEditor.h"
+#include "EditorSupportDelegates.h"
 #include "SLevelEditor.h"
 #include "LevelEditorActions.h"
 #include "SLevelViewport.h"
diff --git a/Engine/Source/Editor/LevelEditor/Private/LightingTools.cpp b/Engine/Source/Editor/LevelEditor/Private/LightingTools.cpp
--- a/Engine/Source/Editor/LevelEditor/Private/LightingTools.cpp
+++ b/Engine/Source/Editor/LevelEditor/Private/LightingTools.cpp
@@ -6,6 +6,7 @@
 
 #include "LevelEditor.h"
 #include "LightingTools.h"
+#include "EditorSupportDelegates.h"
 
 /**
  *	LightingTools settings
diff --git a/Engine/Source/Editor/LevelEditor/Private/LightmapResRatioAdjust.cpp b/Engine/Source/Editor/LevelEditor/Private/LightmapResRatioAdjust.cpp
--- a/Engine/Source/Editor/LevelEditor/Private/LightmapResRatioAdjust.cpp
+++ b/Engine/Source/Editor/LevelEditor/Private/LightmapResRatioAdjust.cpp
@@ -5,6 +5,7 @@
 ================================================================================*/
 
 #include "LevelEditor.h"
+#include "EditorSupportDelegates.h"
 #include "LightmapResRatioAdjust.h"
 #include "SurfaceIterators.h"
 #include "MessageLog.h"
diff --git a/Engine/Source/Editor/Levels/Private/LevelViewModel.cpp b/Engine/Source/Editor/Levels/Private/LevelViewModel.cpp
--- a/Engine/Source/Editor/Levels/Private/LevelViewModel.cpp
+++ b/Engine/Source/Editor/Levels/Private/LevelViewModel.cpp
@@ -6,6 +6,7 @@
 #include "LevelUtils.h"
 #include "EditorLevelUtils.h"
 #include "Toolkits/AssetEditorManager.h"
+#include "EditorSupportDelegates.h"
 
 #define LOCTEXT_NAMESPACE "Level"
 
diff --git a/Engine/Source/Editor/MaterialEditor/Private/MaterialEditor.cpp b/Engine/Source/Editor/MaterialEditor/Private/MaterialEditor.cpp
--- a/Engine/Source/Editor/MaterialEditor/Private/MaterialEditor.cpp
+++ b/Engine/Source/Editor/MaterialEditor/Private/MaterialEditor.cpp
@@ -4,6 +4,7 @@
 #include "MaterialEditorActions.h"
 #include "MaterialExpressionClasses.h"
 #include "MaterialCompiler.h"
+#include "EditorSupportDelegates.h"
 #include "Toolkits/IToolkitHost.h"
 #include "Editor/EditorWidgets/Public/EditorWidgets.h"
 #include "AssetRegistryModule.h"
diff --git a/Engine/Source/Editor/MaterialEditor/Private/MaterialEditorInstanceDetailCustomization.cpp b/Engine/Source/Editor/MaterialEditor/Private/MaterialEditorInstanceDetailCustomization.cpp
--- a/Engine/Source/Editor/MaterialEditor/Private/MaterialEditorInstanceDetailCustomization.cpp
+++ b/Engine/Source/Editor/MaterialEditor/Private/MaterialEditorInstanceDetailCustomization.cpp
@@ -1,6 +1,7 @@
 // Copyright 1998-2014 Epic Games, Inc. All Rights Reserved.
 
 #include "MaterialEditorModule.h"
+#include "EditorSupportDelegates.h"
 #include "UnrealEd.h"
 #include "PropertyEditing.h"
 #include "PropertyCustomizationHelpers.h"
diff --git a/Engine/Source/Editor/Matinee/Private/Matinee.cpp b/Engine/Source/Editor/Matinee/Private/Matinee.cpp
--- a/Engine/Source/Editor/Matinee/Private/Matinee.cpp
+++ b/Engine/Source/Editor/Matinee/Private/Matinee.cpp
@@ -22,6 +22,7 @@
 #include "InterpolationHitProxy.h"
 
 #include "LevelEditorActions.h"
+#include "EditorSupportDelegates.h"
 
 #include "Editor/WorkspaceMenuStructure/Public/WorkspaceMenuStructureModule.h"
 #include "MessageLog.h"
diff --git a/Engine/Source/Editor/Matinee/Private/MatineeTools.cpp b/Engine/Source/Editor/Matinee/Private/MatineeTools.cpp
--- a/Engine/Source/Editor/Matinee/Private/MatineeTools.cpp
+++ b/Engine/Source/Editor/Matinee/Private/MatineeTools.cpp
@@ -13,6 +13,7 @@
 #include "MatineeDelegates.h"
 
 #include "SoundDefinitions.h"
+#include "EditorSupportDelegates.h"
 
 ///// UTILS
 
diff --git a/Engine/Source/Editor/MeshPaint/Private/MeshPaintEdMode.cpp b/Engine/Source/Editor/MeshPaint/Private/MeshPaintEdMode.cpp
--- a/Engine/Source/Editor/MeshPaint/Private/MeshPaintEdMode.cpp
+++ b/Engine/Source/Editor/MeshPaint/Private/MeshPaintEdMode.cpp
@@ -14,6 +14,7 @@
 #include "RawMesh.h"
 #include "Editor/UnrealEd/Public/ObjectTools.h"
 #include "AssetToolsModule.h"
+#include "EditorSupportDelegates.h"
 
 //Slate dependencies
 #include "Editor/LevelEditor/Public/LevelEditor.h"
diff --git a/Engine/Source/Editor/PhAT/Private/PhATSharedData.cpp b/Engine/Source/Editor/PhAT/Private/PhATSharedData.cpp
--- a/Engine/Source/Editor/PhAT/Private/PhATSharedData.cpp
+++ b/Engine/Source/Editor/PhAT/Private/PhATSharedData.cpp
@@ -1,6 +1,7 @@
 // Copyright 1998-2014 Epic Games, Inc. All Rights Reserved.
 
 #include "PhATModule.h"
+#include "EditorSupportDelegates.h"
 #include "ScopedTransaction.h"
 #include "SPhATNewAssetDlg.h"
 #include "PhATEdSkeletalMeshComponent.h"
diff --git a/Engine/Source/Editor/PropertyEditor/Private/PropertyHandleImpl.cpp b/Engine/Source/Editor/PropertyEditor/Private/PropertyHandleImpl.cpp
--- a/Engine/Source/Editor/PropertyEditor/Private/PropertyHandleImpl.cpp
+++ b/Engine/Source/Editor/PropertyEditor/Private/PropertyHandleImpl.cpp
@@ -5,6 +5,7 @@
 #include "PropertyNode.h"
 #include "PropertyEditorHelpers.h"
 #include "ObjectPropertyNode.h"
+#include "EditorSupportDelegates.h"
 #include "ScopedTransaction.h"
 #include "Editor/UnrealEd/Public/Kismet2/KismetEditorUtilities.h"
 #include "IPropertyUtilities.h"
diff --git a/Engine/Source/Editor/TextureAlignMode/Private/TextureAlignEdMode.cpp b/Engine/Source/Editor/TextureAlignMode/Private/TextureAlignEdMode.cpp
--- a/Engine/Source/Editor/TextureAlignMode/Private/TextureAlignEdMode.cpp
+++ b/Engine/Source/Editor/TextureAlignMode/Private/TextureAlignEdMode.cpp
@@ -1,6 +1,7 @@
 // Copyright 1998-2014 Epic Games, Inc. All Rights Reserved.
 
 #include "TextureAlignModePrivatePCH.h"
+#include "EditorSupportDelegates.h"
 
 IMPLEMENT_MODULE( FTextureAlignModeModule, TextureAlignMode );
 
diff --git a/Engine/Source/Editor/UnrealEd/Private/Editor.cpp b/Engine/Source/Editor/UnrealEd/Private/Editor.cpp
--- a/Engine/Source/Editor/UnrealEd/Private/Editor.cpp
+++ b/Engine/Source/Editor/UnrealEd/Private/Editor.cpp
@@ -2,6 +2,7 @@
 
 
 #include "UnrealEd.h"
+#include "EditorSupportDelegates.h"
 #include "Factories.h"
 #include "BSPOps.h"
 
diff --git a/Engine/Source/Editor/UnrealEd/Private/EditorConstraints.cpp b/Engine/Source/Editor/UnrealEd/Private/EditorConstraints.cpp
--- a/Engine/Source/Editor/UnrealEd/Private/EditorConstraints.cpp
+++ b/Engine/Source/Editor/UnrealEd/Private/EditorConstraints.cpp
@@ -5,7 +5,7 @@
 =============================================================================*/
 
 #include "UnrealEd.h"
-
+#include "EditorSupportDelegates.h"
 
 float UEditorEngine::GetGridSize()
 {
diff --git a/Engine/Source/Editor/UnrealEd/Private/EditorLevelUtils.cpp b/Engine/Source/Editor/UnrealEd/Private/EditorLevelUtils.cpp
--- a/Engine/Source/Editor/UnrealEd/Private/EditorLevelUtils.cpp
+++ b/Engine/Source/Editor/UnrealEd/Private/EditorLevelUtils.cpp
@@ -5,6 +5,7 @@
 =============================================================================*/
 
 #include "UnrealEd.h"
+#include "EditorSupportDelegates.h"
 #include "EditorLevelUtils.h"
 
 #include "BusyCursor.h"
diff --git a/Engine/Source/Editor/UnrealEd/Private/EditorModes.cpp b/Engine/Source/Editor/UnrealEd/Private/EditorModes.cpp
--- a/Engine/Source/Editor/UnrealEd/Private/EditorModes.cpp
+++ b/Engine/Source/Editor/UnrealEd/Private/EditorModes.cpp
@@ -1,6 +1,7 @@
 // Copyright 1998-2014 Epic Games, Inc. All Rights Reserved.
 
 #include "UnrealEd.h"
+#include "EditorSupportDelegates.h"
 #include "MouseDeltaTracker.h"
 #include "ScopedTransaction.h"
 #include "SurfaceIterators.h"
diff --git a/Engine/Source/Editor/UnrealEd/Private/EditorServer.cpp b/Engine/Source/Editor/UnrealEd/Private/EditorServer.cpp
--- a/Engine/Source/Editor/UnrealEd/Private/EditorServer.cpp
+++ b/Engine/Source/Editor/UnrealEd/Private/EditorServer.cpp
@@ -2,6 +2,7 @@
 
 
 #include "UnrealEd.h"
+#include "EditorSupportDelegates.h"
 #include "Factories.h"
 #include "BusyCursor.h"
 #include "SoundDefinitions.h"
diff --git a/Engine/Source/Editor/UnrealEd/Private/GeomFitUtils.cpp b/Engine/Source/Editor/UnrealEd/Private/GeomFitUtils.cpp
--- a/Engine/Source/Editor/UnrealEd/Private/GeomFitUtils.cpp
+++ b/Engine/Source/Editor/UnrealEd/Private/GeomFitUtils.cpp
@@ -5,6 +5,7 @@
 =============================================================================*/
 
 #include "UnrealEd.h"
+#include "EditorSupportDelegates.h"
 #include "BSPOps.h"
 #include "../Private/GeomFitUtils.h"
 #include "RawMesh.h"
diff --git a/Engine/Source/Editor/UnrealEd/Private/LevelEditorViewport.cpp b/Engine/Source/Editor/UnrealEd/Private/LevelEditorViewport.cpp
--- a/Engine/Source/Editor/UnrealEd/Private/LevelEditorViewport.cpp
+++ b/Engine/Source/Editor/UnrealEd/Private/LevelEditorViewport.cpp
@@ -2,6 +2,7 @@
 
 
 #include "UnrealEd.h"
+#include "EditorSupportDelegates.h"
 #include "SoundDefinitions.h"
 #include "CameraController.h"
 #include "MouseDeltaTracker.h"
diff --git a/Engine/Source/Editor/UnrealEd/Private/UnrealEdEngine.cpp b/Engine/Source/Editor/UnrealEd/Private/UnrealEdEngine.cpp
--- a/Engine/Source/Editor/UnrealEd/Private/UnrealEdEngine.cpp
+++ b/Engine/Source/Editor/UnrealEd/Private/UnrealEdEngine.cpp
@@ -1,6 +1,7 @@
 // Copyright 1998-2014 Epic Games, Inc. All Rights Reserved.
 
 #include "UnrealEd.h"
+#include "EditorSupportDelegates.h"
 #include "SoundDefinitions.h"
 #include "BusyCursor.h"
 #include "EditorLevelUtils.h"
diff --git a/Engine/Source/Editor/UnrealEd/Private/UnrealEdMisc.cpp b/Engine/Source/Editor/UnrealEd/Private/UnrealEdMisc.cpp
--- a/Engine/Source/Editor/UnrealEd/Private/UnrealEdMisc.cpp
+++ b/Engine/Source/Editor/UnrealEd/Private/UnrealEdMisc.cpp
@@ -1,6 +1,7 @@
 // Copyright 1998-2014 Epic Games, Inc. All Rights Reserved.
 
 #include "UnrealEd.h"
+#include "EditorSupportDelegates.h"
 #include "ISourceControlModule.h"
 #include "Kismet2/DebuggerCommands.h"
 #include "Toolkits/AssetEditorCommonCommands.h"
diff --git a/Engine/Source/Editor/UnrealEd/Private/UnrealEdSrv.cpp b/Engine/Source/Editor/UnrealEd/Private/UnrealEdSrv.cpp
--- a/Engine/Source/Editor/UnrealEd/Private/UnrealEdSrv.cpp
+++ b/Engine/Source/Editor/UnrealEd/Private/UnrealEdSrv.cpp
@@ -2,6 +2,7 @@
 
 
 #include "UnrealEd.h"
+#include "EditorSupportDelegates.h"
 #include "BusyCursor.h"
 #include "ScopedTransaction.h"
 #include "LevelUtils.h"
diff --git a/Engine/Source/Editor/WorldBrowser/Private/LevelCollectionModel.cpp b/Engine/Source/Editor/WorldBrowser/Private/LevelCollectionModel.cpp
--- a/Engine/Source/Editor/WorldBrowser/Private/LevelCollectionModel.cpp
+++ b/Engine/Source/Editor/WorldBrowser/Private/LevelCollectionModel.cpp
@@ -3,6 +3,7 @@
 #include "WorldBrowserPrivatePCH.h"
 #include "SourceControlWindows.h"
 #include "AssetToolsModule.h"
+#include "EditorSupportDelegates.h"
 
 #include "LevelCollectionModel.h"
 
diff --git a/Engine/Source/Runtime/Engine/Classes/AI/Navigation/NavigationSystem.h b/Engine/Source/Runtime/Engine/Classes/AI/Navigation/NavigationSystem.h
--- a/Engine/Source/Runtime/Engine/Classes/AI/Navigation/NavigationSystem.h
+++ b/Engine/Source/Runtime/Engine/Classes/AI/Navigation/NavigationSystem.h
@@ -3,6 +3,7 @@
 #pragma once
 
 #include "NavigationTypes.h"
+#include "GenericOctreePublic.h"
 #include "NavigationSystem.generated.h"
 
 #define NAVSYS_DEBUG (0 && UE_BUILD_DEBUG)
diff --git a/Engine/Source/Runtime/Engine/Classes/Components/DestructibleComponent.h b/Engine/Source/Runtime/Engine/Classes/Components/DestructibleComponent.h
--- a/Engine/Source/Runtime/Engine/Classes/Components/DestructibleComponent.h
+++ b/Engine/Source/Runtime/Engine/Classes/Components/DestructibleComponent.h
@@ -1,8 +1,7 @@
 // Copyright 1998-2014 Epic Games, Inc. All Rights Reserved.
 #pragma once
 
-
-
+#include "PhysxUserData.h"
 #include "DestructibleComponent.generated.h"
 
 #if WITH_PHYSX
diff --git a/Engine/Source/Runtime/Engine/Classes/Distributions/DistributionFloat.h b/Engine/Source/Runtime/Engine/Classes/Distributions/DistributionFloat.h
--- a/Engine/Source/Runtime/Engine/Classes/Distributions/DistributionFloat.h
+++ b/Engine/Source/Runtime/Engine/Classes/Distributions/DistributionFloat.h
@@ -2,6 +2,7 @@
 
 
 #pragma once
+#include "Distributions.h"
 #include "DistributionFloat.generated.h"
 
 /** Type-safe floating point distribution. */
diff --git a/Engine/Source/Runtime/Engine/Classes/Distributions/DistributionFloatUniform.h b/Engine/Source/Runtime/Engine/Classes/Distributions/DistributionFloatUniform.h
--- a/Engine/Source/Runtime/Engine/Classes/Distributions/DistributionFloatUniform.h
+++ b/Engine/Source/Runtime/Engine/Classes/Distributions/DistributionFloatUniform.h
@@ -2,6 +2,7 @@
 
 
 #pragma once
+#include "Distributions.h"
 #include "DistributionFloatUniform.generated.h"
 
 UCLASS(collapsecategories, hidecategories=Object, editinlinenew, MinimalAPI)
diff --git a/Engine/Source/Runtime/Engine/Classes/Distributions/DistributionFloatUniformCurve.h b/Engine/Source/Runtime/Engine/Classes/Distributions/DistributionFloatUniformCurve.h
--- a/Engine/Source/Runtime/Engine/Classes/Distributions/DistributionFloatUniformCurve.h
+++ b/Engine/Source/Runtime/Engine/Classes/Distributions/DistributionFloatUniformCurve.h
@@ -2,6 +2,7 @@
 
 
 #pragma once
+#include "Distributions.h"
 #include "DistributionFloatUniformCurve.generated.h"
 
 UCLASS(collapsecategories, hidecategories=Object, editinlinenew, MinimalAPI)
diff --git a/Engine/Source/Runtime/Engine/Classes/Distributions/DistributionVector.h b/Engine/Source/Runtime/Engine/Classes/Distributions/DistributionVector.h
--- a/Engine/Source/Runtime/Engine/Classes/Distributions/DistributionVector.h
+++ b/Engine/Source/Runtime/Engine/Classes/Distributions/DistributionVector.h
@@ -2,6 +2,7 @@
 
 
 #pragma once
+#include "Distributions.h"
 #include "DistributionVector.generated.h"
 
 UENUM()
diff --git a/Engine/Source/Runtime/Engine/Classes/Distributions/DistributionVectorUniform.h b/Engine/Source/Runtime/Engine/Classes/Distributions/DistributionVectorUniform.h
--- a/Engine/Source/Runtime/Engine/Classes/Distributions/DistributionVectorUniform.h
+++ b/Engine/Source/Runtime/Engine/Classes/Distributions/DistributionVectorUniform.h
@@ -2,6 +2,7 @@
 
 
 #pragma once
+#include "Distributions.h"
 #include "DistributionVectorUniform.generated.h"
 
 UCLASS(collapsecategories, hidecategories=Object, editinlinenew, MinimalAPI)
diff --git a/Engine/Source/Runtime/Engine/Classes/Distributions/DistributionVectorUniformCurve.h b/Engine/Source/Runtime/Engine/Classes/Distributions/DistributionVectorUniformCurve.h
--- a/Engine/Source/Runtime/Engine/Classes/Distributions/DistributionVectorUniformCurve.h
+++ b/Engine/Source/Runtime/Engine/Classes/Distributions/DistributionVectorUniformCurve.h
@@ -2,6 +2,7 @@
 
 
 #pragma once
+#include "Distributions.h"
 #include "DistributionVectorUniformCurve.generated.h"
 
 UCLASS(collapsecategories, hidecategories=Object, editinlinenew, MinimalAPI)
diff --git a/Engine/Source/Runtime/Engine/Classes/EdGraph/EdGraph.h b/Engine/Source/Runtime/Engine/Classes/EdGraph/EdGraph.h
--- a/Engine/Source/Runtime/Engine/Classes/EdGraph/EdGraph.h
+++ b/Engine/Source/Runtime/Engine/Classes/EdGraph/EdGraph.h
@@ -1,6 +1,7 @@
 // Copyright 1998-2014 Epic Games, Inc. All Rights Reserved.
 
 #pragma once
+#include "BlueprintUtilities.h"
 #include "EdGraph.generated.h"
 
 USTRUCT()
diff --git a/Engine/Source/Runtime/Engine/Classes/Engine/GameViewportClient.h b/Engine/Source/Runtime/Engine/Classes/Engine/GameViewportClient.h
--- a/Engine/Source/Runtime/Engine/Classes/Engine/GameViewportClient.h
+++ b/Engine/Source/Runtime/Engine/Classes/Engine/GameViewportClient.h
@@ -14,6 +14,9 @@
  *
  */
 #pragma once
+
+#include "ShowFlags.h"
+
 #include "GameViewportClient.generated.h"
 
 /**
diff --git a/Engine/Source/Runtime/Engine/Classes/Materials/Material.h b/Engine/Source/Runtime/Engine/Classes/Materials/Material.h
--- a/Engine/Source/Runtime/Engine/Classes/Materials/Material.h
+++ b/Engine/Source/Runtime/Engine/Classes/Materials/Material.h
@@ -2,6 +2,8 @@
 
 #pragma once
 #include "../Engine/BlendableInterface.h"
+#include "../Materials/MaterialInterface.h"
+#include "MaterialShared.h"
 #include "Material.generated.h"
 
 #if WITH_EDITOR
diff --git a/Engine/Source/Runtime/Engine/Classes/Materials/MaterialInterface.h b/Engine/Source/Runtime/Engine/Classes/Materials/MaterialInterface.h
--- a/Engine/Source/Runtime/Engine/Classes/Materials/MaterialInterface.h
+++ b/Engine/Source/Runtime/Engine/Classes/Materials/MaterialInterface.h
@@ -2,8 +2,12 @@
 
 
 #pragma once
+
+#include "MaterialShared.h"
 #include "MaterialInterface.generated.h"
 
+class UMaterial;
+
 UENUM()
 enum EMaterialUsage
 {
@@ -175,7 +179,7 @@ public:
 	*						@note: only valid in the editor!
 	* @return	The resource to use for rendering this material instance.
 	*/
-	virtual FMaterialRenderProxy* GetRenderProxy(bool Selected, bool bHovered=false) const PURE_VIRTUAL(UMaterialInterface::GetRenderProxy,return NULL;);
+	virtual class FMaterialRenderProxy* GetRenderProxy(bool Selected, bool bHovered=false) const PURE_VIRTUAL(UMaterialInterface::GetRenderProxy,return NULL;);
 
 	/**
 	* Return a pointer to the physical material used by this material instance.
diff --git a/Engine/Source/Runtime/Engine/Classes/Particles/TypeData/ParticleModuleTypeDataGpu.h b/Engine/Source/Runtime/Engine/Classes/Particles/TypeData/ParticleModuleTypeDataGpu.h
--- a/Engine/Source/Runtime/Engine/Classes/Particles/TypeData/ParticleModuleTypeDataGpu.h
+++ b/Engine/Source/Runtime/Engine/Classes/Particles/TypeData/ParticleModuleTypeDataGpu.h
@@ -6,6 +6,7 @@
 
 #pragma once
 
+#include "Distributions.h"
 #include "ParticleModuleTypeDataGpu.generated.h"
 
 /**
diff --git a/Engine/Source/Runtime/Engine/Classes/PhysicalMaterials/PhysicalMaterial.h b/Engine/Source/Runtime/Engine/Classes/PhysicalMaterials/PhysicalMaterial.h
--- a/Engine/Source/Runtime/Engine/Classes/PhysicalMaterials/PhysicalMaterial.h
+++ b/Engine/Source/Runtime/Engine/Classes/PhysicalMaterials/PhysicalMaterial.h
@@ -2,6 +2,7 @@
 
 #pragma once
 
+#include "PhysxUserData.h"
 #include "PhysicsEngine/PhysicsSettings.h"
 #include "PhysicalMaterial.generated.h"
 
diff --git a/Engine/Source/Runtime/Engine/Classes/PhysicsEngine/BodyInstance.h b/Engine/Source/Runtime/Engine/Classes/PhysicsEngine/BodyInstance.h
--- a/Engine/Source/Runtime/Engine/Classes/PhysicsEngine/BodyInstance.h
+++ b/Engine/Source/Runtime/Engine/Classes/PhysicsEngine/BodyInstance.h
@@ -2,6 +2,7 @@
 
 #pragma once
 
+#include "PhysxUserData.h"
 #include "BodyInstance.generated.h"
 
 #if WITH_PHYSX
diff --git a/Engine/Source/Runtime/Engine/Private/ActorEditor.cpp b/Engine/Source/Runtime/Engine/Private/ActorEditor.cpp
--- a/Engine/Source/Runtime/Engine/Private/ActorEditor.cpp
+++ b/Engine/Source/Runtime/Engine/Private/ActorEditor.cpp
@@ -1,6 +1,7 @@
 // Copyright 1998-2014 Epic Games, Inc. All Rights Reserved.
 
 #include "EnginePrivate.h"
+#include "EditorSupportDelegates.h"
 #include "ActorEditorUtils.h"
 #include "MessageLog.h"
 #include "UObjectToken.h"
diff --git a/Engine/Source/Runtime/Engine/Private/Engine.cpp b/Engine/Source/Runtime/Engine/Private/Engine.cpp
--- a/Engine/Source/Runtime/Engine/Private/Engine.cpp
+++ b/Engine/Source/Runtime/Engine/Private/Engine.cpp
@@ -5,6 +5,7 @@
 =============================================================================*/
 
 #include "EnginePrivate.h"
+#include "EditorSupportDelegates.h"
 
 /*-----------------------------------------------------------------------------
 	Globals.
diff --git a/Engine/Source/Runtime/Engine/Private/Landscape/Landscape.cpp b/Engine/Source/Runtime/Engine/Private/Landscape/Landscape.cpp
--- a/Engine/Source/Runtime/Engine/Private/Landscape/Landscape.cpp
+++ b/Engine/Source/Runtime/Engine/Private/Landscape/Landscape.cpp
@@ -5,6 +5,7 @@ Landscape.cpp: Terrain rendering
 =============================================================================*/
 
 #include "EnginePrivate.h"
+#include "EditorSupportDelegates.h"
 #include "Landscape/LandscapeDataAccess.h"
 #include "Landscape/LandscapeRender.h"
 #include "Landscape/LandscapeRenderMobile.h"
diff --git a/Engine/Source/Runtime/Engine/Private/LevelActor.cpp b/Engine/Source/Runtime/Engine/Private/LevelActor.cpp
--- a/Engine/Source/Runtime/Engine/Private/LevelActor.cpp
+++ b/Engine/Source/Runtime/Engine/Private/LevelActor.cpp
@@ -2,6 +2,7 @@
 
 
 #include "EnginePrivate.h"
+#include "EditorSupportDelegates.h"
 #include "Net/UnrealNetwork.h"
 #include "Collision.h"
 
diff --git a/Engine/Source/Runtime/Engine/Private/LevelUtils.cpp b/Engine/Source/Runtime/Engine/Private/LevelUtils.cpp
--- a/Engine/Source/Runtime/Engine/Private/LevelUtils.cpp
+++ b/Engine/Source/Runtime/Engine/Private/LevelUtils.cpp
@@ -3,6 +3,7 @@
 
 #include "EnginePrivate.h"
 #include "LevelUtils.h"
+#include "EditorSupportDelegates.h"
 
 #if WITH_EDITOR
 #include "ScopedTransaction.h"
diff --git a/Engine/Source/Runtime/Engine/Private/Materials/Material.cpp b/Engine/Source/Runtime/Engine/Private/Materials/Material.cpp
--- a/Engine/Source/Runtime/Engine/Private/Materials/Material.cpp
+++ b/Engine/Source/Runtime/Engine/Private/Materials/Material.cpp
@@ -5,6 +5,7 @@
 =============================================================================*/
 
 #include "EnginePrivate.h"
+#include "EditorSupportDelegates.h"
 #include "MaterialShader.h"
 #include "MaterialInstanceSupport.h"
 #include "UObjectAnnotation.h"
diff --git a/Engine/Source/Runtime/Engine/Private/Materials/MaterialExpressions.cpp b/Engine/Source/Runtime/Engine/Private/Materials/MaterialExpressions.cpp
--- a/Engine/Source/Runtime/Engine/Private/Materials/MaterialExpressions.cpp
+++ b/Engine/Source/Runtime/Engine/Private/Materials/MaterialExpressions.cpp
@@ -5,6 +5,7 @@
 =============================================================================*/
 
 #include "EnginePrivate.h"
+#include "EditorSupportDelegates.h"
 #include "MaterialCompiler.h"
 #if WITH_EDITOR
 #include "Slate.h"
diff --git a/Engine/Source/Runtime/Engine/Private/RendererSupport.cpp b/Engine/Source/Runtime/Engine/Private/RendererSupport.cpp
--- a/Engine/Source/Runtime/Engine/Private/RendererSupport.cpp
+++ b/Engine/Source/Runtime/Engine/Private/RendererSupport.cpp
@@ -5,6 +5,7 @@
 =============================================================================*/
 
 #include "EnginePrivate.h"
+#include "EditorSupportDelegates.h"
 #include "FXSystem.h"
 #include "RenderCore.h"
 #include "GlobalShader.h"
diff --git a/Engine/Source/Runtime/Engine/Private/ShaderCompiler/ShaderCompiler.cpp b/Engine/Source/Runtime/Engine/Private/ShaderCompiler/ShaderCompiler.cpp
--- a/Engine/Source/Runtime/Engine/Private/ShaderCompiler/ShaderCompiler.cpp
+++ b/Engine/Source/Runtime/Engine/Private/ShaderCompiler/ShaderCompiler.cpp
@@ -5,6 +5,7 @@
 =============================================================================*/
 
 #include "EnginePrivate.h"
+#include "EditorSupportDelegates.h"
 #include "ExceptionHandling.h"
 #include "GlobalShader.h"
 #include "TargetPlatform.h"
diff --git a/Engine/Source/Runtime/Engine/Private/SkeletalMesh.cpp b/Engine/Source/Runtime/Engine/Private/SkeletalMesh.cpp
--- a/Engine/Source/Runtime/Engine/Private/SkeletalMesh.cpp
+++ b/Engine/Source/Runtime/Engine/Private/SkeletalMesh.cpp
@@ -5,6 +5,7 @@
 =============================================================================*/
 
 #include "EnginePrivate.h"
+#include "EditorSupportDelegates.h"
 #include "GPUSkinVertexFactory.h"
 #include "SkeletalMeshSorting.h"
 #include "MeshBuild.h"
diff --git a/Engine/Source/Runtime/Engine/Private/Texture.cpp b/Engine/Source/Runtime/Engine/Private/Texture.cpp
--- a/Engine/Source/Runtime/Engine/Private/Texture.cpp
+++ b/Engine/Source/Runtime/Engine/Private/Texture.cpp
@@ -5,6 +5,7 @@
 =============================================================================*/
 
 #include "EnginePrivate.h"
+#include "EditorSupportDelegates.h"
 #include "TargetPlatform.h"
 #include "ImageWrapper.h"
 
diff --git a/Engine/Source/Runtime/Engine/Private/UnrealClient.cpp b/Engine/Source/Runtime/Engine/Private/UnrealClient.cpp
--- a/Engine/Source/Runtime/Engine/Private/UnrealClient.cpp
+++ b/Engine/Source/Runtime/Engine/Private/UnrealClient.cpp
@@ -2,6 +2,7 @@
 
 
 #include "EnginePrivate.h"
+#include "EditorSupportDelegates.h"
 #include "RenderCore.h"
 #include "RenderResource.h"
 #include "RHIStaticStates.h"
diff --git a/Engine/Source/Runtime/Engine/Public/Engine.h b/Engine/Source/Runtime/Engine/Public/Engine.h
--- a/Engine/Source/Runtime/Engine/Public/Engine.h
+++ b/Engine/Source/Runtime/Engine/Public/Engine.h
@@ -10,14 +10,9 @@
 
 #include "Core.h"
 #include "CoreUObject.h"
-#include "Messaging.h"
-#include "TaskGraphInterfaces.h"
 #include "RHI.h"
-#include "RenderCore.h"
 #include "InputCore.h"
-#include "EngineMessages.h"
 #include "EngineSettings.h"
-#include "EditorSupportDelegates.h"
 #include "EngineStats.h"
 #include "EngineLogs.h"
 
@@ -27,18 +22,10 @@
 	Engine public includes.
 -----------------------------------------------------------------------------*/
 
-#include "BlueprintUtilities.h"
 #include "Tickable.h"						// FTickableGameObject interface.
 #include "RenderingThread.h"				// for FRenderCommandFence
-#include "GenericOctreePublic.h"			// for FOctreeElementId
-#include "RenderResource.h"					// for FRenderResource
-#include "HitProxies.h"						// Hit proxy definitions.
 #include "Engine/EngineBaseTypes.h"
 #include "UnrealClient.h"					// for FViewportClient
-#include "ShowFlags.h"						// Flags for enable scene rendering features
-#include "RenderUtils.h"					// Render utility classes.
-#include "Distributions.h"					// Distributions
-#include "PhysxUserData.h"
 
 // EngineClasses.h
 #include "Engine/LatentActionManager.h"
diff --git a/Engine/Source/Runtime/Engine/Public/EngineService.h b/Engine/Source/Runtime/Engine/Public/EngineService.h
--- a/Engine/Source/Runtime/Engine/Public/EngineService.h
+++ b/Engine/Source/Runtime/Engine/Public/EngineService.h
@@ -6,6 +6,8 @@
 
 #pragma once
 
+#include "Messaging.h"
+#include "EngineMessages.h"
 
 /**
  * Implements an application session service.
diff --git a/Engine/Source/Runtime/Engine/Public/MaterialShared.h b/Engine/Source/Runtime/Engine/Public/MaterialShared.h
--- a/Engine/Source/Runtime/Engine/Public/MaterialShared.h
+++ b/Engine/Source/Runtime/Engine/Public/MaterialShared.h
@@ -20,6 +20,7 @@ class FMaterialShaderMap;
 class FMaterialShaderType;
 class FMaterial;
 class FMaterialRenderProxy;
+class UMaterial;
 class UMaterialInstance;
 class UMaterialExpression;
 class UMaterialInterface;
diff --git a/Engine/Source/Runtime/Engine/Public/ParticleEmitterInstances.h b/Engine/Source/Runtime/Engine/Public/ParticleEmitterInstances.h
--- a/Engine/Source/Runtime/Engine/Public/ParticleEmitterInstances.h
+++ b/Engine/Source/Runtime/Engine/Public/ParticleEmitterInstances.h
@@ -8,6 +8,7 @@
 #pragma once
 
 #include "Engine.h"
+#include "Distributions.h"
 #include "ParticleHelper.h"
 
 /*-----------------------------------------------------------------------------
diff --git a/Engine/Source/Runtime/Engine/Public/PhysicsPublic.h b/Engine/Source/Runtime/Engine/Public/PhysicsPublic.h
--- a/Engine/Source/Runtime/Engine/Public/PhysicsPublic.h
+++ b/Engine/Source/Runtime/Engine/Public/PhysicsPublic.h
@@ -8,6 +8,7 @@
 #pragma once
 
 #include "DynamicMeshBuilder.h"
+#include "PhysxUserData.h"
 
 /**
  * Physics stats
diff --git a/Engine/Source/Runtime/Engine/Public/UnrealClient.h b/Engine/Source/Runtime/Engine/Public/UnrealClient.h
--- a/Engine/Source/Runtime/Engine/Public/UnrealClient.h
+++ b/Engine/Source/Runtime/Engine/Public/UnrealClient.h
@@ -6,6 +6,9 @@
 
 #pragma once
 
+#include "RenderResource.h"
+#include "HitProxies.h"		
+
 class FCanvas;
 class FViewportClient;
 
