Fix SkeletalMeshComponent building without apex

diff --git a/Engine/Source/Runtime/Engine/Classes/Components/SkeletalMeshComponent.h b/Engine/Source/Runtime/Engine/Classes/Components/SkeletalMeshComponent.h
--- a/Engine/Source/Runtime/Engine/Classes/Components/SkeletalMeshComponent.h
+++ b/Engine/Source/Runtime/Engine/Classes/Components/SkeletalMeshComponent.h
@@ -5,7 +5,7 @@
 #include "PhysicsEngine/PhysicsConstraintComponent.h"
 #include "SkeletalMeshComponent.generated.h"
 
-#if WITH_APEX
+//#if WITH_APEX
 namespace physx
 { 
 	namespace apex 
@@ -14,7 +14,7 @@ namespace physx
 		class NxClothingCollision;
 	}
 }
-#endif // WITH_APEX
+//#endif // WITH_APEX
 
 class FPhysScene;
 
