BP: Expose FMath::IsNearlyEqual to blueprints ttp335891
Default error-tolerance value in not a SMALL_NUMBER, but 1.e-6

diff --git a/Engine/Source/Runtime/Engine/Classes/Kismet/KismetMathLibrary.h b/Engine/Source/Runtime/Engine/Classes/Kismet/KismetMathLibrary.h
--- a/Engine/Source/Runtime/Engine/Classes/Kismet/KismetMathLibrary.h
+++ b/Engine/Source/Runtime/Engine/Classes/Kismet/KismetMathLibrary.h
@@ -233,6 +233,10 @@ class UKismetMathLibrary : public UBlueprintFunctionLibrary
 	UFUNCTION(BlueprintPure, meta=(FriendlyName = "Equal (float)", CompactNodeTitle = "==", Keywords = "== equal"), Category="Math|Float")
 	static bool EqualEqual_FloatFloat(float A, float B);
 
+	/* Returns true if (|A - B| < ErrorTolerance) */
+	UFUNCTION(BlueprintPure, meta=(FriendlyName = "Nearly Equal (float)", Keywords = "== equal"), Category="Math|Float")
+	static bool NearlyEqual_FloatFloat(float A, float B, float ErrorTolerance = 1.e-6f);
+
 	/* Returns true if A does not equal B (A != B)*/
 	UFUNCTION(BlueprintPure, meta=(FriendlyName = "NotEqual (float)", CompactNodeTitle = "!=", Keywords = "!= not equal"), Category="Math|Float")
 	static bool NotEqual_FloatFloat(float A, float B);
diff --git a/Engine/Source/Runtime/Engine/Private/KismetMathLibrary.cpp b/Engine/Source/Runtime/Engine/Private/KismetMathLibrary.cpp
--- a/Engine/Source/Runtime/Engine/Private/KismetMathLibrary.cpp
+++ b/Engine/Source/Runtime/Engine/Private/KismetMathLibrary.cpp
@@ -308,6 +308,11 @@ bool UKismetMathLibrary::EqualEqual_FloatFloat(float A, float B)
 	return A == B;
 }	
 
+bool UKismetMathLibrary::NearlyEqual_FloatFloat(float A, float B, float ErrorTolerance)
+{
+	return FMath::IsNearlyEqual(A, B, ErrorTolerance);
+}
+
 bool UKismetMathLibrary::NotEqual_FloatFloat(float A, float B)
 {
 	return A != B;
