fix convex mesh show collision ignoring scale
TTP#336289

diff --git a/Engine/Source/Runtime/Engine/Private/PhysicsEngine/PhysDrawing.cpp b/Engine/Source/Runtime/Engine/Private/PhysicsEngine/PhysDrawing.cpp
--- a/Engine/Source/Runtime/Engine/Private/PhysicsEngine/PhysDrawing.cpp
+++ b/Engine/Source/Runtime/Engine/Private/PhysicsEngine/PhysDrawing.cpp
@@ -520,8 +520,7 @@ void FKAggregateGeom::DrawAggGeom(FPrimitiveDrawInterface* PDI, const FTransform
 			{
 				FColor ConvexColor = bPerHullColor ? DebugUtilColor[i%NUM_DEBUG_UTIL_COLORS] : Color;
 				FTransform ElemTM = ConvexElems[i].GetTransform();
-				ElemTM.ScaleTranslation(Scale3D);
-				ElemTM *= ParentTM;
+				ElemTM *= Transform;
 				ConvexElems[i].DrawElemWire(PDI, ElemTM, ConvexColor);
 
 			}
