Prevent infinite cycles in child actor component spawning

diff --git a/Engine/Source/Runtime/Engine/Private/ChildActorComponent.cpp b/Engine/Source/Runtime/Engine/Private/ChildActorComponent.cpp
--- a/Engine/Source/Runtime/Engine/Private/ChildActorComponent.cpp
+++ b/Engine/Source/Runtime/Engine/Private/ChildActorComponent.cpp
@@ -2,6 +2,8 @@
 
 #include "EnginePrivate.h"
 
+DEFINE_LOG_CATEGORY_STATIC(LogChildActorComponent, Warning, All);
+
 UChildActorComponent::UChildActorComponent(const class FPostConstructInitializeProperties& PCIP)
 	: Super(PCIP)
 {
@@ -32,6 +34,21 @@ void UChildActorComponent::CreateChildActor()
 		UWorld* World = GetWorld();
 		if(World != NULL)
 		{
+			// Before we spawn let's try and prevent cyclic disaster
+			bool bSpawn = true;
+			AActor* Actor = GetOwner();
+			while (Actor && bSpawn)
+			{
+				if (Actor->GetClass() == ChildActorClass)
+				{
+					bSpawn = false;
+					UE_LOG(LogChildActorComponent, Error, TEXT("Found cycle in child actor component '%s'.  Not spawning Actor of class '%s' to break."), *GetPathName(), *ChildActorClass->GetName());
+				}
+				Actor = Actor->ParentComponentActor.Get();
+			}
+
+			if (bSpawn)
+			{
 			FActorSpawnParameters Params;
 			Params.bNoCollisionFail = true;
 			Params.bDeferConstruction = true; // We defer construction so that we set ParentComponentActor prior to component registration so they appear selected
@@ -61,6 +78,7 @@ void UChildActorComponent::CreateChildActor()
 		}
 	}
 }
+}
 
 void UChildActorComponent::DestroyChildActor()
 {
