Fix AttachTo code not correctly blocking on blueprint case

diff --git a/Engine/Source/Runtime/Engine/Classes/Components/BrushComponent.h b/Engine/Source/Runtime/Engine/Classes/Components/BrushComponent.h
--- a/Engine/Source/Runtime/Engine/Classes/Components/BrushComponent.h
+++ b/Engine/Source/Runtime/Engine/Classes/Components/BrushComponent.h
@@ -29,7 +29,11 @@ class UBrushComponent : public UPrimitiveComponent
 	virtual FBoxSphereBounds CalcBounds(const FTransform & LocalToWorld) const OVERRIDE;
 	virtual FVector GetCustomLocation() const OVERRIDE;
 protected:
-	virtual FTransform CalcNewComponentToWorld(const FTransform& NewRelativeTransform) const OVERRIDE;
+	
+	/** Calculate the new ComponentToWorld transform for this component.
+	Parent is optional and can be used for computing ComponentToWorld based on arbitrary USceneComponent.
+	If Parent is not passed in we use the component's AttachParent*/
+	virtual FTransform CalcNewComponentToWorld(const FTransform& NewRelativeTransform, const USceneComponent * Parent = NULL) const OVERRIDE;
 	// End USceneComponent interface
 
 public:
diff --git a/Engine/Source/Runtime/Engine/Classes/Components/SceneComponent.h b/Engine/Source/Runtime/Engine/Classes/Components/SceneComponent.h
--- a/Engine/Source/Runtime/Engine/Classes/Components/SceneComponent.h
+++ b/Engine/Source/Runtime/Engine/Classes/Components/SceneComponent.h
@@ -452,6 +452,7 @@ protected:
 private:
 
 	void PropagateTransformUpdate(bool bTransformChanged, bool bSkipPhysicsMove = false);
+	void UpdateComponentToWorldWithParent(USceneComponent * Parent, bool bSkipPhysicsMove);
 
 
 public:
@@ -620,8 +621,10 @@ public:
 
 protected:
 
-	/** Calculate the new ComponentToWorld transform for this component */
-	virtual FTransform CalcNewComponentToWorld(const FTransform& NewRelativeTransform) const;
+	/** Calculate the new ComponentToWorld transform for this component.
+		Parent is optional and can be used for computing ComponentToWorld based on arbitrary USceneComponent.
+		If Parent is not passed in we use the component's AttachParent*/
+	virtual FTransform CalcNewComponentToWorld(const FTransform& NewRelativeTransform, const USceneComponent * Parent = NULL) const;
 
 	
 public:
diff --git a/Engine/Source/Runtime/Engine/Private/BrushComponent.cpp b/Engine/Source/Runtime/Engine/Private/BrushComponent.cpp
--- a/Engine/Source/Runtime/Engine/Private/BrushComponent.cpp
+++ b/Engine/Source/Runtime/Engine/Private/BrushComponent.cpp
@@ -502,9 +502,9 @@ FVector UBrushComponent::GetCustomLocation() const
 	return LocationNoPivot;
 }
 
-FTransform UBrushComponent::CalcNewComponentToWorld(const FTransform& NewRelativeTransform) const
+FTransform UBrushComponent::CalcNewComponentToWorld(const FTransform& NewRelativeTransform, const USceneComponent * Parent) const
 {
-	FTransform CompToWorld = Super::CalcNewComponentToWorld(NewRelativeTransform);
+	FTransform CompToWorld = Super::CalcNewComponentToWorld(NewRelativeTransform, Parent);
 
 	const FVector LocationNoPivot = CompToWorld.GetLocation();
 	const FVector LocationWithPivot = LocationNoPivot + CompToWorld.TransformVector(-PrePivot);
diff --git a/Engine/Source/Runtime/Engine/Private/SceneComponent.cpp b/Engine/Source/Runtime/Engine/Private/SceneComponent.cpp
--- a/Engine/Source/Runtime/Engine/Private/SceneComponent.cpp
+++ b/Engine/Source/Runtime/Engine/Private/SceneComponent.cpp
@@ -26,14 +26,16 @@ USceneComponent::USceneComponent(const class FPostConstructInitializeProperties&
 	NetUpdateTransform = false;
 }
 
-FTransform USceneComponent::CalcNewComponentToWorld(const FTransform& NewRelativeTransform) const
+FTransform USceneComponent::CalcNewComponentToWorld(const FTransform& NewRelativeTransform, const USceneComponent * Parent) const
 {
 	// With no attachment
 	FTransform ParentToWorld = FTransform::Identity;
 
-	if(AttachParent != NULL)
+	Parent = Parent ? Parent : AttachParent;
+
+	if (Parent != NULL)
 	{
-		ParentToWorld = AttachParent->GetSocketTransform(AttachSocketName);
+		ParentToWorld = Parent->GetSocketTransform(AttachSocketName);
 	}
 
 	FTransform NewCompToWorld = NewRelativeTransform * ParentToWorld;
@@ -60,13 +62,13 @@ void USceneComponent::OnUpdateTransform(bool bSkipPhysicsMove)
 {
 }
 
-void USceneComponent::UpdateComponentToWorld(bool bSkipPhysicsMove)
+void USceneComponent::UpdateComponentToWorldWithParent(USceneComponent * Parent, bool bSkipPhysicsMove)
 {
 	// If our parent hasn't been updated before, we'll need walk up our parent attach hierarchy
-	if (AttachParent && !AttachParent->bWorldToComponentUpdated)
+	if (Parent && !Parent->bWorldToComponentUpdated)
 	{
-		AttachParent->UpdateComponentToWorld();
-		
+		Parent->UpdateComponentToWorld();
+
 		// Updating the parent may (depending on if we were already attached to parent) result in our being updated, so just return
 		if (bWorldToComponentUpdated)
 		{
@@ -77,11 +79,11 @@ void USceneComponent::UpdateComponentToWorld(bool bSkipPhysicsMove)
 	bWorldToComponentUpdated = true;
 
 	// Calculate the new ComponentToWorld transform
-	const FTransform RelativeTransform( RelativeRotation, RelativeLocation, RelativeScale3D );
-	const FTransform NewTransform = CalcNewComponentToWorld(RelativeTransform);
+	const FTransform RelativeTransform(RelativeRotation, RelativeLocation, RelativeScale3D);
+	const FTransform NewTransform = CalcNewComponentToWorld(RelativeTransform, Parent);
 
 	// If transform has changed..
-	if(!ComponentToWorld.Equals(NewTransform, SMALL_NUMBER))
+	if (!ComponentToWorld.Equals(NewTransform, SMALL_NUMBER))
 	{
 		// Update transform
 		ComponentToWorld = NewTransform;
@@ -94,6 +96,11 @@ void USceneComponent::UpdateComponentToWorld(bool bSkipPhysicsMove)
 	}
 }
 
+void USceneComponent::UpdateComponentToWorld(bool bSkipPhysicsMove)
+{
+	UpdateComponentToWorldWithParent(AttachParent, bSkipPhysicsMove);
+}
+
 
 void USceneComponent::PropagateTransformUpdate(bool bTransformChanged, bool bSkipPhysicsMove)
 {
@@ -691,14 +698,24 @@ void USceneComponent::AttachTo(class USceneComponent* Parent, FName InSocketName
 			return;
 		}
 
-		if (IsSimulatingPhysics() && GetWorld() && GetWorld()->IsGameWorld())
 		{
-			//Only at run time we break this attachment
-			RelativeLocation = ComponentToWorld.GetLocation();
-			RelativeRotation = ComponentToWorld.GetRotation().Rotator();
-			RelativeScale3D = ComponentToWorld.GetScale3D();
-			return;
+			//This code requires some explaining. Inside the editor we allow user to attach physically simulated objects to other objects. This is done for convenience so that users can group things together in hierarchy.
+			//At runtime we must not attach physically simulated objects as it will cause double transform updates, and you should just use a physical constraint if attachment is the desired behavior.
+			//We must fixup the relative location,rotation,scale as the attachment is no longer valid. Blueprint uses simple construction to try and attach before ComponentToWorld has ever been updated, so we cannot rely on it.
+			//As such we must calculate the proper Relative information
+			//Also physics state may not be created yet so we use bSimulatePhysics to determine if the object has any intention of being physically simulated
+			UPrimitiveComponent * PrimitiveComponent = Cast<UPrimitiveComponent>(this);
+			if (PrimitiveComponent && PrimitiveComponent->BodyInstance.bSimulatePhysics && GetWorld() && GetWorld()->IsGameWorld())
+			{
+				//Since the object is physically simulated it can't be the case that it's a child of object A and being attached to object B (at runtime)
+				UpdateComponentToWorldWithParent(Parent, false);
+				RelativeLocation = ComponentToWorld.GetLocation();
+				RelativeRotation = ComponentToWorld.GetRotation().Rotator();
+				RelativeScale3D = ComponentToWorld.GetScale3D();
+				return;
+			}
 		}
+		
 
 		// Make sure we are detached
 		bool bMaintainWorldPosition = AttachType == EAttachLocation::KeepWorldPosition;
