AnimNotifyState: Call NotifyBegin events after NotifyEnd events.

diff --git a/Engine/Source/Runtime/Engine/Private/Animation/AnimInstance.cpp b/Engine/Source/Runtime/Engine/Private/Animation/AnimInstance.cpp
--- a/Engine/Source/Runtime/Engine/Private/Animation/AnimInstance.cpp
+++ b/Engine/Source/Runtime/Engine/Private/Animation/AnimInstance.cpp
@@ -1051,9 +1051,13 @@ void UAnimInstance::AddCurveValue(const FName & CurveName, float Value, int32 Cu
 
 void UAnimInstance::TriggerAnimNotifies(float DeltaSeconds)
 {
-	TArray<const FAnimNotifyEvent *> NewActiveAnimNotifyState;
 	USkeletalMeshComponent * SkelMeshComp = GetSkelMeshComponent();
 
+	// Array that will replace the 'ActiveAnimNotifyState' at the end of this function.
+	TArray<const FAnimNotifyEvent *> NewActiveAnimNotifyState;
+	// AnimNotifyState freshly added that need their 'NotifyBegin' event called.
+	TArray<const FAnimNotifyEvent *> NotifyStateBeginEvent;
+
 	// Remove NULL entries.
 	ActiveAnimNotifyState.RemoveSwap(NULL);
 
@@ -1066,7 +1070,8 @@ void UAnimInstance::TriggerAnimNotifies(float DeltaSeconds)
 		{
 			if( !ActiveAnimNotifyState.RemoveSingleSwap(AnimNotifyEvent) )
 			{
-				AnimNotifyEvent->NotifyStateClass->NotifyBegin(SkelMeshComp, Cast<UAnimSequence>(AnimNotifyEvent->NotifyStateClass->GetOuter()));
+				// Queue up calls to 'NotifyBegin', so they happen after 'NotifyEnd'.
+				NotifyStateBeginEvent.Add(AnimNotifyEvent);
 			}
 			NewActiveAnimNotifyState.Add(AnimNotifyEvent);
 			continue;
@@ -1119,6 +1124,13 @@ void UAnimInstance::TriggerAnimNotifies(float DeltaSeconds)
 		AnimNotifyEvent->NotifyStateClass->NotifyEnd(SkelMeshComp, Cast<UAnimSequence>(AnimNotifyEvent->NotifyStateClass->GetOuter()));
 	}
 
+	// Call 'NotifyBegin' event on freshly added AnimNotifyState.
+	for (int32 Index = 0; Index < NotifyStateBeginEvent.Num(); Index++)
+	{
+		const FAnimNotifyEvent * AnimNotifyEvent = NotifyStateBeginEvent[Index];
+		AnimNotifyEvent->NotifyStateClass->NotifyBegin(SkelMeshComp, Cast<UAnimSequence>(AnimNotifyEvent->NotifyStateClass->GetOuter()));
+	}
+
 	// Switch our arrays.
 	ActiveAnimNotifyState = NewActiveAnimNotifyState;
 
