#ue4 - Overlap / Touch notifications are now processed before blocking hit notifications during component movement.

- This is a more logical ordering that matches the hit time ordering of the events (blocking hit is always last).

diff --git a/Engine/Source/Runtime/Engine/Private/PrimitiveComponent.cpp b/Engine/Source/Runtime/Engine/Private/PrimitiveComponent.cpp
--- a/Engine/Source/Runtime/Engine/Private/PrimitiveComponent.cpp
+++ b/Engine/Source/Runtime/Engine/Private/PrimitiveComponent.cpp
@@ -1460,23 +1460,9 @@ bool UPrimitiveComponent::MoveComponent( const FVector& Delta, const FRotator& N
 
 		// Update the location.  This will teleport any child components as well (not sweep).
 		bMoved = InternalSetWorldLocationAndRotation(NewLocation, NewRotation, bSkipPhysicsMove);
-
-		// Handle hit notifications.
-		if (BlockingHit.bBlockingHit)
-		{
-			if (IsDeferringMovementUpdates())
-			{
-				FScopedMovementUpdate* ScopedUpdate = GetCurrentScopedMovement();
-				checkSlow(ScopedUpdate != NULL);
-				ScopedUpdate->AppendBlockingHit(BlockingHit);
-			}
-			else
-			{
-				DispatchBlockingHit(*Actor, BlockingHit);
-			}
-		}
 	}
 
+	// Handle overlap notifications.
 	if (bMoved)
 	{
 		// Check if we are deferring the movement updates.
@@ -1492,13 +1478,29 @@ bool UPrimitiveComponent::MoveComponent( const FVector& Delta, const FRotator& N
 			// still need to do this even if bGenerateOverlapEvents is false for this component, since we could have child components where it is true
 			UpdateOverlaps(&PendingOverlaps, true, OverlapsAtEndLocationPtr);
 		}
+	}
 
-		if (Actor && IsRegistered() && Actor->IsNavigationRelevant() && CanEverAffectNavigation() && IsNavigationRelevant() && World != NULL && World->IsGameWorld() && World->GetNetMode() < ENetMode::NM_Client)
+	// Handle blocking hit notifications.
+	if (BlockingHit.bBlockingHit)
+	{
+		if (IsDeferringMovementUpdates())
 		{
-			if (UNavigationSystem* NavSys = World->GetNavigationSystem())
-			{
-				NavSys->UpdateNavOctree(Actor);
-			}
+			FScopedMovementUpdate* ScopedUpdate = GetCurrentScopedMovement();
+			checkSlow(ScopedUpdate != NULL);
+			ScopedUpdate->AppendBlockingHit(BlockingHit);
+		}
+		else
+		{
+			DispatchBlockingHit(*Actor, BlockingHit);
+		}
+	}
+
+	// Update Nav system if relevant.
+	if (bMoved && Actor && IsRegistered() && Actor->IsNavigationRelevant() && CanEverAffectNavigation() && IsNavigationRelevant() && World != NULL && World->IsGameWorld() && World->GetNetMode() < ENetMode::NM_Client)
+	{
+		if (UNavigationSystem* NavSys = World->GetNavigationSystem())
+		{
+			NavSys->UpdateNavOctree(Actor);
 		}
 	}
 
diff --git a/Engine/Source/Runtime/Engine/Private/SceneComponent.cpp b/Engine/Source/Runtime/Engine/Private/SceneComponent.cpp
--- a/Engine/Source/Runtime/Engine/Private/SceneComponent.cpp
+++ b/Engine/Source/Runtime/Engine/Private/SceneComponent.cpp
@@ -197,6 +197,12 @@ void USceneComponent::EndScopedMovementUpdate(class FScopedMovementUpdate& Compl
 				PropagateTransformUpdate(true);
 			}
 
+			// We may have moved somewhere and then moved back to the start, we still need to update overlaps if we touched things along the way.
+			if (bMoved || CurrentScopedUpdate->HasPendingOverlaps() || CurrentScopedUpdate->GetOverlapsAtEnd())
+			{
+				UpdateOverlaps(&CurrentScopedUpdate->GetPendingOverlaps(), true, CurrentScopedUpdate->GetOverlapsAtEnd());
+			}
+
 			// Dispatch all deferred blocking hits
 			if (CompletedScope.BlockingHits.Num() > 0)
 			{
@@ -209,13 +215,7 @@ void USceneComponent::EndScopedMovementUpdate(class FScopedMovementUpdate& Compl
 					{
 						PrimitiveThis->DispatchBlockingHit(*Owner, Hit);
 					}
-				}				
-			}
-
-			// We may have moved somewhere and then moved back to the start, we still need to update overlaps if we touched things along the way.
-			if (bMoved || CurrentScopedUpdate->HasPendingOverlaps() || CurrentScopedUpdate->GetOverlapsAtEnd())
-			{
-				UpdateOverlaps(&CurrentScopedUpdate->GetPendingOverlaps(), true, CurrentScopedUpdate->GetOverlapsAtEnd());
+				}
 			}
 		}
 		else
