Fixed collision disappearing in the static mesh editor when sockets are added

This addresses TTP#335270 - EDITOR: Static Mesh: The collision mesh will disappear after adding a socket

diff --git a/Engine/Source/Runtime/Engine/Private/StaticMeshBuild.cpp b/Engine/Source/Runtime/Engine/Private/StaticMeshBuild.cpp
--- a/Engine/Source/Runtime/Engine/Private/StaticMeshBuild.cpp
+++ b/Engine/Source/Runtime/Engine/Private/StaticMeshBuild.cpp
@@ -100,6 +100,7 @@ void UStaticMesh::Build(bool bSilent)
 	// Invalidate physics data if this has changed.
 	// TODO_STATICMESH: Not necessary any longer?
 	BodySetup->InvalidatePhysicsData();
+	BodySetup->CreatePhysicsMeshes();
 #endif
 
 	// Compare the derived data keys to see if renderable mesh data has actually changed.
