Fixed nontransactional actors getting removed from the level's actor list on undo.
Fixes TTP 335800 - Regression: EDITOR: LANDSCAPE: CRASH: Possible to get a duplicate LandscapeGizmoActiveActor in the scene outliner which can lead to a crash

diff --git a/Engine/Source/Runtime/Engine/Classes/Engine/Level.h b/Engine/Source/Runtime/Engine/Classes/Engine/Level.h
--- a/Engine/Source/Runtime/Engine/Classes/Engine/Level.h
+++ b/Engine/Source/Runtime/Engine/Classes/Engine/Level.h
@@ -299,7 +299,7 @@ class ULevel : public ULevelBase, public IInterface_AssetUserData
 
 #if WITH_EDITORONLY_DATA
 	/** Reference to the blueprint for level scripting */
-	UPROPERTY()
+	UPROPERTY(NonTransactional)
 	class ULevelScriptBlueprint* LevelScriptBlueprint;
 #endif //WITH_EDITORONLY_DATA
 
diff --git a/Engine/Source/Runtime/Engine/Private/Level.cpp b/Engine/Source/Runtime/Engine/Private/Level.cpp
--- a/Engine/Source/Runtime/Engine/Private/Level.cpp
+++ b/Engine/Source/Runtime/Engine/Private/Level.cpp
@@ -1120,13 +1120,18 @@ void ULevel::PostEditUndo()
 		}
 	}
 
-	// Add the levelscriptactor back into the actors array
-	if (LevelScriptActor)
-	{
-		// LSA is transient so won't track the actors list
-		if (!Actors.Contains(LevelScriptActor))
-		{
-			Actors.Add(LevelScriptActor);
+	// Non-transactional actors may disappear from the actors list but still exist, so we need to re-add them
+	// Likewise they won't get recreated if we undo to before they were deleted, so we'll have nulls in the actors list to remove
+	Actors.Remove(nullptr);
+	TSet<AActor*> ActorsSet(Actors);
+	TArray<UObject *> InnerObjects;
+	GetObjectsWithOuter(this, InnerObjects, /*bIncludeNestedObjects*/ false, /*ExclusionFlags*/ RF_PendingKill);
+	for (UObject* InnerObject : InnerObjects)
+	{
+		AActor* InnerActor = Cast<AActor>(InnerObject);
+		if (InnerActor && !ActorsSet.Contains(InnerActor))
+		{
+			Actors.Add(InnerActor);
 		}
 	}
 
