Remove Interpolation.h from Engine.h

diff --git a/Engine/Source/Runtime/Engine/Classes/Matinee/InterpTrack.h b/Engine/Source/Runtime/Engine/Classes/Matinee/InterpTrack.h
--- a/Engine/Source/Runtime/Engine/Classes/Matinee/InterpTrack.h
+++ b/Engine/Source/Runtime/Engine/Classes/Matinee/InterpTrack.h
@@ -7,7 +7,7 @@
  */
 
 #pragma once
-
+#include "Interpolation.h"
 #include "InterpTrack.generated.h"
 
 //
diff --git a/Engine/Source/Runtime/Engine/Public/Engine.h b/Engine/Source/Runtime/Engine/Public/Engine.h
--- a/Engine/Source/Runtime/Engine/Public/Engine.h
+++ b/Engine/Source/Runtime/Engine/Public/Engine.h
@@ -911,7 +911,6 @@
 #include "AnimTree.h"						// Animation.
 #include "SkeletalMeshTypes.h"				// Skeletal animated mesh.
 #include "Animation/SkeletalMeshActor.h"
-#include "Interpolation.h"					// Matinee.
 #include "ContentStreaming.h"				// Content streaming class definitions.
 #include "LightingBuildOptions.h"			// Definition of lighting build option struct.
 #include "PixelFormat.h"
diff --git a/Engine/Source/Runtime/Engine/Public/InterpolationHitProxy.h b/Engine/Source/Runtime/Engine/Public/InterpolationHitProxy.h
--- a/Engine/Source/Runtime/Engine/Public/InterpolationHitProxy.h
+++ b/Engine/Source/Runtime/Engine/Public/InterpolationHitProxy.h
@@ -7,6 +7,8 @@
 #ifndef _INC_INTERPOLATIONHITPROXY
 #define _INC_INTERPOLATIONHITPROXY
 
+#include "Interpolation.h"
+
 /** Input interface hit proxy */
 struct HInterpEdInputInterface : public HHitProxy
 {
