Remove landscape/foliage headers from Engine.h and include only where needed

diff --git a/Engine/Source/Developer/AssetTools/Private/AssetTypeActions/AssetTypeActions_LandscapeLayer.h b/Engine/Source/Developer/AssetTools/Private/AssetTypeActions/AssetTypeActions_LandscapeLayer.h
--- a/Engine/Source/Developer/AssetTools/Private/AssetTypeActions/AssetTypeActions_LandscapeLayer.h
+++ b/Engine/Source/Developer/AssetTools/Private/AssetTypeActions/AssetTypeActions_LandscapeLayer.h
@@ -2,6 +2,8 @@
 
 #pragma once
 
+#include "Landscape/LandscapeLayerInfoObject.h"
+
 class FAssetTypeActions_LandscapeLayer : public FAssetTypeActions_Base
 {
 public:
diff --git a/Engine/Source/Editor/DetailCustomizations/Private/ActorMaterialCategory.cpp b/Engine/Source/Editor/DetailCustomizations/Private/ActorMaterialCategory.cpp
--- a/Engine/Source/Editor/DetailCustomizations/Private/ActorMaterialCategory.cpp
+++ b/Engine/Source/Editor/DetailCustomizations/Private/ActorMaterialCategory.cpp
@@ -5,6 +5,9 @@
 #include "AssetThumbnail.h"
 #include "ActorEditorUtils.h"
 
+#include "Landscape/LandscapeProxy.h"
+#include "Landscape/LandscapeComponent.h"
+
 #define LOCTEXT_NAMESPACE "SMaterialList"
 
 /**
diff --git a/Engine/Source/Editor/FoliageEdit/Private/FoliageEdMode.cpp b/Engine/Source/Editor/FoliageEdit/Private/FoliageEdMode.cpp
--- a/Engine/Source/Editor/FoliageEdit/Private/FoliageEdMode.cpp
+++ b/Engine/Source/Editor/FoliageEdit/Private/FoliageEdMode.cpp
@@ -20,6 +20,12 @@
 #include "Editor/LevelEditor/Public/LevelEditor.h"
 #include "Editor/LevelEditor/Public/SLevelViewport.h"
 
+// Classes
+#include "Foliage/InstancedFoliageActor.h"
+#include "Landscape/LandscapeComponent.h"
+#include "Landscape/LandscapeHeightfieldCollisionComponent.h"
+#include "Landscape/LandscapeInfo.h"
+
 #define FOLIAGE_SNAP_TRACE (10000.f)
 
 //
diff --git a/Engine/Source/Editor/LandscapeEditor/Classes/LandscapeEditorObject.h b/Engine/Source/Editor/LandscapeEditor/Classes/LandscapeEditorObject.h
--- a/Engine/Source/Editor/LandscapeEditor/Classes/LandscapeEditorObject.h
+++ b/Engine/Source/Editor/LandscapeEditor/Classes/LandscapeEditorObject.h
@@ -2,6 +2,7 @@
 
 #pragma once
 
+#include "Landscape/LandscapeProxy.h"
 #include "LandscapeEditorObject.generated.h"
 
 UENUM()
diff --git a/Engine/Source/Editor/LandscapeEditor/Private/LandscapeEdMode.cpp b/Engine/Source/Editor/LandscapeEditor/Private/LandscapeEdMode.cpp
--- a/Engine/Source/Editor/LandscapeEditor/Private/LandscapeEdMode.cpp
+++ b/Engine/Source/Editor/LandscapeEditor/Private/LandscapeEdMode.cpp
@@ -23,6 +23,13 @@
 #include "Editor/LevelEditor/Public/SLevelViewport.h"
 #include "SLandscapeEditor.h"
 
+// Classes
+#include "Landscape/Landscape.h"
+#include "Landscape/LandscapeHeightfieldCollisionComponent.h"
+#include "Landscape/LandscapeMaterialInstanceConstant.h"
+#include "Landscape/LandscapeSplinesComponent.h"
+#include "Foliage/InstancedFoliageActor.h"
+
 #define LOCTEXT_NAMESPACE "Landscape"
 
 DEFINE_LOG_CATEGORY(LogLandscapeEdMode);
diff --git a/Engine/Source/Editor/LandscapeEditor/Private/LandscapeEdModeBrushes.cpp b/Engine/Source/Editor/LandscapeEditor/Private/LandscapeEdModeBrushes.cpp
--- a/Engine/Source/Editor/LandscapeEditor/Private/LandscapeEdModeBrushes.cpp
+++ b/Engine/Source/Editor/LandscapeEditor/Private/LandscapeEdModeBrushes.cpp
@@ -5,8 +5,10 @@
 #include "ObjectTools.h"
 #include "LandscapeEdMode.h"
 #include "ScopedTransaction.h"
-#include "Runtime/Engine/Public/Landscape/LandscapeEdit.h"
-#include "Runtime/Engine/Public/Landscape/LandscapeRender.h"
+#include "Landscape/LandscapeEdit.h"
+#include "Landscape/LandscapeRender.h"
+#include "Landscape/Landscape.h"
+#include "Landscape/LandscapeHeightfieldCollisionComponent.h"
 
 #include "LevelUtils.h"
 
diff --git a/Engine/Source/Editor/LandscapeEditor/Private/LandscapeEdModeComponentTools.cpp b/Engine/Source/Editor/LandscapeEditor/Private/LandscapeEdModeComponentTools.cpp
--- a/Engine/Source/Editor/LandscapeEditor/Private/LandscapeEdModeComponentTools.cpp
+++ b/Engine/Source/Editor/LandscapeEditor/Private/LandscapeEdModeComponentTools.cpp
@@ -4,6 +4,7 @@
 #include "ObjectTools.h"
 #include "LandscapeEdMode.h"
 #include "ScopedTransaction.h"
+#include "Landscape/Landscape.h"
 #include "Landscape/LandscapeEdit.h"
 #include "Landscape/LandscapeRender.h"
 #include "Landscape/LandscapeDataAccess.h"
@@ -11,6 +12,7 @@
 #include "LandscapeEditorModule.h"
 #include "Editor/PropertyEditor/Public/PropertyEditorModule.h"
 #include "LandscapeEdModeTools.h"
+#include "Foliage/InstancedFoliageActor.h"
 
 #define LOCTEXT_NAMESPACE "Landscape"
 
diff --git a/Engine/Source/Editor/LandscapeEditor/Private/LandscapeEdModeErosionTools.cpp b/Engine/Source/Editor/LandscapeEditor/Private/LandscapeEdModeErosionTools.cpp
--- a/Engine/Source/Editor/LandscapeEditor/Private/LandscapeEdModeErosionTools.cpp
+++ b/Engine/Source/Editor/LandscapeEditor/Private/LandscapeEdModeErosionTools.cpp
@@ -11,6 +11,7 @@
 #include "LandscapeEditorModule.h"
 #include "Editor/PropertyEditor/Public/PropertyEditorModule.h"
 #include "LandscapeEdModeTools.h"
+#include "Landscape/Landscape.h"
 
 //
 // FLandscapeToolErosionBase
diff --git a/Engine/Source/Editor/LandscapeEditor/Private/LandscapeEdModeRampTool.cpp b/Engine/Source/Editor/LandscapeEditor/Private/LandscapeEdModeRampTool.cpp
--- a/Engine/Source/Editor/LandscapeEditor/Private/LandscapeEdModeRampTool.cpp
+++ b/Engine/Source/Editor/LandscapeEditor/Private/LandscapeEdModeRampTool.cpp
@@ -1,17 +1,12 @@
 // Copyright 1998-2014 Epic Games, Inc. All Rights Reserved.
 
 #include "LandscapeEditorPrivatePCH.h"
-//#include "ObjectTools.h"
 #include "LandscapeEdMode.h"
 #include "ScopedTransaction.h"
-//#include "EngineFoliageClasses.h"
 #include "Landscape/LandscapeEdit.h"
 #include "Landscape/LandscapeRender.h"
 #include "Landscape/LandscapeDataAccess.h"
-//#include "Landscape/LandscapeSplineProxies.h"
-//#include "LandscapeEditorModule.h"
-//#include "Editor/PropertyEditor/Public/PropertyEditorModule.h"
-//#include "LandscapeEdModeTools.h"
+#include "Landscape/LandscapeHeightfieldCollisionComponent.h"
 #include "Raster.h"
 
 #define LOCTEXT_NAMESPACE "Landscape"
diff --git a/Engine/Source/Editor/LandscapeEditor/Private/LandscapeEdModeSplineTools.cpp b/Engine/Source/Editor/LandscapeEditor/Private/LandscapeEdModeSplineTools.cpp
--- a/Engine/Source/Editor/LandscapeEditor/Private/LandscapeEdModeSplineTools.cpp
+++ b/Engine/Source/Editor/LandscapeEditor/Private/LandscapeEdModeSplineTools.cpp
@@ -13,6 +13,10 @@
 #include "LandscapeEdModeTools.h"
 #include "Components/SplineMeshComponent.h"
 #include "LandscapeSplineImportExport.h"
+#include "Landscape/LandscapeSplinesComponent.h"
+#include "Landscape/LandscapeSplineControlPoint.h"
+#include "Landscape/LandscapeSplineSegment.h"
+#include "Landscape/ControlPointMeshComponent.h"
 
 
 #define LOCTEXT_NAMESPACE "Landscape"
diff --git a/Engine/Source/Editor/LandscapeEditor/Private/LandscapeEdModeTools.h b/Engine/Source/Editor/LandscapeEditor/Private/LandscapeEdModeTools.h
--- a/Engine/Source/Editor/LandscapeEditor/Private/LandscapeEdModeTools.h
+++ b/Engine/Source/Editor/LandscapeEditor/Private/LandscapeEdModeTools.h
@@ -2,6 +2,10 @@
 
 #pragma once
 
+#include "Landscape/Landscape.h"
+#include "Landscape/LandscapeHeightfieldCollisionComponent.h"
+#include "Foliage/InstancedFoliageActor.h"
+
 //
 //	FNoiseParameter - Perlin noise
 //
diff --git a/Engine/Source/Editor/LandscapeEditor/Private/LandscapeEdModeXYOffset.cpp b/Engine/Source/Editor/LandscapeEditor/Private/LandscapeEdModeXYOffset.cpp
--- a/Engine/Source/Editor/LandscapeEditor/Private/LandscapeEdModeXYOffset.cpp
+++ b/Engine/Source/Editor/LandscapeEditor/Private/LandscapeEdModeXYOffset.cpp
@@ -11,6 +11,7 @@
 #include "LandscapeEditorModule.h"
 #include "Editor/PropertyEditor/Public/PropertyEditorModule.h"
 #include "LandscapeEdModeTools.h"
+#include "Landscape/Landscape.h"
 
 namespace
 {
diff --git a/Engine/Source/Editor/LandscapeEditor/Private/LandscapeEditorDetailCustomization_NewLandscape.cpp b/Engine/Source/Editor/LandscapeEditor/Private/LandscapeEditorDetailCustomization_NewLandscape.cpp
--- a/Engine/Source/Editor/LandscapeEditor/Private/LandscapeEditorDetailCustomization_NewLandscape.cpp
+++ b/Engine/Source/Editor/LandscapeEditor/Private/LandscapeEditorDetailCustomization_NewLandscape.cpp
@@ -27,6 +27,8 @@
 #include "AssetRegistryModule.h"
 #include "ImageWrapper.h"
 
+#include "Landscape/Landscape.h"
+
 #define LOCTEXT_NAMESPACE "LandscapeEditor.NewLandscape"
 
 const int32 FLandscapeEditorDetailCustomization_NewLandscape::SectionSizes[] = {7, 15, 31, 63, 127, 255};
diff --git a/Engine/Source/Editor/LandscapeEditor/Private/LandscapeEditorModule.cpp b/Engine/Source/Editor/LandscapeEditor/Private/LandscapeEditorModule.cpp
--- a/Engine/Source/Editor/LandscapeEditor/Private/LandscapeEditorModule.cpp
+++ b/Engine/Source/Editor/LandscapeEditor/Private/LandscapeEditorModule.cpp
@@ -14,6 +14,8 @@
 #include "Editor/LevelEditor/Public/LevelEditor.h"
 #include "LandscapeRender.h"
 
+#include "Landscape/Landscape.h"
+
 #define LOCTEXT_NAMESPACE "LandscapeEditor"
 
 class FLandscapeEditorModule : public ILandscapeEditorModule
diff --git a/Engine/Source/Editor/LandscapeEditor/Private/LandscapeEditorObject.cpp b/Engine/Source/Editor/LandscapeEditor/Private/LandscapeEditorObject.cpp
--- a/Engine/Source/Editor/LandscapeEditor/Private/LandscapeEditorObject.cpp
+++ b/Engine/Source/Editor/LandscapeEditor/Private/LandscapeEditorObject.cpp
@@ -5,6 +5,8 @@
 #include "LandscapeRender.h"
 #include "ImageWrapper.h"
 
+#include "Landscape/LandscapeMaterialInstanceConstant.h"
+
 ULandscapeEditorObject::ULandscapeEditorObject(const class FPostConstructInitializeProperties& PCIP)
 	: Super(PCIP)
 
diff --git a/Engine/Source/Editor/LandscapeEditor/Private/LandscapeSplineImportExport.cpp b/Engine/Source/Editor/LandscapeEditor/Private/LandscapeSplineImportExport.cpp
--- a/Engine/Source/Editor/LandscapeEditor/Private/LandscapeSplineImportExport.cpp
+++ b/Engine/Source/Editor/LandscapeEditor/Private/LandscapeSplineImportExport.cpp
@@ -2,7 +2,8 @@
 
 #include "LandscapeEditorPrivatePCH.h"
 #include "LandscapeSplineImportExport.h"
-
+#include "Landscape/LandscapeSplineControlPoint.h"
+#include "Landscape/LandscapeSplineSegment.h"
 
 #define LOCTEXT_NAMESPACE "Landscape"
 
diff --git a/Engine/Source/Editor/LandscapeEditor/Public/LandscapeEdMode.h b/Engine/Source/Editor/LandscapeEditor/Public/LandscapeEdMode.h
--- a/Engine/Source/Editor/LandscapeEditor/Public/LandscapeEdMode.h
+++ b/Engine/Source/Editor/LandscapeEditor/Public/LandscapeEdMode.h
@@ -6,6 +6,12 @@
 
 DECLARE_LOG_CATEGORY_EXTERN(LogLandscapeEdMode, Log, All);
 
+// Forward declarations
+class ULandscapeLayerInfoObject;
+#include "Landscape/LandscapeInfo.h"
+#include "Landscape/LandscapeProxy.h"
+#include "Landscape/LandscapeGizmoActiveActor.h"
+
 // FLandscapeToolMousePosition - Struct to store mouse positions since the last time we applied the brush
 struct FLandscapeToolMousePosition
 {
diff --git a/Engine/Source/Editor/StatsViewer/Private/StatsPages/PrimitiveStatsPage.cpp b/Engine/Source/Editor/StatsViewer/Private/StatsPages/PrimitiveStatsPage.cpp
--- a/Engine/Source/Editor/StatsViewer/Private/StatsPages/PrimitiveStatsPage.cpp
+++ b/Engine/Source/Editor/StatsViewer/Private/StatsPages/PrimitiveStatsPage.cpp
@@ -3,6 +3,8 @@
 #include "StatsViewerPrivatePCH.h"
 #include "PrimitiveStatsPage.h"
 #include "PrimitiveStats.h"
+#include "Landscape/LandscapeProxy.h"
+#include "Landscape/LandscapeComponent.h"
 
 #define LOCTEXT_NAMESPACE "Editor.StatsViewer.PrimitiveStats"
 
diff --git a/Engine/Source/Editor/UnrealEd/Private/AssetSelection.cpp b/Engine/Source/Editor/UnrealEd/Private/AssetSelection.cpp
--- a/Engine/Source/Editor/UnrealEd/Private/AssetSelection.cpp
+++ b/Engine/Source/Editor/UnrealEd/Private/AssetSelection.cpp
@@ -18,6 +18,7 @@
 #include "SnappingUtils.h"
 #include "ActorEditorUtils.h"
 #include "LevelEditorViewport.h"
+#include "Landscape/Landscape.h"
 
 extern void OnPlaceStaticMeshActor( AActor* MeshActor, bool bUseSurfaceOrientation );
 
diff --git a/Engine/Source/Editor/UnrealEd/Private/EditorActor.cpp b/Engine/Source/Editor/UnrealEd/Private/EditorActor.cpp
--- a/Engine/Source/Editor/UnrealEd/Private/EditorActor.cpp
+++ b/Engine/Source/Editor/UnrealEd/Private/EditorActor.cpp
@@ -13,6 +13,7 @@
 #include "Editor/GeometryMode/Public/GeometryEdMode.h"
 #include "Editor/GeometryMode/Public/EditorGeometry.h"
 #include "ActorEditorUtils.h"
+#include "Foliage/InstancedFoliageActor.h"
 
 #define LOCTEXT_NAMESPACE "UnrealEd.EditorActor"
 
diff --git a/Engine/Source/Editor/UnrealEd/Private/EditorExporters.cpp b/Engine/Source/Editor/UnrealEd/Private/EditorExporters.cpp
--- a/Engine/Source/Editor/UnrealEd/Private/EditorExporters.cpp
+++ b/Engine/Source/Editor/UnrealEd/Private/EditorExporters.cpp
@@ -15,6 +15,8 @@
 #include "RawMesh.h"
 #include "MaterialExportUtils.h"
 #include "ImageUtils.h"
+#include "Foliage/InstancedFoliageActor.h"
+#include "Landscape/Landscape.h"
 
 DEFINE_LOG_CATEGORY_STATIC(LogEditorExporters, Log, All);
 
diff --git a/Engine/Source/Editor/UnrealEd/Private/EditorObject.cpp b/Engine/Source/Editor/UnrealEd/Private/EditorObject.cpp
--- a/Engine/Source/Editor/UnrealEd/Private/EditorObject.cpp
+++ b/Engine/Source/Editor/UnrealEd/Private/EditorObject.cpp
@@ -8,6 +8,8 @@
 #include "Factories.h"
 #include "BSPOps.h"
 
+#include "Foliage/InstancedFoliageActor.h"
+
 DEFINE_LOG_CATEGORY_STATIC(LogEditorObject, Log, All);
 
 /*
diff --git a/Engine/Source/Editor/UnrealEd/Private/EditorServer.cpp b/Engine/Source/Editor/UnrealEd/Private/EditorServer.cpp
--- a/Engine/Source/Editor/UnrealEd/Private/EditorServer.cpp
+++ b/Engine/Source/Editor/UnrealEd/Private/EditorServer.cpp
@@ -20,6 +20,7 @@
 #include "Editor/UnrealEd/Public/Kismet2/BlueprintEditorUtils.h"
 #include "../Private/GeomFitUtils.h"
 #include "Editor/GeometryMode/Public/GeometryEdMode.h"
+#include "Landscape/LandscapeProxy.h"
 
 #include "Matinee/InterpTrackAnimControl.h"
 #include "Matinee/InterpData.h"
diff --git a/Engine/Source/Editor/UnrealEd/Private/Factories/EditorFactories.cpp b/Engine/Source/Editor/UnrealEd/Private/Factories/EditorFactories.cpp
--- a/Engine/Source/Editor/UnrealEd/Private/Factories/EditorFactories.cpp
+++ b/Engine/Source/Editor/UnrealEd/Private/Factories/EditorFactories.cpp
@@ -36,6 +36,8 @@
 #include "EnumEditorUtils.h"
 #include "StructureEditorUtils.h"
 
+#include "Foliage/InstancedFoliageActor.h"
+
 #if PLATFORM_WINDOWS
 // Needed for DDS support.
 #include "AllowWindowsPlatformTypes.h"
diff --git a/Engine/Source/Editor/UnrealEd/Private/FbxExporter.h b/Engine/Source/Editor/UnrealEd/Private/FbxExporter.h
--- a/Engine/Source/Editor/UnrealEd/Private/FbxExporter.h
+++ b/Engine/Source/Editor/UnrealEd/Private/FbxExporter.h
@@ -34,6 +34,8 @@
 #include "MatineeExporter.h"
 #include "FbxImporter.h"
 
+#include "Landscape/LandscapeProxy.h"
+
 namespace UnFbx
 {
 
diff --git a/Engine/Source/Editor/UnrealEd/Private/StaticLightingSystem/StaticLightingExport.cpp b/Engine/Source/Editor/UnrealEd/Private/StaticLightingSystem/StaticLightingExport.cpp
--- a/Engine/Source/Editor/UnrealEd/Private/StaticLightingSystem/StaticLightingExport.cpp
+++ b/Engine/Source/Editor/UnrealEd/Private/StaticLightingSystem/StaticLightingExport.cpp
@@ -9,6 +9,7 @@
 #include "StaticMeshLight.h"
 #include "ModelLight.h"
 #include "Runtime/Engine/Public/Landscape/LandscapeLight.h"
+#include "Runtime/Engine/Classes/Landscape/LandscapeComponent.h"
 
 #include "Lightmass/Lightmass.h"
 
diff --git a/Engine/Source/Editor/UnrealEd/Private/UnrealEdSrv.cpp b/Engine/Source/Editor/UnrealEd/Private/UnrealEdSrv.cpp
--- a/Engine/Source/Editor/UnrealEd/Private/UnrealEdSrv.cpp
+++ b/Engine/Source/Editor/UnrealEd/Private/UnrealEdSrv.cpp
@@ -25,6 +25,9 @@
 #include "HighResScreenshot.h"
 #include "ActorEditorUtils.h"
 #include "Matinee/InterpData.h"
+#include "Landscape/LandscapeGizmoActiveActor.h"
+#include "Landscape/LandscapeComponent.h"
+#include "Landscape/LandscapeHeightfieldCollisionComponent.h"
 
 DEFINE_LOG_CATEGORY_STATIC(LogUnrealEdSrv, Log, All);
 
diff --git a/Engine/Source/Editor/UnrealEd/Public/MaterialExportUtils.h b/Engine/Source/Editor/UnrealEd/Public/MaterialExportUtils.h
--- a/Engine/Source/Editor/UnrealEd/Public/MaterialExportUtils.h
+++ b/Engine/Source/Editor/UnrealEd/Public/MaterialExportUtils.h
@@ -1,6 +1,9 @@
 // Copyright 1998-2014 Epic Games, Inc. All Rights Reserved.
 #pragma once
 
+// Forward declarations
+class ALandscapeProxy;
+
 namespace MaterialExportUtils
 {
 	struct UNREALED_API FFlattenMaterial
diff --git a/Engine/Source/Editor/WorldBrowser/Private/Tiles/STiledLandscapeImportDlg.cpp b/Engine/Source/Editor/WorldBrowser/Private/Tiles/STiledLandscapeImportDlg.cpp
--- a/Engine/Source/Editor/WorldBrowser/Private/Tiles/STiledLandscapeImportDlg.cpp
+++ b/Engine/Source/Editor/WorldBrowser/Private/Tiles/STiledLandscapeImportDlg.cpp
@@ -6,6 +6,7 @@
 #include "IDesktopPlatform.h"
 #include "DesktopPlatformModule.h"
 #include "ContentBrowserModule.h"
+#include "Landscape/LandscapeProxy.h"
 
 #define LOCTEXT_NAMESPACE "WorldBrowser"
 
diff --git a/Engine/Source/Editor/WorldBrowser/Private/Tiles/WorldTileModel.cpp b/Engine/Source/Editor/WorldBrowser/Private/Tiles/WorldTileModel.cpp
--- a/Engine/Source/Editor/WorldBrowser/Private/Tiles/WorldTileModel.cpp
+++ b/Engine/Source/Editor/WorldBrowser/Private/Tiles/WorldTileModel.cpp
@@ -8,6 +8,8 @@
 #include "WorldTileCollectionModel.h"
 #include "WorldTileModel.h"
 
+#include "Landscape/Landscape.h"
+
 #define LOCTEXT_NAMESPACE "WorldBrowser"
 DEFINE_LOG_CATEGORY_STATIC(WorldBrowser, Log, All);
 
diff --git a/Engine/Source/Editor/WorldBrowser/Private/Tiles/WorldTileModel.h b/Engine/Source/Editor/WorldBrowser/Private/Tiles/WorldTileModel.h
--- a/Engine/Source/Editor/WorldBrowser/Private/Tiles/WorldTileModel.h
+++ b/Engine/Source/Editor/WorldBrowser/Private/Tiles/WorldTileModel.h
@@ -1,8 +1,12 @@
 // Copyright 1998-2014 Epic Games, Inc. All Rights Reserved.
+
 #pragma once
 
+#include "Landscape/LandscapeProxy.h"
+
 class FWorldTileCollectionModel;
 class UWorldTileDetails;
+
 //
 typedef TArray<TSharedPtr<class FWorldTileModel>> FWorldTileModelList;
 
diff --git a/Engine/Source/Runtime/Engine/Classes/Landscape/Landscape.h b/Engine/Source/Runtime/Engine/Classes/Landscape/Landscape.h
--- a/Engine/Source/Runtime/Engine/Classes/Landscape/Landscape.h
+++ b/Engine/Source/Runtime/Engine/Classes/Landscape/Landscape.h
@@ -2,6 +2,7 @@
 
 #pragma once
 #include "LandscapeInfo.h"
+#include "LandscapeProxy.h"
 #include "Landscape.generated.h"
 
 UENUM()
diff --git a/Engine/Source/Runtime/Engine/Classes/Landscape/LandscapeGizmoActiveActor.h b/Engine/Source/Runtime/Engine/Classes/Landscape/LandscapeGizmoActiveActor.h
--- a/Engine/Source/Runtime/Engine/Classes/Landscape/LandscapeGizmoActiveActor.h
+++ b/Engine/Source/Runtime/Engine/Classes/Landscape/LandscapeGizmoActiveActor.h
@@ -2,6 +2,8 @@
 
 
 #pragma once
+
+#include "LandscapeGizmoActor.h"
 #include "LandscapeGizmoActiveActor.generated.h"
 
 UENUM()
diff --git a/Engine/Source/Runtime/Engine/Classes/Landscape/LandscapeHeightfieldCollisionComponent.h b/Engine/Source/Runtime/Engine/Classes/Landscape/LandscapeHeightfieldCollisionComponent.h
--- a/Engine/Source/Runtime/Engine/Classes/Landscape/LandscapeHeightfieldCollisionComponent.h
+++ b/Engine/Source/Runtime/Engine/Classes/Landscape/LandscapeHeightfieldCollisionComponent.h
@@ -2,6 +2,9 @@
 
 
 #pragma once
+
+class ULandscapeLayerInfoObject;
+
 #include "LandscapeHeightfieldCollisionComponent.generated.h"
 
 #if WITH_PHYSX
diff --git a/Engine/Source/Runtime/Engine/Classes/Landscape/LandscapeMeshCollisionComponent.h b/Engine/Source/Runtime/Engine/Classes/Landscape/LandscapeMeshCollisionComponent.h
--- a/Engine/Source/Runtime/Engine/Classes/Landscape/LandscapeMeshCollisionComponent.h
+++ b/Engine/Source/Runtime/Engine/Classes/Landscape/LandscapeMeshCollisionComponent.h
@@ -2,6 +2,8 @@
 
 
 #pragma once
+
+#include "Landscape/LandscapeHeightfieldCollisionComponent.h"
 #include "LandscapeMeshCollisionComponent.generated.h"
 
 #if WITH_PHYSX
diff --git a/Engine/Source/Runtime/Engine/Private/ActorEditor.cpp b/Engine/Source/Runtime/Engine/Private/ActorEditor.cpp
--- a/Engine/Source/Runtime/Engine/Private/ActorEditor.cpp
+++ b/Engine/Source/Runtime/Engine/Private/ActorEditor.cpp
@@ -7,6 +7,9 @@
 #include "UObjectToken.h"
 #include "LevelUtils.h"
 #include "MapErrors.h"
+#include "Foliage/InstancedFoliageActor.h"
+#include "Landscape/LandscapeHeightfieldCollisionComponent.h"
+#include "Landscape/LandscapeComponent.h"
 
 #if WITH_EDITOR
 
diff --git a/Engine/Source/Runtime/Engine/Private/InstancedFoliage.cpp b/Engine/Source/Runtime/Engine/Private/InstancedFoliage.cpp
--- a/Engine/Source/Runtime/Engine/Private/InstancedFoliage.cpp
+++ b/Engine/Source/Runtime/Engine/Private/InstancedFoliage.cpp
@@ -8,6 +8,8 @@
 #include "MessageLog.h"
 #include "UObjectToken.h"
 #include "MapErrors.h"
+#include "Foliage/InstancedFoliageActor.h"
+#include "Landscape/LandscapeHeightfieldCollisionComponent.h"
 
 #define LOCTEXT_NAMESPACE "InstancedFoliage"
 
diff --git a/Engine/Source/Runtime/Engine/Private/Landscape/Landscape.cpp b/Engine/Source/Runtime/Engine/Private/Landscape/Landscape.cpp
--- a/Engine/Source/Runtime/Engine/Private/Landscape/Landscape.cpp
+++ b/Engine/Source/Runtime/Engine/Private/Landscape/Landscape.cpp
@@ -14,6 +14,9 @@ Landscape.cpp: Terrain rendering
 #include "MapErrors.h"
 #include "DerivedDataCacheInterface.h"
 #include "TargetPlatform.h"
+#include "Landscape/Landscape.h"
+#include "Landscape/LandscapeMeshCollisionComponent.h"
+#include "Landscape/LandscapeSplinesComponent.h"
 
 // Set this to 0 to disable landscape cooking and thus disable it on device.
 #define ENABLE_LANDSCAPE_COOKING 1
diff --git a/Engine/Source/Runtime/Engine/Private/Landscape/LandscapeCollision.cpp b/Engine/Source/Runtime/Engine/Private/Landscape/LandscapeCollision.cpp
--- a/Engine/Source/Runtime/Engine/Private/Landscape/LandscapeCollision.cpp
+++ b/Engine/Source/Runtime/Engine/Private/Landscape/LandscapeCollision.cpp
@@ -12,6 +12,8 @@ LandscapeCollision.cpp: Landscape collision
 #include "DerivedDataPluginInterface.h"
 #include "DerivedDataCacheInterface.h"
 #include "../PhysicsEngine/PhysDerivedData.h"
+#include "Landscape/LandscapeHeightfieldCollisionComponent.h"
+#include "Landscape/LandscapeMeshCollisionComponent.h"
 
 TMap<FGuid, ULandscapeHeightfieldCollisionComponent::FPhysXHeightfieldRef* > GSharedHeightfieldRefs;
 
diff --git a/Engine/Source/Runtime/Engine/Private/Landscape/LandscapeEdit.cpp b/Engine/Source/Runtime/Engine/Private/Landscape/LandscapeEdit.cpp
--- a/Engine/Source/Runtime/Engine/Private/Landscape/LandscapeEdit.cpp
+++ b/Engine/Source/Runtime/Engine/Private/Landscape/LandscapeEdit.cpp
@@ -9,6 +9,15 @@ LandscapeEdit.cpp: Landscape editing
 #include "Landscape/LandscapeEdit.h"
 #include "Landscape/LandscapeRender.h"
 #include "Landscape/LandscapeRenderMobile.h"
+#include "Landscape/Landscape.h"
+#include "Landscape/LandscapeMaterialInstanceConstant.h"
+#include "Landscape/LandscapeHeightfieldCollisionComponent.h"
+#include "Landscape/LandscapeMeshCollisionComponent.h"
+#include "Landscape/LandscapeSplinesComponent.h"
+#include "Landscape/LandscapeSplineControlPoint.h"
+#include "Landscape/LandscapeSplineSegment.h"
+#include "Landscape/LandscapeGizmoActiveActor.h"
+#include "Foliage/InstancedFoliageActor.h"
 #include "LevelUtils.h"
 #include "MessageLog.h"
 #include "MapErrors.h"
diff --git a/Engine/Source/Runtime/Engine/Private/Landscape/LandscapeEditInterface.cpp b/Engine/Source/Runtime/Engine/Private/Landscape/LandscapeEditInterface.cpp
--- a/Engine/Source/Runtime/Engine/Private/Landscape/LandscapeEditInterface.cpp
+++ b/Engine/Source/Runtime/Engine/Private/Landscape/LandscapeEditInterface.cpp
@@ -8,6 +8,7 @@ LandscapeEditInterface.cpp: Landscape editing interface
 #include "Landscape/LandscapeDataAccess.h"
 #include "Landscape/LandscapeEdit.h"
 #include "Landscape/LandscapeRender.h"
+#include "Landscape/Landscape.h"
 
 #if WITH_EDITOR
 
diff --git a/Engine/Source/Runtime/Engine/Private/Landscape/LandscapeGizmoActor.cpp b/Engine/Source/Runtime/Engine/Private/Landscape/LandscapeGizmoActor.cpp
--- a/Engine/Source/Runtime/Engine/Private/Landscape/LandscapeGizmoActor.cpp
+++ b/Engine/Source/Runtime/Engine/Private/Landscape/LandscapeGizmoActor.cpp
@@ -4,6 +4,9 @@
 #include "EnginePrivate.h"
 #include "Landscape/LandscapeDataAccess.h"
 #include "Landscape/LandscapeRender.h"
+#include "Landscape/LandscapeGizmoActiveActor.h"
+#include "Landscape/LandscapeGizmoRenderComponent.h"
+#include "Landscape/Landscape.h"
 
 namespace
 {
diff --git a/Engine/Source/Runtime/Engine/Private/Landscape/LandscapeRender.cpp b/Engine/Source/Runtime/Engine/Private/Landscape/LandscapeRender.cpp
--- a/Engine/Source/Runtime/Engine/Private/Landscape/LandscapeRender.cpp
+++ b/Engine/Source/Runtime/Engine/Private/Landscape/LandscapeRender.cpp
@@ -10,6 +10,7 @@ LandscapeRender.cpp: New terrain rendering
 #include "Landscape/LandscapeEdit.h"
 #include "LevelUtils.h"
 #include "MaterialCompiler.h"
+#include "Landscape/LandscapeMaterialInstanceConstant.h"
 
 IMPLEMENT_UNIFORM_BUFFER_STRUCT(FLandscapeUniformShaderParameters,TEXT("LandscapeParameters"));
 
diff --git a/Engine/Source/Runtime/Engine/Private/Landscape/LandscapeSplines.cpp b/Engine/Source/Runtime/Engine/Private/Landscape/LandscapeSplines.cpp
--- a/Engine/Source/Runtime/Engine/Private/Landscape/LandscapeSplines.cpp
+++ b/Engine/Source/Runtime/Engine/Private/Landscape/LandscapeSplines.cpp
@@ -6,6 +6,10 @@
 
 #include "EnginePrivate.h"
 #include "Landscape/LandscapeSplineProxies.h"
+#include "Landscape/LandscapeSplinesComponent.h"
+#include "Landscape/LandscapeSplineControlPoint.h"
+#include "Landscape/LandscapeSplineSegment.h"
+#include "Landscape/ControlPointMeshComponent.h"
 
 IMPLEMENT_HIT_PROXY(HLandscapeSplineProxy, HHitProxy);
 IMPLEMENT_HIT_PROXY(HLandscapeSplineProxy_Segment, HLandscapeSplineProxy);
diff --git a/Engine/Source/Runtime/Engine/Private/PhysicsEngine/PhysDerivedData.cpp b/Engine/Source/Runtime/Engine/Private/PhysicsEngine/PhysDerivedData.cpp
--- a/Engine/Source/Runtime/Engine/Private/PhysicsEngine/PhysDerivedData.cpp
+++ b/Engine/Source/Runtime/Engine/Private/PhysicsEngine/PhysDerivedData.cpp
@@ -4,6 +4,7 @@
 #include "TargetPlatform.h"
 
 #if WITH_PHYSX && WITH_EDITOR
+#include "Landscape/LandscapeMeshCollisionComponent.h"
 
 FDerivedDataPhysXCooker::FDerivedDataPhysXCooker( FName InFormat, UBodySetup* InBodySetup )
 	: BodySetup( InBodySetup )
diff --git a/Engine/Source/Runtime/Engine/Private/World.cpp b/Engine/Source/Runtime/Engine/Private/World.cpp
--- a/Engine/Source/Runtime/Engine/Private/World.cpp
+++ b/Engine/Source/Runtime/Engine/Private/World.cpp
@@ -27,6 +27,8 @@
 
 #include "MallocProfiler.h"
 
+#include "Landscape/LandscapeInfo.h"
+
 DEFINE_LOG_CATEGORY_STATIC(LogWorld, Log, All);
 DEFINE_LOG_CATEGORY(LogSpawn);
 
diff --git a/Engine/Source/Runtime/Engine/Public/AI/NavigationSystemHelpers.h b/Engine/Source/Runtime/Engine/Public/AI/NavigationSystemHelpers.h
--- a/Engine/Source/Runtime/Engine/Public/AI/NavigationSystemHelpers.h
+++ b/Engine/Source/Runtime/Engine/Public/AI/NavigationSystemHelpers.h
@@ -4,6 +4,15 @@
 
 #include "AI/Navigation/NavLinkDefinition.h"
 
+#if WITH_PHYSX
+namespace physx
+{
+	class PxTriangleMesh;
+	class PxConvexMesh;
+	class PxHeightField;
+}
+#endif // WITH_PHYSX
+
 struct FNavigableGeometryExport
 {
 	virtual ~FNavigableGeometryExport() {}
diff --git a/Engine/Source/Runtime/Engine/Public/Engine.h b/Engine/Source/Runtime/Engine/Public/Engine.h
--- a/Engine/Source/Runtime/Engine/Public/Engine.h
+++ b/Engine/Source/Runtime/Engine/Public/Engine.h
@@ -105,17 +105,10 @@
 #include "Engine/WindDirectionalSource.h"
 #include "GameFramework/MusicTrackDataStructures.h"
 #include "GameFramework/WorldSettings.h"
-#include "Foliage/InstancedFoliageActor.h"
-#include "Landscape/LandscapeGizmoActor.h"
-#include "Landscape/LandscapeGizmoActiveActor.h"
-#include "Landscape/LandscapeLayerInfoObject.h"
-#include "Landscape/LandscapeInfo.h"
 #include "PhysicsEngine/BodyInstance.h"
-#include "Landscape/LandscapeProxy.h"
 #include "Components/SceneComponent.h"
 #include "Components/LightComponentBase.h"
 #include "Components/LightComponent.h"
-#include "Landscape/Landscape.h"
 #include "Engine/LevelBounds.h"
 #include "Engine/LevelScriptActor.h"
 #include "Engine/Light.h"
@@ -225,11 +218,6 @@
 #include "Debug/DebugDrawService.h"
 #include "AI/EnvironmentQuery/EQSRenderingComponent.h"
 #include "Debug/GameplayDebuggingComponent.h"
-#include "Landscape/LandscapeComponent.h"
-#include "Landscape/LandscapeGizmoRenderComponent.h"
-#include "Landscape/LandscapeHeightfieldCollisionComponent.h"
-#include "Landscape/LandscapeMeshCollisionComponent.h"
-#include "Landscape/LandscapeSplinesComponent.h"
 #include "Components/LineBatchComponent.h"
 #include "Components/MaterialBillboardComponent.h"
 #include "Components/DestructibleComponent.h"
@@ -237,7 +225,6 @@
 #include "Components/PoseableMeshComponent.h"
 #include "Lightmass/LightmassPrimitiveSettingsObject.h"
 #include "Components/StaticMeshComponent.h"
-#include "Landscape/ControlPointMeshComponent.h"
 #include "Components/InstancedStaticMeshComponent.h"
 #include "Components/InteractiveFoliageComponent.h"
 #include "Interfaces/Interface_CollisionDataProvider.h"
@@ -469,8 +456,6 @@
 #include "Engine/InterpCurveEdSetup.h"
 #include "Engine/IntSerialization.h"
 #include "AI/KismetAIAsyncTaskProxy.h"
-#include "Landscape/LandscapeSplineSegment.h"
-#include "Landscape/LandscapeSplineControlPoint.h"
 #include "Layers/Layer.h"
 #include "Engine/LevelBase.h"
 #include "Engine/Level.h"
@@ -615,7 +600,6 @@
 #include "Materials/Material.h"
 #include "Materials/MaterialInstance.h"
 #include "Materials/MaterialInstanceConstant.h"
-#include "Landscape/LandscapeMaterialInstanceConstant.h"
 #include "Materials/MaterialInstanceDynamic.h"
 #include "Materials/MaterialParameterCollection.h"
 #include "Materials/MaterialParameterCollectionInstance.h"
diff --git a/Engine/Source/Runtime/Engine/Public/Landscape/LandscapeDataAccess.h b/Engine/Source/Runtime/Engine/Public/Landscape/LandscapeDataAccess.h
--- a/Engine/Source/Runtime/Engine/Public/Landscape/LandscapeDataAccess.h
+++ b/Engine/Source/Runtime/Engine/Public/Landscape/LandscapeDataAccess.h
@@ -17,6 +17,11 @@ LandscapeDataAccess.h: Classes for the editor to access to Landscape data
 
 #if WITH_EDITOR
 
+// Forward declarations
+class ULandscapeLayerInfoObject;
+#include "Landscape/LandscapeInfo.h"
+#include "Landscape/LandscapeComponent.h"
+
 namespace LandscapeDataAccess
 {
 	const int32 MaxValue = 65535;
diff --git a/Engine/Source/Runtime/Engine/Public/Landscape/LandscapeEdit.h b/Engine/Source/Runtime/Engine/Public/Landscape/LandscapeEdit.h
--- a/Engine/Source/Runtime/Engine/Public/Landscape/LandscapeEdit.h
+++ b/Engine/Source/Runtime/Engine/Public/Landscape/LandscapeEdit.h
@@ -11,6 +11,8 @@ LandscapeEdit.h: Classes for the editor to access to Landscape data
 
 #if WITH_EDITOR
 
+#include "Landscape/LandscapeProxy.h"
+
 struct FLandscapeTextureDataInfo
 {
 	struct FMipInfo
diff --git a/Engine/Source/Runtime/Engine/Public/Landscape/LandscapeLight.h b/Engine/Source/Runtime/Engine/Public/Landscape/LandscapeLight.h
--- a/Engine/Source/Runtime/Engine/Public/Landscape/LandscapeLight.h
+++ b/Engine/Source/Runtime/Engine/Public/Landscape/LandscapeLight.h
@@ -7,6 +7,8 @@ LandscapeLight.h: Static lighting for LandscapeComponents
 #ifndef _LANDSCAPELIGHT_H
 #define _LANDSCAPELIGHT_H
 
+class ULandscapeComponent;
+
 /** A texture mapping for landscapes */
 class FLandscapeStaticLightingTextureMapping : public FStaticLightingTextureMapping
 {
diff --git a/Engine/Source/Runtime/Engine/Public/Landscape/LandscapeRender.h b/Engine/Source/Runtime/Engine/Public/Landscape/LandscapeRender.h
--- a/Engine/Source/Runtime/Engine/Public/Landscape/LandscapeRender.h
+++ b/Engine/Source/Runtime/Engine/Public/Landscape/LandscapeRender.h
@@ -8,6 +8,8 @@ LandscapeRender.h: New terrain rendering
 #define _LANDSCAPERENDER_H
 
 #include "UniformBuffer.h"
+#include "Landscape/LandscapeComponent.h"
+#include "Landscape/LandscapeProxy.h"
 
 class FLandscapeVertexFactory;
 class FLandscapeVertexBuffer;
diff --git a/Engine/Source/Runtime/Engine/Public/Landscape/LandscapeSplineProxies.h b/Engine/Source/Runtime/Engine/Public/Landscape/LandscapeSplineProxies.h
--- a/Engine/Source/Runtime/Engine/Public/Landscape/LandscapeSplineProxies.h
+++ b/Engine/Source/Runtime/Engine/Public/Landscape/LandscapeSplineProxies.h
@@ -2,6 +2,8 @@
 
 #pragma once
 
+#include "Landscape/LandscapeSplineSegment.h"
+
 //////////////////////////////////////////////////////////////////////////
 // LANDSCAPE SPLINES HIT PROXY
 
@@ -57,7 +59,7 @@ struct HLandscapeSplineProxy_Tangent : public HLandscapeSplineProxy
 {
 	DECLARE_HIT_PROXY( ENGINE_API );
 
-	class ULandscapeSplineSegment* SplineSegment;
+	ULandscapeSplineSegment* SplineSegment;
 	uint32 End:1;
 
 	HLandscapeSplineProxy_Tangent(class ULandscapeSplineSegment* InSplineSegment, bool InEnd) :
