UMG - Adding a null for the game viewport client check to prevent crashes when the exit command is sent multiple times.

diff --git a/Engine/Source/Runtime/Engine/Private/LocalPlayer.cpp b/Engine/Source/Runtime/Engine/Private/LocalPlayer.cpp
--- a/Engine/Source/Runtime/Engine/Private/LocalPlayer.cpp
+++ b/Engine/Source/Runtime/Engine/Private/LocalPlayer.cpp
@@ -788,7 +788,7 @@ bool ULocalPlayer::HandleDNCommand( const TCHAR* Cmd, FOutputDevice& Ar )
 bool ULocalPlayer::HandleExitCommand( const TCHAR* Cmd, FOutputDevice& Ar )
 {
 	// If there is no viewport it was already closed.
-	if( ViewportClient->Viewport )
+	if ( ViewportClient && ViewportClient->Viewport )
 	{
 		ViewportClient->CloseRequested(ViewportClient->Viewport);
 	}
