Fix mac build

diff --git a/Engine/Source/Runtime/Engine/Private/PhysicsEngine/PhysDerivedData.h b/Engine/Source/Runtime/Engine/Private/PhysicsEngine/PhysDerivedData.h
--- a/Engine/Source/Runtime/Engine/Private/PhysicsEngine/PhysDerivedData.h
+++ b/Engine/Source/Runtime/Engine/Private/PhysicsEngine/PhysDerivedData.h
@@ -7,10 +7,10 @@
 #include "DerivedDataPluginInterface.h"
 #include "DerivedDataCacheInterface.h"
 #include "TargetPlatform.h"
-
-
 #include "PhysXSupport.h"
 
+// Forward declarations
+class ULandscapeMeshCollisionComponent
 
 //////////////////////////////////////////////////////////////////////////
 // PhysX Cooker
