Fix mac build

diff --git a/Engine/Source/Runtime/Engine/Private/PhysicsEngine/PhysDerivedData.h b/Engine/Source/Runtime/Engine/Private/PhysicsEngine/PhysDerivedData.h
--- a/Engine/Source/Runtime/Engine/Private/PhysicsEngine/PhysDerivedData.h
+++ b/Engine/Source/Runtime/Engine/Private/PhysicsEngine/PhysDerivedData.h
@@ -10,7 +10,7 @@
 #include "PhysXSupport.h"
 
 // Forward declarations
-class ULandscapeMeshCollisionComponent
+class ULandscapeMeshCollisionComponent;
 
 //////////////////////////////////////////////////////////////////////////
 // PhysX Cooker
