Remove more foliage headers from Engine.h and include only where needed

diff --git a/Engine/Source/Developer/AssetTools/Private/AssetTypeActions/AssetTypeActions_InstancedFoliageSettings.cpp b/Engine/Source/Developer/AssetTools/Private/AssetTypeActions/AssetTypeActions_InstancedFoliageSettings.cpp
--- a/Engine/Source/Developer/AssetTools/Private/AssetTypeActions/AssetTypeActions_InstancedFoliageSettings.cpp
+++ b/Engine/Source/Developer/AssetTools/Private/AssetTypeActions/AssetTypeActions_InstancedFoliageSettings.cpp
@@ -1,6 +1,7 @@
 // Copyright 1998-2014 Epic Games, Inc. All Rights Reserved.
 
 #include "AssetToolsPrivatePCH.h"
+#include "Foliage/InstancedFoliageSettings.h"
 
 UClass* FAssetTypeActions_InstancedFoliageSettings::GetSupportedClass() const
 {
diff --git a/Engine/Source/Editor/FoliageEdit/Private/FoliageEdMode.cpp b/Engine/Source/Editor/FoliageEdit/Private/FoliageEdMode.cpp
--- a/Engine/Source/Editor/FoliageEdit/Private/FoliageEdMode.cpp
+++ b/Engine/Source/Editor/FoliageEdit/Private/FoliageEdMode.cpp
@@ -22,6 +22,7 @@
 
 // Classes
 #include "Foliage/InstancedFoliageActor.h"
+#include "Foliage/InstancedFoliageSettings.h"
 #include "Landscape/LandscapeComponent.h"
 #include "Landscape/LandscapeHeightfieldCollisionComponent.h"
 #include "Landscape/LandscapeInfo.h"
diff --git a/Engine/Source/Editor/UnrealEd/Private/Editor.cpp b/Engine/Source/Editor/UnrealEd/Private/Editor.cpp
--- a/Engine/Source/Editor/UnrealEd/Private/Editor.cpp
+++ b/Engine/Source/Editor/UnrealEd/Private/Editor.cpp
@@ -3,6 +3,7 @@
 
 #include "UnrealEd.h"
 #include "Matinee/MatineeActor.h"
+#include "Engine/InteractiveFoliageActor.h"
 #include "EditorSupportDelegates.h"
 #include "Factories.h"
 #include "BSPOps.h"
diff --git a/Engine/Source/Editor/UnrealEd/Private/Factories/ActorFactory.cpp b/Engine/Source/Editor/UnrealEd/Private/Factories/ActorFactory.cpp
--- a/Engine/Source/Editor/UnrealEd/Private/Factories/ActorFactory.cpp
+++ b/Engine/Source/Editor/UnrealEd/Private/Factories/ActorFactory.cpp
@@ -15,6 +15,7 @@ ActorFactory.cpp:
 #include "BSPOps.h"
 #include "Matinee/MatineeActor.h"
 #include "Matinee/InterpData.h"
+#include "Engine/InteractiveFoliageActor.h"
 
 #include "AssetRegistryModule.h"
 
diff --git a/Engine/Source/Runtime/Engine/Private/FoliageComponent.cpp b/Engine/Source/Runtime/Engine/Private/FoliageComponent.cpp
--- a/Engine/Source/Runtime/Engine/Private/FoliageComponent.cpp
+++ b/Engine/Source/Runtime/Engine/Private/FoliageComponent.cpp
@@ -5,6 +5,8 @@
 =============================================================================*/
 
 #include "EnginePrivate.h"
+#include "Engine/InteractiveFoliageActor.h"
+#include "Components/InteractiveFoliageComponent.h"
 
 /** Scene proxy class for UInteractiveFoliageComponent. */
 class FInteractiveFoliageSceneProxy : public FStaticMeshSceneProxy
diff --git a/Engine/Source/Runtime/Engine/Private/InstancedFoliage.cpp b/Engine/Source/Runtime/Engine/Private/InstancedFoliage.cpp
--- a/Engine/Source/Runtime/Engine/Private/InstancedFoliage.cpp
+++ b/Engine/Source/Runtime/Engine/Private/InstancedFoliage.cpp
@@ -9,6 +9,7 @@
 #include "UObjectToken.h"
 #include "MapErrors.h"
 #include "Foliage/InstancedFoliageActor.h"
+#include "Foliage/InstancedFoliageSettings.h"
 #include "Landscape/LandscapeHeightfieldCollisionComponent.h"
 
 #define LOCTEXT_NAMESPACE "InstancedFoliage"
diff --git a/Engine/Source/Runtime/Engine/Private/InteractiveFoliageActor.cpp b/Engine/Source/Runtime/Engine/Private/InteractiveFoliageActor.cpp
--- a/Engine/Source/Runtime/Engine/Private/InteractiveFoliageActor.cpp
+++ b/Engine/Source/Runtime/Engine/Private/InteractiveFoliageActor.cpp
@@ -2,6 +2,9 @@
 
 #include "EnginePrivate.h"
 
+#include "Engine/InteractiveFoliageActor.h"
+#include "Components/InteractiveFoliageComponent.h"
+
 AInteractiveFoliageActor::AInteractiveFoliageActor(const class FPostConstructInitializeProperties& PCIP)
 	: Super(PCIP
 		.SetDefaultSubobjectClass<UInteractiveFoliageComponent>("StaticMeshComponent0"))
diff --git a/Engine/Source/Runtime/Engine/Public/Engine.h b/Engine/Source/Runtime/Engine/Public/Engine.h
--- a/Engine/Source/Runtime/Engine/Public/Engine.h
+++ b/Engine/Source/Runtime/Engine/Public/Engine.h
@@ -167,7 +167,6 @@
 #include "Components/SkeletalMeshComponent.h"
 #include "Animation/SkeletalMeshActor.h"
 #include "Engine/StaticMeshActor.h"
-#include "Engine/InteractiveFoliageActor.h"
 #include "Engine/TargetPoint.h"
 #include "Engine/TextRenderActor.h"
 #include "Engine/TriggerBase.h"
@@ -226,7 +225,6 @@
 #include "Lightmass/LightmassPrimitiveSettingsObject.h"
 #include "Components/StaticMeshComponent.h"
 #include "Components/InstancedStaticMeshComponent.h"
-#include "Components/InteractiveFoliageComponent.h"
 #include "Interfaces/Interface_CollisionDataProvider.h"
 #include "Components/SplineMeshComponent.h"
 #include "Components/ModelComponent.h"
@@ -452,7 +450,6 @@
 #include "GameFramework/GameUserSettings.h"
 #include "GameFramework/PlayerInput.h"
 #include "GameFramework/InputSettings.h"
-#include "Foliage/InstancedFoliageSettings.h"
 #include "Engine/InterpCurveEdSetup.h"
 #include "Engine/IntSerialization.h"
 #include "AI/KismetAIAsyncTaskProxy.h"
