Remove Material Expression headers from Engine.h and only include where needed

diff --git a/Engine/Plugins/Editor/SpeedTreeImporter/Source/SpeedTreeImporter/Private/SpeedTreeImportFactory.cpp b/Engine/Plugins/Editor/SpeedTreeImporter/Source/SpeedTreeImporter/Private/SpeedTreeImportFactory.cpp
--- a/Engine/Plugins/Editor/SpeedTreeImporter/Source/SpeedTreeImporter/Private/SpeedTreeImportFactory.cpp
+++ b/Engine/Plugins/Editor/SpeedTreeImporter/Source/SpeedTreeImporter/Private/SpeedTreeImportFactory.cpp
@@ -5,6 +5,19 @@
 
 #include "UnrealEd.h"
 
+#include "Materials/MaterialExpressionClamp.h"
+#include "Materials/MaterialExpressionComponentMask.h"
+#include "Materials/MaterialExpressionConstant.h"
+#include "Materials/MaterialExpressionLinearInterpolate.h"
+#include "Materials/MaterialExpressionPower.h"
+#include "Materials/MaterialExpressionMultiply.h"
+#include "Materials/MaterialExpressionSpeedTree.h"
+#include "Materials/MaterialExpressionSpeedTreeColorVariation.h"
+#include "Materials/MaterialExpressionTextureCoordinate.h"
+#include "Materials/MaterialExpressionTextureSample.h"
+#include "Materials/MaterialExpressionTwoSidedSign.h"
+#include "Materials/MaterialExpressionVertexColor.h"
+
 #include "Mainframe.h"
 #include "ModuleManager.h"
 #include "ObjectTools.h"
diff --git a/Engine/Source/Developer/AssetTools/Private/AssetTypeActions/AssetTypeActions_MaterialFunction.h b/Engine/Source/Developer/AssetTools/Private/AssetTypeActions/AssetTypeActions_MaterialFunction.h
--- a/Engine/Source/Developer/AssetTools/Private/AssetTypeActions/AssetTypeActions_MaterialFunction.h
+++ b/Engine/Source/Developer/AssetTools/Private/AssetTypeActions/AssetTypeActions_MaterialFunction.h
@@ -2,6 +2,8 @@
 
 #pragma once
 
+#include "Materials/MaterialFunction.h"
+
 class FAssetTypeActions_MaterialFunction : public FAssetTypeActions_Base
 {
 public:
diff --git a/Engine/Source/Editor/DetailCustomizations/Private/MathStructCustomizations.cpp b/Engine/Source/Editor/DetailCustomizations/Private/MathStructCustomizations.cpp
--- a/Engine/Source/Editor/DetailCustomizations/Private/MathStructCustomizations.cpp
+++ b/Engine/Source/Editor/DetailCustomizations/Private/MathStructCustomizations.cpp
@@ -1,6 +1,7 @@
 // Copyright 1998-2014 Epic Games, Inc. All Rights Reserved.
 
 #include "DetailCustomizationsPrivatePCH.h"
+#include "Materials/MaterialExpressionConstant3Vector.h"
 #include "MathStructCustomizations.h"
 
 #define LOCTEXT_NAMESPACE "FMathStructCustomization"
diff --git a/Engine/Source/Editor/MaterialEditor/Private/FindInMaterial.cpp b/Engine/Source/Editor/MaterialEditor/Private/FindInMaterial.cpp
--- a/Engine/Source/Editor/MaterialEditor/Private/FindInMaterial.cpp
+++ b/Engine/Source/Editor/MaterialEditor/Private/FindInMaterial.cpp
@@ -1,5 +1,6 @@
 // Copyright 1998-2014 Epic Games, Inc. All Rights Reserved.
 #include "MaterialEditorModule.h"
+#include "Materials/MaterialExpressionComment.h"
 #include "FindInMaterial.h"
 #include "MaterialEditor.h"
 
diff --git a/Engine/Source/Editor/MaterialEditor/Private/MaterialEditor.cpp b/Engine/Source/Editor/MaterialEditor/Private/MaterialEditor.cpp
--- a/Engine/Source/Editor/MaterialEditor/Private/MaterialEditor.cpp
+++ b/Engine/Source/Editor/MaterialEditor/Private/MaterialEditor.cpp
@@ -1,6 +1,32 @@
 // Copyright 1998-2014 Epic Games, Inc. All Rights Reserved.
 
 #include "MaterialEditorModule.h"
+
+#include "Materials/MaterialExpressionBreakMaterialAttributes.h"
+#include "Materials/MaterialExpressionCollectionParameter.h"
+#include "Materials/MaterialExpressionComment.h"
+#include "Materials/MaterialExpressionComponentMask.h"
+#include "Materials/MaterialExpressionConstant.h"
+#include "Materials/MaterialExpressionConstant2Vector.h"
+#include "Materials/MaterialExpressionConstant3Vector.h"
+#include "Materials/MaterialExpressionConstant4Vector.h"
+#include "Materials/MaterialExpressionFontSampleParameter.h"
+#include "Materials/MaterialExpressionFunctionInput.h"
+#include "Materials/MaterialExpressionFunctionOutput.h"
+#include "Materials/MaterialExpressionParameter.h"
+#include "Materials/MaterialExpressionParticleSubUV.h"
+#include "Materials/MaterialExpressionRotateAboutAxis.h"
+#include "Materials/MaterialExpressionScalarParameter.h"
+#include "Materials/MaterialExpressionStaticComponentMaskParameter.h"
+#include "Materials/MaterialExpressionTextureSampleParameter.h"
+#include "Materials/MaterialExpressionTextureObject.h"
+#include "Materials/MaterialExpressionTextureSampleParameter2D.h"
+#include "Materials/MaterialExpressionTextureSampleParameterCube.h"
+#include "Materials/MaterialExpressionTextureSampleParameterSubUV.h"
+#include "Materials/MaterialExpressionTransformPosition.h"
+#include "Materials/MaterialExpressionVectorParameter.h"
+#include "Materials/MaterialFunction.h"
+
 #include "MaterialEditorActions.h"
 #include "MaterialExpressionClasses.h"
 #include "MaterialCompiler.h"
diff --git a/Engine/Source/Editor/MaterialEditor/Private/MaterialEditor.h b/Engine/Source/Editor/MaterialEditor/Private/MaterialEditor.h
--- a/Engine/Source/Editor/MaterialEditor/Private/MaterialEditor.h
+++ b/Engine/Source/Editor/MaterialEditor/Private/MaterialEditor.h
@@ -1,7 +1,7 @@
 // Copyright 1998-2014 Epic Games, Inc. All Rights Reserved.
 
 #pragma once
-
+#include "Materials/MaterialExpression.h"
 #include "Toolkits/AssetEditorToolkit.h"
 #include "DragAndDrop.h"
 #include "Editor/GraphEditor/Public/GraphEditorDragDropAction.h"
diff --git a/Engine/Source/Editor/MaterialEditor/Private/MaterialEditorInstanceDetailCustomization.cpp b/Engine/Source/Editor/MaterialEditor/Private/MaterialEditorInstanceDetailCustomization.cpp
--- a/Engine/Source/Editor/MaterialEditor/Private/MaterialEditorInstanceDetailCustomization.cpp
+++ b/Engine/Source/Editor/MaterialEditor/Private/MaterialEditorInstanceDetailCustomization.cpp
@@ -1,6 +1,9 @@
 // Copyright 1998-2014 Epic Games, Inc. All Rights Reserved.
 
 #include "MaterialEditorModule.h"
+#include "Materials/MaterialExpressionParameter.h"
+#include "Materials/MaterialExpressionTextureSampleParameter.h"
+#include "Materials/MaterialExpressionFontSampleParameter.h"
 #include "EditorSupportDelegates.h"
 #include "UnrealEd.h"
 #include "PropertyEditing.h"
diff --git a/Engine/Source/Editor/MaterialEditor/Private/MaterialEditorUtilities.cpp b/Engine/Source/Editor/MaterialEditor/Private/MaterialEditorUtilities.cpp
--- a/Engine/Source/Editor/MaterialEditor/Private/MaterialEditorUtilities.cpp
+++ b/Engine/Source/Editor/MaterialEditor/Private/MaterialEditorUtilities.cpp
@@ -2,6 +2,21 @@
 
 
 #include "MaterialEditorModule.h"
+
+#include "Materials/MaterialExpressionFunctionInput.h"
+#include "Materials/MaterialExpressionFunctionOutput.h"
+#include "Materials/MaterialExpressionStaticBoolParameter.h"
+#include "Materials/MaterialExpressionStaticBool.h"
+#include "Materials/MaterialExpressionStaticSwitch.h"
+#include "Materials/MaterialExpressionComment.h"
+#include "Materials/MaterialExpressionParameter.h"
+#include "Materials/MaterialExpressionTextureSampleParameter.h"
+#include "Materials/MaterialExpressionFontSampleParameter.h"
+#include "Materials/MaterialExpressionScalarParameter.h"
+#include "Materials/MaterialExpressionVectorParameter.h"
+#include "Materials/MaterialExpressionStaticSwitchParameter.h"
+#include "Materials/MaterialFunction.h"
+
 #include "MaterialEditorUtilities.h"
 #include "Toolkits/ToolkitManager.h"
 #include "MaterialExpressionClasses.h"
diff --git a/Engine/Source/Editor/MaterialEditor/Private/MaterialExpressionClasses.cpp b/Engine/Source/Editor/MaterialEditor/Private/MaterialExpressionClasses.cpp
--- a/Engine/Source/Editor/MaterialEditor/Private/MaterialExpressionClasses.cpp
+++ b/Engine/Source/Editor/MaterialEditor/Private/MaterialExpressionClasses.cpp
@@ -1,6 +1,10 @@
 // Copyright 1998-2014 Epic Games, Inc. All Rights Reserved.
 
 #include "MaterialEditorModule.h"
+
+#include "Materials/MaterialExpressionComment.h"
+#include "Materials/MaterialExpressionParameter.h"
+
 #include "MaterialExpressionClasses.h"
 
 ///////////////////////////////////////////////////////////////////////////////////////////////////////////////////////
diff --git a/Engine/Source/Editor/MaterialEditor/Private/MaterialInstanceEditor.cpp b/Engine/Source/Editor/MaterialEditor/Private/MaterialInstanceEditor.cpp
--- a/Engine/Source/Editor/MaterialEditor/Private/MaterialInstanceEditor.cpp
+++ b/Engine/Source/Editor/MaterialEditor/Private/MaterialInstanceEditor.cpp
@@ -1,6 +1,10 @@
 // Copyright 1998-2014 Epic Games, Inc. All Rights Reserved.
 
 #include "MaterialEditorModule.h"
+
+#include "Materials/MaterialExpressionTextureBase.h"
+#include "Materials/MaterialExpressionTextureSampleParameter.h"
+
 #include "MaterialInstanceEditor.h"
 #include "MaterialEditor.h"
 #include "MaterialEditorActions.h"
diff --git a/Engine/Source/Editor/MaterialEditor/Private/SMaterialPalette.cpp b/Engine/Source/Editor/MaterialEditor/Private/SMaterialPalette.cpp
--- a/Engine/Source/Editor/MaterialEditor/Private/SMaterialPalette.cpp
+++ b/Engine/Source/Editor/MaterialEditor/Private/SMaterialPalette.cpp
@@ -1,6 +1,10 @@
 // Copyright 1998-2014 Epic Games, Inc. All Rights Reserved.
 
 #include "MaterialEditorModule.h"
+
+#include "Materials/MaterialExpressionComment.h"
+#include "Materials/MaterialFunction.h"
+
 #include "EditorWidgets.h"
 #include "AssetRegistryModule.h"
 #include "MaterialEditor.h"
diff --git a/Engine/Source/Editor/MaterialEditor/Public/MaterialEditorUtilities.h b/Engine/Source/Editor/MaterialEditor/Public/MaterialEditorUtilities.h
--- a/Engine/Source/Editor/MaterialEditor/Public/MaterialEditorUtilities.h
+++ b/Engine/Source/Editor/MaterialEditor/Public/MaterialEditorUtilities.h
@@ -2,6 +2,10 @@
 
 #pragma once
 
+#include "Materials/MaterialExpressionMaterialFunctionCall.h"
+
+struct FExpressionInput;
+
 //////////////////////////////////////////////////////////////////////////
 // FMaterialEditorUtilities
 
diff --git a/Engine/Source/Editor/MeshPaint/Private/MeshPaintEdMode.cpp b/Engine/Source/Editor/MeshPaint/Private/MeshPaintEdMode.cpp
--- a/Engine/Source/Editor/MeshPaint/Private/MeshPaintEdMode.cpp
+++ b/Engine/Source/Editor/MeshPaint/Private/MeshPaintEdMode.cpp
@@ -5,6 +5,12 @@
 ================================================================================*/
 
 #include "MeshPaintPrivatePCH.h"
+
+#include "Materials/MaterialExpressionTextureBase.h"
+#include "Materials/MaterialExpressionTextureSample.h"
+#include "Materials/MaterialExpressionTextureCoordinate.h"
+#include "Materials/MaterialExpressionTextureSampleParameter.h"
+
 #include "MeshPaintEdMode.h"
 #include "Factories.h"
 #include "ScopedTransaction.h"
diff --git a/Engine/Source/Editor/PropertyEditor/Private/UserInterface/PropertyEditor/SPropertyEditorColor.cpp b/Engine/Source/Editor/PropertyEditor/Private/UserInterface/PropertyEditor/SPropertyEditorColor.cpp
--- a/Engine/Source/Editor/PropertyEditor/Private/UserInterface/PropertyEditor/SPropertyEditorColor.cpp
+++ b/Engine/Source/Editor/PropertyEditor/Private/UserInterface/PropertyEditor/SPropertyEditorColor.cpp
@@ -1,6 +1,7 @@
 // Copyright 1998-2014 Epic Games, Inc. All Rights Reserved.
 
 #include "PropertyEditorPrivatePCH.h"
+#include "Materials/MaterialExpressionConstant3Vector.h"
 #include "PropertyEditorHelpers.h"
 #include "IPropertyUtilities.h"
 #include "PropertyNode.h"
diff --git a/Engine/Source/Editor/UnrealEd/Classes/MaterialGraph/MaterialGraphSchema.h b/Engine/Source/Editor/UnrealEd/Classes/MaterialGraph/MaterialGraphSchema.h
--- a/Engine/Source/Editor/UnrealEd/Classes/MaterialGraph/MaterialGraphSchema.h
+++ b/Engine/Source/Editor/UnrealEd/Classes/MaterialGraph/MaterialGraphSchema.h
@@ -38,7 +38,7 @@ struct UNREALED_API FMaterialGraphSchemaAction_NewNode : public FEdGraphSchemaAc
 	 * @param	FunctionInput		The function input to set.
 	 * @param	MaterialValueType	Value type we want input to accept.
 	 */
-	void SetFunctionInputType(UMaterialExpressionFunctionInput* FunctionInput, uint32 MaterialValueType) const;
+	void SetFunctionInputType(class UMaterialExpressionFunctionInput* FunctionInput, uint32 MaterialValueType) const;
 };
 
 /** Action to add a Material Function call to the graph */
diff --git a/Engine/Source/Editor/UnrealEd/Private/EditorExporters.cpp b/Engine/Source/Editor/UnrealEd/Private/EditorExporters.cpp
--- a/Engine/Source/Editor/UnrealEd/Private/EditorExporters.cpp
+++ b/Engine/Source/Editor/UnrealEd/Private/EditorExporters.cpp
@@ -5,6 +5,7 @@
 =============================================================================*/
 
 #include "UnrealEd.h"
+#include "Materials/MaterialExpressionConstant.h"
 #include "SoundDefinitions.h"
 #include "Landscape/LandscapeDataAccess.h"
 #include "Kismet2/DebuggerCommands.h"
diff --git a/Engine/Source/Editor/UnrealEd/Private/Factories/EditorFactories.cpp b/Engine/Source/Editor/UnrealEd/Private/Factories/EditorFactories.cpp
--- a/Engine/Source/Editor/UnrealEd/Private/Factories/EditorFactories.cpp
+++ b/Engine/Source/Editor/UnrealEd/Private/Factories/EditorFactories.cpp
@@ -7,6 +7,8 @@
 #include "UnrealEd.h"
 #include "Matinee/InterpData.h"
 #include "Matinee/InterpGroupCamera.h"
+#include "Materials/MaterialExpressionTextureSample.h"
+#include "Materials/MaterialFunction.h"
 #include "Factories.h"
 #include "SoundDefinitions.h"
 #include "BlueprintUtilities.h"
diff --git a/Engine/Source/Editor/UnrealEd/Private/Fbx/FbxMainExport.cpp b/Engine/Source/Editor/UnrealEd/Private/Fbx/FbxMainExport.cpp
--- a/Engine/Source/Editor/UnrealEd/Private/Fbx/FbxMainExport.cpp
+++ b/Engine/Source/Editor/UnrealEd/Private/Fbx/FbxMainExport.cpp
@@ -33,12 +33,20 @@
 =============================================================================*/
 
 #include "UnrealEd.h"
+
+#include "Materials/MaterialExpressionConstant.h"
+#include "Materials/MaterialExpressionVectorParameter.h"
+#include "Materials/MaterialExpressionConstant2Vector.h"
+#include "Materials/MaterialExpressionConstant3Vector.h"
+#include "Materials/MaterialExpressionConstant4Vector.h"
+
 #include "Matinee/InterpData.h"
 #include "Matinee/InterpTrackMove.h"
 #include "Matinee/InterpTrackMoveAxis.h"
 #include "Matinee/InterpTrackFloatProp.h"
 #include "Matinee/InterpTrackInstFloatProp.h"
 #include "Matinee/InterpTrackInstMove.h"
+
 #include "LandscapeDataAccess.h"
 
 #include "FbxExporter.h"
diff --git a/Engine/Source/Editor/UnrealEd/Private/Fbx/FbxMaterialImport.cpp b/Engine/Source/Editor/UnrealEd/Private/Fbx/FbxMaterialImport.cpp
--- a/Engine/Source/Editor/UnrealEd/Private/Fbx/FbxMaterialImport.cpp
+++ b/Engine/Source/Editor/UnrealEd/Private/Fbx/FbxMaterialImport.cpp
@@ -29,6 +29,9 @@
 */
 
 #include "UnrealEd.h"
+#include "Materials/MaterialExpressionTextureSample.h"
+#include "Materials/MaterialExpressionTextureCoordinate.h"
+#include "Materials/MaterialExpressionVectorParameter.h"
 #include "Factories.h"
 #include "Engine.h"
 #include "TextureLayout.h"
diff --git a/Engine/Source/Editor/UnrealEd/Private/LevelEditorViewport.cpp b/Engine/Source/Editor/UnrealEd/Private/LevelEditorViewport.cpp
--- a/Engine/Source/Editor/UnrealEd/Private/LevelEditorViewport.cpp
+++ b/Engine/Source/Editor/UnrealEd/Private/LevelEditorViewport.cpp
@@ -2,6 +2,7 @@
 
 
 #include "UnrealEd.h"
+#include "Materials/MaterialExpressionTextureSample.h"
 #include "EditorSupportDelegates.h"
 #include "SoundDefinitions.h"
 #include "CameraController.h"
diff --git a/Engine/Source/Editor/UnrealEd/Private/Lightmass/LightmassRender.cpp b/Engine/Source/Editor/UnrealEd/Private/Lightmass/LightmassRender.cpp
--- a/Engine/Source/Editor/UnrealEd/Private/Lightmass/LightmassRender.cpp
+++ b/Engine/Source/Editor/UnrealEd/Private/Lightmass/LightmassRender.cpp
@@ -5,7 +5,7 @@
 =============================================================================*/
 
 #include "UnrealEd.h"
-
+#include "Materials/MaterialExpressionWorldPosition.h"
 #include "Programs/UnrealLightmass/Public/LightmassPublic.h"
 #include "Lightmass.h"
 #include "LightmassRender.h"
diff --git a/Engine/Source/Editor/UnrealEd/Private/MaterialGraph.cpp b/Engine/Source/Editor/UnrealEd/Private/MaterialGraph.cpp
--- a/Engine/Source/Editor/UnrealEd/Private/MaterialGraph.cpp
+++ b/Engine/Source/Editor/UnrealEd/Private/MaterialGraph.cpp
@@ -4,6 +4,10 @@
 // UMaterialGraph
 
 #include "UnrealEd.h"
+
+#include "Materials/MaterialExpressionComment.h"
+#include "Materials/MaterialExpressionFunctionOutput.h"
+
 #include "GraphEditor.h"
 #include "MaterialShared.h"
 #include "MaterialEditorUtilities.h"
diff --git a/Engine/Source/Editor/UnrealEd/Private/MaterialGraphNode.cpp b/Engine/Source/Editor/UnrealEd/Private/MaterialGraphNode.cpp
--- a/Engine/Source/Editor/UnrealEd/Private/MaterialGraphNode.cpp
+++ b/Engine/Source/Editor/UnrealEd/Private/MaterialGraphNode.cpp
@@ -5,6 +5,30 @@
 =============================================================================*/
 
 #include "UnrealEd.h"
+
+#include "Materials/MaterialExpressionComponentMask.h"
+#include "Materials/MaterialExpressionConstant.h"
+#include "Materials/MaterialExpressionConstant2Vector.h"
+#include "Materials/MaterialExpressionConstant3Vector.h"
+#include "Materials/MaterialExpressionConstant4Vector.h"
+#include "Materials/MaterialExpressionCustomTexture.h"
+#include "Materials/MaterialExpressionFontSample.h"
+#include "Materials/MaterialExpressionFontSampleParameter.h"
+#include "Materials/MaterialExpressionFunctionInput.h"
+#include "Materials/MaterialExpressionFunctionOutput.h"
+#include "Materials/MaterialExpressionMaterialFunctionCall.h"
+#include "Materials/MaterialExpressionParameter.h"
+#include "Materials/MaterialExpressionScalarParameter.h"
+#include "Materials/MaterialExpressionStaticBool.h"
+#include "Materials/MaterialExpressionStaticBoolParameter.h"
+#include "Materials/MaterialExpressionTextureBase.h"
+#include "Materials/MaterialExpressionTextureCoordinate.h"
+#include "Materials/MaterialExpressionTextureSample.h"
+#include "Materials/MaterialExpressionTextureSampleParameter.h"
+#include "Materials/MaterialExpressionTextureObject.h"
+#include "Materials/MaterialExpressionVectorParameter.h"
+#include "Materials/MaterialFunction.h"
+
 #include "MaterialEditorUtilities.h"
 #include "MaterialEditorActions.h"
 #include "GraphEditorActions.h"
diff --git a/Engine/Source/Editor/UnrealEd/Private/MaterialGraphNode_Comment.cpp b/Engine/Source/Editor/UnrealEd/Private/MaterialGraphNode_Comment.cpp
--- a/Engine/Source/Editor/UnrealEd/Private/MaterialGraphNode_Comment.cpp
+++ b/Engine/Source/Editor/UnrealEd/Private/MaterialGraphNode_Comment.cpp
@@ -5,6 +5,8 @@
 =============================================================================*/
 
 #include "UnrealEd.h"
+#include "Materials/MaterialExpressionComment.h"
+#include "Materials/MaterialFunction.h"
 #include "MaterialEditorUtilities.h"
 
 /////////////////////////////////////////////////////
diff --git a/Engine/Source/Editor/UnrealEd/Private/MaterialGraphSchema.cpp b/Engine/Source/Editor/UnrealEd/Private/MaterialGraphSchema.cpp
--- a/Engine/Source/Editor/UnrealEd/Private/MaterialGraphSchema.cpp
+++ b/Engine/Source/Editor/UnrealEd/Private/MaterialGraphSchema.cpp
@@ -5,6 +5,14 @@
 =============================================================================*/
 
 #include "UnrealEd.h"
+
+#include "Materials/MaterialExpressionCollectionParameter.h"
+#include "Materials/MaterialExpressionComment.h"
+#include "Materials/MaterialExpressionFunctionInput.h"
+#include "Materials/MaterialExpressionTextureSample.h"
+#include "Materials/MaterialExpressionFunctionOutput.h"
+#include "Materials/MaterialFunction.h"
+
 #include "AssetData.h"
 #include "ScopedTransaction.h"
 #include "MaterialEditorUtilities.h"
diff --git a/Engine/Source/Editor/UnrealEd/Private/ThumbnailRendering/MaterialFunctionThumbnailRenderer.cpp b/Engine/Source/Editor/UnrealEd/Private/ThumbnailRendering/MaterialFunctionThumbnailRenderer.cpp
--- a/Engine/Source/Editor/UnrealEd/Private/ThumbnailRendering/MaterialFunctionThumbnailRenderer.cpp
+++ b/Engine/Source/Editor/UnrealEd/Private/ThumbnailRendering/MaterialFunctionThumbnailRenderer.cpp
@@ -2,6 +2,8 @@
 
 #include "UnrealEd.h"
 
+#include "Materials/MaterialFunction.h"
+
 // FPreviewScene derived helpers for rendering
 #include "ThumbnailHelpers.h"
 
diff --git a/Engine/Source/Runtime/Engine/Classes/Materials/MaterialExpressionAbs.h b/Engine/Source/Runtime/Engine/Classes/Materials/MaterialExpressionAbs.h
--- a/Engine/Source/Runtime/Engine/Classes/Materials/MaterialExpressionAbs.h
+++ b/Engine/Source/Runtime/Engine/Classes/Materials/MaterialExpressionAbs.h
@@ -6,6 +6,7 @@
  */
 
 #pragma once
+#include "Materials/MaterialExpression.h"
 #include "MaterialExpressionAbs.generated.h"
 
 UCLASS(collapsecategories, hidecategories=Object)
diff --git a/Engine/Source/Runtime/Engine/Classes/Materials/MaterialExpressionAntialiasedTextureMask.h b/Engine/Source/Runtime/Engine/Classes/Materials/MaterialExpressionAntialiasedTextureMask.h
--- a/Engine/Source/Runtime/Engine/Classes/Materials/MaterialExpressionAntialiasedTextureMask.h
+++ b/Engine/Source/Runtime/Engine/Classes/Materials/MaterialExpressionAntialiasedTextureMask.h
@@ -2,6 +2,7 @@
 
 
 #pragma once
+#include "Materials/MaterialExpressionTextureSampleParameter2D.h"
 #include "MaterialExpressionAntialiasedTextureMask.generated.h"
 
 UENUM()
diff --git a/Engine/Source/Runtime/Engine/Classes/Materials/MaterialExpressionCeil.h b/Engine/Source/Runtime/Engine/Classes/Materials/MaterialExpressionCeil.h
--- a/Engine/Source/Runtime/Engine/Classes/Materials/MaterialExpressionCeil.h
+++ b/Engine/Source/Runtime/Engine/Classes/Materials/MaterialExpressionCeil.h
@@ -2,8 +2,10 @@
 
 
 #pragma once
+#include "Materials/MaterialExpression.h"
 #include "MaterialExpressionCeil.generated.h"
 
+
 UCLASS(collapsecategories, hidecategories=Object)
 class UMaterialExpressionCeil : public UMaterialExpression
 {
diff --git a/Engine/Source/Runtime/Engine/Classes/Materials/MaterialExpressionDeriveNormalZ.h b/Engine/Source/Runtime/Engine/Classes/Materials/MaterialExpressionDeriveNormalZ.h
--- a/Engine/Source/Runtime/Engine/Classes/Materials/MaterialExpressionDeriveNormalZ.h
+++ b/Engine/Source/Runtime/Engine/Classes/Materials/MaterialExpressionDeriveNormalZ.h
@@ -2,8 +2,10 @@
 
 
 #pragma once
+#include "Materials/MaterialExpression.h"
 #include "MaterialExpressionDeriveNormalZ.generated.h"
 
+
 UCLASS(collapsecategories, hidecategories=Object)
 class UMaterialExpressionDeriveNormalZ : public UMaterialExpression
 {
diff --git a/Engine/Source/Runtime/Engine/Classes/Materials/MaterialExpressionFontSampleParameter.h b/Engine/Source/Runtime/Engine/Classes/Materials/MaterialExpressionFontSampleParameter.h
--- a/Engine/Source/Runtime/Engine/Classes/Materials/MaterialExpressionFontSampleParameter.h
+++ b/Engine/Source/Runtime/Engine/Classes/Materials/MaterialExpressionFontSampleParameter.h
@@ -2,6 +2,7 @@
 
 
 #pragma once
+#include "Materials/MaterialExpressionFontSample.h"
 #include "MaterialExpressionFontSampleParameter.generated.h"
 
 UCLASS(collapsecategories, hidecategories=Object, MinimalAPI)
diff --git a/Engine/Source/Runtime/Engine/Classes/Materials/MaterialExpressionParticleDirection.h b/Engine/Source/Runtime/Engine/Classes/Materials/MaterialExpressionParticleDirection.h
--- a/Engine/Source/Runtime/Engine/Classes/Materials/MaterialExpressionParticleDirection.h
+++ b/Engine/Source/Runtime/Engine/Classes/Materials/MaterialExpressionParticleDirection.h
@@ -6,7 +6,7 @@
 ==============================================================================*/
 
 #pragma once
-
+#include "Materials/MaterialExpression.h"
 #include "MaterialExpressionParticleDirection.generated.h"
 
 UCLASS(collapsecategories, hidecategories=Object)
diff --git a/Engine/Source/Runtime/Engine/Classes/Materials/MaterialExpressionParticleSubUV.h b/Engine/Source/Runtime/Engine/Classes/Materials/MaterialExpressionParticleSubUV.h
--- a/Engine/Source/Runtime/Engine/Classes/Materials/MaterialExpressionParticleSubUV.h
+++ b/Engine/Source/Runtime/Engine/Classes/Materials/MaterialExpressionParticleSubUV.h
@@ -2,6 +2,7 @@
 
 
 #pragma once
+#include "Materials/MaterialExpressionTextureSample.h"
 #include "MaterialExpressionParticleSubUV.generated.h"
 
 UCLASS(MinimalAPI)
diff --git a/Engine/Source/Runtime/Engine/Classes/Materials/MaterialExpressionScalarParameter.h b/Engine/Source/Runtime/Engine/Classes/Materials/MaterialExpressionScalarParameter.h
--- a/Engine/Source/Runtime/Engine/Classes/Materials/MaterialExpressionScalarParameter.h
+++ b/Engine/Source/Runtime/Engine/Classes/Materials/MaterialExpressionScalarParameter.h
@@ -2,6 +2,7 @@
 
 
 #pragma once
+#include "Materials/MaterialExpressionParameter.h"
 #include "MaterialExpressionScalarParameter.generated.h"
 
 UCLASS(collapsecategories, hidecategories=Object, MinimalAPI)
diff --git a/Engine/Source/Runtime/Engine/Classes/Materials/MaterialExpressionStaticBoolParameter.h b/Engine/Source/Runtime/Engine/Classes/Materials/MaterialExpressionStaticBoolParameter.h
--- a/Engine/Source/Runtime/Engine/Classes/Materials/MaterialExpressionStaticBoolParameter.h
+++ b/Engine/Source/Runtime/Engine/Classes/Materials/MaterialExpressionStaticBoolParameter.h
@@ -2,6 +2,7 @@
 
 
 #pragma once
+#include "Materials/MaterialExpressionParameter.h"
 #include "MaterialExpressionStaticBoolParameter.generated.h"
 
 UCLASS(collapsecategories, hidecategories=Object, MinimalAPI)
diff --git a/Engine/Source/Runtime/Engine/Classes/Materials/MaterialExpressionStaticSwitchParameter.h b/Engine/Source/Runtime/Engine/Classes/Materials/MaterialExpressionStaticSwitchParameter.h
--- a/Engine/Source/Runtime/Engine/Classes/Materials/MaterialExpressionStaticSwitchParameter.h
+++ b/Engine/Source/Runtime/Engine/Classes/Materials/MaterialExpressionStaticSwitchParameter.h
@@ -2,6 +2,7 @@
 
 
 #pragma once
+#include "Materials/MaterialExpressionStaticBoolParameter.h"
 #include "MaterialExpressionStaticSwitchParameter.generated.h"
 
 UCLASS(collapsecategories, hidecategories=Object, MinimalAPI)
diff --git a/Engine/Source/Runtime/Engine/Classes/Materials/MaterialExpressionTextureObject.h b/Engine/Source/Runtime/Engine/Classes/Materials/MaterialExpressionTextureObject.h
--- a/Engine/Source/Runtime/Engine/Classes/Materials/MaterialExpressionTextureObject.h
+++ b/Engine/Source/Runtime/Engine/Classes/Materials/MaterialExpressionTextureObject.h
@@ -6,6 +6,7 @@
  */
 
 #pragma once
+#include "Materials/MaterialExpressionTextureBase.h"
 #include "MaterialExpressionTextureObject.generated.h"
 
 UCLASS(collapsecategories, hidecategories=Object, MinimalAPI)
diff --git a/Engine/Source/Runtime/Engine/Classes/Materials/MaterialExpressionTextureObjectParameter.h b/Engine/Source/Runtime/Engine/Classes/Materials/MaterialExpressionTextureObjectParameter.h
--- a/Engine/Source/Runtime/Engine/Classes/Materials/MaterialExpressionTextureObjectParameter.h
+++ b/Engine/Source/Runtime/Engine/Classes/Materials/MaterialExpressionTextureObjectParameter.h
@@ -6,6 +6,7 @@
  */
 
 #pragma once
+#include "Materials/MaterialExpressionTextureSampleParameter.h"
 #include "MaterialExpressionTextureObjectParameter.generated.h"
 
 UCLASS(collapsecategories, hidecategories=Object)
diff --git a/Engine/Source/Runtime/Engine/Classes/Materials/MaterialExpressionTextureSample.h b/Engine/Source/Runtime/Engine/Classes/Materials/MaterialExpressionTextureSample.h
--- a/Engine/Source/Runtime/Engine/Classes/Materials/MaterialExpressionTextureSample.h
+++ b/Engine/Source/Runtime/Engine/Classes/Materials/MaterialExpressionTextureSample.h
@@ -2,6 +2,7 @@
 
 
 #pragma once
+#include "Materials/MaterialExpressionTextureBase.h"
 #include "MaterialExpressionTextureSample.generated.h"
 
 
diff --git a/Engine/Source/Runtime/Engine/Classes/Materials/MaterialExpressionTextureSampleParameter2D.h b/Engine/Source/Runtime/Engine/Classes/Materials/MaterialExpressionTextureSampleParameter2D.h
--- a/Engine/Source/Runtime/Engine/Classes/Materials/MaterialExpressionTextureSampleParameter2D.h
+++ b/Engine/Source/Runtime/Engine/Classes/Materials/MaterialExpressionTextureSampleParameter2D.h
@@ -2,6 +2,7 @@
 
 
 #pragma once
+#include "Materials/MaterialExpressionTextureSampleParameter.h"
 #include "MaterialExpressionTextureSampleParameter2D.generated.h"
 
 UCLASS(collapsecategories, hidecategories=Object, MinimalAPI)
diff --git a/Engine/Source/Runtime/Engine/Classes/Materials/MaterialExpressionTextureSampleParameterSubUV.h b/Engine/Source/Runtime/Engine/Classes/Materials/MaterialExpressionTextureSampleParameterSubUV.h
--- a/Engine/Source/Runtime/Engine/Classes/Materials/MaterialExpressionTextureSampleParameterSubUV.h
+++ b/Engine/Source/Runtime/Engine/Classes/Materials/MaterialExpressionTextureSampleParameterSubUV.h
@@ -2,6 +2,7 @@
 
 
 #pragma once
+#include "Materials/MaterialExpressionTextureSampleParameter2D.h"
 #include "MaterialExpressionTextureSampleParameterSubUV.generated.h"
 
 UCLASS(collapsecategories, hidecategories=Object, MinimalAPI)
diff --git a/Engine/Source/Runtime/Engine/Classes/Materials/MaterialFunction.h b/Engine/Source/Runtime/Engine/Classes/Materials/MaterialFunction.h
--- a/Engine/Source/Runtime/Engine/Classes/Materials/MaterialFunction.h
+++ b/Engine/Source/Runtime/Engine/Classes/Materials/MaterialFunction.h
@@ -68,9 +68,9 @@ public:
 	void UpdateFromFunctionResource();
 
 	/** Get the inputs and outputs that this function exposes, for a function call expression to use. */
-	void GetInputsAndOutputs(TArray<FFunctionExpressionInput>& OutInputs, TArray<FFunctionExpressionOutput>& OutOutputs) const;
+	void GetInputsAndOutputs(TArray<struct FFunctionExpressionInput>& OutInputs, TArray<struct FFunctionExpressionOutput>& OutOutputs) const;
 
-	int32 Compile(class FMaterialCompiler* Compiler, const FFunctionExpressionOutput& Output, int32 MultiplexIndex, const TArray<FFunctionExpressionInput>& Inputs);
+	int32 Compile(class FMaterialCompiler* Compiler, const struct FFunctionExpressionOutput& Output, int32 MultiplexIndex, const TArray<struct FFunctionExpressionInput>& Inputs);
 
 	/** @return true if this function is dependent on the passed in function, directly or indirectly. */
 	ENGINE_API bool IsDependent(UMaterialFunction* OtherFunction);
diff --git a/Engine/Source/Runtime/Engine/Private/Landscape/LandscapeEdit.cpp b/Engine/Source/Runtime/Engine/Private/Landscape/LandscapeEdit.cpp
--- a/Engine/Source/Runtime/Engine/Private/Landscape/LandscapeEdit.cpp
+++ b/Engine/Source/Runtime/Engine/Private/Landscape/LandscapeEdit.cpp
@@ -5,6 +5,10 @@ LandscapeEdit.cpp: Landscape editing
 =============================================================================*/
 
 #include "EnginePrivate.h"
+#include "Materials/MaterialExpressionLandscapeVisibilityMask.h"
+#include "Materials/MaterialExpressionLandscapeLayerWeight.h"
+#include "Materials/MaterialExpressionLandscapeLayerBlend.h"
+#include "Materials/MaterialExpressionLandscapeLayerSwitch.h"
 #include "Landscape/LandscapeDataAccess.h"
 #include "Landscape/LandscapeEdit.h"
 #include "Landscape/LandscapeRender.h"
diff --git a/Engine/Source/Runtime/Engine/Private/Landscape/LandscapeRender.cpp b/Engine/Source/Runtime/Engine/Private/Landscape/LandscapeRender.cpp
--- a/Engine/Source/Runtime/Engine/Private/Landscape/LandscapeRender.cpp
+++ b/Engine/Source/Runtime/Engine/Private/Landscape/LandscapeRender.cpp
@@ -5,6 +5,8 @@ LandscapeRender.cpp: New terrain rendering
 =============================================================================*/
 
 #include "EnginePrivate.h"
+#include "Materials/MaterialExpressionTextureCoordinate.h"
+#include "Materials/MaterialExpressionLandscapeLayerCoords.h"
 #include "ShaderParameters.h"
 #include "Landscape/LandscapeRender.h"
 #include "Landscape/LandscapeEdit.h"
diff --git a/Engine/Source/Runtime/Engine/Private/Materials/HLSLMaterialTranslator.h b/Engine/Source/Runtime/Engine/Private/Materials/HLSLMaterialTranslator.h
--- a/Engine/Source/Runtime/Engine/Private/Materials/HLSLMaterialTranslator.h
+++ b/Engine/Source/Runtime/Engine/Private/Materials/HLSLMaterialTranslator.h
@@ -8,6 +8,16 @@
 #if WITH_EDITORONLY_DATA
 
 #include "EnginePrivate.h"
+#include "Materials/MaterialExpressionWorldPosition.h"
+#include "Materials/MaterialExpressionTextureSample.h"
+#include "Materials/MaterialExpressionSceneTexture.h"
+#include "Materials/MaterialExpressionFunctionInput.h"
+#include "Materials/MaterialExpressionFunctionOutput.h"
+#include "Materials/MaterialExpressionMaterialFunctionCall.h"
+#include "Materials/MaterialExpressionTransform.h"
+#include "Materials/MaterialExpressionTransformPosition.h"
+#include "Materials/MaterialExpressionCustom.h"
+#include "Materials/MaterialFunction.h"
 #include "MaterialCompiler.h"
 #include "MaterialUniformExpressions.h"
 #include "ParameterCollection.h"
diff --git a/Engine/Source/Runtime/Engine/Private/Materials/Material.cpp b/Engine/Source/Runtime/Engine/Private/Materials/Material.cpp
--- a/Engine/Source/Runtime/Engine/Private/Materials/Material.cpp
+++ b/Engine/Source/Runtime/Engine/Private/Materials/Material.cpp
@@ -5,6 +5,24 @@
 =============================================================================*/
 
 #include "EnginePrivate.h"
+#include "Materials/MaterialFunction.h"
+#include "Materials/MaterialExpressionCollectionParameter.h"
+#include "Materials/MaterialExpressionComment.h"
+#include "Materials/MaterialExpressionDynamicParameter.h"
+#include "Materials/MaterialExpressionFontSampleParameter.h"
+#include "Materials/MaterialExpressionLandscapeLayerBlend.h"
+#include "Materials/MaterialExpressionLandscapeLayerSwitch.h"
+#include "Materials/MaterialExpressionLandscapeLayerWeight.h"
+#include "Materials/MaterialExpressionLandscapeVisibilityMask.h"
+#include "Materials/MaterialExpressionMaterialFunctionCall.h"
+#include "Materials/MaterialExpressionMultiply.h"
+#include "Materials/MaterialExpressionQualitySwitch.h"
+#include "Materials/MaterialExpressionScalarParameter.h"
+#include "Materials/MaterialExpressionStaticBoolParameter.h"
+#include "Materials/MaterialExpressionStaticComponentMaskParameter.h"
+#include "Materials/MaterialExpressionStaticSwitchParameter.h"
+#include "Materials/MaterialExpressionTextureSampleParameter.h"
+#include "Materials/MaterialExpressionVectorParameter.h"
 #include "EditorSupportDelegates.h"
 #include "MaterialShader.h"
 #include "MaterialInstanceSupport.h"
diff --git a/Engine/Source/Runtime/Engine/Private/Materials/MaterialExpressions.cpp b/Engine/Source/Runtime/Engine/Private/Materials/MaterialExpressions.cpp
--- a/Engine/Source/Runtime/Engine/Private/Materials/MaterialExpressions.cpp
+++ b/Engine/Source/Runtime/Engine/Private/Materials/MaterialExpressions.cpp
@@ -5,6 +5,136 @@
 =============================================================================*/
 
 #include "EnginePrivate.h"
+
+#include "Materials/MaterialExpressionAbs.h"
+#include "Materials/MaterialExpressionActorPositionWS.h"
+#include "Materials/MaterialExpressionAdd.h"
+#include "Materials/MaterialExpressionAppendVector.h"
+#include "Materials/MaterialExpressionAtmosphericFogColor.h"
+#include "Materials/MaterialExpressionBlackBody.h"
+#include "Materials/MaterialExpressionBreakMaterialAttributes.h"
+#include "Materials/MaterialExpressionBumpOffset.h"
+#include "Materials/MaterialExpressionCameraPositionWS.h"
+#include "Materials/MaterialExpressionCameraVectorWS.h"
+#include "Materials/MaterialExpressionCeil.h"
+#include "Materials/MaterialExpressionClamp.h"
+#include "Materials/MaterialExpressionCollectionParameter.h"
+#include "Materials/MaterialExpressionComment.h"
+#include "Materials/MaterialExpressionComponentMask.h"
+#include "Materials/MaterialExpressionConstant.h"
+#include "Materials/MaterialExpressionConstant2Vector.h"
+#include "Materials/MaterialExpressionConstant3Vector.h"
+#include "Materials/MaterialExpressionConstant4Vector.h"
+#include "Materials/MaterialExpressionConstantBiasScale.h"
+#include "Materials/MaterialExpressionCosine.h"
+#include "Materials/MaterialExpressionCrossProduct.h"
+#include "Materials/MaterialExpressionCustom.h"
+#include "Materials/MaterialExpressionCustomTexture.h"
+#include "Materials/MaterialExpressionDDX.h"
+#include "Materials/MaterialExpressionDDY.h"
+#include "Materials/MaterialExpressionDepthFade.h"
+#include "Materials/MaterialExpressionDepthOfFieldFunction.h"
+#include "Materials/MaterialExpressionDeriveNormalZ.h"
+#include "Materials/MaterialExpressionDesaturation.h"
+#include "Materials/MaterialExpressionDistance.h"
+#include "Materials/MaterialExpressionDistanceCullFade.h"
+#include "Materials/MaterialExpressionDivide.h"
+#include "Materials/MaterialExpressionDotProduct.h"
+#include "Materials/MaterialExpressionDynamicParameter.h"
+#include "Materials/MaterialExpressionEyeAdaptation.h"
+#include "Materials/MaterialExpressionFeatureLevelSwitch.h"
+#include "Materials/MaterialExpressionFloor.h"
+#include "Materials/MaterialExpressionFmod.h"
+#include "Materials/MaterialExpressionFontSample.h"
+#include "Materials/MaterialExpressionFontSampleParameter.h"
+#include "Materials/MaterialExpressionFrac.h"
+#include "Materials/MaterialExpressionFresnel.h"
+#include "Materials/MaterialExpressionFunctionInput.h"
+#include "Materials/MaterialExpressionFunctionOutput.h"
+#include "Materials/MaterialExpressionGIReplace.h"
+#include "Materials/MaterialExpressionIf.h"
+#include "Materials/MaterialExpressionLandscapeLayerBlend.h"
+#include "Materials/MaterialExpressionLandscapeLayerCoords.h"
+#include "Materials/MaterialExpressionLandscapeLayerSwitch.h"
+#include "Materials/MaterialExpressionLandscapeLayerWeight.h"
+#include "Materials/MaterialExpressionLandscapeVisibilityMask.h"
+#include "Materials/MaterialExpressionLightmapUVs.h"
+#include "Materials/MaterialExpressionLightmassReplace.h"
+#include "Materials/MaterialExpressionLightVector.h"
+#include "Materials/MaterialExpressionLinearInterpolate.h"
+#include "Materials/MaterialExpressionMakeMaterialAttributes.h"
+#include "Materials/MaterialExpressionMaterialFunctionCall.h"
+#include "Materials/MaterialExpressionMax.h"
+#include "Materials/MaterialExpressionMin.h"
+#include "Materials/MaterialExpressionMultiply.h"
+#include "Materials/MaterialExpressionNoise.h"
+#include "Materials/MaterialExpressionNormalize.h"
+#include "Materials/MaterialExpressionObjectBounds.h"
+#include "Materials/MaterialExpressionObjectOrientation.h"
+#include "Materials/MaterialExpressionObjectPositionWS.h"
+#include "Materials/MaterialExpressionObjectRadius.h"
+#include "Materials/MaterialExpressionOneMinus.h"
+#include "Materials/MaterialExpressionPanner.h"
+#include "Materials/MaterialExpressionParameter.h"
+#include "Materials/MaterialExpressionScalarParameter.h"
+#include "Materials/MaterialExpressionStaticBoolParameter.h"
+#include "Materials/MaterialExpressionStaticSwitchParameter.h"
+#include "Materials/MaterialExpressionStaticComponentMaskParameter.h"
+#include "Materials/MaterialExpressionVectorParameter.h"
+#include "Materials/MaterialExpressionParticleColor.h"
+#include "Materials/MaterialExpressionParticleDirection.h"
+#include "Materials/MaterialExpressionParticleMacroUV.h"
+#include "Materials/MaterialExpressionParticleMotionBlurFade.h"
+#include "Materials/MaterialExpressionParticlePositionWS.h"
+#include "Materials/MaterialExpressionParticleRadius.h"
+#include "Materials/MaterialExpressionParticleRelativeTime.h"
+#include "Materials/MaterialExpressionParticleSize.h"
+#include "Materials/MaterialExpressionParticleSpeed.h"
+#include "Materials/MaterialExpressionPerInstanceFadeAmount.h"
+#include "Materials/MaterialExpressionPerInstanceRandom.h"
+#include "Materials/MaterialExpressionPixelDepth.h"
+#include "Materials/MaterialExpressionPixelNormalWS.h"
+#include "Materials/MaterialExpressionPower.h"
+#include "Materials/MaterialExpressionQualitySwitch.h"
+#include "Materials/MaterialExpressionReflectionVectorWS.h"
+#include "Materials/MaterialExpressionRotateAboutAxis.h"
+#include "Materials/MaterialExpressionRotator.h"
+#include "Materials/MaterialExpressionSceneColor.h"
+#include "Materials/MaterialExpressionSceneDepth.h"
+#include "Materials/MaterialExpressionSceneTexelSize.h"
+#include "Materials/MaterialExpressionSceneTexture.h"
+#include "Materials/MaterialExpressionScreenPosition.h"
+#include "Materials/MaterialExpressionSine.h"
+#include "Materials/MaterialExpressionSpeedTree.h"
+#include "Materials/MaterialExpressionSpeedTreeColorVariation.h"
+#include "Materials/MaterialExpressionSphereMask.h"
+#include "Materials/MaterialExpressionSphericalParticleOpacity.h"
+#include "Materials/MaterialExpressionSquareRoot.h"
+#include "Materials/MaterialExpressionStaticBool.h"
+#include "Materials/MaterialExpressionStaticSwitch.h"
+#include "Materials/MaterialExpressionSubtract.h"
+#include "Materials/MaterialExpressionTextureBase.h"
+#include "Materials/MaterialExpressionTextureObject.h"
+#include "Materials/MaterialExpressionTextureSample.h"
+#include "Materials/MaterialExpressionParticleSubUV.h"
+#include "Materials/MaterialExpressionTextureSampleParameter.h"
+#include "Materials/MaterialExpressionTextureObjectParameter.h"
+#include "Materials/MaterialExpressionTextureSampleParameter2D.h"
+#include "Materials/MaterialExpressionAntialiasedTextureMask.h"
+#include "Materials/MaterialExpressionTextureSampleParameterSubUV.h"
+#include "Materials/MaterialExpressionTextureSampleParameterCube.h"
+#include "Materials/MaterialExpressionTextureSampleParameterMovie.h"
+#include "Materials/MaterialExpressionTextureCoordinate.h"
+#include "Materials/MaterialExpressionTime.h"
+#include "Materials/MaterialExpressionTransform.h"
+#include "Materials/MaterialExpressionTransformPosition.h"
+#include "Materials/MaterialExpressionTwoSidedSign.h"
+#include "Materials/MaterialExpressionVertexColor.h"
+#include "Materials/MaterialExpressionVertexNormalWS.h"
+#include "Materials/MaterialExpressionViewSize.h"
+#include "Materials/MaterialExpressionWorldPosition.h"
+#include "Materials/MaterialFunction.h"
+
 #include "EditorSupportDelegates.h"
 #include "MaterialCompiler.h"
 #if WITH_EDITOR
diff --git a/Engine/Source/Runtime/Engine/Private/Materials/MaterialInstance.cpp b/Engine/Source/Runtime/Engine/Private/Materials/MaterialInstance.cpp
--- a/Engine/Source/Runtime/Engine/Private/Materials/MaterialInstance.cpp
+++ b/Engine/Source/Runtime/Engine/Private/Materials/MaterialInstance.cpp
@@ -1,6 +1,13 @@
 // Copyright 1998-2014 Epic Games, Inc. All Rights Reserved.
 
 #include "EnginePrivate.h"
+#include "Materials/MaterialExpressionScalarParameter.h"
+#include "Materials/MaterialExpressionVectorParameter.h"
+#include "Materials/MaterialExpressionTextureSampleParameter.h"
+#include "Materials/MaterialExpressionFontSampleParameter.h"
+#include "Materials/MaterialExpressionStaticBoolParameter.h"
+#include "Materials/MaterialExpressionStaticComponentMaskParameter.h"
+#include "Materials/MaterialExpressionLandscapeLayerWeight.h"
 #include "MaterialInstanceSupport.h"
 #include "MaterialShader.h"
 #include "TargetPlatform.h"
diff --git a/Engine/Source/Runtime/Engine/Private/Materials/MaterialShared.cpp b/Engine/Source/Runtime/Engine/Private/Materials/MaterialShared.cpp
--- a/Engine/Source/Runtime/Engine/Private/Materials/MaterialShared.cpp
+++ b/Engine/Source/Runtime/Engine/Private/Materials/MaterialShared.cpp
@@ -5,6 +5,7 @@
 =============================================================================*/
 
 #include "EnginePrivate.h"
+#include "Materials/MaterialExpressionBreakMaterialAttributes.h"
 #include "PixelFormat.h"
 #include "ShaderCompiler.h"
 #include "MaterialCompiler.h"
diff --git a/Engine/Source/Runtime/Engine/Private/Particles/ParticleModules_Parameter.cpp b/Engine/Source/Runtime/Engine/Private/Particles/ParticleModules_Parameter.cpp
--- a/Engine/Source/Runtime/Engine/Private/Particles/ParticleModules_Parameter.cpp
+++ b/Engine/Source/Runtime/Engine/Private/Particles/ParticleModules_Parameter.cpp
@@ -5,6 +5,7 @@
 	Parameter-related particle module implementations.
 =============================================================================*/
 #include "EnginePrivate.h"
+#include "Materials/MaterialExpressionDynamicParameter.h"
 #include "ParticleDefinitions.h"
 #include "../DistributionHelpers.h"
 
diff --git a/Engine/Source/Runtime/Engine/Public/Engine.h b/Engine/Source/Runtime/Engine/Public/Engine.h
--- a/Engine/Source/Runtime/Engine/Public/Engine.h
+++ b/Engine/Source/Runtime/Engine/Public/Engine.h
@@ -464,135 +464,6 @@
 #include "Lightmass/LightmappedSurfaceCollection.h"
 #include "GameFramework/LocalMessage.h"
 #include "GameFramework/EngineMessage.h"
-#include "Materials/MaterialExpression.h"
-#include "Materials/MaterialExpressionAbs.h"
-#include "Materials/MaterialExpressionActorPositionWS.h"
-#include "Materials/MaterialExpressionAdd.h"
-#include "Materials/MaterialExpressionAppendVector.h"
-#include "Materials/MaterialExpressionAtmosphericFogColor.h"
-#include "Materials/MaterialExpressionBlackBody.h"
-#include "Materials/MaterialExpressionBreakMaterialAttributes.h"
-#include "Materials/MaterialExpressionBumpOffset.h"
-#include "Materials/MaterialExpressionCameraPositionWS.h"
-#include "Materials/MaterialExpressionCameraVectorWS.h"
-#include "Materials/MaterialExpressionCeil.h"
-#include "Materials/MaterialExpressionClamp.h"
-#include "Materials/MaterialExpressionCollectionParameter.h"
-#include "Materials/MaterialExpressionComment.h"
-#include "Materials/MaterialExpressionComponentMask.h"
-#include "Materials/MaterialExpressionConstant.h"
-#include "Materials/MaterialExpressionConstant2Vector.h"
-#include "Materials/MaterialExpressionConstant3Vector.h"
-#include "Materials/MaterialExpressionConstant4Vector.h"
-#include "Materials/MaterialExpressionConstantBiasScale.h"
-#include "Materials/MaterialExpressionCosine.h"
-#include "Materials/MaterialExpressionCrossProduct.h"
-#include "Materials/MaterialExpressionCustom.h"
-#include "Materials/MaterialExpressionCustomTexture.h"
-#include "Materials/MaterialExpressionDDX.h"
-#include "Materials/MaterialExpressionDDY.h"
-#include "Materials/MaterialExpressionDepthFade.h"
-#include "Materials/MaterialExpressionDepthOfFieldFunction.h"
-#include "Materials/MaterialExpressionDeriveNormalZ.h"
-#include "Materials/MaterialExpressionDesaturation.h"
-#include "Materials/MaterialExpressionDistance.h"
-#include "Materials/MaterialExpressionDistanceCullFade.h"
-#include "Materials/MaterialExpressionDivide.h"
-#include "Materials/MaterialExpressionDotProduct.h"
-#include "Materials/MaterialExpressionDynamicParameter.h"
-#include "Materials/MaterialExpressionEyeAdaptation.h"
-#include "Materials/MaterialExpressionFeatureLevelSwitch.h"
-#include "Materials/MaterialExpressionFloor.h"
-#include "Materials/MaterialExpressionFmod.h"
-#include "Materials/MaterialExpressionFontSample.h"
-#include "Materials/MaterialExpressionFontSampleParameter.h"
-#include "Materials/MaterialExpressionFrac.h"
-#include "Materials/MaterialExpressionFresnel.h"
-#include "Materials/MaterialExpressionFunctionInput.h"
-#include "Materials/MaterialExpressionFunctionOutput.h"
-#include "Materials/MaterialExpressionGIReplace.h"
-#include "Materials/MaterialExpressionIf.h"
-#include "Materials/MaterialExpressionLandscapeLayerBlend.h"
-#include "Materials/MaterialExpressionLandscapeLayerCoords.h"
-#include "Materials/MaterialExpressionLandscapeLayerSwitch.h"
-#include "Materials/MaterialExpressionLandscapeLayerWeight.h"
-#include "Materials/MaterialExpressionLandscapeVisibilityMask.h"
-#include "Materials/MaterialExpressionLightmapUVs.h"
-#include "Materials/MaterialExpressionLightmassReplace.h"
-#include "Materials/MaterialExpressionLightVector.h"
-#include "Materials/MaterialExpressionLinearInterpolate.h"
-#include "Materials/MaterialExpressionMakeMaterialAttributes.h"
-#include "Materials/MaterialExpressionMaterialFunctionCall.h"
-#include "Materials/MaterialExpressionMax.h"
-#include "Materials/MaterialExpressionMin.h"
-#include "Materials/MaterialExpressionMultiply.h"
-#include "Materials/MaterialExpressionNoise.h"
-#include "Materials/MaterialExpressionNormalize.h"
-#include "Materials/MaterialExpressionObjectBounds.h"
-#include "Materials/MaterialExpressionObjectOrientation.h"
-#include "Materials/MaterialExpressionObjectPositionWS.h"
-#include "Materials/MaterialExpressionObjectRadius.h"
-#include "Materials/MaterialExpressionOneMinus.h"
-#include "Materials/MaterialExpressionPanner.h"
-#include "Materials/MaterialExpressionParameter.h"
-#include "Materials/MaterialExpressionScalarParameter.h"
-#include "Materials/MaterialExpressionStaticBoolParameter.h"
-#include "Materials/MaterialExpressionStaticSwitchParameter.h"
-#include "Materials/MaterialExpressionStaticComponentMaskParameter.h"
-#include "Materials/MaterialExpressionVectorParameter.h"
-#include "Materials/MaterialExpressionParticleColor.h"
-#include "Materials/MaterialExpressionParticleDirection.h"
-#include "Materials/MaterialExpressionParticleMacroUV.h"
-#include "Materials/MaterialExpressionParticleMotionBlurFade.h"
-#include "Materials/MaterialExpressionParticlePositionWS.h"
-#include "Materials/MaterialExpressionParticleRadius.h"
-#include "Materials/MaterialExpressionParticleRelativeTime.h"
-#include "Materials/MaterialExpressionParticleSize.h"
-#include "Materials/MaterialExpressionParticleSpeed.h"
-#include "Materials/MaterialExpressionPerInstanceFadeAmount.h"
-#include "Materials/MaterialExpressionPerInstanceRandom.h"
-#include "Materials/MaterialExpressionPixelDepth.h"
-#include "Materials/MaterialExpressionPixelNormalWS.h"
-#include "Materials/MaterialExpressionPower.h"
-#include "Materials/MaterialExpressionQualitySwitch.h"
-#include "Materials/MaterialExpressionReflectionVectorWS.h"
-#include "Materials/MaterialExpressionRotateAboutAxis.h"
-#include "Materials/MaterialExpressionRotator.h"
-#include "Materials/MaterialExpressionSceneColor.h"
-#include "Materials/MaterialExpressionSceneDepth.h"
-#include "Materials/MaterialExpressionSceneTexelSize.h"
-#include "Materials/MaterialExpressionSceneTexture.h"
-#include "Materials/MaterialExpressionScreenPosition.h"
-#include "Materials/MaterialExpressionSine.h"
-#include "Materials/MaterialExpressionSpeedTree.h"
-#include "Materials/MaterialExpressionSpeedTreeColorVariation.h"
-#include "Materials/MaterialExpressionSphereMask.h"
-#include "Materials/MaterialExpressionSphericalParticleOpacity.h"
-#include "Materials/MaterialExpressionSquareRoot.h"
-#include "Materials/MaterialExpressionStaticBool.h"
-#include "Materials/MaterialExpressionStaticSwitch.h"
-#include "Materials/MaterialExpressionSubtract.h"
-#include "Materials/MaterialExpressionTextureBase.h"
-#include "Materials/MaterialExpressionTextureObject.h"
-#include "Materials/MaterialExpressionTextureSample.h"
-#include "Materials/MaterialExpressionParticleSubUV.h"
-#include "Materials/MaterialExpressionTextureSampleParameter.h"
-#include "Materials/MaterialExpressionTextureObjectParameter.h"
-#include "Materials/MaterialExpressionTextureSampleParameter2D.h"
-#include "Materials/MaterialExpressionAntialiasedTextureMask.h"
-#include "Materials/MaterialExpressionTextureSampleParameterSubUV.h"
-#include "Materials/MaterialExpressionTextureSampleParameterCube.h"
-#include "Materials/MaterialExpressionTextureSampleParameterMovie.h"
-#include "Materials/MaterialExpressionTextureCoordinate.h"
-#include "Materials/MaterialExpressionTime.h"
-#include "Materials/MaterialExpressionTransform.h"
-#include "Materials/MaterialExpressionTransformPosition.h"
-#include "Materials/MaterialExpressionTwoSidedSign.h"
-#include "Materials/MaterialExpressionVertexColor.h"
-#include "Materials/MaterialExpressionVertexNormalWS.h"
-#include "Materials/MaterialExpressionViewSize.h"
-#include "Materials/MaterialExpressionWorldPosition.h"
-#include "Materials/MaterialFunction.h"
 #include "Materials/MaterialInterface.h"
 #include "Materials/Material.h"
 #include "Materials/MaterialInstance.h"
