Fix form extended structures not being properly updated from active transforms

diff --git a/Engine/Source/Runtime/Engine/Private/PhysicsEngine/PhysScene.cpp b/Engine/Source/Runtime/Engine/Private/PhysicsEngine/PhysScene.cpp
--- a/Engine/Source/Runtime/Engine/Private/PhysicsEngine/PhysScene.cpp
+++ b/Engine/Source/Runtime/Engine/Private/PhysicsEngine/PhysScene.cpp
@@ -597,6 +597,8 @@ void FPhysScene::UpdateActiveTransforms(uint32 SceneType)
 	for (PxU32 TransformIdx = 0; TransformIdx < NumTransforms; ++TransformIdx)
 	{
 		ActiveTransforms[SceneType].Add(&PActiveTransforms[TransformIdx]);	//it's ok to use a pointer here because the physics scene leaves the data for us until the next call to fetchTransform;
+		
+
 	}
 }
 #endif
@@ -617,6 +619,7 @@ void FPhysScene::SyncComponentsToBodies(uint32 SceneType)
 		const PxActiveTransform& PActiveTransform = *PActiveTransforms[TransformIdx];
 		PxRigidActor* RigidActor = PActiveTransform.actor->isRigidActor();
 
+
 #if WITH_APEX
 		//Special code for destructible chunk
 		if (const FDestructibleChunkInfo * DestructibleChunkInfo = FPhysxUserData::Get<FDestructibleChunkInfo>(RigidActor->userData))
@@ -634,7 +637,7 @@ void FPhysScene::SyncComponentsToBodies(uint32 SceneType)
 					NxDestructibleActor* DestructibleActor = GApexModuleDestructible->getDestructibleAndChunk(Shape, &ChunkIndex);
 					const physx::PxMat44 ChunkPoseRT = DestructibleActor->getChunkPose(ChunkIndex);
 					const physx::PxTransform Transform(ChunkPoseRT);
-					if (UDestructibleComponent * DestructibleComponent = Cast<UDestructibleComponent>(FPhysxUserData::Get<UPrimitiveComponent>(Shape->userData)))
+					if (UDestructibleComponent * DestructibleComponent = Cast<UDestructibleComponent>(FPhysxUserData::Get<UPrimitiveComponent>(DestructibleActor->userData)))
 					{
 						DestructibleComponent->SetChunkWorldRT(ChunkIndex, P2UQuat(Transform.q), P2UVector(Transform.p));
 					}
