Fix form extended structures crashing due to chunk info realloc
We can't preallocate chunk info for form extended because chunks move between destructible components.
In future we could use a shared pool to remove allocation overhead

diff --git a/Engine/Source/Runtime/Engine/Classes/Components/DestructibleComponent.h b/Engine/Source/Runtime/Engine/Classes/Components/DestructibleComponent.h
--- a/Engine/Source/Runtime/Engine/Classes/Components/DestructibleComponent.h
+++ b/Engine/Source/Runtime/Engine/Classes/Components/DestructibleComponent.h
@@ -70,7 +70,7 @@ class ENGINE_API UDestructibleComponent : public USkinnedMeshComponent
 
 #if WITH_PHYSX
 	/** Per chunk info */
-	TArray<FDestructibleChunkInfo> ChunkInfos;
+	TIndirectArray<FDestructibleChunkInfo> ChunkInfos;
 #endif // WITH_PHYSX 
 
 #if WITH_EDITOR
@@ -191,7 +191,7 @@ private:
 	FPhysxUserData PhysxUserData;
 
 	/** User data wrapper for the chunks passed to physx */
-	TArray<FPhysxUserData> PhysxChunkUserData;
+	TIndirectArray<FPhysxUserData> PhysxChunkUserData;
 	bool IsChunkLarge(int32 ChunkIdx) const;
 	void SetCollisionResponseForActor(physx::PxRigidDynamic* Actor, int32 ChunkIdx);
 #endif
diff --git a/Engine/Source/Runtime/Engine/Private/DestructibleComponent.cpp b/Engine/Source/Runtime/Engine/Private/DestructibleComponent.cpp
--- a/Engine/Source/Runtime/Engine/Private/DestructibleComponent.cpp
+++ b/Engine/Source/Runtime/Engine/Private/DestructibleComponent.cpp
@@ -305,10 +305,8 @@ void UDestructibleComponent::CreatePhysicsState()
 	PhysxUserData = FPhysxUserData(this);
 	ApexDestructibleActor->userData = &PhysxUserData;
 
-	// Setup chunk user data arrays. We have to make sure PhysxChunkUserData will not be reallocated when growing, so we 
-	// reserver ChunkCount here already
-	ChunkInfos.Empty(ChunkCount);
-	PhysxChunkUserData.Empty(ChunkCount);
+	ChunkInfos.Empty();
+	PhysxChunkUserData.Empty();
 
 	// Cache cooked collision data
 	// BRGTODO : cook in asset
@@ -788,9 +786,8 @@ void UDestructibleComponent::SetChunkVisible( int32 ChunkIndex, bool bVisible )
 			else if (PActor->userData == NULL)
 			{
 				// Setup the user data to have a proper chunk - actor mapping 
-				int32 InfoIndex = ChunkInfos.AddUninitialized();
-
-				FDestructibleChunkInfo* CI = &ChunkInfos[InfoIndex];
+				FDestructibleChunkInfo* CI = new FDestructibleChunkInfo();
+				int32 InfoIndex = ChunkInfos.Add(CI);
 				CI->Index = InfoIndex;
 				CI->ChunkIndex = ChunkIndex;
 				CI->OwningComponent = this;
@@ -798,10 +795,11 @@ void UDestructibleComponent::SetChunkVisible( int32 ChunkIndex, bool bVisible )
 
 				PActor->setActorFlag(PxActorFlag::eDISABLE_GRAVITY, !BodyInstance.bEnableGravity);
 
-				int32 UserDataIdx = PhysxChunkUserData.Add(FPhysxUserData(CI));
+				FPhysxUserData * UserData = new FPhysxUserData(CI);
+				int32 UserDataIdx = PhysxChunkUserData.Add(UserData);
 				check(InfoIndex == UserDataIdx);
 
-				PActor->userData = &PhysxChunkUserData[UserDataIdx];
+				PActor->userData = UserData;
 
 				// Set collision response to non-root chunks
 				if (GetDestructibleMesh()->ApexDestructibleAsset->getChunkParentIndex(ChunkIndex) >= 0)
diff --git a/Engine/Source/Runtime/Engine/Private/PhysicsEngine/PhysScene.cpp b/Engine/Source/Runtime/Engine/Private/PhysicsEngine/PhysScene.cpp
--- a/Engine/Source/Runtime/Engine/Private/PhysicsEngine/PhysScene.cpp
+++ b/Engine/Source/Runtime/Engine/Private/PhysicsEngine/PhysScene.cpp
@@ -598,7 +598,10 @@ void FPhysScene::UpdateActiveTransforms(uint32 SceneType)
 	{
 		ActiveTransforms[SceneType].Add(&PActiveTransforms[TransformIdx]);	//it's ok to use a pointer here because the physics scene leaves the data for us until the next call to fetchTransform;
 		
+		const PxActiveTransform& PActiveTransform = PActiveTransforms[TransformIdx];
+		PxRigidActor* RigidActor = PActiveTransform.actor->isRigidActor();
 
+		ensure(!RigidActor->userData || !FPhysxUserData::IsGarbage(RigidActor->userData));
 	}
 }
 #endif
@@ -619,6 +622,7 @@ void FPhysScene::SyncComponentsToBodies(uint32 SceneType)
 		const PxActiveTransform& PActiveTransform = *PActiveTransforms[TransformIdx];
 		PxRigidActor* RigidActor = PActiveTransform.actor->isRigidActor();
 
+		ensure(!RigidActor->userData || !FPhysxUserData::IsGarbage(RigidActor->userData));
 
 #if WITH_APEX
 		//Special code for destructible chunk
diff --git a/Engine/Source/Runtime/Engine/Public/PhysxUserData.h b/Engine/Source/Runtime/Engine/Public/PhysxUserData.h
--- a/Engine/Source/Runtime/Engine/Public/PhysxUserData.h
+++ b/Engine/Source/Runtime/Engine/Public/PhysxUserData.h
@@ -42,9 +42,12 @@ public:
 	FPhysxUserData(FConstraintInstance* InPayload)		:Type(EPhysxUserDataType::ConstraintInstance), Payload(InPayload) {}
 	FPhysxUserData(UPrimitiveComponent* InPayload)		:Type(EPhysxUserDataType::PrimitiveComponent), Payload(InPayload) {}
 	FPhysxUserData(FDestructibleChunkInfo* InPayload)	:Type(EPhysxUserDataType::DestructibleChunk), Payload(InPayload) {}
-
+	
 	template <class T> static T* Get(void* UserData);
 	template <class T> static void Set(void* UserData, T* Payload);
+
+	//helper function to determine if userData is garbage (maybe dangling pointer)
+	static bool IsGarbage(void * UserData){ return ((FPhysxUserData*)UserData)->Type < EPhysxUserDataType::Invalid || ((FPhysxUserData*)UserData)->Type > EPhysxUserDataType::DestructibleChunk; }
 };
 
 template <> FORCEINLINE FBodyInstance* FPhysxUserData::Get(void* UserData)			{ if (!UserData || ((FPhysxUserData*)UserData)->Type != EPhysxUserDataType::BodyInstance) { return NULL; } return (FBodyInstance*)((FPhysxUserData*)UserData)->Payload; }
@@ -61,4 +64,5 @@ template <> FORCEINLINE void FPhysxUserData::Set(void* UserData, FConstraintInst
 template <> FORCEINLINE void FPhysxUserData::Set(void* UserData, UPrimitiveComponent* Payload)		{ check(UserData); ((FPhysxUserData*)UserData)->Type = EPhysxUserDataType::PrimitiveComponent; ((FPhysxUserData*)UserData)->Payload = Payload; }
 template <> FORCEINLINE void FPhysxUserData::Set(void* UserData, FDestructibleChunkInfo* Payload)	{ check(UserData); ((FPhysxUserData*)UserData)->Type = EPhysxUserDataType::DestructibleChunk; ((FPhysxUserData*)UserData)->Payload = Payload; }
 
+
 #endif // WITH_PHYSICS
\ No newline at end of file
