UE4: DLL exporting UShapeComponent. UPrimitiveComponent and some subclasses of UShapeComponent are already exported, so there is no reason UShapeComponent shouldn't be as well.

diff --git a/Engine/Source/Runtime/Engine/Classes/Components/ShapeComponent.h b/Engine/Source/Runtime/Engine/Classes/Components/ShapeComponent.h
--- a/Engine/Source/Runtime/Engine/Classes/Components/ShapeComponent.h
+++ b/Engine/Source/Runtime/Engine/Classes/Components/ShapeComponent.h
@@ -4,8 +4,8 @@
 #pragma once
 #include "ShapeComponent.generated.h"
 
-UCLASS(abstract, hidecategories=(Object,LOD,Lighting,TextureStreaming,Activation,"Components|Activation"), editinlinenew, MinimalAPI, meta=(BlueprintSpawnableComponent), showcategories=(Mobility))
-class UShapeComponent : public UPrimitiveComponent
+UCLASS(abstract, hidecategories=(Object,LOD,Lighting,TextureStreaming,Activation,"Components|Activation"), editinlinenew, meta=(BlueprintSpawnableComponent), showcategories=(Mobility))
+class ENGINE_API UShapeComponent : public UPrimitiveComponent
 {
 	GENERATED_UCLASS_BODY()
 
