fixed issues with Tappy Chicken merge
#ue4
#tappy chicken

diff --git a/Engine/Source/Runtime/Engine/Private/PhysicsEngine/PhysSubstepTasks.cpp b/Engine/Source/Runtime/Engine/Private/PhysicsEngine/PhysSubstepTasks.cpp
--- a/Engine/Source/Runtime/Engine/Private/PhysicsEngine/PhysSubstepTasks.cpp
+++ b/Engine/Source/Runtime/Engine/Private/PhysicsEngine/PhysSubstepTasks.cpp
@@ -296,10 +296,12 @@ void FPhysSubstepTask::SubstepSimulationStart()
 	float DeltaTime = bLastSubstep ? (DeltaSeconds - TotalSubTime) : SubTime;
 	float Interpolation = bLastSubstep ? 1.f : Alpha;
 
+#if WITH_VEHICLE
 	if (VehicleManager)
 	{
 		VehicleManager->Update(DeltaTime);
 	}
+#endif
 
 	SubstepInterpolation(Interpolation);
 
diff --git a/Engine/Source/Runtime/Engine/Public/EngineDefines.h b/Engine/Source/Runtime/Engine/Public/EngineDefines.h
--- a/Engine/Source/Runtime/Engine/Public/EngineDefines.h
+++ b/Engine/Source/Runtime/Engine/Public/EngineDefines.h
@@ -22,6 +22,13 @@
 	#define WITH_APEX (1 && WITH_PHYSX)
 #endif
 
+/** 
+ *   Whether or not compiling with Vehicle extensions to PhysX
+ */
+#ifndef WITH_VEHICLE
+	#define WITH_VEHICLE (1 && WITH_PHYSX)
+#endif
+
 #ifndef WITH_PHYSICS_COOKING
 	#define WITH_PHYSICS_COOKING (WITH_EDITOR || WITH_APEX)		//APEX currently relies on cooking even at runtime
 #endif
