#UE4 #HTML5

 - remove #pragma accidently checked in with merge CL. 2082256

diff --git a/Engine/Source/Runtime/Engine/Public/SaveGameSystem.h b/Engine/Source/Runtime/Engine/Public/SaveGameSystem.h
--- a/Engine/Source/Runtime/Engine/Public/SaveGameSystem.h
+++ b/Engine/Source/Runtime/Engine/Public/SaveGameSystem.h
@@ -28,10 +28,8 @@ public:
 	virtual bool LoadGame(bool bAttemptToUseUI, const TCHAR* Name, const int32 UserIndex, TArray<uint8>& Data) = 0;
 };
 
-#pragma  optimize("",off)
 
-
-//  @todo akhare - move this to core before checking - Properly implement the SaveSystem  for HTML5 
+//  @todo akhare - move this to core before checking - Properly implement the SaveSystem  for HTML5
 /** A generic save game system that just uses IFileManager to save/load with normal files */
 class FGenericSaveGameSystem : public ISaveGameSystem
 {
@@ -106,5 +104,3 @@ protected:
 		return FString::Printf(TEXT("%s/SaveGames/%s.sav"), *FPaths::GameSavedDir(), Name);
 	}
 };
-
-#pragma  optimize("",on)
