Expose more FCanvas functions with ENGINE_API

diff --git a/Engine/Source/Runtime/Engine/Public/Canvas.h b/Engine/Source/Runtime/Engine/Public/Canvas.h
--- a/Engine/Source/Runtime/Engine/Public/Canvas.h
+++ b/Engine/Source/Runtime/Engine/Public/Canvas.h
@@ -53,7 +53,7 @@ public:
 	ENGINE_API ~FCanvas();
 
 	
-	static ESimpleElementBlendMode BlendToSimpleElementBlend(EBlendMode BlendMode);
+	ENGINE_API static ESimpleElementBlendMode BlendToSimpleElementBlend(EBlendMode BlendMode);
 	/**
 	* Returns a FBatchedElements pointer to be used for adding vertices and primitives for rendering.
 	* Adds a new render item to the sort element entry based on the current sort key.
@@ -70,7 +70,7 @@ public:
 	/**
 	* Generates a new FCanvasTileRendererItem for the current sortkey and adds it to the sortelement list of itmes to render
 	*/
-	void AddTileRenderItem(float X,float Y,float SizeX,float SizeY,float U,float V,float SizeU,float SizeV,const FMaterialRenderProxy* MaterialRenderProxy,FHitProxyId HitProxyId,bool bFreezeTime);
+	ENGINE_API void AddTileRenderItem(float X,float Y,float SizeX,float SizeY,float U,float V,float SizeU,float SizeV,const FMaterialRenderProxy* MaterialRenderProxy,FHitProxyId HitProxyId,bool bFreezeTime);
 
 	/** 
 	* Sends a message to the rendering thread to draw the batched elements. 
@@ -100,7 +100,7 @@ public:
 	*
 	* @param Transform - The transform to use for the base
 	*/
-	void SetBaseTransform(const FMatrix& Transform);
+	ENGINE_API void SetBaseTransform(const FMatrix& Transform);
 
 	/**
 	* Generate a 2D projection for the canvas. Use this if you only want to transform in 2D on the XY plane
@@ -120,7 +120,7 @@ public:
 	* @param NearPlane - Distance to the near clip plane
 	* @return Matrix for canvas projection
 	*/
-	static FMatrix CalcBaseTransform3D(uint32 ViewSizeX, uint32 ViewSizeY, float fFOV, float NearPlane);
+	ENGINE_API static FMatrix CalcBaseTransform3D(uint32 ViewSizeX, uint32 ViewSizeY, float fFOV, float NearPlane);
 	
 	/**
 	* Generate a view matrix for the canvas. Used for CalcBaseTransform3D
@@ -130,7 +130,7 @@ public:
 	* @param fFOV - Field of view for the projection
 	* @return Matrix for canvas view orientation
 	*/
-	static FMatrix CalcViewMatrix(uint32 ViewSizeX, uint32 ViewSizeY, float fFOV);
+	ENGINE_API static FMatrix CalcViewMatrix(uint32 ViewSizeX, uint32 ViewSizeY, float fFOV);
 	
 	/**
 	* Generate a projection matrix for the canvas. Used for CalcBaseTransform3D
@@ -141,13 +141,13 @@ public:
 	* @param NearPlane - Distance to the near clip plane
 	* @return Matrix for canvas projection
 	*/
-	static FMatrix CalcProjectionMatrix(uint32 ViewSizeX, uint32 ViewSizeY, float fFOV, float NearPlane);
+	ENGINE_API static FMatrix CalcProjectionMatrix(uint32 ViewSizeX, uint32 ViewSizeY, float fFOV, float NearPlane);
 
 	/**
 	* Get the current top-most transform entry without the canvas projection
 	* @return matrix from transform stack. 
 	*/
-	FMatrix GetTransform() const 
+	ENGINE_API FMatrix GetTransform() const
 	{ 
 		return TransformStack.Top().GetMatrix() * TransformStack[0].GetMatrix().Inverse(); 
 	}
@@ -156,7 +156,7 @@ public:
 	* Get the bottom-most element of the transform stack. 
 	* @return matrix from transform stack. 
 	*/
-	const FMatrix& GetBottomTransform() const 
+	ENGINE_API const FMatrix& GetBottomTransform() const
 	{ 
 		return TransformStack[0].GetMatrix(); 
 	}
@@ -165,7 +165,7 @@ public:
 	* Get the current top-most transform entry 
 	* @return matrix from transform stack. 
 	*/
-	const FMatrix& GetFullTransform() const 
+	ENGINE_API const FMatrix& GetFullTransform() const
 	{ 
 		return TransformStack.Top().GetMatrix(); 
 	}
@@ -175,7 +175,7 @@ public:
 	 *
 	 * @param Copy	canvas to copy from
 	 **/
-	void CopyTransformStack(const FCanvas& Copy);
+	ENGINE_API void CopyTransformStack(const FCanvas& Copy);
 
 	/** 
 	 * Set the current masked region on the canvas
@@ -187,13 +187,13 @@ public:
 	 * @param SizeX - x size in canvas coords
 	 * @param SizeY - y size in canvas coords
 	 */
-	void PushMaskRegion( float X, float Y, float SizeX, float SizeY );
+	ENGINE_API void PushMaskRegion(float X, float Y, float SizeX, float SizeY);
 
 	/**
 	 * Remove the current masking region; if other masking regions were previously pushed onto the stack,
 	 * the next one down will be activated.
 	 */
-	void PopMaskRegion();
+	ENGINE_API void PopMaskRegion();
 
 	/**
 	 * Sets the render target which will be used for subsequent canvas primitives.
@@ -203,7 +203,7 @@ public:
 	/**
 	* Get the current render target for the canvas
 	*/	
-	FORCEINLINE FRenderTarget* GetRenderTarget() const 
+	ENGINE_API FORCEINLINE FRenderTarget* GetRenderTarget() const
 	{ 
 		return RenderTarget; 
 	}
@@ -458,7 +458,7 @@ public:
 	/**
 	 * Get the top-most canvas masking region from the stack.
 	 */
-	FMaskRegion GetCurrentMaskRegion() const;
+	ENGINE_API FMaskRegion GetCurrentMaskRegion() const;
 
 	/**
 	* Clear the canvas
