Fix scaling code where UpdateScale wasn't actually passing info to physx

diff --git a/Engine/Source/Runtime/Engine/Private/PhysicsEngine/BodyInstance.cpp b/Engine/Source/Runtime/Engine/Private/PhysicsEngine/BodyInstance.cpp
--- a/Engine/Source/Runtime/Engine/Private/PhysicsEngine/BodyInstance.cpp
+++ b/Engine/Source/Runtime/Engine/Private/PhysicsEngine/BodyInstance.cpp
@@ -1431,17 +1431,17 @@ bool FBodyInstance::UpdateBodyScale(const FVector & InScale3D)
 					   SCENE_UNLOCK_READ(PScene);
 					   UE_LOG(LogPhysics, Error, TEXT("Unknown geom type."));
 			}
+		}// end switch
 
-			if (UpdatedGeometry)
-			{
-				SCOPED_SCENE_WRITE_LOCK(PScene);
-				PShape->setLocalPose(PLocalPose);
-				PShape->setGeometry(*UpdatedGeometry);
-			}
-			else
-			{
-				FMessageLog("PIE").Warning(FText::Format(LOCTEXT("PhysicsInvalidScale", "Scale ''{0}'' is not valid on object '{1}'."), FText::FromString(InScale3DAdjusted.ToString()), FText::FromString(GetBodyDebugName())));
-			}
+		if (UpdatedGeometry)
+		{
+			SCOPED_SCENE_WRITE_LOCK(PScene);
+			PShape->setLocalPose(PLocalPose);
+			PShape->setGeometry(*UpdatedGeometry);
+		}
+		else
+		{
+			FMessageLog("PIE").Warning(FText::Format(LOCTEXT("PhysicsInvalidScale", "Scale ''{0}'' is not valid on object '{1}'."), FText::FromString(InScale3DAdjusted.ToString()), FText::FromString(GetBodyDebugName())));
 		}
 	}
 
