UE4: Removed unused function UEngine::LoadAdditionalMaps

diff --git a/Engine/Source/Runtime/Engine/Classes/Engine/Engine.h b/Engine/Source/Runtime/Engine/Classes/Engine/Engine.h
--- a/Engine/Source/Runtime/Engine/Classes/Engine/Engine.h
+++ b/Engine/Source/Runtime/Engine/Classes/Engine/Engine.h
@@ -2360,8 +2360,6 @@ public:
 
 	void ClearDebugDisplayProperties();
 
-	bool LoadAdditionalMaps(UWorld* PersistentWorld, const TArray<FString>& MapsToLoad, FString& Error);
-
 	/**
 	 * Loads the PerMapPackages for the given map, and adds them to the RootSet
 	 *
diff --git a/Engine/Source/Runtime/Engine/Private/UnrealEngine.cpp b/Engine/Source/Runtime/Engine/Private/UnrealEngine.cpp
--- a/Engine/Source/Runtime/Engine/Private/UnrealEngine.cpp
+++ b/Engine/Source/Runtime/Engine/Private/UnrealEngine.cpp
@@ -8518,58 +8518,6 @@ void UEngine::ClearDebugDisplayProperties()
 	}
 }
 
-bool UEngine::LoadAdditionalMaps(UWorld* PersistentWorld, const TArray<FString>& MapsToLoad, FString& Error)
-{
-	UPackage* OuterPackage = NULL;
-	
-	for (int32 MapIdx = 0; MapIdx < MapsToLoad.Num(); ++MapIdx)
-	{
-		FString MapName = MapsToLoad[MapIdx];
-
-		// convert MapName to long package name
-		if (FPackageName::IsShortPackageName(MapName))
-		{
-			if (!FPackageName::SearchForPackageOnDisk(MapName, &MapName))
-			{
-				Error = FString::Printf(TEXT("Failed to find package '%s'"), *MapName);
-				return false;
-			}
-		}
-		
-		// Load package 
-		UPackage* MapPackage = LoadPackage(OuterPackage, *MapName, LOAD_None);
-		if (MapPackage == NULL)
-		{
-			Error = FString::Printf(TEXT("Failed to load package '%s'"), *MapName);
-			return false;
-		}
-		
-		//  
-		UWorld* World = UWorld::FindWorldInPackage(MapPackage);
-		if (World == NULL)
-		{
-			Error = FString::Printf(TEXT("Not a map package '%s'"), *MapName);
-			return false;
-		}
-
-		// Add loaded map to Runtime level
-		FName MapFName = FName(*MapName);
-		PersistentWorld->AddToWorld(World->PersistentLevel);
-						
-		ULevelStreaming* StreamingLevel = static_cast<ULevelStreaming*>(StaticConstructObject(ULevelStreamingKismet::StaticClass(), PersistentWorld, NAME_None, RF_NoFlags, NULL));
-		// Associate a package name.
-		StreamingLevel->PackageName			= MapFName;
-		StreamingLevel->PackageNameToLoad	= MapFName;
-		StreamingLevel->bShouldBeLoaded		= true;
-		StreamingLevel->bShouldBeVisible	= true;
-		StreamingLevel->bShouldBlockOnLoad	= true;
-		PersistentWorld->StreamingLevels.Add(StreamingLevel);
-	}
-
-	return true;
-}
-
-
 void UEngine::MovePendingLevel(FWorldContext &Context)
 {
 	check(Context.World());
