UE4: Added code to follow world redirectors where appropriate.

diff --git a/Engine/Source/Runtime/Engine/Classes/Engine/World.h b/Engine/Source/Runtime/Engine/Classes/Engine/World.h
--- a/Engine/Source/Runtime/Engine/Classes/Engine/World.h
+++ b/Engine/Source/Runtime/Engine/Classes/Engine/World.h
@@ -2416,6 +2416,9 @@ public:
 	static UWorld* DuplicateWorldForPIE(const FString& PackageName, UWorld* OwningWorld);
 	static FString RemovePIEPrefix(const FString &Source);
 	static UWorld* FindWorldInPackage(UPackage* Package);
+
+	/** If the specified package contains a redirector to a UWorld, that UWorld is returned. Otherwise, nullptr is returned. */
+	static UWorld* FollowWorldRedirectorInPackage(UPackage* Package, UObjectRedirector** OptionalOutRedirector = nullptr);
 };
 
 /** Global UWorld pointer */
diff --git a/Engine/Source/Runtime/Engine/Private/LevelActor.cpp b/Engine/Source/Runtime/Engine/Private/LevelActor.cpp
--- a/Engine/Source/Runtime/Engine/Private/LevelActor.cpp
+++ b/Engine/Source/Runtime/Engine/Private/LevelActor.cpp
@@ -797,7 +797,12 @@ void UWorld::LoadSecondaryLevels(bool bForce, TSet<FString>* CookedPackages)
 							bLoadedLevelPackage = true;
 
 							// Find the world object in the loaded package.
-							UWorld* const LoadedWorld	= UWorld::FindWorldInPackage(LevelPackage);
+							UWorld* LoadedWorld	= UWorld::FindWorldInPackage(LevelPackage);
+							// If the world was not found, it could be a redirector to a world. If so, follow it to the destination world.
+							if (!LoadedWorld)
+							{
+								LoadedWorld = UWorld::FollowWorldRedirectorInPackage(LevelPackage);
+							}
 							check(LoadedWorld);
 
 							// LoadedWorld won't be serialized as there's a BeginLoad on the stack so we manually serialize it here.
diff --git a/Engine/Source/Runtime/Engine/Private/LevelStreaming.cpp b/Engine/Source/Runtime/Engine/Private/LevelStreaming.cpp
--- a/Engine/Source/Runtime/Engine/Private/LevelStreaming.cpp
+++ b/Engine/Source/Runtime/Engine/Private/LevelStreaming.cpp
@@ -325,6 +325,17 @@ bool ULevelStreaming::RequestLevel(UWorld* PersistentWorld, bool bAllowLevelLoad
 	{
 		// Find world object and use its PersistentLevel pointer.
 		UWorld* World = UWorld::FindWorldInPackage(LevelPackage);
+
+		// Check for a redirector. Follow it, if found.
+		if ( !World )
+		{
+			World = UWorld::FollowWorldRedirectorInPackage(LevelPackage);
+			if ( World )
+			{
+				LevelPackage = World->GetOutermost();
+			}
+		}
+
 		if (World != NULL)
 		{
 #if !(UE_BUILD_SHIPPING || UE_BUILD_TEST)
@@ -504,80 +515,75 @@ void ULevelStreaming::AsyncLevelLoadComplete( const FString& InPackageName, UPac
 		else
 		{
 			// There could have been a redirector in the package. Attempt to follow it.
-			const FString ExpectedRedirectorName = (PackageNameToLoad != NAME_None) ? PackageNameToLoad.ToString() : PackageName.ToString();
-			UObjectRedirector* WorldRedirector = FindObject<UObjectRedirector>(LevelPackage, *FPackageName::GetLongPackageAssetName(*ExpectedRedirectorName));
-			if (WorldRedirector)
+			UObjectRedirector* WorldRedirector = nullptr;
+			UWorld* DestinationWorld = UWorld::FollowWorldRedirectorInPackage(LevelPackage, &WorldRedirector);
+			if (DestinationWorld)
 			{
-				// A redirector was found! Make sure the destination world was loaded.
-				UWorld* DestinationWorld = Cast<UWorld>(WorldRedirector->DestinationObject);
-				if (DestinationWorld)
+				// To follow the world redirector for level streaming...
+				// 1) Update all globals that refer to the redirector package by name
+				// 2) Update the PackageNameToLoad to refer to the new package location
+				// 3) If the package name to load was the same as the destination package name...
+				//         ... update the package name to the new package and let the next RequestLevel try this process again.
+				//    If the package name to load was different...
+				//         ... it means the specified package name was explicit and we will just load from another file.
+
+				FName OldDesiredPackageName = FName(*InPackageName);
+				UWorld** OwningWorldPtr = ULevel::StreamedLevelsOwningWorld.Find(OldDesiredPackageName);
+				UWorld* OwningWorld = OwningWorldPtr ? *OwningWorldPtr : NULL;
+				ULevel::StreamedLevelsOwningWorld.Remove(OldDesiredPackageName);
+
+				// Try again with the destination package to load.
+				// IMPORTANT: check this BEFORE changing PackageNameToLoad, otherwise you wont know if the package name was supposed to be different.
+				const bool bLoadingIntoDifferentPackage = (PackageName != PackageNameToLoad) && (PackageNameToLoad != NAME_None);
+
+				// ... now set PackageNameToLoad
+				PackageNameToLoad = DestinationWorld->GetOutermost()->GetFName();
+
+				if ( PackageNameToLoad != OldDesiredPackageName )
 				{
-					// To follow the world redirector for level streaming...
-					// 1) Update all globals that refer to the redirector package by name
-					// 2) Update the PackageNameToLoad to refer to the new package location
-					// 3) If the package name to load was the same as the destination package name...
-					//         ... update the package name to the new package and let the next RequestLevel try this process again.
-					//    If the package name to load was different...
-					//         ... it means the specified package name was explicit and we will just load from another file.
-
-					FName OldDesiredPackageName = FName(*InPackageName);
-					UWorld** OwningWorldPtr = ULevel::StreamedLevelsOwningWorld.Find(OldDesiredPackageName);
-					UWorld* OwningWorld = OwningWorldPtr ? *OwningWorldPtr : NULL;
-					ULevel::StreamedLevelsOwningWorld.Remove(OldDesiredPackageName);
-
-					// Try again with the destination package to load.
-					// IMPORTANT: check this BEFORE changing PackageNameToLoad, otherwise you wont know if the package name was supposed to be different.
-					const bool bLoadingIntoDifferentPackage = (PackageName != PackageNameToLoad) && (PackageNameToLoad != NAME_None);
-
-					// ... now set PackageNameToLoad
-					PackageNameToLoad = DestinationWorld->GetOutermost()->GetFName();
-
-					if ( PackageNameToLoad != OldDesiredPackageName )
+					EWorldType::Type* OldPackageWorldType = UWorld::WorldTypePreLoadMap.Find(OldDesiredPackageName);
+					if ( OldPackageWorldType )
 					{
-						EWorldType::Type* OldPackageWorldType = UWorld::WorldTypePreLoadMap.Find(OldDesiredPackageName);
-						if ( OldPackageWorldType )
-						{
-							UWorld::WorldTypePreLoadMap.FindOrAdd(PackageNameToLoad) = *OldPackageWorldType;
-							UWorld::WorldTypePreLoadMap.Remove(OldDesiredPackageName);
-						}
+						UWorld::WorldTypePreLoadMap.FindOrAdd(PackageNameToLoad) = *OldPackageWorldType;
+						UWorld::WorldTypePreLoadMap.Remove(OldDesiredPackageName);
 					}
+				}
 
-					// Now determine if we are loading into the package explicitly or if it is okay to just load the other package.
-					if ( bLoadingIntoDifferentPackage )
+				// Now determine if we are loading into the package explicitly or if it is okay to just load the other package.
+				if ( bLoadingIntoDifferentPackage )
+				{
+					// Loading into a new custom package explicitly. Load the destination world directly into the package.
+					// Detach the linker to load from a new file into the same package.
+					ULinkerLoad* PackageLinker = ULinkerLoad::FindExistingLinkerForPackage(LevelPackage);
+					if (PackageLinker)
 					{
-						// Loading into a new custom package explicitly. Load the destination world directly into the package.
-						// Detach the linker to load from a new file into the same package.
-						ULinkerLoad* PackageLinker = ULinkerLoad::FindExistingLinkerForPackage(LevelPackage);
-						if (PackageLinker)
-						{
-							PackageLinker->Detach(false);
-						}
+						PackageLinker->Detach(false);
+					}
 
-						// Make sure the redirector is not in the way of the new world.
-						if (WorldRedirector->GetFName() == DestinationWorld->GetFName())
-						{
-							// Pass NULL as the name to make a new unique name and NULL for the outer to not change the outer.
-							WorldRedirector->Rename(NULL, NULL, REN_DoNotDirty | REN_DontCreateRedirectors | REN_ForceNoResetLoaders | REN_NonTransactional);
-						}
+					// Make sure the redirector is not in the way of the new world.
+					if (WorldRedirector->GetFName() == DestinationWorld->GetFName())
+					{
+						// Pass NULL as the name to make a new unique name and NULL for the outer to not change the outer.
+						WorldRedirector->Rename(NULL, NULL, REN_DoNotDirty | REN_DontCreateRedirectors | REN_ForceNoResetLoaders | REN_NonTransactional);
 					}
-					else
+				}
+				else
+				{
+					// Loading the requested package normally. Fix up the destination world then update the requested package to the destination.
+					if (OwningWorld)
 					{
-						// Loading the requested package normally. Fix up the destination world then update the requested package to the destination.
-						if (OwningWorld)
+						if (DestinationWorld->PersistentLevel)
 						{
-							if (DestinationWorld->PersistentLevel)
-							{
-								DestinationWorld->PersistentLevel->OwningWorld = OwningWorld;
-							}
-
-							// In some cases, BSP render data is not created because the OwningWorld was not set correctly.
-							// Regenerate that render data here
-							DestinationWorld->PersistentLevel->InvalidateModelSurface();
-							DestinationWorld->PersistentLevel->CommitModelSurfaces();
+							DestinationWorld->PersistentLevel->OwningWorld = OwningWorld;
 						}
-						
-						PackageName = PackageNameToLoad;
+
+						// In some cases, BSP render data is not created because the OwningWorld was not set correctly.
+						// Regenerate that render data here
+						DestinationWorld->PersistentLevel->InvalidateModelSurface();
+						DestinationWorld->PersistentLevel->CommitModelSurfaces();
 					}
+						
+					PackageName = PackageNameToLoad;
 				}
 			}
 		}
diff --git a/Engine/Source/Runtime/Engine/Private/UnrealEngine.cpp b/Engine/Source/Runtime/Engine/Private/UnrealEngine.cpp
--- a/Engine/Source/Runtime/Engine/Private/UnrealEngine.cpp
+++ b/Engine/Source/Runtime/Engine/Private/UnrealEngine.cpp
@@ -8243,6 +8243,17 @@ bool UEngine::LoadMap( FWorldContext& WorldContext, FURL URL, class UPendingNetG
 				}
 
 				NewWorld = UWorld::FindWorldInPackage(WorldPackage);
+
+				// If the world was not found, follow a redirector if there is one.
+				if ( !NewWorld )
+				{
+					NewWorld = UWorld::FollowWorldRedirectorInPackage(WorldPackage);
+					if ( NewWorld )
+					{
+						WorldPackage = NewWorld->GetOutermost();
+					}
+				}
+
 				check(NewWorld);
 #if WITH_EDITOR
 				WorldPackage->PIEInstanceID = WorldContext.PIEInstance;
@@ -8286,18 +8297,28 @@ bool UEngine::LoadMap( FWorldContext& WorldContext, FURL URL, class UPendingNetG
 			Error = FString::Printf(TEXT("Failed to load package '%s'"), *URL.Map);
 			return false;
 		}
-		
-		FScopeCycleCounterUObject MapScope(WorldPackage);
-
-		if( FPlatformProperties::RequiresCookedData() && GUseSeekFreeLoading && !(WorldPackage->PackageFlags & PKG_DisallowLazyLoading) )
-		{
-			UE_LOG(LogLoad, Fatal,TEXT("Map '%s' has not been cooked correctly! Most likely stale version on the XDK."),*WorldPackage->GetName());
-		}
 
 		// Find the newly loaded world.
 		NewWorld = UWorld::FindWorldInPackage(WorldPackage);
+
+		// If the world was not found, it could be a redirector to a world. If so, follow it to the destination world.
+		if ( !NewWorld )
+		{
+			NewWorld = UWorld::FollowWorldRedirectorInPackage(WorldPackage);
+			if ( NewWorld )
+			{
+				WorldPackage = NewWorld->GetOutermost();
+			}
+		}
 		check(NewWorld);
 
+		FScopeCycleCounterUObject MapScope(WorldPackage);
+
+		if (FPlatformProperties::RequiresCookedData() && GUseSeekFreeLoading && !(WorldPackage->PackageFlags & PKG_DisallowLazyLoading))
+		{
+			UE_LOG(LogLoad, Fatal, TEXT("Map '%s' has not been cooked correctly! Most likely stale version on the XDK."), *WorldPackage->GetName());
+		}
+
 		if (WorldContext.WorldType == EWorldType::PIE)
 		{
 			// If we are a PIE world and the world we just found is already initialized, then we're probably reloading the editor world and we
@@ -8840,6 +8861,17 @@ static void AsyncMapChangeLevelLoadCompletionCallback(const FString& PackageName
 	{	
 		// Try to find a UWorld object in the level package.
 		UWorld* World = UWorld::FindWorldInPackage(LevelPackage);
+
+		// If the world was not found, try to follow a redirector if it exists
+		if ( !World )
+		{
+			World = UWorld::FollowWorldRedirectorInPackage(LevelPackage);
+			if ( World )
+			{
+				LevelPackage = World->GetOutermost();
+			}
+		}
+
 		ULevel* Level = World ? World->PersistentLevel : NULL;	
 		
 		// Print out a warning and set the error if we couldn't find a level in this package.
diff --git a/Engine/Source/Runtime/Engine/Private/World.cpp b/Engine/Source/Runtime/Engine/Private/World.cpp
--- a/Engine/Source/Runtime/Engine/Private/World.cpp
+++ b/Engine/Source/Runtime/Engine/Private/World.cpp
@@ -1925,6 +1925,17 @@ UWorld* UWorld::DuplicateWorldForPIE(const FString& PackageName, UWorld* OwningW
 
 	// Find world object and use its PersistentLevel pointer.
 	UWorld* EditorLevelWorld = UWorld::FindWorldInPackage(EditorLevelPackage);
+
+	// If the world was not found, try to follow a redirector, if there is one
+	if ( !EditorLevelWorld )
+	{
+		EditorLevelWorld = UWorld::FollowWorldRedirectorInPackage(EditorLevelPackage);
+		if ( EditorLevelPackage )
+		{
+			EditorLevelPackage = EditorLevelWorld->GetOutermost();
+		}
+	}
+
 	if( !EditorLevelWorld )
 		return NULL;
 
@@ -3918,7 +3929,19 @@ void FSeamlessTravelHandler::SeamlessTravelLoadCallback(const FString& PackageNa
 	// defer until tick when it's safe to perform the transition
 	if (IsInTransition())
 	{
-		SetHandlerLoadedData(LevelPackage, UWorld::FindWorldInPackage(LevelPackage));
+		UWorld* World = UWorld::FindWorldInPackage(LevelPackage);
+
+		// If the world could not be found, follow a redirector if there is one.
+		if ( !World )
+		{
+			World = UWorld::FollowWorldRedirectorInPackage(LevelPackage);
+			if ( World )
+			{
+				LevelPackage = World->GetOutermost();
+			}
+		}
+
+		SetHandlerLoadedData(LevelPackage, World);
 	}
 }
 
@@ -4732,6 +4755,31 @@ UWorld* UWorld::FindWorldInPackage(UPackage* Package)
 	return RetVal;
 }
 
+UWorld* UWorld::FollowWorldRedirectorInPackage(UPackage* Package, UObjectRedirector** OptionalOutRedirector)
+{
+	UWorld* RetVal = nullptr;
+	TArray<UObject*> PotentialRedirectors;
+	GetObjectsWithOuter(Package, PotentialRedirectors, false);
+	for (auto ObjIt = PotentialRedirectors.CreateConstIterator(); ObjIt; ++ObjIt)
+	{
+		UObjectRedirector* Redirector = Cast<UObjectRedirector>(*ObjIt);
+		if (Redirector)
+		{
+			RetVal = Cast<UWorld>(Redirector->DestinationObject);
+			if ( RetVal )
+			{
+				if ( OptionalOutRedirector )
+				{
+					(*OptionalOutRedirector) = Redirector;
+				}
+				break;
+			}
+		}
+	}
+
+	return RetVal;
+}
+
 #if WITH_EDITOR
 
 
