Added searching of streaming levels by long package name, to support streaming levels with the same short name

diff --git a/Engine/Source/Runtime/Engine/Private/LevelStreaming.cpp b/Engine/Source/Runtime/Engine/Private/LevelStreaming.cpp
--- a/Engine/Source/Runtime/Engine/Private/LevelStreaming.cpp
+++ b/Engine/Source/Runtime/Engine/Private/LevelStreaming.cpp
@@ -46,17 +46,16 @@ ULevelStreaming* FStreamLevelAction::FindAndCacheLevelStreamingObject( const FNa
 	// Search for the level object by name.
 	if( LevelName != NAME_None )
 	{
-		const FString SearchShortName = MakeSafeShortLevelName( LevelName, InWorld );
+		const FString SearchPackageName = MakeSafeLevelName( LevelName, InWorld );
 
-		// Iterate over all streaming level objects in world info to find one with matching name.
-		for( int32 LevelIndex=0; LevelIndex < InWorld->StreamingLevels.Num(); LevelIndex++ )
+		for (ULevelStreaming* LevelStreaming : InWorld->StreamingLevels)
 		{
-			ULevelStreaming* LevelStreamingObject = InWorld->StreamingLevels[LevelIndex];
-			if( LevelStreamingObject
-				&& FPackageName::GetShortName(LevelStreamingObject->PackageName) == SearchShortName)
+			// We check only suffix of package name, to handle situations when packages were saved for play into a temporary folder
+			// Like Saved/Autosaves/PackageName
+			if (LevelStreaming && 
+				LevelStreaming->PackageName.ToString().EndsWith(SearchPackageName, ESearchCase::IgnoreCase))
 			{
-				// If we have an exact match go ahead and return it
-				return LevelStreamingObject;
+				return LevelStreaming;
 			}
 		}
 	}
@@ -65,15 +64,26 @@ ULevelStreaming* FStreamLevelAction::FindAndCacheLevelStreamingObject( const FNa
 }
 
 /**
- * Given a level name, returns a shoftlevel name that will work with Play on Editor or Play on Console
+ * Given a level name, returns a level name that will work with Play on Editor or Play on Console
  *
  * @param	InLevelName		Raw level name (no UEDPIE or UED<console> prefix)
  * @param	InWorld			World in which to check for other instances of the name
  */
-FString FStreamLevelAction::MakeSafeShortLevelName( const FName& InLevelName, UWorld* InWorld )
+FString FStreamLevelAction::MakeSafeLevelName( const FName& InLevelName, UWorld* InWorld )
 {
 	// Special case for PIE, the PackageName gets mangled.
-	return InWorld->StreamingLevelsPrefix + FPackageName::GetShortName(InLevelName);
+	if (!InWorld->StreamingLevelsPrefix.IsEmpty())
+	{
+		FString PackageName = InWorld->StreamingLevelsPrefix + FPackageName::GetShortName(InLevelName);
+		if (!FPackageName::IsShortPackageName(InLevelName))
+		{
+			PackageName = FPackageName::GetLongPackagePath(InLevelName.ToString()) + TEXT("/") + PackageName;
+		}
+		
+		return PackageName;
+	}
+	
+	return InLevelName.ToString();
 }
 /**
 * Handles "Activated" for single ULevelStreaming object.
