Forgot some files for CL# 2082806

diff --git a/Engine/Source/Runtime/Engine/Classes/Engine/LevelStreaming.h b/Engine/Source/Runtime/Engine/Classes/Engine/LevelStreaming.h
--- a/Engine/Source/Runtime/Engine/Classes/Engine/LevelStreaming.h
+++ b/Engine/Source/Runtime/Engine/Classes/Engine/LevelStreaming.h
@@ -18,11 +18,11 @@ public:
 	FStreamLevelAction(bool bIsLoading, const FName & InLevelName, bool bIsMakeVisibleAfterLoad, bool bIsShouldBlockOnLoad, const FLatentActionInfo& InLatentInfo, UWorld* World);
 
 	/**
-	 * Given a level name, returns a short level name that will work with Play on Editor or Play on Console
+	 * Given a level name, returns level name that will work with Play on Editor or Play on Console
 	 *
 	 * @param	InLevelName		Raw level name (no UEDPIE or UED<console> prefix)
 	 */
-	static FString MakeSafeShortLevelName( const FName& InLevelName, UWorld* InWorld );
+	static FString MakeSafeLevelName( const FName& InLevelName, UWorld* InWorld );
 
 	/**
 	 * Helper function to potentially find a level streaming object by name and cache the result
diff --git a/Engine/Source/Runtime/Engine/Private/GameplayStatics.cpp b/Engine/Source/Runtime/Engine/Private/GameplayStatics.cpp
--- a/Engine/Source/Runtime/Engine/Private/GameplayStatics.cpp
+++ b/Engine/Source/Runtime/Engine/Private/GameplayStatics.cpp
@@ -320,15 +320,20 @@ void UGameplayStatics::UnloadStreamLevel(UObject* WorldContextObject, FName Leve
 
 ULevelStreaming* UGameplayStatics::GetStreamingLevel(UObject* WorldContextObject, FName InPackageName)
 {
-	UWorld* World = GEngine->GetWorldFromContextObject(WorldContextObject);
-	
-	FString ShortPackageName = FStreamLevelAction::MakeSafeShortLevelName(InPackageName, World);
-
-	for (auto It = World->StreamingLevels.CreateConstIterator(); It; ++It)
+	if (InPackageName != NAME_None)
 	{
-		if (FPackageName::GetShortName((*It)->PackageName) == ShortPackageName)
+		UWorld* World = GEngine->GetWorldFromContextObject(WorldContextObject);
+		const FString SearchPackageName = FStreamLevelAction::MakeSafeLevelName(InPackageName, World);
+
+		for (ULevelStreaming* LevelStreaming : World->StreamingLevels)
 		{
-			return *It;
+			// We check only suffix of package name, to handle situations when packages were saved for play into a temporary folder
+			// Like Saved/Autosaves/PackageName
+			if (LevelStreaming && 
+				LevelStreaming->PackageName.ToString().EndsWith(SearchPackageName, ESearchCase::IgnoreCase))
+			{
+				return LevelStreaming;
+			}
 		}
 	}
 	
