Remove RHI.h, RenderingThread.h, UnrealClient.h, MaterialInstanceActor.h, MaterialInstanceBasePropertyOverrides.h from Engine.h
Also remove duplicate EngineBaseTypes.h
Delete AllEngineHeaders.h and Engine.GeneratedWrapper.cpp

diff --git a/Engine/Source/Editor/Matinee/Private/MatineeDraw.cpp b/Engine/Source/Editor/Matinee/Private/MatineeDraw.cpp
--- a/Engine/Source/Editor/Matinee/Private/MatineeDraw.cpp
+++ b/Engine/Source/Editor/Matinee/Private/MatineeDraw.cpp
@@ -19,6 +19,7 @@
 #include "Matinee/InterpTrackLinearColorBase.h"
 #include "Matinee/InterpTrackEvent.h"
 #include "Matinee/InterpTrackMove.h"
+#include "Materials/MaterialInstanceActor.h"
 
 static const int32 GroupHeadHeight = 24;
 static const int32 TrackHeight = 24;
diff --git a/Engine/Source/Editor/Sequencer/Private/SSection.cpp b/Engine/Source/Editor/Sequencer/Private/SSection.cpp
--- a/Engine/Source/Editor/Sequencer/Private/SSection.cpp
+++ b/Engine/Source/Editor/Sequencer/Private/SSection.cpp
@@ -1,6 +1,7 @@
 // Copyright 1998-2014 Epic Games, Inc. All Rights Reserved.
 
 #include "SequencerPrivatePCH.h"
+#include "Sequencer.h"
 #include "SSection.h"
 #include "IKeyArea.h"
 #include "ISequencerSection.h"
diff --git a/Engine/Source/Editor/UnrealEd/Classes/MaterialEditor/MaterialEditorInstanceConstant.h b/Engine/Source/Editor/UnrealEd/Classes/MaterialEditor/MaterialEditorInstanceConstant.h
--- a/Engine/Source/Editor/UnrealEd/Classes/MaterialEditor/MaterialEditorInstanceConstant.h
+++ b/Engine/Source/Editor/UnrealEd/Classes/MaterialEditor/MaterialEditorInstanceConstant.h
@@ -5,6 +5,7 @@
  */
 
 #pragma once
+#include "Materials/MaterialInstanceBasePropertyOverrides.h"
 #include "MaterialEditorInstanceConstant.generated.h"
 
 USTRUCT()
diff --git a/Engine/Source/Runtime/Engine/Classes/Engine/EngineBaseTypes.h b/Engine/Source/Runtime/Engine/Classes/Engine/EngineBaseTypes.h
--- a/Engine/Source/Runtime/Engine/Classes/Engine/EngineBaseTypes.h
+++ b/Engine/Source/Runtime/Engine/Classes/Engine/EngineBaseTypes.h
@@ -7,6 +7,7 @@
  *  The typical use case is for structs used in the renderer and also in script code.
  */
 
+#include "TaskGraphInterfaces.h"
 #include "EngineBaseTypes.generated.h"
 
 //
diff --git a/Engine/Source/Runtime/Engine/Classes/Engine/GameViewportClient.h b/Engine/Source/Runtime/Engine/Classes/Engine/GameViewportClient.h
--- a/Engine/Source/Runtime/Engine/Classes/Engine/GameViewportClient.h
+++ b/Engine/Source/Runtime/Engine/Classes/Engine/GameViewportClient.h
@@ -16,7 +16,7 @@
 #pragma once
 
 #include "ShowFlags.h"
-
+#include "Engine/ScriptViewportClient.h"
 #include "GameViewportClient.generated.h"
 
 /**
diff --git a/Engine/Source/Runtime/Engine/Classes/Engine/ScriptViewportClient.h b/Engine/Source/Runtime/Engine/Classes/Engine/ScriptViewportClient.h
--- a/Engine/Source/Runtime/Engine/Classes/Engine/ScriptViewportClient.h
+++ b/Engine/Source/Runtime/Engine/Classes/Engine/ScriptViewportClient.h
@@ -6,6 +6,7 @@
  */
 
 #pragma once
+#include "UnrealClient.h"
 #include "ScriptViewportClient.generated.h"
 
 UCLASS(transient)
diff --git a/Engine/Source/Runtime/Engine/Classes/GameFramework/Actor.h b/Engine/Source/Runtime/Engine/Classes/GameFramework/Actor.h
--- a/Engine/Source/Runtime/Engine/Classes/GameFramework/Actor.h
+++ b/Engine/Source/Runtime/Engine/Classes/GameFramework/Actor.h
@@ -4,6 +4,7 @@
 
 #include "Runtime/Engine/Classes/Engine/LatentActionManager.h"
 #include "Runtime/Engine/Classes/Camera/CameraTypes.h"
+#include "RenderingThread.h"
 
 UENUM(BlueprintType)
 namespace EEndPlayReason
diff --git a/Engine/Source/Runtime/Engine/Private/AllEngineHeaders.h b/Engine/Source/Runtime/Engine/Private/AllEngineHeaders.h
deleted file mode 100644
diff --git a/Engine/Source/Runtime/Engine/Private/Engine.GeneratedWrapper.cpp b/Engine/Source/Runtime/Engine/Private/Engine.GeneratedWrapper.cpp
deleted file mode 100644
diff --git a/Engine/Source/Runtime/Engine/Private/Interpolation.cpp b/Engine/Source/Runtime/Engine/Private/Interpolation.cpp
--- a/Engine/Source/Runtime/Engine/Private/Interpolation.cpp
+++ b/Engine/Source/Runtime/Engine/Private/Interpolation.cpp
@@ -58,6 +58,7 @@
 #include "Matinee/InterpGroupInstCamera.h"
 #include "Matinee/InterpFilter_Classes.h"
 #include "Matinee/InterpFilter_Custom.h"
+#include "Materials/MaterialInstanceActor.h"
 #include "ParticleDefinitions.h"
 #include "SoundDefinitions.h"
 #include "InterpolationHitProxy.h"
diff --git a/Engine/Source/Runtime/Engine/Private/Materials/MaterialInstance.cpp b/Engine/Source/Runtime/Engine/Private/Materials/MaterialInstance.cpp
--- a/Engine/Source/Runtime/Engine/Private/Materials/MaterialInstance.cpp
+++ b/Engine/Source/Runtime/Engine/Private/Materials/MaterialInstance.cpp
@@ -8,6 +8,7 @@
 #include "Materials/MaterialExpressionStaticBoolParameter.h"
 #include "Materials/MaterialExpressionStaticComponentMaskParameter.h"
 #include "Materials/MaterialExpressionLandscapeLayerWeight.h"
+#include "Materials/MaterialInstanceBasePropertyOverrides.h"
 #include "MaterialInstanceSupport.h"
 #include "MaterialShader.h"
 #include "TargetPlatform.h"
diff --git a/Engine/Source/Runtime/Engine/Private/Materials/MaterialShared.cpp b/Engine/Source/Runtime/Engine/Private/Materials/MaterialShared.cpp
--- a/Engine/Source/Runtime/Engine/Private/Materials/MaterialShared.cpp
+++ b/Engine/Source/Runtime/Engine/Private/Materials/MaterialShared.cpp
@@ -6,6 +6,7 @@
 
 #include "EnginePrivate.h"
 #include "Materials/MaterialExpressionBreakMaterialAttributes.h"
+#include "Materials/MaterialInstanceBasePropertyOverrides.h"
 #include "PixelFormat.h"
 #include "ShaderCompiler.h"
 #include "MaterialCompiler.h"
diff --git a/Engine/Source/Runtime/Engine/Private/UnrealEngine.cpp b/Engine/Source/Runtime/Engine/Private/UnrealEngine.cpp
--- a/Engine/Source/Runtime/Engine/Private/UnrealEngine.cpp
+++ b/Engine/Source/Runtime/Engine/Private/UnrealEngine.cpp
@@ -5,7 +5,6 @@
 =============================================================================*/
 
 #include "EnginePrivate.h"
-#include "AllEngineHeaders.h"
 
 #include "Net/UnrealNetwork.h"
 #include "FileManagerGeneric.h"
@@ -37,6 +36,8 @@
 #include "IHeadMountedDisplay.h"
 #include "Scalability.h"
 #include "StatsData.h"
+#include "AudioDevice.h"
+#include "ActiveSound.h"
 
 #if !UE_BUILD_SHIPPING
 #include "STaskGraph.h"
diff --git a/Engine/Source/Runtime/Engine/Public/Engine.h b/Engine/Source/Runtime/Engine/Public/Engine.h
--- a/Engine/Source/Runtime/Engine/Public/Engine.h
+++ b/Engine/Source/Runtime/Engine/Public/Engine.h
@@ -10,7 +10,6 @@
 
 #include "Core.h"
 #include "CoreUObject.h"
-#include "RHI.h"
 #include "InputCore.h"
 #include "EngineSettings.h"
 #include "EngineStats.h"
@@ -23,15 +22,12 @@
 -----------------------------------------------------------------------------*/
 
 #include "Tickable.h"						// FTickableGameObject interface.
-#include "RenderingThread.h"				// for FRenderCommandFence
 #include "Engine/EngineBaseTypes.h"
-#include "UnrealClient.h"					// for FViewportClient
 
 // EngineClasses.h
 #include "Engine/LatentActionManager.h"
 #include "Engine/Scene.h"
 #include "Camera/CameraTypes.h"
-#include "Engine/EngineBaseTypes.h"
 #include "Engine/NetSerialization.h"
 #include "Engine/EngineTypes.h"
 #include "Interfaces/Interface_AssetUserData.h"
@@ -117,7 +113,6 @@
 #include "Engine/PointLight.h"
 #include "Engine/SpotLight.h"
 #include "Engine/GeneratedMeshAreaLight.h"
-#include "Materials/MaterialInstanceActor.h"
 #include "AI/Navigation/NavigationData.h"
 #include "AI/Navigation/NavigationGraph.h"
 #include "AI/Navigation/NavAreas/NavArea.h"
@@ -717,7 +712,6 @@
 #include "Animation/BoneControllers/AnimNode_HandIKRetargeting.h"
 #include "Animation/AnimNodeSpaceConversions.h"
 #include "Camera/CameraStackTypes.h"
-#include "Materials/MaterialInstanceBasePropertyOverrides.h"
 #include "Engine/StreamableManager.h"
 #include "Engine/TextureDefines.h"
 #include "Tests/TextPropertyTestObject.h"
diff --git a/Engine/Source/Runtime/Engine/Public/SceneManagement.h b/Engine/Source/Runtime/Engine/Public/SceneManagement.h
--- a/Engine/Source/Runtime/Engine/Public/SceneManagement.h
+++ b/Engine/Source/Runtime/Engine/Public/SceneManagement.h
@@ -27,6 +27,7 @@ class FVertexBuffer;
 class FVertexFactory;
 class FAtmosphericFogSceneInfo;
 class ISceneViewExtension;
+class HHitProxy;
 
 
 DECLARE_LOG_CATEGORY_EXTERN(LogBufferVisualization, Log, All);
diff --git a/Engine/Source/Runtime/Engine/Public/SceneUtils.h b/Engine/Source/Runtime/Engine/Public/SceneUtils.h
--- a/Engine/Source/Runtime/Engine/Public/SceneUtils.h
+++ b/Engine/Source/Runtime/Engine/Public/SceneUtils.h
@@ -9,6 +9,9 @@
 // Colors that are defined for a particular mesh type
 // Each event type will be displayed using the defined color
 #pragma once
+
+#include "DynamicRHI.h"
+
 #define DEC_LIGHT			FColor(255,0,0,255)
 #define DEC_SKEL_MESH		FColor(255,0,255,255)
 #define DEC_STATIC_MESH		FColor(0,128,255,255)
diff --git a/Engine/Source/Runtime/Engine/Public/Texture.h b/Engine/Source/Runtime/Engine/Public/Texture.h
--- a/Engine/Source/Runtime/Engine/Public/Texture.h
+++ b/Engine/Source/Runtime/Engine/Public/Texture.h
@@ -7,6 +7,7 @@
 =============================================================================*/
 
 #include "Engine/TextureDefines.h"
+#include "UnrealClient.h"
 
 class UTexture2D;
 
