Remove LatentActionManager.h, Scene.h, NetSerialization.h, Interface_AssetUserData.h and all SoundNode classes from Engine.h

diff --git a/Engine/Source/Developer/MovieSceneTools/Private/AudioTrackEditor.cpp b/Engine/Source/Developer/MovieSceneTools/Private/AudioTrackEditor.cpp
--- a/Engine/Source/Developer/MovieSceneTools/Private/AudioTrackEditor.cpp
+++ b/Engine/Source/Developer/MovieSceneTools/Private/AudioTrackEditor.cpp
@@ -18,6 +18,7 @@
 #include "MovieSceneAudioSection.h"
 #include "CommonMovieSceneTools.h"
 #include "SoundDefinitions.h"
+#include "Sound/SoundNodeWavePlayer.h"
 #include "Runtime/Engine/Public/Slate/SlateTextures.h"
 #include "ObjectTools.h"
 #include "Runtime/Engine/Public/Slate/SceneViewport.h"
diff --git a/Engine/Source/Editor/DetailCustomizations/Private/AmbientSoundDetails.cpp b/Engine/Source/Editor/DetailCustomizations/Private/AmbientSoundDetails.cpp
--- a/Engine/Source/Editor/DetailCustomizations/Private/AmbientSoundDetails.cpp
+++ b/Engine/Source/Editor/DetailCustomizations/Private/AmbientSoundDetails.cpp
@@ -4,6 +4,11 @@
 #include "AmbientSoundDetails.h"
 #include "Toolkits/AssetEditorManager.h"
 #include "SoundDefinitions.h"
+#include "Sound/SoundNodeLooping.h"
+#include "Sound/SoundNodeMixer.h"
+#include "Sound/SoundNodeDelay.h"
+#include "Sound/SoundNodeRandom.h"
+#include "Sound/SoundNodeWavePlayer.h"
 
 
 #define LOCTEXT_NAMESPACE "AmbientSoundDetails"
diff --git a/Engine/Source/Editor/GraphEditor/Private/ConnectionDrawingPolicy.cpp b/Engine/Source/Editor/GraphEditor/Private/ConnectionDrawingPolicy.cpp
--- a/Engine/Source/Editor/GraphEditor/Private/ConnectionDrawingPolicy.cpp
+++ b/Engine/Source/Editor/GraphEditor/Private/ConnectionDrawingPolicy.cpp
@@ -12,6 +12,7 @@
 #include "AnimStateEntryNode.h"
 #include "AnimationGraphSchema.h"
 #include "AnimGraphNode_Base.h"
+#include "Sound/SoundNode.h"
 
 DEFINE_LOG_CATEGORY_STATIC(LogConnectionDrawingPolicy, Log, All);
 
diff --git a/Engine/Source/Editor/SoundCueEditor/Private/SoundCueEditor.cpp b/Engine/Source/Editor/SoundCueEditor/Private/SoundCueEditor.cpp
--- a/Engine/Source/Editor/SoundCueEditor/Private/SoundCueEditor.cpp
+++ b/Engine/Source/Editor/SoundCueEditor/Private/SoundCueEditor.cpp
@@ -2,6 +2,7 @@
 
 #include "SoundCueEditorModule.h"
 #include "SoundDefinitions.h"
+#include "Sound/SoundNodeWavePlayer.h"
 #include "ScopedTransaction.h"
 #include "Toolkits/IToolkitHost.h"
 #include "SoundCueEditor.h"
diff --git a/Engine/Source/Editor/SoundCueEditor/Private/SoundCueEditorUtilities.cpp b/Engine/Source/Editor/SoundCueEditor/Private/SoundCueEditorUtilities.cpp
--- a/Engine/Source/Editor/SoundCueEditor/Private/SoundCueEditorUtilities.cpp
+++ b/Engine/Source/Editor/SoundCueEditor/Private/SoundCueEditorUtilities.cpp
@@ -5,6 +5,7 @@
 #include "SoundCueEditorUtilities.h"
 #include "Toolkits/ToolkitManager.h"
 #include "SoundDefinitions.h"
+#include "Sound/SoundNodeWavePlayer.h"
 #include "GraphEditor.h"
 
 bool FSoundCueEditorUtilities::CanPasteNodes(const class UEdGraph* Graph)
diff --git a/Engine/Source/Editor/UnrealEd/Private/Factories/EditorFactories.cpp b/Engine/Source/Editor/UnrealEd/Private/Factories/EditorFactories.cpp
--- a/Engine/Source/Editor/UnrealEd/Private/Factories/EditorFactories.cpp
+++ b/Engine/Source/Editor/UnrealEd/Private/Factories/EditorFactories.cpp
@@ -9,6 +9,9 @@
 #include "Matinee/InterpGroupCamera.h"
 #include "Materials/MaterialExpressionTextureSample.h"
 #include "Materials/MaterialFunction.h"
+#include "Sound/SoundNodeWavePlayer.h"
+#include "Sound/SoundNodeAttenuation.h"
+#include "Sound/SoundNodeModulator.h"
 #include "Factories.h"
 #include "SoundDefinitions.h"
 #include "BlueprintUtilities.h"
diff --git a/Engine/Source/Editor/UnrealEd/Private/SoundCueGraphNode.cpp b/Engine/Source/Editor/UnrealEd/Private/SoundCueGraphNode.cpp
--- a/Engine/Source/Editor/UnrealEd/Private/SoundCueGraphNode.cpp
+++ b/Engine/Source/Editor/UnrealEd/Private/SoundCueGraphNode.cpp
@@ -5,6 +5,7 @@
 =============================================================================*/
 
 #include "UnrealEd.h"
+#include "Sound/SoundNodeWavePlayer.h"
 #include "SoundDefinitions.h"
 #include "GraphEditorActions.h"
 #include "SoundCueGraphEditorCommands.h"
diff --git a/Engine/Source/Editor/UnrealEd/Private/SoundCueGraphSchema.cpp b/Engine/Source/Editor/UnrealEd/Private/SoundCueGraphSchema.cpp
--- a/Engine/Source/Editor/UnrealEd/Private/SoundCueGraphSchema.cpp
+++ b/Engine/Source/Editor/UnrealEd/Private/SoundCueGraphSchema.cpp
@@ -5,6 +5,8 @@
 =============================================================================*/
 
 #include "UnrealEd.h"
+#include "Sound/SoundNodeWavePlayer.h"
+#include "Sound/SoundNodeDeprecated.h"
 #include "Slate.h"
 #include "AssetData.h"
 #include "GraphEditorActions.h"
diff --git a/Engine/Source/Runtime/Engine/Classes/Sound/SoundNodeAmbient.h b/Engine/Source/Runtime/Engine/Classes/Sound/SoundNodeAmbient.h
--- a/Engine/Source/Runtime/Engine/Classes/Sound/SoundNodeAmbient.h
+++ b/Engine/Source/Runtime/Engine/Classes/Sound/SoundNodeAmbient.h
@@ -6,6 +6,7 @@
  */
 
 #pragma once
+#include "Sound/SoundNodeDeprecated.h"
 #include "SoundNodeAmbient.generated.h"
 
 USTRUCT()
diff --git a/Engine/Source/Runtime/Engine/Classes/Sound/SoundNodeAmbientNonLoop.h b/Engine/Source/Runtime/Engine/Classes/Sound/SoundNodeAmbientNonLoop.h
--- a/Engine/Source/Runtime/Engine/Classes/Sound/SoundNodeAmbientNonLoop.h
+++ b/Engine/Source/Runtime/Engine/Classes/Sound/SoundNodeAmbientNonLoop.h
@@ -6,6 +6,7 @@
  */
 
 #pragma once
+#include "Sound/SoundNodeAmbient.h"
 #include "SoundNodeAmbientNonLoop.generated.h"
 
 UCLASS(dependson=(USoundNodeAttenuation, UDistributionFloat), editinlinenew)
diff --git a/Engine/Source/Runtime/Engine/Classes/Sound/SoundNodeAmbientNonLoopToggle.h b/Engine/Source/Runtime/Engine/Classes/Sound/SoundNodeAmbientNonLoopToggle.h
--- a/Engine/Source/Runtime/Engine/Classes/Sound/SoundNodeAmbientNonLoopToggle.h
+++ b/Engine/Source/Runtime/Engine/Classes/Sound/SoundNodeAmbientNonLoopToggle.h
@@ -6,6 +6,7 @@
  */
 
 #pragma once
+#include "Sound/SoundNodeAmbientNonLoop.h"
 #include "SoundNodeAmbientNonLoopToggle.generated.h"
 
 UCLASS(hidecategories=Object, dependson=USoundNodeAttenuation, editinlinenew)
diff --git a/Engine/Source/Runtime/Engine/Classes/Sound/SoundNodeAttenuation.h b/Engine/Source/Runtime/Engine/Classes/Sound/SoundNodeAttenuation.h
--- a/Engine/Source/Runtime/Engine/Classes/Sound/SoundNodeAttenuation.h
+++ b/Engine/Source/Runtime/Engine/Classes/Sound/SoundNodeAttenuation.h
@@ -2,6 +2,7 @@
 
  
 #pragma once
+#include "Sound/SoundNode.h"
 #include "SoundNodeAttenuation.generated.h"
 
 /** 
diff --git a/Engine/Source/Runtime/Engine/Classes/Sound/SoundNodeBranch.h b/Engine/Source/Runtime/Engine/Classes/Sound/SoundNodeBranch.h
--- a/Engine/Source/Runtime/Engine/Classes/Sound/SoundNodeBranch.h
+++ b/Engine/Source/Runtime/Engine/Classes/Sound/SoundNodeBranch.h
@@ -1,6 +1,7 @@
 // Copyright 1998-2014 Epic Games, Inc. All Rights Reserved.
 
 #pragma once
+#include "Sound/SoundNode.h"
 #include "SoundNodeBranch.generated.h"
 
 /** 
diff --git a/Engine/Source/Runtime/Engine/Classes/Sound/SoundNodeConcatenator.h b/Engine/Source/Runtime/Engine/Classes/Sound/SoundNodeConcatenator.h
--- a/Engine/Source/Runtime/Engine/Classes/Sound/SoundNodeConcatenator.h
+++ b/Engine/Source/Runtime/Engine/Classes/Sound/SoundNodeConcatenator.h
@@ -7,6 +7,7 @@
  */
 
 #pragma once
+#include "Sound/SoundNode.h"
 #include "SoundNodeConcatenator.generated.h"
 
 /**
diff --git a/Engine/Source/Runtime/Engine/Classes/Sound/SoundNodeConcatenatorRadio.h b/Engine/Source/Runtime/Engine/Classes/Sound/SoundNodeConcatenatorRadio.h
--- a/Engine/Source/Runtime/Engine/Classes/Sound/SoundNodeConcatenatorRadio.h
+++ b/Engine/Source/Runtime/Engine/Classes/Sound/SoundNodeConcatenatorRadio.h
@@ -6,6 +6,7 @@
  */
 
 #pragma once
+#include "Sound/SoundNodeDeprecated.h"
 #include "SoundNodeConcatenatorRadio.generated.h"
 
 UCLASS(deprecated, hidecategories=Object, editinlinenew, MinimalAPI)
diff --git a/Engine/Source/Runtime/Engine/Classes/Sound/SoundNodeDelay.h b/Engine/Source/Runtime/Engine/Classes/Sound/SoundNodeDelay.h
--- a/Engine/Source/Runtime/Engine/Classes/Sound/SoundNodeDelay.h
+++ b/Engine/Source/Runtime/Engine/Classes/Sound/SoundNodeDelay.h
@@ -1,6 +1,7 @@
 // Copyright 1998-2014 Epic Games, Inc. All Rights Reserved.
 
 #pragma once
+#include "Sound/SoundNode.h"
 #include "SoundNodeDelay.generated.h"
 
 /** 
diff --git a/Engine/Source/Runtime/Engine/Classes/Sound/SoundNodeDeprecated.h b/Engine/Source/Runtime/Engine/Classes/Sound/SoundNodeDeprecated.h
--- a/Engine/Source/Runtime/Engine/Classes/Sound/SoundNodeDeprecated.h
+++ b/Engine/Source/Runtime/Engine/Classes/Sound/SoundNodeDeprecated.h
@@ -7,6 +7,7 @@
  */
 
 #pragma once
+#include "Sound/SoundNode.h"
 #include "SoundNodeDeprecated.generated.h"
 
 UCLASS(deprecated, abstract, MinimalAPI)
diff --git a/Engine/Source/Runtime/Engine/Classes/Sound/SoundNodeDialoguePlayer.h b/Engine/Source/Runtime/Engine/Classes/Sound/SoundNodeDialoguePlayer.h
--- a/Engine/Source/Runtime/Engine/Classes/Sound/SoundNodeDialoguePlayer.h
+++ b/Engine/Source/Runtime/Engine/Classes/Sound/SoundNodeDialoguePlayer.h
@@ -2,6 +2,7 @@
 
 
 #pragma once
+#include "Sound/SoundNode.h"
 #include "SoundNodeDialoguePlayer.generated.h"
 
 /**
diff --git a/Engine/Source/Runtime/Engine/Classes/Sound/SoundNodeDistanceCrossFade.h b/Engine/Source/Runtime/Engine/Classes/Sound/SoundNodeDistanceCrossFade.h
--- a/Engine/Source/Runtime/Engine/Classes/Sound/SoundNodeDistanceCrossFade.h
+++ b/Engine/Source/Runtime/Engine/Classes/Sound/SoundNodeDistanceCrossFade.h
@@ -2,6 +2,7 @@
 
 
 #pragma once
+#include "Sound/SoundNode.h"
 #include "SoundNodeDistanceCrossFade.generated.h"
 
 USTRUCT()
diff --git a/Engine/Source/Runtime/Engine/Classes/Sound/SoundNodeDoppler.h b/Engine/Source/Runtime/Engine/Classes/Sound/SoundNodeDoppler.h
--- a/Engine/Source/Runtime/Engine/Classes/Sound/SoundNodeDoppler.h
+++ b/Engine/Source/Runtime/Engine/Classes/Sound/SoundNodeDoppler.h
@@ -1,7 +1,7 @@
 // Copyright 1998-2014 Epic Games, Inc. All Rights Reserved.
 
 #pragma once
- 
+#include "Sound/SoundNode.h"
 #include "SoundNodeDoppler.generated.h"
 
 struct FListener;
diff --git a/Engine/Source/Runtime/Engine/Classes/Sound/SoundNodeEnveloper.h b/Engine/Source/Runtime/Engine/Classes/Sound/SoundNodeEnveloper.h
--- a/Engine/Source/Runtime/Engine/Classes/Sound/SoundNodeEnveloper.h
+++ b/Engine/Source/Runtime/Engine/Classes/Sound/SoundNodeEnveloper.h
@@ -1,6 +1,7 @@
 // Copyright 1998-2014 Epic Games, Inc. All Rights Reserved.
 
 #pragma once
+#include "Sound/SoundNode.h"
 #include "SoundNodeEnveloper.generated.h"
 
 /** 
diff --git a/Engine/Source/Runtime/Engine/Classes/Sound/SoundNodeGroupControl.h b/Engine/Source/Runtime/Engine/Classes/Sound/SoundNodeGroupControl.h
--- a/Engine/Source/Runtime/Engine/Classes/Sound/SoundNodeGroupControl.h
+++ b/Engine/Source/Runtime/Engine/Classes/Sound/SoundNodeGroupControl.h
@@ -1,6 +1,7 @@
 // Copyright 1998-2014 Epic Games, Inc. All Rights Reserved.
 
 #pragma once
+#include "Sound/SoundNode.h"
 #include "SoundNodeGroupControl.generated.h"
 
 /** 
diff --git a/Engine/Source/Runtime/Engine/Classes/Sound/SoundNodeLooping.h b/Engine/Source/Runtime/Engine/Classes/Sound/SoundNodeLooping.h
--- a/Engine/Source/Runtime/Engine/Classes/Sound/SoundNodeLooping.h
+++ b/Engine/Source/Runtime/Engine/Classes/Sound/SoundNodeLooping.h
@@ -2,6 +2,7 @@
 
  
 #pragma once
+#include "Sound/SoundNode.h"
 #include "SoundNodeLooping.generated.h"
 
 /** 
diff --git a/Engine/Source/Runtime/Engine/Classes/Sound/SoundNodeMature.h b/Engine/Source/Runtime/Engine/Classes/Sound/SoundNodeMature.h
--- a/Engine/Source/Runtime/Engine/Classes/Sound/SoundNodeMature.h
+++ b/Engine/Source/Runtime/Engine/Classes/Sound/SoundNodeMature.h
@@ -2,6 +2,7 @@
 
  
 #pragma once
+#include "Sound/SoundNode.h"
 #include "SoundNodeMature.generated.h"
 
 /**
diff --git a/Engine/Source/Runtime/Engine/Classes/Sound/SoundNodeMixer.h b/Engine/Source/Runtime/Engine/Classes/Sound/SoundNodeMixer.h
--- a/Engine/Source/Runtime/Engine/Classes/Sound/SoundNodeMixer.h
+++ b/Engine/Source/Runtime/Engine/Classes/Sound/SoundNodeMixer.h
@@ -2,6 +2,7 @@
 
  
 #pragma once
+#include "Sound/SoundNode.h"
 #include "SoundNodeMixer.generated.h"
 
 /** 
diff --git a/Engine/Source/Runtime/Engine/Classes/Sound/SoundNodeModulator.h b/Engine/Source/Runtime/Engine/Classes/Sound/SoundNodeModulator.h
--- a/Engine/Source/Runtime/Engine/Classes/Sound/SoundNodeModulator.h
+++ b/Engine/Source/Runtime/Engine/Classes/Sound/SoundNodeModulator.h
@@ -2,6 +2,7 @@
 
  
 #pragma once
+#include "Sound/SoundNode.h"
 #include "SoundNodeModulator.generated.h"
 
 /** 
diff --git a/Engine/Source/Runtime/Engine/Classes/Sound/SoundNodeModulatorContinuous.h b/Engine/Source/Runtime/Engine/Classes/Sound/SoundNodeModulatorContinuous.h
--- a/Engine/Source/Runtime/Engine/Classes/Sound/SoundNodeModulatorContinuous.h
+++ b/Engine/Source/Runtime/Engine/Classes/Sound/SoundNodeModulatorContinuous.h
@@ -2,6 +2,7 @@
 
 
 #pragma once
+#include "Sound/SoundNode.h"
 #include "SoundNodeModulatorContinuous.generated.h"
 
 UENUM()
diff --git a/Engine/Source/Runtime/Engine/Classes/Sound/SoundNodeOscillator.h b/Engine/Source/Runtime/Engine/Classes/Sound/SoundNodeOscillator.h
--- a/Engine/Source/Runtime/Engine/Classes/Sound/SoundNodeOscillator.h
+++ b/Engine/Source/Runtime/Engine/Classes/Sound/SoundNodeOscillator.h
@@ -2,6 +2,7 @@
 
 
 #pragma once
+#include "Sound/SoundNode.h"
 #include "SoundNodeOscillator.generated.h"
 
 /** 
diff --git a/Engine/Source/Runtime/Engine/Classes/Sound/SoundNodeParamCrossFade.h b/Engine/Source/Runtime/Engine/Classes/Sound/SoundNodeParamCrossFade.h
--- a/Engine/Source/Runtime/Engine/Classes/Sound/SoundNodeParamCrossFade.h
+++ b/Engine/Source/Runtime/Engine/Classes/Sound/SoundNodeParamCrossFade.h
@@ -1,6 +1,7 @@
 // Copyright 1998-2014 Epic Games, Inc. All Rights Reserved.
 
 #pragma once
+#include "Sound/SoundNodeDistanceCrossFade.h"
 #include "SoundNodeParamCrossFade.generated.h"
 
 /** 
diff --git a/Engine/Source/Runtime/Engine/Classes/Sound/SoundNodeRandom.h b/Engine/Source/Runtime/Engine/Classes/Sound/SoundNodeRandom.h
--- a/Engine/Source/Runtime/Engine/Classes/Sound/SoundNodeRandom.h
+++ b/Engine/Source/Runtime/Engine/Classes/Sound/SoundNodeRandom.h
@@ -2,6 +2,7 @@
 
  
 #pragma once
+#include "Sound/SoundNode.h"
 #include "SoundNodeRandom.generated.h"
 
 /** 
diff --git a/Engine/Source/Runtime/Engine/Classes/Sound/SoundNodeSoundClass.h b/Engine/Source/Runtime/Engine/Classes/Sound/SoundNodeSoundClass.h
--- a/Engine/Source/Runtime/Engine/Classes/Sound/SoundNodeSoundClass.h
+++ b/Engine/Source/Runtime/Engine/Classes/Sound/SoundNodeSoundClass.h
@@ -2,6 +2,7 @@
 
  
 #pragma once
+#include "Sound/SoundNode.h"
 #include "SoundNodeSoundClass.generated.h"
 
 /** 
diff --git a/Engine/Source/Runtime/Engine/Classes/Sound/SoundNodeSwitch.h b/Engine/Source/Runtime/Engine/Classes/Sound/SoundNodeSwitch.h
--- a/Engine/Source/Runtime/Engine/Classes/Sound/SoundNodeSwitch.h
+++ b/Engine/Source/Runtime/Engine/Classes/Sound/SoundNodeSwitch.h
@@ -1,6 +1,7 @@
 // Copyright 1998-2014 Epic Games, Inc. All Rights Reserved.
 
 #pragma once
+#include "Sound/SoundNode.h"
 #include "SoundNodeSwitch.generated.h"
 
 /** 
diff --git a/Engine/Source/Runtime/Engine/Classes/Sound/SoundNodeWave.h b/Engine/Source/Runtime/Engine/Classes/Sound/SoundNodeWave.h
--- a/Engine/Source/Runtime/Engine/Classes/Sound/SoundNodeWave.h
+++ b/Engine/Source/Runtime/Engine/Classes/Sound/SoundNodeWave.h
@@ -6,6 +6,7 @@
  */
 
 #pragma once
+#include "Sound/SoundNodeDeprecated.h"
 #include "SoundNodeWave.generated.h"
 
 UCLASS(deprecated, PerObjectConfig, hidecategories=Object, editinlinenew, MinimalAPI)
diff --git a/Engine/Source/Runtime/Engine/Classes/Sound/SoundNodeWaveParam.h b/Engine/Source/Runtime/Engine/Classes/Sound/SoundNodeWaveParam.h
--- a/Engine/Source/Runtime/Engine/Classes/Sound/SoundNodeWaveParam.h
+++ b/Engine/Source/Runtime/Engine/Classes/Sound/SoundNodeWaveParam.h
@@ -1,6 +1,7 @@
 // Copyright 1998-2014 Epic Games, Inc. All Rights Reserved.
 
 #pragma once
+#include "Sound/SoundNode.h"
 #include "SoundNodeWaveParam.generated.h"
 
 /** 
diff --git a/Engine/Source/Runtime/Engine/Classes/Sound/SoundNodeWavePlayer.h b/Engine/Source/Runtime/Engine/Classes/Sound/SoundNodeWavePlayer.h
--- a/Engine/Source/Runtime/Engine/Classes/Sound/SoundNodeWavePlayer.h
+++ b/Engine/Source/Runtime/Engine/Classes/Sound/SoundNodeWavePlayer.h
@@ -2,6 +2,7 @@
 
 
 #pragma once
+#include "Sound/SoundNode.h"
 #include "SoundNodeWavePlayer.generated.h"
 
 /** 
diff --git a/Engine/Source/Runtime/Engine/Private/ActiveSound.cpp b/Engine/Source/Runtime/Engine/Private/ActiveSound.cpp
--- a/Engine/Source/Runtime/Engine/Private/ActiveSound.cpp
+++ b/Engine/Source/Runtime/Engine/Private/ActiveSound.cpp
@@ -2,6 +2,7 @@
 
 #include "EnginePrivate.h"
 #include "SoundDefinitions.h"
+#include "Sound/SoundNodeAttenuation.h"
 #include "SubtitleManager.h"
 
 FActiveSound::FActiveSound()
diff --git a/Engine/Source/Runtime/Engine/Private/AmbientSound.cpp b/Engine/Source/Runtime/Engine/Private/AmbientSound.cpp
--- a/Engine/Source/Runtime/Engine/Private/AmbientSound.cpp
+++ b/Engine/Source/Runtime/Engine/Private/AmbientSound.cpp
@@ -5,6 +5,14 @@
 =============================================================================*/
 
 #include "EnginePrivate.h"
+#include "Sound/SoundNodeAmbient.h" // deprecated
+#include "Sound/SoundNodeAmbientNonLoop.h" // deprecated
+#include "Sound/SoundNodeMixer.h"
+#include "Sound/SoundNodeRandom.h"
+#include "Sound/SoundNodeLooping.h"
+#include "Sound/SoundNodeDelay.h"
+#include "Sound/SoundNodeWavePlayer.h"
+#include "Sound/SoundNodeModulator.h"
 #include "SoundDefinitions.h"
 #include "MessageLog.h"
 #include "UObjectToken.h"
diff --git a/Engine/Source/Runtime/Engine/Private/Audio.cpp b/Engine/Source/Runtime/Engine/Private/Audio.cpp
--- a/Engine/Source/Runtime/Engine/Private/Audio.cpp
+++ b/Engine/Source/Runtime/Engine/Private/Audio.cpp
@@ -5,6 +5,7 @@
 =============================================================================*/
 
 #include "EnginePrivate.h"
+#include "Sound/SoundNodeWavePlayer.h"
 #include "SoundDefinitions.h"
 #include "AudioEffect.h"
 #include "Net/UnrealNetwork.h"
diff --git a/Engine/Source/Runtime/Engine/Private/AudioComponent.cpp b/Engine/Source/Runtime/Engine/Private/AudioComponent.cpp
--- a/Engine/Source/Runtime/Engine/Private/AudioComponent.cpp
+++ b/Engine/Source/Runtime/Engine/Private/AudioComponent.cpp
@@ -2,6 +2,7 @@
 
 #include "EnginePrivate.h"
 #include "SoundDefinitions.h"
+#include "Sound/SoundNodeAttenuation.h"
 #include "SubtitleManager.h"
 
 /*-----------------------------------------------------------------------------
diff --git a/Engine/Source/Runtime/Engine/Private/AudioDevice.cpp b/Engine/Source/Runtime/Engine/Private/AudioDevice.cpp
--- a/Engine/Source/Runtime/Engine/Private/AudioDevice.cpp
+++ b/Engine/Source/Runtime/Engine/Private/AudioDevice.cpp
@@ -1,6 +1,7 @@
 // Copyright 1998-2014 Epic Games, Inc. All Rights Reserved.
 
 #include "EnginePrivate.h"
+#include "Sound/SoundNodeWavePlayer.h"
 #include "SoundDefinitions.h"
 #include "AudioEffect.h"
 #include "AudioDecompress.h"
diff --git a/Engine/Source/Runtime/Engine/Private/Distributions.cpp b/Engine/Source/Runtime/Engine/Private/Distributions.cpp
--- a/Engine/Source/Runtime/Engine/Private/Distributions.cpp
+++ b/Engine/Source/Runtime/Engine/Private/Distributions.cpp
@@ -8,6 +8,7 @@
 #include "DistributionHelpers.h"
 #include "ParticleDefinitions.h"
 #include "SoundDefinitions.h"
+#include "Sound/SoundNode.h"
 
 // Moving UDistributions to PostInitProps to not be default sub-objects:
 // Small enough value to be rounded to 0.0 in the editor 
diff --git a/Engine/Source/Runtime/Engine/Private/Level.cpp b/Engine/Source/Runtime/Engine/Private/Level.cpp
--- a/Engine/Source/Runtime/Engine/Private/Level.cpp
+++ b/Engine/Source/Runtime/Engine/Private/Level.cpp
@@ -5,6 +5,7 @@ Level.cpp: Level-related functions
 =============================================================================*/
 
 #include "EnginePrivate.h"
+#include "Sound/SoundNodeWave.h"
 #include "Net/UnrealNetwork.h"
 #include "SoundDefinitions.h"
 #include "PrecomputedLightVolume.h"
diff --git a/Engine/Source/Runtime/Engine/Private/SoundCue.cpp b/Engine/Source/Runtime/Engine/Private/SoundCue.cpp
--- a/Engine/Source/Runtime/Engine/Private/SoundCue.cpp
+++ b/Engine/Source/Runtime/Engine/Private/SoundCue.cpp
@@ -2,6 +2,9 @@
 
 #include "EnginePrivate.h"
 #include "SoundDefinitions.h"
+#include "Sound/SoundNodeMixer.h"
+#include "Sound/SoundNodeWavePlayer.h"
+#include "Sound/SoundNodeAttenuation.h"
 #if WITH_EDITORONLY_DATA
 #include "Kismet2/BlueprintEditorUtils.h"
 #include "UnrealEd.h"
diff --git a/Engine/Source/Runtime/Engine/Private/SoundNode.cpp b/Engine/Source/Runtime/Engine/Private/SoundNode.cpp
--- a/Engine/Source/Runtime/Engine/Private/SoundNode.cpp
+++ b/Engine/Source/Runtime/Engine/Private/SoundNode.cpp
@@ -3,6 +3,8 @@
 
 #include "EnginePrivate.h"
 #include "SoundDefinitions.h"
+#include "Sound/SoundNode.h"
+
 #if WITH_EDITORONLY_DATA
 #include "UnrealEd.h"
 #endif
diff --git a/Engine/Source/Runtime/Engine/Private/SoundNodeAmbient.cpp b/Engine/Source/Runtime/Engine/Private/SoundNodeAmbient.cpp
--- a/Engine/Source/Runtime/Engine/Private/SoundNodeAmbient.cpp
+++ b/Engine/Source/Runtime/Engine/Private/SoundNodeAmbient.cpp
@@ -3,6 +3,7 @@
 
 #include "EnginePrivate.h"
 #include "SoundDefinitions.h"
+#include "Sound/SoundNodeAmbient.h"
 
 /*-----------------------------------------------------------------------------
 	USoundNodeAmbient implementation.
diff --git a/Engine/Source/Runtime/Engine/Private/SoundNodeAmbientNonLoop.cpp b/Engine/Source/Runtime/Engine/Private/SoundNodeAmbientNonLoop.cpp
--- a/Engine/Source/Runtime/Engine/Private/SoundNodeAmbientNonLoop.cpp
+++ b/Engine/Source/Runtime/Engine/Private/SoundNodeAmbientNonLoop.cpp
@@ -3,6 +3,7 @@
 
 #include "EnginePrivate.h"
 #include "SoundDefinitions.h"
+#include "Sound/SoundNodeAmbientNonLoop.h"
 
 /*-----------------------------------------------------------------------------
 	USoundNodeAmbientNonLoop implementation.
diff --git a/Engine/Source/Runtime/Engine/Private/SoundNodeAmbientNonLoopToggle.cpp b/Engine/Source/Runtime/Engine/Private/SoundNodeAmbientNonLoopToggle.cpp
--- a/Engine/Source/Runtime/Engine/Private/SoundNodeAmbientNonLoopToggle.cpp
+++ b/Engine/Source/Runtime/Engine/Private/SoundNodeAmbientNonLoopToggle.cpp
@@ -3,6 +3,7 @@
 
 #include "EnginePrivate.h"
 #include "SoundDefinitions.h"
+#include "Sound/SoundNodeAmbientNonLoopToggle.h"
 
 /*-----------------------------------------------------------------------------
 	USoundNodeAmbientNonLoopToggle implementation
diff --git a/Engine/Source/Runtime/Engine/Private/SoundNodeAttenuation.cpp b/Engine/Source/Runtime/Engine/Private/SoundNodeAttenuation.cpp
--- a/Engine/Source/Runtime/Engine/Private/SoundNodeAttenuation.cpp
+++ b/Engine/Source/Runtime/Engine/Private/SoundNodeAttenuation.cpp
@@ -3,6 +3,7 @@
 
 #include "EnginePrivate.h"
 #include "SoundDefinitions.h"
+#include "Sound/SoundNodeAttenuation.h"
 
 /*-----------------------------------------------------------------------------
 	USoundNodeAttenuation implementation.
diff --git a/Engine/Source/Runtime/Engine/Private/SoundNodeBranch.cpp b/Engine/Source/Runtime/Engine/Private/SoundNodeBranch.cpp
--- a/Engine/Source/Runtime/Engine/Private/SoundNodeBranch.cpp
+++ b/Engine/Source/Runtime/Engine/Private/SoundNodeBranch.cpp
@@ -3,6 +3,7 @@
 
 #include "EnginePrivate.h"
 #include "SoundDefinitions.h"
+#include "Sound/SoundNodeBranch.h"
 
 #define LOCTEXT_NAMESPACE "SoundNodeBranch"
 
diff --git a/Engine/Source/Runtime/Engine/Private/SoundNodeConcatenator.cpp b/Engine/Source/Runtime/Engine/Private/SoundNodeConcatenator.cpp
--- a/Engine/Source/Runtime/Engine/Private/SoundNodeConcatenator.cpp
+++ b/Engine/Source/Runtime/Engine/Private/SoundNodeConcatenator.cpp
@@ -3,6 +3,7 @@
 
 #include "EnginePrivate.h"
 #include "SoundDefinitions.h"
+#include "Sound/SoundNodeConcatenator.h"
 
 /*-----------------------------------------------------------------------------
          USoundNodeConcatenator implementation.
diff --git a/Engine/Source/Runtime/Engine/Private/SoundNodeConcatenatorRadio.cpp b/Engine/Source/Runtime/Engine/Private/SoundNodeConcatenatorRadio.cpp
--- a/Engine/Source/Runtime/Engine/Private/SoundNodeConcatenatorRadio.cpp
+++ b/Engine/Source/Runtime/Engine/Private/SoundNodeConcatenatorRadio.cpp
@@ -3,6 +3,7 @@
 
 #include "EnginePrivate.h"
 #include "SoundDefinitions.h"
+#include "Sound/SoundNodeConcatenatorRadio.h"
 
 UDEPRECATED_SoundNodeConcatenatorRadio::UDEPRECATED_SoundNodeConcatenatorRadio(const class FPostConstructInitializeProperties& PCIP)
 	: Super(PCIP)
diff --git a/Engine/Source/Runtime/Engine/Private/SoundNodeDelay.cpp b/Engine/Source/Runtime/Engine/Private/SoundNodeDelay.cpp
--- a/Engine/Source/Runtime/Engine/Private/SoundNodeDelay.cpp
+++ b/Engine/Source/Runtime/Engine/Private/SoundNodeDelay.cpp
@@ -3,6 +3,7 @@
 
 #include "EnginePrivate.h"
 #include "SoundDefinitions.h"
+#include "Sound/SoundNodeDelay.h"
 
 /*-----------------------------------------------------------------------------
          USoundNodeDelay implementation.
diff --git a/Engine/Source/Runtime/Engine/Private/SoundNodeDeprecated.cpp b/Engine/Source/Runtime/Engine/Private/SoundNodeDeprecated.cpp
--- a/Engine/Source/Runtime/Engine/Private/SoundNodeDeprecated.cpp
+++ b/Engine/Source/Runtime/Engine/Private/SoundNodeDeprecated.cpp
@@ -3,6 +3,7 @@
 
 #include "EnginePrivate.h"
 #include "SoundDefinitions.h"
+#include "Sound/SoundNodeDeprecated.h"
 
 UDEPRECATED_SoundNodeDeprecated::UDEPRECATED_SoundNodeDeprecated(const class FPostConstructInitializeProperties& PCIP)
 	: Super(PCIP)
diff --git a/Engine/Source/Runtime/Engine/Private/SoundNodeDialoguePlayer.cpp b/Engine/Source/Runtime/Engine/Private/SoundNodeDialoguePlayer.cpp
--- a/Engine/Source/Runtime/Engine/Private/SoundNodeDialoguePlayer.cpp
+++ b/Engine/Source/Runtime/Engine/Private/SoundNodeDialoguePlayer.cpp
@@ -2,6 +2,7 @@
 
 #include "EnginePrivate.h"
 #include "SoundDefinitions.h"
+#include "Sound/SoundNodeDialoguePlayer.h"
 
 USoundNodeDialoguePlayer::USoundNodeDialoguePlayer(const class FPostConstructInitializeProperties& PCIP)
 	: Super(PCIP)
diff --git a/Engine/Source/Runtime/Engine/Private/SoundNodeDistanceCrossFade.cpp b/Engine/Source/Runtime/Engine/Private/SoundNodeDistanceCrossFade.cpp
--- a/Engine/Source/Runtime/Engine/Private/SoundNodeDistanceCrossFade.cpp
+++ b/Engine/Source/Runtime/Engine/Private/SoundNodeDistanceCrossFade.cpp
@@ -3,6 +3,7 @@
 
 #include "EnginePrivate.h"
 #include "SoundDefinitions.h"
+#include "Sound/SoundNodeDistanceCrossFade.h"
 
 /*-----------------------------------------------------------------------------
 	USoundNodeDistanceCrossFade implementation.
diff --git a/Engine/Source/Runtime/Engine/Private/SoundNodeDoppler.cpp b/Engine/Source/Runtime/Engine/Private/SoundNodeDoppler.cpp
--- a/Engine/Source/Runtime/Engine/Private/SoundNodeDoppler.cpp
+++ b/Engine/Source/Runtime/Engine/Private/SoundNodeDoppler.cpp
@@ -3,6 +3,7 @@
 
 #include "EnginePrivate.h"
 #include "SoundDefinitions.h"
+#include "Sound/SoundNodeDoppler.h"
 
 /*-----------------------------------------------------------------------------
          USoundNodeDoppler implementation.
diff --git a/Engine/Source/Runtime/Engine/Private/SoundNodeEnveloper.cpp b/Engine/Source/Runtime/Engine/Private/SoundNodeEnveloper.cpp
--- a/Engine/Source/Runtime/Engine/Private/SoundNodeEnveloper.cpp
+++ b/Engine/Source/Runtime/Engine/Private/SoundNodeEnveloper.cpp
@@ -3,6 +3,7 @@
 
 #include "EnginePrivate.h"
 #include "SoundDefinitions.h"
+#include "Sound/SoundNodeEnveloper.h"
 
 /*-----------------------------------------------------------------------------
 USoundNodeEnveloper implementation.
diff --git a/Engine/Source/Runtime/Engine/Private/SoundNodeGroupControl.cpp b/Engine/Source/Runtime/Engine/Private/SoundNodeGroupControl.cpp
--- a/Engine/Source/Runtime/Engine/Private/SoundNodeGroupControl.cpp
+++ b/Engine/Source/Runtime/Engine/Private/SoundNodeGroupControl.cpp
@@ -3,6 +3,7 @@
 
 #include "EnginePrivate.h"
 #include "SoundDefinitions.h"
+#include "Sound/SoundNodeGroupControl.h"
 
 DEFINE_LOG_CATEGORY_STATIC(LogSoundNodeGroupControl, Log, All);
 
diff --git a/Engine/Source/Runtime/Engine/Private/SoundNodeLooping.cpp b/Engine/Source/Runtime/Engine/Private/SoundNodeLooping.cpp
--- a/Engine/Source/Runtime/Engine/Private/SoundNodeLooping.cpp
+++ b/Engine/Source/Runtime/Engine/Private/SoundNodeLooping.cpp
@@ -3,6 +3,8 @@
 
 #include "EnginePrivate.h"
 #include "SoundDefinitions.h"
+#include "Sound/SoundNodeLooping.h"
+#include "Sound/SoundNodeWavePlayer.h"
 
 /*-----------------------------------------------------------------------------
 	USoundNodeLooping implementation.
diff --git a/Engine/Source/Runtime/Engine/Private/SoundNodeMature.cpp b/Engine/Source/Runtime/Engine/Private/SoundNodeMature.cpp
--- a/Engine/Source/Runtime/Engine/Private/SoundNodeMature.cpp
+++ b/Engine/Source/Runtime/Engine/Private/SoundNodeMature.cpp
@@ -3,6 +3,8 @@
 
 #include "EnginePrivate.h"
 #include "SoundDefinitions.h"
+#include "Sound/SoundNodeMature.h"
+#include "Sound/SoundNodeWavePlayer.h"
 
 /*-----------------------------------------------------------------------------
        USoundNodeMature implementation.
diff --git a/Engine/Source/Runtime/Engine/Private/SoundNodeMixer.cpp b/Engine/Source/Runtime/Engine/Private/SoundNodeMixer.cpp
--- a/Engine/Source/Runtime/Engine/Private/SoundNodeMixer.cpp
+++ b/Engine/Source/Runtime/Engine/Private/SoundNodeMixer.cpp
@@ -3,6 +3,7 @@
 
 #include "EnginePrivate.h"
 #include "SoundDefinitions.h"
+#include "Sound/SoundNodeMixer.h"
 
 /*-----------------------------------------------------------------------------
 	USoundNodeMixer implementation.
diff --git a/Engine/Source/Runtime/Engine/Private/SoundNodeModulator.cpp b/Engine/Source/Runtime/Engine/Private/SoundNodeModulator.cpp
--- a/Engine/Source/Runtime/Engine/Private/SoundNodeModulator.cpp
+++ b/Engine/Source/Runtime/Engine/Private/SoundNodeModulator.cpp
@@ -3,6 +3,7 @@
 
 #include "EnginePrivate.h"
 #include "SoundDefinitions.h"
+#include "Sound/SoundNodeModulator.h"
 
 /*-----------------------------------------------------------------------------
 	USoundNodeModulator implementation.
diff --git a/Engine/Source/Runtime/Engine/Private/SoundNodeModulatorContinuous.cpp b/Engine/Source/Runtime/Engine/Private/SoundNodeModulatorContinuous.cpp
--- a/Engine/Source/Runtime/Engine/Private/SoundNodeModulatorContinuous.cpp
+++ b/Engine/Source/Runtime/Engine/Private/SoundNodeModulatorContinuous.cpp
@@ -4,6 +4,7 @@
 #include "EnginePrivate.h"
 #include "SoundDefinitions.h"
 #include "DistributionHelpers.h"
+#include "Sound/SoundNodeModulatorContinuous.h"
 
 float FModulatorContinuousParams::GetValue(const FActiveSound& ActiveSound) const
 {
diff --git a/Engine/Source/Runtime/Engine/Private/SoundNodeOscillator.cpp b/Engine/Source/Runtime/Engine/Private/SoundNodeOscillator.cpp
--- a/Engine/Source/Runtime/Engine/Private/SoundNodeOscillator.cpp
+++ b/Engine/Source/Runtime/Engine/Private/SoundNodeOscillator.cpp
@@ -3,6 +3,7 @@
 
 #include "EnginePrivate.h"
 #include "SoundDefinitions.h"
+#include "Sound/SoundNodeOscillator.h"
 
 /*-----------------------------------------------------------------------------
 	USoundNodeOscillator implementation.
diff --git a/Engine/Source/Runtime/Engine/Private/SoundNodeParamCrossFade.cpp b/Engine/Source/Runtime/Engine/Private/SoundNodeParamCrossFade.cpp
--- a/Engine/Source/Runtime/Engine/Private/SoundNodeParamCrossFade.cpp
+++ b/Engine/Source/Runtime/Engine/Private/SoundNodeParamCrossFade.cpp
@@ -3,6 +3,7 @@
 
 #include "EnginePrivate.h"
 #include "SoundDefinitions.h"
+#include "Sound/SoundNodeParamCrossFade.h"
 
 /*-----------------------------------------------------------------------------
 	USoundNodeParamCrossFade implementation.
diff --git a/Engine/Source/Runtime/Engine/Private/SoundNodeRandom.cpp b/Engine/Source/Runtime/Engine/Private/SoundNodeRandom.cpp
--- a/Engine/Source/Runtime/Engine/Private/SoundNodeRandom.cpp
+++ b/Engine/Source/Runtime/Engine/Private/SoundNodeRandom.cpp
@@ -3,6 +3,7 @@
 
 #include "EnginePrivate.h"
 #include "SoundDefinitions.h"
+#include "Sound/SoundNodeRandom.h"
 
 #if WITH_EDITOR
 	#include "UnrealEd.h"
diff --git a/Engine/Source/Runtime/Engine/Private/SoundNodeSoundClass.cpp b/Engine/Source/Runtime/Engine/Private/SoundNodeSoundClass.cpp
--- a/Engine/Source/Runtime/Engine/Private/SoundNodeSoundClass.cpp
+++ b/Engine/Source/Runtime/Engine/Private/SoundNodeSoundClass.cpp
@@ -3,6 +3,7 @@
 
 #include "EnginePrivate.h"
 #include "SoundDefinitions.h"
+#include "Sound/SoundNodeSoundClass.h"
 
 /*-----------------------------------------------------------------------------
 	USoundNodeSoundClass implementation.
diff --git a/Engine/Source/Runtime/Engine/Private/SoundNodeSwitch.cpp b/Engine/Source/Runtime/Engine/Private/SoundNodeSwitch.cpp
--- a/Engine/Source/Runtime/Engine/Private/SoundNodeSwitch.cpp
+++ b/Engine/Source/Runtime/Engine/Private/SoundNodeSwitch.cpp
@@ -3,6 +3,8 @@
 
 #include "EnginePrivate.h"
 #include "SoundDefinitions.h"
+#include "Sound/SoundNodeSwitch.h"
+
 #if WITH_EDITORONLY_DATA
 #include "UnrealEd.h"
 #endif
diff --git a/Engine/Source/Runtime/Engine/Private/SoundNodeWave.cpp b/Engine/Source/Runtime/Engine/Private/SoundNodeWave.cpp
--- a/Engine/Source/Runtime/Engine/Private/SoundNodeWave.cpp
+++ b/Engine/Source/Runtime/Engine/Private/SoundNodeWave.cpp
@@ -2,6 +2,7 @@
 
 
 #include "EnginePrivate.h"
+#include "Sound/SoundNodeWave.h" // deprecated
 #include "AudioDecompress.h"
 #include "SoundDefinitions.h"
 #include "SubtitleManager.h"
diff --git a/Engine/Source/Runtime/Engine/Private/SoundNodeWaveParam.cpp b/Engine/Source/Runtime/Engine/Private/SoundNodeWaveParam.cpp
--- a/Engine/Source/Runtime/Engine/Private/SoundNodeWaveParam.cpp
+++ b/Engine/Source/Runtime/Engine/Private/SoundNodeWaveParam.cpp
@@ -3,6 +3,7 @@
 
 #include "EnginePrivate.h"
 #include "SoundDefinitions.h"
+#include "Sound/SoundNodeWaveParam.h"
 
 /*-----------------------------------------------------------------------------
 	USoundNodeWaveParam implementation
diff --git a/Engine/Source/Runtime/Engine/Private/SoundNodeWavePlayer.cpp b/Engine/Source/Runtime/Engine/Private/SoundNodeWavePlayer.cpp
--- a/Engine/Source/Runtime/Engine/Private/SoundNodeWavePlayer.cpp
+++ b/Engine/Source/Runtime/Engine/Private/SoundNodeWavePlayer.cpp
@@ -2,6 +2,7 @@
 
 #include "EnginePrivate.h"
 #include "SoundDefinitions.h"
+#include "Sound/SoundNodeWavePlayer.h"
 
 #define LOCTEXT_NAMESPACE "SoundNodeWavePlayer"
 
diff --git a/Engine/Source/Runtime/Engine/Public/Engine.h b/Engine/Source/Runtime/Engine/Public/Engine.h
--- a/Engine/Source/Runtime/Engine/Public/Engine.h
+++ b/Engine/Source/Runtime/Engine/Public/Engine.h
@@ -25,12 +25,8 @@
 #include "Engine/EngineBaseTypes.h"
 
 // EngineClasses.h
-#include "Engine/LatentActionManager.h"
-#include "Engine/Scene.h"
 #include "Camera/CameraTypes.h"
-#include "Engine/NetSerialization.h"
 #include "Engine/EngineTypes.h"
-#include "Interfaces/Interface_AssetUserData.h"
 #include "Components/ActorComponent.h"
 #include "GameFramework/Actor.h"
 #include "Sound/AmbientSound.h"
@@ -647,34 +643,6 @@
 #include "Sound/SoundWave.h"
 #include "Sound/SoundWaveStreaming.h"
 #include "Sound/SoundMix.h"
-#include "Sound/SoundNode.h"
-#include "Sound/SoundNodeAttenuation.h"
-#include "Sound/SoundNodeBranch.h"
-#include "Sound/SoundNodeConcatenator.h"
-#include "Sound/SoundNodeDelay.h"
-#include "Sound/SoundNodeDeprecated.h"
-#include "Sound/SoundNodeAmbient.h"
-#include "Sound/SoundNodeAmbientNonLoop.h"
-#include "Sound/SoundNodeAmbientNonLoopToggle.h"
-#include "Sound/SoundNodeConcatenatorRadio.h"
-#include "Sound/SoundNodeWave.h"
-#include "Sound/SoundNodeDialoguePlayer.h"
-#include "Sound/SoundNodeDistanceCrossFade.h"
-#include "Sound/SoundNodeParamCrossFade.h"
-#include "Sound/SoundNodeDoppler.h"
-#include "Sound/SoundNodeEnveloper.h"
-#include "Sound/SoundNodeGroupControl.h"
-#include "Sound/SoundNodeLooping.h"
-#include "Sound/SoundNodeMature.h"
-#include "Sound/SoundNodeMixer.h"
-#include "Sound/SoundNodeModulator.h"
-#include "Sound/SoundNodeModulatorContinuous.h"
-#include "Sound/SoundNodeOscillator.h"
-#include "Sound/SoundNodeRandom.h"
-#include "Sound/SoundNodeSoundClass.h"
-#include "Sound/SoundNodeSwitch.h"
-#include "Sound/SoundNodeWaveParam.h"
-#include "Sound/SoundNodeWavePlayer.h"
 #include "Engine/StaticMesh.h"
 #include "Engine/StaticMeshSocket.h"
 #include "Animation/AnimNodeBase.h"
