Removed several more rarely-used-externally headers from Engine.h
Remove EmitterSpawnable class

diff --git a/Engine/Source/Editor/GeometryMode/Private/GeometryModifiers.cpp b/Engine/Source/Editor/GeometryMode/Private/GeometryModifiers.cpp
--- a/Engine/Source/Editor/GeometryMode/Private/GeometryModifiers.cpp
+++ b/Engine/Source/Editor/GeometryMode/Private/GeometryModifiers.cpp
@@ -1,6 +1,7 @@
 // Copyright 1998-2014 Epic Games, Inc. All Rights Reserved.
 
 #include "GeometryModePrivatePCH.h"
+#include "Engine/BrushShape.h"
 #include "EditorSupportDelegates.h"
 #include "ScopedTransaction.h"
 #include "BSPOps.h"
diff --git a/Engine/Source/Editor/UnrealEd/Private/EditorServer.cpp b/Engine/Source/Editor/UnrealEd/Private/EditorServer.cpp
--- a/Engine/Source/Editor/UnrealEd/Private/EditorServer.cpp
+++ b/Engine/Source/Editor/UnrealEd/Private/EditorServer.cpp
@@ -21,6 +21,7 @@
 #include "../Private/GeomFitUtils.h"
 #include "Editor/GeometryMode/Public/GeometryEdMode.h"
 #include "Landscape/LandscapeProxy.h"
+#include "Lightmass/PrecomputedVisibilityOverrideVolume.h"
 
 #include "Matinee/InterpTrackAnimControl.h"
 #include "Matinee/InterpData.h"
diff --git a/Engine/Source/Editor/UnrealEd/Private/Lightmass/Lightmass.cpp b/Engine/Source/Editor/UnrealEd/Private/Lightmass/Lightmass.cpp
--- a/Engine/Source/Editor/UnrealEd/Private/Lightmass/Lightmass.cpp
+++ b/Engine/Source/Editor/UnrealEd/Private/Lightmass/Lightmass.cpp
@@ -13,6 +13,8 @@
 #include "Runtime/Engine/Classes/Matinee/InterpGroupInst.h"
 #include "Runtime/Engine/Classes/Matinee/InterpTrackMove.h"
 #include "Runtime/Engine/Classes/Matinee/InterpTrackInstMove.h"
+#include "Lightmass/PrecomputedVisibilityVolume.h"
+#include "Lightmass/PrecomputedVisibilityOverrideVolume.h"
 
 extern FSwarmDebugOptions GSwarmDebugOptions;
 
diff --git a/Engine/Source/Editor/UnrealEd/Private/Lightmass/Lightmass.h b/Engine/Source/Editor/UnrealEd/Private/Lightmass/Lightmass.h
--- a/Engine/Source/Editor/UnrealEd/Private/Lightmass/Lightmass.h
+++ b/Engine/Source/Editor/UnrealEd/Private/Lightmass/Lightmass.h
@@ -17,6 +17,9 @@
 #include "HideWindowsPlatformTypes.h"
 #endif
 
+#include "Lightmass/LightmassImportanceVolume.h"
+#include "Lightmass/LightmassCharacterIndirectDetailVolume.h"
+
 /** Forward declarations of Lightmass types */
 namespace Lightmass
 {
diff --git a/Engine/Source/Editor/UnrealEd/Private/StaticLightingSystem/StaticLightingSystem.cpp b/Engine/Source/Editor/UnrealEd/Private/StaticLightingSystem/StaticLightingSystem.cpp
--- a/Engine/Source/Editor/UnrealEd/Private/StaticLightingSystem/StaticLightingSystem.cpp
+++ b/Engine/Source/Editor/UnrealEd/Private/StaticLightingSystem/StaticLightingSystem.cpp
@@ -8,6 +8,7 @@
 
 FSwarmDebugOptions GSwarmDebugOptions;
 
+#include "Lightmass/LightmassCharacterIndirectDetailVolume.h"
 #include "LightingBuildOptions.h"
 #include "StaticLightingPrivate.h"
 #include "Database.h"
diff --git a/Engine/Source/Runtime/Engine/Classes/AI/EnvironmentQuery/EnvQueryManager.h b/Engine/Source/Runtime/Engine/Classes/AI/EnvironmentQuery/EnvQueryManager.h
--- a/Engine/Source/Runtime/Engine/Classes/AI/EnvironmentQuery/EnvQueryManager.h
+++ b/Engine/Source/Runtime/Engine/Classes/AI/EnvironmentQuery/EnvQueryManager.h
@@ -2,6 +2,7 @@
 
 #pragma once
 #include "EnvQueryTypes.h"
+#include "Tickable.h"
 #include "EnvQueryManager.generated.h"
 
 /** wrapper for easy query execution */
diff --git a/Engine/Source/Runtime/Engine/Classes/AI/Navigation/CrowdManager.h b/Engine/Source/Runtime/Engine/Classes/AI/Navigation/CrowdManager.h
--- a/Engine/Source/Runtime/Engine/Classes/AI/Navigation/CrowdManager.h
+++ b/Engine/Source/Runtime/Engine/Classes/AI/Navigation/CrowdManager.h
@@ -1,7 +1,7 @@
 // Copyright 1998-2014 Epic Games, Inc. All Rights Reserved.
 
 #pragma once
-
+#include "Tickable.h"
 #include "CrowdManager.generated.h"
 
 /**
diff --git a/Engine/Source/Runtime/Engine/Classes/AI/Navigation/NavigationTestingActor.h b/Engine/Source/Runtime/Engine/Classes/AI/Navigation/NavigationTestingActor.h
--- a/Engine/Source/Runtime/Engine/Classes/AI/Navigation/NavigationTestingActor.h
+++ b/Engine/Source/Runtime/Engine/Classes/AI/Navigation/NavigationTestingActor.h
@@ -1,7 +1,7 @@
 // Copyright 1998-2014 Epic Games, Inc. All Rights Reserved.
 
 #pragma once
-
+#include "Tickable.h"
 #include "NavigationTestingActor.generated.h"
 
 struct FNavTestTickHelper : FTickableGameObject
diff --git a/Engine/Source/Runtime/Engine/Classes/AI/Navigation/RecastNavMesh.h b/Engine/Source/Runtime/Engine/Classes/AI/Navigation/RecastNavMesh.h
--- a/Engine/Source/Runtime/Engine/Classes/AI/Navigation/RecastNavMesh.h
+++ b/Engine/Source/Runtime/Engine/Classes/AI/Navigation/RecastNavMesh.h
@@ -3,6 +3,7 @@
 #pragma once
 
 #include "NavAreas/NavArea.h"
+#include "Tickable.h"
 #include "RecastNavMesh.generated.h"
 
 /** Initial checkin. */
diff --git a/Engine/Source/Runtime/Engine/Classes/Debug/LogVisualizerHUD.h b/Engine/Source/Runtime/Engine/Classes/Debug/LogVisualizerHUD.h
--- a/Engine/Source/Runtime/Engine/Classes/Debug/LogVisualizerHUD.h
+++ b/Engine/Source/Runtime/Engine/Classes/Debug/LogVisualizerHUD.h
@@ -1,7 +1,7 @@
 // Copyright 1998-2014 Epic Games, Inc. All Rights Reserved.
 
 #pragma once
-
+#include "Engine/DebugCameraHUD.h"
 #include "LogVisualizerHUD.generated.h"
 
 UCLASS()
diff --git a/Engine/Source/Runtime/Engine/Classes/Engine/DebugCameraHUD.h b/Engine/Source/Runtime/Engine/Classes/Engine/DebugCameraHUD.h
--- a/Engine/Source/Runtime/Engine/Classes/Engine/DebugCameraHUD.h
+++ b/Engine/Source/Runtime/Engine/Classes/Engine/DebugCameraHUD.h
@@ -1,6 +1,7 @@
 // Copyright 1998-2014 Epic Games, Inc. All Rights Reserved.
 
 #pragma once
+#include "GameFramework/HUD.h"
 #include "DebugCameraHUD.generated.h"
 
 UCLASS(config=Game)
diff --git a/Engine/Source/Runtime/Engine/Classes/Lightmass/PrecomputedVisibilityOverrideVolume.h b/Engine/Source/Runtime/Engine/Classes/Lightmass/PrecomputedVisibilityOverrideVolume.h
--- a/Engine/Source/Runtime/Engine/Classes/Lightmass/PrecomputedVisibilityOverrideVolume.h
+++ b/Engine/Source/Runtime/Engine/Classes/Lightmass/PrecomputedVisibilityOverrideVolume.h
@@ -5,6 +5,7 @@
 //=============================================================================
 
 #pragma once
+#include "GameFramework/Volume.h"
 #include "PrecomputedVisibilityOverrideVolume.generated.h"
 
 UCLASS(hidecategories=(Collision, Brush, Attachment, Physics, Volume), MinimalAPI)
diff --git a/Engine/Source/Runtime/Engine/Classes/Lightmass/PrecomputedVisibilityVolume.h b/Engine/Source/Runtime/Engine/Classes/Lightmass/PrecomputedVisibilityVolume.h
--- a/Engine/Source/Runtime/Engine/Classes/Lightmass/PrecomputedVisibilityVolume.h
+++ b/Engine/Source/Runtime/Engine/Classes/Lightmass/PrecomputedVisibilityVolume.h
@@ -5,6 +5,7 @@
 //=============================================================================
 
 #pragma once
+#include "GameFramework/Volume.h"
 #include "PrecomputedVisibilityVolume.generated.h"
 
 UCLASS(hidecategories=(Collision, Brush, Attachment, Physics, Volume), MinimalAPI)
diff --git a/Engine/Source/Runtime/Engine/Classes/Particles/EmitterSpawnable.h b/Engine/Source/Runtime/Engine/Classes/Particles/EmitterSpawnable.h
deleted file mode 100644
diff --git a/Engine/Source/Runtime/Engine/Private/ActorReplication.cpp b/Engine/Source/Runtime/Engine/Private/ActorReplication.cpp
--- a/Engine/Source/Runtime/Engine/Private/ActorReplication.cpp
+++ b/Engine/Source/Runtime/Engine/Private/ActorReplication.cpp
@@ -2,6 +2,7 @@
 
 #include "EnginePrivate.h"
 #include "Net/UnrealNetwork.h"
+#include "GameFramework/GameNetworkManager.h"
 #include "NavigationPathBuilder.h"
 
 /*-----------------------------------------------------------------------------
diff --git a/Engine/Source/Runtime/Engine/Private/BrushShape.cpp b/Engine/Source/Runtime/Engine/Private/BrushShape.cpp
--- a/Engine/Source/Runtime/Engine/Private/BrushShape.cpp
+++ b/Engine/Source/Runtime/Engine/Private/BrushShape.cpp
@@ -1,6 +1,7 @@
 // Copyright 1998-2014 Epic Games, Inc. All Rights Reserved.
 
 #include "EnginePrivate.h"
+#include "Engine/BrushShape.h"
 
 ABrushShape::ABrushShape(const class FPostConstructInitializeProperties& PCIP)
 	: Super(PCIP)
diff --git a/Engine/Source/Runtime/Engine/Private/CharacterMovementComponent.cpp b/Engine/Source/Runtime/Engine/Private/CharacterMovementComponent.cpp
--- a/Engine/Source/Runtime/Engine/Private/CharacterMovementComponent.cpp
+++ b/Engine/Source/Runtime/Engine/Private/CharacterMovementComponent.cpp
@@ -6,6 +6,7 @@
 =============================================================================*/
 
 #include "EnginePrivate.h"
+#include "GameFramework/GameNetworkManager.h"
 
 DEFINE_LOG_CATEGORY_STATIC(LogCharacterMovement, Log, All);
 
diff --git a/Engine/Source/Runtime/Engine/Private/Debug/LogVisualizerHUD.cpp b/Engine/Source/Runtime/Engine/Private/Debug/LogVisualizerHUD.cpp
--- a/Engine/Source/Runtime/Engine/Private/Debug/LogVisualizerHUD.cpp
+++ b/Engine/Source/Runtime/Engine/Private/Debug/LogVisualizerHUD.cpp
@@ -1,6 +1,8 @@
 // Copyright 1998-2014 Epic Games, Inc. All Rights Reserved.
 
 #include "EnginePrivate.h"
+#include "Debug/LogVisualizerHUD.h"
+#include "Debug/LogVisualizerCameraController.h"
 
 #define LOCTEXT_NAMESPACE "ALogVisualizerHUD"
 
diff --git a/Engine/Source/Runtime/Engine/Private/DebugCameraController.cpp b/Engine/Source/Runtime/Engine/Private/DebugCameraController.cpp
--- a/Engine/Source/Runtime/Engine/Private/DebugCameraController.cpp
+++ b/Engine/Source/Runtime/Engine/Private/DebugCameraController.cpp
@@ -6,6 +6,7 @@
 =============================================================================*/
 
 #include "EnginePrivate.h"
+#include "Engine/DebugCameraHUD.h"
 
 /** The currently selected actor. */
 AActor* GDebugSelectedActor = NULL;
diff --git a/Engine/Source/Runtime/Engine/Private/DebugCameraHUD.cpp b/Engine/Source/Runtime/Engine/Private/DebugCameraHUD.cpp
--- a/Engine/Source/Runtime/Engine/Private/DebugCameraHUD.cpp
+++ b/Engine/Source/Runtime/Engine/Private/DebugCameraHUD.cpp
@@ -5,6 +5,7 @@
 =============================================================================*/
 
 #include "EnginePrivate.h"
+#include "Engine/DebugCameraHUD.h"
 
 // ------------------
 // Externals
diff --git a/Engine/Source/Runtime/Engine/Private/GameMode.cpp b/Engine/Source/Runtime/Engine/Private/GameMode.cpp
--- a/Engine/Source/Runtime/Engine/Private/GameMode.cpp
+++ b/Engine/Source/Runtime/Engine/Private/GameMode.cpp
@@ -5,6 +5,7 @@
 =============================================================================*/
 
 #include "EnginePrivate.h"
+#include "GameFramework/GameNetworkManager.h"
 #include "Matinee/MatineeActor.h"
 #include "OnlineSubsystemUtils.h"
 
diff --git a/Engine/Source/Runtime/Engine/Private/GameNetworkManager.cpp b/Engine/Source/Runtime/Engine/Private/GameNetworkManager.cpp
--- a/Engine/Source/Runtime/Engine/Private/GameNetworkManager.cpp
+++ b/Engine/Source/Runtime/Engine/Private/GameNetworkManager.cpp
@@ -5,6 +5,7 @@
 =============================================================================*/
 
 #include "EnginePrivate.h"
+#include "GameFramework/GameNetworkManager.h"
 
 DEFINE_LOG_CATEGORY_STATIC(LogGameNetworkManager, Log, All);
 
diff --git a/Engine/Source/Runtime/Engine/Private/NetworkDriver.cpp b/Engine/Source/Runtime/Engine/Private/NetworkDriver.cpp
--- a/Engine/Source/Runtime/Engine/Private/NetworkDriver.cpp
+++ b/Engine/Source/Runtime/Engine/Private/NetworkDriver.cpp
@@ -7,6 +7,7 @@
 #include "EnginePrivate.h"
 #include "Net/UnrealNetwork.h"
 #include "Net/NetworkProfiler.h"
+#include "GameFramework/GameNetworkManager.h"
 #include "NavigationPathBuilder.h"
 #include "OnlineSubsystemUtils.h"
 
diff --git a/Engine/Source/Runtime/Engine/Private/Particles/ParticleComponents.cpp b/Engine/Source/Runtime/Engine/Private/Particles/ParticleComponents.cpp
--- a/Engine/Source/Runtime/Engine/Private/Particles/ParticleComponents.cpp
+++ b/Engine/Source/Runtime/Engine/Private/Particles/ParticleComponents.cpp
@@ -6,6 +6,7 @@
 
 #include "EnginePrivate.h"
 #include "ParticleDefinitions.h"
+#include "Particles/EmitterCameraLensEffectBase.h"
 #include "LevelUtils.h"
 #include "ImageUtils.h"
 #include "FXSystem.h"
@@ -5947,39 +5948,4 @@ void AEmitterCameraLensEffectBase::ActivateLensEffect()
 	}
 }
 
-
-AEmitterSpawnable::AEmitterSpawnable(const class FPostConstructInitializeProperties& PCIP)
-	: Super(PCIP)
-{
-	ParticleSystemComponent->SecondsBeforeInactive = 0;
-	
-	bDestroyOnSystemFinish = true;
-	bNetTemporary = true;
-}
-
-void AEmitterSpawnable::SetTemplate(UParticleSystem* NewTemplate)
-{
-	Super::SetTemplate(NewTemplate);
-
-	ParticleTemplate = NewTemplate;
-}
-
-void AEmitterSpawnable::GetLifetimeReplicatedProps( TArray< FLifetimeProperty > & OutLifetimeProps ) const
-{
-	Super::GetLifetimeReplicatedProps( OutLifetimeProps );
-	
-	DOREPLIFETIME_CONDITION( AEmitterSpawnable, ParticleTemplate, COND_InitialOnly );
-}
-
-void AEmitterSpawnable::OnRep_ParticleTemplate()
-{
-	SetTemplate(ParticleTemplate);
-	ParticleSystemComponent->ActivateSystem();
-	if (ParticleTemplate == NULL && bDestroyOnSystemFinish)
-	{
-		// prevent emitter from hanging around forever with no template
-		Destroy();
-	}
-}
-
 #undef LOCTEXT_NAMESPACE
\ No newline at end of file
diff --git a/Engine/Source/Runtime/Engine/Private/PlayerCameraManager.cpp b/Engine/Source/Runtime/Engine/Private/PlayerCameraManager.cpp
--- a/Engine/Source/Runtime/Engine/Private/PlayerCameraManager.cpp
+++ b/Engine/Source/Runtime/Engine/Private/PlayerCameraManager.cpp
@@ -3,6 +3,7 @@
 #include "EnginePrivate.h"
 #include "ParticleDefinitions.h"
 #include "SoundDefinitions.h"
+#include "Particles/EmitterCameraLensEffectBase.h"
 #include "IHeadMountedDisplay.h"
 
 DEFINE_LOG_CATEGORY_STATIC(LogPlayerCameraManager, Log, All);
diff --git a/Engine/Source/Runtime/Engine/Private/PlayerController.cpp b/Engine/Source/Runtime/Engine/Private/PlayerController.cpp
--- a/Engine/Source/Runtime/Engine/Private/PlayerController.cpp
+++ b/Engine/Source/Runtime/Engine/Private/PlayerController.cpp
@@ -2,6 +2,7 @@
 
 #include "EnginePrivate.h"
 #include "Net/UnrealNetwork.h"
+#include "GameFramework/GameNetworkManager.h"
 #include "ConfigCacheIni.h"
 #include "SoundDefinitions.h"
 #include "OnlineSubsystemUtils.h"
diff --git a/Engine/Source/Runtime/Engine/Private/PrecomputedVisibilityVolume.cpp b/Engine/Source/Runtime/Engine/Private/PrecomputedVisibilityVolume.cpp
--- a/Engine/Source/Runtime/Engine/Private/PrecomputedVisibilityVolume.cpp
+++ b/Engine/Source/Runtime/Engine/Private/PrecomputedVisibilityVolume.cpp
@@ -1,6 +1,7 @@
 // Copyright 1998-2014 Epic Games, Inc. All Rights Reserved.
 
 #include "EnginePrivate.h"
+#include "Lightmass/PrecomputedVisibilityVolume.h"
 
 APrecomputedVisibilityVolume::APrecomputedVisibilityVolume(const class FPostConstructInitializeProperties& PCIP)
 	: Super(PCIP)
diff --git a/Engine/Source/Runtime/Engine/Private/PrecomputedVisiblityOverrideVolume.cpp b/Engine/Source/Runtime/Engine/Private/PrecomputedVisiblityOverrideVolume.cpp
--- a/Engine/Source/Runtime/Engine/Private/PrecomputedVisiblityOverrideVolume.cpp
+++ b/Engine/Source/Runtime/Engine/Private/PrecomputedVisiblityOverrideVolume.cpp
@@ -1,6 +1,7 @@
 // Copyright 1998-2014 Epic Games, Inc. All Rights Reserved.
 
 #include "EnginePrivate.h"
+#include "Lightmass/PrecomputedVisibilityOverrideVolume.h"
 
 APrecomputedVisibilityOverrideVolume::APrecomputedVisibilityOverrideVolume(const class FPostConstructInitializeProperties& PCIP)
 	: Super(PCIP)
diff --git a/Engine/Source/Runtime/Engine/Private/World.cpp b/Engine/Source/Runtime/Engine/Private/World.cpp
--- a/Engine/Source/Runtime/Engine/Private/World.cpp
+++ b/Engine/Source/Runtime/Engine/Private/World.cpp
@@ -7,6 +7,7 @@
 #include "EnginePrivate.h"
 #include "Matinee/MatineeActor.h"
 #include "Net/UnrealNetwork.h"
+#include "GameFramework/GameNetworkManager.h"
 #include "ParticleDefinitions.h"
 #include "Database.h"
 #include "Net/NetworkProfiler.h"
diff --git a/Engine/Source/Runtime/Engine/Private/WorldSettings.cpp b/Engine/Source/Runtime/Engine/Private/WorldSettings.cpp
--- a/Engine/Source/Runtime/Engine/Private/WorldSettings.cpp
+++ b/Engine/Source/Runtime/Engine/Private/WorldSettings.cpp
@@ -2,6 +2,7 @@
 
 #include "EnginePrivate.h"
 #include "Net/UnrealNetwork.h"
+#include "GameFramework/GameNetworkManager.h"
 #include "SoundDefinitions.h"
 #include "ParticleDefinitions.h"
 #include "MessageLog.h"
diff --git a/Engine/Source/Runtime/Engine/Public/Engine.h b/Engine/Source/Runtime/Engine/Public/Engine.h
--- a/Engine/Source/Runtime/Engine/Public/Engine.h
+++ b/Engine/Source/Runtime/Engine/Public/Engine.h
@@ -6,39 +6,31 @@
 
 #pragma once
 
-#include "EngineDefines.h"
-
 #include "Core.h"
 #include "CoreUObject.h"
 #include "InputCore.h"
+#include "EngineDefines.h"
 #include "EngineSettings.h"
 #include "EngineStats.h"
 #include "EngineLogs.h"
-
 #include "EngineGlobals.h"
 
 /*-----------------------------------------------------------------------------
 	Engine public includes.
 -----------------------------------------------------------------------------*/
 
-#include "Tickable.h"						// FTickableGameObject interface.
 #include "Engine/EngineBaseTypes.h"
-
-// EngineClasses.h
 #include "Camera/CameraTypes.h"
 #include "Engine/EngineTypes.h"
 #include "Components/ActorComponent.h"
 #include "GameFramework/Actor.h"
 #include "Sound/AmbientSound.h"
 #include "Engine/Brush.h"
-#include "Engine/BrushShape.h"
 #include "GameFramework/Volume.h"
 #include "Engine/BlockingVolume.h"
 #include "GameFramework/CameraBlockingVolume.h"
 #include "Engine/CullDistanceVolume.h"
 #include "Engine/LevelStreamingVolume.h"
-#include "Lightmass/LightmassCharacterIndirectDetailVolume.h"
-#include "Lightmass/LightmassImportanceVolume.h"
 #include "AI/Navigation/NavMeshBoundsVolume.h"
 #include "AI/Navigation/NavRelevantActorInterface.h"
 #include "AI/Navigation/NavModifierVolume.h"
@@ -48,8 +40,6 @@
 #include "GameFramework/PainCausingVolume.h"
 #include "Interfaces/Interface_PostProcessVolume.h"
 #include "Engine/PostProcessVolume.h"
-#include "Lightmass/PrecomputedVisibilityOverrideVolume.h"
-#include "Lightmass/PrecomputedVisibilityVolume.h"
 #include "Sound/ReverbVolume.h"
 #include "Engine/TriggerVolume.h"
 #include "Camera/CameraActor.h"
@@ -69,28 +59,19 @@
 #include "GameFramework/Info.h"
 #include "GameFramework/GameMode.h"
 #include "Components/InputComponent.h"
-#include "Engine/ScriptViewportClient.h"
-#include "Engine/GameViewportClient.h"
 #include "Curves/CurveBase.h"
 #include "Curves/CurveFloat.h"
 #include "GameFramework/ForceFeedbackEffect.h"
-#include "GameFramework/OnlineReplStructs.h"
 #include "GameFramework/PlayerController.h"
 #include "Engine/DebugCameraController.h"
-#include "Debug/LogVisualizerCameraController.h"
 #include "Engine/DecalActor.h"
 #include "PhysicsEngine/DestructibleActor.h"
 #include "Particles/Emitter.h"
-#include "Particles/EmitterCameraLensEffectBase.h"
-#include "Particles/EmitterSpawnable.h"
 #include "Engine/Canvas.h"
 #include "GameFramework/HUD.h"
-#include "Engine/DebugCameraHUD.h"
-#include "Debug/LogVisualizerHUD.h"
 #include "Debug/GameplayDebuggingHUDComponent.h"
 #include "Atmosphere/AtmosphericFog.h"
 #include "Engine/ExponentialHeightFog.h"
-#include "GameFramework/GameNetworkManager.h"
 #include "GameFramework/GameSession.h"
 #include "GameFramework/GameState.h"
 #include "GameFramework/PlayerState.h"
