TTP 336037 : Fix changing defaults of CanvasRenderTarget2D-based BP

diff --git a/Engine/Source/Runtime/Engine/Private/CanvasRenderTarget2D.cpp b/Engine/Source/Runtime/Engine/Private/CanvasRenderTarget2D.cpp
--- a/Engine/Source/Runtime/Engine/Private/CanvasRenderTarget2D.cpp
+++ b/Engine/Source/Runtime/Engine/Private/CanvasRenderTarget2D.cpp
@@ -14,6 +14,12 @@ void UCanvasRenderTarget2D::UpdateResource()
 	// Call parent implementation
 	Super::UpdateResource();
 	
+	// Don't allocate canvas object for CRT2D CDO
+	if(IsTemplate())
+	{
+		return;
+	}
+
 	// Create or find the canvas object to use to render onto the texture.  Multiple canvas render target textures can share the same canvas.
 	static const FName CanvasName( TEXT( "CanvasRenderTarget2DCanvas" ) );
 	UCanvas* Canvas = (UCanvas*)StaticFindObjectFast(UCanvas::StaticClass(), GetTransientPackage(), CanvasName );
