Changing how component instance data works.  Component classes no longer are responsible for adding to the cache or determining which instance data in the cache is relevant to them.
When moving a child actor component around its name will not change in the outliner (TTP #335129)

diff --git a/Engine/Source/Runtime/Engine/Classes/Atmosphere/AtmosphericFogComponent.h b/Engine/Source/Runtime/Engine/Classes/Atmosphere/AtmosphericFogComponent.h
--- a/Engine/Source/Runtime/Engine/Classes/Atmosphere/AtmosphericFogComponent.h
+++ b/Engine/Source/Runtime/Engine/Classes/Atmosphere/AtmosphericFogComponent.h
@@ -83,8 +83,6 @@ struct FAtmospherePrecomputeParameters
 	}
 };
 
-
-
 /**
  *	Used to create fogging effects such as clouds.
  */
@@ -260,10 +258,13 @@ public:
 	ENGINE_API void ReleaseResource();
 
 	// Begin UActorComponent interface.
-	virtual void GetComponentInstanceData(FComponentInstanceDataCache& Cache) const OVERRIDE;
-	virtual void ApplyComponentInstanceData(const FComponentInstanceDataCache& Cache) OVERRIDE;
+	virtual TSharedPtr<FComponentInstanceDataBase> GetComponentInstanceData() const OVERRIDE;
+	virtual FName GetComponentInstanceDataType() const OVERRIDE;
+	virtual void ApplyComponentInstanceData(TSharedPtr<FComponentInstanceDataBase> ComponentInstanceData) OVERRIDE;
 	// End UActorComponent interface.
 
+	const FAtmospherePrecomputeParameters& GetPrecomputeParameters() const { return PrecomputeParams;  }
+
 private:
 #if WITH_EDITORONLY_DATA
 	class FAtmospherePrecomputeDataHandler* PrecomputeDataHandler;
diff --git a/Engine/Source/Runtime/Engine/Classes/Components/ActorComponent.h b/Engine/Source/Runtime/Engine/Classes/Components/ActorComponent.h
--- a/Engine/Source/Runtime/Engine/Classes/Components/ActorComponent.h
+++ b/Engine/Source/Runtime/Engine/Classes/Components/ActorComponent.h
@@ -419,9 +419,13 @@ public:
 	virtual void PostNetReceive() { }
 
 	/** Called before we throw away components during RerunConstructionScripts, to cache any data we wish to persist across that operation */
-	virtual void GetComponentInstanceData(class FComponentInstanceDataCache& Cache) const {}
+	virtual TSharedPtr<class FComponentInstanceDataBase> GetComponentInstanceData() const { return NULL; }
+
+	/** The type of the component instance data that this component is interested in */
+	virtual FName GetComponentInstanceDataType() const { return NAME_None; }
+
 	/** Called after we create new components during RerunConstructionScripts, to optionally apply any data backed up during GetComponentInstanceData */
-	virtual void ApplyComponentInstanceData(const class FComponentInstanceDataCache& Cache) {}
+	virtual void ApplyComponentInstanceData(TSharedPtr<class FComponentInstanceDataBase> ComponentInstanceData ) {}
 
 	// Begin UObject interface.
 	virtual void BeginDestroy() OVERRIDE;
diff --git a/Engine/Source/Runtime/Engine/Classes/Components/ChildActorComponent.h b/Engine/Source/Runtime/Engine/Classes/Components/ChildActorComponent.h
--- a/Engine/Source/Runtime/Engine/Classes/Components/ChildActorComponent.h
+++ b/Engine/Source/Runtime/Engine/Classes/Components/ChildActorComponent.h
@@ -18,9 +18,16 @@ class UChildActorComponent : public USceneComponent
 	UPROPERTY(BlueprintReadOnly, Category = ChildActorComponent, TextExportTransient)
 	AActor*	ChildActor;
 
+	/** We try to keep the child actor's name as best we can, so we store it off here when destroying */
+	FName ChildActorName;
+
 	// Begin ActorComponent interface.
 	virtual void OnComponentCreated() OVERRIDE;
 	virtual void OnComponentDestroyed() OVERRIDE;
+	virtual TSharedPtr<FComponentInstanceDataBase> GetComponentInstanceData() const OVERRIDE;
+	virtual FName GetComponentInstanceDataType() const OVERRIDE;
+	virtual void ApplyComponentInstanceData(TSharedPtr<FComponentInstanceDataBase> ComponentInstanceData) OVERRIDE;
+
 	// End ActorComponent interface.
 
 	/** Create the child actor */
diff --git a/Engine/Source/Runtime/Engine/Classes/Components/InstancedStaticMeshComponent.h b/Engine/Source/Runtime/Engine/Classes/Components/InstancedStaticMeshComponent.h
--- a/Engine/Source/Runtime/Engine/Classes/Components/InstancedStaticMeshComponent.h
+++ b/Engine/Source/Runtime/Engine/Classes/Components/InstancedStaticMeshComponent.h
@@ -120,8 +120,9 @@ public:
 	TArray<FBodyInstance*> InstanceBodies;
 
 	// Begin UActorComponent interface 
-	virtual void GetComponentInstanceData(FComponentInstanceDataCache& Cache) const OVERRIDE;
-	virtual void ApplyComponentInstanceData(const FComponentInstanceDataCache& Cache) OVERRIDE;
+	virtual TSharedPtr<FComponentInstanceDataBase> GetComponentInstanceData() const OVERRIDE;
+	virtual FName GetComponentInstanceDataType() const OVERRIDE;
+	virtual void ApplyComponentInstanceData(TSharedPtr<FComponentInstanceDataBase> ComponentInstanceData) OVERRIDE;
 	// End UActorComponent interface 
 
 	// Begin UPrimitiveComponent Interface
diff --git a/Engine/Source/Runtime/Engine/Classes/Components/LightComponent.h b/Engine/Source/Runtime/Engine/Classes/Components/LightComponent.h
--- a/Engine/Source/Runtime/Engine/Classes/Components/LightComponent.h
+++ b/Engine/Source/Runtime/Engine/Classes/Components/LightComponent.h
@@ -4,7 +4,6 @@
 #pragma once
 #include "LightComponent.generated.h"
 
-
 UCLASS(abstract, HideCategories=(Trigger,Activation,"Components|Activation",Physics), ShowCategories=(Mobility))
 class ENGINE_API ULightComponent : public ULightComponentBase
 {
@@ -245,8 +244,9 @@ public:
 #endif // WITH_EDITOR
 	// End UObject interface.
 
-	virtual void GetComponentInstanceData(FComponentInstanceDataCache& Cache) const OVERRIDE;
-	virtual void ApplyComponentInstanceData(const FComponentInstanceDataCache& Cache) OVERRIDE;
+	virtual TSharedPtr<FComponentInstanceDataBase> GetComponentInstanceData() const OVERRIDE;
+	virtual FName GetComponentInstanceDataType() const OVERRIDE;
+	virtual void ApplyComponentInstanceData(TSharedPtr<FComponentInstanceDataBase> ComponentInstanceData) OVERRIDE;
 
 	/** @return number of material elements in this primitive */
 	virtual int32 GetNumMaterials() const;
diff --git a/Engine/Source/Runtime/Engine/Classes/Components/SkyLightComponent.h b/Engine/Source/Runtime/Engine/Classes/Components/SkyLightComponent.h
--- a/Engine/Source/Runtime/Engine/Classes/Components/SkyLightComponent.h
+++ b/Engine/Source/Runtime/Engine/Classes/Components/SkyLightComponent.h
@@ -122,8 +122,9 @@ class ENGINE_API USkyLightComponent : public ULightComponentBase
 	virtual bool IsReadyForFinishDestroy() OVERRIDE;
 	// End UObject Interface
 
-	virtual void GetComponentInstanceData(FComponentInstanceDataCache& Cache) const OVERRIDE;
-	virtual void ApplyComponentInstanceData(const FComponentInstanceDataCache& Cache) OVERRIDE;
+	virtual TSharedPtr<FComponentInstanceDataBase> GetComponentInstanceData() const OVERRIDE;
+	virtual FName GetComponentInstanceDataType() const OVERRIDE;
+	virtual void ApplyComponentInstanceData(TSharedPtr<FComponentInstanceDataBase> ComponentInstanceData) OVERRIDE;
 
 	/** Called each tick to recapture and queued sky captures. */
 	static void UpdateSkyCaptureContents(UWorld* WorldToUpdate);
diff --git a/Engine/Source/Runtime/Engine/Classes/Components/SplineComponent.h b/Engine/Source/Runtime/Engine/Classes/Components/SplineComponent.h
--- a/Engine/Source/Runtime/Engine/Classes/Components/SplineComponent.h
+++ b/Engine/Source/Runtime/Engine/Classes/Components/SplineComponent.h
@@ -20,8 +20,9 @@ class ENGINE_API USplineComponent : public USceneComponent
 	FInterpCurveFloat SplineReparamTable;
 
 	// Begin UActorComponent interface.
-	virtual void GetComponentInstanceData(FComponentInstanceDataCache& Cache) const OVERRIDE;
-	virtual void ApplyComponentInstanceData(const FComponentInstanceDataCache& Cache) OVERRIDE;
+	virtual TSharedPtr<FComponentInstanceDataBase> GetComponentInstanceData() const OVERRIDE;
+	virtual FName GetComponentInstanceDataType() const OVERRIDE;
+	virtual void ApplyComponentInstanceData(TSharedPtr<FComponentInstanceDataBase> ComponentInstanceData) OVERRIDE;
 	// End UActorComponent interface.
 
 	/** Update the SplineReparamTable */
diff --git a/Engine/Source/Runtime/Engine/Classes/Components/StaticMeshComponent.h b/Engine/Source/Runtime/Engine/Classes/Components/StaticMeshComponent.h
--- a/Engine/Source/Runtime/Engine/Classes/Components/StaticMeshComponent.h
+++ b/Engine/Source/Runtime/Engine/Classes/Components/StaticMeshComponent.h
@@ -227,8 +227,9 @@ public:
 #if WITH_EDITOR
 	virtual void CheckForErrors() OVERRIDE;
 #endif
-	virtual void GetComponentInstanceData(FComponentInstanceDataCache& Cache) const OVERRIDE;
-	virtual void ApplyComponentInstanceData(const FComponentInstanceDataCache& Cache) OVERRIDE;
+	virtual TSharedPtr<FComponentInstanceDataBase> GetComponentInstanceData() const OVERRIDE;
+	virtual FName GetComponentInstanceDataType() const OVERRIDE;
+	virtual void ApplyComponentInstanceData(TSharedPtr<FComponentInstanceDataBase> ComponentInstanceData) OVERRIDE;
 	// End UActorComponent interface.
 
 
diff --git a/Engine/Source/Runtime/Engine/Classes/GameFramework/Actor.h b/Engine/Source/Runtime/Engine/Classes/GameFramework/Actor.h
--- a/Engine/Source/Runtime/Engine/Classes/GameFramework/Actor.h
+++ b/Engine/Source/Runtime/Engine/Classes/GameFramework/Actor.h
@@ -1711,7 +1711,7 @@ public:
 	 * @param	Transform			The transform to construct the actor at.
 	 * @param	InstanceDataCache	Optional cache of state to apply to newly created components (e.g. precomputed lighting)
 	 */
-	void ExecuteConstruction(const FTransform& Transform, const FComponentInstanceDataCache* InstanceDataCache);
+	void ExecuteConstruction(const FTransform& Transform, const class FComponentInstanceDataCache* InstanceDataCache);
 
 	/**
 	 * Called when an instance of this class is placed (in editor) or spawned.
diff --git a/Engine/Source/Runtime/Engine/Private/Atmosphere/Atmosphere.cpp b/Engine/Source/Runtime/Engine/Private/Atmosphere/Atmosphere.cpp
--- a/Engine/Source/Runtime/Engine/Private/Atmosphere/Atmosphere.cpp
+++ b/Engine/Source/Runtime/Engine/Private/Atmosphere/Atmosphere.cpp
@@ -686,17 +686,17 @@ void UAtmosphericFogComponent::Serialize(FArchive& Ar)
 class FAtmospherePrecomputeInstanceData : public FComponentInstanceDataBase
 {
 public:
-	static const FName InstanceDataTypeName;
+	FAtmospherePrecomputeInstanceData(const UAtmosphericFogComponent* SourceComponent)
+		: FComponentInstanceDataBase(SourceComponent)
+	{}
 
 	virtual ~FAtmospherePrecomputeInstanceData()
 	{}
 
-	// Begin FComponentInstanceDataBase interface
-	virtual FName GetDataTypeName() const OVERRIDE
+	virtual bool MatchesComponent(const UActorComponent* Component) const OVERRIDE
 	{
-		return InstanceDataTypeName;
+		return (PrecomputeParameter == CastChecked<UAtmosphericFogComponent>(Component)->GetPrecomputeParameters());
 	}
-	// End FComponentInstanceDataBase interface
 
 	struct FAtmospherePrecomputeParameters PrecomputeParameter;
 
@@ -705,15 +705,21 @@ public:
 	FByteBulkData InscatterData;
 };
 
-const FName FAtmospherePrecomputeInstanceData::InstanceDataTypeName(TEXT("AtmospherePrecomputedInstanceData"));
+FName UAtmosphericFogComponent::GetComponentInstanceDataType() const
+{
+	static const FName InstanceDataTypeName(TEXT("AtmospherePrecomputedInstanceData"));
+	return InstanceDataTypeName;
+}
 
 // Backup the precomputed data before re-running Blueprint construction script
-void UAtmosphericFogComponent::GetComponentInstanceData(FComponentInstanceDataCache& Cache) const
+TSharedPtr<FComponentInstanceDataBase> UAtmosphericFogComponent::GetComponentInstanceData() const
 {
+	TSharedPtr<FAtmospherePrecomputeInstanceData> PrecomputedData;
+
 	if (TransmittanceData.GetElementCount() && IrradianceData.GetElementCount() && InscatterData.GetElementCount() && PrecomputeCounter.GetValue() == EValid)
 	{
 		// Allocate new struct for holding light map data
-		TSharedRef<FAtmospherePrecomputeInstanceData> PrecomputedData = MakeShareable(new FAtmospherePrecomputeInstanceData());
+		 PrecomputedData = MakeShareable(new FAtmospherePrecomputeInstanceData(this));
 
 		// Fill in info
 		PrecomputedData->PrecomputeParameter = PrecomputeParams;
@@ -740,55 +746,44 @@ void UAtmosphericFogComponent::GetComponentInstanceData(FComponentInstanceDataCa
 			InscatterData.GetCopy(&OutData, false);
 			PrecomputedData->InscatterData.Unlock();
 		}
-
-		// Add to cache
-		Cache.AddInstanceData(PrecomputedData);
 	}
+
+	return PrecomputedData;
 }
 
 // Restore the precomputed data after re-running Blueprint construction script
-void UAtmosphericFogComponent::ApplyComponentInstanceData(const FComponentInstanceDataCache& Cache)
+void UAtmosphericFogComponent::ApplyComponentInstanceData(TSharedPtr<FComponentInstanceDataBase> ComponentInstanceData)
 {
-	TArray< TSharedPtr<FComponentInstanceDataBase> > CachedData;
-	Cache.GetInstanceDataOfType(FAtmospherePrecomputeInstanceData::InstanceDataTypeName, CachedData);
-
-	for (int32 DataIdx = 0; DataIdx<CachedData.Num(); DataIdx++)
-	{
-		check(CachedData[DataIdx].IsValid());
-		check(CachedData[DataIdx]->GetDataTypeName() == FAtmospherePrecomputeInstanceData::InstanceDataTypeName);
-		TSharedPtr<FAtmospherePrecomputeInstanceData> PrecomputedData = StaticCastSharedPtr<FAtmospherePrecomputeInstanceData>(CachedData[DataIdx]);
+	check(ComponentInstanceData.IsValid());
+	TSharedPtr<FAtmospherePrecomputeInstanceData> PrecomputedData = StaticCastSharedPtr<FAtmospherePrecomputeInstanceData>(ComponentInstanceData);
 
-		if (PrecomputedData->PrecomputeParameter == PrecomputeParams)
-		{
-			FComponentReregisterContext ReregisterContext(this);
-			ReleaseResource();
-
-			{
-				int32 TotalByte = PrecomputedData->TransmittanceData.GetBulkDataSize();
-				TransmittanceData.Lock(LOCK_READ_WRITE);
-				void* OutData = TransmittanceData.Realloc(TotalByte);
-				PrecomputedData->TransmittanceData.GetCopy(&OutData, false);
-				TransmittanceData.Unlock();
-			}
+	FComponentReregisterContext ReregisterContext(this);
+	ReleaseResource();
 
-			{
-				int32 TotalByte = PrecomputedData->IrradianceData.GetBulkDataSize();
-				IrradianceData.Lock(LOCK_READ_WRITE);
-				void* OutData = IrradianceData.Realloc(TotalByte);
-				PrecomputedData->IrradianceData.GetCopy(&OutData, false);
-				IrradianceData.Unlock();
-			}
+	{
+		int32 TotalByte = PrecomputedData->TransmittanceData.GetBulkDataSize();
+		TransmittanceData.Lock(LOCK_READ_WRITE);
+		void* OutData = TransmittanceData.Realloc(TotalByte);
+		PrecomputedData->TransmittanceData.GetCopy(&OutData, false);
+		TransmittanceData.Unlock();
+	}
 
-			{
-				int32 TotalByte = PrecomputedData->InscatterData.GetBulkDataSize();
-				InscatterData.Lock(LOCK_READ_WRITE);
-				void* OutData = InscatterData.Realloc(TotalByte);
-				PrecomputedData->InscatterData.GetCopy(&OutData, false);
-				InscatterData.Unlock();
-			}
+	{
+		int32 TotalByte = PrecomputedData->IrradianceData.GetBulkDataSize();
+		IrradianceData.Lock(LOCK_READ_WRITE);
+		void* OutData = IrradianceData.Realloc(TotalByte);
+		PrecomputedData->IrradianceData.GetCopy(&OutData, false);
+		IrradianceData.Unlock();
+	}
 
-			PrecomputeCounter.Set(EValid);
-			InitResource();
-		}
+	{
+		int32 TotalByte = PrecomputedData->InscatterData.GetBulkDataSize();
+		InscatterData.Lock(LOCK_READ_WRITE);
+		void* OutData = InscatterData.Realloc(TotalByte);
+		PrecomputedData->InscatterData.GetCopy(&OutData, false);
+		InscatterData.Unlock();
 	}
+
+	PrecomputeCounter.Set(EValid);
+	InitResource();
 }
diff --git a/Engine/Source/Runtime/Engine/Private/ChildActorComponent.cpp b/Engine/Source/Runtime/Engine/Private/ChildActorComponent.cpp
--- a/Engine/Source/Runtime/Engine/Private/ChildActorComponent.cpp
+++ b/Engine/Source/Runtime/Engine/Private/ChildActorComponent.cpp
@@ -1,6 +1,7 @@
 // Copyright 1998-2014 Epic Games, Inc. All Rights Reserved.
 
 #include "EnginePrivate.h"
+#include "ComponentInstanceDataCache.h"
 
 DEFINE_LOG_CATEGORY_STATIC(LogChildActorComponent, Warning, All);
 
@@ -23,6 +24,47 @@ void UChildActorComponent::OnComponentDestroyed()
 	DestroyChildActor();
 }
 
+class FChildActorComponentInstanceData : public FComponentInstanceDataBase
+{
+public:
+	FChildActorComponentInstanceData(const UChildActorComponent* Component)
+		: FComponentInstanceDataBase(Component)
+		, ChildActorName(Component->ChildActorName)
+	{
+	}
+
+	FName ChildActorName;
+};
+
+FName UChildActorComponent::GetComponentInstanceDataType() const
+{
+	static const FName ChildActorComponentInstanceDataName(TEXT("ChildActorInstanceData"));
+	return ChildActorComponentInstanceDataName;
+}
+
+TSharedPtr<FComponentInstanceDataBase> UChildActorComponent::GetComponentInstanceData() const
+{
+	TSharedPtr<FChildActorComponentInstanceData> InstanceData;
+	if (!ChildActorName.IsNone())
+	{
+		// Allocate new struct for holding light map data
+		 InstanceData = MakeShareable(new FChildActorComponentInstanceData(this));
+	}
+
+	return InstanceData;
+}
+
+void UChildActorComponent::ApplyComponentInstanceData(TSharedPtr<FComponentInstanceDataBase> ComponentInstanceData)
+{
+	check(ComponentInstanceData.IsValid());
+	ChildActorName = StaticCastSharedPtr<FChildActorComponentInstanceData>(ComponentInstanceData)->ChildActorName;
+	if (ChildActor)
+	{
+		ChildActor->Rename(*ChildActorName.ToString(), NULL, REN_DoNotDirty);
+	}
+}
+
+
 void UChildActorComponent::CreateChildActor()
 {
 	// Kill spawned actor if we have one
@@ -54,6 +96,7 @@ void UChildActorComponent::CreateChildActor()
 			Params.bDeferConstruction = true; // We defer construction so that we set ParentComponentActor prior to component registration so they appear selected
 			Params.bAllowDuringConstructionScript = true;
 			Params.OverrideLevel = GetOwner()->GetLevel();
+			Params.Name = ChildActorName;
 
 			// Spawn actor of desired class
 			FVector Location = GetComponentLocation();
@@ -63,6 +106,8 @@ void UChildActorComponent::CreateChildActor()
 			// If spawn was successful, 
 			if(ChildActor != NULL)
 			{
+				ChildActorName = ChildActor->GetFName();
+
 #if WITH_EDITOR
 				// Remember which actor spawned it (for selection in editor etc)
 				ChildActor->ParentComponentActor = GetOwner();
@@ -92,6 +137,7 @@ void UChildActorComponent::DestroyChildActor()
 			// World may be NULL during shutdown
 			if(World != NULL)
 			{
+				ChildActor->Rename(*MakeUniqueObjectName(GetOuter(), ChildActor->GetClass(), TEXT("DESTROYED_CHILDACTOR")).ToString(), NULL, REN_DoNotDirty);
 				World->DestroyActor(ChildActor);
 			}
 		}
diff --git a/Engine/Source/Runtime/Engine/Private/ComponentInstanceDataCache.cpp b/Engine/Source/Runtime/Engine/Private/ComponentInstanceDataCache.cpp
--- a/Engine/Source/Runtime/Engine/Private/ComponentInstanceDataCache.cpp
+++ b/Engine/Source/Runtime/Engine/Private/ComponentInstanceDataCache.cpp
@@ -63,20 +63,7 @@ bool FComponentInstanceDataBase::MatchesComponent(const UActorComponent* Compone
 	return bMatches;
 }
 
-void FComponentInstanceDataCache::AddInstanceData(TSharedPtr<FComponentInstanceDataBase> NewData)
-{
-	check(NewData.IsValid());
-	FName TypeName = NewData->GetDataTypeName();
-	TypeToDataMap.Add(TypeName, NewData);
-}
-
-void FComponentInstanceDataCache::GetInstanceDataOfType(FName TypeName, TArray< TSharedPtr<FComponentInstanceDataBase> >& OutData) const
-{
-	TypeToDataMap.MultiFind(TypeName, OutData);
-}
-
-
-void FComponentInstanceDataCache::GetFromActor(AActor* Actor, FComponentInstanceDataCache& Cache)
+FComponentInstanceDataCache::FComponentInstanceDataCache(AActor* Actor)
 {
 	if(Actor != NULL)
 	{
@@ -84,12 +71,16 @@ void FComponentInstanceDataCache::GetFromActor(AActor* Actor, FComponentInstance
 		Actor->GetComponents(Components);
 
 		// Grab per-instance data we want to persist
-		for (int32 CompIdx=0; CompIdx<Components.Num(); CompIdx++)
+		for (UActorComponent* Component : Components)
 		{
-			UActorComponent* Component = Components[CompIdx];
 			if(Component->bCreatedByConstructionScript) // Only cache data from 'created by construction script' components
 			{
-				Component->GetComponentInstanceData(Cache);
+				TSharedPtr<FComponentInstanceDataBase> ComponentInstanceData = Component->GetComponentInstanceData();
+				if (ComponentInstanceData.IsValid())
+				{
+					check(!Component->GetComponentInstanceDataType().IsNone());
+					TypeToDataMap.Add(Component->GetComponentInstanceDataType(), ComponentInstanceData);
+				}
 			}
 		}
 	}
@@ -103,12 +94,26 @@ void FComponentInstanceDataCache::ApplyToActor(AActor* Actor) const
 		Actor->GetComponents(Components);
 
 		// Apply per-instance data.
-		for (int32 CompIdx=0; CompIdx<Components.Num(); CompIdx++)
+		for (UActorComponent* Component : Components)
 		{
-			UActorComponent* Component = Components[CompIdx];
 			if(Component->bCreatedByConstructionScript) // Only try and apply data to 'created by construction script' components
 			{
-				Component->ApplyComponentInstanceData(*this);
+				const FName ComponentInstanceDataType = Component->GetComponentInstanceDataType();
+
+				if (!ComponentInstanceDataType.IsNone())
+				{
+					TArray< TSharedPtr<FComponentInstanceDataBase> > CachedData;
+					TypeToDataMap.MultiFind(ComponentInstanceDataType, CachedData);
+
+					for (TSharedPtr<FComponentInstanceDataBase> ComponentInstanceData : CachedData)
+					{
+						if (ComponentInstanceData.IsValid() && ComponentInstanceData->MatchesComponent(Component))
+						{
+							Component->ApplyComponentInstanceData(ComponentInstanceData);
+							break;
+						}
+					}
+				}
 			}
 		}
 	}
diff --git a/Engine/Source/Runtime/Engine/Private/InstancedStaticMesh.cpp b/Engine/Source/Runtime/Engine/Private/InstancedStaticMesh.cpp
--- a/Engine/Source/Runtime/Engine/Private/InstancedStaticMesh.cpp
+++ b/Engine/Source/Runtime/Engine/Private/InstancedStaticMesh.cpp
@@ -1133,33 +1133,32 @@ public:
 	{
 	}
 
-	virtual FName GetDataTypeName() const
-	{
-		return SourceComponentName;
-	}
-
 public:
 	/** The cached selected instances */
 	TBitArray<> SelectedInstances;
 };
 #endif
 
-void UInstancedStaticMeshComponent::GetComponentInstanceData(FComponentInstanceDataCache& Cache) const
+FName UInstancedStaticMeshComponent::GetComponentInstanceDataType() const
+{
+	static const FName InstanceStaticMeshComponentInstanceDataName(TEXT("InstancedStaticMeshSelectionData"));
+	return InstanceStaticMeshComponentInstanceDataName;
+}
+
+TSharedPtr<FComponentInstanceDataBase> UInstancedStaticMeshComponent::GetComponentInstanceData() const
 {
+	TSharedPtr<FInstancedStaticMeshSelectionData> InstanceData;
 #if WITH_EDITOR
-	Cache.AddInstanceData(MakeShareable(new FInstancedStaticMeshSelectionData(*this)));
+	InstanceData = MakeShareable(new FInstancedStaticMeshSelectionData(*this));
 #endif
+	return InstanceData;
 }
 
-void UInstancedStaticMeshComponent::ApplyComponentInstanceData(const FComponentInstanceDataCache& Cache)
+void UInstancedStaticMeshComponent::ApplyComponentInstanceData(TSharedPtr<FComponentInstanceDataBase> ComponentInstanceData)
 {
 #if WITH_EDITOR
-	TArray< TSharedPtr<FComponentInstanceDataBase> > Data;
-	Cache.GetInstanceDataOfType(GetFName(), Data);
-	if(Data.Num() > 0)
-	{
-		SelectedInstances = StaticCastSharedPtr<FInstancedStaticMeshSelectionData>(Data[0])->SelectedInstances;
-	}
+	check(ComponentInstanceData.IsValid());
+	SelectedInstances = StaticCastSharedPtr<FInstancedStaticMeshSelectionData>(ComponentInstanceData)->SelectedInstances;
 #endif
 }
 
diff --git a/Engine/Source/Runtime/Engine/Private/LightComponent.cpp b/Engine/Source/Runtime/Engine/Private/LightComponent.cpp
--- a/Engine/Source/Runtime/Engine/Private/LightComponent.cpp
+++ b/Engine/Source/Runtime/Engine/Private/LightComponent.cpp
@@ -755,15 +755,22 @@ void ULightComponent::InvalidateLightingCacheInner(bool bRecreateLightGuids)
 class FPrecomputedLightInstanceData : public FComponentInstanceDataBase
 {
 public:
-	static const FName PrecomputedLightInstanceDataTypeName;
-
+	FPrecomputedLightInstanceData(const ULightComponent* SourceComponent)
+		: FComponentInstanceDataBase(SourceComponent)
+		, Transform(SourceComponent->ComponentToWorld)
+		, LightGuid(SourceComponent->LightGuid)
+		, ShadowMapChannel(SourceComponent->ShadowMapChannel)
+		, PreviewShadowMapChannel(SourceComponent->PreviewShadowMapChannel)
+		, bPrecomputedLightingIsValid(SourceComponent->bPrecomputedLightingIsValid)
+	{}
+			
 	virtual ~FPrecomputedLightInstanceData()
 	{}
 
 	// Begin FComponentInstanceDataBase interface
-	virtual FName GetDataTypeName() const OVERRIDE
+	virtual bool MatchesComponent(const UActorComponent* Component) const OVERRIDE
 	{
-		return PrecomputedLightInstanceDataTypeName;
+		return Transform.Equals(CastChecked<ULightComponent>(Component)->ComponentToWorld);
 	}
 	// End FComponentInstanceDataBase interface
 
@@ -774,53 +781,34 @@ public:
 	bool bPrecomputedLightingIsValid;
 };
 
-// Init type name static
-const FName FPrecomputedLightInstanceData::PrecomputedLightInstanceDataTypeName(TEXT("PrecomputedLightInstanceData"));
+FName ULightComponent::GetComponentInstanceDataType() const
+{
+	static const FName PrecomputedLightInstanceDataTypeName(TEXT("PrecomputedLightInstanceData"));
+	return PrecomputedLightInstanceDataTypeName;
+}
 
-void ULightComponent::GetComponentInstanceData(FComponentInstanceDataCache& Cache) const
+TSharedPtr<FComponentInstanceDataBase> ULightComponent::GetComponentInstanceData() const
 {
 	// Allocate new struct for holding light map data
-	TSharedRef<FPrecomputedLightInstanceData> LightMapData = MakeShareable(new FPrecomputedLightInstanceData());
-
-	// Fill in info
-	LightMapData->Transform = ComponentToWorld;
-	LightMapData->LightGuid = LightGuid;
-	LightMapData->ShadowMapChannel = ShadowMapChannel;
-	LightMapData->PreviewShadowMapChannel = PreviewShadowMapChannel;
-	LightMapData->bPrecomputedLightingIsValid = bPrecomputedLightingIsValid;
-
-	// Add to cache
-	Cache.AddInstanceData(LightMapData);
+	return MakeShareable(new FPrecomputedLightInstanceData(this));
 }
 
-void ULightComponent::ApplyComponentInstanceData(const FComponentInstanceDataCache& Cache)
+void ULightComponent::ApplyComponentInstanceData(TSharedPtr<FComponentInstanceDataBase> ComponentInstanceData)
 {
-	TArray< TSharedPtr<FComponentInstanceDataBase> > CachedData;
-	Cache.GetInstanceDataOfType(FPrecomputedLightInstanceData::PrecomputedLightInstanceDataTypeName, CachedData);
+	check(ComponentInstanceData.IsValid());
+	TSharedPtr<FPrecomputedLightInstanceData> LightMapData = StaticCastSharedPtr<FPrecomputedLightInstanceData>(ComponentInstanceData);
 
-	for (int32 DataIdx = 0; DataIdx<CachedData.Num(); DataIdx++)
-	{
-		check(CachedData[DataIdx].IsValid());
-		check(CachedData[DataIdx]->GetDataTypeName() == FPrecomputedLightInstanceData::PrecomputedLightInstanceDataTypeName);
-		TSharedPtr<FPrecomputedLightInstanceData> LightMapData = StaticCastSharedPtr<FPrecomputedLightInstanceData>(CachedData[DataIdx]);
+	LightGuid = LightMapData->LightGuid;
+	ShadowMapChannel = LightMapData->ShadowMapChannel;
+	PreviewShadowMapChannel = LightMapData->PreviewShadowMapChannel;
+	bPrecomputedLightingIsValid = LightMapData->bPrecomputedLightingIsValid;
 
-		// See if data matches current state
-		//@todo - compare all state that affects static lighting like radius, indirect lighting intensity, etc
-		if (LightMapData->Transform.Equals(ComponentToWorld))
-		{
-			LightGuid = LightMapData->LightGuid;
-			ShadowMapChannel = LightMapData->ShadowMapChannel;
-			PreviewShadowMapChannel = LightMapData->PreviewShadowMapChannel;
-			bPrecomputedLightingIsValid = LightMapData->bPrecomputedLightingIsValid;
-
-			MarkRenderStateDirty();
+	MarkRenderStateDirty();
 
 #if WITH_EDITOR
-			// Update the icon with the new state of PreviewShadowMapChannel
-			UpdateLightSpriteTexture();
+	// Update the icon with the new state of PreviewShadowMapChannel
+	UpdateLightSpriteTexture();
 #endif
-		}
-	}
 }
 
 int32 ULightComponent::GetNumMaterials() const
diff --git a/Engine/Source/Runtime/Engine/Private/SkyLightComponent.cpp b/Engine/Source/Runtime/Engine/Private/SkyLightComponent.cpp
--- a/Engine/Source/Runtime/Engine/Private/SkyLightComponent.cpp
+++ b/Engine/Source/Runtime/Engine/Private/SkyLightComponent.cpp
@@ -293,18 +293,10 @@ bool USkyLightComponent::IsReadyForFinishDestroy()
 class FPrecomputedSkyLightInstanceData : public FComponentInstanceDataBase
 {
 public:
-	static const FName PrecomputedSkyLightInstanceDataTypeName;
-
-	virtual ~FPrecomputedSkyLightInstanceData()
+	FPrecomputedSkyLightInstanceData(const USkyLightComponent* SourceComponent)
+		: FComponentInstanceDataBase(SourceComponent)
 	{}
 
-	// Begin FComponentInstanceDataBase interface
-	virtual FName GetDataTypeName() const OVERRIDE
-	{
-		return PrecomputedSkyLightInstanceDataTypeName;
-	}
-	// End FComponentInstanceDataBase interface
-
 	FGuid LightGuid;
 	bool bPrecomputedLightingIsValid;
 	// This has to be refcounted to keep it alive during the handoff without doing a deep copy
@@ -313,48 +305,41 @@ public:
 };
 
 // Init type name static
-const FName FPrecomputedSkyLightInstanceData::PrecomputedSkyLightInstanceDataTypeName(TEXT("PrecomputedSkyLightInstanceData"));
-
-void USkyLightComponent::GetComponentInstanceData(FComponentInstanceDataCache& Cache) const
+FName USkyLightComponent::GetComponentInstanceDataType() const
 {
-	// Allocate new struct for holding light map data
-	TSharedRef<FPrecomputedSkyLightInstanceData> LightMapData = MakeShareable(new FPrecomputedSkyLightInstanceData());
-
-	// Fill in info
-	LightMapData->LightGuid = LightGuid;
-	LightMapData->bPrecomputedLightingIsValid = bPrecomputedLightingIsValid;
-	LightMapData->ProcessedSkyTexture = ProcessedSkyTexture;
-	LightMapData->IrradianceEnvironmentMap = IrradianceEnvironmentMap;
-
-	// Add to cache
-	Cache.AddInstanceData(LightMapData);
+	static const FName PrecomputedSkyLightInstanceDataTypeName(TEXT("PrecomputedSkyLightInstanceData"));
+	return PrecomputedSkyLightInstanceDataTypeName;
 }
 
-void USkyLightComponent::ApplyComponentInstanceData(const FComponentInstanceDataCache& Cache)
+TSharedPtr<FComponentInstanceDataBase> USkyLightComponent::GetComponentInstanceData() const
 {
-	TArray< TSharedPtr<FComponentInstanceDataBase> > CachedData;
-	Cache.GetInstanceDataOfType(FPrecomputedSkyLightInstanceData::PrecomputedSkyLightInstanceDataTypeName, CachedData);
+	TSharedRef<FPrecomputedSkyLightInstanceData> InstanceData = MakeShareable(new FPrecomputedSkyLightInstanceData(this));
+	InstanceData->LightGuid = LightGuid;
+	InstanceData->bPrecomputedLightingIsValid = bPrecomputedLightingIsValid;
+	InstanceData->ProcessedSkyTexture = ProcessedSkyTexture;
+	InstanceData->IrradianceEnvironmentMap = IrradianceEnvironmentMap;
 
-	for (int32 DataIdx = 0; DataIdx<CachedData.Num(); DataIdx++)
-	{
-		check(CachedData[DataIdx].IsValid());
-		check(CachedData[DataIdx]->GetDataTypeName() == FPrecomputedSkyLightInstanceData::PrecomputedSkyLightInstanceDataTypeName);
-		TSharedPtr<FPrecomputedSkyLightInstanceData> LightMapData = StaticCastSharedPtr<FPrecomputedSkyLightInstanceData>(CachedData[DataIdx]);
+	return InstanceData;
+}
 
-		LightGuid = LightMapData->LightGuid;
-		bPrecomputedLightingIsValid = LightMapData->bPrecomputedLightingIsValid;
-		ProcessedSkyTexture = LightMapData->ProcessedSkyTexture;
-		IrradianceEnvironmentMap = LightMapData->IrradianceEnvironmentMap;
+void USkyLightComponent::ApplyComponentInstanceData(TSharedPtr<FComponentInstanceDataBase> ComponentInstanceData)
+{
+	check(ComponentInstanceData.IsValid());
+	TSharedPtr<FPrecomputedSkyLightInstanceData> LightMapData = StaticCastSharedPtr<FPrecomputedSkyLightInstanceData>(ComponentInstanceData);
 
-		if (ProcessedSkyTexture && bSavedConstructionScriptValuesValid)
-		{
-			// We have valid capture state, remove the queued update
-			bCaptureDirty = false;
-			SkyCapturesToUpdate.Remove(this);
-		}
+	LightGuid = LightMapData->LightGuid;
+	bPrecomputedLightingIsValid = LightMapData->bPrecomputedLightingIsValid;
+	ProcessedSkyTexture = LightMapData->ProcessedSkyTexture;
+	IrradianceEnvironmentMap = LightMapData->IrradianceEnvironmentMap;
 
-		MarkRenderStateDirty();
+	if (ProcessedSkyTexture && bSavedConstructionScriptValuesValid)
+	{
+		// We have valid capture state, remove the queued update
+		bCaptureDirty = false;
+		SkyCapturesToUpdate.Remove(this);
 	}
+
+	MarkRenderStateDirty();
 }
 
 void USkyLightComponent::UpdateSkyCaptureContents(UWorld* WorldToUpdate)
diff --git a/Engine/Source/Runtime/Engine/Private/SplineComponent.cpp b/Engine/Source/Runtime/Engine/Private/SplineComponent.cpp
--- a/Engine/Source/Runtime/Engine/Private/SplineComponent.cpp
+++ b/Engine/Source/Runtime/Engine/Private/SplineComponent.cpp
@@ -103,54 +103,29 @@ void USplineComponent::RefreshSplineInputs()
 class FSplineInstanceData : public FComponentInstanceDataBase
 {
 public:
-	static const FName SplineInstanceDataTypeName;
-
-	virtual ~FSplineInstanceData()
-	{}
-
-	// Begin FComponentInstanceDataBase interface
-	virtual FName GetDataTypeName() const OVERRIDE
+	FSplineInstanceData(const USplineComponent* SourceComponent)
+		: FComponentInstanceDataBase(SourceComponent)
+		, SplineInfo(SourceComponent->SplineInfo)
 	{
-		return SplineInstanceDataTypeName;
 	}
-	// End FComponentInstanceDataBase interface
-
-	/** Name of component */
-	FName ComponentName;
 
 	/** Lightmaps from LODData */
 	FInterpCurveVector SplineInfo;
 };
 
-// Init type name static
-const FName FSplineInstanceData::SplineInstanceDataTypeName(TEXT("SplineInstanceData"));
-
-void USplineComponent::GetComponentInstanceData(FComponentInstanceDataCache& Cache) const
+FName USplineComponent::GetComponentInstanceDataType() const
 {
-	TSharedRef<FSplineInstanceData> SplineData = MakeShareable(new FSplineInstanceData());
-	SplineData->ComponentName = GetFName();
-	SplineData->SplineInfo = SplineInfo;
-
-	// Add to cache
-	Cache.AddInstanceData(SplineData);
+	static const FName SplineInstanceDataTypeName(TEXT("SplineInstanceData"));
+	return SplineInstanceDataTypeName;
 }
 
-void USplineComponent::ApplyComponentInstanceData(const FComponentInstanceDataCache& Cache)
+TSharedPtr<FComponentInstanceDataBase> USplineComponent::GetComponentInstanceData() const
 {
-	TArray< TSharedPtr<FComponentInstanceDataBase> > CachedData;
-	Cache.GetInstanceDataOfType(FSplineInstanceData::SplineInstanceDataTypeName, CachedData);
+	return MakeShareable(new FSplineInstanceData(this));
+}
 
-	for (int32 DataIdx = 0; DataIdx < CachedData.Num(); DataIdx++)
-	{
-		check(CachedData[DataIdx].IsValid());
-		check(CachedData[DataIdx]->GetDataTypeName() == FSplineInstanceData::SplineInstanceDataTypeName);
-		TSharedPtr<FSplineInstanceData> SplineData = StaticCastSharedPtr<FSplineInstanceData>(CachedData[DataIdx]);
-
-		// See if data matches current state
-		if (SplineData->ComponentName == GetFName())
-		{
-			SplineInfo = SplineData->SplineInfo;
-			UpdateSplineReparamTable();
-		}
-	}
+void USplineComponent::ApplyComponentInstanceData(TSharedPtr<FComponentInstanceDataBase> ComponentInstanceData)
+{
+	SplineInfo = StaticCastSharedPtr<FSplineInstanceData>(ComponentInstanceData)->SplineInfo;
+	UpdateSplineReparamTable();
 }
\ No newline at end of file
diff --git a/Engine/Source/Runtime/Engine/Private/StaticMeshComponent.cpp b/Engine/Source/Runtime/Engine/Private/StaticMeshComponent.cpp
--- a/Engine/Source/Runtime/Engine/Private/StaticMeshComponent.cpp
+++ b/Engine/Source/Runtime/Engine/Private/StaticMeshComponent.cpp
@@ -9,25 +9,19 @@
 
 #define LOCTEXT_NAMESPACE "StaticMeshComponent"
 
-/**
-* Vertex color information stored during RerunConstructionScripts
-*/
-class FVertexColorInstanceData : public FComponentInstanceDataBase
+class FStaticMeshComponentInstanceData : public FComponentInstanceDataBase
 {
 public:
-
-	FVertexColorInstanceData(const UStaticMeshComponent* SourceComponent)
+	FStaticMeshComponentInstanceData(const UStaticMeshComponent* SourceComponent)
 		: FComponentInstanceDataBase(SourceComponent)
 		, StaticMesh(SourceComponent->StaticMesh)
+		, bHasCachedStaticLighting(false)
 	{
 	}
 
-	static const FName VertexColorInstanceDataName;
-
-	/** FComponentInstanceDataBase interface */
-	virtual FName GetDataTypeName() const OVERRIDE
+	virtual bool MatchesComponent(const UActorComponent* Component) const OVERRIDE
 	{
-		return VertexColorInstanceDataName;
+		return (CastChecked<UStaticMeshComponent>(Component)->StaticMesh == StaticMesh && FComponentInstanceDataBase::MatchesComponent(Component));
 	}
 
 	/** Add vertex color data for a specified LOD before RerunConstructionScripts is called */
@@ -79,23 +73,16 @@ public:
 		return bAppliedAnyData;
 	}
 
-	/** Check whether this vertex color data can match the specified component */
-	virtual bool MatchesComponent(const UActorComponent* Component) const OVERRIDE
+	/** Used to store lightmap data during RerunConstructionScripts */
+	struct FLightMapInstanceData
 	{
-		const UStaticMeshComponent* StaticMeshComponent = CastChecked<const UStaticMeshComponent>(Component, ECastCheckedType::NullAllowed);
-		if (StaticMeshComponent != NULL)
-		{
-			if (StaticMeshComponent->StaticMesh != StaticMesh)
-			{
-				return false;
-			}
-
-			return FComponentInstanceDataBase::MatchesComponent(Component);
-		}
-
-		// this component was not found
-		return false;
-	}
+		/** Transform of instance */
+		FTransform		Transform;
+		/** Lightmaps from LODData */
+		TArray<FLightMapRef>	LODDataLightMap;
+		TArray<FShadowMapRef>	LODDataShadowMap;
+		TArray<FGuid> IrrelevantLights;
+	};
 
 	/** Vertex data stored per-LOD */
 	struct FVertexColorLODData
@@ -121,8 +108,10 @@ public:
 
 	/** Array of cached vertex colors for each LOD */
 	TArray<FVertexColorLODData> VertexColorLODs;
+
+	bool bHasCachedStaticLighting;
+	FLightMapInstanceData CachedStaticLighting;
 };
-const FName FVertexColorInstanceData::VertexColorInstanceDataName(TEXT("VertexColorInstanceData"));
 
 UStaticMeshComponent::UStaticMeshComponent(const class FPostConstructInitializeProperties& PCIP)
 	: Super(PCIP)
@@ -1511,125 +1500,78 @@ int32 UStaticMeshComponent::GetSerializedComponentIndex() const
 	return INDEX_NONE;
 }
 
-/** Used to store lightmap data during RerunConstructionScripts */
-class FLightMapInstanceData : public FComponentInstanceDataBase
+FName UStaticMeshComponent::GetComponentInstanceDataType() const
 {
-public:
-	static const FName LightMapInstanceDataTypeName;
-
-	virtual ~FLightMapInstanceData()
-	{}
-
-	// Begin FComponentInstanceDataBase interface
-	virtual FName GetDataTypeName() const OVERRIDE
-	{
-		return LightMapInstanceDataTypeName;
-	}
-	// End FComponentInstanceDataBase interface
-
-	/** Mesh being used by component */
-	class UStaticMesh*	StaticMesh;
-	/** Transform of instance */
-	FTransform		Transform;
-	/** Lightmaps from LODData */
-	TArray<FLightMapRef>	LODDataLightMap;
-	TArray<FShadowMapRef>	LODDataShadowMap;
-	TArray<FGuid> IrrelevantLights;
-};
-const FName FLightMapInstanceData::LightMapInstanceDataTypeName(TEXT("LightMapInstanceData"));
+	static const FName StaticMeshComponentInstanceDataName(TEXT("StaticMeshInstanceData"));
+	return StaticMeshComponentInstanceDataName;
+}
 
-void UStaticMeshComponent::GetComponentInstanceData(FComponentInstanceDataCache& Cache) const
+TSharedPtr<FComponentInstanceDataBase> UStaticMeshComponent::GetComponentInstanceData() const
 {
+	TSharedPtr<FStaticMeshComponentInstanceData> InstanceData;
+
 	// Don't back up static lighting if there isn't any
 	if(bHasCachedStaticLighting)
 	{
-		// Allocate new struct for holding light map data
-		TSharedRef<FLightMapInstanceData> LightMapData = MakeShareable(new FLightMapInstanceData());
+		InstanceData = MakeShareable(new FStaticMeshComponentInstanceData(this));
 
 		// Fill in info
-		LightMapData->StaticMesh = StaticMesh;
-		LightMapData->Transform = ComponentToWorld;
-		LightMapData->IrrelevantLights = IrrelevantLights;
-		LightMapData->LODDataLightMap.Empty(LODData.Num());
-		for (int32 i = 0; i < LODData.Num(); ++i)
+		InstanceData->bHasCachedStaticLighting = true;
+		InstanceData->CachedStaticLighting.Transform = ComponentToWorld;
+		InstanceData->CachedStaticLighting.IrrelevantLights = IrrelevantLights;
+		InstanceData->CachedStaticLighting.LODDataLightMap.Empty(LODData.Num());
+		for (const FStaticMeshComponentLODInfo& LODDataEntry : LODData)
 		{
-			LightMapData->LODDataLightMap.Add(LODData[i].LightMap);
-			LightMapData->LODDataShadowMap.Add(LODData[i].ShadowMap);
+			InstanceData->CachedStaticLighting.LODDataLightMap.Add(LODDataEntry.LightMap);
+			InstanceData->CachedStaticLighting.LODDataShadowMap.Add(LODDataEntry.ShadowMap);
 		}
-
-		// Add to cache
-		Cache.AddInstanceData(LightMapData);
 	}
 
 
 	// Cache instance vertex colors
-	TSharedPtr<FVertexColorInstanceData> VertexColorData;
 	for( int32 LODIndex = 0; LODIndex < LODData.Num(); ++LODIndex )
 	{
 		const FStaticMeshComponentLODInfo& LODInfo = LODData[LODIndex];
 
 		if ( LODInfo.OverrideVertexColors && LODInfo.OverrideVertexColors->GetNumVertices() > 0 && LODInfo.PaintedVertices.Num() > 0 )
 		{
-			if( !VertexColorData.IsValid() )
+			if (!InstanceData.IsValid())
 			{
-				VertexColorData = MakeShareable( new FVertexColorInstanceData(this) );
+				InstanceData = MakeShareable(new FStaticMeshComponentInstanceData(this));
 			}
 
-			VertexColorData->AddVertexColorData( LODInfo, LODIndex );
+			InstanceData->AddVertexColorData(LODInfo, LODIndex);
 		}
 	}
 
-	if( VertexColorData.IsValid() )
-	{
-		Cache.AddInstanceData( VertexColorData );
-	}
-
+	return InstanceData;
 }
 
-void UStaticMeshComponent::ApplyComponentInstanceData(const FComponentInstanceDataCache& Cache)
+void UStaticMeshComponent::ApplyComponentInstanceData(TSharedPtr<FComponentInstanceDataBase> ComponentInstanceData)
 {
-	TArray< TSharedPtr<FComponentInstanceDataBase> > CachedData;
-	Cache.GetInstanceDataOfType(FLightMapInstanceData::LightMapInstanceDataTypeName, CachedData);
+	check(ComponentInstanceData.IsValid());
 
 	// Note: ApplyComponentInstanceData is called while the component is registered so the rendering thread is already using this component
 	// That means all component state that is modified here must be mirrored on the scene proxy, which will be recreated to receive the changes later due to MarkRenderStateDirty.
-	for(int32 DataIdx=0; DataIdx<CachedData.Num(); DataIdx++)
-	{
-		check(CachedData[DataIdx].IsValid());
-		check(CachedData[DataIdx]->GetDataTypeName() == FLightMapInstanceData::LightMapInstanceDataTypeName);
-		TSharedPtr<FLightMapInstanceData> LightMapData = StaticCastSharedPtr<FLightMapInstanceData>(CachedData[DataIdx]);
-
-		// See if data matches current state
-		if(	LightMapData->StaticMesh == StaticMesh &&
-			LightMapData->Transform.Equals(ComponentToWorld) )
-		{
-			const int32 NumLODLightMaps = LightMapData->LODDataLightMap.Num();
-			SetLODDataCount(NumLODLightMaps, NumLODLightMaps);
-			for (int32 i = 0; i < NumLODLightMaps; ++i)
-			{
-				LODData[i].LightMap = LightMapData->LODDataLightMap[i];
-				LODData[i].ShadowMap = LightMapData->LODDataShadowMap[i];
-			}
-
-			IrrelevantLights = LightMapData->IrrelevantLights;
-			bHasCachedStaticLighting = true;
-
-			MarkRenderStateDirty();
-		}
-	}
-
-	TArray< TSharedPtr<FComponentInstanceDataBase> > AllCachedVertexColorData;
-	Cache.GetInstanceDataOfType( FVertexColorInstanceData::VertexColorInstanceDataName, AllCachedVertexColorData );
+	TSharedPtr<FStaticMeshComponentInstanceData> StaticMeshInstanceData = StaticCastSharedPtr<FStaticMeshComponentInstanceData>(ComponentInstanceData);
 
-	for(const auto& CachedVertexColorData : AllCachedVertexColorData)
+	// See if data matches current state
+	if(	StaticMeshInstanceData->bHasCachedStaticLighting && StaticMeshInstanceData->CachedStaticLighting.Transform.Equals(ComponentToWorld) )
 	{
-		TSharedPtr<FVertexColorInstanceData> ComponentInstanceData = StaticCastSharedPtr<FVertexColorInstanceData>(CachedVertexColorData);
-		if(	ComponentInstanceData->MatchesComponent(this) )
+		const int32 NumLODLightMaps = StaticMeshInstanceData->CachedStaticLighting.LODDataLightMap.Num();
+		SetLODDataCount(NumLODLightMaps, NumLODLightMaps);
+		for (int32 i = 0; i < NumLODLightMaps; ++i)
 		{
-			ComponentInstanceData->ApplyVertexColorData(this);
-			MarkRenderStateDirty();	
+			LODData[i].LightMap = StaticMeshInstanceData->CachedStaticLighting.LODDataLightMap[i];
+			LODData[i].ShadowMap = StaticMeshInstanceData->CachedStaticLighting.LODDataShadowMap[i];
 		}
+
+		IrrelevantLights = StaticMeshInstanceData->CachedStaticLighting.IrrelevantLights;
+		bHasCachedStaticLighting = true;
 	}
+
+	StaticMeshInstanceData->ApplyVertexColorData(this);
+	MarkRenderStateDirty();	
 }
 
 #include "AI/Navigation/RecastHelpers.h"
diff --git a/Engine/Source/Runtime/Engine/Public/ComponentInstanceDataCache.h b/Engine/Source/Runtime/Engine/Public/ComponentInstanceDataCache.h
--- a/Engine/Source/Runtime/Engine/Public/ComponentInstanceDataCache.h
+++ b/Engine/Source/Runtime/Engine/Public/ComponentInstanceDataCache.h
@@ -15,9 +15,6 @@ public:
 	virtual ~FComponentInstanceDataBase()
 	{}
 
-	/** Return the type name of this data */
-	virtual FName GetDataTypeName() const = 0;
-
 	virtual bool MatchesComponent(const UActorComponent* Component) const;
 
 protected:
@@ -39,20 +36,11 @@ public:
 	FComponentInstanceDataCache() {}
 
 	/** Constructor that also populatees cache from Actor */
-	FComponentInstanceDataCache(AActor* InActor)
-	{
-		GetFromActor(InActor, *this);
-	}
+	FComponentInstanceDataCache(AActor* InActor);
 
 	/** Add some new instance data to the cache */
 	void AddInstanceData(TSharedPtr<FComponentInstanceDataBase> NewData);
 
-	/** Get all instance data of a particular type */
-	void GetInstanceDataOfType(FName TypeName, TArray< TSharedPtr<FComponentInstanceDataBase> >& OutData) const;
-
-	/** Util to iterate over components and get data from each */
-	static void GetFromActor(AActor* Actor, FComponentInstanceDataCache& Cache);
-
 	/** Util to iterate over components and apply data to each */
 	void ApplyToActor(AActor* Actor) const;
 
