Fix compilation issues for CL#2083108 when WITH_EDITOR not defined

diff --git a/Engine/Source/Runtime/Engine/Private/InstancedStaticMesh.cpp b/Engine/Source/Runtime/Engine/Private/InstancedStaticMesh.cpp
--- a/Engine/Source/Runtime/Engine/Private/InstancedStaticMesh.cpp
+++ b/Engine/Source/Runtime/Engine/Private/InstancedStaticMesh.cpp
@@ -1147,7 +1147,7 @@ FName UInstancedStaticMeshComponent::GetComponentInstanceDataType() const
 
 TSharedPtr<FComponentInstanceDataBase> UInstancedStaticMeshComponent::GetComponentInstanceData() const
 {
-	TSharedPtr<FInstancedStaticMeshSelectionData> InstanceData;
+	TSharedPtr<FComponentInstanceDataBase> InstanceData;
 #if WITH_EDITOR
 	InstanceData = MakeShareable(new FInstancedStaticMeshSelectionData(*this));
 #endif
