Mark NavigationObjectBase.h and PlayerStart.h as ENGINE_API instead of MinimalAPI
TTP #335704

diff --git a/Engine/Source/Runtime/Engine/Classes/Engine/NavigationObjectBase.h b/Engine/Source/Runtime/Engine/Classes/Engine/NavigationObjectBase.h
--- a/Engine/Source/Runtime/Engine/Classes/Engine/NavigationObjectBase.h
+++ b/Engine/Source/Runtime/Engine/Classes/Engine/NavigationObjectBase.h
@@ -4,8 +4,8 @@
 #pragma once
 #include "NavigationObjectBase.generated.h"
 
-UCLASS(hidecategories=(Lighting, LightColor, Force), ClassGroup=Navigation,MinimalAPI, NotBlueprintable, abstract)
-class ANavigationObjectBase : public AActor, public INavAgentInterface
+UCLASS(hidecategories=(Lighting, LightColor, Force), ClassGroup=Navigation, NotBlueprintable, abstract)
+class ENGINE_API ANavigationObjectBase : public AActor, public INavAgentInterface
 {
 	GENERATED_UCLASS_BODY()
 
@@ -26,32 +26,32 @@ class ANavigationObjectBase : public AActor, public INavAgentInterface
 
 
 	// Begin AActor Interface
-	ENGINE_API virtual void GetSimpleCollisionCylinder(float& CollisionRadius, float& CollisionHalfHeight) const OVERRIDE;
+	virtual void GetSimpleCollisionCylinder(float& CollisionRadius, float& CollisionHalfHeight) const OVERRIDE;
 
 #if	WITH_EDITOR
-	ENGINE_API virtual void PostEditMove(bool bFinished) OVERRIDE;
-	ENGINE_API virtual void PostEditUndo() OVERRIDE;
+	virtual void PostEditMove(bool bFinished) OVERRIDE;
+	virtual void PostEditUndo() OVERRIDE;
 #endif	// WITH_EDITOR
 	// End AActor Interface
 
 	/** @todo document */
-	ENGINE_API virtual bool ShouldBeBased();
+	virtual bool ShouldBeBased();
 	/** @todo document */
-	ENGINE_API virtual void FindBase();
+	virtual void FindBase();
 
 	/** 
 	 * Check to make sure the navigation is at a valid point. 
 	 * Detect when path building is going to move a pathnode around
 	 * This may be undesirable for LDs (ie b/c cover links define slots by offsets)
 	 */
-	ENGINE_API virtual void Validate();
+	virtual void Validate();
 
 	/** Return Physics Volume for this Actor. **/
-	ENGINE_API class APhysicsVolume* GetNavPhysicsVolume();
+	class APhysicsVolume* GetNavPhysicsVolume();
 
 	// INavAgentInterface start
-	ENGINE_API virtual FVector GetNavAgentLocation() const { return GetActorLocation(); }
-	ENGINE_API virtual void GetMoveGoalReachTest(class AActor* MovingActor, const FVector& MoveOffset, FVector& GoalOffset, float& GoalRadius, float& GoalHalfHeight) const;
+	virtual FVector GetNavAgentLocation() const { return GetActorLocation(); }
+	virtual void GetMoveGoalReachTest(class AActor* MovingActor, const FVector& MoveOffset, FVector& GoalOffset, float& GoalRadius, float& GoalHalfHeight) const;
 	// INavAgentInterface end
 };
 
diff --git a/Engine/Source/Runtime/Engine/Classes/GameFramework/PlayerStart.h b/Engine/Source/Runtime/Engine/Classes/GameFramework/PlayerStart.h
--- a/Engine/Source/Runtime/Engine/Classes/GameFramework/PlayerStart.h
+++ b/Engine/Source/Runtime/Engine/Classes/GameFramework/PlayerStart.h
@@ -7,8 +7,8 @@
 #pragma once
 #include "PlayerStart.generated.h"
 
-UCLASS(ClassGroup=Common, hidecategories=Collision,MinimalAPI)
-class APlayerStart : public ANavigationObjectBase
+UCLASS(ClassGroup=Common, hidecategories=Collision)
+class ENGINE_API APlayerStart : public ANavigationObjectBase
 {
 	GENERATED_UCLASS_BODY()
 
@@ -21,9 +21,9 @@ class APlayerStart : public ANavigationObjectBase
 	TSubobjectPtr<class UArrowComponent> ArrowComponent;
 #endif
 
-	ENGINE_API virtual void PostInitializeComponents() OVERRIDE;
+	virtual void PostInitializeComponents() OVERRIDE;
 	
-	ENGINE_API virtual void PostUnregisterAllComponents() OVERRIDE;
+	virtual void PostUnregisterAllComponents() OVERRIDE;
 };
 
 
