UE4: Moving landscape MICs and textures to have the UWorld as their outer instead of the UPackage.

diff --git a/Engine/Source/Editor/LandscapeEditor/Private/LandscapeEdMode.cpp b/Engine/Source/Editor/LandscapeEditor/Private/LandscapeEdMode.cpp
--- a/Engine/Source/Editor/LandscapeEditor/Private/LandscapeEdMode.cpp
+++ b/Engine/Source/Editor/LandscapeEditor/Private/LandscapeEdMode.cpp
@@ -93,15 +93,8 @@ void ALandscape::SplitHeightmap(ULandscapeComponent* Comp, bool bMoveToCurrentLe
 		LandscapeEdit.GetHeightDataFast(Comp->GetSectionBase().X, Comp->GetSectionBase().Y, Comp->GetSectionBase().X + Comp->ComponentSizeQuads, Comp->GetSectionBase().Y + Comp->ComponentSizeQuads, (uint16*)HeightData.GetTypedData(), 0, (uint16*)NormalData.GetTypedData());
 
 		// Construct the heightmap textures
-		UObject* Outer = bMoveToCurrentLevel ? Comp->GetWorld()->GetCurrentLevel()->GetOutermost() : Comp->GetOutermost();
-		HeightmapTexture = ConstructObject<UTexture2D>(UTexture2D::StaticClass(), Outer, NAME_None, RF_Public);
-		HeightmapTexture->Source.Init2DWithMipChain(HeightmapSizeU, HeightmapSizeV, TSF_BGRA8);
-		HeightmapTexture->SRGB = false;
-		HeightmapTexture->CompressionNone = true;
-		HeightmapTexture->MipGenSettings = TMGS_LeaveExistingMips;
-		HeightmapTexture->LODGroup = TEXTUREGROUP_Terrain_Heightmap;
-		HeightmapTexture->AddressX = TA_Clamp;
-		HeightmapTexture->AddressY = TA_Clamp;
+		UObject* TextureOuter = bMoveToCurrentLevel ? Comp->GetWorld()->GetCurrentLevel() : nullptr;
+		HeightmapTexture = Comp->GetLandscapeProxy()->CreateLandscapeTexture(HeightmapSizeU, HeightmapSizeV, TEXTUREGROUP_Terrain_Heightmap, TSF_BGRA8, TextureOuter);
 
 		int32 MipSubsectionSizeQuads = Comp->SubsectionSizeQuads;
 		int32 MipSizeU = HeightmapSizeU;
diff --git a/Engine/Source/Editor/LandscapeEditor/Private/LandscapeEdModeComponentTools.cpp b/Engine/Source/Editor/LandscapeEditor/Private/LandscapeEdModeComponentTools.cpp
--- a/Engine/Source/Editor/LandscapeEditor/Private/LandscapeEdModeComponentTools.cpp
+++ b/Engine/Source/Editor/LandscapeEditor/Private/LandscapeEdModeComponentTools.cpp
@@ -572,14 +572,7 @@ public:
 								int32 WeightmapSize = (Comp->SubsectionSizeQuads+1) * Comp->NumSubsections;
 
 								// We need a new weightmap texture
-								CurrentWeightmapTexture = ConstructObject<UTexture2D>(UTexture2D::StaticClass(), World->GetCurrentLevel()->GetOutermost(), NAME_None, RF_Public);
-								CurrentWeightmapTexture->Source.Init2DWithMipChain(WeightmapSize,WeightmapSize,TSF_BGRA8);
-								CurrentWeightmapTexture->SRGB = false;
-								CurrentWeightmapTexture->CompressionNone = true;
-								CurrentWeightmapTexture->MipGenSettings = TMGS_LeaveExistingMips;
-								CurrentWeightmapTexture->AddressX = TA_Clamp;
-								CurrentWeightmapTexture->AddressY = TA_Clamp;
-								CurrentWeightmapTexture->LODGroup = TEXTUREGROUP_Terrain_Weightmap;
+								CurrentWeightmapTexture = Comp->GetLandscapeProxy()->CreateLandscapeTexture(WeightmapSize, WeightmapSize, TEXTUREGROUP_Terrain_Weightmap, TSF_BGRA8);
 								// Alloc dummy mips
 								Comp->CreateEmptyTextureMips(CurrentWeightmapTexture);
 								CurrentWeightmapTexture->PostEditChange();
diff --git a/Engine/Source/Runtime/Core/Public/UObject/ObjectVersion.h b/Engine/Source/Runtime/Core/Public/UObject/ObjectVersion.h
--- a/Engine/Source/Runtime/Core/Public/UObject/ObjectVersion.h
+++ b/Engine/Source/Runtime/Core/Public/UObject/ObjectVersion.h
@@ -546,6 +546,8 @@ enum EUnrealEngineObjectUE4Version
 	VER_UE4_LOAD_FOR_EDITOR_GAME,
 	// Manual serialization of FRichCurveKey to save space
 	VER_UE4_SERIALIZE_RICH_CURVE_KEY,
+	// Change the outer of ULandscapeMaterialInstanceConstants and Landscape-related textures to the level in which they reside
+	VER_UE4_MOVE_LANDSCAPE_MICS_AND_TEXTURES_WITHIN_LEVEL,
 
 	// -----<new versions can be added before this line>-------------------------------------------------
 	// - this needs to be the last line (see note below)
diff --git a/Engine/Source/Runtime/Engine/Classes/Landscape/LandscapeComponent.h b/Engine/Source/Runtime/Engine/Classes/Landscape/LandscapeComponent.h
--- a/Engine/Source/Runtime/Engine/Classes/Landscape/LandscapeComponent.h
+++ b/Engine/Source/Runtime/Engine/Classes/Landscape/LandscapeComponent.h
@@ -302,6 +302,9 @@ public:
 	/** Get the landscape actor associated with this component. */
 	class ALandscape* GetLandscapeActor() const;
 
+	/** Get the level in which the owning actor resides */
+	class ULevel* GetLevel() const;
+
 	/** @todo document */
 	ENGINE_API class ALandscapeProxy* GetLandscapeProxy() const;
 
diff --git a/Engine/Source/Runtime/Engine/Classes/Landscape/LandscapeProxy.h b/Engine/Source/Runtime/Engine/Classes/Landscape/LandscapeProxy.h
--- a/Engine/Source/Runtime/Engine/Classes/Landscape/LandscapeProxy.h
+++ b/Engine/Source/Runtime/Engine/Classes/Landscape/LandscapeProxy.h
@@ -427,6 +427,9 @@ public:
 
 	/** @return Current size of bounding rectangle in quads space */
 	ENGINE_API FIntRect GetBoundingRect() const;
+
+	/** Creates a Texture2D for use by this landscape proxy or one of it's components. If OptionalOverrideOuter is not specified, the level is used. */
+	ENGINE_API class UTexture2D* CreateLandscapeTexture(int32 InSizeX, int32 InSizeY, TextureGroup InLODGroup, ETextureSourceFormat InFormat, UObject* OptionalOverrideOuter = nullptr) const;
 #endif
 };
 
diff --git a/Engine/Source/Runtime/Engine/Private/Landscape/Landscape.cpp b/Engine/Source/Runtime/Engine/Private/Landscape/Landscape.cpp
--- a/Engine/Source/Runtime/Engine/Private/Landscape/Landscape.cpp
+++ b/Engine/Source/Runtime/Engine/Private/Landscape/Landscape.cpp
@@ -417,6 +417,40 @@ void ULandscapeComponent::PostLoad()
 			}
 		}
 	}
+
+#if WITH_EDITOR
+	if (GIsEditor && !HasAnyFlags(RF_ClassDefaultObject))
+	{
+		if (GetLinkerUE4Version() < VER_UE4_MOVE_LANDSCAPE_MICS_AND_TEXTURES_WITHIN_LEVEL)
+		{
+			ULevel* Level = GetLevel();
+			if (ensure(Level))
+			{
+				TArray<UObject*> ObjectsToMoveFromPackgeToLevel;
+				for (UMaterialInstance* CurrentMIC = MaterialInstance; CurrentMIC; CurrentMIC = Cast<UMaterialInstance>(CurrentMIC->Parent))
+				{
+					ObjectsToMoveFromPackgeToLevel.Add(CurrentMIC);
+				}
+				ObjectsToMoveFromPackgeToLevel.Add(HeightmapTexture);
+				for (auto* Tex : WeightmapTextures)
+				{
+					ObjectsToMoveFromPackgeToLevel.Add(Tex);
+				}
+				ObjectsToMoveFromPackgeToLevel.Add(XYOffsetmapTexture);
+
+				UPackage* MyPackage = GetOutermost();
+				for (auto* Obj : ObjectsToMoveFromPackgeToLevel)
+				{
+					if (Obj && Obj->GetOuter() == MyPackage)
+					{
+						Obj->ClearFlags(RF_Public);
+						Obj->Rename(*Obj->GetName(), Level, REN_DoNotDirty | REN_DontCreateRedirectors | REN_ForceNoResetLoaders | REN_NonTransactional);
+					}
+				}
+			}
+		}
+	}
+#endif
 }
 
 #endif // WITH_EDITOR
@@ -533,6 +567,12 @@ ALandscape* ULandscapeComponent::GetLandscapeActor() const
 	return NULL;
 }
 
+ULevel* ULandscapeComponent::GetLevel() const
+{
+	AActor* MyOwner = GetOwner();
+	return MyOwner ? MyOwner->GetLevel() : NULL;
+}
+
 ALandscapeProxy* ULandscapeComponent::GetLandscapeProxy() const
 {
 	return CastChecked<ALandscapeProxy>(GetOuter());
diff --git a/Engine/Source/Runtime/Engine/Private/Landscape/LandscapeEdit.cpp b/Engine/Source/Runtime/Engine/Private/Landscape/LandscapeEdit.cpp
--- a/Engine/Source/Runtime/Engine/Private/Landscape/LandscapeEdit.cpp
+++ b/Engine/Source/Runtime/Engine/Private/Landscape/LandscapeEdit.cpp
@@ -92,7 +92,12 @@ ULandscapeMaterialInstanceConstant* ALandscapeProxy::GetLayerThumbnailMIC(UMater
 	if( CombinationMaterialInstance == NULL || CombinationMaterialInstance->Parent != LandscapeMaterial || (Proxy && Proxy->GetOutermost() != CombinationMaterialInstance->GetOutermost()) )
 	{
 		FlushRenderingCommands();
-		CombinationMaterialInstance = ConstructObject<ULandscapeMaterialInstanceConstant>(ULandscapeMaterialInstanceConstant::StaticClass(), Proxy ? Proxy->GetOutermost() : GetTransientPackage(), NAME_None, RF_Public);
+		UObject* MICOuter = GetTransientPackage();
+		if ( Proxy && Proxy->GetLevel() )
+		{
+			MICOuter = Proxy->GetLevel();
+		}
+		CombinationMaterialInstance = ConstructObject<ULandscapeMaterialInstanceConstant>(ULandscapeMaterialInstanceConstant::StaticClass(), MICOuter);
 		if (Proxy)
 		{
 			UE_LOG(LogLandscape, Log, TEXT("Looking for key %s, making new combination %s"), *LayerKey, *CombinationMaterialInstance->GetName());
@@ -123,7 +128,7 @@ ULandscapeMaterialInstanceConstant* ALandscapeProxy::GetLayerThumbnailMIC(UMater
 	}
 
 	// Create the instance for this component, that will use the layer combination instance.
-	ULandscapeMaterialInstanceConstant* MaterialInstance = ConstructObject<ULandscapeMaterialInstanceConstant>(ULandscapeMaterialInstanceConstant::StaticClass(), GetTransientPackage(), NAME_None, RF_Public);
+	ULandscapeMaterialInstanceConstant* MaterialInstance = ConstructObject<ULandscapeMaterialInstanceConstant>(ULandscapeMaterialInstanceConstant::StaticClass(), GetTransientPackage());
 	MaterialInstance->SetParentEditorOnly(CombinationMaterialInstance);
 	MaterialInstance->bIsLayerThumbnail = true;
 
@@ -164,7 +169,7 @@ UMaterialInstanceConstant* ULandscapeComponent::GetCombinationMaterial(bool bMob
 		{
 			FlushRenderingCommands();
 
-			CombinationMaterialInstance = ConstructObject<ULandscapeMaterialInstanceConstant>(ULandscapeMaterialInstanceConstant::StaticClass(), GetOutermost(), NAME_None, RF_Public);
+			CombinationMaterialInstance = ConstructObject<ULandscapeMaterialInstanceConstant>(ULandscapeMaterialInstanceConstant::StaticClass(), GetLevel());
 			UE_LOG(LogLandscape, Log, TEXT("Looking for key %s, making new combination %s"), *LayerKey, *CombinationMaterialInstance->GetName());
 			Proxy->MaterialInstanceConstantMap.Add(*LayerKey,CombinationMaterialInstance);
 			CombinationMaterialInstance->SetParentEditorOnly(LandscapeMaterial);
@@ -219,7 +224,7 @@ void ULandscapeComponent::UpdateMaterialInstances()
 		// Create the instance for this component, that will use the layer combination instance.
 		if( MaterialInstance == NULL || GetOutermost() != MaterialInstance->GetOutermost() )
 		{
-			MaterialInstance = ConstructObject<ULandscapeMaterialInstanceConstant>(ULandscapeMaterialInstanceConstant::StaticClass(), GetOutermost(), NAME_None, RF_Public);
+			MaterialInstance = ConstructObject<ULandscapeMaterialInstanceConstant>(ULandscapeMaterialInstanceConstant::StaticClass(), GetLevel());
 		}
 
 		// For undo
@@ -1762,19 +1767,8 @@ void ALandscapeProxy::Import(FGuid Guid, int32 VertsX, int32 VertsY,
 			HeightmapInfo.HeightmapSizeV = (1<<FMath::CeilLogTwo( ((HmY==NumHeightmapsY-1) ? FinalComponentsY : ComponentsPerHeightmap)*ComponentSizeVerts ));
 
 			// Construct the heightmap textures
-			HeightmapInfo.HeightmapTexture = ConstructObject<UTexture2D>(UTexture2D::StaticClass(), GetOutermost(), NAME_None/*FName(TEXT("Heightmap"))*/, RF_Public);
-			HeightmapInfo.HeightmapTexture->Source.Init2DWithMipChain(
-				HeightmapInfo.HeightmapSizeU,
-				HeightmapInfo.HeightmapSizeV,
-				TSF_BGRA8
-				);
-			HeightmapInfo.HeightmapTexture->SRGB = false;
-			HeightmapInfo.HeightmapTexture->CompressionNone = true;
-			HeightmapInfo.HeightmapTexture->MipGenSettings = TMGS_LeaveExistingMips;
-			HeightmapInfo.HeightmapTexture->LODGroup = TEXTUREGROUP_Terrain_Heightmap;
-			HeightmapInfo.HeightmapTexture->AddressX = TA_Clamp;
-			HeightmapInfo.HeightmapTexture->AddressY = TA_Clamp;
-
+			HeightmapInfo.HeightmapTexture = CreateLandscapeTexture(HeightmapInfo.HeightmapSizeU, HeightmapInfo.HeightmapSizeV, TEXTUREGROUP_Terrain_Heightmap, TSF_BGRA8);
+			
 			int32 MipSubsectionSizeQuads = SubsectionSizeQuads;
 			int32 MipSizeU = HeightmapInfo.HeightmapSizeU;
 			int32 MipSizeV = HeightmapInfo.HeightmapSizeV;
@@ -2068,14 +2062,7 @@ void ALandscapeProxy::Import(FGuid Guid, int32 VertsX, int32 VertsY,
 				else
 				{
 					// We couldn't find a suitable place for these layers, so lets make a new one.
-					UTexture2D* WeightmapTexture = ConstructObject<UTexture2D>(UTexture2D::StaticClass(), GetOutermost(), NAME_None, RF_Public);
-					WeightmapTexture->Source.Init2DWithMipChain(WeightmapSize,WeightmapSize,TSF_BGRA8);
-					WeightmapTexture->SRGB = false;
-					WeightmapTexture->CompressionNone = true;
-					WeightmapTexture->MipGenSettings = TMGS_LeaveExistingMips;
-					WeightmapTexture->AddressX = TA_Clamp;
-					WeightmapTexture->AddressY = TA_Clamp;
-					WeightmapTexture->LODGroup = TEXTUREGROUP_Terrain_Weightmap;
+					UTexture2D* WeightmapTexture = CreateLandscapeTexture(WeightmapSize, WeightmapSize, TEXTUREGROUP_Terrain_Weightmap, TSF_BGRA8);
 					FColor* MipData = (FColor*)WeightmapTexture->Source.LockMip(0);
 
 					int32 ThisAllocationLayers = FMath::Min<int32>(RemainingLayers,4);
@@ -4213,14 +4200,7 @@ void ULandscapeComponent::ReallocateWeightmaps(FLandscapeEditDataInterface* Data
 			int32 WeightmapSize = (SubsectionSizeQuads+1) * NumSubsections;
 
 			// We need a new weightmap texture
-			CurrentWeightmapTexture = ConstructObject<UTexture2D>(UTexture2D::StaticClass(), GetOutermost(), NAME_None, RF_Public);
-			CurrentWeightmapTexture->Source.Init2DWithMipChain(WeightmapSize,WeightmapSize,TSF_BGRA8);
-			CurrentWeightmapTexture->SRGB = false;
-			CurrentWeightmapTexture->CompressionNone = true;
-			CurrentWeightmapTexture->MipGenSettings = TMGS_LeaveExistingMips;
-			CurrentWeightmapTexture->AddressX = TA_Clamp;
-			CurrentWeightmapTexture->AddressY = TA_Clamp;
-			CurrentWeightmapTexture->LODGroup = TEXTUREGROUP_Terrain_Weightmap;
+			CurrentWeightmapTexture = GetLandscapeProxy()->CreateLandscapeTexture(WeightmapSize, WeightmapSize, TEXTUREGROUP_Terrain_Weightmap, TSF_BGRA8);
 			// Alloc dummy mips
 			CreateEmptyTextureMips(CurrentWeightmapTexture);
 			CurrentWeightmapTexture->PostEditChange();
@@ -4402,15 +4382,8 @@ void ULandscapeComponent::InitHeightmapData(TArray<FColor>& Heights, bool bUpdat
 	int32 HeightmapSizeV = (1<<FMath::CeilLogTwo( ComponentSizeVerts ));
 
 	// Height map construction
-	HeightmapTexture = ConstructObject<UTexture2D>(UTexture2D::StaticClass(), GetOutermost(), NAME_None, RF_Public);
-	HeightmapTexture->Source.Init2DWithMipChain(HeightmapSizeU,HeightmapSizeV,TSF_BGRA8);
-	HeightmapTexture->SRGB = false;
-	HeightmapTexture->CompressionNone = true;
-	HeightmapTexture->MipGenSettings = TMGS_LeaveExistingMips;
-	HeightmapTexture->LODGroup = TEXTUREGROUP_Terrain_Heightmap;
-	HeightmapTexture->AddressX = TA_Clamp;
-	HeightmapTexture->AddressY = TA_Clamp;
-
+	HeightmapTexture = GetLandscapeProxy()->CreateLandscapeTexture(HeightmapSizeU, HeightmapSizeV, TEXTUREGROUP_Terrain_Heightmap, TSF_BGRA8);
+	
 	int32 MipSubsectionSizeQuads = SubsectionSizeQuads;
 	int32 MipSizeU = HeightmapSizeU;
 	int32 MipSizeV = HeightmapSizeV;
@@ -4785,14 +4758,7 @@ UMaterialInstance* ULandscapeComponent::GeneratePlatformPixelData(TArray<UTextur
 	}
 
 	int32 WeightmapSize = (SubsectionSizeQuads+1) * NumSubsections;
-	UTexture2D* WeightmapTexture = ConstructObject<UTexture2D>(UTexture2D::StaticClass(), GetOutermost(), NAME_None, RF_Public);
-	WeightmapTexture->Source.Init2DWithMipChain(WeightmapSize,WeightmapSize,TSF_BGRA8);
-	WeightmapTexture->SRGB = false;
-	WeightmapTexture->CompressionNone = true;
-	WeightmapTexture->MipGenSettings = TMGS_LeaveExistingMips;
-	WeightmapTexture->AddressX = TA_Clamp;
-	WeightmapTexture->AddressY = TA_Clamp;
-	WeightmapTexture->LODGroup = TEXTUREGROUP_Terrain_Weightmap;
+	UTexture2D* WeightmapTexture = GetLandscapeActor()->CreateLandscapeTexture(WeightmapSize, WeightmapSize, TEXTUREGROUP_Terrain_Weightmap, TSF_BGRA8);
 	CreateEmptyTextureMips(WeightmapTexture);
 
 	{
@@ -4839,7 +4805,7 @@ UMaterialInstance* ULandscapeComponent::GeneratePlatformPixelData(TArray<UTextur
 
 	if (!bIsCooking)
 	{
-		UMaterialInstanceDynamic* MobileMaterialInstance = UMaterialInstanceDynamic::Create(MaterialInstance, GetOutermost()); 
+		UMaterialInstanceDynamic* MobileMaterialInstance = UMaterialInstanceDynamic::Create(MaterialInstance, GetLevel()); 
 		// Set the layer mask
 		int32 CurrentIdx = 0;
 		for( int32 AllocIdx=0;AllocIdx<WeightmapLayerAllocations.Num();AllocIdx++ )
@@ -4860,7 +4826,7 @@ UMaterialInstance* ULandscapeComponent::GeneratePlatformPixelData(TArray<UTextur
 	else // for cooking
 	{
 		UMaterialInstanceConstant* CombinationMaterialInstance = GetCombinationMaterial(true); 
-		UMaterialInstanceConstant* MobileMaterialInstance = ConstructObject<ULandscapeMaterialInstanceConstant>(ULandscapeMaterialInstanceConstant::StaticClass(), GetOutermost(), NAME_None, RF_Public);
+		UMaterialInstanceConstant* MobileMaterialInstance = ConstructObject<ULandscapeMaterialInstanceConstant>(ULandscapeMaterialInstanceConstant::StaticClass(), GetLevel());
 
 		MobileMaterialInstance->SetParentEditorOnly(CombinationMaterialInstance);
 
@@ -5055,6 +5021,21 @@ void ULandscapeComponent::GeneratePlatformVertexData()
 	PlatformData.InitializeFromUncompressedData(NewPlatformData);
 }
 
+UTexture2D* ALandscapeProxy::CreateLandscapeTexture(int32 InSizeX, int32 InSizeY, TextureGroup InLODGroup, ETextureSourceFormat InFormat, UObject* OptionalOverrideOuter) const
+{
+	UObject* TexOuter = OptionalOverrideOuter ? OptionalOverrideOuter : GetLevel();
+	UTexture2D* NewTexture = ConstructObject<UTexture2D>(UTexture2D::StaticClass(), TexOuter);
+	NewTexture->Source.Init2DWithMipChain(InSizeX, InSizeY, InFormat);
+	NewTexture->SRGB = false;
+	NewTexture->CompressionNone = true;
+	NewTexture->MipGenSettings = TMGS_LeaveExistingMips;
+	NewTexture->AddressX = TA_Clamp;
+	NewTexture->AddressY = TA_Clamp;
+	NewTexture->LODGroup = InLODGroup;
+
+	return NewTexture;
+}
+
 #endif //WITH_EDITOR
 
 #undef LOCTEXT_NAMESPACE
diff --git a/Engine/Source/Runtime/Engine/Private/Landscape/LandscapeEditInterface.cpp b/Engine/Source/Runtime/Engine/Private/Landscape/LandscapeEditInterface.cpp
--- a/Engine/Source/Runtime/Engine/Private/Landscape/LandscapeEditInterface.cpp
+++ b/Engine/Source/Runtime/Engine/Private/Landscape/LandscapeEditInterface.cpp
@@ -3764,14 +3764,7 @@ void FLandscapeEditDataInterface::SetSelectData(int32 X1, int32 Y1, int32 X2, in
 				//FlushRenderingCommands();
 				// Construct Texture...
 				int32 WeightmapSize = (Component->SubsectionSizeQuads+1) * Component->NumSubsections;
-				DataTexture = ConstructObject<UTexture2D>(UTexture2D::StaticClass(), Component->GetOutermost(), NAME_None, RF_Public);
-				DataTexture->Source.Init2DWithMipChain(WeightmapSize,WeightmapSize,TSF_G8);
-				DataTexture->SRGB = false;
-				DataTexture->CompressionNone = true;
-				DataTexture->MipGenSettings = TMGS_LeaveExistingMips;
-				DataTexture->AddressX = TA_Clamp;
-				DataTexture->AddressY = TA_Clamp;
-				DataTexture->LODGroup = TEXTUREGROUP_Terrain_Weightmap;
+				DataTexture = Component->GetLandscapeProxy()->CreateLandscapeTexture(WeightmapSize, WeightmapSize, TEXTUREGROUP_Terrain_Weightmap, TSF_G8);
 				// Alloc dummy mips
 				ULandscapeComponent::CreateEmptyTextureMips(DataTexture, true);
 				DataTexture->PostEditChange();
@@ -3911,14 +3904,7 @@ void FLandscapeEditDataInterface::SetXYOffsetDataTempl(int32 X1, int32 Y1, int32
 					//FlushRenderingCommands();
 					// Construct Texture...
 					int32 WeightmapSize = (Component->SubsectionSizeQuads+1) * Component->NumSubsections;
-					XYOffsetTexture = ConstructObject<UTexture2D>(UTexture2D::StaticClass(), Component->GetOutermost(), NAME_None, RF_Public);
-					XYOffsetTexture->Source.Init2DWithMipChain(WeightmapSize,WeightmapSize,TSF_BGRA8);
-					XYOffsetTexture->SRGB = false;
-					XYOffsetTexture->CompressionNone = true;
-					XYOffsetTexture->MipGenSettings = TMGS_LeaveExistingMips;
-					XYOffsetTexture->AddressX = TA_Clamp;
-					XYOffsetTexture->AddressY = TA_Clamp;
-					XYOffsetTexture->LODGroup = TEXTUREGROUP_Terrain_Weightmap; // for now...
+					XYOffsetTexture = Component->GetLandscapeProxy()->CreateLandscapeTexture(WeightmapSize, WeightmapSize, TEXTUREGROUP_Terrain_Weightmap, TSF_BGRA8);
 					// Alloc dummy mips
 					ULandscapeComponent::CreateEmptyTextureMips(XYOffsetTexture, true);
 					XYOffsetTexture->PostEditChange();
