#ue4 - Minor optimizations and variable name cleanup in CharacterMovementComponent.

- UCharacterMovementComponent::ApplyAccumulatedForces (avoid GetGravityZ() call if possible, and remove LHS on Velocity).
- InOutRequestedSpeed in ApplyRequestedVelocity is actually just an out param.

diff --git a/Engine/Source/Runtime/Engine/Classes/GameFramework/CharacterMovementComponent.h b/Engine/Source/Runtime/Engine/Classes/GameFramework/CharacterMovementComponent.h
--- a/Engine/Source/Runtime/Engine/Classes/GameFramework/CharacterMovementComponent.h
+++ b/Engine/Source/Runtime/Engine/Classes/GameFramework/CharacterMovementComponent.h
@@ -782,7 +782,7 @@ public:
 	virtual FRotator ComputeOrientToMovementRotation(const FRotator& CurrentRotation, float DeltaTime, FRotator& DeltaRotation);
 
 	/** use velocity requested by path following */
-	virtual bool ApplyRequestedMove(FVector& NewAcceleration, float DeltaTime, float MaxAccel, float MaxSpeed, float& InOutRequestedSpeed);
+	virtual bool ApplyRequestedMove(FVector& NewAcceleration, float DeltaTime, float MaxAccel, float MaxSpeed, float& OutRequestedSpeed);
 
 	/** Called if bNotifyApex is true and character has just passed the apex of its jump. */
 	virtual void NotifyJumpApex();
diff --git a/Engine/Source/Runtime/Engine/Private/CharacterMovementComponent.cpp b/Engine/Source/Runtime/Engine/Private/CharacterMovementComponent.cpp
--- a/Engine/Source/Runtime/Engine/Private/CharacterMovementComponent.cpp
+++ b/Engine/Source/Runtime/Engine/Private/CharacterMovementComponent.cpp
@@ -1947,22 +1947,22 @@ void UCharacterMovementComponent::CalcVelocity(float DeltaTime, float Friction,
 	}
 }
 
-bool UCharacterMovementComponent::ApplyRequestedMove(FVector& NewAcceleration, float DeltaTime, float MaxAccel, float MaxSpeed, float& InOutRequestedSpeed)
+bool UCharacterMovementComponent::ApplyRequestedMove(FVector& NewAcceleration, float DeltaTime, float MaxAccel, float MaxSpeed, float& OutRequestedSpeed)
 {
 	if (bHasRequestedVelocity)
 	{
 		bHasRequestedVelocity = false;
 
-		InOutRequestedSpeed = RequestedVelocity.Size();
-		if (InOutRequestedSpeed < KINDA_SMALL_NUMBER)
+		OutRequestedSpeed = RequestedVelocity.Size();
+		if (OutRequestedSpeed < KINDA_SMALL_NUMBER)
 		{
 			return true;
 		}
 
-		const FVector RequestedMoveDir = RequestedVelocity / InOutRequestedSpeed;
+		const FVector RequestedMoveDir = RequestedVelocity / OutRequestedSpeed;
 
-		InOutRequestedSpeed = bRequestedMoveWithMaxSpeed ? MaxSpeed : FMath::Min(MaxSpeed, InOutRequestedSpeed);
-		const FVector MoveVelocity = RequestedMoveDir * InOutRequestedSpeed;
+		OutRequestedSpeed = bRequestedMoveWithMaxSpeed ? MaxSpeed : FMath::Min(MaxSpeed, OutRequestedSpeed);
+		const FVector MoveVelocity = RequestedMoveDir * OutRequestedSpeed;
 
 		NewAcceleration = MoveVelocity / DeltaTime;
 		if (NewAcceleration.SizeSquared() > FMath::Square(MaxAccel) || bForceMaxAccel)
@@ -5672,14 +5672,16 @@ void UCharacterMovementComponent::ApplyRepulsionForce( float DeltaTime )
 
 void UCharacterMovementComponent::ApplyAccumulatedForces(float DeltaSeconds)
 {
-	// check to see if applied momentum is enough to overcome gravity
-	if ( IsMovingOnGround() && (PendingImpulseToApply.Z + (PendingForceToApply.Z * DeltaSeconds) + (GetGravityZ() * DeltaSeconds) > SMALL_NUMBER))
+	if (PendingImpulseToApply.Z != 0.f || PendingForceToApply.Z != 0.f)
 	{
-		SetMovementMode(MOVE_Falling);
+		// check to see if applied momentum is enough to overcome gravity
+		if ( IsMovingOnGround() && (PendingImpulseToApply.Z + (PendingForceToApply.Z * DeltaSeconds) + (GetGravityZ() * DeltaSeconds) > SMALL_NUMBER))
+		{
+			SetMovementMode(MOVE_Falling);
+		}
 	}
 
-	Velocity += PendingImpulseToApply;
-	Velocity += PendingForceToApply * DeltaSeconds;
+	Velocity += PendingImpulseToApply + (PendingForceToApply * DeltaSeconds);
 
 	PendingImpulseToApply = FVector::ZeroVector;
 	PendingForceToApply = FVector::ZeroVector;
