Add a flush to pending physics commands when shutting down physics module

diff --git a/Engine/Source/Runtime/Engine/Private/PhysicsEngine/PhysLevel.cpp b/Engine/Source/Runtime/Engine/Private/PhysicsEngine/PhysLevel.cpp
--- a/Engine/Source/Runtime/Engine/Private/PhysicsEngine/PhysLevel.cpp
+++ b/Engine/Source/Runtime/Engine/Private/PhysicsEngine/PhysLevel.cpp
@@ -372,6 +372,13 @@ void TermGamePhys()
 		return;
 	}
 
+	if (GPhysCommandHandler != NULL)
+	{
+		GPhysCommandHandler->Flush();	//finish off any remaining commands
+		delete GPhysCommandHandler;
+		GPhysCommandHandler = NULL;
+	}
+
 #if WITH_APEX
 #if WITH_APEX_LEGACY
 	if(GApexModuleLegacy != NULL)
@@ -398,12 +405,6 @@ void TermGamePhys()
 	PxCloseExtensions();
 	PxCloseVehicleSDK();
 
-	if (GPhysCommandHandler != NULL)
-	{
-		delete GPhysCommandHandler;
-		GPhysCommandHandler = NULL;
-	}
-
 	if(GPhysXSDK != NULL)
 	{
 		GPhysXSDK->release();
