Fix collision detection for instanced static mesh

diff --git a/Engine/Source/Runtime/Engine/Private/InstancedStaticMesh.cpp b/Engine/Source/Runtime/Engine/Private/InstancedStaticMesh.cpp
--- a/Engine/Source/Runtime/Engine/Private/InstancedStaticMesh.cpp
+++ b/Engine/Source/Runtime/Engine/Private/InstancedStaticMesh.cpp
@@ -1213,9 +1213,9 @@ void UInstancedStaticMeshComponent::InitInstanceBody(int32 InstanceIdx, FBodyIns
 
 	UBodySetup* BodySetup = GetBodySetup();
 	check(BodySetup);
-
+	
 	// Get transform of the instance
-	FTransform InstanceTransform = ComponentToWorld * FTransform(PerInstanceSMData[InstanceIdx].Transform);
+	FTransform InstanceTransform = FTransform(PerInstanceSMData[InstanceIdx].Transform) * ComponentToWorld;
 	
 	InstanceBodyInstance->CopyBodyInstancePropertiesFrom(&BodyInstance);
 	InstanceBodyInstance->InstanceBodyIndex = InstanceIdx; // Set body index 
