Prevent crashing when there are NULLs in the SerializedComponent array

diff --git a/Engine/Source/Runtime/Engine/Private/ComponentInstanceDataCache.cpp b/Engine/Source/Runtime/Engine/Private/ComponentInstanceDataCache.cpp
--- a/Engine/Source/Runtime/Engine/Private/ComponentInstanceDataCache.cpp
+++ b/Engine/Source/Runtime/Engine/Private/ComponentInstanceDataCache.cpp
@@ -15,15 +15,18 @@ FComponentInstanceDataBase::FComponentInstanceDataBase(const UActorComponent* So
 		bool bFound = false;
 		for (const UActorComponent* SerializedComponent : ComponentOwner->SerializedComponents)
 		{
-			if (SerializedComponent == SourceComponent)
+			if (SerializedComponent)
 			{
-				++SourceComponentTypeSerializedIndex;
-				bFound = true;
-				break;
-			}
-			else if (SerializedComponent->GetClass() == SourceComponent->GetClass())
-			{
-				++SourceComponentTypeSerializedIndex;
+				if (SerializedComponent == SourceComponent)
+				{
+					++SourceComponentTypeSerializedIndex;
+					bFound = true;
+					break;
+				}
+				else if (SerializedComponent->GetClass() == SourceComponent->GetClass())
+				{
+					++SourceComponentTypeSerializedIndex;
+				}
 			}
 		}
 		if (!bFound)
@@ -50,7 +53,8 @@ bool FComponentInstanceDataBase::MatchesComponent(const UActorComponent* Compone
 			{
 				for (const UActorComponent* SerializedComponent : ComponentOwner->SerializedComponents)
 				{
-					if (   (SerializedComponent->GetClass() == Component->GetClass())
+					if (   SerializedComponent
+						&& (SerializedComponent->GetClass() == Component->GetClass())
 						&& (++FoundSerializedComponentsOfType == SourceComponentTypeSerializedIndex))
 					{
 						bMatches = true;
