Cleanup how sub-objects generate network GUIDs

* We no longer generate GUID's on demand for relatively stably named sub-objects
* We no longer send the names of sub-objects when assigning guid's at SerializeNewActor time
* We now generate ALL guids for ALL sub-objects at SerializeNewActor time
** This way, we don't have to generate them on demand, and get into weird timing issues that can happen if the owning actor hasn't been serialized

diff --git a/Engine/Source/Runtime/Engine/Private/ActorReplication.cpp b/Engine/Source/Runtime/Engine/Private/ActorReplication.cpp
--- a/Engine/Source/Runtime/Engine/Private/ActorReplication.cpp
+++ b/Engine/Source/Runtime/Engine/Private/ActorReplication.cpp
@@ -258,7 +258,6 @@ bool AActor::ReplicateSubobjects(UActorChannel *Channel, FOutBunch *Bunch, FRepl
 
 void AActor::GetSubobjectsWithStableNamesForNetworking(TArray<UObject*> &ObjList)
 {	
-#if 0
 	// For experimenting with replicating ALL stably named components initially
 	for (UActorComponent* Component : OwnedComponents)
 	{
@@ -268,17 +267,17 @@ void AActor::GetSubobjectsWithStableNamesForNetworking(TArray<UObject*> &ObjList
 			Component->GetSubobjectsWithStableNamesForNetworking(ObjList);
 		}
 	}
-#else
-	for (TWeakObjectPtr<class UActorComponent> WeakPtr : ReplicatedComponents)
+
+	// Sort the list so that we generate the same list on client/server
+	struct FCompareComponentNames
 	{
-		UActorComponent * Component = WeakPtr.Get();
-		if (Component && Component->IsNameStableForNetworking())
+		FORCEINLINE bool operator()( UObject & A, UObject & B ) const
 		{
-			ObjList.Add(Component);
-			Component->GetSubobjectsWithStableNamesForNetworking(ObjList);
+			return A.GetName() < B.GetName();
 		}
-	}
-#endif
+	};
+
+	Sort( ObjList.GetTypedData(), ObjList.Num(), FCompareComponentNames() );
 }
 
 void AActor::OnSubobjectCreatedFromReplication(UObject *NewSubobject)
diff --git a/Engine/Source/Runtime/Engine/Private/PackageMapClient.cpp b/Engine/Source/Runtime/Engine/Private/PackageMapClient.cpp
--- a/Engine/Source/Runtime/Engine/Private/PackageMapClient.cpp
+++ b/Engine/Source/Runtime/Engine/Private/PackageMapClient.cpp
@@ -177,7 +177,7 @@ bool UPackageMapClient::WriteObject( FArchive& Ar, UObject* ObjOuter, FNetworkGU
  */
 bool UPackageMapClient::SerializeNewActor(FArchive& Ar, class UActorChannel *Channel, class AActor*& Actor)
 {
-	UE_LOG(LogNetPackageMap, VeryVerbose, TEXT("START SerializeNewActor"));
+	UE_LOG(LogNetPackageMap, VeryVerbose, TEXT("SerializeNewActor START"));
 
 	uint8 bIsClosingChannel = 0;
 
@@ -288,51 +288,144 @@ bool UPackageMapClient::SerializeNewActor(FArchive& Ar, class UActorChannel *Cha
 			}
 		}
 
-		// Serialize Initially replicating components here
-		if (Ar.IsSaving())
+		//
+		// Generate net guids for every subobject that has a stable name, this way they can safely be referenced over the network
+		//
+
+		// The idea here is to assume both the server and client will generate the same list when calling GetSubobjectsWithStableNamesForNetworking
+		// We can take advantage of this by only sending the first GUID, and then having the client generate the same guids deterministically from there
+		// The only caveat is that some guids have already been assigned before getting here (they might be referenced, and were pre-assigned)
+		// To get around this, we cherry pick these guids, and send them anyhow, skipping over them when doing the range assign
+
+		TArray< UObject * > Subobjs;
+
+		Actor->GetSubobjectsWithStableNamesForNetworking( Subobjs );
+
+		if ( Ar.IsSaving() )
 		{
-			TArray<UObject*> Subobjs;
-			Actor->GetSubobjectsWithStableNamesForNetworking(Subobjs);
+			uint32 Checksum = 0;
 
-			int32 NumReplicatedComponents = Subobjs.Num();
+			uint16 NumSubobjects = Subobjs.Num();
+			Ar << NumSubobjects;
 
-			Ar << NumReplicatedComponents;
+			bool bRangeStartSend = false;
 
-			for (UObject *Obj : Subobjs)
+			for ( UObject * Obj : Subobjs )
 			{
-				FNetworkGUID SubobjNetGUID;
-				InternalGetOrAssignNetGUID(Obj, SubobjNetGUID);
-				Ar << SubobjNetGUID;
+				FNetworkGUID SubobjNetGUID = Cache->NetGUIDLookup.FindRef( Obj );
+
+				uint8 bForceSendGuid = SubobjNetGUID.IsValid() ? 1 : 0;
+
+				if ( !bForceSendGuid )
+				{
+					InternalGetOrAssignNetGUID( Obj, SubobjNetGUID );
+				}
+
+				check( SubobjNetGUID.IsDynamic() );		// Make sure some assumptions are correct
+
+				Ar.SerializeBits( &bForceSendGuid, 1 );
 
-				FString ObjName = Obj->GetName();
-				Ar << ObjName;
+				if ( bForceSendGuid || !bRangeStartSend )
+				{
+					Ar << SubobjNetGUID;				// Save guid if forced, or if we haven't sent the start of the range
+
+					if ( !bForceSendGuid )
+					{
+						// If this is the start of the range, we don't need to send it again
+						bRangeStartSend = true;
+					}
+				}
 
-				NetGUIDAckStatus.Add(SubobjNetGUID, GUID_PACKET_ACKED);		// We know this GUID is acked since the export is in this bunch, and this bunch is reliable
-				UE_LOG(LogNetPackageMap, Log, TEXT("   Saving: Assigned NetGUID <%s> to subobject %s"), *SubobjNetGUID.ToString(), *Obj->GetPathName());
+				// Evolve checksum so we can sanity check
+				Checksum = FCrc::MemCrc32( &SubobjNetGUID.Value, sizeof( SubobjNetGUID.Value ), Checksum );
+				Checksum = FCrc::StrCrc32( *Obj->GetName(), Checksum );
+
+				check( !PendingAckGUIDs.Contains( SubobjNetGUID ) );
+
+				UE_LOG( LogNetPackageMap, Log, TEXT( "SerializeNewActor: [Saving] Assigned NetGUID <%s> to subobject %s" ), *SubobjNetGUID.ToString(), *Obj->GetPathName() );
+			}
+
+			if ( NumSubobjects > 0 )
+			{
+				Ar << Checksum;
 			}
 		}
 		else
 		{
-			int32 NumReplicatedComponents = 0;
-			Ar << NumReplicatedComponents;
+			uint16 NumSubobjects = 0;
+			Ar << NumSubobjects;
 
-			for (int32 idx=0; idx < NumReplicatedComponents && !Ar.IsError(); ++idx)
+			if ( NumSubobjects > 0 )
 			{
-				FNetworkGUID SubobjNetGUID;
-				Ar << SubobjNetGUID;
+				FNetworkGUID LastRangedNetGUID;
+
+				uint32 Checksum = 0;
 
-				FString ObjName;
-				Ar << ObjName;
+				TArray< FNetworkGUID > SubObjGuids;
+
+				for ( int32 i = 0; i < NumSubobjects; i++ )
+				{
+					FNetworkGUID LocalNetGUID;
+
+					uint8 bForceSendGuid = false;
+					Ar.SerializeBits( &bForceSendGuid, 1 );
+
+					if ( bForceSendGuid || !LastRangedNetGUID.IsValid() )
+					{
+						// Server only saves out the very first guid (or guid that aren't part of the range), we derive the rest
+						Ar << LocalNetGUID;
+					}
+					else
+					{
+						// Derive the next guid from the last one
+						LocalNetGUID.Value = ( ( LastRangedNetGUID.Value >> 1 ) + 1 ) << 1;
+					}
+
+					if ( !bForceSendGuid )
+					{
+						LastRangedNetGUID = LocalNetGUID;
+					}
+
+					SubObjGuids.Add( LocalNetGUID );
+
+					if ( i < Subobjs.Num() )
+					{
+						UObject * Obj = Subobjs[i];
+
+						// Evolve checksum so we can sanity check
+						Checksum = FCrc::MemCrc32( &LocalNetGUID.Value, sizeof( LocalNetGUID.Value ), Checksum );
+						Checksum = FCrc::StrCrc32( *Obj->GetName(), Checksum );
+					}
+				}
 
-				UObject * Subobj = NetGUIDAssign(SubobjNetGUID, ObjName, Actor);
+				uint32 ServerChecksum = 0;
+				Ar << ServerChecksum;
 
-				UE_LOG(LogNetPackageMap, Log, TEXT("   Loading: Assigned NetGUID <%s> to subobject %s"), *SubobjNetGUID.ToString(), *Subobj->GetPathName());
+				// Sanity check the results
+				if ( NumSubobjects != Subobjs.Num() )
+				{
+					UE_LOG( LogNetPackageMap, Error, TEXT( "SerializeNewActor: Invalid number of sub-objects: %s (%i != %i)" ), *Actor->GetFullName(), NumSubobjects, Subobjs.Num() );
+				}
+				else if ( Checksum != ServerChecksum )
+				{
+					UE_LOG( LogNetPackageMap, Error, TEXT( "SerializeNewActor: Subobject checksum FAILED: %s" ), *Actor->GetFullName() );
+				}
+				else
+				{
+					// Everything worked, assign the guids
+					check( SubObjGuids.Num() == Subobjs.Num() );
+
+					for ( int32 i = 0; i < Subobjs.Num(); i++ )
+					{
+						AssignNetGUID( Subobjs[i], SubObjGuids[i] );
+						UE_LOG( LogNetPackageMap, Log, TEXT( "SerializeNewActor: [Loading] Assigned NetGUID <%s> to subobject %s" ), *SubObjGuids[i].ToString(), *Subobjs[i]->GetPathName() );
+					}
+				}
 			}
 		}
 	}
 
-	UE_LOG(LogNetPackageMap, VeryVerbose, TEXT("END SerializeNewActor"));
-	UE_LOG(LogNetPackageMap, Log, TEXT("Finished Serializing Actor %s <%s> on Channel %d"), Actor ? *Actor->GetName() : TEXT("NULL"), *NetGUID.ToString(), Channel->ChIndex );
+	UE_LOG(LogNetPackageMap, Log, TEXT("SerializeNewActor END: Finished Serializing Actor %s <%s> on Channel %d"), Actor ? *Actor->GetName() : TEXT("NULL"), *NetGUID.ToString(), Channel->ChIndex );
 
 	return SpawnedActor;
 }
@@ -1159,8 +1252,10 @@ void UPackageMapClient::HandleUnAssignedObject(const UObject* Obj)
 /** Do we need to include the full path of this object for the client to resolve it? */
 bool UPackageMapClient::ShouldSendFullPath(const UObject* Object, const FNetworkGUID &NetGUID)
 {
-	if (!Connection)
+	if ( !Connection )
+	{
 		return false;
+	}
 
 	// NetGUID is already exported
 	if ( CurrentExportBunch != NULL && CurrentExportBunch->ExportNetGUIDs.Contains( NetGUID ) )
@@ -1168,26 +1263,22 @@ bool UPackageMapClient::ShouldSendFullPath(const UObject* Object, const FNetwork
 		return false;
 	}
 
-	if (!NetGUID.IsValid())
+	if ( !NetGUID.IsValid() )
+	{
 		return false;
+	}
 
-	if (NetGUID.IsDefault())
-		return true;
-
-	bool IsDynamic = IsDynamicObject(Object);
-
-	// If static object that hasn't been ack'd
-	if (!IsDynamic && (Locked || !NetGUIDHasBeenAckd(NetGUID)))
-		return true;
+	if ( NetGUID.IsDynamic() )
+	{
+		return false;		// We only export fully stably named objects
+	}
 
-	// If it is a dynamic actor component that hasn't been Ack'd, then send its path
-	const UActorComponent *Component = Cast<const UActorComponent>(Object);
-	if (IsDynamic && (Component && Component->IsNameStableForNetworking()) && !NetGUIDHasBeenAckd(NetGUID) )
+	if ( NetGUID.IsDefault() )
 	{
 		return true;
 	}
 
-	return false;
+	return Locked || !NetGUIDHasBeenAckd( NetGUID );
 }
 
 /**
