Reset physx actor on destructible component when destryong it.
Read comment for more info

diff --git a/Engine/Source/Runtime/Engine/Private/DestructibleComponent.cpp b/Engine/Source/Runtime/Engine/Private/DestructibleComponent.cpp
--- a/Engine/Source/Runtime/Engine/Private/DestructibleComponent.cpp
+++ b/Engine/Source/Runtime/Engine/Private/DestructibleComponent.cpp
@@ -353,6 +353,14 @@ void UDestructibleComponent::DestroyPhysicsState()
 		}
 		
 		ApexDestructibleActor = NULL;
+		
+		//Destructible component uses the BodyInstance in PrimitiveComponent in a very dangerous way. It assigns PxRigidDynamic to it as it needs it.
+		//Destructible PxRigidDynamic actors can be deleted from under us as PhysX sees fit.
+		//Ideally we wouldn't ever have a dangling pointer, but in practice this is hard to avoid.
+		//In theory anyone using BodyInstance on a PrimitiveComponent should be using functions like GetBodyInstance - in which case we properly fix up the dangling pointer
+		BodyInstance.RigidActorSync = NULL;
+		BodyInstance.RigidActorAsync = NULL;
+
 	}
 #endif	// #if WITH_APEX
 	Super::DestroyPhysicsState();
@@ -1088,6 +1096,12 @@ bool UDestructibleComponent::SweepComponent(FHitResult& OutHit, const FVector St
 #if WITH_APEX
 void UDestructibleComponent::SetupFakeBodyInstance( physx::PxRigidActor* NewRigidActor, int32 InstanceIdx, FFakeBodyInstanceState* PrevState )
 {
+	//This code is very dangerous, but at the moment I have no better solution:
+	//Destructible component assigns PxRigidDynamic to the BodyInstance as it needs it.
+	//Destructible PxRigidDynamic actors can be deleted from under us as PhysX sees fit.
+	//Ideally we wouldn't ever have a dangling pointer, but in practice this is hard to avoid.
+	//In theory anyone using BodyInstance on a PrimitiveComponent should be using functions like GetBodyInstance - in which case we properly fix up the dangling pointer
+
 	if (PrevState != NULL)
 	{
 		PrevState->ActorSync = BodyInstance.RigidActorSync;
