#ue4 - Revert a change from an integration that changed the behavior or bOrientRotationToMovement and bUseControllerDesiredRotation.

- Also add an early-out to PhysicsRotation() in case there is nothing to be done based on the rotation settings.
- Reordered some code in CalcVelocity() for minor optimization wins.

diff --git a/Engine/Source/Runtime/Engine/Private/CharacterMovementComponent.cpp b/Engine/Source/Runtime/Engine/Private/CharacterMovementComponent.cpp
--- a/Engine/Source/Runtime/Engine/Private/CharacterMovementComponent.cpp
+++ b/Engine/Source/Runtime/Engine/Private/CharacterMovementComponent.cpp
@@ -1918,7 +1918,7 @@ void UCharacterMovementComponent::CalcVelocity(float DeltaTime, float Friction,
 		ApplyVelocityBraking(DeltaTime, Friction, BrakingDeceleration);
 	
 		// Don't allow braking to lower us below max speed if we started above it.
-		if (FVector::DotProduct(Acceleration, OldVelocity) > 0.0f && bVelocityOverMax && Velocity.SizeSquared() < FMath::Square(MaxSpeed))
+		if (bVelocityOverMax && Velocity.SizeSquared() < FMath::Square(MaxSpeed) && FVector::DotProduct(Acceleration, OldVelocity) > 0.0f)
 		{
 			Velocity = SafeNormalPrecise(OldVelocity) * MaxSpeed;
 		}
@@ -3364,11 +3364,16 @@ void UCharacterMovementComponent::PhysicsRotation(float DeltaTime)
 		return;
 	}
 
+	if (!(bOrientRotationToMovement || bUseControllerDesiredRotation))
+	{
+		return;
+	}
+
 	const FRotator CurrentRotation = CharacterOwner->GetActorRotation();
 	FRotator DeltaRot = GetDeltaRotation(DeltaTime);
 	FRotator DesiredRotation = CurrentRotation;
 
-	if (bOrientRotationToMovement && (Velocity.IsZero() == false))
+	if (bOrientRotationToMovement)
 	{
 		DesiredRotation = ComputeOrientToMovementRotation(CurrentRotation, DeltaTime, DeltaRot);
 	}
