UActorComponent should call super IsNameStableForNetworking

diff --git a/Engine/Source/Runtime/Engine/Private/ActorComponent.cpp b/Engine/Source/Runtime/Engine/Private/ActorComponent.cpp
--- a/Engine/Source/Runtime/Engine/Private/ActorComponent.cpp
+++ b/Engine/Source/Runtime/Engine/Private/ActorComponent.cpp
@@ -1088,7 +1088,7 @@ bool UActorComponent::IsNameStableForNetworking() const
 	 *	-They were explicitly set to bNetAddressable (blueprint components created by SCS)
 	 */
 
-	return bNetAddressable || IsDefaultSubobject() || HasAnyFlags( RF_WasLoaded | RF_DefaultSubObject );
+	return bNetAddressable || Super::IsNameStableForNetworking();
 }
 
 bool UActorComponent::IsSupportedForNetworking() const
