OnMovementModeChange multicast delegate

diff --git a/Engine/Source/Runtime/Engine/Classes/GameFramework/Character.h b/Engine/Source/Runtime/Engine/Classes/GameFramework/Character.h
--- a/Engine/Source/Runtime/Engine/Classes/GameFramework/Character.h
+++ b/Engine/Source/Runtime/Engine/Classes/GameFramework/Character.h
@@ -4,6 +4,7 @@
 
 #include "Character.generated.h"
 
+DECLARE_DYNAMIC_MULTICAST_DELEGATE_ThreeParams(FMovementModeChangedSignature, class ACharacter*, Character, EMovementMode, PrevMovementMode, uint8, PreviousCustomMode);
 //
 // Forward declarations
 //
@@ -449,6 +450,8 @@ public:
 	 */
 	virtual void OnMovementModeChanged(EMovementMode PrevMovementMode, uint8 PreviousCustomMode = 0);
 
+	/** Native multicast delegate for MovementMode changing. */
+	FMovementModeChangedSignature MovementModeChangedDelegate;
 	/**
 	 * Called from CharacterMovementComponent to notify the character that the movement mode has changed.
 	 * @param	PrevMovementMode	Movement mode before the change
diff --git a/Engine/Source/Runtime/Engine/Private/Character.cpp b/Engine/Source/Runtime/Engine/Private/Character.cpp
--- a/Engine/Source/Runtime/Engine/Private/Character.cpp
+++ b/Engine/Source/Runtime/Engine/Private/Character.cpp
@@ -666,6 +666,7 @@ void ACharacter::LaunchCharacter(FVector LaunchVelocity, bool bXYOverride, bool
 void ACharacter::OnMovementModeChanged(EMovementMode PrevMovementMode, uint8 PrevCustomMode)
 {
 	K2_OnMovementModeChanged(PrevMovementMode, CharacterMovement->MovementMode, PrevCustomMode, CharacterMovement->CustomMovementMode);
+	MovementModeChangedDelegate.Broadcast(this, PrevMovementMode, PrevCustomMode);
 }
 
 
