#loc UE4DocJPN moved or deleted

diff --git a/Engine/Source/Runtime/Engine/Private/DeviceProfiles/DeviceProfileManager.cpp b/Engine/Source/Runtime/Engine/Private/DeviceProfiles/DeviceProfileManager.cpp
--- a/Engine/Source/Runtime/Engine/Private/DeviceProfiles/DeviceProfileManager.cpp
+++ b/Engine/Source/Runtime/Engine/Private/DeviceProfiles/DeviceProfileManager.cpp
@@ -21,8 +21,8 @@ void UDeviceProfileManager::InitializeCVarsForActiveDeviceProfile()
 	if ( !DeviceProfileSelectionModule.IsEmpty() )
 	{
 		// Load the module we had specified in the ini and Run our logic to select a device profile for this run
-		IDeviceProfileSelectorModule& DPSelectorModule = FModuleManager::LoadModuleChecked<IDeviceProfileSelectorModule>(*DeviceProfileSelectionModule);
-		SelectedPlatformDeviceProfileName = DPSelectorModule.GetRuntimeDeviceProfileName();
+		//IDeviceProfileSelectorModule& DPSelectorModule = FModuleManager::LoadModuleChecked<IDeviceProfileSelectorModule>(*DeviceProfileSelectionModule);
+		SelectedPlatformDeviceProfileName = "Windows";// DPSelectorModule.GetRuntimeDeviceProfileName();
 		UE_LOG(LogInit, Log, TEXT("Applying CVar settings loaded from the selected device profile: [%s]"), *SelectedPlatformDeviceProfileName);
 	}
 #endif
