Fixed crash when foliage instances are based on Blueprint-generated components
TTP#336530

diff --git a/Engine/Source/Editor/FoliageEdit/Private/FoliageEdMode.cpp b/Engine/Source/Editor/FoliageEdit/Private/FoliageEdMode.cpp
--- a/Engine/Source/Editor/FoliageEdit/Private/FoliageEdMode.cpp
+++ b/Engine/Source/Editor/FoliageEdit/Private/FoliageEdMode.cpp
@@ -1846,6 +1846,11 @@ bool FEdModeFoliage::InputKey( FLevelEditorViewportClient* ViewportClient, FView
 								Instance.Location = Hit.Location;
 								Instance.ZOffset = 0.f;
 								Instance.Base = Hit.Component.Get();
+								// We cannot be based on an a blueprint component as these will disappear when the construction script is re-run
+								if (Instance.Base && Instance.Base->bCreatedByConstructionScript)
+								{
+									Instance.Base = NULL;
+								}
 
 								if( Instance.Flags & FOLIAGE_AlignToNormal )
 								{
diff --git a/Engine/Source/Runtime/Engine/Private/InstancedFoliage.cpp b/Engine/Source/Runtime/Engine/Private/InstancedFoliage.cpp
--- a/Engine/Source/Runtime/Engine/Private/InstancedFoliage.cpp
+++ b/Engine/Source/Runtime/Engine/Private/InstancedFoliage.cpp
@@ -1518,6 +1518,12 @@ void AInstancedFoliageActor::PostLoad()
 				}
 				else
 				{
+					// Clear out the Base for any instances based on blueprint components, as they will be destroyed when the construction scripts are re-run.
+					if (Instance.Base && Instance.Base->bCreatedByConstructionScript)
+					{
+						Instance.Base = NULL;
+					}
+
 					// Add valid instances to the hash.
 					MeshInfo.InstanceHash->InsertInstance(Instance.Location, InstanceIdx);
 					FFoliageComponentHashInfo* ComponentHashInfo = MeshInfo.ComponentHash.Find(Instance.Base);
