World Browser: Added support for tiles uniqueness by long package name instead of just by short name

diff --git a/Engine/Source/Runtime/Engine/Private/WorldComposition.cpp b/Engine/Source/Runtime/Engine/Private/WorldComposition.cpp
--- a/Engine/Source/Runtime/Engine/Private/WorldComposition.cpp
+++ b/Engine/Source/Runtime/Engine/Private/WorldComposition.cpp
@@ -109,9 +109,12 @@ struct FPackageNameAndLODIndex
 struct FWorldTilesGatherer
 	: public IPlatformFile::FDirectoryVisitor
 {
-	// Mapping PackageName -> LOD package names
-	TMap<FString, FTileLODCollection> TileLODCollections;
-
+	// List of tile long pakage names (non LOD)
+	TArray<FString> TilesCollection;
+	
+	// Tile short pkg name -> Tiles LOD
+	TMultiMap<FString, FPackageNameAndLODIndex> TilesLODCollection;
+	
 	virtual bool Visit(const TCHAR* FilenameOrDirectory, bool bIsDirectory) OVERRIDE
 	{
 		FString FullPath = FilenameOrDirectory;
@@ -121,12 +124,20 @@ struct FWorldTilesGatherer
 		{
 			FString TilePackageName = FPackageName::FilenameToLongPackageName(FullPath);
 			FPackageNameAndLODIndex PackageNameLOD = BreakToNameAndLODIndex(TilePackageName);
-			FString ShortPackageName = FPackageName::GetShortName(PackageNameLOD.PackageName);
-
+						
 			if (PackageNameLOD.LODIndex != INDEX_NONE)
 			{
-				// Store in collection
-				TileLODCollections.FindOrAdd(ShortPackageName).PackageNames[PackageNameLOD.LODIndex] = TilePackageName;
+				if (PackageNameLOD.LODIndex == 0) 
+				{
+					// non-LOD tile
+					TilesCollection.Add(TilePackageName);
+				}
+				else
+				{
+					// LOD tile
+					FString TileShortName = FPackageName::GetShortName(PackageNameLOD.PackageName);
+					TilesLODCollection.Add(TileShortName, PackageNameLOD);
+				}
 			}
 		}
 	
@@ -189,17 +200,8 @@ void UWorldComposition::Rescan()
 	FString PersistentLevelPackageName = OwningWorld->GetOutermost()->GetName();
 	
 	// Add found tiles to a world composition, except persistent level
-	for (const auto& TileLODEntry : Gatherer.TileLODCollections)
+	for (const auto& TilePackageName : Gatherer.TilesCollection)
 	{
-		const FTileLODCollection& LODCollection = TileLODEntry.Value;
-		const FString& TilePackageName = LODCollection.PackageNames[0];
-
-		// Discard this entry in case original level was not found
-		if (TilePackageName.IsEmpty())
-		{
-			continue;
-		}
-
 		// Discard persistent level entry
 		if (TilePackageName == PersistentLevelPackageName)
 		{
@@ -218,17 +220,31 @@ void UWorldComposition::Rescan()
 		Tile.Info = Info;
 		
 		// Assign LOD tiles
-		for (int32 LODIdx = 1; LODIdx <= WORLDTILE_LOD_MAX_INDEX; ++LODIdx)
+		FString TileShortName = FPackageName::GetShortName(TilePackageName);
+		TArray<FPackageNameAndLODIndex> TileLODList;
+		Gatherer.TilesLODCollection.MultiFind(TileShortName, TileLODList);
+		if (TileLODList.Num())
 		{
-			const FString& TileLODPackageName = LODCollection.PackageNames[LODIdx];
-			if (TileLODPackageName.IsEmpty())
+			Tile.LODPackageNames.SetNum(WORLDTILE_LOD_MAX_INDEX);
+			FString TilePath = FPackageName::GetLongPackagePath(TilePackageName);
+			for (const auto& TileLOD : TileLODList)
 			{
-				break;
+				// LOD tiles should be in the same directory or in nested directory
+				// Basically tile path should be a prefix of a LOD tile path
+				if (TileLOD.PackageName.StartsWith(TilePath))
+				{
+					Tile.LODPackageNames[TileLOD.LODIndex-1] = FName(*FString::Printf(TEXT("%s_LOD%d"), *TileLOD.PackageName, TileLOD.LODIndex));
+				}
 			}
-			
-			Tile.LODPackageNames.Add(FName(*TileLODPackageName));
-		}
 
+			// Remove null entries in LOD list
+			int32 NullEntryIdx;
+			if (Tile.LODPackageNames.Find(FName(), NullEntryIdx))
+			{
+				Tile.LODPackageNames.SetNum(NullEntryIdx);
+			}
+		}
+		
 		Tiles.Add(Tile);
 	}
 
