Editoronly Bluepirnts: Blueprint wont be loaded outside editor, so UObjectLibrary must load BlueprintGeneratedClasses (that are not assets) instead of Blueprint objects.

Work in progress.

diff --git a/Engine/Source/Runtime/AssetRegistry/Public/AssetData.h b/Engine/Source/Runtime/AssetRegistry/Public/AssetData.h
--- a/Engine/Source/Runtime/AssetRegistry/Public/AssetData.h
+++ b/Engine/Source/Runtime/AssetRegistry/Public/AssetData.h
@@ -173,6 +173,26 @@ public:
 		return Asset;
 	}
 
+	UPackage* GetPackage() const
+	{
+		if (PackageName == NAME_None)
+		{
+			return NULL;
+		}
+
+		UPackage* Package = FindPackage(NULL, *PackageName.ToString());
+		if (Package)
+		{
+			Package->FullyLoad();
+		}
+		else
+		{
+			Package = LoadPackage(NULL, *PackageName.ToString(), LOAD_None);
+		}
+
+		return Package;
+	}
+
 	/** Returns true if the asset is loaded */
 	bool IsAssetLoaded() const
 	{
diff --git a/Engine/Source/Runtime/Engine/Private/EngineUtils.cpp b/Engine/Source/Runtime/Engine/Private/EngineUtils.cpp
--- a/Engine/Source/Runtime/Engine/Private/EngineUtils.cpp
+++ b/Engine/Source/Runtime/Engine/Private/EngineUtils.cpp
@@ -247,7 +247,7 @@ void FContentComparisonHelper::RecursiveObjectCollection(UObject* InStartObject,
 }
 #endif
 
-bool EngineUtils::FindOrLoadAssetsByPath(const FString& Path, TArray<UObject*>& OutAssets)
+bool EngineUtils::FindOrLoadAssetsByPath(const FString& Path, TArray<UObject*>& OutAssets, EAssetToLoad Type)
 {
 	if ( !FPackageName::IsValidLongPackageName(Path, true) )
 	{
@@ -272,7 +272,6 @@ bool EngineUtils::FindOrLoadAssetsByPath(const FString& Path, TArray<UObject*>&
 	}
 
 	// Load packages or find existing ones and fully load them
-	TSet<UPackage*> Packages;
 	for (int32 FileIdx = 0; FileIdx < Filenames.Num(); ++FileIdx)
 	{
 		const FString& Filename = Filenames[FileIdx];
@@ -290,22 +289,20 @@ bool EngineUtils::FindOrLoadAssetsByPath(const FString& Path, TArray<UObject*>&
 
 		if (Package)
 		{
-			Packages.Add(Package);
-		}
-	}
-
-	// If any packages were successfully loaded, find all assets that were in the packages and add them to OutAssets
-	if ( Packages.Num() > 0 )
-	{
-		for (FObjectIterator ObjIt; ObjIt; ++ObjIt)
-		{
-			if ( Packages.Contains(ObjIt->GetOutermost()) && ObjIt->IsAsset() )
+			TArray<UObject*> ObjectsInPackage;
+			GetObjectsWithOuter(Package, ObjectsInPackage);
+			for (UObject* Object : ObjectsInPackage)
 			{
-				OutAssets.Add(*ObjIt);
+				const bool bWantedType = 
+					((EAssetToLoad::ATL_Regular == Type) && Object->IsAsset()) ||
+					((EAssetToLoad::ATL_Class == Type) && Object->IsA<UClass>());
+				if (bWantedType)
+				{
+					OutAssets.Add(Object);
+				}
 			}
 		}
 	}
-
 	return true;
 }
 
diff --git a/Engine/Source/Runtime/Engine/Private/ObjectLibrary.cpp b/Engine/Source/Runtime/Engine/Private/ObjectLibrary.cpp
--- a/Engine/Source/Runtime/Engine/Private/ObjectLibrary.cpp
+++ b/Engine/Source/Runtime/Engine/Private/ObjectLibrary.cpp
@@ -197,7 +197,7 @@ int32 UObjectLibrary::LoadAssetsFromPaths(TArray<FString> Paths)
 	{
 		TArray<UObject*> LoadedObjects;
 		FString Path = Paths[PathIndex];
-		if (EngineUtils::FindOrLoadAssetsByPath(Path, LoadedObjects))
+		if (EngineUtils::FindOrLoadAssetsByPath(Path, LoadedObjects, bHasBlueprintClasses ? EngineUtils::ATL_Class : EngineUtils::ATL_Regular))
 		{
 			for (int32 i = 0; i < LoadedObjects.Num(); ++i)
 			{
@@ -239,24 +239,19 @@ int32 UObjectLibrary::LoadBlueprintsFromPaths(TArray<FString> Paths)
 		TArray<UObject*> LoadedObjects;
 		FString Path = Paths[PathIndex];
 
-		if (EngineUtils::FindOrLoadAssetsByPath(Path, LoadedObjects))
+		if (EngineUtils::FindOrLoadAssetsByPath(Path, LoadedObjects, EngineUtils::ATL_Class))
 		{
 			for (int32 i = 0; i < LoadedObjects.Num(); ++i)
 			{
-				UBlueprintCore* Blueprint = Cast<UBlueprintCore>(LoadedObjects[i]);
+				auto Class = Cast<UBlueprintGeneratedClass>(LoadedObjects[i]);
 
-				if (Blueprint && Blueprint->GeneratedClass)
+				if (Class == NULL || (ObjectBaseClass && !Class->IsChildOf(ObjectBaseClass)))
 				{
-					UClass* Class = Blueprint->GeneratedClass;
-
-					if (Class == NULL || (ObjectBaseClass && !Class->IsChildOf(ObjectBaseClass)))
-					{
-						continue;
-					}
-
-					AddObject(Class);
-					Count++;
+					continue;
 				}
+
+				AddObject(Class);
+				Count++;
 			}
 		}
 	}
@@ -396,28 +391,38 @@ int32 UObjectLibrary::LoadAssetsFromAssetData()
 	{
 		FAssetData& Data = AssetDataList[AssetIdx];
 
-		UObject *LoadedObject = Data.GetAsset();
-
-		if (!LoadedObject)
-		{
-			continue;
-		}
-
-		if (bHasBlueprintClasses && LoadedObject->IsA(UBlueprintCore::StaticClass()))
+		UObject *LoadedObject = NULL;
+			
+		if (!bHasBlueprintClasses)
 		{
-			// If we're looking for non-blueprints and are a blueprint, look for generated class
-			UBlueprintCore* LoadedBlueprint = Cast<UBlueprintCore>(LoadedObject);
-			checkSlow(!ObjectBaseClass || LoadedBlueprint->GeneratedClass->IsChildOf(ObjectBaseClass));
+			LoadedObject = Data.GetAsset();
 
-			LoadedObject = LoadedBlueprint->GeneratedClass;
+			checkSlow(!LoadedObject || !ObjectBaseClass || LoadedObject->IsA(ObjectBaseClass));
 		}
 		else
 		{
-			checkSlow(!ObjectBaseClass || LoadedObject->IsA(ObjectBaseClass));
+			UPackage* Package = Data.GetPackage();
+			if (Package)
+			{
+				TArray<UObject*> ObjectsInPackage;
+				GetObjectsWithOuter(Package, ObjectsInPackage);
+				for (UObject* PotentialBPGC : ObjectsInPackage)
+				{
+					if (auto LoadedBPGC = Cast<UBlueprintGeneratedClass>(PotentialBPGC))
+					{
+						checkSlow(!ObjectBaseClass || LoadedBPGC->IsChildOf(ObjectBaseClass));
+						LoadedObject = LoadedBPGC;
+						break; //there is usually only one BGPC in a package
+					}
+				}
+			}
 		}
 
-		AddObject(LoadedObject);
-		Count++;
+		if (LoadedObject)
+		{
+			AddObject(LoadedObject);
+			Count++;
+		}
 	}
 
 	return Count;
diff --git a/Engine/Source/Runtime/Engine/Private/Slate/SlateGameResources.cpp b/Engine/Source/Runtime/Engine/Private/Slate/SlateGameResources.cpp
--- a/Engine/Source/Runtime/Engine/Private/Slate/SlateGameResources.cpp
+++ b/Engine/Source/Runtime/Engine/Private/Slate/SlateGameResources.cpp
@@ -173,7 +173,7 @@ void FSlateGameResources::Initialize( const FString& ScopeToDirectory, const FSt
 	BasePath = InBasePath;
 
 	TArray<UObject*> LoadedObjects;
-	if ( EngineUtils::FindOrLoadAssetsByPath( ContentRootDir, LoadedObjects ) )
+	if (EngineUtils::FindOrLoadAssetsByPath(ContentRootDir, LoadedObjects, EngineUtils::ATL_Regular))
 	{
 		for( auto ObjectIter = LoadedObjects.CreateConstIterator(); ObjectIter; ++ObjectIter )
 		{
diff --git a/Engine/Source/Runtime/Engine/Public/EngineUtils.h b/Engine/Source/Runtime/Engine/Public/EngineUtils.h
--- a/Engine/Source/Runtime/Engine/Public/EngineUtils.h
+++ b/Engine/Source/Runtime/Engine/Public/EngineUtils.h
@@ -650,8 +650,14 @@ protected:
 
 namespace EngineUtils
 {
+	enum EAssetToLoad
+	{
+		ATL_Regular,
+		ATL_Class,
+	};
+
 	/** Loads all the assets found in the specified path and subpaths */
-	ENGINE_API bool FindOrLoadAssetsByPath(const FString& Path, TArray<UObject*>& OutAssets);
+	ENGINE_API bool FindOrLoadAssetsByPath(const FString& Path, TArray<UObject*>& OutAssets, EAssetToLoad Type);
 }
 
 /** Helper class for serializing flags describing which data have been stripped (if any). */
