Added LOD information to "stat levels" console command, also fixed several issues with it

diff --git a/Engine/Source/Runtime/Engine/Private/EngineUtils.cpp b/Engine/Source/Runtime/Engine/Private/EngineUtils.cpp
--- a/Engine/Source/Runtime/Engine/Private/EngineUtils.cpp
+++ b/Engine/Source/Runtime/Engine/Private/EngineUtils.cpp
@@ -306,14 +306,24 @@ bool EngineUtils::FindOrLoadAssetsByPath(const FString& Path, TArray<UObject*>&
 	return true;
 }
 
-
-/** This will set the StreamingLevels TMap with the current Streaming Level Status and also set which level the player is in **/
-void GetLevelStreamingStatus( UWorld* World, TMap<FName,int32>& StreamingLevels, FString& LevelPlayerIsInName )
+TArray<FSubLevelStatus> GetSubLevelsStatus( UWorld* World )
 {
+	TArray<FSubLevelStatus> Result;
 	FWorldContext &Context = GEngine->GetWorldContextFromWorldChecked(World);
 
+	Result.Reserve(World->StreamingLevels.Num() + 1);
+	
+	// Add persistent level
+	{
+		FSubLevelStatus LevelStatus = {};
+		LevelStatus.PackageName = World->GetOutermost()->GetFName();
+		LevelStatus.StreamingStatus = LEVEL_Visible;
+		LevelStatus.LODIndex	= INDEX_NONE;
+		Result.Add(LevelStatus);
+	}
+	
 	// Iterate over the world info's level streaming objects to find and see whether levels are loaded, visible or neither.
-	for( int32 LevelIndex=0; LevelIndex<World->StreamingLevels.Num(); LevelIndex++ )
+	for( int32 LevelIndex=0; LevelIndex < World->StreamingLevels.Num(); LevelIndex++ )
 	{
 		ULevelStreaming* LevelStreaming = World->StreamingLevels[LevelIndex];
 
@@ -322,28 +332,32 @@ void GetLevelStreamingStatus( UWorld* World, TMap<FName,int32>& StreamingLevels,
 			&&	LevelStreaming->PackageName != World->GetOutermost()->GetFName() )
 		{
 			ULevel* Level = LevelStreaming->GetLoadedLevel();
+			FSubLevelStatus LevelStatus = {};
+			LevelStatus.PackageName = LevelStreaming->PackageName;
+			LevelStatus.LODIndex	= LevelStreaming->GetLODIndex(World);
+
 			if( Level != NULL )
 			{
 				if( World->ContainsLevel( Level ) == true )
 				{
 					if( World->CurrentLevelPendingVisibility == Level )
 					{
-						StreamingLevels.Add( LevelStreaming->PackageName, LEVEL_MakingVisible );
+						LevelStatus.StreamingStatus = LEVEL_MakingVisible;
 					}
 					else
 					{
-						StreamingLevels.Add( LevelStreaming->PackageName, LEVEL_Visible );
+						LevelStatus.StreamingStatus = LEVEL_Visible;
 					}
 				}
 				else
 				{
-					StreamingLevels.Add( LevelStreaming->PackageName, LEVEL_Loaded );
+					LevelStatus.StreamingStatus = LEVEL_Loaded;
 				}
 			}
 			else
 			{
 				// See whether the level's world object is still around.
-				UPackage* LevelPackage	= Cast<UPackage>(StaticFindObjectFast( UPackage::StaticClass(), NULL, LevelStreaming->PackageName ));
+				UPackage* LevelPackage	= FindObjectFast<UPackage>(NULL, LevelStatus.PackageName);
 				UWorld*	  LevelWorld	= NULL;
 				if( LevelPackage )
 				{
@@ -352,17 +366,19 @@ void GetLevelStreamingStatus( UWorld* World, TMap<FName,int32>& StreamingLevels,
 
 				if( LevelWorld )
 				{
-					StreamingLevels.Add( LevelStreaming->PackageName, LEVEL_UnloadedButStillAround );
+					LevelStatus.StreamingStatus = LEVEL_UnloadedButStillAround;
 				}
-				else if( GetAsyncLoadPercentage( *LevelStreaming->PackageName.ToString() ) >= 0 )
+				else if( GetAsyncLoadPercentage( LevelStatus.PackageName.ToString() ) >= 0 )
 				{
-					StreamingLevels.Add( LevelStreaming->PackageName, LEVEL_Loading );
+					LevelStatus.StreamingStatus = LEVEL_Loading;
 				}
 				else
 				{
-					StreamingLevels.Add( LevelStreaming->PackageName, LEVEL_Unloaded );
+					LevelStatus.StreamingStatus = LEVEL_Unloaded;
 				}
 			}
+
+			Result.Add(LevelStatus);
 		}
 	}
 
@@ -371,12 +387,15 @@ void GetLevelStreamingStatus( UWorld* World, TMap<FName,int32>& StreamingLevels,
 	for( int32 LevelIndex=0; LevelIndex < Context.LevelsToLoadForPendingMapChange.Num(); LevelIndex++ )
 	{
 		const FName LevelName = Context.LevelsToLoadForPendingMapChange[LevelIndex];
-		StreamingLevels.Add(LevelName, LEVEL_Preloading);
+		
+		FSubLevelStatus LevelStatus = {};
+		LevelStatus.PackageName = LevelName;
+		LevelStatus.StreamingStatus = LEVEL_Preloading;
+		LevelStatus.LODIndex = INDEX_NONE;
+		Result.Add(LevelStatus);
 	}
 
 
-	ULevel* LevelPlayerIsIn = NULL;
-
 	for( FConstPlayerControllerIterator Iterator = World->GetPlayerControllerIterator(); Iterator; ++Iterator )
 	{
 		APlayerController* PlayerController = *Iterator;
@@ -391,35 +410,34 @@ void GetLevelStreamingStatus( UWorld* World, TMap<FName,int32>& StreamingLevels,
 			static FName NAME_FindLevel = FName(TEXT("FindLevel"), true);			
 			PlayerController->GetWorld()->LineTraceSingle(Hit,PlayerController->GetPawn()->GetActorLocation(), (PlayerController->GetPawn()->GetActorLocation()-FVector(0.f, 0.f, 256.f)), FCollisionQueryParams(NAME_FindLevel, true, PlayerController->GetPawn()), FCollisionObjectQueryParams(ECC_WorldStatic));
 
-			/** @todo UE4 FIXME
-			if( Hit.Level != NULL )
-			{
-				LevelPlayerIsIn = Hit.Level;
-			}
-			else 
-			*/
+			ULevel* LevelPlayerIsIn = NULL;
+
 			if( Hit.GetActor() != NULL )
 			{
 				LevelPlayerIsIn = Hit.GetActor()->GetLevel();
 			}
-			else if( Hit.Component != NULL && Hit.Component->GetOwner() != NULL )
+			else if( Hit.Component != NULL )
+			{
+				LevelPlayerIsIn = Hit.Component->GetComponentLevel();
+			}
+
+			if (LevelPlayerIsIn)
 			{
-				AActor* Owner = Hit.Component->GetOwner();
-				if (Owner)
+				FName LevelName = LevelPlayerIsIn->GetOutermost()->GetFName();
+				FSubLevelStatus* LevelStatusPlayerIn = Result.FindByPredicate([LevelName](const FSubLevelStatus& InLevelStatus)
 				{
-					LevelPlayerIsIn = Owner->GetLevel();
-				}
-				else
+					return InLevelStatus.PackageName == LevelName;
+				});
+
+				if (LevelStatusPlayerIn)
 				{
-					// This happens for BSP where the ModelComponent's outer is the level
-					LevelPlayerIsIn = Hit.Component->GetTypedOuter<ULevel>();
+					LevelStatusPlayerIn->bPlayerInside = true;
 				}
 			}
 		}
 	}
 
-	// this no longer seems to be getting the correct level name :-(
-	LevelPlayerIsInName = LevelPlayerIsIn != NULL ? LevelPlayerIsIn->GetOutermost()->GetName() : TEXT("None");
+	return Result;
 }
 
 //////////////////////////////////////////////////////////////////////////
diff --git a/Engine/Source/Runtime/Engine/Private/UnrealEngine.cpp b/Engine/Source/Runtime/Engine/Private/UnrealEngine.cpp
--- a/Engine/Source/Runtime/Engine/Private/UnrealEngine.cpp
+++ b/Engine/Source/Runtime/Engine/Private/UnrealEngine.cpp
@@ -5855,34 +5855,35 @@ static FVector2D RotateVec2D(const FVector2D InVec, float RotAngle)
 #if !UE_BUILD_SHIPPING
 bool UEngine::HandleLogoutStatLevelsCommand( const TCHAR* Cmd, FOutputDevice& Ar, UWorld* InWorld )
 {
-	TMap<FName,int32> StreamingLevels;
-	FString LevelPlayerIsInName;
-	GetLevelStreamingStatus( InWorld, StreamingLevels, LevelPlayerIsInName );
+	const TArray<FSubLevelStatus> SubLevelsStatusList = GetSubLevelsStatus(InWorld);
 
-	Ar.Logf( TEXT( "Level Streaming:" ) );
+	Ar.Logf( TEXT( "Levels:" ) );
 
 	// now draw the "map" name
-	FString MapName	= InWorld->GetCurrentLevel()->GetOutermost()->GetName();
-
-	if( LevelPlayerIsInName == MapName )
+	if (SubLevelsStatusList.Num())
 	{
-		MapName = *FString::Printf( TEXT("->  %s"), *MapName );
-	}
-	else
-	{
-		MapName = *FString::Printf( TEXT("    %s"), *MapName );
-	}
+		// First entry - always persistent level
+		FString MapName	= SubLevelsStatusList[0].PackageName.ToString();
+		if (SubLevelsStatusList[0].bPlayerInside)
+		{
+			MapName = *FString::Printf( TEXT("->  %s"), *MapName );
+		}
+		else
+		{
+			MapName = *FString::Printf( TEXT("    %s"), *MapName );
+		}
 
-	Ar.Logf( TEXT( "%s" ), *MapName );
+		Ar.Logf( TEXT( "%s" ), *MapName );
+	}
 
 	// now log the levels
-	for( TMap<FName,int32>::TIterator It(StreamingLevels); It; ++It )
+	for (int32 LevelIdx = 1; LevelIdx < SubLevelsStatusList.Num(); ++LevelIdx)
 	{
-		FString LevelName = It.Key().ToString();
-		const int32 Status = It.Value();
+		const FSubLevelStatus& LevelStatus = SubLevelsStatusList[LevelIdx];
+		FString DisplayName = LevelStatus.PackageName.ToString();
 		FString StatusName;
 
-		switch( Status )
+		switch( LevelStatus.StreamingStatus )
 		{
 		case LEVEL_Visible:
 			StatusName = TEXT( "red loaded and visible" );
@@ -5906,33 +5907,37 @@ bool UEngine::HandleLogoutStatLevelsCommand( const TCHAR* Cmd, FOutputDevice& Ar
 			break;
 		};
 
+		if (LevelStatus.LODIndex != INDEX_NONE)
+		{
+			DisplayName += FString::Printf(TEXT(" [LOD%d]"), LevelStatus.LODIndex+1);
+		}
 
-		UPackage* LevelPackage = FindObject<UPackage>( NULL, *LevelName );
+		UPackage* LevelPackage = FindObjectFast<UPackage>( NULL, LevelStatus.PackageName );
 
 		if( LevelPackage 
 			&& (LevelPackage->GetLoadTime() > 0) 
-			&& (Status != LEVEL_Unloaded) )
+			&& (LevelStatus.StreamingStatus != LEVEL_Unloaded) )
 		{
-			LevelName += FString::Printf(TEXT(" - %4.1f sec"), LevelPackage->GetLoadTime());
+			DisplayName += FString::Printf(TEXT(" - %4.1f sec"), LevelPackage->GetLoadTime());
 		}
-		else if( GetAsyncLoadPercentage( *LevelName ) >= 0 )
+		else if( GetAsyncLoadPercentage( *LevelStatus.PackageName.ToString() ) >= 0 )
 		{
-			const int32 Percentage = FMath::TruncToInt( GetAsyncLoadPercentage( *LevelName ) );
-			LevelName += FString::Printf(TEXT(" - %3i %%"), Percentage ); 
+			const int32 Percentage = FMath::TruncToInt( GetAsyncLoadPercentage( *LevelStatus.PackageName.ToString() ) );
+			DisplayName += FString::Printf(TEXT(" - %3i %%"), Percentage ); 
 		}
 
-		if( LevelPlayerIsInName == LevelName )
+		if ( LevelStatus.bPlayerInside )
 		{
-			LevelName = *FString::Printf( TEXT("->  %s"), *LevelName );
+			DisplayName = *FString::Printf( TEXT("->  %s"), *DisplayName );
 		}
 		else
 		{
-			LevelName = *FString::Printf( TEXT("    %s"), *LevelName );
+			DisplayName = *FString::Printf( TEXT("    %s"), *DisplayName );
 		}
 
-		LevelName = FString::Printf( TEXT("%s \t\t%s"), *LevelName, *StatusName );
+		DisplayName = FString::Printf( TEXT("%s \t\t%s"), *DisplayName, *StatusName );
 
-		Ar.Logf( TEXT( "%s" ), *LevelName );
+		Ar.Logf( TEXT( "%s" ), *DisplayName );
 
 	}
 
@@ -10019,37 +10024,39 @@ int32 UEngine::RenderStatColorList(UWorld* World, FViewport* Viewport, FCanvas*
 int32 UEngine::RenderStatLevels(UWorld* World, FViewport* Viewport, FCanvas* Canvas, int32 X, int32 Y, const FVector* ViewLocation, const FRotator* ViewRotation)
 {
 	int32 MaxY = Y;
-	TMap<FName, int32> StreamingLevels;
-	FString LevelPlayerIsInName;
-	GetLevelStreamingStatus(World, StreamingLevels, LevelPlayerIsInName);
+	const TArray<FSubLevelStatus> SubLevelsStatusList = GetSubLevelsStatus(World);
 
 	// now do drawing to the screen
 
 	// Render unloaded levels in red, loaded ones in yellow and visible ones in green. Blue signifies that a level is unloaded but
 	// hasn't been garbage collected yet.
-	Canvas->DrawShadowedString(X, Y, TEXT("Level streaming"), GetSmallFont(), FLinearColor::White);
+	Canvas->DrawShadowedString(X, Y, TEXT("Levels"), GetSmallFont(), FLinearColor::White);
 	Y += 12;
 
-	// now draw the "map" name
-	FString MapName = World->GetCurrentLevel()->GetOutermost()->GetName();
-
-	if (LevelPlayerIsInName == MapName)
-	{
-		MapName = *FString::Printf(TEXT("->  %s"), *MapName);
-	}
-	else
+	if (SubLevelsStatusList.Num())
 	{
-		MapName = *FString::Printf(TEXT("    %s"), *MapName);
-	}
+		// First entry - always persistent level
+		FString MapName	= SubLevelsStatusList[0].PackageName.ToString();
+		if (SubLevelsStatusList[0].bPlayerInside)
+		{
+			MapName = *FString::Printf( TEXT("->  %s"), *MapName );
+		}
+		else
+		{
+			MapName = *FString::Printf( TEXT("    %s"), *MapName );
+		}
 
-	Canvas->DrawShadowedString(X, Y, *MapName, GetSmallFont(), FColor(127, 127, 127));
-	Y += 12;
+		Canvas->DrawShadowedString(X, Y, *MapName, GetSmallFont(), FColor(127, 127, 127));
+		Y += 12;
+	}
 
 	int32 BaseY = Y;
 
 	// now draw the levels
-	for (TMap<FName, int32>::TIterator It(StreamingLevels); It; ++It)
+	for (int32 LevelIdx = 1; LevelIdx < SubLevelsStatusList.Num(); ++LevelIdx)
 	{
+		const FSubLevelStatus& LevelStatus = SubLevelsStatusList[LevelIdx];
+		
 		// Wrap around at the bottom.
 		if (Y > Viewport->GetSizeXY().Y - 30)
 		{
@@ -10058,34 +10065,38 @@ int32 UEngine::RenderStatLevels(UWorld* World, FViewport* Viewport, FCanvas* Can
 			X += 250;
 		}
 
-		FString	LevelName = It.Key().ToString();
-		int32		Status = It.Value();
-		FColor	Color = GetColorForLevelStatus(Status);
+		FColor	Color = GetColorForLevelStatus(LevelStatus.StreamingStatus);
+		FString DisplayName = LevelStatus.PackageName.ToString();
 
-		UPackage* LevelPackage = FindObject<UPackage>(NULL, *LevelName);
+		if (LevelStatus.LODIndex != INDEX_NONE)
+		{
+			DisplayName += FString::Printf(TEXT(" [LOD%d]"), LevelStatus.LODIndex+1);
+		}
 
+		UPackage* LevelPackage = FindObjectFast<UPackage>(NULL, LevelStatus.PackageName);
+				
 		if (LevelPackage
 			&& (LevelPackage->GetLoadTime() > 0)
-			&& (Status != LEVEL_Unloaded))
+			&& (LevelStatus.StreamingStatus != LEVEL_Unloaded))
 		{
-			LevelName += FString::Printf(TEXT(" - %4.1f sec"), LevelPackage->GetLoadTime());
+			DisplayName += FString::Printf(TEXT(" - %4.1f sec"), LevelPackage->GetLoadTime());
 		}
-		else if (GetAsyncLoadPercentage(*LevelName) >= 0)
+		else if (GetAsyncLoadPercentage(*LevelStatus.PackageName.ToString()) >= 0)
 		{
-			const int32 Percentage = FMath::TruncToInt(GetAsyncLoadPercentage(*LevelName));
-			LevelName += FString::Printf(TEXT(" - %3i %%"), Percentage);
+			const int32 Percentage = FMath::TruncToInt(GetAsyncLoadPercentage(*LevelStatus.PackageName.ToString()));
+			DisplayName += FString::Printf(TEXT(" - %3i %%"), Percentage);
 		}
 
-		if (LevelPlayerIsInName == LevelName)
+		if (LevelStatus.bPlayerInside)
 		{
-			LevelName = *FString::Printf(TEXT("->  %s"), *LevelName);
+			DisplayName = *FString::Printf(TEXT("->  %s"), *DisplayName);
 		}
 		else
 		{
-			LevelName = *FString::Printf(TEXT("    %s"), *LevelName);
+			DisplayName = *FString::Printf(TEXT("    %s"), *DisplayName);
 		}
-
-		Canvas->DrawShadowedString(X + 4, Y, *LevelName, GetSmallFont(), Color);
+		
+		Canvas->DrawShadowedString(X + 4, Y, *DisplayName, GetSmallFont(), Color);
 		Y += 12;
 	}
 	return FMath::Max(MaxY, Y);
@@ -10098,17 +10109,13 @@ int32 UEngine::RenderStatLevelMap(UWorld* World, FViewport* Viewport, FCanvas* C
 	const FVector2D MapSize = FVector2D(512, 512);
 
 	// Get status of each sublevel (by name)
-	TMap<FName, int32> StreamingLevels;
-	FString LevelPlayerIsInName;
-	GetLevelStreamingStatus(World, StreamingLevels, LevelPlayerIsInName);
+	const TArray<FSubLevelStatus> SubLevelsStatusList = GetSubLevelsStatus(World);
 
 	// First iterate to find bounds of all streaming volumes
 	FBox AllVolBounds(0);
-	for (TMap<FName, int32>::TIterator It(StreamingLevels); It; ++It)
+	for (const FSubLevelStatus& LevelStatus : SubLevelsStatusList)
 	{
-		FName	LevelName = It.Key();
-
-		ULevelStreaming* LevelStreaming = World->GetLevelStreamingForPackageName(LevelName);
+		ULevelStreaming* LevelStreaming = World->GetLevelStreamingForPackageName(LevelStatus.PackageName);
 		if (LevelStreaming && LevelStreaming->bDrawOnLevelStatusMap)
 		{
 			AllVolBounds += LevelStreaming->GetStreamingVolumeBounds();
@@ -10137,16 +10144,13 @@ int32 UEngine::RenderStatLevelMap(UWorld* World, FViewport* Viewport, FCanvas* C
 
 
 	// Now we iterate and actually draw volumes
-	for (TMap<FName, int32>::TIterator It(StreamingLevels); It; ++It)
+	for (const FSubLevelStatus& LevelStatus : SubLevelsStatusList)
 	{
-		FName	LevelName = It.Key();
-		int32	Status = It.Value();
-
 		// Find the color to draw this level in
-		FColor StatusColor = GetColorForLevelStatus(Status);
+		FColor StatusColor = GetColorForLevelStatus(LevelStatus.StreamingStatus);
 		StatusColor.A = 64; // make it translucent
 
-		ULevelStreaming* LevelStreaming = World->GetLevelStreamingForPackageName(LevelName);
+		ULevelStreaming* LevelStreaming = World->GetLevelStreamingForPackageName(LevelStatus.PackageName);
 		if (LevelStreaming && LevelStreaming->bDrawOnLevelStatusMap)
 		{
 			for (int32 VolIdx = 0; VolIdx < LevelStreaming->EditorStreamingVolumes.Num(); VolIdx++)
diff --git a/Engine/Source/Runtime/Engine/Public/EngineUtils.h b/Engine/Source/Runtime/Engine/Public/EngineUtils.h
--- a/Engine/Source/Runtime/Engine/Public/EngineUtils.h
+++ b/Engine/Source/Runtime/Engine/Public/EngineUtils.h
@@ -568,8 +568,21 @@ private:
  */
 ENGINE_API void DrawStatsHUD( UWorld* InWorld, FViewport* Viewport, FCanvas* Canvas, UCanvas* CanvasObject, TArray<FDebugDisplayProperty>& DebugProperties, const FVector& ViewLocation, const FRotator& ViewRotation );
 
-/** This will set the StreamingLevels TMap with the current Streaming Level Status and also set which level the player is in **/
-void GetLevelStreamingStatus( UWorld* InWorld, TMap<FName,int32>& StreamingLevels, FString& LevelPlayerIsInName );
+/** SubLevel status information */
+struct FSubLevelStatus
+{
+	FName				PackageName;
+	EStreamingStatus	StreamingStatus;
+	int32				LODIndex;
+	bool				bPlayerInside;
+};
+
+/**
+ *	Gathers SubLevels status from a provided world
+ *	@param InWorld		World to gather sublevels stats from
+ *	@return				sublevels status (streaming state, LOD index, where player is)
+ */
+TArray<FSubLevelStatus> GetSubLevelsStatus( UWorld* InWorld );
 
 #if !UE_BUILD_SHIPPING
 
