fixed ActiveGenerators array growing on every navmesh rebuild

diff --git a/Engine/Source/Runtime/Engine/Private/AI/Navigation/RecastNavMeshGenerator.cpp b/Engine/Source/Runtime/Engine/Private/AI/Navigation/RecastNavMeshGenerator.cpp
--- a/Engine/Source/Runtime/Engine/Private/AI/Navigation/RecastNavMeshGenerator.cpp
+++ b/Engine/Source/Runtime/Engine/Private/AI/Navigation/RecastNavMeshGenerator.cpp
@@ -3323,6 +3323,7 @@ void FRecastNavMeshGenerator::Init()
 	/** setup maximum number of active tile generator*/
 	const int32 NumberOfWorkerThreads = FTaskGraphInterface::Get().GetNumWorkerThreads();
 	const int32 MaxGeneratorsThreads = NumberOfWorkerThreads > 1 ? NumberOfWorkerThreads - 1 : 1;
+	ActiveGenerators.Reset();
 	ActiveGenerators.AddZeroed(MaxGeneratorsThreads);
 	NumActiveTiles = 0;
 
