More Engine.h cleanup (AnimCompression, Niagara, AnimSet, AutomationTestSettings, BookMark, BrushBuilder, Channels)

diff --git a/Engine/Source/Developer/AssetTools/Private/AssetTypeActions/AssetTypeActions_NiagaraScript.cpp b/Engine/Source/Developer/AssetTools/Private/AssetTypeActions/AssetTypeActions_NiagaraScript.cpp
--- a/Engine/Source/Developer/AssetTools/Private/AssetTypeActions/AssetTypeActions_NiagaraScript.cpp
+++ b/Engine/Source/Developer/AssetTools/Private/AssetTypeActions/AssetTypeActions_NiagaraScript.cpp
@@ -1,6 +1,8 @@
 // Copyright 1998-2014 Epic Games, Inc. All Rights Reserved.
 
 #include "AssetToolsPrivatePCH.h"
+#include "Engine/NiagaraScript.h"
+
 #include "Toolkits/AssetEditorManager.h"
 
 #include "Editor/NiagaraEditor/Public/NiagaraEditorModule.h"
diff --git a/Engine/Source/Editor/EditorSettingsViewer/Private/EditorSettingsViewerModule.cpp b/Engine/Source/Editor/EditorSettingsViewer/Private/EditorSettingsViewerModule.cpp
--- a/Engine/Source/Editor/EditorSettingsViewer/Private/EditorSettingsViewerModule.cpp
+++ b/Engine/Source/Editor/EditorSettingsViewer/Private/EditorSettingsViewerModule.cpp
@@ -5,7 +5,7 @@
 =============================================================================*/
 
 #include "EditorSettingsViewerPrivatePCH.h"
-
+#include "Tests/AutomationTestSettings.h"
 
 #define LOCTEXT_NAMESPACE "FEditorSettingsViewerModule"
 
diff --git a/Engine/Source/Editor/NiagaraEditor/Private/NiagaraCompiler.cpp b/Engine/Source/Editor/NiagaraEditor/Private/NiagaraCompiler.cpp
--- a/Engine/Source/Editor/NiagaraEditor/Private/NiagaraCompiler.cpp
+++ b/Engine/Source/Editor/NiagaraEditor/Private/NiagaraCompiler.cpp
@@ -1,6 +1,8 @@
 // Copyright 1998-2014 Epic Games, Inc. All Rights Reserved.
 
 #include "NiagaraEditorPrivatePCH.h"
+#include "Engine/NiagaraScript.h"
+#include "Components/NiagaraComponent.h"
 #include "CompilerResultsLog.h"
 #include "EdGraphUtilities.h"
 #include "VectorVM.h"
diff --git a/Engine/Source/Editor/NiagaraEditor/Private/NiagaraEditor.cpp b/Engine/Source/Editor/NiagaraEditor/Private/NiagaraEditor.cpp
--- a/Engine/Source/Editor/NiagaraEditor/Private/NiagaraEditor.cpp
+++ b/Engine/Source/Editor/NiagaraEditor/Private/NiagaraEditor.cpp
@@ -1,6 +1,7 @@
 // Copyright 1998-2014 Epic Games, Inc. All Rights Reserved.
 
 #include "NiagaraEditorPrivatePCH.h"
+#include "Engine/NiagaraScript.h"
 
 #include "Toolkits/IToolkitHost.h"
 #include "Editor/WorkspaceMenuStructure/Public/WorkspaceMenuStructureModule.h"
diff --git a/Engine/Source/Editor/UnrealEd/Classes/Builders/EditorBrushBuilder.h b/Engine/Source/Editor/UnrealEd/Classes/Builders/EditorBrushBuilder.h
--- a/Engine/Source/Editor/UnrealEd/Classes/Builders/EditorBrushBuilder.h
+++ b/Engine/Source/Editor/UnrealEd/Classes/Builders/EditorBrushBuilder.h
@@ -1,6 +1,7 @@
 // Copyright 1998-2014 Epic Games, Inc. All Rights Reserved.
 
 #pragma once
+#include "Engine/BrushBuilder.h"
 #include "EditorBrushBuilder.generated.h"
 
 UCLASS(abstract, MinimalAPI)
diff --git a/Engine/Source/Editor/UnrealEd/Classes/Niagara/NiagaraScriptSource.h b/Engine/Source/Editor/UnrealEd/Classes/Niagara/NiagaraScriptSource.h
--- a/Engine/Source/Editor/UnrealEd/Classes/Niagara/NiagaraScriptSource.h
+++ b/Engine/Source/Editor/UnrealEd/Classes/Niagara/NiagaraScriptSource.h
@@ -2,6 +2,7 @@
 
 
 #pragma once
+#include "Engine/NiagaraScriptSourceBase.h"
 #include "NiagaraScriptSource.generated.h"
 
 UCLASS(MinimalAPI)
diff --git a/Engine/Source/Editor/UnrealEd/Private/AnimationCompressionPanel.cpp b/Engine/Source/Editor/UnrealEd/Private/AnimationCompressionPanel.cpp
--- a/Engine/Source/Editor/UnrealEd/Private/AnimationCompressionPanel.cpp
+++ b/Engine/Source/Editor/UnrealEd/Private/AnimationCompressionPanel.cpp
@@ -1,6 +1,7 @@
 // Copyright 1998-2014 Epic Games, Inc. All Rights Reserved.
 
 #include "UnrealEd.h"
+#include "Animation/AnimCompress.h"
 #include "AnimationCompressionPanel.h"
 
 #define LOCTEXT_NAMESPACE "AnimationCompression"
diff --git a/Engine/Source/Editor/UnrealEd/Private/Commandlets/PackageUtilities.cpp b/Engine/Source/Editor/UnrealEd/Private/Commandlets/PackageUtilities.cpp
--- a/Engine/Source/Editor/UnrealEd/Private/Commandlets/PackageUtilities.cpp
+++ b/Engine/Source/Editor/UnrealEd/Private/Commandlets/PackageUtilities.cpp
@@ -5,6 +5,8 @@
 =============================================================================*/
 #include "UnrealEd.h"
 
+#include "Animation/AnimCompress_BitwiseCompressOnly.h"
+
 #include "ISourceControlModule.h"
 
 #include "PackageHelperFunctions.h"
diff --git a/Engine/Source/Editor/UnrealEd/Private/EditorModes.cpp b/Engine/Source/Editor/UnrealEd/Private/EditorModes.cpp
--- a/Engine/Source/Editor/UnrealEd/Private/EditorModes.cpp
+++ b/Engine/Source/Editor/UnrealEd/Private/EditorModes.cpp
@@ -1,6 +1,7 @@
 // Copyright 1998-2014 Epic Games, Inc. All Rights Reserved.
 
 #include "UnrealEd.h"
+#include "Engine/BookMark.h"
 #include "EditorSupportDelegates.h"
 #include "MouseDeltaTracker.h"
 #include "ScopedTransaction.h"
diff --git a/Engine/Source/Editor/UnrealEd/Private/EditorServer.cpp b/Engine/Source/Editor/UnrealEd/Private/EditorServer.cpp
--- a/Engine/Source/Editor/UnrealEd/Private/EditorServer.cpp
+++ b/Engine/Source/Editor/UnrealEd/Private/EditorServer.cpp
@@ -22,7 +22,7 @@
 #include "Editor/GeometryMode/Public/GeometryEdMode.h"
 #include "Landscape/LandscapeProxy.h"
 #include "Lightmass/PrecomputedVisibilityOverrideVolume.h"
-
+#include "Animation/AnimSet.h"
 #include "Matinee/InterpTrackAnimControl.h"
 #include "Matinee/InterpData.h"
 
diff --git a/Engine/Source/Editor/UnrealEd/Private/Factories/NiagaraScriptFactory.cpp b/Engine/Source/Editor/UnrealEd/Private/Factories/NiagaraScriptFactory.cpp
--- a/Engine/Source/Editor/UnrealEd/Private/Factories/NiagaraScriptFactory.cpp
+++ b/Engine/Source/Editor/UnrealEd/Private/Factories/NiagaraScriptFactory.cpp
@@ -1,6 +1,7 @@
 // Copyright 1998-2014 Epic Games, Inc. All Rights Reserved.
 
 #include "UnrealEd.h"
+#include "Engine/NiagaraScript.h"
 #include "NiagaraEditorModule.h"
 #include "VectorVM.h"
 
diff --git a/Engine/Source/Editor/UnrealEd/Private/Niagara/NiagaraScriptSource.cpp b/Engine/Source/Editor/UnrealEd/Private/Niagara/NiagaraScriptSource.cpp
--- a/Engine/Source/Editor/UnrealEd/Private/Niagara/NiagaraScriptSource.cpp
+++ b/Engine/Source/Editor/UnrealEd/Private/Niagara/NiagaraScriptSource.cpp
@@ -1,6 +1,7 @@
 // Copyright 1998-2014 Epic Games, Inc. All Rights Reserved.
 
 #include "UnrealEd.h"
+#include "Engine/NiagaraScript.h"
 #include "NiagaraEditorModule.h"
 
 //////////////////////////////////////////////////////////////////////////
diff --git a/Engine/Source/Editor/UnrealEd/Private/SSkeletonWidget.cpp b/Engine/Source/Editor/UnrealEd/Private/SSkeletonWidget.cpp
--- a/Engine/Source/Editor/UnrealEd/Private/SSkeletonWidget.cpp
+++ b/Engine/Source/Editor/UnrealEd/Private/SSkeletonWidget.cpp
@@ -1,6 +1,7 @@
 // Copyright 1998-2014 Epic Games, Inc. All Rights Reserved.
 
 #include "UnrealEd.h"
+#include "Animation/AnimSet.h"
 #include "SSkeletonWidget.h"
 #include "AssetData.h"
 #include "MainFrame.h"
diff --git a/Engine/Source/Editor/UnrealEd/Private/Tests/EditorAutomationTests.cpp b/Engine/Source/Editor/UnrealEd/Private/Tests/EditorAutomationTests.cpp
--- a/Engine/Source/Editor/UnrealEd/Private/Tests/EditorAutomationTests.cpp
+++ b/Engine/Source/Editor/UnrealEd/Private/Tests/EditorAutomationTests.cpp
@@ -2,6 +2,7 @@
 
 
 #include "UnrealEd.h"
+#include "Tests/AutomationTestSettings.h"
 #include "ObjectTools.h"
 #include "Json.h"
 #include "AssetRegistryModule.h"
diff --git a/Engine/Source/Editor/UnrealEd/Private/UnrealEdEngine.cpp b/Engine/Source/Editor/UnrealEd/Private/UnrealEdEngine.cpp
--- a/Engine/Source/Editor/UnrealEd/Private/UnrealEdEngine.cpp
+++ b/Engine/Source/Editor/UnrealEd/Private/UnrealEdEngine.cpp
@@ -4,6 +4,7 @@
 
 #include "Matinee/InterpData.h"
 #include "Matinee/MatineeActor.h"
+#include "Animation/AnimCompress.h"
 
 #include "EditorSupportDelegates.h"
 #include "SoundDefinitions.h"
diff --git a/Engine/Source/Editor/UnrealEd/Public/SSkeletonWidget.h b/Engine/Source/Editor/UnrealEd/Public/SSkeletonWidget.h
--- a/Engine/Source/Editor/UnrealEd/Public/SSkeletonWidget.h
+++ b/Engine/Source/Editor/UnrealEd/Public/SSkeletonWidget.h
@@ -6,6 +6,8 @@
 
 #include "Slate.h"
 
+class UAnimSet;
+
 /**
  * This below code is to select skeleton from the list 
  */
diff --git a/Engine/Source/Runtime/Engine/Classes/Animation/AnimCompress_Automatic.h b/Engine/Source/Runtime/Engine/Classes/Animation/AnimCompress_Automatic.h
--- a/Engine/Source/Runtime/Engine/Classes/Animation/AnimCompress_Automatic.h
+++ b/Engine/Source/Runtime/Engine/Classes/Animation/AnimCompress_Automatic.h
@@ -7,6 +7,7 @@
  */
 
 #pragma once
+#include "Animation/AnimCompress.h"
 #include "AnimCompress_Automatic.generated.h"
 
 UCLASS(hidecategories=AnimationCompressionAlgorithm)
diff --git a/Engine/Source/Runtime/Engine/Classes/Animation/AnimCompress_BitwiseCompressOnly.h b/Engine/Source/Runtime/Engine/Classes/Animation/AnimCompress_BitwiseCompressOnly.h
--- a/Engine/Source/Runtime/Engine/Classes/Animation/AnimCompress_BitwiseCompressOnly.h
+++ b/Engine/Source/Runtime/Engine/Classes/Animation/AnimCompress_BitwiseCompressOnly.h
@@ -6,6 +6,7 @@
  */
 
 #pragma once
+#include "Animation/AnimCompress.h"
 #include "AnimCompress_BitwiseCompressOnly.generated.h"
 
 UCLASS(MinimalAPI)
diff --git a/Engine/Source/Runtime/Engine/Classes/Animation/AnimCompress_LeastDestructive.h b/Engine/Source/Runtime/Engine/Classes/Animation/AnimCompress_LeastDestructive.h
--- a/Engine/Source/Runtime/Engine/Classes/Animation/AnimCompress_LeastDestructive.h
+++ b/Engine/Source/Runtime/Engine/Classes/Animation/AnimCompress_LeastDestructive.h
@@ -6,7 +6,7 @@
  * Reverts any animation compression, restoring the animation to the raw data.
  *
  */
-
+#include "Animation/AnimCompress.h"
 #include "AnimCompress_LeastDestructive.generated.h"
 
 UCLASS()
diff --git a/Engine/Source/Runtime/Engine/Classes/Animation/AnimCompress_PerTrackCompression.h b/Engine/Source/Runtime/Engine/Classes/Animation/AnimCompress_PerTrackCompression.h
--- a/Engine/Source/Runtime/Engine/Classes/Animation/AnimCompress_PerTrackCompression.h
+++ b/Engine/Source/Runtime/Engine/Classes/Animation/AnimCompress_PerTrackCompression.h
@@ -7,7 +7,7 @@
  * well as choosing the best bitwise compression for each track independently.
  *
  */
-
+#include "Animation/AnimCompress_RemoveLinearKeys.h"
 #include "AnimCompress_PerTrackCompression.generated.h"
 
 UCLASS(hidecategories=AnimCompress)
diff --git a/Engine/Source/Runtime/Engine/Classes/Animation/AnimCompress_RemoveEverySecondKey.h b/Engine/Source/Runtime/Engine/Classes/Animation/AnimCompress_RemoveEverySecondKey.h
--- a/Engine/Source/Runtime/Engine/Classes/Animation/AnimCompress_RemoveEverySecondKey.h
+++ b/Engine/Source/Runtime/Engine/Classes/Animation/AnimCompress_RemoveEverySecondKey.h
@@ -6,6 +6,7 @@
  */
 
 #pragma once
+#include "Animation/AnimCompress.h"
 #include "AnimCompress_RemoveEverySecondKey.generated.h"
 
 UCLASS(MinimalAPI)
diff --git a/Engine/Source/Runtime/Engine/Classes/Animation/AnimCompress_RemoveLinearKeys.h b/Engine/Source/Runtime/Engine/Classes/Animation/AnimCompress_RemoveLinearKeys.h
--- a/Engine/Source/Runtime/Engine/Classes/Animation/AnimCompress_RemoveLinearKeys.h
+++ b/Engine/Source/Runtime/Engine/Classes/Animation/AnimCompress_RemoveLinearKeys.h
@@ -6,6 +6,7 @@
  */
 
 #pragma once
+#include "Animation/AnimCompress.h"
 #include "AnimCompress_RemoveLinearKeys.generated.h"
 
 UCLASS(MinimalAPI)
diff --git a/Engine/Source/Runtime/Engine/Classes/Animation/AnimCompress_RemoveTrivialKeys.h b/Engine/Source/Runtime/Engine/Classes/Animation/AnimCompress_RemoveTrivialKeys.h
--- a/Engine/Source/Runtime/Engine/Classes/Animation/AnimCompress_RemoveTrivialKeys.h
+++ b/Engine/Source/Runtime/Engine/Classes/Animation/AnimCompress_RemoveTrivialKeys.h
@@ -8,6 +8,7 @@
  */
 
 #pragma once
+#include "Animation/AnimCompress.h"
 #include "AnimCompress_RemoveTrivialKeys.generated.h"
 
 UCLASS(MinimalAPI)
diff --git a/Engine/Source/Runtime/Engine/Classes/Animation/AnimCompress_RevertToRaw.h b/Engine/Source/Runtime/Engine/Classes/Animation/AnimCompress_RevertToRaw.h
--- a/Engine/Source/Runtime/Engine/Classes/Animation/AnimCompress_RevertToRaw.h
+++ b/Engine/Source/Runtime/Engine/Classes/Animation/AnimCompress_RevertToRaw.h
@@ -6,6 +6,7 @@
  */
 
 #pragma once
+#include "Animation/AnimCompress.h"
 #include "AnimCompress_RevertToRaw.generated.h"
 
 UCLASS(deprecated)
diff --git a/Engine/Source/Runtime/Engine/Classes/Engine/ActorChannel.h b/Engine/Source/Runtime/Engine/Classes/Engine/ActorChannel.h
--- a/Engine/Source/Runtime/Engine/Classes/Engine/ActorChannel.h
+++ b/Engine/Source/Runtime/Engine/Classes/Engine/ActorChannel.h
@@ -35,6 +35,7 @@
 #pragma once
 #include "Net/DataBunch.h"
 #include "Net/DataReplication.h"
+#include "Engine/Channel.h"
 #include "ActorChannel.generated.h"
 
 UCLASS(transient, customConstructor)
diff --git a/Engine/Source/Runtime/Engine/Classes/Engine/NetConnection.h b/Engine/Source/Runtime/Engine/Classes/Engine/NetConnection.h
--- a/Engine/Source/Runtime/Engine/Classes/Engine/NetConnection.h
+++ b/Engine/Source/Runtime/Engine/Classes/Engine/NetConnection.h
@@ -8,6 +8,7 @@
 #include "Net/DataBunch.h"
 #include "NetConnection.generated.h"
 
+class FObjectReplicator;
 
 /*-----------------------------------------------------------------------------
 	Types.
diff --git a/Engine/Source/Runtime/Engine/Classes/Engine/NetDriver.h b/Engine/Source/Runtime/Engine/Classes/Engine/NetDriver.h
--- a/Engine/Source/Runtime/Engine/Classes/Engine/NetDriver.h
+++ b/Engine/Source/Runtime/Engine/Classes/Engine/NetDriver.h
@@ -6,6 +6,9 @@
 //
 #include "NetDriver.generated.h"
 
+class FRepChangedPropertyTracker;
+class FRepLayout;
+class FObjectReplicator;
 
 //
 // Whether to support net lag and packet loss testing.
diff --git a/Engine/Source/Runtime/Engine/Classes/Engine/PackageMapClient.h b/Engine/Source/Runtime/Engine/Classes/Engine/PackageMapClient.h
--- a/Engine/Source/Runtime/Engine/Classes/Engine/PackageMapClient.h
+++ b/Engine/Source/Runtime/Engine/Classes/Engine/PackageMapClient.h
@@ -19,6 +19,7 @@
  */
 
 #pragma once
+#include "Net/DataBunch.h"
 #include "PackageMapClient.generated.h"
 
 UCLASS(transient)
diff --git a/Engine/Source/Runtime/Engine/Classes/Engine/VoiceChannel.h b/Engine/Source/Runtime/Engine/Classes/Engine/VoiceChannel.h
--- a/Engine/Source/Runtime/Engine/Classes/Engine/VoiceChannel.h
+++ b/Engine/Source/Runtime/Engine/Classes/Engine/VoiceChannel.h
@@ -7,6 +7,7 @@
 #pragma once
 #include "Net/DataBunch.h"
 #include "Net/VoiceDataCommon.h"
+#include "Engine/Channel.h"
 #include "VoiceChannel.generated.h"
 
 UCLASS(transient, customConstructor)
diff --git a/Engine/Source/Runtime/Engine/Private/ActorReplication.cpp b/Engine/Source/Runtime/Engine/Private/ActorReplication.cpp
--- a/Engine/Source/Runtime/Engine/Private/ActorReplication.cpp
+++ b/Engine/Source/Runtime/Engine/Private/ActorReplication.cpp
@@ -2,6 +2,7 @@
 
 #include "EnginePrivate.h"
 #include "Net/UnrealNetwork.h"
+#include "Engine/ActorChannel.h"
 #include "GameFramework/GameNetworkManager.h"
 #include "NavigationPathBuilder.h"
 
diff --git a/Engine/Source/Runtime/Engine/Private/Animation/AnimCompress.cpp b/Engine/Source/Runtime/Engine/Private/Animation/AnimCompress.cpp
--- a/Engine/Source/Runtime/Engine/Private/Animation/AnimCompress.cpp
+++ b/Engine/Source/Runtime/Engine/Private/Animation/AnimCompress.cpp
@@ -5,6 +5,7 @@
 =============================================================================*/ 
 
 #include "EnginePrivate.h"
+#include "Animation/AnimCompress.h"
 #include "AnimationCompression.h"
 #include "AnimationUtils.h"
 #include "FloatPacker.h"
diff --git a/Engine/Source/Runtime/Engine/Private/Animation/AnimCompress_Automatic.cpp b/Engine/Source/Runtime/Engine/Private/Animation/AnimCompress_Automatic.cpp
--- a/Engine/Source/Runtime/Engine/Private/Animation/AnimCompress_Automatic.cpp
+++ b/Engine/Source/Runtime/Engine/Private/Animation/AnimCompress_Automatic.cpp
@@ -3,7 +3,7 @@
 
 #include "EnginePrivate.h"
 #include "AnimationUtils.h"
-
+#include "Animation/AnimCompress_Automatic.h"
 
 UAnimCompress_Automatic::UAnimCompress_Automatic(const class FPostConstructInitializeProperties& PCIP)
 	: Super(PCIP)
diff --git a/Engine/Source/Runtime/Engine/Private/Animation/AnimCompress_BitwiseCompressOnly.cpp b/Engine/Source/Runtime/Engine/Private/Animation/AnimCompress_BitwiseCompressOnly.cpp
--- a/Engine/Source/Runtime/Engine/Private/Animation/AnimCompress_BitwiseCompressOnly.cpp
+++ b/Engine/Source/Runtime/Engine/Private/Animation/AnimCompress_BitwiseCompressOnly.cpp
@@ -5,6 +5,7 @@
 =============================================================================*/ 
 
 #include "EnginePrivate.h"
+#include "Animation/AnimCompress_BitwiseCompressOnly.h"
 #include "AnimationUtils.h"
 #include "AnimEncoding.h"
 #include "AnimationCompression.h"
diff --git a/Engine/Source/Runtime/Engine/Private/Animation/AnimCompress_LeastDestructive.cpp b/Engine/Source/Runtime/Engine/Private/Animation/AnimCompress_LeastDestructive.cpp
--- a/Engine/Source/Runtime/Engine/Private/Animation/AnimCompress_LeastDestructive.cpp
+++ b/Engine/Source/Runtime/Engine/Private/Animation/AnimCompress_LeastDestructive.cpp
@@ -5,6 +5,8 @@
 =============================================================================*/ 
 
 #include "EnginePrivate.h"
+#include "Animation/AnimCompress_LeastDestructive.h"
+#include "Animation/AnimCompress_BitwiseCompressOnly.h"
 #include "AnimationUtils.h"
 #include "AnimEncoding.h"
 
diff --git a/Engine/Source/Runtime/Engine/Private/Animation/AnimCompress_PerTrackCompression.cpp b/Engine/Source/Runtime/Engine/Private/Animation/AnimCompress_PerTrackCompression.cpp
--- a/Engine/Source/Runtime/Engine/Private/Animation/AnimCompress_PerTrackCompression.cpp
+++ b/Engine/Source/Runtime/Engine/Private/Animation/AnimCompress_PerTrackCompression.cpp
@@ -5,6 +5,8 @@
 =============================================================================*/ 
 
 #include "EnginePrivate.h"
+#include "Animation/AnimCompress_RemoveLinearKeys.h"
+#include "Animation/AnimCompress_PerTrackCompression.h"
 #include "AnimationUtils.h"
 #include "AnimationCompression.h"
 #include "AnimEncoding.h"
diff --git a/Engine/Source/Runtime/Engine/Private/Animation/AnimCompress_RemoveEverySecondKey.cpp b/Engine/Source/Runtime/Engine/Private/Animation/AnimCompress_RemoveEverySecondKey.cpp
--- a/Engine/Source/Runtime/Engine/Private/Animation/AnimCompress_RemoveEverySecondKey.cpp
+++ b/Engine/Source/Runtime/Engine/Private/Animation/AnimCompress_RemoveEverySecondKey.cpp
@@ -5,6 +5,7 @@
 =============================================================================*/ 
 
 #include "EnginePrivate.h"
+#include "Animation/AnimCompress_RemoveEverySecondKey.h"
 #include "AnimationUtils.h"
 #include "AnimEncoding.h"
 #include "AnimationCompression.h"
diff --git a/Engine/Source/Runtime/Engine/Private/Animation/AnimCompress_RemoveLinearKeys.cpp b/Engine/Source/Runtime/Engine/Private/Animation/AnimCompress_RemoveLinearKeys.cpp
--- a/Engine/Source/Runtime/Engine/Private/Animation/AnimCompress_RemoveLinearKeys.cpp
+++ b/Engine/Source/Runtime/Engine/Private/Animation/AnimCompress_RemoveLinearKeys.cpp
@@ -5,6 +5,7 @@
 =============================================================================*/ 
 
 #include "EnginePrivate.h"
+#include "Animation/AnimCompress_RemoveLinearKeys.h"
 #include "AnimationUtils.h"
 #include "AnimEncoding.h"
 #include "AnimationCompression.h"
diff --git a/Engine/Source/Runtime/Engine/Private/Animation/AnimCompress_RemoveTrivialKeys.cpp b/Engine/Source/Runtime/Engine/Private/Animation/AnimCompress_RemoveTrivialKeys.cpp
--- a/Engine/Source/Runtime/Engine/Private/Animation/AnimCompress_RemoveTrivialKeys.cpp
+++ b/Engine/Source/Runtime/Engine/Private/Animation/AnimCompress_RemoveTrivialKeys.cpp
@@ -5,6 +5,7 @@
 =============================================================================*/ 
 
 #include "EnginePrivate.h"
+#include "Animation/AnimCompress_RemoveTrivialKeys.h"
 #include "AnimationUtils.h"
 #include "AnimEncoding.h"
 #include "AnimationCompression.h"
diff --git a/Engine/Source/Runtime/Engine/Private/Animation/AnimCompress_RevertToRaw.cpp b/Engine/Source/Runtime/Engine/Private/Animation/AnimCompress_RevertToRaw.cpp
--- a/Engine/Source/Runtime/Engine/Private/Animation/AnimCompress_RevertToRaw.cpp
+++ b/Engine/Source/Runtime/Engine/Private/Animation/AnimCompress_RevertToRaw.cpp
@@ -5,6 +5,8 @@
 =============================================================================*/ 
 
 #include "EnginePrivate.h"
+#include "Animation/AnimCompress_RevertToRaw.h"
+#include "Animation/AnimCompress_BitwiseCompressOnly.h"
 #include "AnimationUtils.h"
 #include "AnimEncoding.h"
 
diff --git a/Engine/Source/Runtime/Engine/Private/Animation/AnimSequence.cpp b/Engine/Source/Runtime/Engine/Private/Animation/AnimSequence.cpp
--- a/Engine/Source/Runtime/Engine/Private/Animation/AnimSequence.cpp
+++ b/Engine/Source/Runtime/Engine/Private/Animation/AnimSequence.cpp
@@ -10,6 +10,7 @@
 #include "AnimationUtils.h"
 #include "AnimationRuntime.h"
 #include "TargetPlatform.h"
+#include "Animation/AnimCompress_RevertToRaw.h"
 
 #define USE_SLERP 0
 
diff --git a/Engine/Source/Runtime/Engine/Private/Animation/AnimSet.cpp b/Engine/Source/Runtime/Engine/Private/Animation/AnimSet.cpp
--- a/Engine/Source/Runtime/Engine/Private/Animation/AnimSet.cpp
+++ b/Engine/Source/Runtime/Engine/Private/Animation/AnimSet.cpp
@@ -1,6 +1,7 @@
 // Copyright 1998-2014 Epic Games, Inc. All Rights Reserved.
 
 #include "EnginePrivate.h"
+#include "Animation/AnimSet.h"
 #include "AnimationUtils.h"
 
 /////////////////////////////////////////////////////
diff --git a/Engine/Source/Runtime/Engine/Private/Animation/AnimationUtils.cpp b/Engine/Source/Runtime/Engine/Private/Animation/AnimationUtils.cpp
--- a/Engine/Source/Runtime/Engine/Private/Animation/AnimationUtils.cpp
+++ b/Engine/Source/Runtime/Engine/Private/Animation/AnimationUtils.cpp
@@ -5,6 +5,12 @@
 =============================================================================*/ 
 
 #include "EnginePrivate.h"
+#include "Animation/AnimCompress_BitwiseCompressOnly.h"
+#include "Animation/AnimCompress_PerTrackCompression.h"
+#include "Animation/AnimCompress_RevertToRaw.h"
+#include "Animation/AnimCompress_LeastDestructive.h"
+#include "Animation/AnimCompress_RemoveEverySecondKey.h"
+#include "Animation/AnimSet.h"
 #include "AnimationUtils.h"
 #include "AnimationCompression.h"
 #include "AnimEncoding.h"
diff --git a/Engine/Source/Runtime/Engine/Private/BookMark.cpp b/Engine/Source/Runtime/Engine/Private/BookMark.cpp
--- a/Engine/Source/Runtime/Engine/Private/BookMark.cpp
+++ b/Engine/Source/Runtime/Engine/Private/BookMark.cpp
@@ -1,6 +1,8 @@
 // Copyright 1998-2014 Epic Games, Inc. All Rights Reserved.
 
 #include "EnginePrivate.h"
+#include "Engine/BookMark.h"
+#include "Engine/BookMark2D.h"
 
 UBookMark::UBookMark(const class FPostConstructInitializeProperties& PCIP)
 	: Super(PCIP)
diff --git a/Engine/Source/Runtime/Engine/Private/Brush.cpp b/Engine/Source/Runtime/Engine/Private/Brush.cpp
--- a/Engine/Source/Runtime/Engine/Private/Brush.cpp
+++ b/Engine/Source/Runtime/Engine/Private/Brush.cpp
@@ -6,6 +6,7 @@
 
 #include "EnginePrivate.h"
 #include "Net/UnrealNetwork.h"
+#include "Engine/BrushBuilder.h"
 #include "ActorEditorUtils.h"
 
 #if WITH_EDITOR
diff --git a/Engine/Source/Runtime/Engine/Private/BrushBuilder.cpp b/Engine/Source/Runtime/Engine/Private/BrushBuilder.cpp
--- a/Engine/Source/Runtime/Engine/Private/BrushBuilder.cpp
+++ b/Engine/Source/Runtime/Engine/Private/BrushBuilder.cpp
@@ -5,6 +5,7 @@
 =============================================================================*/
 
 #include "EnginePrivate.h"
+#include "Engine/BrushBuilder.h"
 
 /*-----------------------------------------------------------------------------
 	UBrushBuilder.
diff --git a/Engine/Source/Runtime/Engine/Private/DataBunch.cpp b/Engine/Source/Runtime/Engine/Private/DataBunch.cpp
--- a/Engine/Source/Runtime/Engine/Private/DataBunch.cpp
+++ b/Engine/Source/Runtime/Engine/Private/DataBunch.cpp
@@ -6,6 +6,7 @@
 
 #include "EnginePrivate.h"
 #include "Net/UnrealNetwork.h"
+#include "Engine/ControlChannel.h"
 
 static const int32 MAX_BUNCH_SIZE = 1024 * 1024; 
 
diff --git a/Engine/Source/Runtime/Engine/Private/DataChannel.cpp b/Engine/Source/Runtime/Engine/Private/DataChannel.cpp
--- a/Engine/Source/Runtime/Engine/Private/DataChannel.cpp
+++ b/Engine/Source/Runtime/Engine/Private/DataChannel.cpp
@@ -9,6 +9,8 @@
 
 #include "Net/NetworkProfiler.h"
 #include "Net/DataChannel.h"
+#include "Engine/ActorChannel.h"
+#include "Engine/ControlChannel.h"
 
 DEFINE_LOG_CATEGORY(LogNet);
 DEFINE_LOG_CATEGORY(LogNetPlayerMovement);
diff --git a/Engine/Source/Runtime/Engine/Private/DataReplication.cpp b/Engine/Source/Runtime/Engine/Private/DataReplication.cpp
--- a/Engine/Source/Runtime/Engine/Private/DataReplication.cpp
+++ b/Engine/Source/Runtime/Engine/Private/DataReplication.cpp
@@ -9,6 +9,7 @@
 #include "Net/NetworkProfiler.h"
 #include "Net/RepLayout.h"
 #include "Net/DataReplication.h"
+#include "Engine/ActorChannel.h"
 
 class FNetSerializeCB : public INetSerializeCB
 {
diff --git a/Engine/Source/Runtime/Engine/Private/NetConnection.cpp b/Engine/Source/Runtime/Engine/Private/NetConnection.cpp
--- a/Engine/Source/Runtime/Engine/Private/NetConnection.cpp
+++ b/Engine/Source/Runtime/Engine/Private/NetConnection.cpp
@@ -7,6 +7,8 @@
 #include "EnginePrivate.h"
 #include "Net/UnrealNetwork.h"
 #include "Net/NetworkProfiler.h"
+#include "Net/DataReplication.h"
+#include "Engine/ActorChannel.h"
 
 /*-----------------------------------------------------------------------------
 	UNetConnection implementation.
diff --git a/Engine/Source/Runtime/Engine/Private/NetworkDriver.cpp b/Engine/Source/Runtime/Engine/Private/NetworkDriver.cpp
--- a/Engine/Source/Runtime/Engine/Private/NetworkDriver.cpp
+++ b/Engine/Source/Runtime/Engine/Private/NetworkDriver.cpp
@@ -7,6 +7,8 @@
 #include "EnginePrivate.h"
 #include "Net/UnrealNetwork.h"
 #include "Net/NetworkProfiler.h"
+#include "Engine/ActorChannel.h"
+#include "Engine/VoiceChannel.h"
 #include "GameFramework/GameNetworkManager.h"
 #include "NavigationPathBuilder.h"
 #include "OnlineSubsystemUtils.h"
diff --git a/Engine/Source/Runtime/Engine/Private/NiagaraActor.cpp b/Engine/Source/Runtime/Engine/Private/NiagaraActor.cpp
--- a/Engine/Source/Runtime/Engine/Private/NiagaraActor.cpp
+++ b/Engine/Source/Runtime/Engine/Private/NiagaraActor.cpp
@@ -2,6 +2,8 @@
 
 
 #include "EnginePrivate.h"
+#include "Engine/NiagaraActor.h"
+#include "Components/NiagaraComponent.h"
 
 ANiagaraActor::ANiagaraActor(const class FPostConstructInitializeProperties& PCIP)
 	: Super(PCIP)
diff --git a/Engine/Source/Runtime/Engine/Private/NiagaraComponent.cpp b/Engine/Source/Runtime/Engine/Private/NiagaraComponent.cpp
--- a/Engine/Source/Runtime/Engine/Private/NiagaraComponent.cpp
+++ b/Engine/Source/Runtime/Engine/Private/NiagaraComponent.cpp
@@ -1,6 +1,8 @@
 // Copyright 1998-2014 Epic Games, Inc. All Rights Reserved.
 
 #include "EnginePrivate.h"
+#include "Components/NiagaraComponent.h"
+#include "Engine/NiagaraScript.h"
 #include "VectorVM.h"
 #include "ParticleHelper.h"
 #include "Particles/ParticleResources.h"
diff --git a/Engine/Source/Runtime/Engine/Private/NiagaraScript.cpp b/Engine/Source/Runtime/Engine/Private/NiagaraScript.cpp
--- a/Engine/Source/Runtime/Engine/Private/NiagaraScript.cpp
+++ b/Engine/Source/Runtime/Engine/Private/NiagaraScript.cpp
@@ -1,6 +1,8 @@
 // Copyright 1998-2014 Epic Games, Inc. All Rights Reserved.
 
 #include "EnginePrivate.h"
+#include "Engine/NiagaraScript.h"
+#include "Engine/NiagaraScriptSourceBase.h"
 
 UNiagaraScript::UNiagaraScript(const class FPostConstructInitializeProperties& PCIP)
 	: Super(PCIP)
diff --git a/Engine/Source/Runtime/Engine/Private/PackageMapClient.cpp b/Engine/Source/Runtime/Engine/Private/PackageMapClient.cpp
--- a/Engine/Source/Runtime/Engine/Private/PackageMapClient.cpp
+++ b/Engine/Source/Runtime/Engine/Private/PackageMapClient.cpp
@@ -4,6 +4,8 @@
 #include "Net/UnrealNetwork.h"
 #include "Net/NetworkProfiler.h"
 #include "Net/DataChannel.h"
+#include "Net/DataReplication.h"
+#include "Engine/ActorChannel.h"
 
 // ( OutPacketId == GUID_PACKET_NOT_ACKED ) == NAK'd		(this GUID is not acked, and is not pending either, so sort of waiting)
 // ( OutPacketId == GUID_PACKET_ACKED )		== FULLY ACK'd	(this GUID is fully acked, and we no longer need to send full path)
diff --git a/Engine/Source/Runtime/Engine/Private/PlayerController.cpp b/Engine/Source/Runtime/Engine/Private/PlayerController.cpp
--- a/Engine/Source/Runtime/Engine/Private/PlayerController.cpp
+++ b/Engine/Source/Runtime/Engine/Private/PlayerController.cpp
@@ -13,6 +13,7 @@
 #include "DisplayDebugHelpers.h"
 #include "Matinee/InterpTrackInstDirector.h"
 #include "Matinee/MatineeActor.h"
+#include "Engine/ActorChannel.h"
 
 DEFINE_LOG_CATEGORY_STATIC(LogPlayerController, Log, All);
 
diff --git a/Engine/Source/Runtime/Engine/Private/RepLayout.cpp b/Engine/Source/Runtime/Engine/Private/RepLayout.cpp
--- a/Engine/Source/Runtime/Engine/Private/RepLayout.cpp
+++ b/Engine/Source/Runtime/Engine/Private/RepLayout.cpp
@@ -8,6 +8,7 @@
 #include "Net/RepLayout.h"
 #include "Net/DataReplication.h"
 #include "Net/NetworkProfiler.h"
+#include "Engine/ActorChannel.h"
 
 static TAutoConsoleVariable<int32> CVarAllowPropertySkipping( TEXT( "net.AllowPropertySkipping" ), 1, TEXT( "Allow skipping of properties that haven't changed for other clients" ) );
 
diff --git a/Engine/Source/Runtime/Engine/Private/Tests/EngineAutomationTests.cpp b/Engine/Source/Runtime/Engine/Private/Tests/EngineAutomationTests.cpp
--- a/Engine/Source/Runtime/Engine/Private/Tests/EngineAutomationTests.cpp
+++ b/Engine/Source/Runtime/Engine/Private/Tests/EngineAutomationTests.cpp
@@ -2,6 +2,8 @@
 
 
 #include "EnginePrivate.h"
+#include "Tests/AutomationTestSettings.h"
+
 #if WITH_EDITOR
 #include "FileHelpers.h"
 #endif
diff --git a/Engine/Source/Runtime/Engine/Private/Tests/ExternalToolTest.cpp b/Engine/Source/Runtime/Engine/Private/Tests/ExternalToolTest.cpp
--- a/Engine/Source/Runtime/Engine/Private/Tests/ExternalToolTest.cpp
+++ b/Engine/Source/Runtime/Engine/Private/Tests/ExternalToolTest.cpp
@@ -2,6 +2,7 @@
 
 #include "EnginePrivate.h"
 #include "AutomationTestCommon.h"
+#include "Tests/AutomationTestSettings.h"
 
 DEFINE_LOG_CATEGORY_STATIC(LogBaseAutomationTests, Log, All);
 
diff --git a/Engine/Source/Runtime/Engine/Private/VoiceChannel.cpp b/Engine/Source/Runtime/Engine/Private/VoiceChannel.cpp
--- a/Engine/Source/Runtime/Engine/Private/VoiceChannel.cpp
+++ b/Engine/Source/Runtime/Engine/Private/VoiceChannel.cpp
@@ -9,6 +9,7 @@
 #include "OnlineSubsystemUtils.h"
 #include "OnlineSubsystemTypes.h"
 #include "Net/NetworkProfiler.h"
+#include "Engine/VoiceChannel.h"
 
 /** Cleans up any voice data remaining in the queue */
 void UVoiceChannel::CleanUp()
diff --git a/Engine/Source/Runtime/Engine/Public/Engine.h b/Engine/Source/Runtime/Engine/Public/Engine.h
--- a/Engine/Source/Runtime/Engine/Public/Engine.h
+++ b/Engine/Source/Runtime/Engine/Public/Engine.h
@@ -102,7 +102,6 @@
 #include "AI/Navigation/NavigationTestingActor.h"
 #include "AI/Navigation/NavLinkHostInterface.h"
 #include "AI/Navigation/NavLinkProxy.h"
-#include "Engine/NiagaraActor.h"
 #include "Engine/Note.h"
 #include "Components/PrimitiveComponent.h"
 #include "Particles/ParticleSystem.h"
@@ -206,7 +205,6 @@
 #include "AI/Navigation/NavLinkRenderingComponent.h"
 #include "AI/Navigation/NavMeshRenderingComponent.h"
 #include "AI/Navigation/NavTestRenderingComponent.h"
-#include "Components/NiagaraComponent.h"
 #include "Components/ShapeComponent.h"
 #include "Components/BoxComponent.h"
 #include "Components/CapsuleComponent.h"
@@ -238,26 +236,15 @@
 #include "Animation/BlendSpace1D.h"
 #include "Animation/AimOffsetBlendSpace1D.h"
 #include "Animation/AnimationTypes.h"
-#include "Animation/AnimCompress.h"
-#include "Animation/AnimCompress_Automatic.h"
-#include "Animation/AnimCompress_BitwiseCompressOnly.h"
-#include "Animation/AnimCompress_LeastDestructive.h"
-#include "Animation/AnimCompress_RemoveEverySecondKey.h"
-#include "Animation/AnimCompress_RemoveLinearKeys.h"
-#include "Animation/AnimCompress_PerTrackCompression.h"
-#include "Animation/AnimCompress_RemoveTrivialKeys.h"
-#include "Animation/AnimCompress_RevertToRaw.h"
 #include "Animation/AnimSingleNodeInstance.h"
 #include "Vehicles/VehicleAnimInstance.h"
 #include "Animation/AnimNotifies/AnimNotify.h"
 #include "Animation/AnimNotifies/AnimNotifyState.h"
 #include "Animation/AnimNotifies/AnimNotifyState_TimedParticleEffect.h"
 #include "Animation/AnimNotifies/AnimNotifyState_Trail.h"
-#include "Animation/AnimSet.h"
 #include "EditorFramework/AssetImportData.h"
 #include "Engine/AssetUserData.h"
 #include "Sound/AudioSettings.h"
-#include "Tests/AutomationTestSettings.h"
 #include "AI/Navigation/AvoidanceManager.h"
 #include "AI/BehaviorTree/BTNode.h"
 #include "AI/BehaviorTree/BTCompositeNode.h"
@@ -300,10 +287,7 @@
 #include "Animation/AnimBlueprintGeneratedClass.h"
 #include "PhysicsEngine/BodySetup.h"
 #include "Animation/AnimData/BoneMaskFilter.h"
-#include "Engine/BookMark.h"
-#include "Engine/BookMark2D.h"
 #include "Engine/Breakpoint.h"
-#include "Engine/BrushBuilder.h"
 #include "AI/BehaviorTree/BTAuxiliaryNode.h"
 #include "AI/BehaviorTree/BTDecorator.h"
 #include "AI/BehaviorTree/Decorators/BTDecorator_BlackboardBase.h"
@@ -341,10 +325,6 @@
 #include "Camera/CameraModifier.h"
 #include "Camera/CameraShake.h"
 #include "Camera/CameraModifier_CameraShake.h"
-#include "Engine/Channel.h"
-#include "Engine/ActorChannel.h"
-#include "Engine/ControlChannel.h"
-#include "Engine/VoiceChannel.h"
 #include "GameFramework/CheatManager.h"
 #include "Slate/CheckboxStyleAsset.h"
 #include "Engine/ClipPadEntry.h"
@@ -457,8 +437,6 @@
 #include "AI/Navigation/NavLinkTrivial.h"
 #include "AI/Navigation/NavNodeInterface.h"
 #include "Engine/NetDriver.h"
-#include "Engine/NiagaraScript.h"
-#include "Engine/NiagaraScriptSourceBase.h"
 #include "Engine/ObjectLibrary.h"
 #include "Engine/ObjectReferencer.h"
 #include "GameFramework/OnlineSession.h"
diff --git a/Engine/Source/Runtime/Engine/Public/Net/DataChannel.h b/Engine/Source/Runtime/Engine/Public/Net/DataChannel.h
--- a/Engine/Source/Runtime/Engine/Public/Net/DataChannel.h
+++ b/Engine/Source/Runtime/Engine/Public/Net/DataChannel.h
@@ -5,6 +5,7 @@
 =============================================================================*/
 
 #pragma once
+#include "Engine/Channel.h"
 
 /*-----------------------------------------------------------------------------
 	UControlChannel base class.
diff --git a/Engine/Source/Runtime/Online/OnlineSubsystemSteam/Private/VoiceEngineSteam.h b/Engine/Source/Runtime/Online/OnlineSubsystemSteam/Private/VoiceEngineSteam.h
--- a/Engine/Source/Runtime/Online/OnlineSubsystemSteam/Private/VoiceEngineSteam.h
+++ b/Engine/Source/Runtime/Online/OnlineSubsystemSteam/Private/VoiceEngineSteam.h
@@ -3,6 +3,7 @@
 #pragma once
 
 #include "VoiceInterface.h"
+#include "Net/VoiceDataCommon.h"
 #include "OnlineSubsystemTypes.h"
 #include "OnlineSubsystemSteamPackage.h"
 #include "OnlineSubsystemSteamTypes.h"
diff --git a/Engine/Source/Runtime/Online/OnlineSubsystemUtils/Private/IpNetDriver.cpp b/Engine/Source/Runtime/Online/OnlineSubsystemUtils/Private/IpNetDriver.cpp
--- a/Engine/Source/Runtime/Online/OnlineSubsystemUtils/Private/IpNetDriver.cpp
+++ b/Engine/Source/Runtime/Online/OnlineSubsystemUtils/Private/IpNetDriver.cpp
@@ -8,6 +8,7 @@ Notes:
 =============================================================================*/
 
 #include "OnlineSubsystemUtilsPrivatePCH.h"
+#include "Engine/Channel.h"
 
 #include "IPAddress.h"
 #include "Sockets.h"
diff --git a/Engine/Source/Runtime/Online/OnlineSubsystemUtils/Public/VoiceEngineImpl.h b/Engine/Source/Runtime/Online/OnlineSubsystemUtils/Public/VoiceEngineImpl.h
--- a/Engine/Source/Runtime/Online/OnlineSubsystemUtils/Public/VoiceEngineImpl.h
+++ b/Engine/Source/Runtime/Online/OnlineSubsystemUtils/Public/VoiceEngineImpl.h
@@ -6,6 +6,7 @@
 #include "Voice.h"
 #include "OnlineSubsystemTypes.h"
 #include "OnlineSubsystemUtilsPackage.h"
+#include "Net/VoiceDataCommon.h"
 
 /**
  * Container for unprocessed voice data
diff --git a/Engine/Source/Runtime/SkillSystem/Private/AttributeComponent.cpp b/Engine/Source/Runtime/SkillSystem/Private/AttributeComponent.cpp
--- a/Engine/Source/Runtime/SkillSystem/Private/AttributeComponent.cpp
+++ b/Engine/Source/Runtime/SkillSystem/Private/AttributeComponent.cpp
@@ -4,6 +4,7 @@
 #include "SkillSystemModulePrivatePCH.h"
 
 #include "Net/UnrealNetwork.h"
+#include "Engine/ActorChannel.h"
 #include "MessageLog.h"
 #include "UObjectToken.h"
 #include "MapErrors.h"
