Remove Particle headers from Engine.h and include only where needed.

diff --git a/Engine/Plugins/Editor/IntroTutorials/Source/IntroTutorials/Private/IntroTutorials.cpp b/Engine/Plugins/Editor/IntroTutorials/Source/IntroTutorials/Private/IntroTutorials.cpp
--- a/Engine/Plugins/Editor/IntroTutorials/Source/IntroTutorials/Private/IntroTutorials.cpp
+++ b/Engine/Plugins/Editor/IntroTutorials/Source/IntroTutorials/Private/IntroTutorials.cpp
@@ -14,6 +14,7 @@
 
 #include "EngineAnalytics.h"
 #include "Runtime/Analytics/Analytics/Public/Interfaces/IAnalyticsProvider.h"
+#include "Particles/ParticleSystem.h"
 
 #define LOCTEXT_NAMESPACE "IntroTutorials"
 
diff --git a/Engine/Source/Developer/AssetTools/Private/AssetTypeActions/AssetTypeActions_ParticleSystem.cpp b/Engine/Source/Developer/AssetTools/Private/AssetTypeActions/AssetTypeActions_ParticleSystem.cpp
--- a/Engine/Source/Developer/AssetTools/Private/AssetTypeActions/AssetTypeActions_ParticleSystem.cpp
+++ b/Engine/Source/Developer/AssetTools/Private/AssetTypeActions/AssetTypeActions_ParticleSystem.cpp
@@ -3,6 +3,7 @@
 #include "AssetToolsPrivatePCH.h"
 #include "Toolkits/AssetEditorManager.h"
 #include "Editor/Cascade/Public/CascadeModule.h"
+#include "Particles/ParticleEmitter.h"
 
 #define LOCTEXT_NAMESPACE "AssetTypeActions"
 
diff --git a/Engine/Source/Developer/AssetTools/Private/AssetTypeActions/AssetTypeActions_ParticleSystem.h b/Engine/Source/Developer/AssetTools/Private/AssetTypeActions/AssetTypeActions_ParticleSystem.h
--- a/Engine/Source/Developer/AssetTools/Private/AssetTypeActions/AssetTypeActions_ParticleSystem.h
+++ b/Engine/Source/Developer/AssetTools/Private/AssetTypeActions/AssetTypeActions_ParticleSystem.h
@@ -2,6 +2,8 @@
 
 #pragma once
 
+#include "Particles/ParticleSystem.h"
+
 class FAssetTypeActions_ParticleSystem : public FAssetTypeActions_Base
 {
 public:
diff --git a/Engine/Source/Editor/BlueprintGraph/Private/K2Node_AddComponent.cpp b/Engine/Source/Editor/BlueprintGraph/Private/K2Node_AddComponent.cpp
--- a/Engine/Source/Editor/BlueprintGraph/Private/K2Node_AddComponent.cpp
+++ b/Engine/Source/Editor/BlueprintGraph/Private/K2Node_AddComponent.cpp
@@ -6,6 +6,7 @@
 #include "ParticleDefinitions.h"
 #include "CompilerResultsLog.h"
 #include "CallFunctionHandler.h"
+#include "Particles/ParticleSystemComponent.h"
 
 #define LOCTEXT_NAMESPACE "K2Node_AddComponent"
 
diff --git a/Engine/Source/Editor/Cascade/Classes/CascadeParticleSystemComponent.h b/Engine/Source/Editor/Cascade/Classes/CascadeParticleSystemComponent.h
--- a/Engine/Source/Editor/Cascade/Classes/CascadeParticleSystemComponent.h
+++ b/Engine/Source/Editor/Cascade/Classes/CascadeParticleSystemComponent.h
@@ -2,7 +2,7 @@
 
 #pragma once
 
-
+#include "Particles/ParticleSystemComponent.h"
 #include "CascadeParticleSystemComponent.generated.h"
 
 /** 
diff --git a/Engine/Source/Editor/Cascade/Private/Cascade.cpp b/Engine/Source/Editor/Cascade/Private/Cascade.cpp
--- a/Engine/Source/Editor/Cascade/Private/Cascade.cpp
+++ b/Engine/Source/Editor/Cascade/Private/Cascade.cpp
@@ -19,6 +19,18 @@
 #include "Editor/PropertyEditor/Public/PropertyEditorModule.h"
 #include "Editor/PropertyEditor/Public/IDetailsView.h"
 
+#include "Particles/Event/ParticleModuleEventGenerator.h"
+#include "Particles/Parameter/ParticleModuleParameterDynamic.h"
+#include "Particles/Spawn/ParticleModuleSpawn.h"
+#include "Particles/TypeData/ParticleModuleTypeDataMesh.h"
+#include "Particles/VectorField/ParticleModuleVectorFieldLocal.h"
+#include "Particles/ParticleEmitter.h"
+#include "Particles/ParticleLODLevel.h"
+#include "Particles/ParticleModuleRequired.h"
+#include "Particles/ParticleSpriteEmitter.h"
+#include "Particles/ParticleSystem.h"
+#include "Particles/ParticleSystemComponent.h"
+
 static const FName Cascade_PreviewViewportTab("Cascade_PreviewViewport");
 static const FName Cascade_EmmitterCanvasTab("Cascade_EmitterCanvas");
 static const FName Cascade_PropertiesTab("Cascade_Properties");
diff --git a/Engine/Source/Editor/Cascade/Private/Cascade.h b/Engine/Source/Editor/Cascade/Private/Cascade.h
--- a/Engine/Source/Editor/Cascade/Private/Cascade.h
+++ b/Engine/Source/Editor/Cascade/Private/Cascade.h
@@ -6,12 +6,15 @@
 #include "IDistCurveEditor.h"
 #include "CascadeConfiguration.h"
 #include "CascadeParticleSystemComponent.h"
+#include "Particles/ParticleEmitter.h"
+#include "Particles/ParticleModule.h"
 
 class SCascadePreviewViewport;
 class SCascadeEmitterCanvas;
 class UCascadeParticleSystemComponent;
 class UCascadeConfiguration;
 
+
 DECLARE_LOG_CATEGORY_EXTERN(LogCascade, Log, All);
 
 
diff --git a/Engine/Source/Editor/Cascade/Private/CascadeEmitterCanvasClient.cpp b/Engine/Source/Editor/Cascade/Private/CascadeEmitterCanvasClient.cpp
--- a/Engine/Source/Editor/Cascade/Private/CascadeEmitterCanvasClient.cpp
+++ b/Engine/Source/Editor/Cascade/Private/CascadeEmitterCanvasClient.cpp
@@ -8,6 +8,14 @@
 #include "SCascadeEmitterCanvas.h"
 #include "CascadeEmitterCanvasClient.h"
 #include "Runtime/Engine/Public/Slate/SceneViewport.h"
+#include "Particles/Material/ParticleModuleMeshMaterial.h"
+#include "Particles/Spawn/ParticleModuleSpawn.h"
+#include "Particles/TypeData/ParticleModuleTypeDataMesh.h"
+#include "Particles/ParticleLODLevel.h"
+#include "Particles/ParticleModule.h"
+#include "Particles/ParticleModuleRequired.h"
+#include "Particles/ParticleSpriteEmitter.h"
+#include "Particles/ParticleSystem.h"
 
 
 FCascadeEmitterCanvasClient::FCascadeEmitterCanvasClient(TWeakPtr<FCascade> InCascade, TWeakPtr<SCascadeEmitterCanvas> InCascadeViewport)
diff --git a/Engine/Source/Editor/Cascade/Private/CascadeEmitterCanvasClient.h b/Engine/Source/Editor/Cascade/Private/CascadeEmitterCanvasClient.h
--- a/Engine/Source/Editor/Cascade/Private/CascadeEmitterCanvasClient.h
+++ b/Engine/Source/Editor/Cascade/Private/CascadeEmitterCanvasClient.h
@@ -2,6 +2,7 @@
 
 #pragma once
 
+#include "Particles/ParticleModule.h"
 
 class FCascade;
 class SCascadeEmitterCanvas;
diff --git a/Engine/Source/Editor/Cascade/Private/CascadeModule.cpp b/Engine/Source/Editor/Cascade/Private/CascadeModule.cpp
--- a/Engine/Source/Editor/Cascade/Private/CascadeModule.cpp
+++ b/Engine/Source/Editor/Cascade/Private/CascadeModule.cpp
@@ -7,6 +7,7 @@
 
 const FName CascadeAppIdentifier = FName(TEXT("CascadeApp"));
 
+class FCascade;
 
 /*-----------------------------------------------------------------------------
    FCascadeModule
diff --git a/Engine/Source/Editor/Cascade/Private/CascadePreviewViewportClient.cpp b/Engine/Source/Editor/Cascade/Private/CascadePreviewViewportClient.cpp
--- a/Engine/Source/Editor/Cascade/Private/CascadePreviewViewportClient.cpp
+++ b/Engine/Source/Editor/Cascade/Private/CascadePreviewViewportClient.cpp
@@ -6,6 +6,12 @@
 #include "PreviewScene.h"
 #include "ImageUtils.h"
 #include "CascadePreviewViewportClient.h"
+#include "Particles/Spawn/ParticleModuleSpawn.h"
+#include "Particles/TypeData/ParticleModuleTypeDataGpu.h"
+#include "Particles/VectorField/ParticleModuleVectorFieldLocal.h"
+#include "Particles/ParticleLODLevel.h"
+#include "Particles/ParticleModuleRequired.h"
+#include "Particles/ParticleSpriteEmitter.h"
 
 #define LOCTEXT_NAMESPACE "CascadeViewportClient"
 
diff --git a/Engine/Source/Editor/DetailCustomizations/Private/EmitterDetails.cpp b/Engine/Source/Editor/DetailCustomizations/Private/EmitterDetails.cpp
--- a/Engine/Source/Editor/DetailCustomizations/Private/EmitterDetails.cpp
+++ b/Engine/Source/Editor/DetailCustomizations/Private/EmitterDetails.cpp
@@ -2,7 +2,8 @@
 
 #include "DetailCustomizationsPrivatePCH.h"
 #include "EmitterDetails.h"
-
+#include "Particles/Emitter.h"
+#include "Particles/ParticleSystemComponent.h"
 
 TSharedRef<IDetailCustomization> FEmitterDetails::MakeInstance()
 {
diff --git a/Engine/Source/Editor/DetailCustomizations/Private/ParticleModuleDetails.cpp b/Engine/Source/Editor/DetailCustomizations/Private/ParticleModuleDetails.cpp
--- a/Engine/Source/Editor/DetailCustomizations/Private/ParticleModuleDetails.cpp
+++ b/Engine/Source/Editor/DetailCustomizations/Private/ParticleModuleDetails.cpp
@@ -4,6 +4,7 @@
 #include "ParticleModuleDetails.h"
 #include "SharedPointer.h"
 #include "PropertyRestriction.h"
+#include "Particles/ParticleModule.h"
 
 DECLARE_LOG_CATEGORY_CLASS(LogParticleModuleDetails, Log, All);
 
diff --git a/Engine/Source/Editor/MaterialEditor/Private/MaterialEditor.cpp b/Engine/Source/Editor/MaterialEditor/Private/MaterialEditor.cpp
--- a/Engine/Source/Editor/MaterialEditor/Private/MaterialEditor.cpp
+++ b/Engine/Source/Editor/MaterialEditor/Private/MaterialEditor.cpp
@@ -60,6 +60,7 @@
 #include "FindInMaterial.h"
 
 #include "Developer/MessageLog/Public/MessageLogModule.h"
+#include "Particles/ParticleSystemComponent.h"
 
 #define LOCTEXT_NAMESPACE "MaterialEditor"
 
diff --git a/Engine/Source/Editor/Matinee/Private/MatineeDraw.cpp b/Engine/Source/Editor/Matinee/Private/MatineeDraw.cpp
--- a/Engine/Source/Editor/Matinee/Private/MatineeDraw.cpp
+++ b/Engine/Source/Editor/Matinee/Private/MatineeDraw.cpp
@@ -21,6 +21,8 @@
 #include "Matinee/InterpTrackMove.h"
 #include "Materials/MaterialInstanceActor.h"
 
+#include "Particles/Emitter.h"
+
 static const int32 GroupHeadHeight = 24;
 static const int32 TrackHeight = 24;
 static const int32 SubTrackHeight = 19;
diff --git a/Engine/Source/Editor/Persona/Private/Persona.cpp b/Engine/Source/Editor/Persona/Private/Persona.cpp
--- a/Engine/Source/Editor/Persona/Private/Persona.cpp
+++ b/Engine/Source/Editor/Persona/Private/Persona.cpp
@@ -51,6 +51,9 @@
 #include "Customization/AnimGraphNodeSlotDetails.h"
 
 #include "AnimPreviewInstance.h"
+
+#include "Particles/ParticleSystemComponent.h"
+
 #define LOCTEXT_NAMESPACE "FPersona"
 
 /////////////////////////////////////////////////////
diff --git a/Engine/Source/Editor/PropertyEditor/Private/UserInterface/PropertyEditor/SPropertyEditorAsset.cpp b/Engine/Source/Editor/PropertyEditor/Private/UserInterface/PropertyEditor/SPropertyEditorAsset.cpp
--- a/Engine/Source/Editor/PropertyEditor/Private/UserInterface/PropertyEditor/SPropertyEditorAsset.cpp
+++ b/Engine/Source/Editor/PropertyEditor/Private/UserInterface/PropertyEditor/SPropertyEditorAsset.cpp
@@ -13,7 +13,7 @@
 #include "Developer/AssetTools/Public/AssetToolsModule.h"
 #include "SAssetDropTarget.h"
 #include "AssetRegistryModule.h"
-
+#include "Particles/ParticleSystem.h"
 #define LOCTEXT_NAMESPACE "PropertyEditor"
 
 DECLARE_DELEGATE( FOnCopy );
diff --git a/Engine/Source/Editor/UnrealEd/Private/AssetSelection.cpp b/Engine/Source/Editor/UnrealEd/Private/AssetSelection.cpp
--- a/Engine/Source/Editor/UnrealEd/Private/AssetSelection.cpp
+++ b/Engine/Source/Editor/UnrealEd/Private/AssetSelection.cpp
@@ -20,6 +20,8 @@
 #include "LevelEditorViewport.h"
 #include "Landscape/Landscape.h"
 
+#include "Particles/Emitter.h"
+
 extern void OnPlaceStaticMeshActor( AActor* MeshActor, bool bUseSurfaceOrientation );
 
 namespace AssetSelectionUtils
diff --git a/Engine/Source/Editor/UnrealEd/Private/Commandlets/ContentCommandlets.cpp b/Engine/Source/Editor/UnrealEd/Private/Commandlets/ContentCommandlets.cpp
--- a/Engine/Source/Editor/UnrealEd/Private/Commandlets/ContentCommandlets.cpp
+++ b/Engine/Source/Editor/UnrealEd/Private/Commandlets/ContentCommandlets.cpp
@@ -16,6 +16,14 @@ DEFINE_LOG_CATEGORY_STATIC(LogContentCommandlet, Log, All);
 #include "Kismet2/KismetEditorUtilities.h"
 #include "Kismet2/BlueprintEditorUtils.h"
 
+#include "Particles/Material/ParticleModuleMeshMaterial.h"
+#include "Particles/ParticleEmitter.h"
+#include "Particles/ParticleLODLevel.h"
+#include "Particles/ParticleModule.h"
+#include "Particles/ParticleModuleRequired.h"
+#include "Particles/ParticleSystem.h"
+#include "Particles/TypeData/ParticleModuleTypeDataMesh.h"
+
 
 /**-----------------------------------------------------------------------------
  *	UResavePackages commandlet.
diff --git a/Engine/Source/Editor/UnrealEd/Private/ComponentAssetBroker.cpp b/Engine/Source/Editor/UnrealEd/Private/ComponentAssetBroker.cpp
--- a/Engine/Source/Editor/UnrealEd/Private/ComponentAssetBroker.cpp
+++ b/Engine/Source/Editor/UnrealEd/Private/ComponentAssetBroker.cpp
@@ -5,6 +5,9 @@
 
 #include "SoundDefinitions.h"
 
+#include "Particles/ParticleSystem.h"
+#include "Particles/ParticleSystemComponent.h"
+
 //////////////////////////////////////////////////////////////////////////
 // FStaticMeshComponentBroker
 
diff --git a/Engine/Source/Editor/UnrealEd/Private/EditorActor.cpp b/Engine/Source/Editor/UnrealEd/Private/EditorActor.cpp
--- a/Engine/Source/Editor/UnrealEd/Private/EditorActor.cpp
+++ b/Engine/Source/Editor/UnrealEd/Private/EditorActor.cpp
@@ -15,6 +15,9 @@
 #include "ActorEditorUtils.h"
 #include "Foliage/InstancedFoliageActor.h"
 
+#include "Particles/Emitter.h"
+#include "Particles/ParticleSystemComponent.h"
+
 #define LOCTEXT_NAMESPACE "UnrealEd.EditorActor"
 
 DEFINE_LOG_CATEGORY_STATIC(LogEditorActor, Log, All);
diff --git a/Engine/Source/Editor/UnrealEd/Private/EditorServer.cpp b/Engine/Source/Editor/UnrealEd/Private/EditorServer.cpp
--- a/Engine/Source/Editor/UnrealEd/Private/EditorServer.cpp
+++ b/Engine/Source/Editor/UnrealEd/Private/EditorServer.cpp
@@ -44,6 +44,9 @@
 #include "UObjectToken.h"
 #include "MapErrors.h"
 
+#include "Particles/Emitter.h"
+#include "Particles/ParticleSystemComponent.h"
+
 DEFINE_LOG_CATEGORY_STATIC(LogEditorServer, Log, All);
 
 /** Used for the "tagsounds" and "checksounds" commands only			*/
diff --git a/Engine/Source/Editor/UnrealEd/Private/Factories/ActorFactory.cpp b/Engine/Source/Editor/UnrealEd/Private/Factories/ActorFactory.cpp
--- a/Engine/Source/Editor/UnrealEd/Private/Factories/ActorFactory.cpp
+++ b/Engine/Source/Editor/UnrealEd/Private/Factories/ActorFactory.cpp
@@ -19,6 +19,10 @@ ActorFactory.cpp:
 
 #include "AssetRegistryModule.h"
 
+#include "Particles/Emitter.h"
+#include "Particles/ParticleSystem.h"
+#include "Particles/ParticleSystemComponent.h"
+
 DEFINE_LOG_CATEGORY_STATIC(LogActorFactory, Log, All);
 
 #define LOCTEXT_NAMESPACE "ActorFactory"
diff --git a/Engine/Source/Editor/UnrealEd/Private/Factories/EditorFactories.cpp b/Engine/Source/Editor/UnrealEd/Private/Factories/EditorFactories.cpp
--- a/Engine/Source/Editor/UnrealEd/Private/Factories/EditorFactories.cpp
+++ b/Engine/Source/Editor/UnrealEd/Private/Factories/EditorFactories.cpp
@@ -43,6 +43,8 @@
 
 #include "Foliage/InstancedFoliageActor.h"
 
+#include "Particles/ParticleSystem.h"
+
 #if PLATFORM_WINDOWS
 // Needed for DDS support.
 #include "AllowWindowsPlatformTypes.h"
diff --git a/Engine/Source/Editor/UnrealEd/Private/Factories/VectorFieldFactory.cpp b/Engine/Source/Editor/UnrealEd/Private/Factories/VectorFieldFactory.cpp
--- a/Engine/Source/Editor/UnrealEd/Private/Factories/VectorFieldFactory.cpp
+++ b/Engine/Source/Editor/UnrealEd/Private/Factories/VectorFieldFactory.cpp
@@ -6,6 +6,8 @@
 
 #include "UnrealEd.h"
 
+#include "Particles/ParticleSystemComponent.h"
+
 DEFINE_LOG_CATEGORY_STATIC(LogVectorFieldFactory, Log, All);
 
 #define LOCTEXT_NAMESPACE "VectorFieldFactory"
diff --git a/Engine/Source/Editor/UnrealEd/Private/Fbx/FbxMainExport.cpp b/Engine/Source/Editor/UnrealEd/Private/Fbx/FbxMainExport.cpp
--- a/Engine/Source/Editor/UnrealEd/Private/Fbx/FbxMainExport.cpp
+++ b/Engine/Source/Editor/UnrealEd/Private/Fbx/FbxMainExport.cpp
@@ -52,6 +52,7 @@
 
 #include "FbxExporter.h"
 #include "RawMesh.h"
+#include "Particles/Emitter.h"
 
 namespace UnFbx
 {
diff --git a/Engine/Source/Editor/UnrealEd/Private/ThumbnailHelpers.cpp b/Engine/Source/Editor/UnrealEd/Private/ThumbnailHelpers.cpp
--- a/Engine/Source/Editor/UnrealEd/Private/ThumbnailHelpers.cpp
+++ b/Engine/Source/Editor/UnrealEd/Private/ThumbnailHelpers.cpp
@@ -4,6 +4,8 @@
 #include "ThumbnailHelpers.h"
 #include "FXSystem.h"
 
+#include "Particles/ParticleSystem.h"
+#include "Particles/ParticleSystemComponent.h"
 /*
 ***************************************************************
   FThumbnailPreviewScene
diff --git a/Engine/Source/Runtime/Engine/Classes/Particles/Acceleration/ParticleModuleAcceleration.h b/Engine/Source/Runtime/Engine/Classes/Particles/Acceleration/ParticleModuleAcceleration.h
--- a/Engine/Source/Runtime/Engine/Classes/Particles/Acceleration/ParticleModuleAcceleration.h
+++ b/Engine/Source/Runtime/Engine/Classes/Particles/Acceleration/ParticleModuleAcceleration.h
@@ -2,6 +2,7 @@
 
 
 #pragma once
+#include "Particles/Acceleration/ParticleModuleAccelerationBase.h"
 #include "ParticleModuleAcceleration.generated.h"
 
 UCLASS(editinlinenew, hidecategories=Object, meta=(DisplayName = "Acceleration"), MinimalAPI)
diff --git a/Engine/Source/Runtime/Engine/Classes/Particles/Acceleration/ParticleModuleAccelerationConstant.h b/Engine/Source/Runtime/Engine/Classes/Particles/Acceleration/ParticleModuleAccelerationConstant.h
--- a/Engine/Source/Runtime/Engine/Classes/Particles/Acceleration/ParticleModuleAccelerationConstant.h
+++ b/Engine/Source/Runtime/Engine/Classes/Particles/Acceleration/ParticleModuleAccelerationConstant.h
@@ -5,7 +5,7 @@
 ==============================================================================*/
 
 #pragma once
-
+#include "Particles/Acceleration/ParticleModuleAccelerationBase.h"
 #include "ParticleModuleAccelerationConstant.generated.h"
 
 UCLASS(editinlinenew, hidecategories=(Object, Acceleration), meta=(DisplayName = "Const Acceleration"))
diff --git a/Engine/Source/Runtime/Engine/Classes/Particles/Acceleration/ParticleModuleAccelerationDrag.h b/Engine/Source/Runtime/Engine/Classes/Particles/Acceleration/ParticleModuleAccelerationDrag.h
--- a/Engine/Source/Runtime/Engine/Classes/Particles/Acceleration/ParticleModuleAccelerationDrag.h
+++ b/Engine/Source/Runtime/Engine/Classes/Particles/Acceleration/ParticleModuleAccelerationDrag.h
@@ -5,7 +5,7 @@
 ==============================================================================*/
 
 #pragma once
-
+#include "Particles/Acceleration/ParticleModuleAccelerationBase.h"
 #include "ParticleModuleAccelerationDrag.generated.h"
 
 UCLASS(editinlinenew, hidecategories=(Object, Acceleration), MinimalAPI, meta=(DisplayName = "Drag"))
diff --git a/Engine/Source/Runtime/Engine/Classes/Particles/Acceleration/ParticleModuleAccelerationOverLifetime.h b/Engine/Source/Runtime/Engine/Classes/Particles/Acceleration/ParticleModuleAccelerationOverLifetime.h
--- a/Engine/Source/Runtime/Engine/Classes/Particles/Acceleration/ParticleModuleAccelerationOverLifetime.h
+++ b/Engine/Source/Runtime/Engine/Classes/Particles/Acceleration/ParticleModuleAccelerationOverLifetime.h
@@ -2,6 +2,7 @@
 
 
 #pragma once
+#include "Particles/Acceleration/ParticleModuleAccelerationBase.h"
 #include "ParticleModuleAccelerationOverLifetime.generated.h"
 
 UCLASS(editinlinenew, hidecategories=Object, meta=(DisplayName = "Acceleration/Life"), MinimalAPI)
diff --git a/Engine/Source/Runtime/Engine/Classes/Particles/Attractor/ParticleModuleAttractorLine.h b/Engine/Source/Runtime/Engine/Classes/Particles/Attractor/ParticleModuleAttractorLine.h
--- a/Engine/Source/Runtime/Engine/Classes/Particles/Attractor/ParticleModuleAttractorLine.h
+++ b/Engine/Source/Runtime/Engine/Classes/Particles/Attractor/ParticleModuleAttractorLine.h
@@ -2,6 +2,7 @@
 
 
 #pragma once
+#include "Particles/Attractor/ParticleModuleAttractorBase.h"
 #include "ParticleModuleAttractorLine.generated.h"
 
 UCLASS(editinlinenew, hidecategories=Object, MinimalAPI, meta=(DisplayName = "Line Attractor"))
diff --git a/Engine/Source/Runtime/Engine/Classes/Particles/Attractor/ParticleModuleAttractorParticle.h b/Engine/Source/Runtime/Engine/Classes/Particles/Attractor/ParticleModuleAttractorParticle.h
--- a/Engine/Source/Runtime/Engine/Classes/Particles/Attractor/ParticleModuleAttractorParticle.h
+++ b/Engine/Source/Runtime/Engine/Classes/Particles/Attractor/ParticleModuleAttractorParticle.h
@@ -2,6 +2,7 @@
 
 
 #pragma once
+#include "Particles/Attractor/ParticleModuleAttractorBase.h"
 #include "ParticleModuleAttractorParticle.generated.h"
 
 UENUM()
diff --git a/Engine/Source/Runtime/Engine/Classes/Particles/Attractor/ParticleModuleAttractorPointGravity.h b/Engine/Source/Runtime/Engine/Classes/Particles/Attractor/ParticleModuleAttractorPointGravity.h
--- a/Engine/Source/Runtime/Engine/Classes/Particles/Attractor/ParticleModuleAttractorPointGravity.h
+++ b/Engine/Source/Runtime/Engine/Classes/Particles/Attractor/ParticleModuleAttractorPointGravity.h
@@ -6,7 +6,7 @@
 ==============================================================================*/
 
 #pragma once
-
+#include "Particles/Attractor/ParticleModuleAttractorBase.h"
 #include "ParticleModuleAttractorPointGravity.generated.h"
 
 UCLASS(editinlinenew, hidecategories=Object, meta=(DisplayName = "Point Gravity"))
diff --git a/Engine/Source/Runtime/Engine/Classes/Particles/Beam/ParticleModuleBeamModifier.h b/Engine/Source/Runtime/Engine/Classes/Particles/Beam/ParticleModuleBeamModifier.h
--- a/Engine/Source/Runtime/Engine/Classes/Particles/Beam/ParticleModuleBeamModifier.h
+++ b/Engine/Source/Runtime/Engine/Classes/Particles/Beam/ParticleModuleBeamModifier.h
@@ -8,6 +8,7 @@
  */
 
 #pragma once
+#include "Particles/Beam/ParticleModuleBeamBase.h"
 #include "ParticleModuleBeamModifier.generated.h"
 
 /**	What to modify. */
diff --git a/Engine/Source/Runtime/Engine/Classes/Particles/Beam/ParticleModuleBeamNoise.h b/Engine/Source/Runtime/Engine/Classes/Particles/Beam/ParticleModuleBeamNoise.h
--- a/Engine/Source/Runtime/Engine/Classes/Particles/Beam/ParticleModuleBeamNoise.h
+++ b/Engine/Source/Runtime/Engine/Classes/Particles/Beam/ParticleModuleBeamNoise.h
@@ -8,6 +8,7 @@
  */
 
 #pragma once
+#include "Particles/Beam/ParticleModuleBeamBase.h"
 #include "ParticleModuleBeamNoise.generated.h"
 
 UCLASS(editinlinenew, hidecategories=Object, MinimalAPI, meta=(DisplayName = "Noise"))
diff --git a/Engine/Source/Runtime/Engine/Classes/Particles/Camera/ParticleModuleCameraOffset.h b/Engine/Source/Runtime/Engine/Classes/Particles/Camera/ParticleModuleCameraOffset.h
--- a/Engine/Source/Runtime/Engine/Classes/Particles/Camera/ParticleModuleCameraOffset.h
+++ b/Engine/Source/Runtime/Engine/Classes/Particles/Camera/ParticleModuleCameraOffset.h
@@ -2,6 +2,7 @@
 
 
 #pragma once
+#include "Particles/Camera/ParticleModuleCameraBase.h"
 #include "ParticleModuleCameraOffset.generated.h"
 
 /**
diff --git a/Engine/Source/Runtime/Engine/Classes/Particles/Collision/ParticleModuleCollision.h b/Engine/Source/Runtime/Engine/Classes/Particles/Collision/ParticleModuleCollision.h
--- a/Engine/Source/Runtime/Engine/Classes/Particles/Collision/ParticleModuleCollision.h
+++ b/Engine/Source/Runtime/Engine/Classes/Particles/Collision/ParticleModuleCollision.h
@@ -2,6 +2,7 @@
 
 
 #pragma once
+#include "Particles/Collision/ParticleModuleCollisionBase.h"
 #include "ParticleModuleCollision.generated.h"
 
 UCLASS(editinlinenew, hidecategories=Object, MinimalAPI, meta=(DisplayName = "Collision"))
diff --git a/Engine/Source/Runtime/Engine/Classes/Particles/Collision/ParticleModuleCollisionGPU.h b/Engine/Source/Runtime/Engine/Classes/Particles/Collision/ParticleModuleCollisionGPU.h
--- a/Engine/Source/Runtime/Engine/Classes/Particles/Collision/ParticleModuleCollisionGPU.h
+++ b/Engine/Source/Runtime/Engine/Classes/Particles/Collision/ParticleModuleCollisionGPU.h
@@ -1,6 +1,7 @@
 // Copyright 1998-2014 Epic Games, Inc. All Rights Reserved.
 
 #pragma once
+#include "Particles/Collision/ParticleModuleCollisionBase.h"
 #include "ParticleModuleCollisionGPU.generated.h"
 
 /**
diff --git a/Engine/Source/Runtime/Engine/Classes/Particles/Color/ParticleModuleColor.h b/Engine/Source/Runtime/Engine/Classes/Particles/Color/ParticleModuleColor.h
--- a/Engine/Source/Runtime/Engine/Classes/Particles/Color/ParticleModuleColor.h
+++ b/Engine/Source/Runtime/Engine/Classes/Particles/Color/ParticleModuleColor.h
@@ -2,6 +2,7 @@
 
 
 #pragma once
+#include "Particles/Color/ParticleModuleColorBase.h"
 #include "ParticleModuleColor.generated.h"
 
 UCLASS(editinlinenew, hidecategories=Object, MinimalAPI, meta=(DisplayName = "Initial Color"))
diff --git a/Engine/Source/Runtime/Engine/Classes/Particles/Color/ParticleModuleColorOverLife.h b/Engine/Source/Runtime/Engine/Classes/Particles/Color/ParticleModuleColorOverLife.h
--- a/Engine/Source/Runtime/Engine/Classes/Particles/Color/ParticleModuleColorOverLife.h
+++ b/Engine/Source/Runtime/Engine/Classes/Particles/Color/ParticleModuleColorOverLife.h
@@ -2,6 +2,7 @@
 
 
 #pragma once
+#include "Particles/Color/ParticleModuleColorBase.h"
 #include "ParticleModuleColorOverLife.generated.h"
 
 UCLASS(editinlinenew, hidecategories=Object, MinimalAPI, meta=(DisplayName = "Color Over Life"))
diff --git a/Engine/Source/Runtime/Engine/Classes/Particles/Color/ParticleModuleColor_Seeded.h b/Engine/Source/Runtime/Engine/Classes/Particles/Color/ParticleModuleColor_Seeded.h
--- a/Engine/Source/Runtime/Engine/Classes/Particles/Color/ParticleModuleColor_Seeded.h
+++ b/Engine/Source/Runtime/Engine/Classes/Particles/Color/ParticleModuleColor_Seeded.h
@@ -2,6 +2,7 @@
 
 
 #pragma once
+#include "Particles/Color/ParticleModuleColor.h"
 #include "ParticleModuleColor_Seeded.generated.h"
 
 UCLASS(editinlinenew, hidecategories=Object, meta=(DisplayName = "Init Color (Seed)"))
diff --git a/Engine/Source/Runtime/Engine/Classes/Particles/EmitterCameraLensEffectBase.h b/Engine/Source/Runtime/Engine/Classes/Particles/EmitterCameraLensEffectBase.h
--- a/Engine/Source/Runtime/Engine/Classes/Particles/EmitterCameraLensEffectBase.h
+++ b/Engine/Source/Runtime/Engine/Classes/Particles/EmitterCameraLensEffectBase.h
@@ -7,6 +7,7 @@
  */
 
 #pragma once
+#include "Particles/Emitter.h"
 #include "EmitterCameraLensEffectBase.generated.h"
 
 UCLASS(abstract, Blueprintable, dependson=AEmitter, MinimalAPI)
diff --git a/Engine/Source/Runtime/Engine/Classes/Particles/Event/ParticleModuleEventGenerator.h b/Engine/Source/Runtime/Engine/Classes/Particles/Event/ParticleModuleEventGenerator.h
--- a/Engine/Source/Runtime/Engine/Classes/Particles/Event/ParticleModuleEventGenerator.h
+++ b/Engine/Source/Runtime/Engine/Classes/Particles/Event/ParticleModuleEventGenerator.h
@@ -2,6 +2,7 @@
 
 
 #pragma once
+#include "Particles/Event/ParticleModuleEventBase.h"
 #include "ParticleModuleEventGenerator.generated.h"
 
 /**
diff --git a/Engine/Source/Runtime/Engine/Classes/Particles/Event/ParticleModuleEventReceiverBase.h b/Engine/Source/Runtime/Engine/Classes/Particles/Event/ParticleModuleEventReceiverBase.h
--- a/Engine/Source/Runtime/Engine/Classes/Particles/Event/ParticleModuleEventReceiverBase.h
+++ b/Engine/Source/Runtime/Engine/Classes/Particles/Event/ParticleModuleEventReceiverBase.h
@@ -2,6 +2,8 @@
 
 
 #pragma once
+#include "Particles/Event/ParticleModuleEventBase.h"
+#include "Particles/ParticleSystemComponent.h"
 #include "ParticleModuleEventReceiverBase.generated.h"
 
 UCLASS(editinlinenew, hidecategories=Object, abstract)
diff --git a/Engine/Source/Runtime/Engine/Classes/Particles/Event/ParticleModuleEventReceiverKillParticles.h b/Engine/Source/Runtime/Engine/Classes/Particles/Event/ParticleModuleEventReceiverKillParticles.h
--- a/Engine/Source/Runtime/Engine/Classes/Particles/Event/ParticleModuleEventReceiverKillParticles.h
+++ b/Engine/Source/Runtime/Engine/Classes/Particles/Event/ParticleModuleEventReceiverKillParticles.h
@@ -2,6 +2,7 @@
 
 
 #pragma once
+#include "Particles/Event/ParticleModuleEventReceiverBase.h"
 #include "ParticleModuleEventReceiverKillParticles.generated.h"
 
 UCLASS(editinlinenew, hidecategories=Object, meta=(DisplayName = "EventReceiver Kill All"))
diff --git a/Engine/Source/Runtime/Engine/Classes/Particles/Event/ParticleModuleEventReceiverSpawn.h b/Engine/Source/Runtime/Engine/Classes/Particles/Event/ParticleModuleEventReceiverSpawn.h
--- a/Engine/Source/Runtime/Engine/Classes/Particles/Event/ParticleModuleEventReceiverSpawn.h
+++ b/Engine/Source/Runtime/Engine/Classes/Particles/Event/ParticleModuleEventReceiverSpawn.h
@@ -2,6 +2,7 @@
 
 
 #pragma once
+#include "Particles/Event/ParticleModuleEventReceiverBase.h"
 #include "ParticleModuleEventReceiverSpawn.generated.h"
 
 UCLASS(editinlinenew, hidecategories=Object, meta=(DisplayName = "EventReceiver Spawn"))
diff --git a/Engine/Source/Runtime/Engine/Classes/Particles/Kill/ParticleModuleKillBox.h b/Engine/Source/Runtime/Engine/Classes/Particles/Kill/ParticleModuleKillBox.h
--- a/Engine/Source/Runtime/Engine/Classes/Particles/Kill/ParticleModuleKillBox.h
+++ b/Engine/Source/Runtime/Engine/Classes/Particles/Kill/ParticleModuleKillBox.h
@@ -2,6 +2,7 @@
 
 
 #pragma once
+#include "Particles/Kill/ParticleModuleKillBase.h"
 #include "ParticleModuleKillBox.generated.h"
 
 UCLASS(editinlinenew, hidecategories=Object, meta=(DisplayName = "Kill Box"))
diff --git a/Engine/Source/Runtime/Engine/Classes/Particles/Kill/ParticleModuleKillHeight.h b/Engine/Source/Runtime/Engine/Classes/Particles/Kill/ParticleModuleKillHeight.h
--- a/Engine/Source/Runtime/Engine/Classes/Particles/Kill/ParticleModuleKillHeight.h
+++ b/Engine/Source/Runtime/Engine/Classes/Particles/Kill/ParticleModuleKillHeight.h
@@ -2,6 +2,7 @@
 
 
 #pragma once
+#include "Particles/Kill/ParticleModuleKillBase.h"
 #include "ParticleModuleKillHeight.generated.h"
 
 UCLASS(editinlinenew, hidecategories=Object, meta=(DisplayName = "Kill Height"))
diff --git a/Engine/Source/Runtime/Engine/Classes/Particles/Lifetime/ParticleModuleLifetime.h b/Engine/Source/Runtime/Engine/Classes/Particles/Lifetime/ParticleModuleLifetime.h
--- a/Engine/Source/Runtime/Engine/Classes/Particles/Lifetime/ParticleModuleLifetime.h
+++ b/Engine/Source/Runtime/Engine/Classes/Particles/Lifetime/ParticleModuleLifetime.h
@@ -2,6 +2,7 @@
 
 
 #pragma once
+#include "Particles/Lifetime/ParticleModuleLifetimeBase.h"
 #include "ParticleModuleLifetime.generated.h"
 
 UCLASS(editinlinenew, hidecategories=Object, MinimalAPI, meta=(DisplayName = "Lifetime"))
diff --git a/Engine/Source/Runtime/Engine/Classes/Particles/Lifetime/ParticleModuleLifetime_Seeded.h b/Engine/Source/Runtime/Engine/Classes/Particles/Lifetime/ParticleModuleLifetime_Seeded.h
--- a/Engine/Source/Runtime/Engine/Classes/Particles/Lifetime/ParticleModuleLifetime_Seeded.h
+++ b/Engine/Source/Runtime/Engine/Classes/Particles/Lifetime/ParticleModuleLifetime_Seeded.h
@@ -2,6 +2,7 @@
 
 
 #pragma once
+#include "Particles/Lifetime/ParticleModuleLifetime.h"
 #include "ParticleModuleLifetime_Seeded.generated.h"
 
 UCLASS(editinlinenew, hidecategories=Object, meta=(DisplayName = "Lifetime (Seed)"))
diff --git a/Engine/Source/Runtime/Engine/Classes/Particles/Light/ParticleModuleLight.h b/Engine/Source/Runtime/Engine/Classes/Particles/Light/ParticleModuleLight.h
--- a/Engine/Source/Runtime/Engine/Classes/Particles/Light/ParticleModuleLight.h
+++ b/Engine/Source/Runtime/Engine/Classes/Particles/Light/ParticleModuleLight.h
@@ -2,6 +2,8 @@
 
 
 #pragma once
+
+#include "Particles/Light/ParticleModuleLightBase.h"
 #include "ParticleModuleLight.generated.h"
 
 UCLASS(editinlinenew, hidecategories=Object, MinimalAPI, meta=(DisplayName = "Light"))
diff --git a/Engine/Source/Runtime/Engine/Classes/Particles/Light/ParticleModuleLight_Seeded.h b/Engine/Source/Runtime/Engine/Classes/Particles/Light/ParticleModuleLight_Seeded.h
--- a/Engine/Source/Runtime/Engine/Classes/Particles/Light/ParticleModuleLight_Seeded.h
+++ b/Engine/Source/Runtime/Engine/Classes/Particles/Light/ParticleModuleLight_Seeded.h
@@ -2,6 +2,7 @@
 
 
 #pragma once
+#include "Particles/Light/ParticleModuleLight.h"
 #include "ParticleModuleLight_Seeded.generated.h"
 
 UCLASS(editinlinenew, hidecategories=Object, MinimalAPI, meta=(DisplayName = "Light (Seed)"))
diff --git a/Engine/Source/Runtime/Engine/Classes/Particles/Location/ParticleModuleLocation.h b/Engine/Source/Runtime/Engine/Classes/Particles/Location/ParticleModuleLocation.h
--- a/Engine/Source/Runtime/Engine/Classes/Particles/Location/ParticleModuleLocation.h
+++ b/Engine/Source/Runtime/Engine/Classes/Particles/Location/ParticleModuleLocation.h
@@ -2,6 +2,7 @@
 
 
 #pragma once
+#include "Particles/Location/ParticleModuleLocationBase.h"
 #include "ParticleModuleLocation.generated.h"
 
 UCLASS(editinlinenew, hidecategories=Object, MinimalAPI, meta=(DisplayName = "Initial Location"))
diff --git a/Engine/Source/Runtime/Engine/Classes/Particles/Location/ParticleModuleLocationBoneSocket.h b/Engine/Source/Runtime/Engine/Classes/Particles/Location/ParticleModuleLocationBoneSocket.h
--- a/Engine/Source/Runtime/Engine/Classes/Particles/Location/ParticleModuleLocationBoneSocket.h
+++ b/Engine/Source/Runtime/Engine/Classes/Particles/Location/ParticleModuleLocationBoneSocket.h
@@ -2,6 +2,7 @@
 
 
 #pragma once
+#include "Particles/Location/ParticleModuleLocationBase.h"
 #include "ParticleModuleLocationBoneSocket.generated.h"
 
 UENUM()
diff --git a/Engine/Source/Runtime/Engine/Classes/Particles/Location/ParticleModuleLocationPrimitiveCylinder.h b/Engine/Source/Runtime/Engine/Classes/Particles/Location/ParticleModuleLocationPrimitiveCylinder.h
--- a/Engine/Source/Runtime/Engine/Classes/Particles/Location/ParticleModuleLocationPrimitiveCylinder.h
+++ b/Engine/Source/Runtime/Engine/Classes/Particles/Location/ParticleModuleLocationPrimitiveCylinder.h
@@ -6,6 +6,7 @@
 //=============================================================================
 
 #pragma once
+#include "Particles/Location/ParticleModuleLocationPrimitiveBase.h"
 #include "ParticleModuleLocationPrimitiveCylinder.generated.h"
 
 UENUM()
diff --git a/Engine/Source/Runtime/Engine/Classes/Particles/Location/ParticleModuleLocationPrimitiveCylinder_Seeded.h b/Engine/Source/Runtime/Engine/Classes/Particles/Location/ParticleModuleLocationPrimitiveCylinder_Seeded.h
--- a/Engine/Source/Runtime/Engine/Classes/Particles/Location/ParticleModuleLocationPrimitiveCylinder_Seeded.h
+++ b/Engine/Source/Runtime/Engine/Classes/Particles/Location/ParticleModuleLocationPrimitiveCylinder_Seeded.h
@@ -2,6 +2,7 @@
 
 
 #pragma once
+#include "Particles/Location/ParticleModuleLocationPrimitiveCylinder.h"
 #include "ParticleModuleLocationPrimitiveCylinder_Seeded.generated.h"
 
 UCLASS(editinlinenew, hidecategories=Object, meta=(DisplayName = "Cylinder (Seed)"))
diff --git a/Engine/Source/Runtime/Engine/Classes/Particles/Location/ParticleModuleLocationPrimitiveSphere.h b/Engine/Source/Runtime/Engine/Classes/Particles/Location/ParticleModuleLocationPrimitiveSphere.h
--- a/Engine/Source/Runtime/Engine/Classes/Particles/Location/ParticleModuleLocationPrimitiveSphere.h
+++ b/Engine/Source/Runtime/Engine/Classes/Particles/Location/ParticleModuleLocationPrimitiveSphere.h
@@ -6,6 +6,7 @@
 //=============================================================================
 
 #pragma once
+#include "Particles/Location/ParticleModuleLocationPrimitiveBase.h"
 #include "ParticleModuleLocationPrimitiveSphere.generated.h"
 
 UCLASS(editinlinenew, hidecategories=Object, MinimalAPI, meta=(DisplayName = "Sphere"))
diff --git a/Engine/Source/Runtime/Engine/Classes/Particles/Location/ParticleModuleLocationPrimitiveSphere_Seeded.h b/Engine/Source/Runtime/Engine/Classes/Particles/Location/ParticleModuleLocationPrimitiveSphere_Seeded.h
--- a/Engine/Source/Runtime/Engine/Classes/Particles/Location/ParticleModuleLocationPrimitiveSphere_Seeded.h
+++ b/Engine/Source/Runtime/Engine/Classes/Particles/Location/ParticleModuleLocationPrimitiveSphere_Seeded.h
@@ -2,6 +2,7 @@
 
 
 #pragma once
+#include "Particles/Location/ParticleModuleLocationPrimitiveSphere.h"
 #include "ParticleModuleLocationPrimitiveSphere_Seeded.generated.h"
 
 UCLASS(editinlinenew, hidecategories=Object, meta=(DisplayName = "Sphere (Seed)"))
diff --git a/Engine/Source/Runtime/Engine/Classes/Particles/Location/ParticleModuleLocationWorldOffset.h b/Engine/Source/Runtime/Engine/Classes/Particles/Location/ParticleModuleLocationWorldOffset.h
--- a/Engine/Source/Runtime/Engine/Classes/Particles/Location/ParticleModuleLocationWorldOffset.h
+++ b/Engine/Source/Runtime/Engine/Classes/Particles/Location/ParticleModuleLocationWorldOffset.h
@@ -2,6 +2,7 @@
 
 
 #pragma once
+#include "Particles/Location/ParticleModuleLocation.h"
 #include "ParticleModuleLocationWorldOffset.generated.h"
 
 UCLASS(editinlinenew, meta=(DisplayName = "World Offset"))
diff --git a/Engine/Source/Runtime/Engine/Classes/Particles/Location/ParticleModuleLocationWorldOffset_Seeded.h b/Engine/Source/Runtime/Engine/Classes/Particles/Location/ParticleModuleLocationWorldOffset_Seeded.h
--- a/Engine/Source/Runtime/Engine/Classes/Particles/Location/ParticleModuleLocationWorldOffset_Seeded.h
+++ b/Engine/Source/Runtime/Engine/Classes/Particles/Location/ParticleModuleLocationWorldOffset_Seeded.h
@@ -2,6 +2,7 @@
 
 
 #pragma once
+#include "Particles/Location/ParticleModuleLocationWorldOffset.h"
 #include "ParticleModuleLocationWorldOffset_Seeded.generated.h"
 
 UCLASS(editinlinenew, meta=(DisplayName = "World Offset (Seed)"))
diff --git a/Engine/Source/Runtime/Engine/Classes/Particles/Location/ParticleModuleLocation_Seeded.h b/Engine/Source/Runtime/Engine/Classes/Particles/Location/ParticleModuleLocation_Seeded.h
--- a/Engine/Source/Runtime/Engine/Classes/Particles/Location/ParticleModuleLocation_Seeded.h
+++ b/Engine/Source/Runtime/Engine/Classes/Particles/Location/ParticleModuleLocation_Seeded.h
@@ -2,6 +2,7 @@
 
 
 #pragma once
+#include "Particles/Location/ParticleModuleLocation.h"
 #include "ParticleModuleLocation_Seeded.generated.h"
 
 UCLASS(editinlinenew, hidecategories=Object, meta=(DisplayName = "Initial Location (Seed)"))
diff --git a/Engine/Source/Runtime/Engine/Classes/Particles/Material/ParticleModuleMeshMaterial.h b/Engine/Source/Runtime/Engine/Classes/Particles/Material/ParticleModuleMeshMaterial.h
--- a/Engine/Source/Runtime/Engine/Classes/Particles/Material/ParticleModuleMeshMaterial.h
+++ b/Engine/Source/Runtime/Engine/Classes/Particles/Material/ParticleModuleMeshMaterial.h
@@ -2,6 +2,7 @@
 
 
 #pragma once
+#include "Particles/Material/ParticleModuleMaterialBase.h"
 #include "ParticleModuleMeshMaterial.generated.h"
 
 UCLASS(editinlinenew, hidecategories=Object, MinimalAPI, meta=(DisplayName = "Mesh Material"))
diff --git a/Engine/Source/Runtime/Engine/Classes/Particles/Orbit/ParticleModuleOrbit.h b/Engine/Source/Runtime/Engine/Classes/Particles/Orbit/ParticleModuleOrbit.h
--- a/Engine/Source/Runtime/Engine/Classes/Particles/Orbit/ParticleModuleOrbit.h
+++ b/Engine/Source/Runtime/Engine/Classes/Particles/Orbit/ParticleModuleOrbit.h
@@ -2,6 +2,7 @@
 
 
 #pragma once
+#include "Particles/Orbit/ParticleModuleOrbitBase.h"
 #include "ParticleModuleOrbit.generated.h"
 
 UENUM()
diff --git a/Engine/Source/Runtime/Engine/Classes/Particles/Orientation/ParticleModuleOrientationAxisLock.h b/Engine/Source/Runtime/Engine/Classes/Particles/Orientation/ParticleModuleOrientationAxisLock.h
--- a/Engine/Source/Runtime/Engine/Classes/Particles/Orientation/ParticleModuleOrientationAxisLock.h
+++ b/Engine/Source/Runtime/Engine/Classes/Particles/Orientation/ParticleModuleOrientationAxisLock.h
@@ -2,6 +2,7 @@
 
 
 #pragma once
+#include "Particles/Orientation/ParticleModuleOrientationBase.h"
 #include "ParticleModuleOrientationAxisLock.generated.h"
 
 // Flags indicating lock
diff --git a/Engine/Source/Runtime/Engine/Classes/Particles/Parameter/ParticleModuleParameterDynamic.h b/Engine/Source/Runtime/Engine/Classes/Particles/Parameter/ParticleModuleParameterDynamic.h
--- a/Engine/Source/Runtime/Engine/Classes/Particles/Parameter/ParticleModuleParameterDynamic.h
+++ b/Engine/Source/Runtime/Engine/Classes/Particles/Parameter/ParticleModuleParameterDynamic.h
@@ -6,6 +6,8 @@
  */
 
 #pragma once
+#include "Particles/Parameter/ParticleModuleParameterBase.h"
+#include "Particles/ParticleSystemComponent.h"
 #include "ParticleModuleParameterDynamic.generated.h"
 
 /**
diff --git a/Engine/Source/Runtime/Engine/Classes/Particles/Parameter/ParticleModuleParameterDynamic_Seeded.h b/Engine/Source/Runtime/Engine/Classes/Particles/Parameter/ParticleModuleParameterDynamic_Seeded.h
--- a/Engine/Source/Runtime/Engine/Classes/Particles/Parameter/ParticleModuleParameterDynamic_Seeded.h
+++ b/Engine/Source/Runtime/Engine/Classes/Particles/Parameter/ParticleModuleParameterDynamic_Seeded.h
@@ -2,6 +2,7 @@
 
 
 #pragma once
+#include "Particles/Parameter/ParticleModuleParameterDynamic.h"
 #include "ParticleModuleParameterDynamic_Seeded.generated.h"
 
 UCLASS(editinlinenew, hidecategories=Object, meta=(DisplayName = "Dynamic (Seed)"))
diff --git a/Engine/Source/Runtime/Engine/Classes/Particles/ParticleEventManager.h b/Engine/Source/Runtime/Engine/Classes/Particles/ParticleEventManager.h
--- a/Engine/Source/Runtime/Engine/Classes/Particles/ParticleEventManager.h
+++ b/Engine/Source/Runtime/Engine/Classes/Particles/ParticleEventManager.h
@@ -2,8 +2,14 @@
 
 
 #pragma once
+#include "GameFramework/Actor.h"
 #include "ParticleEventManager.generated.h"
 
+struct FParticleEventSpawnData;
+struct FParticleEventDeathData;
+struct FParticleEventCollideData;
+struct FParticleEventBurstData;
+
 UCLASS(dependson=UParticleSystemComponent, config=Game, NotBlueprintable, notplaceable)
 class AParticleEventManager : public AActor
 {
diff --git a/Engine/Source/Runtime/Engine/Classes/Particles/ParticleModule.h b/Engine/Source/Runtime/Engine/Classes/Particles/ParticleModule.h
--- a/Engine/Source/Runtime/Engine/Classes/Particles/ParticleModule.h
+++ b/Engine/Source/Runtime/Engine/Classes/Particles/ParticleModule.h
@@ -2,8 +2,13 @@
 
 
 #pragma once
+#include "ParticleHelper.h"
 #include "ParticleModule.generated.h"
 
+class UParticleEmitter;
+class UParticleLODLevel;
+struct FParticleEmitterInstance;
+
 /** ModuleType
  *	Indicates the kind of emitter the module can be applied to.
  *	ie, EPMT_Beam - only applies to beam emitters.
diff --git a/Engine/Source/Runtime/Engine/Classes/Particles/ParticleSpriteEmitter.h b/Engine/Source/Runtime/Engine/Classes/Particles/ParticleSpriteEmitter.h
--- a/Engine/Source/Runtime/Engine/Classes/Particles/ParticleSpriteEmitter.h
+++ b/Engine/Source/Runtime/Engine/Classes/Particles/ParticleSpriteEmitter.h
@@ -2,6 +2,7 @@
 
 
 #pragma once
+#include "Particles/ParticleEmitter.h"
 #include "ParticleSpriteEmitter.generated.h"
 
 UENUM()
diff --git a/Engine/Source/Runtime/Engine/Classes/Particles/ParticleSystemComponent.h b/Engine/Source/Runtime/Engine/Classes/Particles/ParticleSystemComponent.h
--- a/Engine/Source/Runtime/Engine/Classes/Particles/ParticleSystemComponent.h
+++ b/Engine/Source/Runtime/Engine/Classes/Particles/ParticleSystemComponent.h
@@ -2,7 +2,11 @@
 
 
 #pragma once
+#include "Components/PrimitiveComponent.h"
 #include "ParticleHelper.h"
+#include "Particles/Emitter.h"
+#include "Particles/ParticleEmitter.h"
+#include "Particles/ParticleSystem.h"
 #include "ParticleSystemComponent.generated.h"
 
 //
diff --git a/Engine/Source/Runtime/Engine/Classes/Particles/Rotation/ParticleModuleMeshRotation.h b/Engine/Source/Runtime/Engine/Classes/Particles/Rotation/ParticleModuleMeshRotation.h
--- a/Engine/Source/Runtime/Engine/Classes/Particles/Rotation/ParticleModuleMeshRotation.h
+++ b/Engine/Source/Runtime/Engine/Classes/Particles/Rotation/ParticleModuleMeshRotation.h
@@ -2,6 +2,7 @@
 
 
 #pragma once
+#include "Particles/Rotation/ParticleModuleRotationBase.h"
 #include "ParticleModuleMeshRotation.generated.h"
 
 UCLASS(editinlinenew, hidecategories=Object, MinimalAPI, meta=(DisplayName = "Init Mesh Rotation"))
diff --git a/Engine/Source/Runtime/Engine/Classes/Particles/Rotation/ParticleModuleMeshRotation_Seeded.h b/Engine/Source/Runtime/Engine/Classes/Particles/Rotation/ParticleModuleMeshRotation_Seeded.h
--- a/Engine/Source/Runtime/Engine/Classes/Particles/Rotation/ParticleModuleMeshRotation_Seeded.h
+++ b/Engine/Source/Runtime/Engine/Classes/Particles/Rotation/ParticleModuleMeshRotation_Seeded.h
@@ -2,6 +2,7 @@
 
 
 #pragma once
+#include "Particles/Rotation/ParticleModuleMeshRotation.h"
 #include "ParticleModuleMeshRotation_Seeded.generated.h"
 
 UCLASS(editinlinenew, hidecategories=Object, meta=(DisplayName = "Init Mesh Rot (Seed)"))
diff --git a/Engine/Source/Runtime/Engine/Classes/Particles/Rotation/ParticleModuleRotation.h b/Engine/Source/Runtime/Engine/Classes/Particles/Rotation/ParticleModuleRotation.h
--- a/Engine/Source/Runtime/Engine/Classes/Particles/Rotation/ParticleModuleRotation.h
+++ b/Engine/Source/Runtime/Engine/Classes/Particles/Rotation/ParticleModuleRotation.h
@@ -2,6 +2,7 @@
 
 
 #pragma once
+#include "Particles/Rotation/ParticleModuleRotationBase.h"
 #include "ParticleModuleRotation.generated.h"
 
 UCLASS(editinlinenew, hidecategories=Object, MinimalAPI, meta=(DisplayName = "Initial Rotation"))
diff --git a/Engine/Source/Runtime/Engine/Classes/Particles/Rotation/ParticleModuleRotation_Seeded.h b/Engine/Source/Runtime/Engine/Classes/Particles/Rotation/ParticleModuleRotation_Seeded.h
--- a/Engine/Source/Runtime/Engine/Classes/Particles/Rotation/ParticleModuleRotation_Seeded.h
+++ b/Engine/Source/Runtime/Engine/Classes/Particles/Rotation/ParticleModuleRotation_Seeded.h
@@ -2,6 +2,7 @@
 
 
 #pragma once
+#include "Particles/Rotation/ParticleModuleRotation.h"
 #include "ParticleModuleRotation_Seeded.generated.h"
 
 UCLASS(editinlinenew, hidecategories=Object, meta=(DisplayName = "Init Rotation (Seed)"))
diff --git a/Engine/Source/Runtime/Engine/Classes/Particles/RotationRate/ParticleModuleMeshRotationRate.h b/Engine/Source/Runtime/Engine/Classes/Particles/RotationRate/ParticleModuleMeshRotationRate.h
--- a/Engine/Source/Runtime/Engine/Classes/Particles/RotationRate/ParticleModuleMeshRotationRate.h
+++ b/Engine/Source/Runtime/Engine/Classes/Particles/RotationRate/ParticleModuleMeshRotationRate.h
@@ -2,6 +2,7 @@
 
 
 #pragma once
+#include "Particles/RotationRate/ParticleModuleRotationRateBase.h"
 #include "ParticleModuleMeshRotationRate.generated.h"
 
 UCLASS(editinlinenew, hidecategories=Object, MinimalAPI, meta=(DisplayName = "Init Mesh Rotation Rate"))
diff --git a/Engine/Source/Runtime/Engine/Classes/Particles/RotationRate/ParticleModuleMeshRotationRateMultiplyLife.h b/Engine/Source/Runtime/Engine/Classes/Particles/RotationRate/ParticleModuleMeshRotationRateMultiplyLife.h
--- a/Engine/Source/Runtime/Engine/Classes/Particles/RotationRate/ParticleModuleMeshRotationRateMultiplyLife.h
+++ b/Engine/Source/Runtime/Engine/Classes/Particles/RotationRate/ParticleModuleMeshRotationRateMultiplyLife.h
@@ -2,6 +2,7 @@
 
 
 #pragma once
+#include "Particles/RotationRate/ParticleModuleRotationRateBase.h"
 #include "ParticleModuleMeshRotationRateMultiplyLife.generated.h"
 
 UCLASS(editinlinenew, hidecategories=Object, meta=(DisplayName = "Mesh Rotation Rate * Life"))
diff --git a/Engine/Source/Runtime/Engine/Classes/Particles/RotationRate/ParticleModuleMeshRotationRate_Seeded.h b/Engine/Source/Runtime/Engine/Classes/Particles/RotationRate/ParticleModuleMeshRotationRate_Seeded.h
--- a/Engine/Source/Runtime/Engine/Classes/Particles/RotationRate/ParticleModuleMeshRotationRate_Seeded.h
+++ b/Engine/Source/Runtime/Engine/Classes/Particles/RotationRate/ParticleModuleMeshRotationRate_Seeded.h
@@ -2,6 +2,7 @@
 
 
 #pragma once
+#include "Particles/RotationRate/ParticleModuleRotationRate.h"
 #include "ParticleModuleMeshRotationRate_Seeded.generated.h"
 
 UCLASS(editinlinenew, hidecategories=Object, meta=(DisplayName = "Init Mesh Rotation Rate (Seed)"))
diff --git a/Engine/Source/Runtime/Engine/Classes/Particles/RotationRate/ParticleModuleRotationRate_Seeded.h b/Engine/Source/Runtime/Engine/Classes/Particles/RotationRate/ParticleModuleRotationRate_Seeded.h
--- a/Engine/Source/Runtime/Engine/Classes/Particles/RotationRate/ParticleModuleRotationRate_Seeded.h
+++ b/Engine/Source/Runtime/Engine/Classes/Particles/RotationRate/ParticleModuleRotationRate_Seeded.h
@@ -2,6 +2,7 @@
 
 
 #pragma once
+#include "Particles/RotationRate/ParticleModuleMeshRotationRate.h"
 #include "ParticleModuleRotationRate_Seeded.generated.h"
 
 UCLASS(editinlinenew, hidecategories=Object, meta=(DisplayName = "Init Rotation Rate (Seed)"))
diff --git a/Engine/Source/Runtime/Engine/Classes/Particles/Size/ParticleModuleSize.h b/Engine/Source/Runtime/Engine/Classes/Particles/Size/ParticleModuleSize.h
--- a/Engine/Source/Runtime/Engine/Classes/Particles/Size/ParticleModuleSize.h
+++ b/Engine/Source/Runtime/Engine/Classes/Particles/Size/ParticleModuleSize.h
@@ -2,6 +2,7 @@
 
 
 #pragma once
+#include "Particles/Size/ParticleModuleSizeBase.h"
 #include "ParticleModuleSize.generated.h"
 
 UCLASS(editinlinenew, hidecategories=Object, MinimalAPI, meta=(DisplayName = "Initial Size"))
diff --git a/Engine/Source/Runtime/Engine/Classes/Particles/Size/ParticleModuleSizeMultiplyLife.h b/Engine/Source/Runtime/Engine/Classes/Particles/Size/ParticleModuleSizeMultiplyLife.h
--- a/Engine/Source/Runtime/Engine/Classes/Particles/Size/ParticleModuleSizeMultiplyLife.h
+++ b/Engine/Source/Runtime/Engine/Classes/Particles/Size/ParticleModuleSizeMultiplyLife.h
@@ -2,6 +2,7 @@
 
 
 #pragma once
+#include "Particles/Size/ParticleModuleSizeBase.h"
 #include "ParticleModuleSizeMultiplyLife.generated.h"
 
 UCLASS(editinlinenew, hidecategories=Object, MinimalAPI, meta=(DisplayName = "Size By Life"))
diff --git a/Engine/Source/Runtime/Engine/Classes/Particles/Size/ParticleModuleSizeScaleBySpeed.h b/Engine/Source/Runtime/Engine/Classes/Particles/Size/ParticleModuleSizeScaleBySpeed.h
--- a/Engine/Source/Runtime/Engine/Classes/Particles/Size/ParticleModuleSizeScaleBySpeed.h
+++ b/Engine/Source/Runtime/Engine/Classes/Particles/Size/ParticleModuleSizeScaleBySpeed.h
@@ -5,7 +5,7 @@
 ==============================================================================*/
 
 #pragma once
-
+#include "Particles/Size/ParticleModuleSizeBase.h"
 #include "ParticleModuleSizeScaleBySpeed.generated.h"
 
 UCLASS(editinlinenew, hidecategories=Object, meta=(DisplayName = "Size By Speed"))
diff --git a/Engine/Source/Runtime/Engine/Classes/Particles/Size/ParticleModuleSize_Seeded.h b/Engine/Source/Runtime/Engine/Classes/Particles/Size/ParticleModuleSize_Seeded.h
--- a/Engine/Source/Runtime/Engine/Classes/Particles/Size/ParticleModuleSize_Seeded.h
+++ b/Engine/Source/Runtime/Engine/Classes/Particles/Size/ParticleModuleSize_Seeded.h
@@ -2,6 +2,7 @@
 
 
 #pragma once
+#include "Particles/Size/ParticleModuleSize.h"
 #include "ParticleModuleSize_Seeded.generated.h"
 
 UCLASS(editinlinenew, hidecategories=Object, meta=(DisplayName = "Initial Size (Seed)"))
diff --git a/Engine/Source/Runtime/Engine/Classes/Particles/Spawn/ParticleModuleSpawn.h b/Engine/Source/Runtime/Engine/Classes/Particles/Spawn/ParticleModuleSpawn.h
--- a/Engine/Source/Runtime/Engine/Classes/Particles/Spawn/ParticleModuleSpawn.h
+++ b/Engine/Source/Runtime/Engine/Classes/Particles/Spawn/ParticleModuleSpawn.h
@@ -2,8 +2,11 @@
 
 
 #pragma once
+#include "Particles/Spawn/ParticleModuleSpawnBase.h"
 #include "ParticleModuleSpawn.generated.h"
 
+class UParticleLODLevel;
+
 UCLASS(editinlinenew, hidecategories=Object, hidecategories=ParticleModuleSpawnBase, MinimalAPI, meta=(DisplayName = "Spawn"))
 class UParticleModuleSpawn : public UParticleModuleSpawnBase
 {
diff --git a/Engine/Source/Runtime/Engine/Classes/Particles/Spawn/ParticleModuleSpawnBase.h b/Engine/Source/Runtime/Engine/Classes/Particles/Spawn/ParticleModuleSpawnBase.h
--- a/Engine/Source/Runtime/Engine/Classes/Particles/Spawn/ParticleModuleSpawnBase.h
+++ b/Engine/Source/Runtime/Engine/Classes/Particles/Spawn/ParticleModuleSpawnBase.h
@@ -2,8 +2,11 @@
 
 
 #pragma once
+#include "Particles/ParticleModule.h"
 #include "ParticleModuleSpawnBase.generated.h"
 
+struct FParticleEmitterInstance;
+
 UCLASS(editinlinenew, hidecategories=Object, abstract, meta=(DisplayName = "Spawn"))
 class UParticleModuleSpawnBase : public UParticleModule
 {
diff --git a/Engine/Source/Runtime/Engine/Classes/Particles/Spawn/ParticleModuleSpawnPerUnit.h b/Engine/Source/Runtime/Engine/Classes/Particles/Spawn/ParticleModuleSpawnPerUnit.h
--- a/Engine/Source/Runtime/Engine/Classes/Particles/Spawn/ParticleModuleSpawnPerUnit.h
+++ b/Engine/Source/Runtime/Engine/Classes/Particles/Spawn/ParticleModuleSpawnPerUnit.h
@@ -2,6 +2,7 @@
 
 
 #pragma once
+#include "Particles/Spawn/ParticleModuleSpawnBase.h"
 #include "ParticleModuleSpawnPerUnit.generated.h"
 
 UCLASS(editinlinenew, hidecategories=Object, meta=(DisplayName = "Spawn PerUnit"))
diff --git a/Engine/Source/Runtime/Engine/Classes/Particles/SubUV/ParticleModuleSubUV.h b/Engine/Source/Runtime/Engine/Classes/Particles/SubUV/ParticleModuleSubUV.h
--- a/Engine/Source/Runtime/Engine/Classes/Particles/SubUV/ParticleModuleSubUV.h
+++ b/Engine/Source/Runtime/Engine/Classes/Particles/SubUV/ParticleModuleSubUV.h
@@ -2,6 +2,8 @@
 
 
 #pragma once
+#include "Particles/SubUV/ParticleModuleSubUVBase.h"
+#include "Particles/ParticleEmitter.h"
 #include "ParticleModuleSubUV.generated.h"
 
 UCLASS(editinlinenew, hidecategories=Object, MinimalAPI, meta=(DisplayName = "SubImage Index"))
diff --git a/Engine/Source/Runtime/Engine/Classes/Particles/SubUV/ParticleModuleSubUVMovie.h b/Engine/Source/Runtime/Engine/Classes/Particles/SubUV/ParticleModuleSubUVMovie.h
--- a/Engine/Source/Runtime/Engine/Classes/Particles/SubUV/ParticleModuleSubUVMovie.h
+++ b/Engine/Source/Runtime/Engine/Classes/Particles/SubUV/ParticleModuleSubUVMovie.h
@@ -2,6 +2,7 @@
 
 
 #pragma once
+#include "Particles/SubUV/ParticleModuleSubUV.h"
 #include "ParticleModuleSubUVMovie.generated.h"
 
 UCLASS(editinlinenew, hidecategories=Object, hidecategories=SubUV, meta=(DisplayName = "SubUV Movie"))
diff --git a/Engine/Source/Runtime/Engine/Classes/Particles/Trail/ParticleModuleTrailSource.h b/Engine/Source/Runtime/Engine/Classes/Particles/Trail/ParticleModuleTrailSource.h
--- a/Engine/Source/Runtime/Engine/Classes/Particles/Trail/ParticleModuleTrailSource.h
+++ b/Engine/Source/Runtime/Engine/Classes/Particles/Trail/ParticleModuleTrailSource.h
@@ -8,6 +8,7 @@
  */
 
 #pragma once
+#include "Particles/Trail/ParticleModuleTrailBase.h"
 #include "ParticleModuleTrailSource.generated.h"
 
 UENUM()
diff --git a/Engine/Source/Runtime/Engine/Classes/Particles/TypeData/ParticleModuleTypeDataAnimTrail.h b/Engine/Source/Runtime/Engine/Classes/Particles/TypeData/ParticleModuleTypeDataAnimTrail.h
--- a/Engine/Source/Runtime/Engine/Classes/Particles/TypeData/ParticleModuleTypeDataAnimTrail.h
+++ b/Engine/Source/Runtime/Engine/Classes/Particles/TypeData/ParticleModuleTypeDataAnimTrail.h
@@ -7,6 +7,7 @@
  */
 
 #pragma once
+#include "Particles/TypeData/ParticleModuleTypeDataBase.h"
 #include "ParticleModuleTypeDataAnimTrail.generated.h"
 
 UCLASS(editinlinenew, hidecategories=Object, meta=(DisplayName = "AnimTrail Data"))
diff --git a/Engine/Source/Runtime/Engine/Classes/Particles/TypeData/ParticleModuleTypeDataBase.h b/Engine/Source/Runtime/Engine/Classes/Particles/TypeData/ParticleModuleTypeDataBase.h
--- a/Engine/Source/Runtime/Engine/Classes/Particles/TypeData/ParticleModuleTypeDataBase.h
+++ b/Engine/Source/Runtime/Engine/Classes/Particles/TypeData/ParticleModuleTypeDataBase.h
@@ -2,8 +2,11 @@
 
 
 #pragma once
+#include "Particles/ParticleModule.h"
 #include "ParticleModuleTypeDataBase.generated.h"
 
+struct FParticleEmitterInstance;
+
 UCLASS(editinlinenew, hidecategories=Object, abstract, MinimalAPI)
 class UParticleModuleTypeDataBase : public UParticleModule
 {
diff --git a/Engine/Source/Runtime/Engine/Classes/Particles/TypeData/ParticleModuleTypeDataBeam2.h b/Engine/Source/Runtime/Engine/Classes/Particles/TypeData/ParticleModuleTypeDataBeam2.h
--- a/Engine/Source/Runtime/Engine/Classes/Particles/TypeData/ParticleModuleTypeDataBeam2.h
+++ b/Engine/Source/Runtime/Engine/Classes/Particles/TypeData/ParticleModuleTypeDataBeam2.h
@@ -2,6 +2,7 @@
 
 
 #pragma once
+#include "Particles/TypeData/ParticleModuleTypeDataBase.h"
 #include "ParticleModuleTypeDataBeam2.generated.h"
 
 UENUM()
diff --git a/Engine/Source/Runtime/Engine/Classes/Particles/TypeData/ParticleModuleTypeDataMesh.h b/Engine/Source/Runtime/Engine/Classes/Particles/TypeData/ParticleModuleTypeDataMesh.h
--- a/Engine/Source/Runtime/Engine/Classes/Particles/TypeData/ParticleModuleTypeDataMesh.h
+++ b/Engine/Source/Runtime/Engine/Classes/Particles/TypeData/ParticleModuleTypeDataMesh.h
@@ -2,8 +2,11 @@
 
 
 #pragma once
+#include "Particles/TypeData/ParticleModuleTypeDataBase.h"
 #include "ParticleModuleTypeDataMesh.generated.h"
 
+class UParticleEmitter;
+
 UENUM()
 enum EMeshScreenAlignment
 {
diff --git a/Engine/Source/Runtime/Engine/Classes/Particles/VectorField/ParticleModuleVectorFieldBase.h b/Engine/Source/Runtime/Engine/Classes/Particles/VectorField/ParticleModuleVectorFieldBase.h
--- a/Engine/Source/Runtime/Engine/Classes/Particles/VectorField/ParticleModuleVectorFieldBase.h
+++ b/Engine/Source/Runtime/Engine/Classes/Particles/VectorField/ParticleModuleVectorFieldBase.h
@@ -6,7 +6,7 @@
 ==============================================================================*/
 
 #pragma once
-
+#include "Particles/ParticleModule.h"
 #include "ParticleModuleVectorFieldBase.generated.h"
 
 UCLASS(editinlinenew, hidecategories=Object, abstract, meta=(DisplayName = "Vector Field"))
diff --git a/Engine/Source/Runtime/Engine/Classes/Particles/VectorField/ParticleModuleVectorFieldGlobal.h b/Engine/Source/Runtime/Engine/Classes/Particles/VectorField/ParticleModuleVectorFieldGlobal.h
--- a/Engine/Source/Runtime/Engine/Classes/Particles/VectorField/ParticleModuleVectorFieldGlobal.h
+++ b/Engine/Source/Runtime/Engine/Classes/Particles/VectorField/ParticleModuleVectorFieldGlobal.h
@@ -5,7 +5,7 @@
 ==============================================================================*/
 
 #pragma once
-
+#include "Particles/VectorField/ParticleModuleVectorFieldBase.h"
 #include "ParticleModuleVectorFieldGlobal.generated.h"
 
 UCLASS(editinlinenew, hidecategories=Object, meta=(DisplayName = "Global Vector Fields"))
diff --git a/Engine/Source/Runtime/Engine/Classes/Particles/VectorField/ParticleModuleVectorFieldLocal.h b/Engine/Source/Runtime/Engine/Classes/Particles/VectorField/ParticleModuleVectorFieldLocal.h
--- a/Engine/Source/Runtime/Engine/Classes/Particles/VectorField/ParticleModuleVectorFieldLocal.h
+++ b/Engine/Source/Runtime/Engine/Classes/Particles/VectorField/ParticleModuleVectorFieldLocal.h
@@ -5,7 +5,7 @@
 ==============================================================================*/
 
 #pragma once
-
+#include "Particles/VectorField/ParticleModuleVectorFieldBase.h"
 #include "ParticleModuleVectorFieldLocal.generated.h"
 
 UCLASS(editinlinenew, hidecategories=Object, MinimalAPI, meta=(DisplayName = "Local Vector Field"))
diff --git a/Engine/Source/Runtime/Engine/Classes/Particles/Velocity/ParticleModuleVelocity.h b/Engine/Source/Runtime/Engine/Classes/Particles/Velocity/ParticleModuleVelocity.h
--- a/Engine/Source/Runtime/Engine/Classes/Particles/Velocity/ParticleModuleVelocity.h
+++ b/Engine/Source/Runtime/Engine/Classes/Particles/Velocity/ParticleModuleVelocity.h
@@ -2,6 +2,7 @@
 
 
 #pragma once
+#include "Particles/Velocity/ParticleModuleVelocityBase.h"
 #include "ParticleModuleVelocity.generated.h"
 
 UCLASS(editinlinenew, hidecategories=Object, MinimalAPI, meta=(DisplayName = "Initial Velocity"))
diff --git a/Engine/Source/Runtime/Engine/Classes/Particles/Velocity/ParticleModuleVelocityCone.h b/Engine/Source/Runtime/Engine/Classes/Particles/Velocity/ParticleModuleVelocityCone.h
--- a/Engine/Source/Runtime/Engine/Classes/Particles/Velocity/ParticleModuleVelocityCone.h
+++ b/Engine/Source/Runtime/Engine/Classes/Particles/Velocity/ParticleModuleVelocityCone.h
@@ -2,6 +2,7 @@
 
 
 #pragma once
+#include "Particles/Velocity/ParticleModuleVelocityBase.h"
 #include "ParticleModuleVelocityCone.generated.h"
 
 UCLASS(editinlinenew, hidecategories=Object, meta=(DisplayName = "Velocity Cone"))
diff --git a/Engine/Source/Runtime/Engine/Classes/Particles/Velocity/ParticleModuleVelocity_Seeded.h b/Engine/Source/Runtime/Engine/Classes/Particles/Velocity/ParticleModuleVelocity_Seeded.h
--- a/Engine/Source/Runtime/Engine/Classes/Particles/Velocity/ParticleModuleVelocity_Seeded.h
+++ b/Engine/Source/Runtime/Engine/Classes/Particles/Velocity/ParticleModuleVelocity_Seeded.h
@@ -2,6 +2,7 @@
 
 
 #pragma once
+#include "Particles/Velocity/ParticleModuleVelocity.h"
 #include "ParticleModuleVelocity_Seeded.generated.h"
 
 UCLASS(editinlinenew, hidecategories=Object, meta=(DisplayName = "Initial Velocity (Seed)"))
diff --git a/Engine/Source/Runtime/Engine/Private/Animation/AnimNotifyState_Trail.cpp b/Engine/Source/Runtime/Engine/Private/Animation/AnimNotifyState_Trail.cpp
--- a/Engine/Source/Runtime/Engine/Private/Animation/AnimNotifyState_Trail.cpp
+++ b/Engine/Source/Runtime/Engine/Private/Animation/AnimNotifyState_Trail.cpp
@@ -3,7 +3,9 @@
 #include "EnginePrivate.h"
 #include "MessageLog.h"
 #include "ParticleDefinitions.h"
-
+#include "Particles/TypeData/ParticleModuleTypeDataAnimTrail.h"
+#include "Particles/ParticleSystem.h"
+#include "Particles/ParticleSystemComponent.h"
 #define LOCTEXT_NAMESPACE "AnimNotifyState_Trail"
 
 DEFINE_LOG_CATEGORY(LogAnimTrails);
diff --git a/Engine/Source/Runtime/Engine/Private/Distributions.cpp b/Engine/Source/Runtime/Engine/Private/Distributions.cpp
--- a/Engine/Source/Runtime/Engine/Private/Distributions.cpp
+++ b/Engine/Source/Runtime/Engine/Private/Distributions.cpp
@@ -10,6 +10,8 @@
 #include "SoundDefinitions.h"
 #include "Sound/SoundNode.h"
 
+#include "Particles/ParticleModule.h"
+
 // Moving UDistributions to PostInitProps to not be default sub-objects:
 // Small enough value to be rounded to 0.0 in the editor 
 // but significant enough to properly detect uninitialized defaults.
diff --git a/Engine/Source/Runtime/Engine/Private/GameViewportClient.cpp b/Engine/Source/Runtime/Engine/Private/GameViewportClient.cpp
--- a/Engine/Source/Runtime/Engine/Private/GameViewportClient.cpp
+++ b/Engine/Source/Runtime/Engine/Private/GameViewportClient.cpp
@@ -17,6 +17,8 @@
 
 #include "HighResScreenshot.h"
 
+#include "Particles/ParticleSystemComponent.h"
+
 /** This variable allows forcing full screen of the first player controller viewport, even if there are multiple controllers plugged in and no cinematic playing. */
 bool GForceFullscreen = false;
 
diff --git a/Engine/Source/Runtime/Engine/Private/Interpolation.cpp b/Engine/Source/Runtime/Engine/Private/Interpolation.cpp
--- a/Engine/Source/Runtime/Engine/Private/Interpolation.cpp
+++ b/Engine/Source/Runtime/Engine/Private/Interpolation.cpp
@@ -65,6 +65,9 @@
 #include "AVIWriter.h"
 #include "AnimationUtils.h"
 #include "MatineeUtils.h"
+#include "Particles/Emitter.h"
+#include "Particles/ParticleSystemComponent.h"
+#include "Particles/ParticleSystemReplay.h"
 
 #if WITH_EDITOR
 #include "UnrealEd.h"
diff --git a/Engine/Source/Runtime/Engine/Private/Particles/ParticleBeam2EmitterInstance.cpp b/Engine/Source/Runtime/Engine/Private/Particles/ParticleBeam2EmitterInstance.cpp
--- a/Engine/Source/Runtime/Engine/Private/Particles/ParticleBeam2EmitterInstance.cpp
+++ b/Engine/Source/Runtime/Engine/Private/Particles/ParticleBeam2EmitterInstance.cpp
@@ -6,6 +6,20 @@
 =============================================================================*/
 #include "EnginePrivate.h"
 #include "ParticleDefinitions.h"
+#include "Particles/Beam/ParticleModuleBeamBase.h"
+#include "Particles/Beam/ParticleModuleBeamModifier.h"
+#include "Particles/Beam/ParticleModuleBeamNoise.h"
+#include "Particles/Beam/ParticleModuleBeamSource.h"
+#include "Particles/Beam/ParticleModuleBeamTarget.h"
+#include "Particles/Event/ParticleModuleEventGenerator.h"
+#include "Particles/Spawn/ParticleModuleSpawn.h"
+#include "Particles/TypeData/ParticleModuleTypeDataBeam2.h"
+#include "Particles/ParticleEmitter.h"
+#include "Particles/ParticleLODLevel.h"
+#include "Particles/ParticleModuleRequired.h"
+#include "Particles/ParticleSystemComponent.h"
+#include "Particles/ParticleSpriteEmitter.h"
+#include "ParticleEmitterInstances.h"
 
 /** Beam particle stat objects */
 
diff --git a/Engine/Source/Runtime/Engine/Private/Particles/ParticleBeamModules.cpp b/Engine/Source/Runtime/Engine/Private/Particles/ParticleBeamModules.cpp
--- a/Engine/Source/Runtime/Engine/Private/Particles/ParticleBeamModules.cpp
+++ b/Engine/Source/Runtime/Engine/Private/Particles/ParticleBeamModules.cpp
@@ -7,6 +7,14 @@
 #include "EnginePrivate.h"
 #include "ParticleDefinitions.h"
 #include "../DistributionHelpers.h"
+#include "Particles/Beam/ParticleModuleBeamBase.h"
+#include "Particles/Beam/ParticleModuleBeamModifier.h"
+#include "Particles/Beam/ParticleModuleBeamNoise.h"
+#include "Particles/Beam/ParticleModuleBeamSource.h"
+#include "Particles/Beam/ParticleModuleBeamTarget.h"
+#include "Particles/TypeData/ParticleModuleTypeDataBeam2.h"
+#include "Particles/ParticleSystem.h"
+#include "Particles/ParticleSystemComponent.h"
 
 /*-----------------------------------------------------------------------------
 	Abstract base modules used for categorization.
diff --git a/Engine/Source/Runtime/Engine/Private/Particles/ParticleComponents.cpp b/Engine/Source/Runtime/Engine/Private/Particles/ParticleComponents.cpp
--- a/Engine/Source/Runtime/Engine/Private/Particles/ParticleComponents.cpp
+++ b/Engine/Source/Runtime/Engine/Private/Particles/ParticleComponents.cpp
@@ -18,6 +18,34 @@
 #include "ObjectEditorUtils.h"
 #endif
 
+#include "Particles/Collision/ParticleModuleCollision.h"
+#include "Particles/Color/ParticleModuleColorOverLife.h"
+#include "Particles/Event/ParticleModuleEventGenerator.h"
+#include "Particles/Event/ParticleModuleEventReceiverBase.h"
+#include "Particles/Lifetime/ParticleModuleLifetimeBase.h"
+#include "Particles/Lifetime/ParticleModuleLifetime.h"
+#include "Particles/Material/ParticleModuleMeshMaterial.h"
+#include "Particles/Orbit/ParticleModuleOrbit.h"
+#include "Particles/Size/ParticleModuleSize.h"
+#include "Particles/Spawn/ParticleModuleSpawn.h"
+#include "Particles/Spawn/ParticleModuleSpawnBase.h"
+#include "Particles/TypeData/ParticleModuleTypeDataBase.h"
+#include "Particles/TypeData/ParticleModuleTypeDataBeam2.h"
+#include "Particles/TypeData/ParticleModuleTypeDataGpu.h"
+#include "Particles/TypeData/ParticleModuleTypeDataMesh.h"
+#include "Particles/Velocity/ParticleModuleVelocity.h"
+#include "Particles/Emitter.h"
+#include "Particles/EmitterCameraLensEffectBase.h"
+#include "Particles/ParticleEventManager.h"
+#include "Particles/ParticleLODLevel.h"
+#include "Particles/ParticleModule.h"
+#include "Particles/ParticleModuleRequired.h"
+#include "Particles/ParticleSpriteEmitter.h"
+#include "Particles/ParticleSystem.h"
+#include "Particles/ParticleSystemComponent.h"
+#include "Particles/ParticleSystemReplay.h"
+
+
 #define LOCTEXT_NAMESPACE "ParticleComponents"
 
 DEFINE_LOG_CATEGORY(LogParticles);
diff --git a/Engine/Source/Runtime/Engine/Private/Particles/ParticleEmitterInstances.cpp b/Engine/Source/Runtime/Engine/Private/Particles/ParticleEmitterInstances.cpp
--- a/Engine/Source/Runtime/Engine/Private/Particles/ParticleEmitterInstances.cpp
+++ b/Engine/Source/Runtime/Engine/Private/Particles/ParticleEmitterInstances.cpp
@@ -9,6 +9,24 @@
 #include "LevelUtils.h"
 #include "FXSystem.h"
 
+#include "Particles/Camera/ParticleModuleCameraOffset.h"
+#include "Particles/Collision/ParticleModuleCollisionGPU.h"
+#include "Particles/Event/ParticleModuleEventGenerator.h"
+#include "Particles/Event/ParticleModuleEventReceiverBase.h"
+#include "Particles/Light/ParticleModuleLightBase.h"
+#include "Particles/Material/ParticleModuleMeshMaterial.h"
+#include "Particles/Modules/Location/ParticleModulePivotOffset.h"
+#include "Particles/Orbit/ParticleModuleOrbit.h"
+#include "Particles/Orientation/ParticleModuleOrientationAxisLock.h"
+#include "Particles/Parameter/ParticleModuleParameterDynamic.h"
+#include "Particles/Spawn/ParticleModuleSpawn.h"
+#include "Particles/TypeData/ParticleModuleTypeDataMesh.h"
+#include "Particles/ParticleLODLevel.h"
+#include "Particles/ParticleModule.h"
+#include "Particles/ParticleModuleRequired.h"
+#include "Particles/ParticleSpriteEmitter.h"
+#include "Particles/ParticleSystemComponent.h"
+
 /*-----------------------------------------------------------------------------
 FParticlesStatGroup
 -----------------------------------------------------------------------------*/
diff --git a/Engine/Source/Runtime/Engine/Private/Particles/ParticleEventManager.cpp b/Engine/Source/Runtime/Engine/Private/Particles/ParticleEventManager.cpp
--- a/Engine/Source/Runtime/Engine/Private/Particles/ParticleEventManager.cpp
+++ b/Engine/Source/Runtime/Engine/Private/Particles/ParticleEventManager.cpp
@@ -2,6 +2,9 @@
 
 #include "EnginePrivate.h"
 #include "ParticleDefinitions.h"
+#include "Particles/ParticleEventManager.h"
+#include "Particles/ParticleSystemComponent.h"
+#include "Particles/Emitter.h"
 
 /*-----------------------------------------------------------------------------
 	AParticleEventManager implementation.
@@ -67,8 +70,3 @@ void AParticleEventManager::HandleParticleBurstEvents( UParticleSystemComponent*
 		}
 	}
 }
-
-
-
-
-
diff --git a/Engine/Source/Runtime/Engine/Private/Particles/ParticleGpuSimulation.cpp b/Engine/Source/Runtime/Engine/Private/Particles/ParticleGpuSimulation.cpp
--- a/Engine/Source/Runtime/Engine/Private/Particles/ParticleGpuSimulation.cpp
+++ b/Engine/Source/Runtime/Engine/Private/Particles/ParticleGpuSimulation.cpp
@@ -17,6 +17,14 @@
 #include "GlobalShader.h"
 #include "../VectorField.h"
 #include "../VectorFieldVisualization.h"
+#include "Particles/Orientation/ParticleModuleOrientationAxisLock.h"
+#include "Particles/Spawn/ParticleModuleSpawn.h"
+#include "Particles/Spawn/ParticleModuleSpawnPerUnit.h"
+#include "Particles/TypeData/ParticleModuleTypeDataGpu.h"
+#include "Particles/ParticleLODLevel.h"
+#include "Particles/ParticleModuleRequired.h"
+#include "Particles/ParticleSpriteEmitter.h"
+#include "Particles/ParticleSystemComponent.h"
 
 /*------------------------------------------------------------------------------
 	Constants to tune memory and performance for GPU particle simulation.
diff --git a/Engine/Source/Runtime/Engine/Private/Particles/ParticleModules.cpp b/Engine/Source/Runtime/Engine/Private/Particles/ParticleModules.cpp
--- a/Engine/Source/Runtime/Engine/Private/Particles/ParticleModules.cpp
+++ b/Engine/Source/Runtime/Engine/Private/Particles/ParticleModules.cpp
@@ -8,6 +8,68 @@
 #include "FXSystem.h"
 #include "ParticleDefinitions.h"
 #include "../DistributionHelpers.h"
+#include "Particles/Acceleration/ParticleModuleAcceleration.h"
+#include "Particles/Acceleration/ParticleModuleAccelerationBase.h"
+#include "Particles/Acceleration/ParticleModuleAccelerationConstant.h"
+#include "Particles/Acceleration/ParticleModuleAccelerationDrag.h"
+#include "Particles/Acceleration/ParticleModuleAccelerationDragScaleOverLife.h"
+#include "Particles/Acceleration/ParticleModuleAccelerationOverLifetime.h"
+#include "Particles/Attractor/ParticleModuleAttractorPointGravity.h"
+#include "Particles/Attractor/ParticleModuleAttractorLine.h"
+#include "Particles/Attractor/ParticleModuleAttractorParticle.h"
+#include "Particles/Attractor/ParticleModuleAttractorPoint.h"
+#include "Particles/Attractor/ParticleModuleAttractorPointGravity.h"
+#include "Particles/Collision/ParticleModuleCollisionGPU.h"
+#include "Particles/Kill/ParticleModuleKillBase.h"
+#include "Particles/Kill/ParticleModuleKillBox.h"
+#include "Particles/Kill/ParticleModuleKillHeight.h"
+#include "Particles/Light/ParticleModuleLight_Seeded.h"
+#include "Particles/Lifetime/ParticleModuleLifetime_Seeded.h"
+#include "Particles/Location/ParticleModuleLocation.h"
+#include "Particles/Location/ParticleModuleLocationBase.h"
+#include "Particles/Location/ParticleModuleLocationBoneSocket.h"
+#include "Particles/Location/ParticleModuleLocationDirect.h"
+#include "Particles/Location/ParticleModuleLocationEmitter.h"
+#include "Particles/Location/ParticleModuleLocationEmitterDirect.h"
+#include "Particles/Location/ParticleModuleLocationPrimitiveBase.h"
+#include "Particles/Location/ParticleModuleLocationPrimitiveCylinder.h"
+#include "Particles/Location/ParticleModuleLocationPrimitiveCylinder_Seeded.h"
+#include "Particles/Location/ParticleModuleLocationPrimitiveSphere.h"
+#include "Particles/Location/ParticleModuleLocationPrimitiveSphere_Seeded.h"
+#include "Particles/Location/ParticleModuleLocationPrimitiveTriangle.h"
+#include "Particles/Location/ParticleModuleLocationSkelVertSurface.h"
+#include "Particles/Location/ParticleModuleLocationWorldOffset.h"
+#include "Particles/Location/ParticleModuleLocationWorldOffset_Seeded.h"
+#include "Particles/Location/ParticleModuleLocation_Seeded.h"
+#include "Particles/Location/ParticleModuleSourceMovement.h"
+#include "Particles/Modules/Location/ParticleModulePivotOffset.h"
+#include "Particles/Orientation/ParticleModuleOrientationAxisLock.h"
+#include "Particles/Rotation/ParticleModuleRotation.h"
+#include "Particles/Rotation/ParticleModuleRotation_Seeded.h"
+#include "Particles/Rotation/ParticleModuleMeshRotation.h"
+#include "Particles/Rotation/ParticleModuleMeshRotation_Seeded.h"
+#include "Particles/Rotation/ParticleModuleRotationOverLifetime.h"
+#include "Particles/RotationRate/ParticleModuleRotationRate.h"
+#include "Particles/RotationRate/ParticleModuleRotationRateMultiplyLife.h"
+#include "Particles/RotationRate/ParticleModuleRotationRate_Seeded.h"
+#include "Particles/RotationRate/ParticleModuleMeshRotationRate.h"
+#include "Particles/RotationRate/ParticleModuleMeshRotationRate_Seeded.h"
+#include "Particles/RotationRate/ParticleModuleMeshRotationRateMultiplyLife.h"
+#include "Particles/RotationRate/ParticleModuleMeshRotationRateOverLife.h"
+#include "Particles/SubUV/ParticleModuleSubUVMovie.h"
+#include "Particles/TypeData/ParticleModuleTypeDataAnimTrail.h"
+#include "Particles/TypeData/ParticleModuleTypeDataBase.h"
+#include "Particles/TypeData/ParticleModuleTypeDataBeam2.h"
+#include "Particles/TypeData/ParticleModuleTypeDataGpu.h"
+#include "Particles/TypeData/ParticleModuleTypeDataMesh.h"
+#include "Particles/TypeData/ParticleModuleTypeDataRibbon.h"
+#include "Particles/ParticleEmitter.h"
+#include "Particles/ParticleLODLevel.h"
+#include "Particles/ParticleModule.h"
+#include "Particles/ParticleModuleRequired.h"
+#include "Particles/ParticleSpriteEmitter.h"
+#include "Particles/ParticleSystem.h"
+#include "Particles/ParticleSystemComponent.h"
 
 /*-----------------------------------------------------------------------------
 	Abstract base modules used for categorization.
diff --git a/Engine/Source/Runtime/Engine/Private/Particles/ParticleModules_Camera.cpp b/Engine/Source/Runtime/Engine/Private/Particles/ParticleModules_Camera.cpp
--- a/Engine/Source/Runtime/Engine/Private/Particles/ParticleModules_Camera.cpp
+++ b/Engine/Source/Runtime/Engine/Private/Particles/ParticleModules_Camera.cpp
@@ -7,6 +7,12 @@
 #include "EnginePrivate.h"
 #include "ParticleDefinitions.h"
 #include "../DistributionHelpers.h"
+#include "Particles/Camera/ParticleModuleCameraOffset.h"
+#include "Particles/TypeData/ParticleModuleTypeDataMesh.h"
+#include "Particles/ParticleLODLevel.h"
+#include "Particles/ParticleModuleRequired.h"
+#include "Particles/ParticleSpriteEmitter.h"
+#include "Particles/ParticleSystemComponent.h"
 
 UParticleModuleCameraBase::UParticleModuleCameraBase(const class FPostConstructInitializeProperties& PCIP)
 	: Super(PCIP)
diff --git a/Engine/Source/Runtime/Engine/Private/Particles/ParticleModules_Collision.cpp b/Engine/Source/Runtime/Engine/Private/Particles/ParticleModules_Collision.cpp
--- a/Engine/Source/Runtime/Engine/Private/Particles/ParticleModules_Collision.cpp
+++ b/Engine/Source/Runtime/Engine/Private/Particles/ParticleModules_Collision.cpp
@@ -7,6 +7,16 @@
 #include "EnginePrivate.h"
 #include "ParticleDefinitions.h"
 #include "../DistributionHelpers.h"
+#include "Particles/Collision/ParticleModuleCollision.h"
+#include "Particles/Collision/ParticleModuleCollisionBase.h"
+#include "Particles/Collision/ParticleModuleCollisionGPU.h"
+#include "Particles/Event/ParticleModuleEventGenerator.h"
+#include "Particles/TypeData/ParticleModuleTypeDataMesh.h"
+#include "Particles/TypeData/ParticleModuleTypeDataGpu.h"
+#include "Particles/ParticleLODLevel.h"
+#include "Particles/ParticleModuleRequired.h"
+#include "Particles/ParticleSpriteEmitter.h"
+#include "Particles/ParticleSystemComponent.h"
 
 UParticleModuleCollisionBase::UParticleModuleCollisionBase(const class FPostConstructInitializeProperties& PCIP)
 	: Super(PCIP)
diff --git a/Engine/Source/Runtime/Engine/Private/Particles/ParticleModules_Color.cpp b/Engine/Source/Runtime/Engine/Private/Particles/ParticleModules_Color.cpp
--- a/Engine/Source/Runtime/Engine/Private/Particles/ParticleModules_Color.cpp
+++ b/Engine/Source/Runtime/Engine/Private/Particles/ParticleModules_Color.cpp
@@ -7,6 +7,15 @@
 #include "EnginePrivate.h"
 #include "ParticleDefinitions.h"
 #include "../DistributionHelpers.h"
+#include "Particles/Color/ParticleModuleColor.h"
+#include "Particles/Color/ParticleModuleColorBase.h"
+#include "Particles/Color/ParticleModuleColorOverLife.h"
+#include "Particles/Color/ParticleModuleColorScaleOverLife.h"
+#include "Particles/Color/ParticleModuleColor_Seeded.h"
+#include "Particles/ParticleEmitter.h"
+#include "Particles/ParticleLODLevel.h"
+#include "Particles/ParticleSystem.h"
+#include "Particles/ParticleSystemComponent.h"
 
 UParticleModuleColorBase::UParticleModuleColorBase(const class FPostConstructInitializeProperties& PCIP)
 	: Super(PCIP)
diff --git a/Engine/Source/Runtime/Engine/Private/Particles/ParticleModules_Event.cpp b/Engine/Source/Runtime/Engine/Private/Particles/ParticleModules_Event.cpp
--- a/Engine/Source/Runtime/Engine/Private/Particles/ParticleModules_Event.cpp
+++ b/Engine/Source/Runtime/Engine/Private/Particles/ParticleModules_Event.cpp
@@ -6,6 +6,13 @@
 #include "EnginePrivate.h"
 #include "ParticleDefinitions.h"
 #include "../DistributionHelpers.h"
+#include "Particles/Event/ParticleModuleEventBase.h"
+#include "Particles/Event/ParticleModuleEventGenerator.h"
+#include "Particles/Event/ParticleModuleEventReceiverBase.h"
+#include "Particles/Event/ParticleModuleEventReceiverKillParticles.h"
+#include "Particles/Event/ParticleModuleEventReceiverSpawn.h"
+#include "Particles/Event/ParticleModuleEventSendToGame.h"
+#include "Particles/ParticleLODLevel.h"
 
 /*-----------------------------------------------------------------------------
 	Abstract base modules used for categorization.
diff --git a/Engine/Source/Runtime/Engine/Private/Particles/ParticleModules_Location.cpp b/Engine/Source/Runtime/Engine/Private/Particles/ParticleModules_Location.cpp
--- a/Engine/Source/Runtime/Engine/Private/Particles/ParticleModules_Location.cpp
+++ b/Engine/Source/Runtime/Engine/Private/Particles/ParticleModules_Location.cpp
@@ -8,6 +8,25 @@
 #include "EnginePrivate.h"
 #include "ParticleDefinitions.h"
 #include "../DistributionHelpers.h"
+#include "Particles/Location/ParticleModuleLocation.h"
+#include "Particles/Location/ParticleModuleLocationBoneSocket.h"
+#include "Particles/Location/ParticleModuleLocationDirect.h"
+#include "Particles/Location/ParticleModuleLocationEmitter.h"
+#include "Particles/Location/ParticleModuleLocationEmitterDirect.h"
+#include "Particles/Location/ParticleModuleLocationPrimitiveBase.h"
+#include "Particles/Location/ParticleModuleLocationPrimitiveCylinder.h"
+#include "Particles/Location/ParticleModuleLocationPrimitiveCylinder_Seeded.h"
+#include "Particles/Location/ParticleModuleLocationPrimitiveSphere.h"
+#include "Particles/Location/ParticleModuleLocationPrimitiveSphere_Seeded.h"
+#include "Particles/Location/ParticleModuleLocationPrimitiveTriangle.h"
+#include "Particles/Location/ParticleModuleLocationSkelVertSurface.h"
+#include "Particles/Location/ParticleModuleLocationWorldOffset.h"
+#include "Particles/Location/ParticleModuleLocationWorldOffset_Seeded.h"
+#include "Particles/Location/ParticleModuleLocation_Seeded.h"
+#include "Particles/ParticleLODLevel.h"
+#include "Particles/ParticleModuleRequired.h"
+#include "Particles/ParticleSpriteEmitter.h"
+#include "Particles/ParticleSystemComponent.h"
 
 UParticleModuleLocationBase::UParticleModuleLocationBase(const class FPostConstructInitializeProperties& PCIP)
 	: Super(PCIP)
diff --git a/Engine/Source/Runtime/Engine/Private/Particles/ParticleModules_Material.cpp b/Engine/Source/Runtime/Engine/Private/Particles/ParticleModules_Material.cpp
--- a/Engine/Source/Runtime/Engine/Private/Particles/ParticleModules_Material.cpp
+++ b/Engine/Source/Runtime/Engine/Private/Particles/ParticleModules_Material.cpp
@@ -6,6 +6,7 @@
 =============================================================================*/
 #include "EnginePrivate.h"
 #include "ParticleDefinitions.h"
+#include "Particles/Material/ParticleModuleMeshMaterial.h"
 
 UParticleModuleMaterialBase::UParticleModuleMaterialBase(const class FPostConstructInitializeProperties& PCIP)
 	: Super(PCIP)
diff --git a/Engine/Source/Runtime/Engine/Private/Particles/ParticleModules_Orbit.cpp b/Engine/Source/Runtime/Engine/Private/Particles/ParticleModules_Orbit.cpp
--- a/Engine/Source/Runtime/Engine/Private/Particles/ParticleModules_Orbit.cpp
+++ b/Engine/Source/Runtime/Engine/Private/Particles/ParticleModules_Orbit.cpp
@@ -6,6 +6,10 @@
 #include "EnginePrivate.h"
 #include "ParticleDefinitions.h"
 #include "../DistributionHelpers.h"
+#include "Particles/Orbit/ParticleModuleOrbit.h"
+#include "Particles/TypeData/ParticleModuleTypeDataGpu.h"
+#include "Particles/ParticleLODLevel.h"
+#include "Particles/ParticleSystemComponent.h"
 
 UParticleModuleOrbitBase::UParticleModuleOrbitBase(const class FPostConstructInitializeProperties& PCIP)
 	: Super(PCIP)
diff --git a/Engine/Source/Runtime/Engine/Private/Particles/ParticleModules_Parameter.cpp b/Engine/Source/Runtime/Engine/Private/Particles/ParticleModules_Parameter.cpp
--- a/Engine/Source/Runtime/Engine/Private/Particles/ParticleModules_Parameter.cpp
+++ b/Engine/Source/Runtime/Engine/Private/Particles/ParticleModules_Parameter.cpp
@@ -8,6 +8,13 @@
 #include "Materials/MaterialExpressionDynamicParameter.h"
 #include "ParticleDefinitions.h"
 #include "../DistributionHelpers.h"
+#include "Particles/Parameter/ParticleModuleParameterBase.h"
+#include "Particles/Parameter/ParticleModuleParameterDynamic.h"
+#include "Particles/Parameter/ParticleModuleParameterDynamic_Seeded.h"
+#include "Particles/TypeData/ParticleModuleTypeDataMesh.h"
+#include "Particles/ParticleEmitter.h"
+#include "Particles/ParticleLODLevel.h"
+#include "Particles/ParticleModuleRequired.h"
 
 UParticleModuleParameterBase::UParticleModuleParameterBase(const class FPostConstructInitializeProperties& PCIP)
 	: Super(PCIP)
diff --git a/Engine/Source/Runtime/Engine/Private/Particles/ParticleModules_Size.cpp b/Engine/Source/Runtime/Engine/Private/Particles/ParticleModules_Size.cpp
--- a/Engine/Source/Runtime/Engine/Private/Particles/ParticleModules_Size.cpp
+++ b/Engine/Source/Runtime/Engine/Private/Particles/ParticleModules_Size.cpp
@@ -8,6 +8,13 @@
 #include "EnginePrivate.h"
 #include "ParticleDefinitions.h"
 #include "../DistributionHelpers.h"
+#include "Particles/Size/ParticleModuleSize_Seeded.h"
+#include "Particles/Size/ParticleModuleSizeMultiplyLife.h"
+#include "Particles/Size/ParticleModuleSizeScale.h"
+#include "Particles/Size/ParticleModuleSizeScaleBySpeed.h"
+#include "Particles/TypeData/ParticleModuleTypeDataGpu.h"
+#include "Particles/ParticleLODLevel.h"
+#include "Particles/ParticleSystemComponent.h"
 
 UParticleModuleSizeBase::UParticleModuleSizeBase(const class FPostConstructInitializeProperties& PCIP)
 	: Super(PCIP)
diff --git a/Engine/Source/Runtime/Engine/Private/Particles/ParticleModules_Spawn.cpp b/Engine/Source/Runtime/Engine/Private/Particles/ParticleModules_Spawn.cpp
--- a/Engine/Source/Runtime/Engine/Private/Particles/ParticleModules_Spawn.cpp
+++ b/Engine/Source/Runtime/Engine/Private/Particles/ParticleModules_Spawn.cpp
@@ -6,6 +6,11 @@
 #include "EnginePrivate.h"
 #include "ParticleDefinitions.h"
 #include "../DistributionHelpers.h"
+#include "Particles/Spawn/ParticleModuleSpawn.h"
+#include "Particles/Spawn/ParticleModuleSpawnBase.h"
+#include "Particles/Spawn/ParticleModuleSpawnPerUnit.h"
+#include "Particles/ParticleEmitter.h"
+#include "Particles/ParticleSystemComponent.h"
 
 UParticleModuleSpawnBase::UParticleModuleSpawnBase(const class FPostConstructInitializeProperties& PCIP)
 	: Super(PCIP)
diff --git a/Engine/Source/Runtime/Engine/Private/Particles/ParticleModules_VectorField.cpp b/Engine/Source/Runtime/Engine/Private/Particles/ParticleModules_VectorField.cpp
--- a/Engine/Source/Runtime/Engine/Private/Particles/ParticleModules_VectorField.cpp
+++ b/Engine/Source/Runtime/Engine/Private/Particles/ParticleModules_VectorField.cpp
@@ -8,6 +8,14 @@
 #include "FXSystem.h"
 #include "ParticleDefinitions.h"
 #include "../DistributionHelpers.h"
+#include "Particles/TypeData/ParticleModuleTypeDataGpu.h"
+#include "Particles/VectorField/ParticleModuleVectorFieldLocal.h"
+#include "Particles/VectorField/ParticleModuleVectorFieldGlobal.h"
+#include "Particles/VectorField/ParticleModuleVectorFieldRotation.h"
+#include "Particles/VectorField/ParticleModuleVectorFieldRotationRate.h"
+#include "Particles/VectorField/ParticleModuleVectorFieldScale.h"
+#include "Particles/VectorField/ParticleModuleVectorFieldScaleOverLife.h"
+#include "Particles/ParticleLODLevel.h"
 
 /*------------------------------------------------------------------------------
 	Global vector field scale.
diff --git a/Engine/Source/Runtime/Engine/Private/Particles/ParticleModules_Velocity.cpp b/Engine/Source/Runtime/Engine/Private/Particles/ParticleModules_Velocity.cpp
--- a/Engine/Source/Runtime/Engine/Private/Particles/ParticleModules_Velocity.cpp
+++ b/Engine/Source/Runtime/Engine/Private/Particles/ParticleModules_Velocity.cpp
@@ -8,6 +8,18 @@
 #include "EnginePrivate.h"
 #include "ParticleDefinitions.h"
 #include "../DistributionHelpers.h"
+#include "Particles/Lifetime/ParticleModuleLifetimeBase.h"
+#include "Particles/TypeData/ParticleModuleTypeDataMesh.h"
+#include "Particles/Velocity/ParticleModuleVelocity.h"
+#include "Particles/Velocity/ParticleModuleVelocityBase.h"
+#include "Particles/Velocity/ParticleModuleVelocityCone.h"
+#include "Particles/Velocity/ParticleModuleVelocityInheritParent.h"
+#include "Particles/Velocity/ParticleModuleVelocityOverLifetime.h"
+#include "Particles/Velocity/ParticleModuleVelocity_Seeded.h"
+#include "Particles/ParticleLODLevel.h"
+#include "Particles/ParticleModuleRequired.h"
+#include "Particles/ParticleSpriteEmitter.h"
+#include "Particles/ParticleSystemComponent.h"
 
 UParticleModuleVelocityBase::UParticleModuleVelocityBase(const class FPostConstructInitializeProperties& PCIP)
 	: Super(PCIP)
diff --git a/Engine/Source/Runtime/Engine/Private/Particles/ParticleSystemRender.cpp b/Engine/Source/Runtime/Engine/Private/Particles/ParticleSystemRender.cpp
--- a/Engine/Source/Runtime/Engine/Private/Particles/ParticleSystemRender.cpp
+++ b/Engine/Source/Runtime/Engine/Private/Particles/ParticleSystemRender.cpp
@@ -7,6 +7,16 @@
 #include "ParticleDefinitions.h"
 #include "DiagnosticTable.h"
 #include "ParticleResources.h"
+#include "Particles/Orientation/ParticleModuleOrientationAxisLock.h"
+#include "Particles/Rotation/ParticleModuleRotation.h"
+#include "Particles/Rotation/ParticleModuleRotation_Seeded.h"
+#include "Particles/TypeData/ParticleModuleTypeDataBeam2.h"
+#include "Particles/TypeData/ParticleModuleTypeDataMesh.h"
+#include "Particles/TypeData/ParticleModuleTypeDataRibbon.h"
+#include "Particles/ParticleLODLevel.h"
+#include "Particles/ParticleModuleRequired.h"
+#include "Particles/ParticleSpriteEmitter.h"
+#include "Particles/ParticleSystemComponent.h"
 
 /** 
  * Whether to track particle rendering stats.  
diff --git a/Engine/Source/Runtime/Engine/Private/Particles/ParticleTrail2EmitterInstance.cpp b/Engine/Source/Runtime/Engine/Private/Particles/ParticleTrail2EmitterInstance.cpp
--- a/Engine/Source/Runtime/Engine/Private/Particles/ParticleTrail2EmitterInstance.cpp
+++ b/Engine/Source/Runtime/Engine/Private/Particles/ParticleTrail2EmitterInstance.cpp
@@ -7,7 +7,19 @@
 
 #include "EnginePrivate.h"
 #include "ParticleDefinitions.h"
-
+#include "ParticleHelper.h"
+#include "Particles/Event/ParticleModuleEventGenerator.h"
+#include "Particles/Lifetime/ParticleModuleLifetime.h"
+#include "Particles/Spawn/ParticleModuleSpawn.h"
+#include "Particles/Spawn/ParticleModuleSpawnPerUnit.h"
+#include "Particles/Trail/ParticleModuleTrailSource.h"
+#include "Particles/TypeData/ParticleModuleTypeDataAnimTrail.h"
+#include "Particles/TypeData/ParticleModuleTypeDataRibbon.h"
+#include "Particles/ParticleLODLevel.h"
+#include "Particles/ParticleEmitter.h"
+#include "Particles/ParticleSpriteEmitter.h"
+#include "Particles/ParticleSystemComponent.h"
+#include "Particles/ParticleModuleRequired.h"
 /** trail stats */
 
 
diff --git a/Engine/Source/Runtime/Engine/Private/Particles/ParticleTrailModules.cpp b/Engine/Source/Runtime/Engine/Private/Particles/ParticleTrailModules.cpp
--- a/Engine/Source/Runtime/Engine/Private/Particles/ParticleTrailModules.cpp
+++ b/Engine/Source/Runtime/Engine/Private/Particles/ParticleTrailModules.cpp
@@ -7,6 +7,10 @@
 #include "EnginePrivate.h"
 #include "ParticleDefinitions.h"
 #include "../DistributionHelpers.h"
+#include "Particles/Trail/ParticleModuleTrailSource.h"
+#include "Particles/TypeData/ParticleModuleTypeDataAnimTrail.h"
+#include "Particles/TypeData/ParticleModuleTypeDataRibbon.h"
+#include "Particles/ParticleSystemComponent.h"
 
 UParticleModuleTrailBase::UParticleModuleTrailBase(const class FPostConstructInitializeProperties& PCIP)
 	: Super(PCIP)
diff --git a/Engine/Source/Runtime/Engine/Private/UnrealEngine.cpp b/Engine/Source/Runtime/Engine/Private/UnrealEngine.cpp
--- a/Engine/Source/Runtime/Engine/Private/UnrealEngine.cpp
+++ b/Engine/Source/Runtime/Engine/Private/UnrealEngine.cpp
@@ -39,6 +39,16 @@
 #include "AudioDevice.h"
 #include "ActiveSound.h"
 
+#include "Particles/Spawn/ParticleModuleSpawn.h"
+#include "Particles/TypeData/ParticleModuleTypeDataMesh.h"
+#include "Particles/ParticleEmitter.h"
+#include "Particles/ParticleLODLevel.h"
+#include "Particles/ParticleModule.h"
+#include "Particles/ParticleModuleRequired.h"
+#include "Particles/ParticleSpriteEmitter.h"
+#include "Particles/ParticleSystem.h"
+#include "Particles/ParticleSystemComponent.h"
+
 #if !UE_BUILD_SHIPPING
 #include "STaskGraph.h"
 #endif
diff --git a/Engine/Source/Runtime/Engine/Private/World.cpp b/Engine/Source/Runtime/Engine/Private/World.cpp
--- a/Engine/Source/Runtime/Engine/Private/World.cpp
+++ b/Engine/Source/Runtime/Engine/Private/World.cpp
@@ -29,6 +29,7 @@
 #include "MallocProfiler.h"
 
 #include "Landscape/LandscapeInfo.h"
+#include "Particles/ParticleEventManager.h"
 
 DEFINE_LOG_CATEGORY_STATIC(LogWorld, Log, All);
 DEFINE_LOG_CATEGORY(LogSpawn);
diff --git a/Engine/Source/Runtime/Engine/Private/WorldSettings.cpp b/Engine/Source/Runtime/Engine/Private/WorldSettings.cpp
--- a/Engine/Source/Runtime/Engine/Private/WorldSettings.cpp
+++ b/Engine/Source/Runtime/Engine/Private/WorldSettings.cpp
@@ -8,6 +8,7 @@
 #include "MessageLog.h"
 #include "UObjectToken.h"
 #include "MapErrors.h"
+#include "Particles/ParticleEventManager.h"
 
 #define LOCTEXT_NAMESPACE "ErrorChecking"
 
diff --git a/Engine/Source/Runtime/Engine/Public/Engine.h b/Engine/Source/Runtime/Engine/Public/Engine.h
--- a/Engine/Source/Runtime/Engine/Public/Engine.h
+++ b/Engine/Source/Runtime/Engine/Public/Engine.h
@@ -66,7 +66,6 @@
 #include "Engine/DebugCameraController.h"
 #include "Engine/DecalActor.h"
 #include "PhysicsEngine/DestructibleActor.h"
-#include "Particles/Emitter.h"
 #include "Engine/Canvas.h"
 #include "GameFramework/HUD.h"
 #include "Debug/GameplayDebuggingHUDComponent.h"
@@ -104,9 +103,6 @@
 #include "AI/Navigation/NavLinkProxy.h"
 #include "Engine/Note.h"
 #include "Components/PrimitiveComponent.h"
-#include "Particles/ParticleSystem.h"
-#include "Particles/ParticleSystemComponent.h"
-#include "Particles/ParticleEventManager.h"
 #include "GameFramework/MovementComponent.h"
 #include "GameFramework/NavMovementComponent.h"
 #include "GameFramework/PawnMovementComponent.h"
@@ -441,127 +437,6 @@
 #include "Engine/ObjectReferencer.h"
 #include "GameFramework/OnlineSession.h"
 #include "Engine/PackageMapClient.h"
-#include "Particles/ParticleLODLevel.h"
-#include "Particles/ParticleEmitter.h"
-#include "Particles/ParticleSpriteEmitter.h"
-#include "Particles/ParticleModule.h"
-#include "Particles/Acceleration/ParticleModuleAccelerationBase.h"
-#include "Particles/Acceleration/ParticleModuleAcceleration.h"
-#include "Particles/Acceleration/ParticleModuleAccelerationConstant.h"
-#include "Particles/Acceleration/ParticleModuleAccelerationDrag.h"
-#include "Particles/Acceleration/ParticleModuleAccelerationDragScaleOverLife.h"
-#include "Particles/Acceleration/ParticleModuleAccelerationOverLifetime.h"
-#include "Particles/Attractor/ParticleModuleAttractorBase.h"
-#include "Particles/Attractor/ParticleModuleAttractorLine.h"
-#include "Particles/Attractor/ParticleModuleAttractorParticle.h"
-#include "Particles/Attractor/ParticleModuleAttractorPoint.h"
-#include "Particles/Attractor/ParticleModuleAttractorPointGravity.h"
-#include "Particles/Beam/ParticleModuleBeamBase.h"
-#include "Particles/Beam/ParticleModuleBeamModifier.h"
-#include "Particles/Beam/ParticleModuleBeamNoise.h"
-#include "Particles/Beam/ParticleModuleBeamSource.h"
-#include "Particles/Beam/ParticleModuleBeamTarget.h"
-#include "Particles/Camera/ParticleModuleCameraBase.h"
-#include "Particles/Camera/ParticleModuleCameraOffset.h"
-#include "Particles/Collision/ParticleModuleCollisionBase.h"
-#include "Particles/Collision/ParticleModuleCollision.h"
-#include "Particles/Collision/ParticleModuleCollisionGPU.h"
-#include "Particles/Color/ParticleModuleColorBase.h"
-#include "Particles/Color/ParticleModuleColor.h"
-#include "Particles/Color/ParticleModuleColor_Seeded.h"
-#include "Particles/Color/ParticleModuleColorOverLife.h"
-#include "Particles/Color/ParticleModuleColorScaleOverLife.h"
-#include "Particles/Event/ParticleModuleEventBase.h"
-#include "Particles/Event/ParticleModuleEventGenerator.h"
-#include "Particles/Event/ParticleModuleEventReceiverBase.h"
-#include "Particles/Event/ParticleModuleEventReceiverKillParticles.h"
-#include "Particles/Event/ParticleModuleEventReceiverSpawn.h"
-#include "Particles/Kill/ParticleModuleKillBase.h"
-#include "Particles/Kill/ParticleModuleKillBox.h"
-#include "Particles/Kill/ParticleModuleKillHeight.h"
-#include "Particles/Lifetime/ParticleModuleLifetimeBase.h"
-#include "Particles/Lifetime/ParticleModuleLifetime.h"
-#include "Particles/Lifetime/ParticleModuleLifetime_Seeded.h"
-#include "Particles/Light/ParticleModuleLightBase.h"
-#include "Particles/Light/ParticleModuleLight.h"
-#include "Particles/Light/ParticleModuleLight_Seeded.h"
-#include "Particles/Location/ParticleModuleLocationBase.h"
-#include "Particles/Location/ParticleModuleLocation.h"
-#include "Particles/Location/ParticleModuleLocation_Seeded.h"
-#include "Particles/Location/ParticleModuleLocationWorldOffset.h"
-#include "Particles/Location/ParticleModuleLocationWorldOffset_Seeded.h"
-#include "Particles/Location/ParticleModuleLocationBoneSocket.h"
-#include "Particles/Location/ParticleModuleLocationDirect.h"
-#include "Particles/Location/ParticleModuleLocationEmitter.h"
-#include "Particles/Location/ParticleModuleLocationEmitterDirect.h"
-#include "Particles/Location/ParticleModuleLocationPrimitiveBase.h"
-#include "Particles/Location/ParticleModuleLocationPrimitiveCylinder.h"
-#include "Particles/Location/ParticleModuleLocationPrimitiveCylinder_Seeded.h"
-#include "Particles/Location/ParticleModuleLocationPrimitiveSphere.h"
-#include "Particles/Location/ParticleModuleLocationPrimitiveSphere_Seeded.h"
-#include "Particles/Location/ParticleModuleLocationPrimitiveTriangle.h"
-#include "Particles/Location/ParticleModuleLocationSkelVertSurface.h"
-#include "Particles/Modules/Location/ParticleModulePivotOffset.h"
-#include "Particles/Location/ParticleModuleSourceMovement.h"
-#include "Particles/Material/ParticleModuleMaterialBase.h"
-#include "Particles/Material/ParticleModuleMeshMaterial.h"
-#include "Particles/Orbit/ParticleModuleOrbitBase.h"
-#include "Particles/Orbit/ParticleModuleOrbit.h"
-#include "Particles/Orientation/ParticleModuleOrientationBase.h"
-#include "Particles/Orientation/ParticleModuleOrientationAxisLock.h"
-#include "Particles/Parameter/ParticleModuleParameterBase.h"
-#include "Particles/Parameter/ParticleModuleParameterDynamic.h"
-#include "Particles/Parameter/ParticleModuleParameterDynamic_Seeded.h"
-#include "Particles/ParticleModuleRequired.h"
-#include "Particles/Rotation/ParticleModuleRotationBase.h"
-#include "Particles/Rotation/ParticleModuleMeshRotation.h"
-#include "Particles/Rotation/ParticleModuleMeshRotation_Seeded.h"
-#include "Particles/Rotation/ParticleModuleRotation.h"
-#include "Particles/Rotation/ParticleModuleRotation_Seeded.h"
-#include "Particles/Rotation/ParticleModuleRotationOverLifetime.h"
-#include "Particles/RotationRate/ParticleModuleRotationRateBase.h"
-#include "Particles/RotationRate/ParticleModuleMeshRotationRate.h"
-#include "Particles/RotationRate/ParticleModuleMeshRotationRate_Seeded.h"
-#include "Particles/RotationRate/ParticleModuleMeshRotationRateMultiplyLife.h"
-#include "Particles/RotationRate/ParticleModuleMeshRotationRateOverLife.h"
-#include "Particles/RotationRate/ParticleModuleRotationRate.h"
-#include "Particles/RotationRate/ParticleModuleRotationRate_Seeded.h"
-#include "Particles/RotationRate/ParticleModuleRotationRateMultiplyLife.h"
-#include "Particles/Size/ParticleModuleSizeBase.h"
-#include "Particles/Size/ParticleModuleSize.h"
-#include "Particles/Size/ParticleModuleSize_Seeded.h"
-#include "Particles/Size/ParticleModuleSizeMultiplyLife.h"
-#include "Particles/Size/ParticleModuleSizeScale.h"
-#include "Particles/Size/ParticleModuleSizeScaleBySpeed.h"
-#include "Particles/Spawn/ParticleModuleSpawnBase.h"
-#include "Particles/Spawn/ParticleModuleSpawn.h"
-#include "Particles/Spawn/ParticleModuleSpawnPerUnit.h"
-#include "Particles/SubUV/ParticleModuleSubUVBase.h"
-#include "Particles/SubUV/ParticleModuleSubUV.h"
-#include "Particles/SubUV/ParticleModuleSubUVMovie.h"
-#include "Particles/Trail/ParticleModuleTrailBase.h"
-#include "Particles/Trail/ParticleModuleTrailSource.h"
-#include "Particles/TypeData/ParticleModuleTypeDataBase.h"
-#include "Particles/TypeData/ParticleModuleTypeDataAnimTrail.h"
-#include "Particles/TypeData/ParticleModuleTypeDataBeam2.h"
-#include "Particles/TypeData/ParticleModuleTypeDataGpu.h"
-#include "Particles/TypeData/ParticleModuleTypeDataMesh.h"
-#include "Particles/TypeData/ParticleModuleTypeDataRibbon.h"
-#include "Particles/VectorField/ParticleModuleVectorFieldBase.h"
-#include "Particles/VectorField/ParticleModuleVectorFieldGlobal.h"
-#include "Particles/VectorField/ParticleModuleVectorFieldLocal.h"
-#include "Particles/VectorField/ParticleModuleVectorFieldRotation.h"
-#include "Particles/VectorField/ParticleModuleVectorFieldRotationRate.h"
-#include "Particles/VectorField/ParticleModuleVectorFieldScale.h"
-#include "Particles/VectorField/ParticleModuleVectorFieldScaleOverLife.h"
-#include "Particles/Velocity/ParticleModuleVelocityBase.h"
-#include "Particles/Velocity/ParticleModuleVelocity.h"
-#include "Particles/Velocity/ParticleModuleVelocity_Seeded.h"
-#include "Particles/Velocity/ParticleModuleVelocityCone.h"
-#include "Particles/Velocity/ParticleModuleVelocityInheritParent.h"
-#include "Particles/Velocity/ParticleModuleVelocityOverLifetime.h"
-#include "Particles/Event/ParticleModuleEventSendToGame.h"
-#include "Particles/ParticleSystemReplay.h"
 #include "Engine/PendingNetGame.h"
 #include "PhysicalMaterials/PhysicalMaterial.h"
 #include "PhysicalMaterials/PhysicalMaterialPropertyBase.h"
