Fixed anim trails texture tiling.

#ttp 336607 - UE4:ORION:Particle anim trail tiling UV option is probably broken

diff --git a/Engine/Source/Runtime/Engine/Private/Particles/ParticleTrail2EmitterInstance.cpp b/Engine/Source/Runtime/Engine/Private/Particles/ParticleTrail2EmitterInstance.cpp
--- a/Engine/Source/Runtime/Engine/Private/Particles/ParticleTrail2EmitterInstance.cpp
+++ b/Engine/Source/Runtime/Engine/Private/Particles/ParticleTrail2EmitterInstance.cpp
@@ -3912,9 +3912,14 @@ FDynamicEmitterDataBase* FParticleAnimTrailEmitterInstance::GetDynamicData(bool
 	NewEmitterData->bRenderParticles = false;
 	NewEmitterData->bRenderTangents = false;
 	NewEmitterData->bRenderTessellation = false;
-#endif
+#endif	
+	
+	//These are unused for anim trails.
+	//NewEmitterData->DistanceTessellationStepSize = TrailTypeData->DistanceTessellationStepSize;
+	//NewEmitterData->TangentTessellationScalar = TrailTypeData->TangentTessellationScalar;
+	//NewEmitterData->TextureTileDistance = TrailTypeData->TilingDistance;
 
-	if (NewEmitterData->TextureTileDistance > 0.0f)
+	if (TrailTypeData->TilingDistance > 0.0f)
 	{
 		NewEmitterData->bTextureTileDistance = true;
 	}
