Extracting AI code from the engine into a separate module, pass 1 #UE4

diff --git a/Engine/Source/Runtime/Engine/Classes/GameFramework/Pawn.h b/Engine/Source/Runtime/Engine/Classes/GameFramework/Pawn.h
--- a/Engine/Source/Runtime/Engine/Classes/GameFramework/Pawn.h
+++ b/Engine/Source/Runtime/Engine/Classes/GameFramework/Pawn.h
@@ -82,7 +82,7 @@ public:
 
 	/** default class to use when pawn is controlled by AI. */
 	UPROPERTY(EditAnywhere, BlueprintReadWrite, Category=AI)
-	TSubclassOf<class AAIController>  AIControllerClass;
+	TSubclassOf<class AController>  AIControllerClass;
 
 	/** If Pawn is possessed by a player, points to his playerstate.  Needed for network play as controllers are not replicated to clients. */
 	UPROPERTY(editinline, replicatedUsing=OnRep_PlayerState, BlueprintReadOnly, Category="Pawn")
diff --git a/Engine/Source/Runtime/Engine/Engine.Build.cs b/Engine/Source/Runtime/Engine/Engine.Build.cs
--- a/Engine/Source/Runtime/Engine/Engine.Build.cs
+++ b/Engine/Source/Runtime/Engine/Engine.Build.cs
@@ -51,7 +51,8 @@ public class Engine : ModuleRules
 				"AssetRegistry", // Here until FAssetData is moved to engine
 				"EngineMessages",
 				"EngineSettings",
-				"SynthBenchmark",
+				"SynthBenchmark",                
+                "AIModule",
 			}
 		);
 
@@ -62,7 +63,6 @@ public class Engine : ModuleRules
 				"SlateCore",
 				"SlateReflector",
 				"VectorVM",
-                "AIModule"
 			}
 		);
 
diff --git a/Engine/Source/Runtime/Engine/Private/AI/Navigation/AvoidanceManager.cpp b/Engine/Source/Runtime/Engine/Private/AI/Navigation/AvoidanceManager.cpp
--- a/Engine/Source/Runtime/Engine/Private/AI/Navigation/AvoidanceManager.cpp
+++ b/Engine/Source/Runtime/Engine/Private/AI/Navigation/AvoidanceManager.cpp
@@ -1,6 +1,7 @@
 // Copyright 1998-2014 Epic Games, Inc. All Rights Reserved.
 
 #include "EnginePrivate.h"
+#include "AI/Navigation/AvoidanceManager.h"
 
 DEFINE_STAT(STAT_AI_ObstacleAvoidance);
 
diff --git a/Engine/Source/Runtime/Engine/Private/AI/Navigation/NavAreas/NavAreaMeta.cpp b/Engine/Source/Runtime/Engine/Private/AI/Navigation/NavAreas/NavAreaMeta.cpp
--- a/Engine/Source/Runtime/Engine/Private/AI/Navigation/NavAreas/NavAreaMeta.cpp
+++ b/Engine/Source/Runtime/Engine/Private/AI/Navigation/NavAreas/NavAreaMeta.cpp
@@ -1,6 +1,7 @@
 // Copyright 1998-2014 Epic Games, Inc. All Rights Reserved.
 
 #include "EnginePrivate.h"
+#include "AI/Navigation/NavAreas/NavAreaMeta.h"
 
 UNavAreaMeta::UNavAreaMeta(const class FPostConstructInitializeProperties& PCIP) : Super(PCIP)
 {
diff --git a/Engine/Source/Runtime/Engine/Private/AI/Navigation/NavAreas/NavAreaMeta_SwitchByAgent.cpp b/Engine/Source/Runtime/Engine/Private/AI/Navigation/NavAreas/NavAreaMeta_SwitchByAgent.cpp
--- a/Engine/Source/Runtime/Engine/Private/AI/Navigation/NavAreas/NavAreaMeta_SwitchByAgent.cpp
+++ b/Engine/Source/Runtime/Engine/Private/AI/Navigation/NavAreas/NavAreaMeta_SwitchByAgent.cpp
@@ -1,6 +1,8 @@
 // Copyright 1998-2014 Epic Games, Inc. All Rights Reserved.
 
 #include "EnginePrivate.h"
+#include "AI/Navigation/NavAreas/NavArea_Default.h"
+#include "AI/Navigation/NavAreas/NavAreaMeta_SwitchByAgent.h"
 
 UNavAreaMeta_SwitchByAgent::UNavAreaMeta_SwitchByAgent(const class FPostConstructInitializeProperties& PCIP) : Super(PCIP)
 {
diff --git a/Engine/Source/Runtime/Engine/Private/AI/Navigation/NavAreas/NavArea_Default.cpp b/Engine/Source/Runtime/Engine/Private/AI/Navigation/NavAreas/NavArea_Default.cpp
--- a/Engine/Source/Runtime/Engine/Private/AI/Navigation/NavAreas/NavArea_Default.cpp
+++ b/Engine/Source/Runtime/Engine/Private/AI/Navigation/NavAreas/NavArea_Default.cpp
@@ -1,6 +1,7 @@
 // Copyright 1998-2014 Epic Games, Inc. All Rights Reserved.
 
 #include "EnginePrivate.h"
+#include "AI/Navigation/NavAreas/NavArea_Default.h"
 
 UNavArea_Default::UNavArea_Default(const class FPostConstructInitializeProperties& PCIP) : Super(PCIP)
 {
diff --git a/Engine/Source/Runtime/Engine/Private/AI/Navigation/NavAreas/NavArea_Null.cpp b/Engine/Source/Runtime/Engine/Private/AI/Navigation/NavAreas/NavArea_Null.cpp
--- a/Engine/Source/Runtime/Engine/Private/AI/Navigation/NavAreas/NavArea_Null.cpp
+++ b/Engine/Source/Runtime/Engine/Private/AI/Navigation/NavAreas/NavArea_Null.cpp
@@ -1,6 +1,7 @@
 // Copyright 1998-2014 Epic Games, Inc. All Rights Reserved.
 
 #include "EnginePrivate.h"
+#include "AI/Navigation/NavAreas/NavArea_Null.h"
 
 UNavArea_Null::UNavArea_Null(const class FPostConstructInitializeProperties& PCIP) : Super(PCIP)
 {
diff --git a/Engine/Source/Runtime/Engine/Private/AI/Navigation/NavCollision.cpp b/Engine/Source/Runtime/Engine/Private/AI/Navigation/NavCollision.cpp
--- a/Engine/Source/Runtime/Engine/Private/AI/Navigation/NavCollision.cpp
+++ b/Engine/Source/Runtime/Engine/Private/AI/Navigation/NavCollision.cpp
@@ -2,7 +2,7 @@
 
 #include "EnginePrivate.h"
 #include "AI/Navigation/NavCollision.h"
-#include "AI/NavigationModifier.h"
+#include "NavigationModifier.h"
 #include "../Private/AI/Navigation/RecastNavMeshGenerator.h"
 #include "DerivedDataPluginInterface.h"
 #include "DerivedDataCacheInterface.h"
diff --git a/Engine/Source/Runtime/Engine/Private/AI/Navigation/NavFilters/NavigationQueryFilter.cpp b/Engine/Source/Runtime/Engine/Private/AI/Navigation/NavFilters/NavigationQueryFilter.cpp
--- a/Engine/Source/Runtime/Engine/Private/AI/Navigation/NavFilters/NavigationQueryFilter.cpp
+++ b/Engine/Source/Runtime/Engine/Private/AI/Navigation/NavFilters/NavigationQueryFilter.cpp
@@ -1,6 +1,10 @@
 // Copyright 1998-2014 Epic Games, Inc. All Rights Reserved.
 
 #include "EnginePrivate.h"
+// @todo AIModule circular dependency
+#include "AIController.h"
+#include "Navigation/NavigationComponent.h"
+#include "AI/Navigation/NavFilters/NavigationQueryFilter.h"
 
 UNavigationQueryFilter::UNavigationQueryFilter(const class FPostConstructInitializeProperties& PCIP) : Super(PCIP)
 {
@@ -64,12 +68,6 @@ void UNavigationQueryFilter::InitializeFilter(const ANavigationData* NavData, FN
 	Filter->SetExcludeFlags(ExcludeFlags.Packed);
 }
 
-TSharedPtr<const struct FNavigationQueryFilter> UNavigationQueryFilter::GetQueryFilter(const AAIController* AI, UClass* FilterClass)
-{
-	const ANavigationData* NavData = AI && AI->NavComponent ? AI->NavComponent->GetNavData() : NULL;
-	return GetQueryFilter(NavData, FilterClass);
-}
-
 TSharedPtr<const struct FNavigationQueryFilter> UNavigationQueryFilter::GetQueryFilter(const ANavigationData* NavData, UClass* FilterClass)
 {
 	UNavigationQueryFilter* DefFilterOb = FilterClass ? FilterClass->GetDefaultObject<UNavigationQueryFilter>() : NULL;
diff --git a/Engine/Source/Runtime/Engine/Private/AI/Navigation/NavFilters/RecastFilter_UseDefaultArea.cpp b/Engine/Source/Runtime/Engine/Private/AI/Navigation/NavFilters/RecastFilter_UseDefaultArea.cpp
--- a/Engine/Source/Runtime/Engine/Private/AI/Navigation/NavFilters/RecastFilter_UseDefaultArea.cpp
+++ b/Engine/Source/Runtime/Engine/Private/AI/Navigation/NavFilters/RecastFilter_UseDefaultArea.cpp
@@ -1,6 +1,7 @@
 // Copyright 1998-2014 Epic Games, Inc. All Rights Reserved.
 
 #include "EnginePrivate.h"
+#include "AI/Navigation/NavFilters/RecastFilter_UseDefaultArea.h"
 
 URecastFilter_UseDefaultArea::URecastFilter_UseDefaultArea(const class FPostConstructInitializeProperties& PCIP) : Super(PCIP)
 {
diff --git a/Engine/Source/Runtime/Engine/Private/AI/Navigation/NavLinkProxy.cpp b/Engine/Source/Runtime/Engine/Private/AI/Navigation/NavLinkProxy.cpp
--- a/Engine/Source/Runtime/Engine/Private/AI/Navigation/NavLinkProxy.cpp
+++ b/Engine/Source/Runtime/Engine/Private/AI/Navigation/NavLinkProxy.cpp
@@ -2,7 +2,10 @@
 
 #include "EnginePrivate.h"
 #include "ObjectEditorUtils.h"
-#include "AI/NavigationModifier.h"
+#include "NavigationModifier.h"
+#include "Navigation/PathFollowingComponent.h"
+#include "AI/Navigation/SmartNavLinkComponent.h"
+#include "AI/Navigation/NavLinkProxy.h"
 
 ANavLinkProxy::ANavLinkProxy(const class FPostConstructInitializeProperties& PCIP) : Super(PCIP)
 {
diff --git a/Engine/Source/Runtime/Engine/Private/AI/Navigation/NavLinkRenderingComponent.cpp b/Engine/Source/Runtime/Engine/Private/AI/Navigation/NavLinkRenderingComponent.cpp
--- a/Engine/Source/Runtime/Engine/Private/AI/Navigation/NavLinkRenderingComponent.cpp
+++ b/Engine/Source/Runtime/Engine/Private/AI/Navigation/NavLinkRenderingComponent.cpp
@@ -1,7 +1,8 @@
 // Copyright 1998-2014 Epic Games, Inc. All Rights Reserved.
 
 #include "EnginePrivate.h"
-#include "../../Public/AI/NavLinkRenderingProxy.h"
+#include "AI/NavLinkRenderingProxy.h"
+#include "AI/Navigation/NavLinkRenderingComponent.h"
 
 //----------------------------------------------------------------------//
 // UNavLinkRenderingComponent
diff --git a/Engine/Source/Runtime/Engine/Private/AI/Navigation/NavModifierVolume.cpp b/Engine/Source/Runtime/Engine/Private/AI/Navigation/NavModifierVolume.cpp
--- a/Engine/Source/Runtime/Engine/Private/AI/Navigation/NavModifierVolume.cpp
+++ b/Engine/Source/Runtime/Engine/Private/AI/Navigation/NavModifierVolume.cpp
@@ -2,6 +2,7 @@
 
 #include "EnginePrivate.h"
 #include "AI/NavigationModifier.h"
+#include "AI/Navigation/NavModifierVolume.h"
 
 //----------------------------------------------------------------------//
 // ANavModifierVolume
diff --git a/Engine/Source/Runtime/Engine/Private/AI/Navigation/NavRelevantComponent.cpp b/Engine/Source/Runtime/Engine/Private/AI/Navigation/NavRelevantComponent.cpp
--- a/Engine/Source/Runtime/Engine/Private/AI/Navigation/NavRelevantComponent.cpp
+++ b/Engine/Source/Runtime/Engine/Private/AI/Navigation/NavRelevantComponent.cpp
@@ -1,6 +1,7 @@
 // Copyright 1998-2014 Epic Games, Inc. All Rights Reserved.
 
 #include "EnginePrivate.h"
+#include "AI/Navigation/NavRelevantComponent.h"
 
 UNavRelevantComponent::UNavRelevantComponent(const class FPostConstructInitializeProperties& PCIP) : Super(PCIP)
 {
diff --git a/Engine/Source/Runtime/Engine/Private/AI/Navigation/NavTestRenderingComponent.cpp b/Engine/Source/Runtime/Engine/Private/AI/Navigation/NavTestRenderingComponent.cpp
--- a/Engine/Source/Runtime/Engine/Private/AI/Navigation/NavTestRenderingComponent.cpp
+++ b/Engine/Source/Runtime/Engine/Private/AI/Navigation/NavTestRenderingComponent.cpp
@@ -2,6 +2,7 @@
 
 #include "EnginePrivate.h"
 #include "DebugRenderSceneProxy.h"
+#include "AI/Navigation/NavTestRenderingComponent.h"
 
 static const FColor NavMeshRenderColor_OpenSet(255,128,0,255);
 static const FColor NavMeshRenderColor_ClosedSet(255,196,0,255);
diff --git a/Engine/Source/Runtime/Engine/Private/AI/Navigation/NavigationData.cpp b/Engine/Source/Runtime/Engine/Private/AI/Navigation/NavigationData.cpp
--- a/Engine/Source/Runtime/Engine/Private/AI/Navigation/NavigationData.cpp
+++ b/Engine/Source/Runtime/Engine/Private/AI/Navigation/NavigationData.cpp
@@ -1,7 +1,9 @@
 // Copyright 1998-2014 Epic Games, Inc. All Rights Reserved.
 
 #include "EnginePrivate.h"
-#include "../../Public/AI/NavDataGenerator.h"
+#include "AI/NavDataGenerator.h"
+#include "AI/Navigation/NavAreas/NavAreaMeta.h"
+#include "AI/Navigation/NavigationData.h"
 
 //----------------------------------------------------------------------//
 // FNavDataGenerator
diff --git a/Engine/Source/Runtime/Engine/Private/AI/Navigation/NavigationGraph.cpp b/Engine/Source/Runtime/Engine/Private/AI/Navigation/NavigationGraph.cpp
--- a/Engine/Source/Runtime/Engine/Private/AI/Navigation/NavigationGraph.cpp
+++ b/Engine/Source/Runtime/Engine/Private/AI/Navigation/NavigationGraph.cpp
@@ -1,7 +1,11 @@
 // Copyright 1998-2014 Epic Games, Inc. All Rights Reserved.
 #include "EnginePrivate.h"
-#include "../../Public/AI/NavDataGenerator.h"
 #include "NavGraphGenerator.h"
+#include "AI/NavDataGenerator.h"
+#include "AI/Navigation/NavNodeInterface.h"
+#include "AI/Navigation/NavigationGraphNodeComponent.h"
+#include "AI/Navigation/NavigationGraphNode.h"
+#include "AI/Navigation/NavigationGraph.h"
 
 //----------------------------------------------------------------------//
 // FNavGraphNode
diff --git a/Engine/Source/Runtime/Engine/Private/AI/Navigation/NavigationInterfaces.cpp b/Engine/Source/Runtime/Engine/Private/AI/Navigation/NavigationInterfaces.cpp
--- a/Engine/Source/Runtime/Engine/Private/AI/Navigation/NavigationInterfaces.cpp
+++ b/Engine/Source/Runtime/Engine/Private/AI/Navigation/NavigationInterfaces.cpp
@@ -1,17 +1,18 @@
 // Copyright 1998-2014 Epic Games, Inc. All Rights Reserved.
 
 #include "EnginePrivate.h"
+#include "AI/Navigation/NavAgentInterface.h"
+#include "AI/Navigation/NavigationPathGenerator.h"
+#include "AI/Navigation/NavNodeInterface.h"
+#include "AI/Navigation/NavLinkHostInterface.h"
+#include "AI/Navigation/NavPathObserverInterface.h"
+#include "AI/Navigation/NavRelevantActorInterface.h"
 
 UNavAgentInterface::UNavAgentInterface(const class FPostConstructInitializeProperties& PCIP)
 	: Super(PCIP)
 {
 }
 
-UCrowdAgentInterface::UCrowdAgentInterface(const class FPostConstructInitializeProperties& PCIP)
-	: Super(PCIP)
-{
-}
-
 UNavigationPathGenerator::UNavigationPathGenerator(const class FPostConstructInitializeProperties& PCIP)
 	: Super(PCIP)
 {
@@ -22,7 +23,6 @@ UNavNodeInterface::UNavNodeInterface(const class FPostConstructInitializePropert
 {
 }
 
-
 UNavLinkHostInterface::UNavLinkHostInterface(const class FPostConstructInitializeProperties& PCIP)
 	: Super(PCIP)
 {
diff --git a/Engine/Source/Runtime/Engine/Private/AI/Navigation/NavigationModifier.cpp b/Engine/Source/Runtime/Engine/Private/AI/Navigation/NavigationModifier.cpp
--- a/Engine/Source/Runtime/Engine/Private/AI/Navigation/NavigationModifier.cpp
+++ b/Engine/Source/Runtime/Engine/Private/AI/Navigation/NavigationModifier.cpp
@@ -1,13 +1,16 @@
 // Copyright 1998-2014 Epic Games, Inc. All Rights Reserved.
 
 #include "EnginePrivate.h"
-#include "AI/NavigationModifier.h"
-#include "AI/NavigationOctree.h"
+#include "NavigationModifier.h"
+#include "NavigationOctree.h"
 #include "RecastHelpers.h"
 #include "ConvexHull2d.h"
 #if WITH_RECAST
 #include "PImplRecastNavMesh.h"
 #endif // WITH_RECAST
+#include "AI/Navigation/NavLinkDefinition.h"
+#include "AI/Navigation/NavLinkTrivial.h"
+#include "AI/Navigation/NavAreas/NavAreaMeta.h"
 
 // if square distance between two points is less than this the those points
 // will be considered identical when calculating convex hull
diff --git a/Engine/Source/Runtime/Engine/Private/AI/Navigation/NavigationOctree.cpp b/Engine/Source/Runtime/Engine/Private/AI/Navigation/NavigationOctree.cpp
--- a/Engine/Source/Runtime/Engine/Private/AI/Navigation/NavigationOctree.cpp
+++ b/Engine/Source/Runtime/Engine/Private/AI/Navigation/NavigationOctree.cpp
@@ -1,8 +1,9 @@
 // Copyright 1998-2014 Epic Games, Inc. All Rights Reserved.
 
 #include "EnginePrivate.h"
-#include "AI/NavigationOctree.h"
+#include "NavigationOctree.h"
 #include "RecastHelpers.h"
+#include "AI/Navigation/NavRelevantComponent.h"
 
 #if NAVOCTREE_CONTAINS_COLLISION_DATA
 #include "RecastNavMeshGenerator.h"
diff --git a/Engine/Source/Runtime/Engine/Private/AI/Navigation/NavigationPath.cpp b/Engine/Source/Runtime/Engine/Private/AI/Navigation/NavigationPath.cpp
--- a/Engine/Source/Runtime/Engine/Private/AI/Navigation/NavigationPath.cpp
+++ b/Engine/Source/Runtime/Engine/Private/AI/Navigation/NavigationPath.cpp
@@ -2,6 +2,8 @@
 
 #include "EnginePrivate.h"
 #include "VisualLog.h"
+#include "AI/Navigation/NavigationTypes.h"
+#include "AI/Navigation/RecastNavMesh.h"
 
 #define DEBUG_DRAW_OFFSET 0
 #define PATH_OFFSET_KEEP_VISIBLE_POINTS 1
diff --git a/Engine/Source/Runtime/Engine/Private/AI/Navigation/NavigationProxy.cpp b/Engine/Source/Runtime/Engine/Private/AI/Navigation/NavigationProxy.cpp
--- a/Engine/Source/Runtime/Engine/Private/AI/Navigation/NavigationProxy.cpp
+++ b/Engine/Source/Runtime/Engine/Private/AI/Navigation/NavigationProxy.cpp
@@ -1,6 +1,7 @@
 // Copyright 1998-2014 Epic Games, Inc. All Rights Reserved.
 
 #include "EnginePrivate.h"
+#include "AI/Navigation/NavigationProxy.h"
 
 UNavigationProxy::UNavigationProxy(const class FPostConstructInitializeProperties& PCIP) : Super(PCIP)
 {
diff --git a/Engine/Source/Runtime/Engine/Private/AI/Navigation/NavigationSystem.cpp b/Engine/Source/Runtime/Engine/Private/AI/Navigation/NavigationSystem.cpp
--- a/Engine/Source/Runtime/Engine/Private/AI/Navigation/NavigationSystem.cpp
+++ b/Engine/Source/Runtime/Engine/Private/AI/Navigation/NavigationSystem.cpp
@@ -1,8 +1,8 @@
 // Copyright 1998-2014 Epic Games, Inc. All Rights Reserved.
 
 #include "EnginePrivate.h"
-#include "AI/NavDataGenerator.h"
-#include "AI/NavigationOctree.h"
+#include "NavDataGenerator.h"
+#include "NavigationOctree.h"
 #if WITH_RECAST
 #include "RecastNavMeshGenerator.h"
 #endif // WITH_RECAST
@@ -12,6 +12,15 @@
 #include "Editor/GeometryMode/Public/EditorGeometry.h"
 #endif
 
+// @todo this is here only due to circular dependency to AIModule. To be removed
+#include "Navigation/CrowdManager.h"
+#include "Navigation/PathFollowingComponent.h"
+#include "Navigation/NavigationComponent.h"
+#include "AI/Navigation/NavAreas/NavArea_Null.h"
+#include "AI/Navigation/NavAreas/NavArea_Default.h"
+#include "AI/Navigation/SmartNavLinkComponent.h"
+#include "AI/Navigation/NavigationSystem.h"
+
 static const uint32 INITIAL_ASYNC_QUERIES_SIZE = 32;
 static const uint32 REGISTRATION_QUEUE_SIZE = 16;	// and we'll not reallocate
 #if WITH_RECAST
@@ -514,7 +523,25 @@ void UNavigationSystem::OnWorldInitDone(FNavigationSystem::EMode Mode)
 
 void UNavigationSystem::CreateCrowdManager()
 {
-	CrowdManager = NewObject<UCrowdManager>(this);
+	SetCrowdManager(NewObject<UCrowdManager>(this));
+}
+
+void UNavigationSystem::SetCrowdManager(UCrowdManager* NewCrowdManager)
+{
+	if (NewCrowdManager == CrowdManager.Get())
+	{
+		return;
+	}
+
+	if (CrowdManager.IsValid())
+	{
+		CrowdManager->RemoveFromRoot();
+	}
+	CrowdManager = NewCrowdManager;
+	if (NewCrowdManager != NULL)
+	{
+		CrowdManager->AddToRoot();
+	}
 }
 
 void UNavigationSystem::Tick(float DeltaSeconds)
@@ -571,7 +598,7 @@ void UNavigationSystem::Tick(float DeltaSeconds)
 		AsyncPathFindingQueries.Reset();
 	}
 
-	if (CrowdManager)
+	if (CrowdManager.IsValid())
 	{
 		CrowdManager->Tick(DeltaSeconds);
 	}
@@ -593,6 +620,10 @@ void UNavigationSystem::AddReferencedObjects(UObject* InThis, FReferenceCollecto
 	{
 		Collector.AddReferencedObject(PendingNavAreaRegistration[PendingAreaIndex], InThis);
 	}
+	
+	UNavigationSystem* This = CastChecked<UNavigationSystem>(InThis);
+	UCrowdManager* CrowdManager = This->GetCrowdManager();
+	Collector.AddReferencedObject(CrowdManager, InThis);
 }
 
 #if WITH_EDITOR
@@ -925,7 +956,7 @@ void UNavigationSystem::SimpleMoveToLocation(AController* Controller, const FVec
 	}
 }
 
-bool UNavigationSystem::NavigationRaycast(UObject* WorldContextObject, const FVector& RayStart, const FVector& RayEnd, FVector& HitLocation, TSubclassOf<class UNavigationQueryFilter> FilterClass, AAIController* Querier)
+bool UNavigationSystem::NavigationRaycast(UObject* WorldContextObject, const FVector& RayStart, const FVector& RayEnd, FVector& HitLocation, TSubclassOf<class UNavigationQueryFilter> FilterClass, AController* Querier)
 {
 	UWorld* World = NULL;
 
@@ -2697,7 +2728,9 @@ void UNavigationSystem::CleanUp()
 			NavDataSet[NavDataIndex]->CleanUp();
 		}
 	}	
-	
+
+	SetCrowdManager(NULL);
+
 	NavDataSet.Reset();
 }
 
diff --git a/Engine/Source/Runtime/Engine/Private/AI/Navigation/NavigationSystemHelpers.cpp b/Engine/Source/Runtime/Engine/Private/AI/Navigation/NavigationSystemHelpers.cpp
--- a/Engine/Source/Runtime/Engine/Private/AI/Navigation/NavigationSystemHelpers.cpp
+++ b/Engine/Source/Runtime/Engine/Private/AI/Navigation/NavigationSystemHelpers.cpp
@@ -1,11 +1,12 @@
 // Copyright 1998-2014 Epic Games, Inc. All Rights Reserved.
 
 #include "EnginePrivate.h"
-#include "AI/NavDataGenerator.h"
-#include "AI/NavigationOctree.h"
+#include "NavDataGenerator.h"
+#include "NavigationOctree.h"
 #if WITH_RECAST
 #include "RecastNavMeshGenerator.h"
 #endif // WITH_RECAST
+#include "AI/NavigationSystemHelpers.h"
 
 namespace NavigationHelper
 {
@@ -38,7 +39,7 @@ namespace NavigationHelper
 	{
 		float FallDownHeight = 0.f;
 
-		UE_VLOG_SEGMENT(Querier, FallStart, FallStart + FVector(0,0,-FallLimit)
+		UE_VLOG_SEGMENT(Querier, LogNavigation, Log, FallStart, FallStart + FVector(0, 0, -FallLimit)
 				, FColor::Red, TEXT("TerrainTrace"));
 
 		FCollisionQueryParams TraceParams(NAME_None, true, Querier);
@@ -46,7 +47,7 @@ namespace NavigationHelper
 		const bool bHit = Querier->GetWorld()->LineTraceSingle(Hit, FallStart, FallStart+FVector(0,0,-FallLimit), TraceParams, FCollisionObjectQueryParams(ECC_WorldStatic));
 		if( bHit )
 		{
-			UE_VLOG_LOCATION(Querier, Hit.Location, 15, FColor::Red, TEXT("%s")
+			UE_VLOG_LOCATION(Querier, LogNavigation, Log, Hit.Location, 15, FColor::Red, TEXT("%s")
 				, Hit.Actor.IsValid() ? *Hit.Actor->GetName() : TEXT("NULL"));
 
 			if (Cast<UStaticMeshComponent>(Hit.Component.Get()))
@@ -86,7 +87,7 @@ namespace NavigationHelper
 					if (FallDownHeight > 0.f)
 					{
 						// @todo maybe it's a good idea to clear ModifiedLink.MaxFallDownLength here
-						UE_VLOG_SEGMENT(Actor, WorldRight, WorldRight + FVector(0,0,-FallDownHeight)
+						UE_VLOG_SEGMENT(Actor, LogNavigation, Log, WorldRight, WorldRight + FVector(0, 0, -FallDownHeight)
 							, FColor::Green, TEXT("FallDownHeight %d"), LinkIndex);
 
 						Link.Right.Z -= FallDownHeight;
@@ -127,7 +128,7 @@ namespace NavigationHelper
 					if (FallDownHeightStart > 0.f)
 					{
 						// @todo maybe it's a good idea to clear ModifiedLink.MaxFallDownLength here
-						UE_VLOG_SEGMENT(Actor, WorldRightStart, WorldRightStart + FVector(0,0,-FallDownHeightStart)
+						UE_VLOG_SEGMENT(Actor, LogNavigation, Log, WorldRightStart, WorldRightStart + FVector(0, 0, -FallDownHeightStart)
 							, FColor::Green, TEXT("FallDownHeightStart %d"), LinkIndex);
 
 						Link.RightStart.Z -= FallDownHeightStart;
@@ -135,7 +136,7 @@ namespace NavigationHelper
 					if (FallDownHeightEnd > 0.f)
 					{
 						// @todo maybe it's a good idea to clear ModifiedLink.MaxFallDownLength here
-						UE_VLOG_SEGMENT(Actor, WorldRightEnd, WorldRightEnd + FVector(0,0,-FallDownHeightEnd)
+						UE_VLOG_SEGMENT(Actor, LogNavigation, Log, WorldRightEnd, WorldRightEnd + FVector(0, 0, -FallDownHeightEnd)
 							, FColor::Green, TEXT("FallDownHeightEnd %d"), LinkIndex);
 
 						Link.RightEnd.Z -= FallDownHeightEnd;
diff --git a/Engine/Source/Runtime/Engine/Private/AI/Navigation/NavigationTestingActor.cpp b/Engine/Source/Runtime/Engine/Private/AI/Navigation/NavigationTestingActor.cpp
--- a/Engine/Source/Runtime/Engine/Private/AI/Navigation/NavigationTestingActor.cpp
+++ b/Engine/Source/Runtime/Engine/Private/AI/Navigation/NavigationTestingActor.cpp
@@ -4,6 +4,7 @@
 #if WITH_EDITOR
 #include "ObjectEditorUtils.h"
 #endif
+#include "AI/Navigation/NavigationTestingActor.h"
 
 void FNavTestTickHelper::Tick(float DeltaTime)
 {
diff --git a/Engine/Source/Runtime/Engine/Private/AI/Navigation/PImplRecastNavMesh.cpp b/Engine/Source/Runtime/Engine/Private/AI/Navigation/PImplRecastNavMesh.cpp
--- a/Engine/Source/Runtime/Engine/Private/AI/Navigation/PImplRecastNavMesh.cpp
+++ b/Engine/Source/Runtime/Engine/Private/AI/Navigation/PImplRecastNavMesh.cpp
@@ -18,6 +18,8 @@
 	#include "RecastNavMeshGenerator.h"
 #endif // WITH_NAVIGATION_GENERATOR
 
+#include "AI/Navigation/SmartNavLinkComponent.h"
+
 //----------------------------------------------------------------------//
 // bunch of compile-time checks to assure types used by Recast and our
 // mid-layer are the same size
@@ -908,9 +910,9 @@ bool FPImplRecastNavMesh::InitPathfinding(const FVector& UnrealStart, const FVec
 	if (StartPoly == INVALID_NAVNODEREF)
 	{
 		UE_VLOG(NavMeshOwner, LogNavigation, Warning, TEXT("FPImplRecastNavMesh::InitPathfinding start point not on navmesh"));
-		UE_VLOG_SEGMENT(NavMeshOwner, UnrealStart, UnrealEnd, FColor::Red, TEXT("Failed path"));
-		UE_VLOG_LOCATION(NavMeshOwner, UnrealStart, 15, FColor::Red, TEXT("Start failed"));
-		UE_VLOG_BOX(NavMeshOwner, FBox(UnrealStart - NavExtent, UnrealStart + NavExtent), FColor::Red, TEXT_EMPTY);
+		UE_VLOG_SEGMENT(NavMeshOwner, LogNavigation, Warning, UnrealStart, UnrealEnd, FColor::Red, TEXT("Failed path"));
+		UE_VLOG_LOCATION(NavMeshOwner, LogNavigation, Warning, UnrealStart, 15, FColor::Red, TEXT("Start failed"));
+		UE_VLOG_BOX(NavMeshOwner, LogNavigation, Warning, FBox(UnrealStart - NavExtent, UnrealStart + NavExtent), FColor::Red, TEXT_EMPTY);
 
 		return false;
 	}
@@ -920,9 +922,9 @@ bool FPImplRecastNavMesh::InitPathfinding(const FVector& UnrealStart, const FVec
 	if (EndPoly == INVALID_NAVNODEREF)
 	{
 		UE_VLOG(NavMeshOwner, LogNavigation, Warning, TEXT("FPImplRecastNavMesh::InitPathfinding end point not on navmesh"));
-		UE_VLOG_SEGMENT(NavMeshOwner, UnrealEnd, UnrealEnd, FColor::Red, TEXT("Failed path"));
-		UE_VLOG_LOCATION(NavMeshOwner, UnrealEnd, 15, FColor::Red, TEXT("End failed"));
-		UE_VLOG_BOX(NavMeshOwner, FBox(UnrealEnd - NavExtent, UnrealEnd + NavExtent), FColor::Red, TEXT_EMPTY);
+		UE_VLOG_SEGMENT(NavMeshOwner, LogNavigation, Warning, UnrealEnd, UnrealEnd, FColor::Red, TEXT("Failed path"));
+		UE_VLOG_LOCATION(NavMeshOwner, LogNavigation, Warning, UnrealEnd, 15, FColor::Red, TEXT("End failed"));
+		UE_VLOG_BOX(NavMeshOwner, LogNavigation, Warning, FBox(UnrealEnd - NavExtent, UnrealEnd + NavExtent), FColor::Red, TEXT_EMPTY);
 
 		return false;
 	}
diff --git a/Engine/Source/Runtime/Engine/Private/AI/Navigation/PImplRecastNavMesh.h b/Engine/Source/Runtime/Engine/Private/AI/Navigation/PImplRecastNavMesh.h
--- a/Engine/Source/Runtime/Engine/Private/AI/Navigation/PImplRecastNavMesh.h
+++ b/Engine/Source/Runtime/Engine/Private/AI/Navigation/PImplRecastNavMesh.h
@@ -39,7 +39,7 @@ public:
 	const dtQueryFilter* GetAsDetourQueryFilter() const { return this; }
 };
 
-struct FRecastSpeciaLinkFilter : public dtQuerySpecialLinkFilter
+struct ENGINE_API FRecastSpeciaLinkFilter : public dtQuerySpecialLinkFilter
 {
 	FRecastSpeciaLinkFilter(UNavigationSystem* NavSystem, const UObject* Owner) : NavSys(NavSystem), SearchOwner(Owner) {}
 	virtual bool isLinkAllowed(const int32 UserId) const OVERRIDE;
@@ -58,7 +58,7 @@ namespace EClusterPath
 }
 
 /** Engine Private! - Private Implementation details of ARecastNavMesh */
-class FPImplRecastNavMesh
+class ENGINE_API FPImplRecastNavMesh
 {
 public:
 
diff --git a/Engine/Source/Runtime/Engine/Private/AI/Navigation/RecastHelpers.cpp b/Engine/Source/Runtime/Engine/Private/AI/Navigation/RecastHelpers.cpp
--- a/Engine/Source/Runtime/Engine/Private/AI/Navigation/RecastHelpers.cpp
+++ b/Engine/Source/Runtime/Engine/Private/AI/Navigation/RecastHelpers.cpp
@@ -1,6 +1,7 @@
 // Copyright 1998-2014 Epic Games, Inc. All Rights Reserved.
 
 #include "EnginePrivate.h"
+#include "AI/Navigation/RecastHelpers.h"
 
 FVector Unreal2RecastPoint(const FVector& UnrealPoint)
 {
diff --git a/Engine/Source/Runtime/Engine/Private/AI/Navigation/RecastHelpers.h b/Engine/Source/Runtime/Engine/Private/AI/Navigation/RecastHelpers.h
--- a/Engine/Source/Runtime/Engine/Private/AI/Navigation/RecastHelpers.h
+++ b/Engine/Source/Runtime/Engine/Private/AI/Navigation/RecastHelpers.h
@@ -10,12 +10,12 @@
 
 #pragma once
 
-FVector Unreal2RecastPoint(const float* UnrealPoint);
-FVector Unreal2RecastPoint(const FVector& UnrealPoint);
-FBox Unreal2RecastBox(const FBox& UnrealBox);
-FMatrix Unreal2RecastMatrix();
+ENGINE_API FVector Unreal2RecastPoint(const float* UnrealPoint);
+ENGINE_API FVector Unreal2RecastPoint(const FVector& UnrealPoint);
+ENGINE_API FBox Unreal2RecastBox(const FBox& UnrealBox);
+ENGINE_API FMatrix Unreal2RecastMatrix();
 
-FVector Recast2UnrealPoint(const float* RecastPoint);
-FVector Recast2UnrealPoint(const FVector& RecastPoint);
-FBox Recast2UnrealBox(const float* RecastMin, const float* RecastMax);
-FBox Recast2UnrealBox(const FBox& RecastBox);
+ENGINE_API FVector Recast2UnrealPoint(const float* RecastPoint);
+ENGINE_API FVector Recast2UnrealPoint(const FVector& RecastPoint);
+ENGINE_API FBox Recast2UnrealBox(const float* RecastMin, const float* RecastMax);
+ENGINE_API FBox Recast2UnrealBox(const FBox& RecastBox);
diff --git a/Engine/Source/Runtime/Engine/Private/AI/Navigation/RecastNavMesh.cpp b/Engine/Source/Runtime/Engine/Private/AI/Navigation/RecastNavMesh.cpp
--- a/Engine/Source/Runtime/Engine/Private/AI/Navigation/RecastNavMesh.cpp
+++ b/Engine/Source/Runtime/Engine/Private/AI/Navigation/RecastNavMesh.cpp
@@ -2,6 +2,10 @@
 
 #include "EnginePrivate.h"
 #include "RecastHelpers.h"
+#include "AI/Navigation/NavAreas/NavArea_Null.h"
+#include "AI/Navigation/NavAreas/NavArea_Default.h"
+#include "AI/Navigation/SmartNavLinkComponent.h"
+#include "AI/Navigation/RecastNavMesh.h"
 
 #if WITH_EDITOR
 #include "ObjectEditorUtils.h"
diff --git a/Engine/Source/Runtime/Engine/Private/AI/Navigation/SmartNavLinkComponent.cpp b/Engine/Source/Runtime/Engine/Private/AI/Navigation/SmartNavLinkComponent.cpp
--- a/Engine/Source/Runtime/Engine/Private/AI/Navigation/SmartNavLinkComponent.cpp
+++ b/Engine/Source/Runtime/Engine/Private/AI/Navigation/SmartNavLinkComponent.cpp
@@ -2,6 +2,12 @@
 
 #include "EnginePrivate.h"
 #include "VisualLog.h"
+#include "AI/Navigation/NavAreas/NavArea_Null.h"
+#include "AI/Navigation/NavAreas/NavArea_Default.h"
+#include "AI/Navigation/SmartNavLinkComponent.h"
+// @todo to be addressed when removing AIModule circular dependency
+#include "Navigation/PathFollowingComponent.h"
+#include "Navigation/NavigationComponent.h"
 
 USmartNavLinkComponent::USmartNavLinkComponent(const class FPostConstructInitializeProperties& PCIP) : Super(PCIP)
 {
diff --git a/Engine/Source/Runtime/Engine/Private/ActorReplication.cpp b/Engine/Source/Runtime/Engine/Private/ActorReplication.cpp
--- a/Engine/Source/Runtime/Engine/Private/ActorReplication.cpp
+++ b/Engine/Source/Runtime/Engine/Private/ActorReplication.cpp
@@ -4,7 +4,7 @@
 #include "Net/UnrealNetwork.h"
 #include "Engine/ActorChannel.h"
 #include "GameFramework/GameNetworkManager.h"
-#include "NavigationPathBuilder.h"
+#include "NetworkingDistanceConstants.h"
 
 /*-----------------------------------------------------------------------------
 	AActor networking implementation.
diff --git a/Engine/Source/Runtime/Engine/Private/Animation/AnimInstance.cpp b/Engine/Source/Runtime/Engine/Private/Animation/AnimInstance.cpp
--- a/Engine/Source/Runtime/Engine/Private/Animation/AnimInstance.cpp
+++ b/Engine/Source/Runtime/Engine/Private/Animation/AnimInstance.cpp
@@ -12,6 +12,12 @@
 
 #include "MessageLog.h"
 
+// @todo this is here only due to circular dependency to AIModule. To be removed
+#include "AIController.h"
+#include "Navigation/PathFollowingComponent.h"
+#include "Navigation/NavigationComponent.h"
+#include "BrainComponent.h"
+
 /** Anim stats */
 
 
diff --git a/Engine/Source/Runtime/Engine/Private/CharacterMovementComponent.cpp b/Engine/Source/Runtime/Engine/Private/CharacterMovementComponent.cpp
--- a/Engine/Source/Runtime/Engine/Private/CharacterMovementComponent.cpp
+++ b/Engine/Source/Runtime/Engine/Private/CharacterMovementComponent.cpp
@@ -7,6 +7,8 @@
 
 #include "EnginePrivate.h"
 #include "GameFramework/GameNetworkManager.h"
+// @todo this is here only due to circular dependency to AIModule. To be removed
+#include "Navigation/PathFollowingComponent.h"
 
 DEFINE_LOG_CATEGORY_STATIC(LogCharacterMovement, Log, All);
 
diff --git a/Engine/Source/Runtime/Engine/Private/CheatManager.cpp b/Engine/Source/Runtime/Engine/Private/CheatManager.cpp
--- a/Engine/Source/Runtime/Engine/Private/CheatManager.cpp
+++ b/Engine/Source/Runtime/Engine/Private/CheatManager.cpp
@@ -4,8 +4,12 @@
 
 #include "Slate.h"
 #include "SlateReflector.h"
-#include "AI/NavDataGenerator.h"
+#include "NavDataGenerator.h"
 #include "OnlineSubsystemUtils.h"
+// @todo this is here only due to circular dependency to AIModule. To be removed
+#include "AIController.h"
+#include "EnvironmentQuery/EnvQuery.h"
+#include "EnvironmentQuery/EnvQueryManager.h"
 
 #if WITH_EDITOR
 #include "UnrealEd.h"
diff --git a/Engine/Source/Runtime/Engine/Private/Controller.cpp b/Engine/Source/Runtime/Engine/Private/Controller.cpp
--- a/Engine/Source/Runtime/Engine/Private/Controller.cpp
+++ b/Engine/Source/Runtime/Engine/Private/Controller.cpp
@@ -7,7 +7,11 @@
 #include "EnginePrivate.h"
 #include "Net/UnrealNetwork.h"
 #include "ConfigCacheIni.h"
-#include "NavigationPathBuilder.h"
+#include "NetworkingDistanceConstants.h"
+
+// @todo this is here only due to circular dependency to AIModule. To be removed
+#include "Navigation/PathFollowingComponent.h"
+#include "Navigation/NavigationComponent.h"
 
 DEFINE_LOG_CATEGORY(LogPath);
 
diff --git a/Engine/Source/Runtime/Engine/Private/Debug/GameplayDebuggingComponent.cpp b/Engine/Source/Runtime/Engine/Private/Debug/GameplayDebuggingComponent.cpp
--- a/Engine/Source/Runtime/Engine/Private/Debug/GameplayDebuggingComponent.cpp
+++ b/Engine/Source/Runtime/Engine/Private/Debug/GameplayDebuggingComponent.cpp
@@ -2,9 +2,17 @@
 
 #include "EnginePrivate.h"
 #include "Net/UnrealNetwork.h"
-#include "AI/BehaviorTreeDelegates.h"
+#include "BehaviorTreeDelegates.h"
 #include "../AI/Navigation/RecastNavMeshGenerator.h"
 #include "../NavMeshRenderingHelpers.h"
+#include "../../Classes/AI/Navigation/NavigationSystem.h"
+
+// @todo this is here only due to circular dependency to AIModule. To be removed
+#include "AIController.h"
+#include "BrainComponent.h"
+#include "Navigation/NavigationComponent.h"
+#include "Navigation/PathFollowingComponent.h"
+#include "EnvironmentQuery/EnvQueryTypes.h"
 
 #if WITH_EDITOR
 #include "UnrealEd.h"
@@ -108,9 +116,11 @@ const uint32 UGameplayDebuggingComponent::ShowFlagIndex = uint32(FEngineShowFlag
 
 UGameplayDebuggingComponent::UGameplayDebuggingComponent(const class FPostConstructInitializeProperties& PCIP) 
 	: Super(PCIP)
+#if WITH_EQS
 	, EQSTimestamp(0.f)
 	, bDrawEQSLabels(true)
 	, bDrawEQSFailedItems(true)
+#endif // WITH_EQS
 {
 	PrimaryComponentTick.bCanEverTick = true;
 	bWantsInitializeComponent = true;
@@ -127,7 +137,9 @@ UGameplayDebuggingComponent::UGameplayDebuggingComponent(const class FPostConstr
 	bHiddenInGame = false;
 	bEnabledTargetSelection = false;
 
+#if WITH_EQS
 	bEnableClientEQSSceneProxy = false;
+#endif // WITH_EQS
 	NextTargrtSelectionTime = 0;
 	ReplicateViewDataCounters.Init(0, EAIDebugDrawDataView::MAX);
 }
@@ -179,10 +191,12 @@ void UGameplayDebuggingComponent::GetLifetimeReplicatedProps( TArray< FLifetimeP
 
 	DOREPLIFETIME( UGameplayDebuggingComponent, TargetActor );
 
+#if WITH_EQS
 	DOREPLIFETIME(UGameplayDebuggingComponent, EQSTimestamp);
 	DOREPLIFETIME(UGameplayDebuggingComponent, EQSName);
 	DOREPLIFETIME(UGameplayDebuggingComponent, EQSId);
 	DOREPLIFETIME(UGameplayDebuggingComponent, EQSRepData);
+#endif // WITH_EQS
 }
 
 bool UGameplayDebuggingComponent::ServerEnableTargetSelection_Validate(bool bEnable)
@@ -329,10 +343,12 @@ void UGameplayDebuggingComponent::CollectDataToReplicate()
 			CollectBehaviorTreeData();
 		}
 		
+#if WITH_EQS
 		if (ShouldReplicateData(EAIDebugDrawDataView::EQS))
 		{
 			CollectEQSData();
 		}
+#endif // WITH_EQS
 	}
 }
 
@@ -387,36 +403,10 @@ void UGameplayDebuggingComponent::CollectPathData()
 			if (!CurrentPath.HasSameObject(NewPath.Get()))
 			{
 				PathPoints.Reset();
-				
-#if WITH_RECAST
-				ARecastNavMesh* NavMesh = NewPath.IsValid() ? Cast<ARecastNavMesh>(NewPath->GetOwner()) : NULL;
-				FNavMeshPath* NavMeshPath = NewPath->CastPath<FNavMeshPath>();
-				if (NavMeshPath && NavMesh && !NavMeshPath->IsStringPulled() && NavMeshPath->PathPoints.Num() == 2)
-				{				
-					PathPoints.Add(NavMeshPath->PathPoints[0].Location);
-					NavMesh->BeginBatchQuery();
-
-					for (int32 Idx = 1; Idx < NavMeshPath->PathCorridor.Num() - 1; Idx++)
-					{
-						FVector CenterPt;
-						NavMesh->GetPolyCenter(NavMeshPath->PathCorridor[Idx], CenterPt);
-						PathPoints.Add(CenterPt);
-					}
-
-					NavMesh->FinishBatchQuery();
-					PathPoints.Add(NavMeshPath->PathPoints[1].Location);
-				}
-				else
+				for (int32 Index=0; Index < NewPath->PathPoints.Num(); ++Index)
 				{
-#endif // WITH_RECAST
-					for (int32 Index=0; Index < NewPath->PathPoints.Num(); ++Index)
-					{
-						PathPoints.Add( NewPath->PathPoints[Index].Location );
-					}
-#if WITH_RECAST
+					PathPoints.Add( NewPath->PathPoints[Index].Location );
 				}
-#endif // WITH_RECAST
-
 				CurrentPath = NewPath;
 			}
 		}
@@ -488,11 +478,13 @@ void UGameplayDebuggingComponent::SelectForDebugging(bool bNewStatus)
 		}
 #endif
 
+#if WITH_EQS
 		if (bNewStatus == false)
 		{
 			CachedQueryInstance.Reset();
 			MarkRenderStateDirty();
 		}
+#endif // WITH_EQS
 #endif
 	}
 }
@@ -506,7 +498,9 @@ void UGameplayDebuggingComponent::EnableDebugDraw(bool bEnable, bool InFocusedCo
 	else
 	{
 		SelectForDebugging(false);
+#if WITH_EQS
 		EnableClientEQSSceneProxy(false);
+#endif // WITH_EQS
 	}
 }
 
@@ -584,6 +578,7 @@ void UGameplayDebuggingComponent::GetKeyboardDesc(TArray<FDebugCategoryView>& Ca
 	Categories.Add(FDebugCategoryView(EAIDebugDrawDataView::Perception, "Perception"));		// Num4
 }
 
+#if WITH_EQS
 //////////////////////////////////////////////////////////////////////////
 // EQS Data
 //////////////////////////////////////////////////////////////////////////
@@ -779,6 +774,7 @@ void UGameplayDebuggingComponent::CollectEQSData()
 	}
 #endif
 }
+#endif // WITH_EQS
 
 //----------------------------------------------------------------------//
 // NavMesh rendering
@@ -1205,11 +1201,13 @@ FPrimitiveSceneProxy* UGameplayDebuggingComponent::CreateSceneProxy()
 	}
 #endif
 
+#if WITH_EQS
 	if (ShouldReplicateData(EAIDebugDrawDataView::EQS) && IsClientEQSSceneProxyEnabled() )
 	{
 		CompositeProxy = CompositeProxy ? CompositeProxy : (new FDebugRenderSceneCompositeProxy(this));
 		CompositeProxy->AddChild(new FEQSSceneProxy(this, TEXT("GameplayDebug"), /*bDrawOnlyWhenSelected=*/false, EQSLocalData.SolidSpheres, EQSLocalData.Texts));
 	}
+#endif // WITH_EQS
 
 	return CompositeProxy;
 }
@@ -1226,10 +1224,12 @@ FBoxSphereBounds UGameplayDebuggingComponent::CalcBounds(const FTransform & Loca
 	}
 #endif
 
+#if WITH_EQS
 	if (EQSRepData.Num())
 	{
 		MyBounds = FBox(FVector(-HALF_WORLD_MAX, -HALF_WORLD_MAX, -HALF_WORLD_MAX), FVector(HALF_WORLD_MAX, HALF_WORLD_MAX, HALF_WORLD_MAX));
 	}
+#endif // WITH_EQS
 
 	return MyBounds;
 }
diff --git a/Engine/Source/Runtime/Engine/Private/Debug/GameplayDebuggingControllerComponent.cpp b/Engine/Source/Runtime/Engine/Private/Debug/GameplayDebuggingControllerComponent.cpp
--- a/Engine/Source/Runtime/Engine/Private/Debug/GameplayDebuggingControllerComponent.cpp
+++ b/Engine/Source/Runtime/Engine/Private/Debug/GameplayDebuggingControllerComponent.cpp
@@ -2,7 +2,7 @@
 
 #include "EnginePrivate.h"
 #include "Net/UnrealNetwork.h"
-#include "AI/BehaviorTreeDelegates.h"
+#include "BehaviorTreeDelegates.h"
 
 #if WITH_EDITOR
 #include "UnrealEd.h"
@@ -186,10 +186,12 @@ void UGameplayDebuggingControllerComponent::ToggleActiveView(EAIDebugDrawDataVie
 	if (DebugAITargetActor)
 	{
 		MyPC->ServerReplicateMessageToAIDebugView(DebugAITargetActor, (ActiveViews & (1 << NewView)) ? EDebugComponentMessage::ActivateDataView : EDebugComponentMessage::DeactivateDataView, NewView);
+#if WITH_EQS
 		if (DebugAITargetActor->GetDebugComponent() && NewView == EAIDebugDrawDataView::EQS)
 		{
 			DebugAITargetActor->GetDebugComponent()->EnableClientEQSSceneProxy(IsViewActive(EAIDebugDrawDataView::EQS));
 		}
+#endif // WITH_EQS
 	}
 }
 
diff --git a/Engine/Source/Runtime/Engine/Private/Debug/GameplayDebuggingHUDComponent.cpp b/Engine/Source/Runtime/Engine/Private/Debug/GameplayDebuggingHUDComponent.cpp
--- a/Engine/Source/Runtime/Engine/Private/Debug/GameplayDebuggingHUDComponent.cpp
+++ b/Engine/Source/Runtime/Engine/Private/Debug/GameplayDebuggingHUDComponent.cpp
@@ -5,6 +5,7 @@
 =============================================================================*/
 
 #include "EnginePrivate.h"
+#include "Debug/GameplayDebuggingComponent.h"
 #include "Net/UnrealNetwork.h"
 
 DEFINE_LOG_CATEGORY_STATIC(LogHUD, Log, All);
@@ -303,7 +304,7 @@ void AGameplayDebuggingHUDComponent::DrawBehaviorTreeData(APlayerController* PC,
 
 void AGameplayDebuggingHUDComponent::DrawEQSData(APlayerController* PC, class UGameplayDebuggingComponent *DebugComponent)
 {
-#if !(UE_BUILD_SHIPPING || UE_BUILD_TEST)
+#if !(UE_BUILD_SHIPPING || UE_BUILD_TEST) && WITH_EQS
 	PrintString(DefaultContext, TEXT("\n{green}EQS\n"));
 	if (DebugComponent->EQSRepData.Num())
 	{
@@ -379,7 +380,7 @@ void AGameplayDebuggingHUDComponent::DrawEQSData(APlayerController* PC, class UG
 
 void AGameplayDebuggingHUDComponent::DrawEQSItemDetails(int32 ItemIdx, class UGameplayDebuggingComponent *DebugComponent)
 {
-#if USE_EQS_DEBUGGER
+#if USE_EQS_DEBUGGER && WITH_EQS
 	const float PosY = DefaultContext.CursorY + 1.0f;
 	float PosX = DefaultContext.CursorX;
 
diff --git a/Engine/Source/Runtime/Engine/Private/InstancedStaticMesh.cpp b/Engine/Source/Runtime/Engine/Private/InstancedStaticMesh.cpp
--- a/Engine/Source/Runtime/Engine/Private/InstancedStaticMesh.cpp
+++ b/Engine/Source/Runtime/Engine/Private/InstancedStaticMesh.cpp
@@ -7,7 +7,26 @@
 #include "EnginePrivate.h"
 #include "ShaderParameters.h"
 
-#include "Components/InstancedStaticMeshComponent.h"
+#include "Components/InteractiveFoliageComponent.h"
+#include "Components/SplineMeshComponent.h"
+#include "Components/ModelComponent.h"
+#include "Components/NiagaraComponent.h"
+#include "Components/ShapeComponent.h"
+#include "Components/BoxComponent.h"
+#include "Components/CapsuleComponent.h"
+#include "Components/SphereComponent.h"
+#include "Components/DrawSphereComponent.h"
+#include "Components/TextRenderComponent.h"
+#include "Components/VectorFieldComponent.h"
+#include "PhysicsEngine/RadialForceComponent.h"
+#include "GameFramework/SpringArmComponent.h"
+#include "Components/WindDirectionalSourceComponent.h"
+#include "Components/TimelineComponent.h"
+#include "Slate.h"
+#include "NavDataGenerator.h"
+#include "OnlineSubsystemUtils.h"
+#include "AI/Navigation/RecastHelpers.h"
+
 #include "StaticMeshLight.h"
 #include "SpeedTreeWind.h"
 #include "ComponentInstanceDataCache.h"
@@ -1459,7 +1478,6 @@ void UInstancedStaticMeshComponent::ApplyWorldOffset(const FVector& InOffset, bo
 	MarkRenderStateDirty();
 }
 
-#include "AI/Navigation/RecastHelpers.h"
 bool UInstancedStaticMeshComponent::DoCustomNavigableGeometryExport(struct FNavigableGeometryExport* GeomExport) const
 {
 	if (StaticMesh != NULL)
diff --git a/Engine/Source/Runtime/Engine/Private/Landscape/LandscapeSplines.cpp b/Engine/Source/Runtime/Engine/Private/Landscape/LandscapeSplines.cpp
--- a/Engine/Source/Runtime/Engine/Private/Landscape/LandscapeSplines.cpp
+++ b/Engine/Source/Runtime/Engine/Private/Landscape/LandscapeSplines.cpp
@@ -5,6 +5,7 @@
 =============================================================================*/
 
 #include "EnginePrivate.h"
+#include "AI/Navigation/RecastHelpers.h"
 #include "Landscape/LandscapeSplineProxies.h"
 #include "Landscape/LandscapeSplinesComponent.h"
 #include "Landscape/LandscapeSplineControlPoint.h"
@@ -2568,7 +2569,6 @@ UBodySetup* USplineMeshComponent::GetBodySetup()
 	return NULL;
 }
 
-#include "../AI/Navigation/RecastHelpers.h"
 bool USplineMeshComponent::DoCustomNavigableGeometryExport(struct FNavigableGeometryExport* GeomExport) const
 {
 	// the NavCollision is supposed to be faster than exporting the regular collision,
diff --git a/Engine/Source/Runtime/Engine/Private/NetworkDriver.cpp b/Engine/Source/Runtime/Engine/Private/NetworkDriver.cpp
--- a/Engine/Source/Runtime/Engine/Private/NetworkDriver.cpp
+++ b/Engine/Source/Runtime/Engine/Private/NetworkDriver.cpp
@@ -10,8 +10,8 @@
 #include "Engine/ActorChannel.h"
 #include "Engine/VoiceChannel.h"
 #include "GameFramework/GameNetworkManager.h"
-#include "NavigationPathBuilder.h"
 #include "OnlineSubsystemUtils.h"
+#include "NetworkingDistanceConstants.h"
 
 // Default net driver stats
 DEFINE_STAT(STAT_Ping);
diff --git a/Engine/Source/Runtime/Engine/Private/Pawn.cpp b/Engine/Source/Runtime/Engine/Private/Pawn.cpp
--- a/Engine/Source/Runtime/Engine/Private/Pawn.cpp
+++ b/Engine/Source/Runtime/Engine/Private/Pawn.cpp
@@ -9,9 +9,9 @@
 #include "EnginePrivate.h"
 #include "Net/UnrealNetwork.h"
 #include "ConfigCacheIni.h"
-#include "NavigationPathBuilder.h"
 #include "ParticleDefinitions.h"
 #include "DisplayDebugHelpers.h"
+#include "NetworkingDistanceConstants.h"
 
 DEFINE_LOG_CATEGORY(LogDamage);
 DEFINE_LOG_CATEGORY_STATIC(LogPawn, Warning, All);
@@ -26,7 +26,7 @@ APawn::APawn(const class FPostConstructInitializeProperties& PCIP)
 {
 	PrimaryActorTick.bCanEverTick = true;
 	PrimaryActorTick.TickGroup = TG_PrePhysics;
-	AIControllerClass = AAIController::StaticClass();
+	AIControllerClass = NULL;
 #if !(UE_BUILD_SHIPPING || UE_BUILD_TEST)
 	GameplayDebuggingComponentClass = UGameplayDebuggingComponent::StaticClass();
 #endif //!(UE_BUILD_SHIPPING || UE_BUILD_TEST)
diff --git a/Engine/Source/Runtime/Engine/Private/PlayerState.cpp b/Engine/Source/Runtime/Engine/Private/PlayerState.cpp
--- a/Engine/Source/Runtime/Engine/Private/PlayerState.cpp
+++ b/Engine/Source/Runtime/Engine/Private/PlayerState.cpp
@@ -7,6 +7,8 @@
 #include "EnginePrivate.h"
 #include "Net/UnrealNetwork.h"
 #include "OnlineSubsystemUtils.h"
+// @todo this is here only due to circular dependency to AIModule. To be removed
+#include "AIController.h"
 
 APlayerState::APlayerState(const class FPostConstructInitializeProperties& PCIP)
 	: Super(PCIP
diff --git a/Engine/Source/Runtime/Engine/Private/Replication.cpp b/Engine/Source/Runtime/Engine/Private/Replication.cpp
--- a/Engine/Source/Runtime/Engine/Private/Replication.cpp
+++ b/Engine/Source/Runtime/Engine/Private/Replication.cpp
@@ -8,7 +8,6 @@
 #include "Matinee/MatineeInterface.h"
 #include "Matinee/MatineeActor.h"
 #include "Net/UnrealNetwork.h"
-#include "NavigationPathBuilder.h"
 #include "SoundDefinitions.h"
 
 void AMatineeActor::GetLifetimeReplicatedProps( TArray< FLifetimeProperty > & OutLifetimeProps ) const
diff --git a/Engine/Source/Runtime/Engine/Private/UnrealEngine.cpp b/Engine/Source/Runtime/Engine/Private/UnrealEngine.cpp
--- a/Engine/Source/Runtime/Engine/Private/UnrealEngine.cpp
+++ b/Engine/Source/Runtime/Engine/Private/UnrealEngine.cpp
@@ -49,6 +49,10 @@
 #include "Particles/ParticleSystem.h"
 #include "Particles/ParticleSystemComponent.h"
 
+// @todo this is here only due to circular dependency to AIModule. To be removed
+#include "BehaviorTree/BehaviorTreeManager.h"
+#include "EnvironmentQuery/EnvQueryManager.h"
+
 #if !UE_BUILD_SHIPPING
 #include "STaskGraph.h"
 #endif
@@ -1245,16 +1249,6 @@ void UEngine::InitializeObjectReferences()
 		NavigationSystemClass = LoadClass<UNavigationSystem>(NULL, *NavigationSystemClassName.ClassName, NULL, LOAD_None, NULL);
 	}
 
-	if ( BehaviorTreeManagerClass == NULL )
-	{
-		BehaviorTreeManagerClass = LoadClass<UBehaviorTreeManager>(NULL, *BehaviorTreeManagerClassName.ClassName, NULL, LOAD_None, NULL);
-	}
-
-	if ( EnvironmentQueryManagerClass == NULL )
-	{
-		EnvironmentQueryManagerClass = LoadClass<UEnvQueryManager>(NULL, *EnvironmentQueryManagerClassName.ClassName, NULL, LOAD_None, NULL);
-	}
-
 	if ( AvoidanceManagerClass == NULL )
 	{
 		AvoidanceManagerClass = LoadClass<UAvoidanceManager>(NULL, *AvoidanceManagerClassName.ClassName, NULL, LOAD_None, NULL);
diff --git a/Engine/Source/Runtime/Engine/Private/World.cpp b/Engine/Source/Runtime/Engine/Private/World.cpp
--- a/Engine/Source/Runtime/Engine/Private/World.cpp
+++ b/Engine/Source/Runtime/Engine/Private/World.cpp
@@ -31,6 +31,10 @@
 #include "Landscape/LandscapeInfo.h"
 #include "Particles/ParticleEventManager.h"
 
+// @todo this is here only due to circular dependency to AIModule. To be removed
+#include "BehaviorTree/BehaviorTreeManager.h"
+#include "EnvironmentQuery/EnvQueryManager.h"
+
 DEFINE_LOG_CATEGORY_STATIC(LogWorld, Log, All);
 DEFINE_LOG_CATEGORY(LogSpawn);
 
@@ -175,8 +179,10 @@ void UWorld::AddReferencedObjects(UObject* InThis, FReferenceCollector& Collecto
 		Collector.AddReferencedObject( This->GameState, This );
 		Collector.AddReferencedObject( This->AuthorityGameMode, This );
 		Collector.AddReferencedObject( This->NetworkManager, This );
-		Collector.AddReferencedObject( This->BehaviorTreeManager, This );
-		Collector.AddReferencedObject( This->EnvironmentQueryManager, This );
+		UBehaviorTreeManager* BTManager = This->GetBehaviorTreeManager();
+		Collector.AddReferencedObject( BTManager, This );
+		UEnvQueryManager* EQSManager = This->GetEnvironmentQueryManager();
+		Collector.AddReferencedObject( EQSManager, This );
 		Collector.AddReferencedObject( This->NavigationSystem, This );
 		Collector.AddReferencedObject( This->AvoidanceManager, This );
 
@@ -664,15 +670,11 @@ void UWorld::InitWorld(const InitializationValues IVS)
 		UNavigationSystem::CreateNavigationSystem(this);
 	}
 
-	if (GEngine->BehaviorTreeManagerClass != NULL)
-	{
-		BehaviorTreeManager = NewObject<UBehaviorTreeManager>(this, GEngine->BehaviorTreeManagerClass);
-	}
-
-	if (GEngine->EnvironmentQueryManagerClass != NULL)
-	{
-		EnvironmentQueryManager = NewObject<UEnvQueryManager>(this, GEngine->EnvironmentQueryManagerClass);
-	}
+	BehaviorTreeManager = NewObject<UBehaviorTreeManager>(this);
+	// @hackyness
+	BehaviorTreeManager->AddToRoot();
+	EnvironmentQueryManager = NewObject<UEnvQueryManager>(this);
+	EnvironmentQueryManager->AddToRoot();
 
 	if (GEngine->AvoidanceManagerClass != NULL)
 	{
@@ -2748,9 +2750,22 @@ void UWorld::CleanupWorld(bool bSessionEnded, bool bCleanupResources, UWorld* Ne
 
 	FWorldDelegates::OnWorldCleanup.Broadcast(this, bSessionEnded, bCleanupResources);
 
-	if (NavigationSystem != NULL && bCleanupResources == true)
+	if (bCleanupResources == true)
 	{
-		NavigationSystem->CleanUp();
+		if (NavigationSystem != NULL)
+		{
+			NavigationSystem->CleanUp();
+		}
+
+		if (BehaviorTreeManager.Get())
+		{
+			BehaviorTreeManager->RemoveFromRoot();
+		}
+
+		if (EnvironmentQueryManager.Get())
+		{
+			EnvironmentQueryManager->RemoveFromRoot();
+		}
 	}
 
 	NetworkActors.Empty();
diff --git a/Engine/Source/Runtime/Engine/Public/AI/NavigationOctree.h b/Engine/Source/Runtime/Engine/Public/AI/NavigationOctree.h
--- a/Engine/Source/Runtime/Engine/Public/AI/NavigationOctree.h
+++ b/Engine/Source/Runtime/Engine/Public/AI/NavigationOctree.h
@@ -3,6 +3,7 @@
 #pragma once
 
 #include "GenericOctree.h"
+#include "AI/Navigation/NavigationProxy.h"
 
 struct ENGINE_API FNavigationOctreeFilter
 {
diff --git a/Engine/Source/Runtime/Engine/Public/DebugRenderSceneProxy.h b/Engine/Source/Runtime/Engine/Public/DebugRenderSceneProxy.h
--- a/Engine/Source/Runtime/Engine/Public/DebugRenderSceneProxy.h
+++ b/Engine/Source/Runtime/Engine/Public/DebugRenderSceneProxy.h
@@ -9,13 +9,15 @@
 #ifndef _INC_DEBUGRENDERSCENEPROXY
 #define _INC_DEBUGRENDERSCENEPROXY
 
+#include "Debug/DebugDrawService.h"
+
 class FDebugRenderSceneProxy : public FPrimitiveSceneProxy
 {
 public:
 
 	ENGINE_API FDebugRenderSceneProxy(const UPrimitiveComponent* InComponent);
 	// FPrimitiveSceneProxy interface.
-	  
+
 	/** 
 	 * Draw the scene proxy as a dynamic element
 	 *
diff --git a/Engine/Source/Runtime/Engine/Public/Engine.h b/Engine/Source/Runtime/Engine/Public/Engine.h
--- a/Engine/Source/Runtime/Engine/Public/Engine.h
+++ b/Engine/Source/Runtime/Engine/Public/Engine.h
@@ -44,16 +44,12 @@
 #include "Engine/TriggerVolume.h"
 #include "Camera/CameraActor.h"
 #include "AI/Navigation/NavAgentInterface.h"
-#include "AI/Navigation/CrowdAgentInterface.h"
 #include "GameFramework/Controller.h"
-#include "AI/AITypes.h"
 #include "AI/Navigation/NavLinkDefinition.h"
 #include "AI/Navigation/NavigationTypes.h"
 #include "Kismet/BlueprintFunctionLibrary.h"
 #include "AI/Navigation/NavigationSystem.h"
-#include "AI/AIResourceInterface.h"
-#include "AI/Navigation/PathFollowingComponent.h"
-#include "GameFramework/AIController.h"
+#include "AI/Navigation/NavCollision.h"
 #include "GameFramework/PlayerMuteList.h"
 #include "Camera/PlayerCameraManager.h"
 #include "GameFramework/Info.h"
@@ -110,9 +106,6 @@
 #include "Interfaces/NetworkPredictionInterface.h"
 #include "GameFramework/CharacterMovementComponent.h"
 #include "GameFramework/Character.h"
-#include "AI/EnvironmentQuery/EQSQueryResultSourceInterface.h"
-#include "AI/EnvironmentQuery/EnvQueryTypes.h"
-#include "AI/EnvironmentQuery/EQSTestingPawn.h"
 #include "GameFramework/FloatingPawnMovement.h"
 #include "GameFramework/DefaultPawn.h"
 #include "GameFramework/SpectatorPawn.h"
@@ -143,13 +136,7 @@
 #include "Engine/TriggerSphere.h"
 #include "VectorField/VectorFieldVolume.h"
 #include "Components/ApplicationLifecycleComponent.h"
-#include "AI/BehaviorTree/BehaviorTreeTypes.h"
-#include "AI/BehaviorTree/Blackboard/BlackboardKeyType.h"
 #include "Engine/DataAsset.h"
-#include "AI/BehaviorTree/BlackboardData.h"
-#include "AI/BehaviorTree/BlackboardComponent.h"
-#include "AI/BrainComponent.h"
-#include "AI/BehaviorTree/BehaviorTreeComponent.h"
 #include "Debug/GameplayDebuggingControllerComponent.h"
 #include "GameFramework/SpectatorPawnMovement.h"
 #include "Vehicles/VehicleWheel.h"
@@ -157,14 +144,7 @@
 #include "Vehicles/WheeledVehicleMovementComponent4W.h"
 #include "GameFramework/ProjectileMovementComponent.h"
 #include "GameFramework/RotatingMovementComponent.h"
-#include "AI/Navigation/NavigationPathGenerator.h"
-#include "AI/Navigation/NavigationComponent.h"
-#include "AI/Navigation/NavRelevantComponent.h"
-#include "AI/Navigation/SmartNavLinkComponent.h"
-#include "AI/Navigation/CrowdManager.h"
-#include "AI/Navigation/CrowdFollowingComponent.h"
 #include "Components/PawnNoiseEmitterComponent.h"
-#include "Components/PawnSensingComponent.h"
 #include "PhysicsEngine/PhysicsHandleComponent.h"
 #include "Atmosphere/AtmosphericFogComponent.h"
 #include "Sound/SoundAttenuation.h"
@@ -185,7 +165,6 @@
 #include "Components/BrushComponent.h"
 #include "Components/DrawFrustumComponent.h"
 #include "Debug/DebugDrawService.h"
-#include "AI/EnvironmentQuery/EQSRenderingComponent.h"
 #include "Debug/GameplayDebuggingComponent.h"
 #include "Components/LineBatchComponent.h"
 #include "Components/MaterialBillboardComponent.h"
@@ -242,21 +221,6 @@
 #include "Engine/AssetUserData.h"
 #include "Sound/AudioSettings.h"
 #include "AI/Navigation/AvoidanceManager.h"
-#include "AI/BehaviorTree/BTNode.h"
-#include "AI/BehaviorTree/BTCompositeNode.h"
-#include "AI/BehaviorTree/BTTaskNode.h"
-#include "AI/BehaviorTree/BehaviorTree.h"
-#include "AI/BehaviorTree/BehaviorTreeManager.h"
-#include "AI/BehaviorTree/Blackboard/BlackboardKeyType_Bool.h"
-#include "AI/BehaviorTree/Blackboard/BlackboardKeyType_Class.h"
-#include "AI/BehaviorTree/Blackboard/BlackboardKeyType_Enum.h"
-#include "AI/BehaviorTree/Blackboard/BlackboardKeyType_Float.h"
-#include "AI/BehaviorTree/Blackboard/BlackboardKeyType_Int.h"
-#include "AI/BehaviorTree/Blackboard/BlackboardKeyType_Name.h"
-#include "AI/BehaviorTree/Blackboard/BlackboardKeyType_NativeEnum.h"
-#include "AI/BehaviorTree/Blackboard/BlackboardKeyType_Object.h"
-#include "AI/BehaviorTree/Blackboard/BlackboardKeyType_String.h"
-#include "AI/BehaviorTree/Blackboard/BlackboardKeyType_Vector.h"
 #include "Engine/BlendableInterface.h"
 #include "Engine/BlueprintCore.h"
 #include "EdGraph/EdGraphPin.h"
@@ -265,12 +229,9 @@
 #include "Engine/Blueprint.h"
 #include "Animation/AnimBlueprint.h"
 #include "Engine/LevelScriptBlueprint.h"
-#include "AI/AISystem.h"
-#include "AI/BehaviorTree/BTFunctionLibrary.h"
 #include "Sound/DialogueTypes.h"
 #include "Kismet/GameplayStatics.h"
 #include "Kismet/HeadMountedDisplayFunctionLibrary.h"
-#include "Kismet/KismetAIHelperLibrary.h"
 #include "Kismet/KismetArrayLibrary.h"
 #include "Kismet/KismetInputLibrary.h"
 #include "Kismet/KismetMaterialLibrary.h"
@@ -284,37 +245,6 @@
 #include "PhysicsEngine/BodySetup.h"
 #include "Animation/AnimData/BoneMaskFilter.h"
 #include "Engine/Breakpoint.h"
-#include "AI/BehaviorTree/BTAuxiliaryNode.h"
-#include "AI/BehaviorTree/BTDecorator.h"
-#include "AI/BehaviorTree/Decorators/BTDecorator_BlackboardBase.h"
-#include "AI/BehaviorTree/Decorators/BTDecorator_Blackboard.h"
-#include "AI/BehaviorTree/Decorators/BTDecorator_BlueprintBase.h"
-#include "AI/BehaviorTree/Decorators/BTDecorator_CompareBBEntries.h"
-#include "AI/BehaviorTree/Decorators/BTDecorator_ConeCheck.h"
-#include "AI/BehaviorTree/Decorators/BTDecorator_Cooldown.h"
-#include "AI/BehaviorTree/Decorators/BTDecorator_DoesPathExist.h"
-#include "AI/BehaviorTree/Decorators/BTDecorator_ForceSuccess.h"
-#include "AI/BehaviorTree/Decorators/BTDecorator_KeepInCone.h"
-#include "AI/BehaviorTree/Decorators/BTDecorator_Loop.h"
-#include "AI/BehaviorTree/Decorators/BTDecorator_Optional.h"
-#include "AI/BehaviorTree/Decorators/BTDecorator_ReachedMoveGoal.h"
-#include "AI/BehaviorTree/Decorators/BTDecorator_TimeLimit.h"
-#include "AI/BehaviorTree/BTService.h"
-#include "AI/BehaviorTree/Services/BTService_BlackboardBase.h"
-#include "AI/BehaviorTree/Services/BTService_DefaultFocus.h"
-#include "AI/BehaviorTree/Services/BTService_BlueprintBase.h"
-#include "AI/BehaviorTree/Composites/BTComposite_Selector.h"
-#include "AI/BehaviorTree/Composites/BTComposite_Sequence.h"
-#include "AI/BehaviorTree/Composites/BTComposite_SimpleParallel.h"
-#include "AI/BehaviorTree/Tasks/BTTask_BlackboardBase.h"
-#include "AI/BehaviorTree/Tasks/BTTask_MoveDirectlyToward.h"
-#include "AI/BehaviorTree/Tasks/BTTask_MoveTo.h"
-#include "AI/BehaviorTree/Tasks/BTTask_RunEQSQuery.h"
-#include "AI/BehaviorTree/Tasks/BTTask_BlueprintBase.h"
-#include "AI/BehaviorTree/Tasks/BTTask_MakeNoise.h"
-#include "AI/BehaviorTree/Tasks/BTTask_PlaySound.h"
-#include "AI/BehaviorTree/Tasks/BTTask_RunBehavior.h"
-#include "AI/BehaviorTree/Tasks/BTTask_Wait.h"
 #include "Slate/ButtonStyleAsset.h"
 #include "Camera/CameraAnim.h"
 #include "Camera/CameraAnimInst.h"
@@ -371,30 +301,6 @@
 #include "EdGraph/EdGraphSchema.h"
 #include "Engine/Engine.h"
 #include "Engine/GameEngine.h"
-#include "AI/EnvironmentQuery/EnvQuery.h"
-#include "AI/EnvironmentQuery/EnvQueryContext.h"
-#include "AI/EnvironmentQuery/Contexts/EnvQueryContext_BlueprintBase.h"
-#include "AI/EnvironmentQuery/Contexts/EnvQueryContext_Item.h"
-#include "AI/EnvironmentQuery/Contexts/EnvQueryContext_Querier.h"
-#include "AI/EnvironmentQuery/EnvQueryGenerator.h"
-#include "AI/EnvironmentQuery/Generators/EnvQueryGenerator_Composite.h"
-#include "AI/EnvironmentQuery/Generators/EnvQueryGenerator_PathingGrid.h"
-#include "AI/EnvironmentQuery/Generators/EnvQueryGenerator_ProjectedPoints.h"
-#include "AI/EnvironmentQuery/Generators/EnvQueryGenerator_OnCircle.h"
-#include "AI/EnvironmentQuery/Generators/EnvQueryGenerator_SimpleGrid.h"
-#include "AI/EnvironmentQuery/Items/EnvQueryItemType.h"
-#include "AI/EnvironmentQuery/Items/EnvQueryItemType_VectorBase.h"
-#include "AI/EnvironmentQuery/Items/EnvQueryItemType_ActorBase.h"
-#include "AI/EnvironmentQuery/Items/EnvQueryItemType_Actor.h"
-#include "AI/EnvironmentQuery/Items/EnvQueryItemType_Direction.h"
-#include "AI/EnvironmentQuery/Items/EnvQueryItemType_Point.h"
-#include "AI/EnvironmentQuery/EnvQueryManager.h"
-#include "AI/EnvironmentQuery/EnvQueryOption.h"
-#include "AI/EnvironmentQuery/EnvQueryTest.h"
-#include "AI/EnvironmentQuery/Tests/EnvQueryTest_Distance.h"
-#include "AI/EnvironmentQuery/Tests/EnvQueryTest_Dot.h"
-#include "AI/EnvironmentQuery/Tests/EnvQueryTest_Pathfinding.h"
-#include "AI/EnvironmentQuery/Tests/EnvQueryTest_Trace.h"
 #include "Exporters/Exporter.h"
 #include "Engine/FontImportOptions.h"
 #include "Engine/Font.h"
@@ -403,7 +309,6 @@
 #include "GameFramework/InputSettings.h"
 #include "Engine/InterpCurveEdSetup.h"
 #include "Engine/IntSerialization.h"
-#include "AI/KismetAIAsyncTaskProxy.h"
 #include "Layers/Layer.h"
 #include "Engine/LevelBase.h"
 #include "Engine/Level.h"
@@ -422,16 +327,6 @@
 #include "Materials/MaterialInstanceDynamic.h"
 #include "Materials/MaterialParameterCollection.h"
 #include "Materials/MaterialParameterCollectionInstance.h"
-#include "AI/Navigation/NavAreas/NavArea_Default.h"
-#include "AI/Navigation/NavAreas/NavArea_Null.h"
-#include "AI/Navigation/NavAreas/NavAreaMeta.h"
-#include "AI/Navigation/NavAreas/NavAreaMeta_SwitchByAgent.h"
-#include "AI/Navigation/NavCollision.h"
-#include "AI/Navigation/NavigationProxy.h"
-#include "AI/Navigation/NavFilters/NavigationQueryFilter.h"
-#include "AI/Navigation/NavFilters/RecastFilter_UseDefaultArea.h"
-#include "AI/Navigation/NavLinkTrivial.h"
-#include "AI/Navigation/NavNodeInterface.h"
 #include "Engine/NetDriver.h"
 #include "Engine/ObjectLibrary.h"
 #include "Engine/ObjectReferencer.h"
@@ -555,10 +450,10 @@
 #include "ShadowMap.h"
 #include "Model.h"							// Model class.
 
-#include "AI/NavDataGenerator.h"
-#include "AI/NavLinkRenderingProxy.h"
-#include "AI/NavigationModifier.h"
-#include "AI/NavigationOctree.h"
+#include "NavDataGenerator.h"
+#include "NavLinkRenderingProxy.h"
+#include "NavigationModifier.h"
+#include "NavigationOctree.h"
 #include "StaticMeshResources.h"			// Static meshes.
 #include "AnimTree.h"						// Animation.
 #include "SkeletalMeshTypes.h"				// Skeletal animated mesh.
@@ -576,7 +471,7 @@
 #include "UnrealExporter.h"					// Exporter definition.
 #include "TimerManager.h"					// Game play timers
 #include "EngineService.h"
-#include "AI/NavigationSystemHelpers.h"
+#include "NavigationSystemHelpers.h"
 #include "HardwareInfo.h"
 
 #include "EngineModule.h"
diff --git a/Engine/Source/Runtime/Engine/Public/VisualLog.h b/Engine/Source/Runtime/Engine/Public/VisualLog.h
--- a/Engine/Source/Runtime/Engine/Public/VisualLog.h
+++ b/Engine/Source/Runtime/Engine/Public/VisualLog.h
@@ -195,7 +195,7 @@ struct ENGINE_API FVisLogEntry
 	UE_VLOG(Actor, CategoryName, Verbosity, Format, ##__VA_ARGS__); \
 }
 
-#define UE_VLOG_SEGMENT_THICK(Actor, SegmentStart, SegmentEnd, Color, Thickness, DescriptionFormat, ...) \
+#define UE_VLOG_SEGMENT_THICK(Actor, CategoryName, Verbosity, SegmentStart, SegmentEnd, Color, Thickness, DescriptionFormat, ...) \
 { \
 	SCOPE_CYCLE_COUNTER(STAT_VisualLog); \
 	if (FVisualLog::Get()->IsRecording()) \
@@ -204,9 +204,9 @@ struct ENGINE_API FVisLogEntry
 } \
 }
 
-#define UE_VLOG_SEGMENT(Actor, SegmentStart, SegmentEnd, Color, DescriptionFormat, ...) UE_VLOG_SEGMENT_THICK(Actor, SegmentStart, SegmentEnd, Color, 0, DescriptionFormat, ##__VA_ARGS__)
+#define UE_VLOG_SEGMENT(Actor, CategoryName, Verbosity, SegmentStart, SegmentEnd, Color, DescriptionFormat, ...) UE_VLOG_SEGMENT_THICK(Actor, CategoryName, Verbosity, SegmentStart, SegmentEnd, Color, 0, DescriptionFormat, ##__VA_ARGS__)
 
-#define UE_VLOG_LOCATION(Actor, Location, Radius, Color, DescriptionFormat, ...) \
+#define UE_VLOG_LOCATION(Actor, CategoryName, Verbosity, Location, Radius, Color, DescriptionFormat, ...) \
 { \
 	SCOPE_CYCLE_COUNTER(STAT_VisualLog); \
 	if (FVisualLog::Get()->IsRecording()) \
@@ -215,7 +215,7 @@ struct ENGINE_API FVisLogEntry
 	} \
 }
 
-#define UE_VLOG_BOX(Actor, Box, Color, DescriptionFormat, ...) \
+#define UE_VLOG_BOX(Actor, CategoryName, Verbosity, Box, Color, DescriptionFormat, ...) \
 { \
 	SCOPE_CYCLE_COUNTER(STAT_VisualLog); \
 	if (FVisualLog::Get()->IsRecording()) \
@@ -318,11 +318,11 @@ private:
 };
 
 #else
-	#define UE_VLOG(Actor, CategoryName, Verbosity, Format, ...) UE_LOG(CategoryName, Verbosity, Format, ##__VA_ARGS__)
+	#define UE_VLOG(Actor, CategoryName, Verbosity, Verbosity, Format, ...) UE_LOG(CategoryName, Verbosity, Format, ##__VA_ARGS__)
 	#define UE_CVLOG(Condition, Actor, CategoryName, Verbosity, Format, ...) UE_CLOG(Condition, CategoryName, Verbosity, Format, ##__VA_ARGS__)
-	#define UE_VLOG_SEGMENT(Actor, SegmentStart, SegmentEnd, Color, DescriptionFormat, ...)
-	#define UE_VLOG_LOCATION(Actor, Location, Radius, Color, DescriptionFormat, ...)
-	#define UE_VLOG_BOX(Actor, Box, Color, DescriptionFormat, ...) 
+	#define UE_VLOG_SEGMENT(Actor, CategoryName, Verbosity, SegmentStart, SegmentEnd, Color, DescriptionFormat, ...)
+	#define UE_VLOG_LOCATION(Actor, CategoryName, Verbosity, Location, Radius, Color, DescriptionFormat, ...)
+	#define UE_VLOG_BOX(Actor, CategoryName, Verbosity, Box, Color, DescriptionFormat, ...) 
 	#define REDIRECT_TO_VLOG(Dest)
 	#define REDIRECT_ACTOR_TO_VLOG(Src, Destination) 
 #endif
