- Fix enum forward decl in NetConnection.h

diff --git a/Engine/Source/Runtime/Engine/Classes/Engine/NetConnection.h b/Engine/Source/Runtime/Engine/Classes/Engine/NetConnection.h
--- a/Engine/Source/Runtime/Engine/Classes/Engine/NetConnection.h
+++ b/Engine/Source/Runtime/Engine/Classes/Engine/NetConnection.h
@@ -6,6 +6,7 @@
 
 #pragma once
 #include "Net/DataBunch.h"
+#include "Engine/Channel.h"
 #include "NetConnection.generated.h"
 
 class FObjectReplicator;
@@ -514,7 +515,7 @@ public:
 	class UControlChannel* GetControlChannel();
 
 	/** Create a channel. */
-	ENGINE_API UChannel* CreateChannel( enum EChannelType Type, bool bOpenedLocally, int32 ChannelIndex=INDEX_NONE );
+	ENGINE_API UChannel* CreateChannel( EChannelType Type, bool bOpenedLocally, int32 ChannelIndex=INDEX_NONE );
 
 	/** Handle a packet we just received. */
 	void ReceivedPacket( FBitReader& Reader );
