Fixed a silly VisualLog macro bug #UE4

diff --git a/Engine/Source/Runtime/Engine/Public/VisualLog.h b/Engine/Source/Runtime/Engine/Public/VisualLog.h
--- a/Engine/Source/Runtime/Engine/Public/VisualLog.h
+++ b/Engine/Source/Runtime/Engine/Public/VisualLog.h
@@ -318,7 +318,7 @@ private:
 };
 
 #else
-	#define UE_VLOG(Actor, CategoryName, Verbosity, Verbosity, Format, ...) UE_LOG(CategoryName, Verbosity, Format, ##__VA_ARGS__)
+	#define UE_VLOG(Actor, CategoryName, Verbosity, Format, ...) UE_LOG(CategoryName, Verbosity, Format, ##__VA_ARGS__)
 	#define UE_CVLOG(Condition, Actor, CategoryName, Verbosity, Format, ...) UE_CLOG(Condition, CategoryName, Verbosity, Format, ##__VA_ARGS__)
 	#define UE_VLOG_SEGMENT(Actor, CategoryName, Verbosity, SegmentStart, SegmentEnd, Color, DescriptionFormat, ...)
 	#define UE_VLOG_LOCATION(Actor, CategoryName, Verbosity, Location, Radius, Color, DescriptionFormat, ...)
