- Improved GetPlatformName description

diff --git a/Engine/Source/Runtime/Engine/Classes/Kismet/GameplayStatics.h b/Engine/Source/Runtime/Engine/Classes/Kismet/GameplayStatics.h
--- a/Engine/Source/Runtime/Engine/Classes/Kismet/GameplayStatics.h
+++ b/Engine/Source/Runtime/Engine/Classes/Kismet/GameplayStatics.h
@@ -453,7 +453,9 @@ class ENGINE_API UGameplayStatics : public UBlueprintFunctionLibrary
 	UFUNCTION(BlueprintCallable, Category="Game", meta=(Enable="true"))
 	static void EnableLiveStreaming(bool Enable);
 
-	/** Returns the string name of the current platform, to perform different behavior based on platorm */
+	/**
+	 * Returns the string name of the current platform, to perform different behavior based on platform. 
+	 * (Platform names include Windows, Mac, IOS, Android, PS4, XboxOne, HTML5, Linux) */
 	UFUNCTION(BlueprintCallable, Category="Game")
 	static FString GetPlatformName();
 
