CIS compilation fixes on Main #UE4

diff --git a/Engine/Source/Runtime/AIModule/Classes/AIController.h b/Engine/Source/Runtime/AIModule/Classes/AIController.h
--- a/Engine/Source/Runtime/AIModule/Classes/AIController.h
+++ b/Engine/Source/Runtime/AIModule/Classes/AIController.h
@@ -13,6 +13,7 @@
 #include "AI/Navigation/NavigationTypes.h"
 #include "AI/Navigation/NavigationSystem.h"
 #include "Navigation/PathFollowingComponent.h"
+#include "GameFramework/Controller.h"
 #include "AIController.generated.h"
 
 DECLARE_DYNAMIC_MULTICAST_DELEGATE_TwoParams(FAIMoveCompletedSignature, FAIRequestID, RequestID, EPathFollowingResult::Type, Result);
diff --git a/Engine/Source/Runtime/AIModule/Classes/Navigation/PathFollowingComponent.h b/Engine/Source/Runtime/AIModule/Classes/Navigation/PathFollowingComponent.h
--- a/Engine/Source/Runtime/AIModule/Classes/Navigation/PathFollowingComponent.h
+++ b/Engine/Source/Runtime/AIModule/Classes/Navigation/PathFollowingComponent.h
@@ -286,7 +286,7 @@ protected:
 	TWeakObjectPtr<AActor> DestinationActor;
 
 	/** cached DestinationActor cast to INavAgentInterface. Use SetDestinationActor to set this */
-	const INavAgentInterface* DestinationAgent;
+	const class INavAgentInterface* DestinationAgent;
 
 	/** destination for current path segment */
 	FBasedPosition CurrentDestination;
diff --git a/Engine/Source/Runtime/AIModule/Private/AIController.cpp b/Engine/Source/Runtime/AIModule/Private/AIController.cpp
--- a/Engine/Source/Runtime/AIModule/Private/AIController.cpp
+++ b/Engine/Source/Runtime/AIModule/Private/AIController.cpp
@@ -10,6 +10,7 @@
 #include "BehaviorTree/BlackboardComponent.h"
 #include "Navigation/PathFollowingComponent.h"
 #include "Navigation/NavigationComponent.h"
+#include "GameFramework/Pawn.h"
 #include "AIController.h"
 
 // mz@todo these need to be removed, legacy code
diff --git a/Engine/Source/Runtime/AIModule/Private/AIModulePrivate.h b/Engine/Source/Runtime/AIModule/Private/AIModulePrivate.h
--- a/Engine/Source/Runtime/AIModule/Private/AIModulePrivate.h
+++ b/Engine/Source/Runtime/AIModule/Private/AIModulePrivate.h
@@ -9,6 +9,13 @@
 #endif
 
 #include "Core.h"
+//#include "CoreUObject.h"
+//#include "Engine/EngineBaseTypes.h"
+//#include "Engine/EngineTypes.h"
+//// @todo consider making this non-shippable
+//#include "DisplayDebugHelpers.h"
+//#include "InputCoreTypes.h"
+//#include "GameFramework/Actor.h"
 #include "Engine.h"
 
 #include "AIModule.h"
diff --git a/Engine/Source/Runtime/Engine/Classes/AI/Navigation/NavigationSystem.h b/Engine/Source/Runtime/Engine/Classes/AI/Navigation/NavigationSystem.h
--- a/Engine/Source/Runtime/Engine/Classes/AI/Navigation/NavigationSystem.h
+++ b/Engine/Source/Runtime/Engine/Classes/AI/Navigation/NavigationSystem.h
@@ -4,7 +4,9 @@
 
 #include "AI/Navigation/NavFilters/NavigationQueryFilter.h"
 #include "AI/Navigation/NavigationTypes.h"
+#include "Kismet/BlueprintFunctionLibrary.h"
 #include "GenericOctreePublic.h"
+//#include "AI/Navigation/NavigationData.h"
 #include "NavigationSystem.generated.h"
 
 #define NAVSYS_DEBUG (0 && UE_BUILD_DEBUG)
@@ -620,10 +622,10 @@ public:
 	static bool IsNavigationBeingBuilt(UObject* WorldContext);
 
 	UFUNCTION(BlueprintCallable, Category="AI|Navigation")
-	static void SimpleMoveToActor(AController* Controller, const AActor* Goal);
+	static void SimpleMoveToActor(class AController* Controller, const AActor* Goal);
 
 	UFUNCTION(BlueprintCallable, Category="AI|Navigation")
-	static void SimpleMoveToLocation(AController* Controller, const FVector& Goal);
+	static void SimpleMoveToLocation(class AController* Controller, const FVector& Goal);
 
 	/** Performs navigation raycast on NavigationData appropriate for given Querier.
 	 *	@param Querier if not passed default navigation data will be used
diff --git a/Engine/Source/Runtime/Engine/Classes/GameFramework/Controller.h b/Engine/Source/Runtime/Engine/Classes/GameFramework/Controller.h
--- a/Engine/Source/Runtime/Engine/Classes/GameFramework/Controller.h
+++ b/Engine/Source/Runtime/Engine/Classes/GameFramework/Controller.h
@@ -20,6 +20,7 @@
 //=============================================================================
 
 #pragma once
+#include "AI/Navigation/NavAgentInterface.h"
 #include "Controller.generated.h"
 
 class UNavigationComponent;
