Hopefully fix forward declared enum compile errors in particles

diff --git a/Engine/Source/Runtime/Engine/Classes/Particles/Event/ParticleModuleEventGenerator.h b/Engine/Source/Runtime/Engine/Classes/Particles/Event/ParticleModuleEventGenerator.h
--- a/Engine/Source/Runtime/Engine/Classes/Particles/Event/ParticleModuleEventGenerator.h
+++ b/Engine/Source/Runtime/Engine/Classes/Particles/Event/ParticleModuleEventGenerator.h
@@ -3,6 +3,7 @@
 
 #pragma once
 #include "Particles/Event/ParticleModuleEventBase.h"
+#include "Particles/ParticleSystemComponent.h"
 #include "ParticleModuleEventGenerator.generated.h"
 
 /**
@@ -14,7 +15,7 @@ struct FParticleEvent_GenerateInfo
 
 	/** The type of event. */
 	UPROPERTY(EditAnywhere, Category=ParticleEvent_GenerateInfo)
-	TEnumAsByte<enum EParticleEventType> Type;
+	TEnumAsByte<EParticleEventType> Type;
 
 	/** How often to trigger the event (<= 1 means EVERY time). */
 	UPROPERTY(EditAnywhere, Category=ParticleEvent_GenerateInfo)
diff --git a/Engine/Source/Runtime/Engine/Classes/Particles/Orientation/ParticleModuleOrientationAxisLock.h b/Engine/Source/Runtime/Engine/Classes/Particles/Orientation/ParticleModuleOrientationAxisLock.h
--- a/Engine/Source/Runtime/Engine/Classes/Particles/Orientation/ParticleModuleOrientationAxisLock.h
+++ b/Engine/Source/Runtime/Engine/Classes/Particles/Orientation/ParticleModuleOrientationAxisLock.h
@@ -52,7 +52,7 @@ class UParticleModuleOrientationAxisLock : public UParticleModuleOrientationBase
 	 *		EPAL_ROTATE_Z		Lock the sprite rotation on the Z-axis.
 	 */
 	UPROPERTY(EditAnywhere, Category=Orientation)
-	TEnumAsByte<enum EParticleAxisLock> LockAxisFlags;
+	TEnumAsByte<EParticleAxisLock> LockAxisFlags;
 
 
 	// Begin UObject Interface
diff --git a/Engine/Source/Runtime/Engine/Classes/Particles/ParticleModuleRequired.h b/Engine/Source/Runtime/Engine/Classes/Particles/ParticleModuleRequired.h
--- a/Engine/Source/Runtime/Engine/Classes/Particles/ParticleModuleRequired.h
+++ b/Engine/Source/Runtime/Engine/Classes/Particles/ParticleModuleRequired.h
@@ -3,6 +3,8 @@
 
 #pragma once
 #include "Particles/ParticleModule.h"
+#include "Particles/ParticleEmitter.h"
+#include "Particles/ParticleSpriteEmitter.h"
 #include "ParticleModuleRequired.generated.h"
 
 UENUM()
@@ -58,7 +60,7 @@ class UParticleModuleRequired : public UParticleModule
 	 *	PSA_TypeSpecific	- Use the alignment method indicated int he type data module.
 	 */
 	UPROPERTY(EditAnywhere, Category=Emitter)
-	TEnumAsByte<enum EParticleScreenAlignment> ScreenAlignment;
+	TEnumAsByte<EParticleScreenAlignment> ScreenAlignment;
 
 	/** If true, update the emitter in local space										*/
 	UPROPERTY(EditAnywhere, Category=Emitter)
@@ -145,7 +147,7 @@ class UParticleModuleRequired : public UParticleModule
 	//
 	/** The method to utilize when burst-emitting particles						*/
 	UPROPERTY()
-	TEnumAsByte<enum EParticleBurstMethod> ParticleBurstMethod;
+	TEnumAsByte<EParticleBurstMethod> ParticleBurstMethod;
 
 	/** The array of burst entries.												*/
 	UPROPERTY(export, noclear)
@@ -198,7 +200,7 @@ class UParticleModuleRequired : public UParticleModule
 	 *						  and the next 
 	 */
 	UPROPERTY(EditAnywhere, Category=SubUV)
-	TEnumAsByte<enum EParticleSubUVInterpMethod> InterpolationMethod;
+	TEnumAsByte<EParticleSubUVInterpMethod> InterpolationMethod;
 
 	/** The number of sub-images horizontally in the texture							*/
 	UPROPERTY(EditAnywhere, Category=SubUV)
diff --git a/Engine/Source/Runtime/Engine/Classes/Particles/Spawn/ParticleModuleSpawn.h b/Engine/Source/Runtime/Engine/Classes/Particles/Spawn/ParticleModuleSpawn.h
--- a/Engine/Source/Runtime/Engine/Classes/Particles/Spawn/ParticleModuleSpawn.h
+++ b/Engine/Source/Runtime/Engine/Classes/Particles/Spawn/ParticleModuleSpawn.h
@@ -3,6 +3,7 @@
 
 #pragma once
 #include "Particles/Spawn/ParticleModuleSpawnBase.h"
+#include "Particles/ParticleEmitter.h"
 #include "ParticleModuleSpawn.generated.h"
 
 class UParticleLODLevel;
@@ -22,11 +23,11 @@ class UParticleModuleSpawn : public UParticleModuleSpawnBase
 
 	/** The method to utilize when burst-emitting particles. */
 	UPROPERTY(EditAnywhere, Category=Burst)
-	TEnumAsByte<enum EParticleBurstMethod> ParticleBurstMethod;
+	TEnumAsByte<EParticleBurstMethod> ParticleBurstMethod;
 
 	/** The array of burst entries. */
 	UPROPERTY(EditAnywhere, export, noclear, Category=Burst)
-	TArray<struct FParticleBurst> BurstList;
+	TArray<FParticleBurst> BurstList;
 
 	/** Scale all burst entries by this amount. */
 	UPROPERTY(EditAnywhere, Category=Burst)
diff --git a/Engine/Source/Runtime/Engine/Classes/Particles/TypeData/ParticleModuleTypeDataGpu.h b/Engine/Source/Runtime/Engine/Classes/Particles/TypeData/ParticleModuleTypeDataGpu.h
--- a/Engine/Source/Runtime/Engine/Classes/Particles/TypeData/ParticleModuleTypeDataGpu.h
+++ b/Engine/Source/Runtime/Engine/Classes/Particles/TypeData/ParticleModuleTypeDataGpu.h
@@ -8,6 +8,7 @@
 
 #include "Distributions.h"
 #include "Particles/TypeData/ParticleModuleTypeDataBase.h"
+#include "Particles/Orientation/ParticleModuleOrientationAxisLock.h"
 #include "ParticleModuleTypeDataGpu.generated.h"
 
 /**
@@ -161,7 +162,7 @@ struct FGPUSpriteEmitterInfo
 
 	/** The method for locking the particles to a particular axis. */
 	UPROPERTY()
-	TEnumAsByte<enum EParticleAxisLock> LockAxisFlag;
+	TEnumAsByte<EParticleAxisLock> LockAxisFlag;
 
 	/** If true, collisions are enabled for this emitter. */
 	UPROPERTY()
diff --git a/Engine/Source/Runtime/Engine/Classes/Particles/TypeData/ParticleModuleTypeDataMesh.h b/Engine/Source/Runtime/Engine/Classes/Particles/TypeData/ParticleModuleTypeDataMesh.h
--- a/Engine/Source/Runtime/Engine/Classes/Particles/TypeData/ParticleModuleTypeDataMesh.h
+++ b/Engine/Source/Runtime/Engine/Classes/Particles/TypeData/ParticleModuleTypeDataMesh.h
@@ -3,6 +3,7 @@
 
 #pragma once
 #include "Particles/TypeData/ParticleModuleTypeDataBase.h"
+#include "Particles/Orientation/ParticleModuleOrientationAxisLock.h"
 #include "ParticleModuleTypeDataMesh.generated.h"
 
 class UParticleEmitter;
@@ -114,7 +115,7 @@ class UParticleModuleTypeDataMesh : public UParticleModuleTypeDataBase
 	 *		EPAL_ROTATE_Z	 -	Ignored for mesh emitters. Treated as EPAL_NONE.
 	 */
 	UPROPERTY(EditAnywhere, Category=Orientation)
-	TEnumAsByte<enum EParticleAxisLock> AxisLockOption;
+	TEnumAsByte<EParticleAxisLock> AxisLockOption;
 
 	/**
 	 *	If true, then point the X-axis of the mesh towards the camera.
