Remaining fixes from PR 185.

Fixes
- Incorrect NULL usage.
- Incorrect proc handle initialization (default constructor should create an invalid handle).
- Filename case not matching actual file on disk.
- Non-const string literals.
- Superfluous expression without side effect.

This is a subset of changes (compiler warning fixes needed for Linux editor) from pull request 185 by sbc100, 3dluvr, wshearn, abarbu et al (#ue4linux community): https://github.com/EpicGames/UnrealEngine/pull/185

diff --git a/Engine/Source/Editor/SwarmInterface/Private/SwarmInterfaceLocal.cpp b/Engine/Source/Editor/SwarmInterface/Private/SwarmInterfaceLocal.cpp
--- a/Engine/Source/Editor/SwarmInterface/Private/SwarmInterfaceLocal.cpp
+++ b/Engine/Source/Editor/SwarmInterface/Private/SwarmInterfaceLocal.cpp
@@ -101,7 +101,7 @@ FSwarmInterfaceLocalImpl::FSwarmInterfaceLocalImpl( void )
 #if USE_LOCAL_SWARM_INTERFACE
 ,	bIsConnected( false )
 ,	bIsEditor( false )
-,	LightmassProcHandle( NULL )
+,	LightmassProcHandle()
 #endif
 {
 }
diff --git a/Engine/Source/Editor/UnrealEd/Private/Fbx/FbxFactory.cpp b/Engine/Source/Editor/UnrealEd/Private/Fbx/FbxFactory.cpp
--- a/Engine/Source/Editor/UnrealEd/Private/Fbx/FbxFactory.cpp
+++ b/Engine/Source/Editor/UnrealEd/Private/Fbx/FbxFactory.cpp
@@ -225,7 +225,7 @@ UObject* UFbxFactory::FactoryCreateBinary
 		{
 			// Log the import message and import the mesh.
 			const TCHAR* errorMessage = FbxImporter->GetErrorMessage();
-			if (errorMessage[0] != NULL)
+			if (errorMessage[0] != '\0')
 			{
 				Warn->Log( errorMessage );
 			}
@@ -588,4 +588,4 @@ bool UFbxImportUI::CanEditChange( const UProperty* InProperty ) const
 	return bIsMutable;
 }
 
-#undef LOCTEXT_NAMESPACE
\ No newline at end of file
+#undef LOCTEXT_NAMESPACE
diff --git a/Engine/Source/Editor/UnrealEd/Private/Fbx/FbxMainExport.cpp b/Engine/Source/Editor/UnrealEd/Private/Fbx/FbxMainExport.cpp
--- a/Engine/Source/Editor/UnrealEd/Private/Fbx/FbxMainExport.cpp
+++ b/Engine/Source/Editor/UnrealEd/Private/Fbx/FbxMainExport.cpp
@@ -825,7 +825,9 @@ void FFbxExporter::ExportSkeletalMesh( AActor* Actor, USkeletalMeshComponent* Sk
 
 FbxSurfaceMaterial* FFbxExporter::CreateDefaultMaterial()
 {
-	FbxSurfaceMaterial* FbxMaterial = Scene->GetMaterial("Fbx Default Material");
+	// TODO(sbc): the below cast is needed to avoid clang warning.  The upstream
+	// signature in FBX should really use 'const char *'.
+	FbxSurfaceMaterial* FbxMaterial = Scene->GetMaterial(const_cast<char*>("Fbx Default Material"));
 	
 	if (!FbxMaterial)
 	{
@@ -2477,4 +2479,4 @@ FbxNode* FFbxExporter::ExportLandscapeToFbx(ALandscapeProxy* Landscape, const TC
 }
 
 
-} // namespace UnFbx
\ No newline at end of file
+} // namespace UnFbx
diff --git a/Engine/Source/Editor/UnrealEd/Private/MeshPaintRendering.cpp b/Engine/Source/Editor/UnrealEd/Private/MeshPaintRendering.cpp
--- a/Engine/Source/Editor/UnrealEd/Private/MeshPaintRendering.cpp
+++ b/Engine/Source/Editor/UnrealEd/Private/MeshPaintRendering.cpp
@@ -53,7 +53,7 @@ namespace MeshPaintRendering
 	};
 
 
-	IMPLEMENT_SHADER_TYPE( , TMeshPaintVertexShader, TEXT( "MeshPaintVertexShader" ), TEXT( "Main" ), SF_Vertex);
+	IMPLEMENT_SHADER_TYPE( , TMeshPaintVertexShader, TEXT( "meshpaintvertexshader" ), TEXT( "Main" ), SF_Vertex);
 
 
 
@@ -215,7 +215,7 @@ namespace MeshPaintRendering
 	};
 
 
-	IMPLEMENT_SHADER_TYPE( , TMeshPaintDilateVertexShader, TEXT( "MeshPaintDilateVertexShader" ), TEXT( "Main" ), SF_Vertex );
+	IMPLEMENT_SHADER_TYPE( , TMeshPaintDilateVertexShader, TEXT( "meshpaintdilatevertexshader" ), TEXT( "Main" ), SF_Vertex );
 
 
 
diff --git a/Engine/Source/Programs/UnrealLightmass/Private/LightmassCore/Misc/LMDebug.cpp b/Engine/Source/Programs/UnrealLightmass/Private/LightmassCore/Misc/LMDebug.cpp
--- a/Engine/Source/Programs/UnrealLightmass/Private/LightmassCore/Misc/LMDebug.cpp
+++ b/Engine/Source/Programs/UnrealLightmass/Private/LightmassCore/Misc/LMDebug.cpp
@@ -151,7 +151,7 @@ void appHandleCriticalError()
 							// Read each line, looking for the URL
 							const uint32 LineBufferSize = 1024;
 							char LineBuffer[LineBufferSize];
-							char* URLSearchText = "CrashReport url = ";
+							const char* URLSearchText = "CrashReport url = ";
 							char* URLFoundText = NULL;
 							while( fgets( LineBuffer, LineBufferSize, GeneratedAutoReportFile ) != NULL )
 							{
diff --git a/Engine/Source/Runtime/Engine/Private/Materials/MaterialInstance.cpp b/Engine/Source/Runtime/Engine/Private/Materials/MaterialInstance.cpp
--- a/Engine/Source/Runtime/Engine/Private/Materials/MaterialInstance.cpp
+++ b/Engine/Source/Runtime/Engine/Private/Materials/MaterialInstance.cpp
@@ -736,7 +736,6 @@ void UMaterialInstance::OverrideTexture( const UTexture* InTextureToOverride, UT
 			// Iterate over both the 2D textures and cube texture expressions.
 			ExpressionsByType[0] = &SourceMaterialResource->GetUniform2DTextureExpressions();
 			ExpressionsByType[1] = &SourceMaterialResource->GetUniformCubeTextureExpressions();
-			SourceMaterialResource = SourceMaterialResource;
 		}
 		else
 		{
@@ -747,7 +746,6 @@ void UMaterialInstance::OverrideTexture( const UTexture* InTextureToOverride, UT
 			// Iterate over both the 2D textures and cube texture expressions.
 			ExpressionsByType[0] = &SourceMaterialResource->GetUniform2DTextureExpressions();
 			ExpressionsByType[1] = &SourceMaterialResource->GetUniformCubeTextureExpressions();
-			SourceMaterialResource = SourceMaterialResource;
 		}
 		
 		for(int32 TypeIndex = 0;TypeIndex < ARRAY_COUNT(ExpressionsByType);TypeIndex++)
@@ -2402,4 +2400,4 @@ bool UMaterialInstance::IsPropertyActive(EMaterialProperty InProperty) const
 int32 UMaterialInstance::CompileProperty( class FMaterialCompiler* Compiler, EMaterialProperty Property, float DefaultFloat, FLinearColor DefaultColor, const FVector4& DefaultVector )
 {
 	return Parent ? Parent->CompileProperty(Compiler,Property,DefaultFloat,DefaultColor,DefaultVector) : INDEX_NONE;
-}
\ No newline at end of file
+}
