Fix one more enum compile error in particles

diff --git a/Engine/Source/Runtime/Engine/Classes/Particles/TypeData/ParticleModuleTypeDataGpu.h b/Engine/Source/Runtime/Engine/Classes/Particles/TypeData/ParticleModuleTypeDataGpu.h
--- a/Engine/Source/Runtime/Engine/Classes/Particles/TypeData/ParticleModuleTypeDataGpu.h
+++ b/Engine/Source/Runtime/Engine/Classes/Particles/TypeData/ParticleModuleTypeDataGpu.h
@@ -8,6 +8,7 @@
 
 #include "Distributions.h"
 #include "Particles/TypeData/ParticleModuleTypeDataBase.h"
+#include "Particles/ParticleSpriteEmitter.h"
 #include "Particles/Orientation/ParticleModuleOrientationAxisLock.h"
 #include "ParticleModuleTypeDataGpu.generated.h"
 
@@ -158,7 +159,7 @@ struct FGPUSpriteEmitterInfo
 
 	/** The method for aligning the particle based on the camera. */
 	UPROPERTY()
-	TEnumAsByte<enum EParticleScreenAlignment> ScreenAlignment;
+	TEnumAsByte<EParticleScreenAlignment> ScreenAlignment;
 
 	/** The method for locking the particles to a particular axis. */
 	UPROPERTY()
