Engine.h trimming (MaterialShared.h, Components.h, Texture.h, SystemSettings.h, ConvexVolume.h)

diff --git a/Engine/Source/Runtime/Engine/Public/Engine.h b/Engine/Source/Runtime/Engine/Public/Engine.h
--- a/Engine/Source/Runtime/Engine/Public/Engine.h
+++ b/Engine/Source/Runtime/Engine/Public/Engine.h
@@ -422,11 +422,7 @@
 #include "Engine/WorldComposition.h"
 
 #include "VisualLog.h"
-#include "MaterialShared.h"					// Shared material definitions.
-#include "Components.h"						// Forward declarations of object components of actors
-#include "Texture.h"						// Textures.
 #include "SystemSettings.h"					// Scalability options.
-#include "ConvexVolume.h"					// Convex volume definition.
 #include "SceneManagement.h"				// Scene management.
 #include "StaticLighting.h"					// Static lighting definitions.
 #include "LightMap.h"						// Light-maps.
