FNetGUID cleanup

diff --git a/Engine/Source/Runtime/Engine/Classes/Engine/PackageMapClient.h b/Engine/Source/Runtime/Engine/Classes/Engine/PackageMapClient.h
--- a/Engine/Source/Runtime/Engine/Classes/Engine/PackageMapClient.h
+++ b/Engine/Source/Runtime/Engine/Classes/Engine/PackageMapClient.h
@@ -89,7 +89,7 @@ protected:
 	bool	ExportNetGUID( FNetworkGUID NetGUID, const UObject * Object, FString PathName, UObject * ObjOuter );
 	void	ExportNetGUIDHeader();
 
-	void	InternalGetOrAssignNetGUID(UObject* Obj, FNetworkGUID &NetGUID);
+	FNetworkGUID	InternalGetOrAssignNetGUID( const UObject * Obj );
 	bool	InternalIsMapped(UObject* Obj, FNetworkGUID &NetGUID);
 
 	void	InternalWriteObject( FArchive& Ar, FNetworkGUID NetGUID, const UObject * Object, FString ObjectPathName, UObject* ObjectOuter );	
diff --git a/Engine/Source/Runtime/Engine/Private/PackageMapClient.cpp b/Engine/Source/Runtime/Engine/Private/PackageMapClient.cpp
--- a/Engine/Source/Runtime/Engine/Private/PackageMapClient.cpp
+++ b/Engine/Source/Runtime/Engine/Private/PackageMapClient.cpp
@@ -57,9 +57,7 @@ bool UPackageMapClient::SerializeObject( FArchive& Ar, UClass* Class, UObject*&
 		}
 
 		bSerializedUnAckedObject = false;
-		FNetworkGUID NetGUID;
-		
-		InternalGetOrAssignNetGUID(Object, NetGUID);
+		FNetworkGUID NetGUID = InternalGetOrAssignNetGUID( Object );
 		bool Mapped = InternalIsMapped(Object, NetGUID);
 
 		// Write out NetGUID to caller if necessary
@@ -320,7 +318,7 @@ bool UPackageMapClient::SerializeNewActor(FArchive& Ar, class UActorChannel *Cha
 
 				if ( !bForceSendGuid )
 				{
-					InternalGetOrAssignNetGUID( Obj, SubobjNetGUID );
+					SubobjNetGUID = InternalGetOrAssignNetGUID( Obj );
 				}
 
 				check( SubobjNetGUID.IsDynamic() );		// Make sure some assumptions are correct
@@ -439,22 +437,24 @@ bool UPackageMapClient::SerializeNewActor(FArchive& Ar, class UActorChannel *Cha
 //--------------------------------------------------------------------
 
 /** Gets or assigns a new NetGUID to this object. Returns whether the object is fully mapped or not */
-void UPackageMapClient::InternalGetOrAssignNetGUID(UObject* Object, FNetworkGUID &NetGUID)
+FNetworkGUID UPackageMapClient::InternalGetOrAssignNetGUID( const UObject * Object )
 {
 	if (!Object || !SupportsObject(Object))
 	{
 		// Null of unsupported object, leave as default NetGUID and just return mapped=true
-		return;
+		return FNetworkGUID();
 	}
 
 	// ----------------
 	// Assign NetGUID if necessary
 	// ----------------
-	NetGUID = Cache->NetGUIDLookup.FindRef(Object);
+	FNetworkGUID NetGUID = Cache->NetGUIDLookup.FindRef(Object);
 	if (!NetGUID.IsValid())
 	{
 		NetGUID = AssignNewNetGUID(Object);
 	}
+
+	return NetGUID;
 }
 
 bool UPackageMapClient::InternalIsMapped(UObject* Object, FNetworkGUID &NetGUID)
@@ -692,9 +692,7 @@ void UPackageMapClient::InternalWriteObjectPath( FArchive & Ar, FNetworkGUID Net
 	if ( SerializeOuter )
 	{
 		// Serialize reference to outer. This is basically a form of compression.
-		FNetworkGUID OuterNetGUID;
-			
-		InternalGetOrAssignNetGUID( ObjectOuter, OuterNetGUID );
+		FNetworkGUID OuterNetGUID = InternalGetOrAssignNetGUID( ObjectOuter );
 		InternalWriteObject( Ar, OuterNetGUID, ObjectOuter, TEXT( "" ), NULL );
 	}
 
@@ -1230,11 +1228,7 @@ void UPackageMapClient::HandleUnAssignedObject(const UObject* Obj)
 {
 	check( Obj != NULL );
 
-	FNetworkGUID NetGUID = Cache->NetGUIDLookup.FindRef(Obj);
-	if (!NetGUID.IsValid())
-	{
-		NetGUID = AssignNewNetGUID(Obj);
-	}
+	FNetworkGUID NetGUID = InternalGetOrAssignNetGUID( Obj );
 
 	if ( !NetGUID.IsDefault() && ShouldSendFullPath( Obj, NetGUID ) )
 	{
