Make USkeletalMesh::AnimBlueprintGeneratedClass readonly at runtime.

diff --git a/Engine/Source/Runtime/Engine/Classes/Components/SkeletalMeshComponent.h b/Engine/Source/Runtime/Engine/Classes/Components/SkeletalMeshComponent.h
--- a/Engine/Source/Runtime/Engine/Classes/Components/SkeletalMeshComponent.h
+++ b/Engine/Source/Runtime/Engine/Classes/Components/SkeletalMeshComponent.h
@@ -234,7 +234,11 @@ public:
 	class UAnimBlueprint* AnimationBlueprint_DEPRECATED;
 #endif
 
-	UPROPERTY(EditAnywhere, BlueprintReadWrite, Category=Animation)
+	/** 
+	 * The AnimBlueprint class to use
+	 * Use 'SetAnimClass' to change at runtime. 
+	 */
+	UPROPERTY(EditAnywhere, BlueprintReadOnly, Category = Animation)
 	class UAnimBlueprintGeneratedClass* AnimBlueprintGeneratedClass;
 
 	/** The active animation graph program instance */
