Add a MaxThreshold for Passive Sound Mixes

diff --git a/Engine/Config/BaseEngine.ini b/Engine/Config/BaseEngine.ini
--- a/Engine/Config/BaseEngine.ini
+++ b/Engine/Config/BaseEngine.ini
@@ -601,6 +601,7 @@ bUseStreamingPause=false
 +TaggedPropertyRedirects=(ClassName="PawnSensingComponent",OldPropertyName="bWantsSeePlayerNotify",NewPropertyName="bSeePawns")
 +TaggedPropertyRedirects=(ClassName="FloatingPawnMovement",OldPropertyName="AccelRate",NewPropertyName="Acceleration")
 +TaggedPropertyRedirects=(ClassName="FloatingPawnMovement",OldPropertyName="DecelRate",NewPropertyName="Deceleration")
++TaggedPropertyRedirects=(ClassName="PassiveSoundMixModifier",OldPropertyName="VolumeThreshold",NewPropertyName="MinVolumeThreshold")
 +TaggedPropertyRedirects=(ClassName="ProjectileMovementComponent",OldPropertyName="Speed",NewPropertyName="InitialSpeed")
 
 +TaggedPropertyRedirects=(ClassName="MaterialEditorInstanceConstant",OldPropertyName="RefractionBias",NewPropertyName="RefractionDepthBias")
diff --git a/Engine/Source/Runtime/Engine/Classes/Sound/SoundClass.h b/Engine/Source/Runtime/Engine/Classes/Sound/SoundClass.h
--- a/Engine/Source/Runtime/Engine/Classes/Sound/SoundClass.h
+++ b/Engine/Source/Runtime/Engine/Classes/Sound/SoundClass.h
@@ -138,13 +138,18 @@ struct FPassiveSoundMixModifier
 	UPROPERTY(EditAnywhere, Category=PassiveSoundMixModifier)
 	class USoundMix* SoundMix;
 
-	/** Minimum volume level to activate SoundMix at. */
+	/** Minimum volume level required to activate SoundMix. Below this value the SoundMix will not be active. */
 	UPROPERTY(EditAnywhere, Category=PassiveSoundMixModifier)
-	float VolumeThreshold;
+	float MinVolumeThreshold;
+
+	/** Maximum volume level required to activate SoundMix. Above this value the SoundMix will not be active. */
+	UPROPERTY(EditAnywhere, Category=PassiveSoundMixModifier)
+	float MaxVolumeThreshold;
 
 	FPassiveSoundMixModifier()
-	: SoundMix(NULL)
-	, VolumeThreshold(0.0f)
+		: SoundMix(NULL)
+		, MinVolumeThreshold(0.f)
+		, MaxVolumeThreshold(10.f)
 	{
 	}
 	
diff --git a/Engine/Source/Runtime/Engine/Private/AudioDevice.cpp b/Engine/Source/Runtime/Engine/Private/AudioDevice.cpp
--- a/Engine/Source/Runtime/Engine/Private/AudioDevice.cpp
+++ b/Engine/Source/Runtime/Engine/Private/AudioDevice.cpp
@@ -1115,17 +1115,19 @@ void FAudioDevice::UpdatePassiveSoundMixModifiers(TArray<FWaveInstance*>& WaveIn
 	// Find all passive SoundMixes from currently active wave instances
 	for( int32 WaveIndex = FirstActiveIndex; WaveIndex < WaveInstances.Num(); WaveIndex++ )
 	{
-		if( WaveInstances[WaveIndex] )
+		FWaveInstance* WaveInstance = WaveInstances[WaveIndex];
+		if (WaveInstance)
 		{
-			USoundClass* SoundClass = WaveInstances[WaveIndex]->SoundClass;
+			USoundClass* SoundClass = WaveInstance->SoundClass;
 			if( SoundClass ) 
 			{
+				const float WaveInstanceActualVolume = WaveInstance->GetActualVolume();
 				// Check each SoundMix individually for volume levels
-				for( int32 MixIndex = 0; MixIndex < SoundClass->PassiveSoundMixModifiers.Num(); ++MixIndex )
+				for (const FPassiveSoundMixModifier& PassiveSoundMixModifier : SoundClass->PassiveSoundMixModifiers)
 				{
-					if( WaveInstances[WaveIndex]->GetActualVolume() >= SoundClass->PassiveSoundMixModifiers[MixIndex].VolumeThreshold )
+					if (WaveInstanceActualVolume >= PassiveSoundMixModifier.MinVolumeThreshold && WaveInstanceActualVolume <= PassiveSoundMixModifier.MaxVolumeThreshold)
 					{
-						CurrPassiveSoundMixModifiers.AddUnique( SoundClass->PassiveSoundMixModifiers[MixIndex].SoundMix );
+						CurrPassiveSoundMixModifiers.AddUnique( PassiveSoundMixModifier.SoundMix );
 					}
 				}
 			}
@@ -1133,20 +1135,20 @@ void FAudioDevice::UpdatePassiveSoundMixModifiers(TArray<FWaveInstance*>& WaveIn
 	}
 
 	// Push SoundMixes that weren't previously active
-	for( int32 MixIndex = 0; MixIndex < CurrPassiveSoundMixModifiers.Num(); MixIndex++ )
+	for (USoundMix* CurrPassiveSoundMixModifier : CurrPassiveSoundMixModifiers)
 	{
-		if( PrevPassiveSoundMixModifiers.Find( CurrPassiveSoundMixModifiers[MixIndex] ) == INDEX_NONE )
+		if (PrevPassiveSoundMixModifiers.Find(CurrPassiveSoundMixModifier) == INDEX_NONE)
 		{
-			PushSoundMixModifier(CurrPassiveSoundMixModifiers[MixIndex], true);
+			PushSoundMixModifier(CurrPassiveSoundMixModifier, true);
 		}
 	}
 
 	// Pop SoundMixes that are no longer active
-	for( int32 MixIndex = 0; MixIndex < PrevPassiveSoundMixModifiers.Num(); MixIndex++ )
+	for (USoundMix* PrevPassiveSoundMixModifier : PrevPassiveSoundMixModifiers)
 	{
-		if( CurrPassiveSoundMixModifiers.Find( PrevPassiveSoundMixModifiers[MixIndex] ) == INDEX_NONE )
+		if (CurrPassiveSoundMixModifiers.Find(PrevPassiveSoundMixModifier) == INDEX_NONE)
 		{
-			PopSoundMixModifier(PrevPassiveSoundMixModifiers[MixIndex], true);
+			PopSoundMixModifier(PrevPassiveSoundMixModifier, true);
 		}
 	}
 
diff --git a/Engine/Source/Runtime/Engine/Private/SoundMix.cpp b/Engine/Source/Runtime/Engine/Private/SoundMix.cpp
--- a/Engine/Source/Runtime/Engine/Private/SoundMix.cpp
+++ b/Engine/Source/Runtime/Engine/Private/SoundMix.cpp
@@ -105,8 +105,8 @@ bool USoundMix::CheckForDependencyLoop(USoundClass* SoundClass, TArray<USoundCla
 	{
 		const FPassiveSoundMixModifier& CurrentSoundMix = SoundClass->PassiveSoundMixModifiers[Index];
 
-		// Ignore if volume threshold is 0 as the Sound Mix will not be deactivated
-		if (CurrentSoundMix.SoundMix == this && CurrentSoundMix.VolumeThreshold > 0.0f)
+		// Ignore if volume threshold is between 0 and 10 (arbitrarily large upper value) as the Sound Mix will not be deactivated
+		if (CurrentSoundMix.SoundMix == this && CurrentSoundMix.MinVolumeThreshold > 0.f && CurrentSoundMix.MaxVolumeThreshold < 10.f)
 		{
 			ProblemClasses.AddUnique(SoundClass);
 			bFoundProblemClass = true;
