Expose HitResult info for sweep related overlap event

diff --git a/Engine/Source/Runtime/Engine/Classes/Components/PrimitiveComponent.h b/Engine/Source/Runtime/Engine/Classes/Components/PrimitiveComponent.h
--- a/Engine/Source/Runtime/Engine/Classes/Components/PrimitiveComponent.h
+++ b/Engine/Source/Runtime/Engine/Classes/Components/PrimitiveComponent.h
@@ -170,9 +170,9 @@ struct FSpriteCategoryInfo
  */
 DECLARE_DYNAMIC_MULTICAST_DELEGATE_FourParams( FComponentHitSignature, class AActor*, OtherActor, class UPrimitiveComponent*, OtherComp, FVector, NormalImpulse, const FHitResult&, Hit );
 /** Delegate for notification of start of overlap of a specific component */
-DECLARE_DYNAMIC_MULTICAST_DELEGATE_ThreeParams( FComponentBeginOverlapSignature,class AActor*, OtherActor, class UPrimitiveComponent*, OtherComp, int32, OtherBodyIndex );
+DECLARE_DYNAMIC_MULTICAST_DELEGATE_FiveParams( FComponentBeginOverlapSignature,class AActor*, OtherActor, class UPrimitiveComponent*, OtherComp, int32, OtherBodyIndex, bool, bFromSweep, const FHitResult &, SweepResult);
 /** Delegate for notification of end of overlap of a specific component */
-DECLARE_DYNAMIC_MULTICAST_DELEGATE_ThreeParams( FComponentEndOverlapSignature, class AActor*, OtherActor, class UPrimitiveComponent*, OtherComp, int32, OtherBodyIndex );
+DECLARE_DYNAMIC_MULTICAST_DELEGATE_ThreeParams( FComponentEndOverlapSignature, class AActor*, OtherActor, class UPrimitiveComponent*, OtherComp, int32, OtherBodyIndex);
 
 DECLARE_DYNAMIC_MULTICAST_DELEGATE_OneParam( FComponentBeginCursorOverSignature, UPrimitiveComponent*, TouchedComponent );
 DECLARE_DYNAMIC_MULTICAST_DELEGATE_OneParam( FComponentEndCursorOverSignature, UPrimitiveComponent*, TouchedComponent );
diff --git a/Engine/Source/Runtime/Engine/Classes/Components/SceneComponent.h b/Engine/Source/Runtime/Engine/Classes/Components/SceneComponent.h
--- a/Engine/Source/Runtime/Engine/Classes/Components/SceneComponent.h
+++ b/Engine/Source/Runtime/Engine/Classes/Components/SceneComponent.h
@@ -14,16 +14,26 @@ struct FOverlapInfo
 	FOverlapInfo()
 	{}
 
+	FOverlapInfo(const FHitResult & InSweepResult)
+		: bFromSweep(true), OverlapInfo(InSweepResult)
+	{
+	}
+
 	FOverlapInfo(class UPrimitiveComponent* InComponent, int32 InBodyIndex = INDEX_NONE)
-		:Component(InComponent), BodyIndex(InBodyIndex) 
-	{}
+		: bFromSweep(false)
+	{
+		OverlapInfo.Component = InComponent;
+		OverlapInfo.Item = InBodyIndex;
+	}
+	
+	int32 GetBodyIndex() const { return OverlapInfo.Item;  }
 
-	friend bool operator == (const FOverlapInfo& LHS, const FOverlapInfo& RHS) { return LHS.Component == RHS.Component && LHS.BodyIndex == RHS.BodyIndex; }
+	//This function completely ignores SweepResult information. It seems that places that use this function do not care, but it still seems risky
+	friend bool operator == (const FOverlapInfo& LHS, const FOverlapInfo& RHS) { return LHS.OverlapInfo.Component == RHS.OverlapInfo.Component && LHS.OverlapInfo.Item == RHS.OverlapInfo.Item; }
+	bool bFromSweep;
 
-	/** Overlapping component */
-	TWeakObjectPtr<class UPrimitiveComponent> Component;
-	/** BodyIndex of the component that is overlapping */
-	int32 BodyIndex;
+	/** Information for both sweep and overlap queries. Different parts are valid depending on bFromSweep*/
+	FHitResult OverlapInfo;
 };
 
 /** Detail mode for scene component rendering. */
diff --git a/Engine/Source/Runtime/Engine/Classes/Engine/InteractiveFoliageActor.h b/Engine/Source/Runtime/Engine/Classes/Engine/InteractiveFoliageActor.h
--- a/Engine/Source/Runtime/Engine/Classes/Engine/InteractiveFoliageActor.h
+++ b/Engine/Source/Runtime/Engine/Classes/Engine/InteractiveFoliageActor.h
@@ -84,7 +84,7 @@ protected:
 
 	/** Called when capsule is touched */
 	UFUNCTION()
-	void CapsuleTouched(AActor* Other, UPrimitiveComponent* OtherComp, int32 OtherBodyIndex);
+	void CapsuleTouched(AActor* Other, UPrimitiveComponent* OtherComp, int32 OtherBodyIndex, bool bFromSweep, const FHitResult & OverlapInfo);
 public:
 	// Begin AActor interface
 	virtual void Tick(float DeltaSeconds) OVERRIDE;
diff --git a/Engine/Source/Runtime/Engine/Classes/GameFramework/CharacterMovementComponent.h b/Engine/Source/Runtime/Engine/Classes/GameFramework/CharacterMovementComponent.h
--- a/Engine/Source/Runtime/Engine/Classes/GameFramework/CharacterMovementComponent.h
+++ b/Engine/Source/Runtime/Engine/Classes/GameFramework/CharacterMovementComponent.h
@@ -1260,7 +1260,7 @@ protected:
 
 	/** Called when the collision capsule touches another primitive component */
 	UFUNCTION()
-	void CapsuleTouched(AActor* Other, UPrimitiveComponent* OtherComp, int32 OtherBodyIndex);
+	void CapsuleTouched(AActor* Other, UPrimitiveComponent* OtherComp, int32 OtherBodyIndex, bool bFromSweep, const FHitResult & SweepResult);
 
 	// Enum used to control GetPawnCapsuleExtent behavior
 	enum EShrinkCapsuleExtent
diff --git a/Engine/Source/Runtime/Engine/Private/CharacterMovementComponent.cpp b/Engine/Source/Runtime/Engine/Private/CharacterMovementComponent.cpp
--- a/Engine/Source/Runtime/Engine/Private/CharacterMovementComponent.cpp
+++ b/Engine/Source/Runtime/Engine/Private/CharacterMovementComponent.cpp
@@ -5519,7 +5519,7 @@ void UCharacterMovementComponent::ClientAckGoodMove_Implementation(float TimeSta
 	ClientData->AckMove(MoveIndex);
 }
 
-void UCharacterMovementComponent::CapsuleTouched( AActor* Other, UPrimitiveComponent* OtherComp, int32 OtherBodyIndex )
+void UCharacterMovementComponent::CapsuleTouched( AActor* Other, UPrimitiveComponent* OtherComp, int32 OtherBodyIndex, bool bFromSweep, const FHitResult & SweepResult )
 {
 	check(bEnablePhysicsInteraction);
 
@@ -5605,16 +5605,16 @@ void UCharacterMovementComponent::ApplyRepulsionForce( float DeltaTime )
 		{
 			const FOverlapInfo& Overlap = Overlaps[i];
 
-			UPrimitiveComponent* OverlapComp = Overlap.Component.Get();
+			UPrimitiveComponent* OverlapComp = Overlap.OverlapInfo.Component.Get();
 			if (!OverlapComp || OverlapComp->Mobility < EComponentMobility::Movable)
 			{ 
 				continue; 
 			}
 
 			FName BoneName = NAME_None;
-			if (Overlap.BodyIndex != INDEX_NONE && OverlapComp->IsA(USkinnedMeshComponent::StaticClass()))
+			if (Overlap.GetBodyIndex() != INDEX_NONE && OverlapComp->IsA(USkinnedMeshComponent::StaticClass()))
 			{
-				BoneName = ((USkinnedMeshComponent*)OverlapComp)->GetBoneName(Overlap.BodyIndex);
+				BoneName = ((USkinnedMeshComponent*)OverlapComp)->GetBoneName(Overlap.GetBodyIndex());
 			}
 
 			// Use the body instead of the component for cases where we have multi-body overlaps enabled
diff --git a/Engine/Source/Runtime/Engine/Private/FoliageComponent.cpp b/Engine/Source/Runtime/Engine/Private/FoliageComponent.cpp
--- a/Engine/Source/Runtime/Engine/Private/FoliageComponent.cpp
+++ b/Engine/Source/Runtime/Engine/Private/FoliageComponent.cpp
@@ -96,7 +96,7 @@ float AInteractiveFoliageActor::TakeDamage(float DamageAmount, FDamageEvent cons
 	return 0.f;
 }
 
-void AInteractiveFoliageActor::CapsuleTouched(AActor* Other, UPrimitiveComponent* OtherComp, int32 OtherBodyIndex)
+void AInteractiveFoliageActor::CapsuleTouched(AActor* Other, UPrimitiveComponent* OtherComp, int32 OtherBodyIndex, bool bFromSweep, const FHitResult & SweepResult)
 {
 	if (Other != NULL && OtherComp != NULL)
 	{
diff --git a/Engine/Source/Runtime/Engine/Private/PrimitiveComponent.cpp b/Engine/Source/Runtime/Engine/Private/PrimitiveComponent.cpp
--- a/Engine/Source/Runtime/Engine/Private/PrimitiveComponent.cpp
+++ b/Engine/Source/Runtime/Engine/Private/PrimitiveComponent.cpp
@@ -1364,7 +1364,7 @@ bool UPrimitiveComponent::MoveComponent( const FVector& Delta, const FRotator& N
 								}
 
 								// cache touches
-								PendingOverlaps.AddUnique(FOverlapInfo(TestHit.Component.Get(), TestHit.Item));
+								PendingOverlaps.AddUnique(FOverlapInfo(TestHit));
 							}
 						}
 					}
@@ -1764,7 +1764,7 @@ bool UPrimitiveComponent::IsOverlappingComponent(UPrimitiveComponent const* Othe
 {
 	for (int32 i=0; i < OverlappingComponents.Num(); ++i)
 	{
-		if (OverlappingComponents[i].Component == OtherComp) { return true; }
+		if (OverlappingComponents[i].OverlapInfo.Component.Get() == OtherComp) { return true; }
 	}
 	return false;
 }
@@ -1780,7 +1780,7 @@ bool UPrimitiveComponent::IsOverlappingActor(const AActor* Other) const
 	{
 		for (int32 OverlapIdx=0; OverlapIdx<OverlappingComponents.Num(); ++OverlapIdx)
 		{
-			UPrimitiveComponent const* const PrimComp = OverlappingComponents[OverlapIdx].Component.Get();
+			UPrimitiveComponent const* const PrimComp = OverlappingComponents[OverlapIdx].OverlapInfo.Component.Get();
 			if ( PrimComp && (PrimComp->GetOwner() == Other) )
 			{
 				return true;
@@ -1798,7 +1798,7 @@ bool UPrimitiveComponent::GetOverlapsWithActor(const AActor* Actor, TArray<FOver
 	{
 		for (int32 OverlapIdx=0; OverlapIdx<OverlappingComponents.Num(); ++OverlapIdx)
 		{
-			UPrimitiveComponent const* const PrimComp = OverlappingComponents[OverlapIdx].Component.Get();
+			UPrimitiveComponent const* const PrimComp = OverlappingComponents[OverlapIdx].OverlapInfo.Component.Get();
 			if ( PrimComp && (PrimComp->GetOwner() == Actor) )
 			{
 				OutOverlaps.Add(OverlappingComponents[OverlapIdx]);
@@ -1825,7 +1825,7 @@ void UPrimitiveComponent::BeginComponentOverlap(const FOverlapInfo& OtherOverlap
 {
 	//	UE_LOG(LogActor, Log, TEXT("BEGIN OVERLAP! Self=%s SelfComp=%s, Other=%s, OtherComp=%s"), *GetNameSafe(this), *GetNameSafe(MyComp), *GetNameSafe(OtherComp->GetOwner()), *GetNameSafe(OtherComp));
 
-	UPrimitiveComponent* OtherComp = OtherOverlap.Component.Get();
+	UPrimitiveComponent* OtherComp = OtherOverlap.OverlapInfo.Component.Get();
 
 	if (OtherComp)
 	{
@@ -1847,12 +1847,12 @@ void UPrimitiveComponent::BeginComponentOverlap(const FOverlapInfo& OtherOverlap
 				// first execute component delegates
 				if (!IsPendingKill())
 				{
-					OnComponentBeginOverlap.Broadcast(OtherActor, OtherComp, OtherOverlap.BodyIndex);
+					OnComponentBeginOverlap.Broadcast(OtherActor, OtherComp, OtherOverlap.GetBodyIndex(), OtherOverlap.bFromSweep, OtherOverlap.OverlapInfo);
 				}
 
 				if (!OtherComp->IsPendingKill())
 				{
-					OtherComp->OnComponentBeginOverlap.Broadcast(MyActor, this, INDEX_NONE);
+					OtherComp->OnComponentBeginOverlap.Broadcast(MyActor, this, INDEX_NONE, OtherOverlap.bFromSweep, OtherOverlap.OverlapInfo);
 				}
 
 				// then execute actor notification if this is a new actor touch
@@ -1888,7 +1888,7 @@ void UPrimitiveComponent::BeginComponentOverlap(const FOverlapInfo& OtherOverlap
 
 void UPrimitiveComponent::EndComponentOverlap(const FOverlapInfo& OtherOverlap, bool bDoNotifies, bool bNoNotifySelf)
 {
-	UPrimitiveComponent* OtherComp = OtherOverlap.Component.Get();
+	UPrimitiveComponent* OtherComp = OtherOverlap.OverlapInfo.Component.Get();
 
 	AActor* const OtherActor = OtherComp ? OtherComp->GetOwner() : NULL;
 	AActor* const MyActor = GetOwner();
@@ -1899,7 +1899,7 @@ void UPrimitiveComponent::EndComponentOverlap(const FOverlapInfo& OtherOverlap,
 	{
 		if (!bNoNotifySelf && IsPrimCompValidAndAlive(this))
 		{
-			OnComponentEndOverlap.Broadcast(OtherActor, OtherComp, OtherOverlap.BodyIndex);
+			OnComponentEndOverlap.Broadcast(OtherActor, OtherComp, OtherOverlap.GetBodyIndex());
 		}
 
 		if(IsPrimCompValidAndAlive(OtherComp))
@@ -1950,9 +1950,9 @@ void UPrimitiveComponent::GetOverlappingActors(TArray<AActor*>& OutOverlappingAc
 	{
 		const FOverlapInfo& OtherOvelap = *CompIt;
 		
-		if (OtherOvelap.Component.IsValid())
+		if (OtherOvelap.OverlapInfo.Component.IsValid())
 		{
-			AActor* OtherActor = OtherOvelap.Component->GetOwner();
+			AActor* OtherActor = OtherOvelap.OverlapInfo.Component->GetOwner();
 			if (OtherActor)
 			{
 				if ( !ClassFilter || OtherActor->IsA(ClassFilter) )
@@ -1971,7 +1971,7 @@ void UPrimitiveComponent::GetOverlappingComponents(TArray<UPrimitiveComponent*>&
 
 	for (auto CompIt = OverlappingComponents.CreateConstIterator(); CompIt; ++CompIt)
 	{
-		UPrimitiveComponent* const OtherComp = (*CompIt).Component.Get();
+		UPrimitiveComponent* const OtherComp = (*CompIt).OverlapInfo.Component.Get();
 		if (OtherComp)
 		{
 			OutOverlappingComponents.Add(OtherComp);
@@ -2122,7 +2122,7 @@ void UPrimitiveComponent::UpdateOverlaps(TArray<FOverlapInfo> const* PendingOver
 			for (auto CompIt = OldOverlappingComponents.CreateIterator(); CompIt; ++CompIt)
 			{
 				const FOverlapInfo& OtherOverlap = *CompIt;
-				if (OtherOverlap.Component.IsValid())
+				if (OtherOverlap.OverlapInfo.Component.IsValid())
 				{
 					EndComponentOverlap(OtherOverlap, bDoNotifies, false);
 				}
@@ -2132,7 +2132,7 @@ void UPrimitiveComponent::UpdateOverlaps(TArray<FOverlapInfo> const* PendingOver
 			for (auto CompIt = NewOverlappingComponents.CreateIterator(); CompIt; ++CompIt)
 			{
 				const FOverlapInfo& OtherOverlap = *CompIt;
-				if (OtherOverlap.Component.IsValid())
+				if (OtherOverlap.OverlapInfo.Component.IsValid())
 				{
 					BeginComponentOverlap(OtherOverlap, bDoNotifies);
 				}
@@ -2147,7 +2147,7 @@ void UPrimitiveComponent::UpdateOverlaps(TArray<FOverlapInfo> const* PendingOver
 		for (auto CompIt = OverlappingComponents.CreateIterator(); CompIt; ++CompIt)
 		{
 			const FOverlapInfo& OtherOverlap = *CompIt;
-			if (OtherOverlap.Component.IsValid())
+			if (OtherOverlap.OverlapInfo.Component.IsValid())
 			{
 				EndComponentOverlap(OtherOverlap, bDoNotifies, false);
 			}
@@ -2238,7 +2238,7 @@ void UPrimitiveComponent::UpdatePhysicsVolume( bool bTriggerNotifiers )
 			for (auto CompIt = OverlappingComponents.CreateIterator(); CompIt; ++CompIt)
 			{
 				const FOverlapInfo& Overlap = *CompIt;
-				UPrimitiveComponent* OtherComponent = Overlap.Component.Get();
+				UPrimitiveComponent* OtherComponent = Overlap.OverlapInfo.Component.Get();
 				if (OtherComponent)
 				{
 					APhysicsVolume* V = Cast<APhysicsVolume>(OtherComponent->GetOwner());
