#ue4 - Remove client/server prediction data when TornOff.

UDN: https://udn.unrealengine.com/questions/208507/root-motion-movement-on-torn-off-pawns.html

diff --git a/Engine/Source/Runtime/Engine/Classes/GameFramework/Character.h b/Engine/Source/Runtime/Engine/Classes/GameFramework/Character.h
--- a/Engine/Source/Runtime/Engine/Classes/GameFramework/Character.h
+++ b/Engine/Source/Runtime/Engine/Classes/GameFramework/Character.h
@@ -261,6 +261,7 @@ public:
 	virtual void GetSimpleCollisionCylinder(float& CollisionRadius, float& CollisionHalfHeight) const OVERRIDE;
 	virtual void ApplyWorldOffset(const FVector& InOffset, bool bWorldShift) OVERRIDE;
 	virtual UActorComponent* FindComponentByClass(const TSubclassOf<UActorComponent> ComponentClass) const OVERRIDE;
+	virtual void TornOff() OVERRIDE;
 	// End AActor Interface
 
 	template<class T>
diff --git a/Engine/Source/Runtime/Engine/Private/Character.cpp b/Engine/Source/Runtime/Engine/Private/Character.cpp
--- a/Engine/Source/Runtime/Engine/Private/Character.cpp
+++ b/Engine/Source/Runtime/Engine/Private/Character.cpp
@@ -587,6 +587,18 @@ void ACharacter::UnPossessed()
 }
 
 
+void ACharacter::TornOff()
+{
+	Super::TornOff();
+
+	if (CharacterMovement)
+	{
+		CharacterMovement->ResetPredictionData_Client();
+		CharacterMovement->ResetPredictionData_Server();
+	}
+}
+
+
 void ACharacter::BaseChange()
 {
 	if (CharacterMovement && CharacterMovement->MovementMode != MOVE_None)
