FBodyInstance::SetBodyTransform - replace isValid check with warning and early out to prevent crash on bad data.

diff --git a/Engine/Source/Runtime/Engine/Private/PhysicsEngine/BodyInstance.cpp b/Engine/Source/Runtime/Engine/Private/PhysicsEngine/BodyInstance.cpp
--- a/Engine/Source/Runtime/Engine/Private/PhysicsEngine/BodyInstance.cpp
+++ b/Engine/Source/Runtime/Engine/Private/PhysicsEngine/BodyInstance.cpp
@@ -1635,7 +1635,15 @@ void FBodyInstance::SetBodyTransform(const FTransform& NewTransform, bool bTelep
 	}
 
 	const PxTransform PNewPose = U2PTransform(NewTransform);
-	check(PNewPose.isValid());
+
+	if (!PNewPose.isValid())
+	{
+		return;
+	}
+	else
+	{
+		UE_LOG(LogPhysics, Warning, TEXT("FBodyInstance::SetBodyTransform: Trying to set new transform with bad data [p=(%f,%f,%f) q=(%f,%f,%f,%f)]"), PNewPose.p.x, PNewPose.p.y, PNewPose.p.z, PNewPose.q.x, PNewPose.q.y, PNewPose.q.z, PNewPose.q.w);
+	}
 
 	SCENE_LOCK_WRITE(RigidActor->getScene());
 	// SIMULATED & KINEMATIC
