#UE4 Feature to fix TTP 335117...

diff --git a/Engine/Source/Runtime/Engine/Classes/Kismet/KismetStringLibrary.h b/Engine/Source/Runtime/Engine/Classes/Kismet/KismetStringLibrary.h
--- a/Engine/Source/Runtime/Engine/Classes/Kismet/KismetStringLibrary.h
+++ b/Engine/Source/Runtime/Engine/Classes/Kismet/KismetStringLibrary.h
@@ -36,6 +36,10 @@ class UKismetStringLibrary : public UBlueprintFunctionLibrary
 	UFUNCTION(BlueprintPure, meta=(FriendlyName = "ToString (rotator)", CompactNodeTitle = "->"), Category="Utilities|String")
 	static FString Conv_RotatorToString(FRotator InRot);
 
+	/** Converts a transform value to a string, in the form 'Translation: X= Y= Z= Rotation: P= Y= R= Scale: X= Y= Z=' */
+	UFUNCTION(BlueprintPure, meta=(FriendlyName = "ToString (transform)", CompactNodeTitle = "->"), Category="Utilities|String")
+	static FString Conv_TransformToString(FTransform InTrans);
+
 	/** Converts a UObject value to a string by calling the object's GetName method */
 	UFUNCTION(BlueprintPure, meta=(FriendlyName = "ToString (object)", CompactNodeTitle = "->"), Category="Utilities|String")
 	static FString Conv_ObjectToString(class UObject* InObj);
diff --git a/Engine/Source/Runtime/Engine/Private/KismetStringLibrary.cpp b/Engine/Source/Runtime/Engine/Private/KismetStringLibrary.cpp
--- a/Engine/Source/Runtime/Engine/Private/KismetStringLibrary.cpp
+++ b/Engine/Source/Runtime/Engine/Private/KismetStringLibrary.cpp
@@ -81,6 +81,11 @@ FString UKismetStringLibrary::Conv_RotatorToString(FRotator InRot)
 	return InRot.ToString();	
 }
 
+FString UKismetStringLibrary::Conv_TransformToString(FTransform InTrans)
+{
+	return FString::Printf(TEXT("Translation: %s Rotation: %s Scale %s"), *InTrans.GetTranslation().ToString(), *InTrans.Rotator().ToString(), *InTrans.GetScale3D().ToString());
+}
+
 FString UKismetStringLibrary::Conv_ObjectToString(class UObject* InObj)
 {
 	return (InObj != NULL) ? InObj->GetName() : FString(TEXT("None"));
