Fix warning on wrong part of if :(

diff --git a/Engine/Source/Runtime/Engine/Private/PhysicsEngine/BodyInstance.cpp b/Engine/Source/Runtime/Engine/Private/PhysicsEngine/BodyInstance.cpp
--- a/Engine/Source/Runtime/Engine/Private/PhysicsEngine/BodyInstance.cpp
+++ b/Engine/Source/Runtime/Engine/Private/PhysicsEngine/BodyInstance.cpp
@@ -1637,14 +1637,11 @@ void FBodyInstance::SetBodyTransform(const FTransform& NewTransform, bool bTelep
 	const PxTransform PNewPose = U2PTransform(NewTransform);
 
 	if (!PNewPose.isValid())
-	{
-		return;
-	}
-	else
 	{
 		UE_LOG(LogPhysics, Warning, TEXT("FBodyInstance::SetBodyTransform: Trying to set new transform with bad data [p=(%f,%f,%f) q=(%f,%f,%f,%f)]"), PNewPose.p.x, PNewPose.p.y, PNewPose.p.z, PNewPose.q.x, PNewPose.q.y, PNewPose.q.z, PNewPose.q.w);
+		return;
 	}
-
+	
 	SCENE_LOCK_WRITE(RigidActor->getScene());
 	// SIMULATED & KINEMATIC
 	if(PRigidDynamic)
