fixed Win32 build

diff --git a/Engine/Source/Runtime/Engine/Classes/Kismet/KismetStringLibrary.h b/Engine/Source/Runtime/Engine/Classes/Kismet/KismetStringLibrary.h
--- a/Engine/Source/Runtime/Engine/Classes/Kismet/KismetStringLibrary.h
+++ b/Engine/Source/Runtime/Engine/Classes/Kismet/KismetStringLibrary.h
@@ -38,7 +38,7 @@ class UKismetStringLibrary : public UBlueprintFunctionLibrary
 
 	/** Converts a transform value to a string, in the form 'Translation: X= Y= Z= Rotation: P= Y= R= Scale: X= Y= Z=' */
 	UFUNCTION(BlueprintPure, meta=(FriendlyName = "ToString (transform)", CompactNodeTitle = "->"), Category="Utilities|String")
-	static FString Conv_TransformToString(FTransform InTrans);
+	static FString Conv_TransformToString(const FTransform& InTrans);
 
 	/** Converts a UObject value to a string by calling the object's GetName method */
 	UFUNCTION(BlueprintPure, meta=(FriendlyName = "ToString (object)", CompactNodeTitle = "->"), Category="Utilities|String")
diff --git a/Engine/Source/Runtime/Engine/Private/KismetStringLibrary.cpp b/Engine/Source/Runtime/Engine/Private/KismetStringLibrary.cpp
--- a/Engine/Source/Runtime/Engine/Private/KismetStringLibrary.cpp
+++ b/Engine/Source/Runtime/Engine/Private/KismetStringLibrary.cpp
@@ -81,7 +81,7 @@ FString UKismetStringLibrary::Conv_RotatorToString(FRotator InRot)
 	return InRot.ToString();	
 }
 
-FString UKismetStringLibrary::Conv_TransformToString(FTransform InTrans)
+FString UKismetStringLibrary::Conv_TransformToString(const FTransform& InTrans)
 {
 	return FString::Printf(TEXT("Translation: %s Rotation: %s Scale %s"), *InTrans.GetTranslation().ToString(), *InTrans.Rotator().ToString(), *InTrans.GetScale3D().ToString());
 }
