Fixed crash upon exiting PIE when a Level Blueprint has been created.
#ttp 336690 - Regression: MATINEE: Crash: Exiting Play in Editor after Playing a Matinee with the level blueprint
#branch UE4
#proj Editor.UnrealEd, Runtime.Engine
#add Added code to UEditorEngine::EndPlayMap() to repoint the outers of the GeneratedClass and ClassDefaultObject of the LevelScriptBlueprint to the Editor World package when quitting PIE, so GC isn't left with hanging references.
#remove Removed condition in ULevel::PostLoad() to only rename the LevelScriptBlueprint when it's not named the same as the outer world.  This is because the rename also has the side effect of repointing the outers of the other generated BP objects, which we always want to do.

diff --git a/Engine/Source/Runtime/Engine/Private/Level.cpp b/Engine/Source/Runtime/Engine/Private/Level.cpp
--- a/Engine/Source/Runtime/Engine/Private/Level.cpp
+++ b/Engine/Source/Runtime/Engine/Private/Level.cpp
@@ -610,12 +610,15 @@ void ULevel::PostLoad()
 	}
 
 #if WITH_EDITOR
-	// Rename the LevelScriptBlueprint after the outer world.
-	UWorld* OuterWorld = Cast<UWorld>(GetOuter());
-	if (LevelScriptBlueprint && OuterWorld && LevelScriptBlueprint->GetFName() != OuterWorld->GetFName())
+	if (!(GetOutermost()->PackageFlags & PKG_PlayInEditor))
 	{
-		// Use LevelScriptBlueprint->GetOuter() instead of NULL to make sure the generated top level objects are moved appropriately
-		LevelScriptBlueprint->Rename(*OuterWorld->GetName(), LevelScriptBlueprint->GetOuter(), REN_DoNotDirty | REN_DontCreateRedirectors | REN_ForceNoResetLoaders | REN_NonTransactional);
+		// Rename the LevelScriptBlueprint after the outer world.
+		UWorld* OuterWorld = Cast<UWorld>(GetOuter());
+		if (LevelScriptBlueprint && OuterWorld && LevelScriptBlueprint->GetFName() != OuterWorld->GetFName())
+		{
+			// Use LevelScriptBlueprint->GetOuter() instead of NULL to make sure the generated top level objects are moved appropriately
+			LevelScriptBlueprint->Rename(*OuterWorld->GetName(), LevelScriptBlueprint->GetOuter(), REN_DoNotDirty | REN_DontCreateRedirectors | REN_ForceNoResetLoaders | REN_NonTransactional);
+		}
 	}
 #endif
 }
