Remove more headers from Engine.h (StaticMeshResources.h, AnimTree.h, SkeletalMeshTypes.h, SkeletalMeshActor.h, LightingBuildOptions.h, PixelFormat.h, WorldComposition.h, VisualLog.h, StaticLighting.h, Lightmap.h, ShadowMap.h, Model.h)

diff --git a/Engine/Source/Developer/FunctionalTesting/Private/FunctionalTest.cpp b/Engine/Source/Developer/FunctionalTesting/Private/FunctionalTest.cpp
--- a/Engine/Source/Developer/FunctionalTesting/Private/FunctionalTest.cpp
+++ b/Engine/Source/Developer/FunctionalTesting/Private/FunctionalTest.cpp
@@ -2,6 +2,7 @@
 
 #include "FunctionalTestingPrivatePCH.h"
 #include "ObjectEditorUtils.h"
+#include "VisualLog.h"
 
 AFunctionalTest::AFunctionalTest( const class FPostConstructInitializeProperties& PCIP )
 	: Super(PCIP)
diff --git a/Engine/Source/Developer/MeshUtilities/Private/MeshUtilities.cpp b/Engine/Source/Developer/MeshUtilities/Private/MeshUtilities.cpp
--- a/Engine/Source/Developer/MeshUtilities/Private/MeshUtilities.cpp
+++ b/Engine/Source/Developer/MeshUtilities/Private/MeshUtilities.cpp
@@ -1,6 +1,7 @@
 // Copyright 1998-2014 Epic Games, Inc. All Rights Reserved.
 
 #include "MeshUtilitiesPrivate.h"
+#include "StaticMeshResources.h"
 #include "SkeletalMeshTypes.h"
 #include "Landscape/LandscapeRender.h"
 #include "MeshBuild.h"
diff --git a/Engine/Source/Developer/Windows/WindowsTargetPlatform/Private/GenericWindowsTargetPlatform.h b/Engine/Source/Developer/Windows/WindowsTargetPlatform/Private/GenericWindowsTargetPlatform.h
--- a/Engine/Source/Developer/Windows/WindowsTargetPlatform/Private/GenericWindowsTargetPlatform.h
+++ b/Engine/Source/Developer/Windows/WindowsTargetPlatform/Private/GenericWindowsTargetPlatform.h
@@ -6,6 +6,7 @@
 
 #pragma once
 
+#include "StaticMeshResources.h"
 
 /**
  * Template for Windows target platforms
diff --git a/Engine/Source/Developer/iOS/IOSTargetPlatform/Private/IOSTargetPlatform.h b/Engine/Source/Developer/iOS/IOSTargetPlatform/Private/IOSTargetPlatform.h
--- a/Engine/Source/Developer/iOS/IOSTargetPlatform/Private/IOSTargetPlatform.h
+++ b/Engine/Source/Developer/iOS/IOSTargetPlatform/Private/IOSTargetPlatform.h
@@ -7,7 +7,7 @@
 #pragma once
 
 #include "Ticker.h"
-
+#include "StaticMeshResources.h"
 
 /**
  * FIOSTargetPlatform, abstraction for cooking iOS platforms
diff --git a/Engine/Source/Editor/FoliageEdit/Private/FoliageEdMode.cpp b/Engine/Source/Editor/FoliageEdit/Private/FoliageEdMode.cpp
--- a/Engine/Source/Editor/FoliageEdit/Private/FoliageEdMode.cpp
+++ b/Engine/Source/Editor/FoliageEdit/Private/FoliageEdMode.cpp
@@ -5,6 +5,7 @@
 ================================================================================*/
 
 #include "UnrealEd.h"
+#include "StaticMeshResources.h"
 #include "ObjectTools.h"
 #include "FoliageEdMode.h"
 #include "ScopedTransaction.h"
diff --git a/Engine/Source/Editor/LevelEditor/Private/LevelEditorActions.cpp b/Engine/Source/Editor/LevelEditor/Private/LevelEditorActions.cpp
--- a/Engine/Source/Editor/LevelEditor/Private/LevelEditorActions.cpp
+++ b/Engine/Source/Editor/LevelEditor/Private/LevelEditorActions.cpp
@@ -4,6 +4,7 @@
 #include "LevelEditor.h"
 #include "Matinee/MatineeActor.h"
 #include "Engine/LevelScriptBlueprint.h"
+#include "LightingBuildOptions.h"
 #include "EditorSupportDelegates.h"
 #include "SLevelEditor.h"
 #include "LevelEditorActions.h"
diff --git a/Engine/Source/Editor/LevelEditor/Private/SLevelEditorBuildAndSubmit.cpp b/Engine/Source/Editor/LevelEditor/Private/SLevelEditorBuildAndSubmit.cpp
--- a/Engine/Source/Editor/LevelEditor/Private/SLevelEditorBuildAndSubmit.cpp
+++ b/Engine/Source/Editor/LevelEditor/Private/SLevelEditorBuildAndSubmit.cpp
@@ -7,6 +7,7 @@
 #include "LevelEditorActions.h"
 #include "EditorBuildUtils.h"
 #include "MessageLog.h"
+#include "LightingBuildOptions.h"
 
 #include "AssetToolsModule.h"
 
diff --git a/Engine/Source/Editor/MeshPaint/Private/MeshPaintEdMode.cpp b/Engine/Source/Editor/MeshPaint/Private/MeshPaintEdMode.cpp
--- a/Engine/Source/Editor/MeshPaint/Private/MeshPaintEdMode.cpp
+++ b/Engine/Source/Editor/MeshPaint/Private/MeshPaintEdMode.cpp
@@ -11,6 +11,7 @@
 #include "Materials/MaterialExpressionTextureCoordinate.h"
 #include "Materials/MaterialExpressionTextureSampleParameter.h"
 
+#include "StaticMeshResources.h"
 #include "MeshPaintEdMode.h"
 #include "Factories.h"
 #include "ScopedTransaction.h"
diff --git a/Engine/Source/Editor/StaticMeshEditor/Private/StaticMeshEditor.cpp b/Engine/Source/Editor/StaticMeshEditor/Private/StaticMeshEditor.cpp
--- a/Engine/Source/Editor/StaticMeshEditor/Private/StaticMeshEditor.cpp
+++ b/Engine/Source/Editor/StaticMeshEditor/Private/StaticMeshEditor.cpp
@@ -10,6 +10,7 @@
 #include "StaticMeshEditorActions.h"
 
 #include "UnrealEd.h"
+#include "StaticMeshResources.h"
 #include "ISocketManager.h"
 #include "PreviewScene.h"
 #include "ScopedTransaction.h"
diff --git a/Engine/Source/Editor/StatsViewer/Private/StatsPages/PrimitiveStatsPage.cpp b/Engine/Source/Editor/StatsViewer/Private/StatsPages/PrimitiveStatsPage.cpp
--- a/Engine/Source/Editor/StatsViewer/Private/StatsPages/PrimitiveStatsPage.cpp
+++ b/Engine/Source/Editor/StatsViewer/Private/StatsPages/PrimitiveStatsPage.cpp
@@ -1,6 +1,7 @@
 // Copyright 1998-2014 Epic Games, Inc. All Rights Reserved.
 
 #include "StatsViewerPrivatePCH.h"
+#include "StaticMeshResources.h"
 #include "PrimitiveStatsPage.h"
 #include "PrimitiveStats.h"
 #include "Landscape/LandscapeProxy.h"
diff --git a/Engine/Source/Editor/UnrealEd/Private/Commandlets/CookCommandlet.cpp b/Engine/Source/Editor/UnrealEd/Private/Commandlets/CookCommandlet.cpp
--- a/Engine/Source/Editor/UnrealEd/Private/Commandlets/CookCommandlet.cpp
+++ b/Engine/Source/Editor/UnrealEd/Private/Commandlets/CookCommandlet.cpp
@@ -6,6 +6,7 @@
 
 #include "UnrealEd.h"
 
+#include "Engine/WorldComposition.h"
 #include "PackageHelperFunctions.h"
 #include "DerivedDataCacheInterface.h"
 #include "ISourceControlModule.h"
diff --git a/Engine/Source/Editor/UnrealEd/Private/DragTool_FrustumSelect.cpp b/Engine/Source/Editor/UnrealEd/Private/DragTool_FrustumSelect.cpp
--- a/Engine/Source/Editor/UnrealEd/Private/DragTool_FrustumSelect.cpp
+++ b/Engine/Source/Editor/UnrealEd/Private/DragTool_FrustumSelect.cpp
@@ -2,6 +2,7 @@
 
 
 #include "UnrealEd.h"
+#include "StaticMeshResources.h"
 #include "DragTool_FrustumSelect.h"
 #include "ActorEditorUtils.h"
 #include "ScopedTransaction.h"
diff --git a/Engine/Source/Editor/UnrealEd/Private/Editor.cpp b/Engine/Source/Editor/UnrealEd/Private/Editor.cpp
--- a/Engine/Source/Editor/UnrealEd/Private/Editor.cpp
+++ b/Engine/Source/Editor/UnrealEd/Private/Editor.cpp
@@ -4,6 +4,7 @@
 #include "UnrealEd.h"
 #include "Matinee/MatineeActor.h"
 #include "Engine/InteractiveFoliageActor.h"
+#include "Engine/WorldComposition.h"
 #include "EditorSupportDelegates.h"
 #include "Factories.h"
 #include "BSPOps.h"
diff --git a/Engine/Source/Editor/UnrealEd/Private/EditorExporters.cpp b/Engine/Source/Editor/UnrealEd/Private/EditorExporters.cpp
--- a/Engine/Source/Editor/UnrealEd/Private/EditorExporters.cpp
+++ b/Engine/Source/Editor/UnrealEd/Private/EditorExporters.cpp
@@ -5,6 +5,7 @@
 =============================================================================*/
 
 #include "UnrealEd.h"
+#include "StaticMeshResources.h"
 #include "Materials/MaterialExpressionConstant.h"
 #include "SoundDefinitions.h"
 #include "Landscape/LandscapeDataAccess.h"
diff --git a/Engine/Source/Editor/UnrealEd/Private/EditorModes.cpp b/Engine/Source/Editor/UnrealEd/Private/EditorModes.cpp
--- a/Engine/Source/Editor/UnrealEd/Private/EditorModes.cpp
+++ b/Engine/Source/Editor/UnrealEd/Private/EditorModes.cpp
@@ -2,6 +2,7 @@
 
 #include "UnrealEd.h"
 #include "Engine/BookMark.h"
+#include "StaticMeshResources.h"
 #include "EditorSupportDelegates.h"
 #include "MouseDeltaTracker.h"
 #include "ScopedTransaction.h"
diff --git a/Engine/Source/Editor/UnrealEd/Private/Factories/ActorFactory.cpp b/Engine/Source/Editor/UnrealEd/Private/Factories/ActorFactory.cpp
--- a/Engine/Source/Editor/UnrealEd/Private/Factories/ActorFactory.cpp
+++ b/Engine/Source/Editor/UnrealEd/Private/Factories/ActorFactory.cpp
@@ -7,6 +7,7 @@ ActorFactory.cpp:
 #include "UnrealEd.h"
 #include "ParticleDefinitions.h"
 #include "SoundDefinitions.h"
+#include "StaticMeshResources.h"
 #include "BlueprintUtilities.h"
 #include "Kismet2/BlueprintEditorUtils.h"
 #include "Kismet2/KismetEditorUtilities.h"
diff --git a/Engine/Source/Editor/UnrealEd/Private/Fbx/FbxMainExport.cpp b/Engine/Source/Editor/UnrealEd/Private/Fbx/FbxMainExport.cpp
--- a/Engine/Source/Editor/UnrealEd/Private/Fbx/FbxMainExport.cpp
+++ b/Engine/Source/Editor/UnrealEd/Private/Fbx/FbxMainExport.cpp
@@ -47,6 +47,7 @@
 #include "Matinee/InterpTrackInstFloatProp.h"
 #include "Matinee/InterpTrackInstMove.h"
 
+#include "StaticMeshResources.h"
 #include "LandscapeDataAccess.h"
 #include "Landscape/LandscapeProxy.h"
 
diff --git a/Engine/Source/Editor/UnrealEd/Private/Fbx/FbxStaticMeshImport.cpp b/Engine/Source/Editor/UnrealEd/Private/Fbx/FbxStaticMeshImport.cpp
--- a/Engine/Source/Editor/UnrealEd/Private/Fbx/FbxStaticMeshImport.cpp
+++ b/Engine/Source/Editor/UnrealEd/Private/Fbx/FbxStaticMeshImport.cpp
@@ -38,6 +38,7 @@
 #include "MeshUtilities.h"
 
 #include "Engine.h"
+#include "StaticMeshResources.h"
 #include "TextureLayout.h"
 #include "ObjectTools.h"
 #include "PackageTools.h"
diff --git a/Engine/Source/Editor/UnrealEd/Private/FbxMeshUtils.cpp b/Engine/Source/Editor/UnrealEd/Private/FbxMeshUtils.cpp
--- a/Engine/Source/Editor/UnrealEd/Private/FbxMeshUtils.cpp
+++ b/Engine/Source/Editor/UnrealEd/Private/FbxMeshUtils.cpp
@@ -1,6 +1,7 @@
 // Copyright 1998-2014 Epic Games, Inc. All Rights Reserved.
 
 #include "UnrealEd.h"
+#include "StaticMeshResources.h"
 #include "Factories.h"
 
 #include "FbxMeshUtils.h"
diff --git a/Engine/Source/Editor/UnrealEd/Private/GeomFitUtils.cpp b/Engine/Source/Editor/UnrealEd/Private/GeomFitUtils.cpp
--- a/Engine/Source/Editor/UnrealEd/Private/GeomFitUtils.cpp
+++ b/Engine/Source/Editor/UnrealEd/Private/GeomFitUtils.cpp
@@ -5,6 +5,7 @@
 =============================================================================*/
 
 #include "UnrealEd.h"
+#include "StaticMeshResources.h"
 #include "EditorSupportDelegates.h"
 #include "BSPOps.h"
 #include "../Private/GeomFitUtils.h"
diff --git a/Engine/Source/Editor/UnrealEd/Private/GeomTools.cpp b/Engine/Source/Editor/UnrealEd/Private/GeomTools.cpp
--- a/Engine/Source/Editor/UnrealEd/Private/GeomTools.cpp
+++ b/Engine/Source/Editor/UnrealEd/Private/GeomTools.cpp
@@ -4,6 +4,7 @@
  =============================================================================*/
 
 #include "UnrealEd.h"
+#include "StaticMeshResources.h"
 #include "GeomTools.h"
 
 DEFINE_LOG_CATEGORY_STATIC(LogGeomTools, Log, All);
diff --git a/Engine/Source/Editor/UnrealEd/Private/Lightmass/Lightmass.cpp b/Engine/Source/Editor/UnrealEd/Private/Lightmass/Lightmass.cpp
--- a/Engine/Source/Editor/UnrealEd/Private/Lightmass/Lightmass.cpp
+++ b/Engine/Source/Editor/UnrealEd/Private/Lightmass/Lightmass.cpp
@@ -6,6 +6,7 @@
 
 #include "UnrealEd.h"
 #include "PrecomputedLightVolume.h"
+#include "Runtime/Engine/Public/StaticMeshResources.h"
 #include "Runtime/Engine/Public/Landscape/LandscapeRender.h"
 #include "Runtime/Engine/Public/Landscape/LandscapeLight.h"
 #include "Runtime/Engine/Classes/Matinee/MatineeActor.h"
diff --git a/Engine/Source/Editor/UnrealEd/Private/StaticLightingSystem/StaticLightingDebug.cpp b/Engine/Source/Editor/UnrealEd/Private/StaticLightingSystem/StaticLightingDebug.cpp
--- a/Engine/Source/Editor/UnrealEd/Private/StaticLightingSystem/StaticLightingDebug.cpp
+++ b/Engine/Source/Editor/UnrealEd/Private/StaticLightingSystem/StaticLightingDebug.cpp
@@ -5,6 +5,7 @@
 =============================================================================*/
 
 #include "UnrealEd.h"
+#include "StaticMeshResources.h"
 #include "StaticLightingPrivate.h"
 
 /** Information about the texel that is selected */
diff --git a/Engine/Source/Editor/UnrealEd/Private/StaticLightingSystem/StaticLightingExport.cpp b/Engine/Source/Editor/UnrealEd/Private/StaticLightingSystem/StaticLightingExport.cpp
--- a/Engine/Source/Editor/UnrealEd/Private/StaticLightingSystem/StaticLightingExport.cpp
+++ b/Engine/Source/Editor/UnrealEd/Private/StaticLightingSystem/StaticLightingExport.cpp
@@ -8,6 +8,7 @@
 
 #include "StaticMeshLight.h"
 #include "ModelLight.h"
+#include "StaticMeshResources.h"
 #include "Runtime/Engine/Public/Landscape/LandscapeLight.h"
 #include "Runtime/Engine/Classes/Landscape/LandscapeComponent.h"
 
diff --git a/Engine/Source/Editor/UnrealEd/Private/StaticLightingSystem/StaticLightingPrivate.h b/Engine/Source/Editor/UnrealEd/Private/StaticLightingSystem/StaticLightingPrivate.h
--- a/Engine/Source/Editor/UnrealEd/Private/StaticLightingSystem/StaticLightingPrivate.h
+++ b/Engine/Source/Editor/UnrealEd/Private/StaticLightingSystem/StaticLightingPrivate.h
@@ -9,6 +9,8 @@
 // Includes.
 #include "ChunkedArray.h"
 #include "GenericOctree.h"
+#include "StaticLighting.h"
+#include "LightingBuildOptions.h"
 
 DECLARE_LOG_CATEGORY_EXTERN(LogStaticLightingSystem, Log, All);
 
diff --git a/Engine/Source/Editor/UnrealEd/Private/StaticMeshEdit.cpp b/Engine/Source/Editor/UnrealEd/Private/StaticMeshEdit.cpp
--- a/Engine/Source/Editor/UnrealEd/Private/StaticMeshEdit.cpp
+++ b/Engine/Source/Editor/UnrealEd/Private/StaticMeshEdit.cpp
@@ -5,6 +5,7 @@
 =============================================================================*/
 
 #include "UnrealEd.h"
+#include "StaticMeshResources.h"
 #include "Factories.h"
 #include "TextureLayout.h"
 #include "BSPOps.h"
diff --git a/Engine/Source/Editor/UnrealEd/Private/UnrealEdSrv.cpp b/Engine/Source/Editor/UnrealEd/Private/UnrealEdSrv.cpp
--- a/Engine/Source/Editor/UnrealEd/Private/UnrealEdSrv.cpp
+++ b/Engine/Source/Editor/UnrealEd/Private/UnrealEdSrv.cpp
@@ -2,6 +2,7 @@
 
 
 #include "UnrealEd.h"
+#include "StaticMeshResources.h"
 #include "EditorSupportDelegates.h"
 #include "BusyCursor.h"
 #include "ScopedTransaction.h"
diff --git a/Engine/Source/Editor/UnrealEd/Private/VertexSnapping.cpp b/Engine/Source/Editor/UnrealEd/Private/VertexSnapping.cpp
--- a/Engine/Source/Editor/UnrealEd/Private/VertexSnapping.cpp
+++ b/Engine/Source/Editor/UnrealEd/Private/VertexSnapping.cpp
@@ -2,6 +2,7 @@
 
 #include "UnrealEd.h"
 #include "VertexSnapping.h"
+#include "StaticMeshResources.h"
 
 namespace VertexSnappingConstants
 {
diff --git a/Engine/Source/Editor/UnrealEd/Public/TexAlignTools.h b/Engine/Source/Editor/UnrealEd/Public/TexAlignTools.h
--- a/Engine/Source/Editor/UnrealEd/Public/TexAlignTools.h
+++ b/Engine/Source/Editor/UnrealEd/Public/TexAlignTools.h
@@ -6,6 +6,8 @@
 
 #pragma once
 
+#include "Model.h"
+
 class FBspSurfIdx
 {
 public:
diff --git a/Engine/Source/Editor/WorldBrowser/Private/Tiles/WorldTileCollectionModel.cpp b/Engine/Source/Editor/WorldBrowser/Private/Tiles/WorldTileCollectionModel.cpp
--- a/Engine/Source/Editor/WorldBrowser/Private/Tiles/WorldTileCollectionModel.cpp
+++ b/Engine/Source/Editor/WorldBrowser/Private/Tiles/WorldTileCollectionModel.cpp
@@ -1,6 +1,7 @@
 // Copyright 1998-2014 Epic Games, Inc. All Rights Reserved.
 #include "WorldBrowserPrivatePCH.h"
 
+#include "Engine/WorldComposition.h"
 #include "Editor/PropertyEditor/Public/IDetailsView.h"
 #include "Editor/PropertyEditor/Public/PropertyEditorModule.h"
 #include "FileHelpers.h"
diff --git a/Engine/Source/Editor/WorldBrowser/Private/Tiles/WorldTileDetailsCustomization.cpp b/Engine/Source/Editor/WorldBrowser/Private/Tiles/WorldTileDetailsCustomization.cpp
--- a/Engine/Source/Editor/WorldBrowser/Private/Tiles/WorldTileDetailsCustomization.cpp
+++ b/Engine/Source/Editor/WorldBrowser/Private/Tiles/WorldTileDetailsCustomization.cpp
@@ -1,6 +1,8 @@
 // Copyright 1998-2014 Epic Games, Inc. All Rights Reserved.
 #include "WorldBrowserPrivatePCH.h"
 
+#include "Engine/WorldComposition.h"
+
 #include "Editor/PropertyEditor/Public/IDetailsView.h"
 #include "Editor/PropertyEditor/Public/PropertyEditing.h"
 #include "Editor/PropertyEditor/Public/PropertyEditorModule.h"
diff --git a/Engine/Source/Editor/WorldBrowser/Private/Tiles/WorldTileModel.cpp b/Engine/Source/Editor/WorldBrowser/Private/Tiles/WorldTileModel.cpp
--- a/Engine/Source/Editor/WorldBrowser/Private/Tiles/WorldTileModel.cpp
+++ b/Engine/Source/Editor/WorldBrowser/Private/Tiles/WorldTileModel.cpp
@@ -8,6 +8,7 @@
 #include "WorldTileCollectionModel.h"
 #include "WorldTileModel.h"
 
+#include "Engine/WorldComposition.h"
 #include "Landscape/Landscape.h"
 
 #define LOCTEXT_NAMESPACE "WorldBrowser"
diff --git a/Engine/Source/Editor/WorldBrowser/Private/WorldBrowserModule.cpp b/Engine/Source/Editor/WorldBrowser/Private/WorldBrowserModule.cpp
--- a/Engine/Source/Editor/WorldBrowser/Private/WorldBrowserModule.cpp
+++ b/Engine/Source/Editor/WorldBrowser/Private/WorldBrowserModule.cpp
@@ -1,6 +1,7 @@
 // Copyright 1998-2014 Epic Games, Inc. All Rights Reserved.
 #include "WorldBrowserPrivatePCH.h"
 
+#include "Engine/WorldComposition.h"
 #include "ModuleManager.h"
 #include "StreamingLevels/StreamingLevelEdMode.h"
 #include "Tiles/WorldTileCollectionModel.h"
diff --git a/Engine/Source/Runtime/AIModule/Private/BehaviorTree/BTCompositeNode.cpp b/Engine/Source/Runtime/AIModule/Private/BehaviorTree/BTCompositeNode.cpp
--- a/Engine/Source/Runtime/AIModule/Private/BehaviorTree/BTCompositeNode.cpp
+++ b/Engine/Source/Runtime/AIModule/Private/BehaviorTree/BTCompositeNode.cpp
@@ -2,6 +2,7 @@
 
 #include "AIModulePrivate.h"
 #include "BehaviorTree/BTCompositeNode.h"
+#include "VisualLog.h"
 
 UBTCompositeNode::UBTCompositeNode(const class FPostConstructInitializeProperties& PCIP) : Super(PCIP)
 {
diff --git a/Engine/Source/Runtime/AIModule/Private/BehaviorTree/BehaviorTreeComponent.cpp b/Engine/Source/Runtime/AIModule/Private/BehaviorTree/BehaviorTreeComponent.cpp
--- a/Engine/Source/Runtime/AIModule/Private/BehaviorTree/BehaviorTreeComponent.cpp
+++ b/Engine/Source/Runtime/AIModule/Private/BehaviorTree/BehaviorTreeComponent.cpp
@@ -8,6 +8,7 @@
 #include "BehaviorTree/BehaviorTreeManager.h"
 #include "BehaviorTree/BlackboardComponent.h"
 #include "BehaviorTree/BehaviorTreeComponent.h"
+#include "VisualLog.h"
 
 //----------------------------------------------------------------------//
 // UBehaviorTreeComponent
diff --git a/Engine/Source/Runtime/AIModule/Private/BehaviorTree/BehaviorTreeTypes.cpp b/Engine/Source/Runtime/AIModule/Private/BehaviorTree/BehaviorTreeTypes.cpp
--- a/Engine/Source/Runtime/AIModule/Private/BehaviorTree/BehaviorTreeTypes.cpp
+++ b/Engine/Source/Runtime/AIModule/Private/BehaviorTree/BehaviorTreeTypes.cpp
@@ -5,6 +5,7 @@
 #include "BehaviorTree/BTCompositeNode.h"
 #include "BehaviorTree/Tasks/BTTask_RunBehavior.h"
 #include "BehaviorTree/BehaviorTreeTypes.h"
+#include "VisualLog.h"
 
 //----------------------------------------------------------------------//
 // FBehaviorTreeInstance
diff --git a/Engine/Source/Runtime/AIModule/Private/BehaviorTree/BlackboardComponent.cpp b/Engine/Source/Runtime/AIModule/Private/BehaviorTree/BlackboardComponent.cpp
--- a/Engine/Source/Runtime/AIModule/Private/BehaviorTree/BlackboardComponent.cpp
+++ b/Engine/Source/Runtime/AIModule/Private/BehaviorTree/BlackboardComponent.cpp
@@ -4,6 +4,7 @@
 #include "EngineUtils.h"
 #include "BehaviorTree/Blackboard/BlackboardKeyAllTypes.h"
 #include "BehaviorTree/BlackboardComponent.h"
+#include "VisualLog.h"
 
 UBlackboardComponent::UBlackboardComponent(const class FPostConstructInitializeProperties& PCIP) : Super(PCIP)
 {
diff --git a/Engine/Source/Runtime/AIModule/Private/BehaviorTree/Decorators/BTDecorator_Blackboard.cpp b/Engine/Source/Runtime/AIModule/Private/BehaviorTree/Decorators/BTDecorator_Blackboard.cpp
--- a/Engine/Source/Runtime/AIModule/Private/BehaviorTree/Decorators/BTDecorator_Blackboard.cpp
+++ b/Engine/Source/Runtime/AIModule/Private/BehaviorTree/Decorators/BTDecorator_Blackboard.cpp
@@ -4,6 +4,7 @@
 #include "BehaviorTree/Decorators/BTDecorator_Blackboard.h"
 #include "BehaviorTree/BlackboardComponent.h"
 #include "BehaviorTree/Blackboard/BlackboardKeyAllTypes.h"
+#include "VisualLog.h"
 
 UBTDecorator_Blackboard::UBTDecorator_Blackboard(const class FPostConstructInitializeProperties& PCIP) : Super(PCIP)
 {
diff --git a/Engine/Source/Runtime/AIModule/Private/BehaviorTree/Tasks/BTTask_MoveTo.cpp b/Engine/Source/Runtime/AIModule/Private/BehaviorTree/Tasks/BTTask_MoveTo.cpp
--- a/Engine/Source/Runtime/AIModule/Private/BehaviorTree/Tasks/BTTask_MoveTo.cpp
+++ b/Engine/Source/Runtime/AIModule/Private/BehaviorTree/Tasks/BTTask_MoveTo.cpp
@@ -6,6 +6,7 @@
 #include "BehaviorTree/BlackboardComponent.h"
 #include "BehaviorTree/Blackboard/BlackboardKeyType_Object.h"
 #include "BehaviorTree/Blackboard/BlackboardKeyType_Vector.h"
+#include "VisualLog.h"
 
 UBTTask_MoveTo::UBTTask_MoveTo(const class FPostConstructInitializeProperties& PCIP) 
 	: Super(PCIP)
diff --git a/Engine/Source/Runtime/AIModule/Private/BehaviorTree/Tasks/BTTask_RunEQSQuery.cpp b/Engine/Source/Runtime/AIModule/Private/BehaviorTree/Tasks/BTTask_RunEQSQuery.cpp
--- a/Engine/Source/Runtime/AIModule/Private/BehaviorTree/Tasks/BTTask_RunEQSQuery.cpp
+++ b/Engine/Source/Runtime/AIModule/Private/BehaviorTree/Tasks/BTTask_RunEQSQuery.cpp
@@ -5,6 +5,7 @@
 #include "EnvironmentQuery/Items/EnvQueryItemType.h"
 #include "EnvironmentQuery/EnvQueryManager.h"
 #include "EnvironmentQuery/EnvQuery.h"
+#include "VisualLog.h"
 
 
 UBTTask_RunEQSQuery::UBTTask_RunEQSQuery(const class FPostConstructInitializeProperties& PCIP) : Super(PCIP)
diff --git a/Engine/Source/Runtime/AIModule/Private/EnvironmentQuery/EQSRenderingComponent.cpp b/Engine/Source/Runtime/AIModule/Private/EnvironmentQuery/EQSRenderingComponent.cpp
--- a/Engine/Source/Runtime/AIModule/Private/EnvironmentQuery/EQSRenderingComponent.cpp
+++ b/Engine/Source/Runtime/AIModule/Private/EnvironmentQuery/EQSRenderingComponent.cpp
@@ -5,6 +5,7 @@
 #include "EnvironmentQuery/EQSRenderingComponent.h"
 #include "EnvironmentQuery/EQSQueryResultSourceInterface.h"
 #include "EnvironmentQuery/Items/EnvQueryItemType_VectorBase.h"
+#include "VisualLog.h"
 
 static const int32 EQSMaxItemsDrawn = 10000;
 
diff --git a/Engine/Source/Runtime/AIModule/Private/EnvironmentQuery/Generators/EnvQueryGenerator_OnCircle.cpp b/Engine/Source/Runtime/AIModule/Private/EnvironmentQuery/Generators/EnvQueryGenerator_OnCircle.cpp
--- a/Engine/Source/Runtime/AIModule/Private/EnvironmentQuery/Generators/EnvQueryGenerator_OnCircle.cpp
+++ b/Engine/Source/Runtime/AIModule/Private/EnvironmentQuery/Generators/EnvQueryGenerator_OnCircle.cpp
@@ -4,6 +4,7 @@
 #include "EnvironmentQuery/Contexts/EnvQueryContext_Querier.h"
 #include "EnvironmentQuery/Items/EnvQueryItemType_Point.h"
 #include "EnvironmentQuery/Generators/EnvQueryGenerator_OnCircle.h"
+#include "VisualLog.h"
 
 #define LOCTEXT_NAMESPACE "EnvQueryGenerator"
 
diff --git a/Engine/Source/Runtime/AIModule/Private/Navigation/CrowdFollowingComponent.cpp b/Engine/Source/Runtime/AIModule/Private/Navigation/CrowdFollowingComponent.cpp
--- a/Engine/Source/Runtime/AIModule/Private/Navigation/CrowdFollowingComponent.cpp
+++ b/Engine/Source/Runtime/AIModule/Private/Navigation/CrowdFollowingComponent.cpp
@@ -5,6 +5,7 @@
 #include "Navigation/CrowdManager.h"
 #include "Navigation/NavigationComponent.h"
 #include "AI/Navigation/SmartNavLinkComponent.h"
+#include "VisualLog.h"
 
 DEFINE_LOG_CATEGORY(LogCrowdFollowing);
 
diff --git a/Engine/Source/Runtime/AIModule/Private/Navigation/CrowdManager.cpp b/Engine/Source/Runtime/AIModule/Private/Navigation/CrowdManager.cpp
--- a/Engine/Source/Runtime/AIModule/Private/Navigation/CrowdManager.cpp
+++ b/Engine/Source/Runtime/AIModule/Private/Navigation/CrowdManager.cpp
@@ -11,6 +11,7 @@
 #include "Navigation/CrowdManager.h"
 #include "Navigation/CrowdFollowingComponent.h"
 #include "Navigation/CrowdAgentInterface.h"
+#include "VisualLog.h"
 
 DECLARE_STATS_GROUP(TEXT("Crowd"), STATGROUP_AICrowd, STATCAT_Advanced);
 
diff --git a/Engine/Source/Runtime/AIModule/Private/Navigation/PathFollowingComponent.cpp b/Engine/Source/Runtime/AIModule/Private/Navigation/PathFollowingComponent.cpp
--- a/Engine/Source/Runtime/AIModule/Private/Navigation/PathFollowingComponent.cpp
+++ b/Engine/Source/Runtime/AIModule/Private/Navigation/PathFollowingComponent.cpp
@@ -5,8 +5,8 @@
 #	include "Detour/DetourNavMeshQuery.h"
 #endif
 #include "Navigation/NavigationComponent.h"
-
 #include "AI/Navigation/SmartNavLinkComponent.h"
+#include "VisualLog.h"
 
 #define USE_PHYSIC_FOR_VISIBILITY_TESTS 1 // Physic will be used for visibility tests if set or only raycasts on navmesh if not
 
diff --git a/Engine/Source/Runtime/Engine/Classes/Engine/Brush.h b/Engine/Source/Runtime/Engine/Classes/Engine/Brush.h
--- a/Engine/Source/Runtime/Engine/Classes/Engine/Brush.h
+++ b/Engine/Source/Runtime/Engine/Classes/Engine/Brush.h
@@ -5,6 +5,7 @@
 //=============================================================================
 
 #pragma once
+#include "Model.h"
 #include "Brush.generated.h"
 
 //-----------------------------------------------------------------------------
diff --git a/Engine/Source/Runtime/Engine/Classes/Engine/LightMapTexture2D.h b/Engine/Source/Runtime/Engine/Classes/Engine/LightMapTexture2D.h
--- a/Engine/Source/Runtime/Engine/Classes/Engine/LightMapTexture2D.h
+++ b/Engine/Source/Runtime/Engine/Classes/Engine/LightMapTexture2D.h
@@ -5,6 +5,7 @@
  */
 
 #pragma once
+#include "Engine/Texture2D.h"
 #include "LightMapTexture2D.generated.h"
 
 /**
diff --git a/Engine/Source/Runtime/Engine/Classes/Engine/World.h b/Engine/Source/Runtime/Engine/Classes/Engine/World.h
--- a/Engine/Source/Runtime/Engine/Classes/Engine/World.h
+++ b/Engine/Source/Runtime/Engine/Classes/Engine/World.h
@@ -14,6 +14,11 @@
 
 class FPhysScene;
 class FSceneViewFamily;
+class IInterface_PostProcessVolume;
+class ULocalPlayer;
+class UGameViewportClient;
+class ABrush;
+class UModel;
 
 template<typename,typename> class TOctree;
 
diff --git a/Engine/Source/Runtime/Engine/Private/AI/Navigation/NavigationSystemHelpers.cpp b/Engine/Source/Runtime/Engine/Private/AI/Navigation/NavigationSystemHelpers.cpp
--- a/Engine/Source/Runtime/Engine/Private/AI/Navigation/NavigationSystemHelpers.cpp
+++ b/Engine/Source/Runtime/Engine/Private/AI/Navigation/NavigationSystemHelpers.cpp
@@ -7,6 +7,7 @@
 #include "RecastNavMeshGenerator.h"
 #endif // WITH_RECAST
 #include "AI/NavigationSystemHelpers.h"
+#include "VisualLog.h"
 
 namespace NavigationHelper
 {
diff --git a/Engine/Source/Runtime/Engine/Private/AI/Navigation/PImplRecastNavMesh.cpp b/Engine/Source/Runtime/Engine/Private/AI/Navigation/PImplRecastNavMesh.cpp
--- a/Engine/Source/Runtime/Engine/Private/AI/Navigation/PImplRecastNavMesh.cpp
+++ b/Engine/Source/Runtime/Engine/Private/AI/Navigation/PImplRecastNavMesh.cpp
@@ -19,6 +19,8 @@
 #endif // WITH_NAVIGATION_GENERATOR
 
 #include "AI/Navigation/SmartNavLinkComponent.h"
+#include "VisualLog.h"
+
 
 //----------------------------------------------------------------------//
 // bunch of compile-time checks to assure types used by Recast and our
diff --git a/Engine/Source/Runtime/Engine/Private/Actor.cpp b/Engine/Source/Runtime/Engine/Private/Actor.cpp
--- a/Engine/Source/Runtime/Engine/Private/Actor.cpp
+++ b/Engine/Source/Runtime/Engine/Private/Actor.cpp
@@ -13,6 +13,7 @@
 #include "Matinee/InterpGroupInst.h"
 #include "Particles/ParticleSystemComponent.h"
 #include "AI/Navigation/NavRelevantComponent.h"
+#include "VisualLog.h"
 
 //DEFINE_LOG_CATEGORY_STATIC(LogActor, Log, All);
 DEFINE_LOG_CATEGORY(LogActor);
diff --git a/Engine/Source/Runtime/Engine/Private/Animation/AnimInstance.cpp b/Engine/Source/Runtime/Engine/Private/Animation/AnimInstance.cpp
--- a/Engine/Source/Runtime/Engine/Private/Animation/AnimInstance.cpp
+++ b/Engine/Source/Runtime/Engine/Private/Animation/AnimInstance.cpp
@@ -7,11 +7,13 @@
 #include "EnginePrivate.h"
 #include "AnimationRuntime.h"
 #include "AnimationUtils.h"
+#include "AnimTree.h"
 #include "ParticleDefinitions.h"
 #include "DisplayDebugHelpers.h"
 
 #include "MessageLog.h"
 
+
 // @todo this is here only due to circular dependency to AIModule. To be removed
 #include "AIController.h"
 #include "Navigation/PathFollowingComponent.h"
diff --git a/Engine/Source/Runtime/Engine/Private/Animation/AnimNode_BlendListBase.cpp b/Engine/Source/Runtime/Engine/Private/Animation/AnimNode_BlendListBase.cpp
--- a/Engine/Source/Runtime/Engine/Private/Animation/AnimNode_BlendListBase.cpp
+++ b/Engine/Source/Runtime/Engine/Private/Animation/AnimNode_BlendListBase.cpp
@@ -2,6 +2,7 @@
 
 #include "EnginePrivate.h"
 #include "AnimationRuntime.h"
+#include "AnimTree.h"
 
 /////////////////////////////////////////////////////
 // FAnimNode_BlendListBase
diff --git a/Engine/Source/Runtime/Engine/Private/Animation/AnimNode_LayeredBoneBlend.cpp b/Engine/Source/Runtime/Engine/Private/Animation/AnimNode_LayeredBoneBlend.cpp
--- a/Engine/Source/Runtime/Engine/Private/Animation/AnimNode_LayeredBoneBlend.cpp
+++ b/Engine/Source/Runtime/Engine/Private/Animation/AnimNode_LayeredBoneBlend.cpp
@@ -2,6 +2,7 @@
 
 #include "EnginePrivate.h"
 #include "AnimationRuntime.h"
+#include "AnimTree.h"
 
 /////////////////////////////////////////////////////
 // FAnimNode_LayeredBoneBlend
diff --git a/Engine/Source/Runtime/Engine/Private/Animation/AnimNode_StateMachine.cpp b/Engine/Source/Runtime/Engine/Private/Animation/AnimNode_StateMachine.cpp
--- a/Engine/Source/Runtime/Engine/Private/Animation/AnimNode_StateMachine.cpp
+++ b/Engine/Source/Runtime/Engine/Private/Animation/AnimNode_StateMachine.cpp
@@ -2,6 +2,7 @@
 
 #include "EnginePrivate.h"
 #include "AnimationRuntime.h"
+#include "AnimTree.h"
 
 //////////////////////////////////////////////////////////////////////////
 
diff --git a/Engine/Source/Runtime/Engine/Private/Brush.cpp b/Engine/Source/Runtime/Engine/Private/Brush.cpp
--- a/Engine/Source/Runtime/Engine/Private/Brush.cpp
+++ b/Engine/Source/Runtime/Engine/Private/Brush.cpp
@@ -6,6 +6,7 @@
 
 #include "EnginePrivate.h"
 #include "Net/UnrealNetwork.h"
+#include "Model.h"
 #include "Engine/BrushBuilder.h"
 #include "ActorEditorUtils.h"
 
diff --git a/Engine/Source/Runtime/Engine/Private/BrushComponent.cpp b/Engine/Source/Runtime/Engine/Private/BrushComponent.cpp
--- a/Engine/Source/Runtime/Engine/Private/BrushComponent.cpp
+++ b/Engine/Source/Runtime/Engine/Private/BrushComponent.cpp
@@ -5,6 +5,7 @@
 =============================================================================*/
 
 #include "EnginePrivate.h"
+#include "Model.h"
 #include "LevelUtils.h"
 
 #include "DebuggingDefines.h"
diff --git a/Engine/Source/Runtime/Engine/Private/CheatManager.cpp b/Engine/Source/Runtime/Engine/Private/CheatManager.cpp
--- a/Engine/Source/Runtime/Engine/Private/CheatManager.cpp
+++ b/Engine/Source/Runtime/Engine/Private/CheatManager.cpp
@@ -6,6 +6,7 @@
 #include "SlateReflector.h"
 #include "NavDataGenerator.h"
 #include "OnlineSubsystemUtils.h"
+#include "VisualLog.h"
 // @todo this is here only due to circular dependency to AIModule. To be removed
 #include "AIController.h"
 #include "EnvironmentQuery/EnvQuery.h"
diff --git a/Engine/Source/Runtime/Engine/Private/Controller.cpp b/Engine/Source/Runtime/Engine/Private/Controller.cpp
--- a/Engine/Source/Runtime/Engine/Private/Controller.cpp
+++ b/Engine/Source/Runtime/Engine/Private/Controller.cpp
@@ -8,6 +8,7 @@
 #include "Net/UnrealNetwork.h"
 #include "ConfigCacheIni.h"
 #include "NetworkingDistanceConstants.h"
+#include "VisualLog.h"
 
 // @todo this is here only due to circular dependency to AIModule. To be removed
 #include "Navigation/PathFollowingComponent.h"
diff --git a/Engine/Source/Runtime/Engine/Private/Debug/GameplayDebuggingControllerComponent.cpp b/Engine/Source/Runtime/Engine/Private/Debug/GameplayDebuggingControllerComponent.cpp
--- a/Engine/Source/Runtime/Engine/Private/Debug/GameplayDebuggingControllerComponent.cpp
+++ b/Engine/Source/Runtime/Engine/Private/Debug/GameplayDebuggingControllerComponent.cpp
@@ -3,6 +3,7 @@
 #include "EnginePrivate.h"
 #include "Net/UnrealNetwork.h"
 #include "BehaviorTreeDelegates.h"
+#include "VisualLog.h"
 
 #if WITH_EDITOR
 #include "UnrealEd.h"
diff --git a/Engine/Source/Runtime/Engine/Private/Debug/LogVisualizerHUD.cpp b/Engine/Source/Runtime/Engine/Private/Debug/LogVisualizerHUD.cpp
--- a/Engine/Source/Runtime/Engine/Private/Debug/LogVisualizerHUD.cpp
+++ b/Engine/Source/Runtime/Engine/Private/Debug/LogVisualizerHUD.cpp
@@ -3,6 +3,7 @@
 #include "EnginePrivate.h"
 #include "Debug/LogVisualizerHUD.h"
 #include "Debug/LogVisualizerCameraController.h"
+#include "VisualLog.h"
 
 #define LOCTEXT_NAMESPACE "ALogVisualizerHUD"
 
diff --git a/Engine/Source/Runtime/Engine/Private/DestructibleMesh.cpp b/Engine/Source/Runtime/Engine/Private/DestructibleMesh.cpp
--- a/Engine/Source/Runtime/Engine/Private/DestructibleMesh.cpp
+++ b/Engine/Source/Runtime/Engine/Private/DestructibleMesh.cpp
@@ -5,6 +5,7 @@
 =============================================================================*/
 
 #include "EnginePrivate.h"
+#include "StaticMeshResources.h"
 #include "Engine/DestructibleFractureSettings.h"
 #include "PhysicsEngine/PhysXSupport.h"
 
diff --git a/Engine/Source/Runtime/Engine/Private/ErrorChecking.cpp b/Engine/Source/Runtime/Engine/Private/ErrorChecking.cpp
--- a/Engine/Source/Runtime/Engine/Private/ErrorChecking.cpp
+++ b/Engine/Source/Runtime/Engine/Private/ErrorChecking.cpp
@@ -10,6 +10,7 @@
 #include "MessageLog.h"
 #include "UObjectToken.h"
 #include "MapErrors.h"
+#include "Model.h"
 
 #define LOCTEXT_NAMESPACE "ErrorChecking"
 
diff --git a/Engine/Source/Runtime/Engine/Private/FoliageComponent.cpp b/Engine/Source/Runtime/Engine/Private/FoliageComponent.cpp
--- a/Engine/Source/Runtime/Engine/Private/FoliageComponent.cpp
+++ b/Engine/Source/Runtime/Engine/Private/FoliageComponent.cpp
@@ -5,6 +5,7 @@
 =============================================================================*/
 
 #include "EnginePrivate.h"
+#include "StaticMeshResources.h"
 #include "Engine/InteractiveFoliageActor.h"
 #include "Components/InteractiveFoliageComponent.h"
 
diff --git a/Engine/Source/Runtime/Engine/Private/HModel.cpp b/Engine/Source/Runtime/Engine/Private/HModel.cpp
--- a/Engine/Source/Runtime/Engine/Private/HModel.cpp
+++ b/Engine/Source/Runtime/Engine/Private/HModel.cpp
@@ -5,6 +5,7 @@
 =============================================================================*/
 
 #include "EnginePrivate.h"
+#include "Model.h"
 #include "HModel.h"
 #include "Raster.h"
 
diff --git a/Engine/Source/Runtime/Engine/Private/InstancedFoliage.cpp b/Engine/Source/Runtime/Engine/Private/InstancedFoliage.cpp
--- a/Engine/Source/Runtime/Engine/Private/InstancedFoliage.cpp
+++ b/Engine/Source/Runtime/Engine/Private/InstancedFoliage.cpp
@@ -5,6 +5,7 @@
 =============================================================================*/
 
 #include "EnginePrivate.h"
+#include "StaticMeshResources.h"
 #include "MessageLog.h"
 #include "UObjectToken.h"
 #include "MapErrors.h"
diff --git a/Engine/Source/Runtime/Engine/Private/InstancedStaticMesh.cpp b/Engine/Source/Runtime/Engine/Private/InstancedStaticMesh.cpp
--- a/Engine/Source/Runtime/Engine/Private/InstancedStaticMesh.cpp
+++ b/Engine/Source/Runtime/Engine/Private/InstancedStaticMesh.cpp
@@ -6,6 +6,7 @@
 
 #include "EnginePrivate.h"
 #include "ShaderParameters.h"
+#include "StaticLighting.h"
 
 #include "Components/InteractiveFoliageComponent.h"
 #include "Components/SplineMeshComponent.h"
@@ -27,6 +28,7 @@
 #include "OnlineSubsystemUtils.h"
 #include "AI/Navigation/RecastHelpers.h"
 
+#include "StaticMeshResources.h"
 #include "StaticMeshLight.h"
 #include "SpeedTreeWind.h"
 #include "ComponentInstanceDataCache.h"
diff --git a/Engine/Source/Runtime/Engine/Private/Landscape/LandscapeRender.cpp b/Engine/Source/Runtime/Engine/Private/Landscape/LandscapeRender.cpp
--- a/Engine/Source/Runtime/Engine/Private/Landscape/LandscapeRender.cpp
+++ b/Engine/Source/Runtime/Engine/Private/Landscape/LandscapeRender.cpp
@@ -13,6 +13,7 @@ LandscapeRender.cpp: New terrain rendering
 #include "LevelUtils.h"
 #include "MaterialCompiler.h"
 #include "Landscape/LandscapeMaterialInstanceConstant.h"
+#include "RawIndexBuffer.h"
 
 IMPLEMENT_UNIFORM_BUFFER_STRUCT(FLandscapeUniformShaderParameters,TEXT("LandscapeParameters"));
 
diff --git a/Engine/Source/Runtime/Engine/Private/Landscape/LandscapeSplines.cpp b/Engine/Source/Runtime/Engine/Private/Landscape/LandscapeSplines.cpp
--- a/Engine/Source/Runtime/Engine/Private/Landscape/LandscapeSplines.cpp
+++ b/Engine/Source/Runtime/Engine/Private/Landscape/LandscapeSplines.cpp
@@ -5,6 +5,7 @@
 =============================================================================*/
 
 #include "EnginePrivate.h"
+#include "StaticMeshResources.h"
 #include "AI/Navigation/RecastHelpers.h"
 #include "Landscape/LandscapeSplineProxies.h"
 #include "Landscape/LandscapeSplinesComponent.h"
diff --git a/Engine/Source/Runtime/Engine/Private/Level.cpp b/Engine/Source/Runtime/Engine/Private/Level.cpp
--- a/Engine/Source/Runtime/Engine/Private/Level.cpp
+++ b/Engine/Source/Runtime/Engine/Private/Level.cpp
@@ -6,8 +6,11 @@ Level.cpp: Level-related functions
 
 #include "EnginePrivate.h"
 #include "Engine/LevelScriptBlueprint.h"
+#include "Engine/WorldComposition.h"
 #include "Sound/SoundNodeWave.h"
 #include "Net/UnrealNetwork.h"
+#include "Model.h"
+#include "StaticLighting.h"
 #include "SoundDefinitions.h"
 #include "PrecomputedLightVolume.h"
 #include "TickTaskManagerInterface.h"
diff --git a/Engine/Source/Runtime/Engine/Private/Light.cpp b/Engine/Source/Runtime/Engine/Private/Light.cpp
--- a/Engine/Source/Runtime/Engine/Private/Light.cpp
+++ b/Engine/Source/Runtime/Engine/Private/Light.cpp
@@ -2,6 +2,7 @@
 
 
 #include "EnginePrivate.h"
+#include "LightingBuildOptions.h"
 #include "Net/UnrealNetwork.h"
 
 ALight::ALight(const class FPostConstructInitializeProperties& PCIP)
diff --git a/Engine/Source/Runtime/Engine/Private/LightMap.cpp b/Engine/Source/Runtime/Engine/Private/LightMap.cpp
--- a/Engine/Source/Runtime/Engine/Private/LightMap.cpp
+++ b/Engine/Source/Runtime/Engine/Private/LightMap.cpp
@@ -6,6 +6,8 @@
 
 #include "EnginePrivate.h"
 #include "TargetPlatform.h"
+#include "StaticLighting.h"
+
 
 DEFINE_LOG_CATEGORY_STATIC(LogLightMap, Log, All);
 
diff --git a/Engine/Source/Runtime/Engine/Private/Model.cpp b/Engine/Source/Runtime/Engine/Private/Model.cpp
--- a/Engine/Source/Runtime/Engine/Private/Model.cpp
+++ b/Engine/Source/Runtime/Engine/Private/Model.cpp
@@ -5,6 +5,7 @@
 =============================================================================*/
 
 #include "EnginePrivate.h"
+#include "Model.h"
 #include "MeshBuild.h"
 
 float UModel::BSPTexelScale = 100.0f;
diff --git a/Engine/Source/Runtime/Engine/Private/ModelCollision.cpp b/Engine/Source/Runtime/Engine/Private/ModelCollision.cpp
--- a/Engine/Source/Runtime/Engine/Private/ModelCollision.cpp
+++ b/Engine/Source/Runtime/Engine/Private/ModelCollision.cpp
@@ -5,6 +5,7 @@
 =============================================================================*/
 
 #include "EnginePrivate.h"
+#include "Model.h"
 
 
 /*---------------------------------------------------------------------------------------
diff --git a/Engine/Source/Runtime/Engine/Private/ModelComponent.cpp b/Engine/Source/Runtime/Engine/Private/ModelComponent.cpp
--- a/Engine/Source/Runtime/Engine/Private/ModelComponent.cpp
+++ b/Engine/Source/Runtime/Engine/Private/ModelComponent.cpp
@@ -5,6 +5,7 @@
 =============================================================================*/
 
 #include "EnginePrivate.h"
+#include "Model.h"
 
 FModelElement::FModelElement(UModelComponent* InComponent,UMaterialInterface* InMaterial):
 	Component(InComponent),
diff --git a/Engine/Source/Runtime/Engine/Private/ModelLight.cpp b/Engine/Source/Runtime/Engine/Private/ModelLight.cpp
--- a/Engine/Source/Runtime/Engine/Private/ModelLight.cpp
+++ b/Engine/Source/Runtime/Engine/Private/ModelLight.cpp
@@ -8,6 +8,7 @@
 #include "TextureLayout.h"
 #include "ModelLight.h"
 #include "Collision.h"
+#include "Model.h"
 
 //
 //	Static vars
diff --git a/Engine/Source/Runtime/Engine/Private/ModelRender.cpp b/Engine/Source/Runtime/Engine/Private/ModelRender.cpp
--- a/Engine/Source/Runtime/Engine/Private/ModelRender.cpp
+++ b/Engine/Source/Runtime/Engine/Private/ModelRender.cpp
@@ -6,6 +6,7 @@
 
 #include "EnginePrivate.h"
 #include "LevelUtils.h"
+#include "Model.h"
 #include "HModel.h"
 
 
diff --git a/Engine/Source/Runtime/Engine/Private/MorphTools.cpp b/Engine/Source/Runtime/Engine/Private/MorphTools.cpp
--- a/Engine/Source/Runtime/Engine/Private/MorphTools.cpp
+++ b/Engine/Source/Runtime/Engine/Private/MorphTools.cpp
@@ -5,9 +5,11 @@
 =============================================================================*/ 
 
 #include "EnginePrivate.h"
+#include "RawIndexBuffer.h"
 #include "MeshBuild.h"
 #include "PhysicsEngine/PhysXSupport.h"
 
+
 /** compare based on base mesh source vertex indices */
 struct FCompareVertexAnimDeltas
 {
diff --git a/Engine/Source/Runtime/Engine/Private/Particles/ParticleComponents.cpp b/Engine/Source/Runtime/Engine/Private/Particles/ParticleComponents.cpp
--- a/Engine/Source/Runtime/Engine/Private/Particles/ParticleComponents.cpp
+++ b/Engine/Source/Runtime/Engine/Private/Particles/ParticleComponents.cpp
@@ -5,6 +5,7 @@
 =============================================================================*/
 
 #include "EnginePrivate.h"
+#include "StaticMeshResources.h"
 #include "ParticleDefinitions.h"
 #include "Particles/EmitterCameraLensEffectBase.h"
 #include "LevelUtils.h"
diff --git a/Engine/Source/Runtime/Engine/Private/Particles/ParticleEmitterInstances.cpp b/Engine/Source/Runtime/Engine/Private/Particles/ParticleEmitterInstances.cpp
--- a/Engine/Source/Runtime/Engine/Private/Particles/ParticleEmitterInstances.cpp
+++ b/Engine/Source/Runtime/Engine/Private/Particles/ParticleEmitterInstances.cpp
@@ -5,6 +5,7 @@
 =============================================================================*/
 
 #include "EnginePrivate.h"
+#include "StaticMeshResources.h"
 #include "ParticleDefinitions.h"
 #include "LevelUtils.h"
 #include "FXSystem.h"
diff --git a/Engine/Source/Runtime/Engine/Private/Particles/ParticleModules_Location.cpp b/Engine/Source/Runtime/Engine/Private/Particles/ParticleModules_Location.cpp
--- a/Engine/Source/Runtime/Engine/Private/Particles/ParticleModules_Location.cpp
+++ b/Engine/Source/Runtime/Engine/Private/Particles/ParticleModules_Location.cpp
@@ -7,6 +7,7 @@
 
 #include "EnginePrivate.h"
 #include "ParticleDefinitions.h"
+#include "RawIndexBuffer.h"
 #include "../DistributionHelpers.h"
 #include "Particles/Location/ParticleModuleLocation.h"
 #include "Particles/Location/ParticleModuleLocationBoneSocket.h"
diff --git a/Engine/Source/Runtime/Engine/Private/Particles/ParticleSystemRender.cpp b/Engine/Source/Runtime/Engine/Private/Particles/ParticleSystemRender.cpp
--- a/Engine/Source/Runtime/Engine/Private/Particles/ParticleSystemRender.cpp
+++ b/Engine/Source/Runtime/Engine/Private/Particles/ParticleSystemRender.cpp
@@ -4,6 +4,7 @@
 	ParticleSystemRender.cpp: Particle system rendering functions.
 =============================================================================*/
 #include "EnginePrivate.h"
+#include "StaticMeshResources.h"
 #include "ParticleDefinitions.h"
 #include "DiagnosticTable.h"
 #include "ParticleResources.h"
diff --git a/Engine/Source/Runtime/Engine/Private/Pawn.cpp b/Engine/Source/Runtime/Engine/Private/Pawn.cpp
--- a/Engine/Source/Runtime/Engine/Private/Pawn.cpp
+++ b/Engine/Source/Runtime/Engine/Private/Pawn.cpp
@@ -12,6 +12,7 @@
 #include "ParticleDefinitions.h"
 #include "DisplayDebugHelpers.h"
 #include "NetworkingDistanceConstants.h"
+#include "VisualLog.h"
 
 DEFINE_LOG_CATEGORY(LogDamage);
 DEFINE_LOG_CATEGORY_STATIC(LogPawn, Warning, All);
diff --git a/Engine/Source/Runtime/Engine/Private/PhysicsEngine/PhysAnim.cpp b/Engine/Source/Runtime/Engine/Private/PhysicsEngine/PhysAnim.cpp
--- a/Engine/Source/Runtime/Engine/Private/PhysicsEngine/PhysAnim.cpp
+++ b/Engine/Source/Runtime/Engine/Private/PhysicsEngine/PhysAnim.cpp
@@ -5,6 +5,7 @@
 =============================================================================*/ 
 
 #include "EnginePrivate.h"
+#include "AnimTree.h"
 
 #if WITH_PHYSX
 	#include "PhysXSupport.h"
diff --git a/Engine/Source/Runtime/Engine/Private/PhysicsEngine/PhysUtils.cpp b/Engine/Source/Runtime/Engine/Private/PhysicsEngine/PhysUtils.cpp
--- a/Engine/Source/Runtime/Engine/Private/PhysicsEngine/PhysUtils.cpp
+++ b/Engine/Source/Runtime/Engine/Private/PhysicsEngine/PhysUtils.cpp
@@ -8,6 +8,7 @@
  ===========================================================================*/
 
 #include "EnginePrivate.h"
+#include "Model.h"
 #include "PhysXSupport.h"
 
 /* *********************************************************************** */
diff --git a/Engine/Source/Runtime/Engine/Private/PlayerController.cpp b/Engine/Source/Runtime/Engine/Private/PlayerController.cpp
--- a/Engine/Source/Runtime/Engine/Private/PlayerController.cpp
+++ b/Engine/Source/Runtime/Engine/Private/PlayerController.cpp
@@ -3,6 +3,7 @@
 #include "EnginePrivate.h"
 #include "Engine/Console.h"
 #include "Net/UnrealNetwork.h"
+#include "Engine/WorldComposition.h"
 #include "GameFramework/GameNetworkManager.h"
 #include "ConfigCacheIni.h"
 #include "SoundDefinitions.h"
diff --git a/Engine/Source/Runtime/Engine/Private/Polygon.cpp b/Engine/Source/Runtime/Engine/Private/Polygon.cpp
--- a/Engine/Source/Runtime/Engine/Private/Polygon.cpp
+++ b/Engine/Source/Runtime/Engine/Private/Polygon.cpp
@@ -5,6 +5,8 @@
 =============================================================================*/
 
 #include "EnginePrivate.h"
+#include "Model.h"
+
 #if WITH_EDITOR
 #include "GeomTools.h"
 #endif
diff --git a/Engine/Source/Runtime/Engine/Private/PoseableMeshComponent.cpp b/Engine/Source/Runtime/Engine/Private/PoseableMeshComponent.cpp
--- a/Engine/Source/Runtime/Engine/Private/PoseableMeshComponent.cpp
+++ b/Engine/Source/Runtime/Engine/Private/PoseableMeshComponent.cpp
@@ -5,6 +5,7 @@
 =============================================================================*/
 
 #include "EnginePrivate.h"
+#include "AnimTree.h"
 
 UPoseableMeshComponent::UPoseableMeshComponent(const class FPostConstructInitializeProperties& PCIP)
 	: Super(PCIP)
diff --git a/Engine/Source/Runtime/Engine/Private/PositionVertexBuffer.cpp b/Engine/Source/Runtime/Engine/Private/PositionVertexBuffer.cpp
--- a/Engine/Source/Runtime/Engine/Private/PositionVertexBuffer.cpp
+++ b/Engine/Source/Runtime/Engine/Private/PositionVertexBuffer.cpp
@@ -2,6 +2,7 @@
 
 #include "EnginePrivate.h"
 
+#include "StaticMeshResources.h"
 #include "StaticMeshVertexData.h"
 
 /*-----------------------------------------------------------------------------
diff --git a/Engine/Source/Runtime/Engine/Private/PrimitiveDrawingUtils.cpp b/Engine/Source/Runtime/Engine/Private/PrimitiveDrawingUtils.cpp
--- a/Engine/Source/Runtime/Engine/Private/PrimitiveDrawingUtils.cpp
+++ b/Engine/Source/Runtime/Engine/Private/PrimitiveDrawingUtils.cpp
@@ -2,6 +2,7 @@
 
 #include "EnginePrivate.h"
 #include "SceneManagement.h"
+#include "StaticMeshResources.h"
 
 /** Emits draw events for a given FMeshBatch and the FPrimitiveSceneProxy corresponding to that mesh element. */
 void EmitMeshDrawEvents(const FPrimitiveSceneProxy* PrimitiveSceneProxy, const FMeshBatch& Mesh)
diff --git a/Engine/Source/Runtime/Engine/Private/RawIndexBuffer.cpp b/Engine/Source/Runtime/Engine/Private/RawIndexBuffer.cpp
--- a/Engine/Source/Runtime/Engine/Private/RawIndexBuffer.cpp
+++ b/Engine/Source/Runtime/Engine/Private/RawIndexBuffer.cpp
@@ -5,6 +5,7 @@
 =============================================================================*/
 
 #include "EnginePrivate.h"
+#include "RawIndexBuffer.h"
 
 #if WITH_EDITOR
 #include "MeshUtilities.h"
diff --git a/Engine/Source/Runtime/Engine/Private/SceneManagement.cpp b/Engine/Source/Runtime/Engine/Private/SceneManagement.cpp
--- a/Engine/Source/Runtime/Engine/Private/SceneManagement.cpp
+++ b/Engine/Source/Runtime/Engine/Private/SceneManagement.cpp
@@ -1,6 +1,7 @@
 // Copyright 1998-2014 Epic Games, Inc. All Rights Reserved.
 
 #include "EnginePrivate.h"
+#include "StaticMeshResources.h"
 #include "../../Renderer/Private/ScenePrivate.h"
 
 /*
diff --git a/Engine/Source/Runtime/Engine/Private/SkeletalMesh.cpp b/Engine/Source/Runtime/Engine/Private/SkeletalMesh.cpp
--- a/Engine/Source/Runtime/Engine/Private/SkeletalMesh.cpp
+++ b/Engine/Source/Runtime/Engine/Private/SkeletalMesh.cpp
@@ -17,6 +17,7 @@
 #include "UObjectToken.h"
 #include "MapErrors.h"
 #include "SkeletalRenderGPUSkin.h"
+#include "RawIndexBuffer.h"
 
 #if WITH_EDITOR
 #include "MeshUtilities.h"
diff --git a/Engine/Source/Runtime/Engine/Private/SkeletalMeshComponent.cpp b/Engine/Source/Runtime/Engine/Private/SkeletalMeshComponent.cpp
--- a/Engine/Source/Runtime/Engine/Private/SkeletalMeshComponent.cpp
+++ b/Engine/Source/Runtime/Engine/Private/SkeletalMeshComponent.cpp
@@ -13,6 +13,7 @@
 #include "AnimationUtils.h"
 #include "AnimationRuntime.h"
 #include "PhysXASync.h"
+#include "AnimTree.h"
 
 #include "Particles/ParticleSystemComponent.h"
 
diff --git a/Engine/Source/Runtime/Engine/Private/SkeletalMeshMerge.cpp b/Engine/Source/Runtime/Engine/Private/SkeletalMeshMerge.cpp
--- a/Engine/Source/Runtime/Engine/Private/SkeletalMeshMerge.cpp
+++ b/Engine/Source/Runtime/Engine/Private/SkeletalMeshMerge.cpp
@@ -5,6 +5,7 @@
 =============================================================================*/
 
 #include "EnginePrivate.h"
+#include "RawIndexBuffer.h"
 #include "SkeletalMeshMerge.h"
 
 /*-----------------------------------------------------------------------------
diff --git a/Engine/Source/Runtime/Engine/Private/SkinnedMeshComponent.cpp b/Engine/Source/Runtime/Engine/Private/SkinnedMeshComponent.cpp
--- a/Engine/Source/Runtime/Engine/Private/SkinnedMeshComponent.cpp
+++ b/Engine/Source/Runtime/Engine/Private/SkinnedMeshComponent.cpp
@@ -10,6 +10,7 @@
 #include "SkeletalRenderCPUSkin.h"
 #include "SkeletalRenderGPUSkin.h"
 #include "AnimationUtils.h"
+#include "AnimTree.h"
 
 DEFINE_LOG_CATEGORY_STATIC(LogSkinnedMeshComp, Log, All);
 
diff --git a/Engine/Source/Runtime/Engine/Private/StaticMesh.cpp b/Engine/Source/Runtime/Engine/Private/StaticMesh.cpp
--- a/Engine/Source/Runtime/Engine/Private/StaticMesh.cpp
+++ b/Engine/Source/Runtime/Engine/Private/StaticMesh.cpp
@@ -5,6 +5,7 @@
 =============================================================================*/
 
 #include "EnginePrivate.h"
+#include "StaticMeshResources.h"
 #include "MeshBuild.h"
 #include "GenericOctree.h"
 #include "TessellationRendering.h"
diff --git a/Engine/Source/Runtime/Engine/Private/StaticMeshActor.cpp b/Engine/Source/Runtime/Engine/Private/StaticMeshActor.cpp
--- a/Engine/Source/Runtime/Engine/Private/StaticMeshActor.cpp
+++ b/Engine/Source/Runtime/Engine/Private/StaticMeshActor.cpp
@@ -6,6 +6,7 @@
 
 #include "EnginePrivate.h"
 #include "MessageLog.h"
+#include "StaticMeshResources.h"
 #include "UObjectToken.h"
 #include "MapErrors.h"
 
diff --git a/Engine/Source/Runtime/Engine/Private/StaticMeshBuild.cpp b/Engine/Source/Runtime/Engine/Private/StaticMeshBuild.cpp
--- a/Engine/Source/Runtime/Engine/Private/StaticMeshBuild.cpp
+++ b/Engine/Source/Runtime/Engine/Private/StaticMeshBuild.cpp
@@ -5,6 +5,7 @@
 =============================================================================*/
 
 #include "EnginePrivate.h"
+#include "StaticMeshResources.h"
 
 #if WITH_EDITORONLY_DATA
 #include "RawMesh.h"
diff --git a/Engine/Source/Runtime/Engine/Private/StaticMeshComponent.cpp b/Engine/Source/Runtime/Engine/Private/StaticMeshComponent.cpp
--- a/Engine/Source/Runtime/Engine/Private/StaticMeshComponent.cpp
+++ b/Engine/Source/Runtime/Engine/Private/StaticMeshComponent.cpp
@@ -1,6 +1,7 @@
 // Copyright 1998-2014 Epic Games, Inc. All Rights Reserved.
 
 #include "EnginePrivate.h"
+#include "StaticMeshResources.h"
 #include "Net/UnrealNetwork.h"
 #include "MessageLog.h"
 #include "UObjectToken.h"
diff --git a/Engine/Source/Runtime/Engine/Private/StaticMeshLight.cpp b/Engine/Source/Runtime/Engine/Private/StaticMeshLight.cpp
--- a/Engine/Source/Runtime/Engine/Private/StaticMeshLight.cpp
+++ b/Engine/Source/Runtime/Engine/Private/StaticMeshLight.cpp
@@ -5,6 +5,7 @@
 =============================================================================*/
 
 #include "EnginePrivate.h"
+#include "StaticMeshResources.h"
 #include "Raster.h"
 #include "StaticMeshLight.h"
 
diff --git a/Engine/Source/Runtime/Engine/Private/StaticMeshRender.cpp b/Engine/Source/Runtime/Engine/Private/StaticMeshRender.cpp
--- a/Engine/Source/Runtime/Engine/Private/StaticMeshRender.cpp
+++ b/Engine/Source/Runtime/Engine/Private/StaticMeshRender.cpp
@@ -5,6 +5,7 @@
 =============================================================================*/
 
 #include "EnginePrivate.h"
+#include "StaticMeshResources.h"
 #include "LevelUtils.h"
 #include "TessellationRendering.h"
 
diff --git a/Engine/Source/Runtime/Engine/Private/StaticMeshVertexData.h b/Engine/Source/Runtime/Engine/Private/StaticMeshVertexData.h
--- a/Engine/Source/Runtime/Engine/Private/StaticMeshVertexData.h
+++ b/Engine/Source/Runtime/Engine/Private/StaticMeshVertexData.h
@@ -2,6 +2,8 @@
 
 #pragma once
 
+#include "StaticMeshResources.h"
+
 /** The implementation of the static mesh vertex data storage type. */
 template<typename VertexDataType>
 class TStaticMeshVertexData :
diff --git a/Engine/Source/Runtime/Engine/Private/UnrealEngine.cpp b/Engine/Source/Runtime/Engine/Private/UnrealEngine.cpp
--- a/Engine/Source/Runtime/Engine/Private/UnrealEngine.cpp
+++ b/Engine/Source/Runtime/Engine/Private/UnrealEngine.cpp
@@ -8,6 +8,7 @@
 
 #include "Net/UnrealNetwork.h"
 #include "Engine/Console.h"
+#include "VisualLog.h"
 #include "FileManagerGeneric.h"
 #include "Database.h"
 #include "SkeletalMeshMerge.h"
diff --git a/Engine/Source/Runtime/Engine/Private/UnrealExporter.cpp b/Engine/Source/Runtime/Engine/Private/UnrealExporter.cpp
--- a/Engine/Source/Runtime/Engine/Private/UnrealExporter.cpp
+++ b/Engine/Source/Runtime/Engine/Private/UnrealExporter.cpp
@@ -6,6 +6,7 @@
 
 // Engine includes.
 #include "EnginePrivate.h"
+#include "Model.h"
 
 DEFINE_LOG_CATEGORY_STATIC(LogExporter, Log, All);
 
diff --git a/Engine/Source/Runtime/Engine/Private/World.cpp b/Engine/Source/Runtime/Engine/Private/World.cpp
--- a/Engine/Source/Runtime/Engine/Private/World.cpp
+++ b/Engine/Source/Runtime/Engine/Private/World.cpp
@@ -7,9 +7,11 @@
 #include "EnginePrivate.h"
 #include "Engine/LevelScriptBlueprint.h"
 #include "Engine/Console.h"
+#include "Engine/WorldComposition.h"
 #include "Matinee/MatineeActor.h"
 #include "Net/UnrealNetwork.h"
 #include "GameFramework/GameNetworkManager.h"
+#include "Model.h"
 #include "ParticleDefinitions.h"
 #include "Database.h"
 #include "Net/NetworkProfiler.h"
diff --git a/Engine/Source/Runtime/Engine/Private/WorldComposition.cpp b/Engine/Source/Runtime/Engine/Private/WorldComposition.cpp
--- a/Engine/Source/Runtime/Engine/Private/WorldComposition.cpp
+++ b/Engine/Source/Runtime/Engine/Private/WorldComposition.cpp
@@ -4,6 +4,7 @@
 	WorldComposition.cpp: UWorldComposition implementation
 =============================================================================*/
 #include "EnginePrivate.h"
+#include "Engine/WorldComposition.h"
 #include "LevelUtils.h"
 
 #if WITH_EDITOR
diff --git a/Engine/Source/Runtime/Engine/Private/WorldSettings.cpp b/Engine/Source/Runtime/Engine/Private/WorldSettings.cpp
--- a/Engine/Source/Runtime/Engine/Private/WorldSettings.cpp
+++ b/Engine/Source/Runtime/Engine/Private/WorldSettings.cpp
@@ -1,6 +1,7 @@
 // Copyright 1998-2014 Epic Games, Inc. All Rights Reserved.
 
 #include "EnginePrivate.h"
+#include "Engine/WorldComposition.h"
 #include "Net/UnrealNetwork.h"
 #include "GameFramework/GameNetworkManager.h"
 #include "SoundDefinitions.h"
diff --git a/Engine/Source/Runtime/Engine/Public/AnimEncoding.h b/Engine/Source/Runtime/Engine/Public/AnimEncoding.h
--- a/Engine/Source/Runtime/Engine/Public/AnimEncoding.h
+++ b/Engine/Source/Runtime/Engine/Public/AnimEncoding.h
@@ -49,6 +49,9 @@ struct BoneTrackPair
 typedef TArray<BoneTrackPair> BoneTrackArray;
 
 
+/** Array of FTransform using the game memory stack */
+typedef TArray< FTransform, TMemStackAllocator<> > FTransformArray;
+
 /**
  * Extracts a single BoneAtom from an Animation Sequence.
  *
diff --git a/Engine/Source/Runtime/Engine/Public/AnimTree.h b/Engine/Source/Runtime/Engine/Public/AnimTree.h
--- a/Engine/Source/Runtime/Engine/Public/AnimTree.h
+++ b/Engine/Source/Runtime/Engine/Public/AnimTree.h
@@ -33,6 +33,5 @@ DECLARE_CYCLE_STAT_EXTERN(TEXT("MeshObject Update"), STAT_MeshObjectUpdate, STAT
 DECLARE_CYCLE_STAT_EXTERN(TEXT("Update SkelMesh Bounds"), STAT_UpdateSkelMeshBounds, STATGROUP_Anim, );
 DECLARE_CYCLE_STAT_EXTERN(TEXT("BlendInPhysics"), STAT_BlendInPhysics, STATGROUP_Anim, );
 
-/** Array of FTransform using the game memory stack */
-typedef TArray< FTransform, TMemStackAllocator<> > FTransformArray;
+
 
diff --git a/Engine/Source/Runtime/Engine/Public/Components.h b/Engine/Source/Runtime/Engine/Public/Components.h
--- a/Engine/Source/Runtime/Engine/Public/Components.h
+++ b/Engine/Source/Runtime/Engine/Public/Components.h
@@ -2,6 +2,8 @@
 
 #pragma once
 
+#include "RenderUtils.h"
+
 /*=============================================================================
 	Components.h: Forward declarations of object components of actors
 =============================================================================*/
diff --git a/Engine/Source/Runtime/Engine/Public/Engine.h b/Engine/Source/Runtime/Engine/Public/Engine.h
--- a/Engine/Source/Runtime/Engine/Public/Engine.h
+++ b/Engine/Source/Runtime/Engine/Public/Engine.h
@@ -419,27 +419,15 @@
 #include "Animation/VertexAnim/MorphTarget.h"
 #include "Animation/VertexAnim/VertexAnimation.h"
 #include "Engine/World.h"
-#include "Engine/WorldComposition.h"
 
-#include "VisualLog.h"
 #include "SystemSettings.h"					// Scalability options.
 #include "SceneManagement.h"				// Scene management.
-#include "StaticLighting.h"					// Static lighting definitions.
-#include "LightMap.h"						// Light-maps.
-#include "ShadowMap.h"
-#include "Model.h"							// Model class.
 
 #include "NavDataGenerator.h"
 #include "NavLinkRenderingProxy.h"
 #include "NavigationModifier.h"
 #include "NavigationOctree.h"
-#include "StaticMeshResources.h"			// Static meshes.
-#include "AnimTree.h"						// Animation.
-#include "SkeletalMeshTypes.h"				// Skeletal animated mesh.
-#include "Animation/SkeletalMeshActor.h"
 #include "ContentStreaming.h"				// Content streaming class definitions.
-#include "LightingBuildOptions.h"			// Definition of lighting build option struct.
-#include "PixelFormat.h"
 #include "PhysicsPublic.h"
 #include "ComponentReregisterContext.h"	
 #include "DrawDebugHelpers.h"
diff --git a/Engine/Source/Runtime/Engine/Public/Landscape/LandscapeLight.h b/Engine/Source/Runtime/Engine/Public/Landscape/LandscapeLight.h
--- a/Engine/Source/Runtime/Engine/Public/Landscape/LandscapeLight.h
+++ b/Engine/Source/Runtime/Engine/Public/Landscape/LandscapeLight.h
@@ -7,6 +7,8 @@ LandscapeLight.h: Static lighting for LandscapeComponents
 #ifndef _LANDSCAPELIGHT_H
 #define _LANDSCAPELIGHT_H
 
+#include "StaticLighting.h"
+
 class ULandscapeComponent;
 
 /** A texture mapping for landscapes */
diff --git a/Engine/Source/Runtime/Engine/Public/LightMap.h b/Engine/Source/Runtime/Engine/Public/LightMap.h
--- a/Engine/Source/Runtime/Engine/Public/LightMap.h
+++ b/Engine/Source/Runtime/Engine/Public/LightMap.h
@@ -7,6 +7,7 @@
 #pragma once
 
 #include "ChunkedArray.h"
+#include "SceneManagement.h"
 #include "Engine/LightMapTexture2D.h"
 
 /** Whether to use bilinear filtering on lightmaps */
diff --git a/Engine/Source/Runtime/Engine/Public/LocalVertexFactory.h b/Engine/Source/Runtime/Engine/Public/LocalVertexFactory.h
--- a/Engine/Source/Runtime/Engine/Public/LocalVertexFactory.h
+++ b/Engine/Source/Runtime/Engine/Public/LocalVertexFactory.h
@@ -2,6 +2,9 @@
 
 #pragma once
 
+#include "VertexFactory.h"
+#include "Components.h"
+
 /*=============================================================================
 	LocalVertexFactory.h: Local vertex factory definitions.
 =============================================================================*/
diff --git a/Engine/Source/Runtime/Engine/Public/Model.h b/Engine/Source/Runtime/Engine/Public/Model.h
--- a/Engine/Source/Runtime/Engine/Public/Model.h
+++ b/Engine/Source/Runtime/Engine/Public/Model.h
@@ -9,6 +9,15 @@
 #include "Engine.h"
 #include "RawIndexBuffer.h"
 #include "LocalVertexFactory.h"
+#include "RenderUtils.h"
+#include "LightMap.h"
+#include "ShadowMap.h"
+
+struct FStaticLightingVertex;
+class ABrush;
+class UModel;
+class ULightComponent;
+class UPolys;
 
 //
 // One vertex associated with a Bsp node's polygon.  Contains a vertex index
diff --git a/Engine/Source/Runtime/Engine/Public/ModelLight.h b/Engine/Source/Runtime/Engine/Public/ModelLight.h
--- a/Engine/Source/Runtime/Engine/Public/ModelLight.h
+++ b/Engine/Source/Runtime/Engine/Public/ModelLight.h
@@ -6,6 +6,9 @@
 
 #pragma once
 
+#include "StaticLighting.h"
+#include "Model.h"
+
 /** Represents a BSP surface to the static lighting system. */
 class FBSPSurfaceStaticLighting : public FStaticLightingTextureMapping, public FStaticLightingMesh
 {
diff --git a/Engine/Source/Runtime/Engine/Public/PhysicsPublic.h b/Engine/Source/Runtime/Engine/Public/PhysicsPublic.h
--- a/Engine/Source/Runtime/Engine/Public/PhysicsPublic.h
+++ b/Engine/Source/Runtime/Engine/Public/PhysicsPublic.h
@@ -9,6 +9,7 @@
 
 #include "DynamicMeshBuilder.h"
 #include "PhysxUserData.h"
+#include "LocalVertexFactory.h"
 
 /**
  * Physics stats
diff --git a/Engine/Source/Runtime/Engine/Public/RawIndexBuffer.h b/Engine/Source/Runtime/Engine/Public/RawIndexBuffer.h
--- a/Engine/Source/Runtime/Engine/Public/RawIndexBuffer.h
+++ b/Engine/Source/Runtime/Engine/Public/RawIndexBuffer.h
@@ -7,6 +7,7 @@
 #pragma once
 
 #include "Engine.h"
+#include "RenderResource.h"
 
 class FRawIndexBuffer : public FIndexBuffer
 {
diff --git a/Engine/Source/Runtime/Engine/Public/SceneManagement.h b/Engine/Source/Runtime/Engine/Public/SceneManagement.h
--- a/Engine/Source/Runtime/Engine/Public/SceneManagement.h
+++ b/Engine/Source/Runtime/Engine/Public/SceneManagement.h
@@ -15,6 +15,7 @@
 #include "Engine/TextureLightProfile.h"
 #include "Engine/World.h"
 #include "RendererInterface.h"
+#include "GameViewportClient.h"
 
 // Forward declarations.
 class FSceneViewFamily;
diff --git a/Engine/Source/Runtime/Engine/Public/StaticMeshLight.h b/Engine/Source/Runtime/Engine/Public/StaticMeshLight.h
--- a/Engine/Source/Runtime/Engine/Public/StaticMeshLight.h
+++ b/Engine/Source/Runtime/Engine/Public/StaticMeshLight.h
@@ -6,6 +6,8 @@
 
 #pragma once
 
+#include "StaticLighting.h"
+
 /** Represents the triangles of one LOD of a static mesh primitive to the static lighting system. */
 class FStaticMeshStaticLightingMesh : public FStaticLightingMesh
 {
diff --git a/Engine/Source/Runtime/Engine/Public/StaticMeshResources.h b/Engine/Source/Runtime/Engine/Public/StaticMeshResources.h
--- a/Engine/Source/Runtime/Engine/Public/StaticMeshResources.h
+++ b/Engine/Source/Runtime/Engine/Public/StaticMeshResources.h
@@ -6,6 +6,8 @@
 
 #pragma once
 
+#include "RawIndexBuffer.h"
+
 /**
  * The LOD settings to use for a group of static meshes.
  */
diff --git a/Engine/Source/Runtime/Engine/Public/SurfaceIterators.h b/Engine/Source/Runtime/Engine/Public/SurfaceIterators.h
--- a/Engine/Source/Runtime/Engine/Public/SurfaceIterators.h
+++ b/Engine/Source/Runtime/Engine/Public/SurfaceIterators.h
@@ -7,6 +7,8 @@
 #ifndef __SURFACE_ITERATORS_H__
 #define __SURFACE_ITERATORS_H__
 
+#include "Model.h"
+
 ///////////////////////////////////////////////////////////////////////////////////////////////////////////////////////
 //
 // Level filters
diff --git a/Engine/Source/Runtime/Engine/Public/SystemSettings.h b/Engine/Source/Runtime/Engine/Public/SystemSettings.h
--- a/Engine/Source/Runtime/Engine/Public/SystemSettings.h
+++ b/Engine/Source/Runtime/Engine/Public/SystemSettings.h
@@ -7,6 +7,7 @@
 #pragma once
 
 #include "Texture.h"
+#include "ShowFlags.h"
 
 /*-----------------------------------------------------------------------------
 	System settings and scalability options.
diff --git a/Engine/Source/Runtime/Renderer/Private/RendererScene.cpp b/Engine/Source/Runtime/Renderer/Private/RendererScene.cpp
--- a/Engine/Source/Runtime/Renderer/Private/RendererScene.cpp
+++ b/Engine/Source/Runtime/Renderer/Private/RendererScene.cpp
@@ -7,6 +7,7 @@
 #include "RendererPrivate.h"
 #include "ScenePrivate.h"
 #include "ShaderCompiler.h"
+#include "StaticMeshResources.h"
 
 
 // Enable this define to do slow checks for components being added to the wrong
diff --git a/Engine/Source/Runtime/ShaderCore/Public/VertexFactory.h b/Engine/Source/Runtime/ShaderCore/Public/VertexFactory.h
--- a/Engine/Source/Runtime/ShaderCore/Public/VertexFactory.h
+++ b/Engine/Source/Runtime/ShaderCore/Public/VertexFactory.h
@@ -7,6 +7,7 @@
 #pragma once
 
 #include "ShaderCore.h"
+#include "Shader.h"
 #include "RenderResource.h"
 #include "SecureHash.h"
 
