Fix rendering of multiline strings using DrawDebugString.
ttp#332307 DrawDebugString doesn't understand newline characters
#proj core
#branch UE4
#summary Modified text rendering algorithm. When '\\n' is found in rendered string, current character rendering position is moved to the beginning of next line and character is skipped.

diff --git a/Engine/Source/Runtime/Engine/Private/UserInterface/CanvasItem.cpp b/Engine/Source/Runtime/Engine/Private/UserInterface/CanvasItem.cpp
--- a/Engine/Source/Runtime/Engine/Private/UserInterface/CanvasItem.cpp
+++ b/Engine/Source/Runtime/Engine/Private/UserInterface/CanvasItem.cpp
@@ -774,8 +774,8 @@ void FCanvasTextItem::Draw( class FCanvas* InCanvas )
 }
 
 void FCanvasTextItem::DrawStringInternal( class FCanvas* InCanvas, const FVector2D& DrawPos, const FLinearColor& InColour )
-{		
-	DrawnSize.X = 0.0f;
+{
+	DrawnSize = FVector2D(EForceInit::ForceInitToZero);
 	FHitProxyId HitProxyId = InCanvas->GetHitProxyId();
 	FTexture* LastTexture = NULL;
 	UTexture2D* Tex = NULL;
@@ -786,78 +786,91 @@ void FCanvasTextItem::DrawStringInternal( class FCanvas* InCanvas, const FVector
 	{
 		int32 Ch = (int32)Font->RemapChar(Chars[i]);
 
-		// Process character if it's valid.
-		if( Font->Characters.IsValidIndex(Ch) )
+		// Skip invalid characters.
+		if (!Font->Characters.IsValidIndex(Ch))
 		{
-			FFontCharacter& Char = Font->Characters[Ch];
-			if( Font->Textures.IsValidIndex(Char.TextureIndex) && 
-				(Tex=Font->Textures[Char.TextureIndex])!=NULL && 
-				Tex->Resource != NULL )
+			continue;
+		}
+
+		FFontCharacter& Char = Font->Characters[Ch];
+
+		if (FChar::IsLinebreak(Chars[i]))
+		{
+			// Set current character offset to the beginning of next line.
+			DrawnSize.X = 0.0f;
+			DrawnSize.Y += Char.VSize * YScale;
+
+			// Don't draw newline character
+			continue;
+		}
+
+		if( Font->Textures.IsValidIndex(Char.TextureIndex) && 
+			(Tex=Font->Textures[Char.TextureIndex])!=NULL && 
+			Tex->Resource != NULL )
+		{
+			if( LastTexture != Tex->Resource || BatchedElements == NULL )
 			{
-				if( LastTexture != Tex->Resource || BatchedElements == NULL )
-				{
-					FBatchedElementParameters* BatchedElementParameters = NULL;
-					BatchedElements = InCanvas->GetBatchedElements(FCanvas::ET_Triangle, BatchedElementParameters, Tex->Resource, BlendMode, FontRenderInfo.GlowInfo);
-					check(BatchedElements != NULL);
-					// trade-off between memory and performance by pre-allocating more reserved space 
-					// for the triangles/vertices of the batched elements used to render the text tiles
-					//BatchedElements->AddReserveTriangles(TextLen*2,Tex->Resource,BlendMode);
-					//BatchedElements->AddReserveVertices(TextLen*4);
-
-					FIntPoint ImportedTextureSize = Tex->GetImportedSize();
-					InvTextureSize.X = 1.0f / (float)ImportedTextureSize.X;
-					InvTextureSize.Y = 1.0f / (float)ImportedTextureSize.Y;
-				}
-				LastTexture = Tex->Resource;
-
-				const float X		= DrawnSize.X + DrawPos.X;
-				const float Y		= DrawPos.Y + + Char.VerticalOffset * YScale;
-				float SizeX			= Char.USize * XScale;
-				const float SizeY	= Char.VSize * YScale;
-				const float U		= Char.StartU * InvTextureSize.X;
-				const float V		= Char.StartV * InvTextureSize.Y;
-				const float SizeU	= Char.USize * InvTextureSize.X;
-				const float SizeV	= Char.VSize * InvTextureSize.Y;				
-
-				float Left, Top, Right, Bottom;
-				Left = X * Depth;
-				Top = Y * Depth;
-				Right = (X + SizeX) * Depth;
-				Bottom = (Y + SizeY) * Depth;
-
-				int32 V00 = BatchedElements->AddVertex(
-					FVector4( Left, Top, 0.f, Depth ),
-					FVector2D( U, V ),
-					InColour,
-					HitProxyId );
-				int32 V10 = BatchedElements->AddVertex(
-					FVector4( Right, Top, 0.0f, Depth ),
-					FVector2D( U + SizeU, V ),			
-					InColour,
-					HitProxyId );
-				int32 V01 = BatchedElements->AddVertex(
-					FVector4( Left, Bottom, 0.0f, Depth ),
-					FVector2D( U, V + SizeV ),	
-					InColour,
-					HitProxyId);
-				int32 V11 = BatchedElements->AddVertex(
-					FVector4( Right, Bottom, 0.0f, Depth ),
-					FVector2D( U + SizeU, V + SizeV ),
-					InColour,
-					HitProxyId);
-
-				BatchedElements->AddTriangle(V00, V10, V11, Tex->Resource, BlendMode, FontRenderInfo.GlowInfo);
-				BatchedElements->AddTriangle(V00, V11, V01, Tex->Resource, BlendMode, FontRenderInfo.GlowInfo);
-
-				// if we have another non-whitespace character to render, add the font's kerning.
-				if ( Chars[i+1] && !FChar::IsWhitespace(Chars[i+1]) )
-				{
-					SizeX += CharIncrement;
-				}
-
-				// Update the current rendering position
-				DrawnSize.X += SizeX;
+				FBatchedElementParameters* BatchedElementParameters = NULL;
+				BatchedElements = InCanvas->GetBatchedElements(FCanvas::ET_Triangle, BatchedElementParameters, Tex->Resource, BlendMode, FontRenderInfo.GlowInfo);
+				check(BatchedElements != NULL);
+				// trade-off between memory and performance by pre-allocating more reserved space 
+				// for the triangles/vertices of the batched elements used to render the text tiles
+				//BatchedElements->AddReserveTriangles(TextLen*2,Tex->Resource,BlendMode);
+				//BatchedElements->AddReserveVertices(TextLen*4);
+
+				FIntPoint ImportedTextureSize = Tex->GetImportedSize();
+				InvTextureSize.X = 1.0f / (float)ImportedTextureSize.X;
+				InvTextureSize.Y = 1.0f / (float)ImportedTextureSize.Y;
 			}
+			LastTexture = Tex->Resource;
+
+			const float X		= DrawnSize.X + DrawPos.X;
+			const float Y		= DrawnSize.Y + DrawPos.Y + Char.VerticalOffset * YScale;
+			float SizeX			= Char.USize * XScale;
+			const float SizeY	= Char.VSize * YScale;
+			const float U		= Char.StartU * InvTextureSize.X;
+			const float V		= Char.StartV * InvTextureSize.Y;
+			const float SizeU	= Char.USize * InvTextureSize.X;
+			const float SizeV	= Char.VSize * InvTextureSize.Y;				
+
+			float Left, Top, Right, Bottom;
+			Left = X * Depth;
+			Top = Y * Depth;
+			Right = (X + SizeX) * Depth;
+			Bottom = (Y + SizeY) * Depth;
+
+			int32 V00 = BatchedElements->AddVertex(
+				FVector4( Left, Top, 0.f, Depth ),
+				FVector2D( U, V ),
+				InColour,
+				HitProxyId );
+			int32 V10 = BatchedElements->AddVertex(
+				FVector4( Right, Top, 0.0f, Depth ),
+				FVector2D( U + SizeU, V ),			
+				InColour,
+				HitProxyId );
+			int32 V01 = BatchedElements->AddVertex(
+				FVector4( Left, Bottom, 0.0f, Depth ),
+				FVector2D( U, V + SizeV ),	
+				InColour,
+				HitProxyId);
+			int32 V11 = BatchedElements->AddVertex(
+				FVector4( Right, Bottom, 0.0f, Depth ),
+				FVector2D( U + SizeU, V + SizeV ),
+				InColour,
+				HitProxyId);
+
+			BatchedElements->AddTriangle(V00, V10, V11, Tex->Resource, BlendMode, FontRenderInfo.GlowInfo);
+			BatchedElements->AddTriangle(V00, V11, V01, Tex->Resource, BlendMode, FontRenderInfo.GlowInfo);
+
+			// if we have another non-whitespace character to render, add the font's kerning.
+			if ( Chars[i+1] && !FChar::IsWhitespace(Chars[i+1]) )
+			{
+				SizeX += CharIncrement;
+			}
+
+			// Update the current rendering position
+			DrawnSize.X += SizeX;
 		}
 	}
 }
