Engine.h trimming (AnimNode classes, SkeletalMeshActor.h)
Remove unused AnimationTypes.h/.cpp

diff --git a/Engine/Source/Developer/AssetTools/Private/AssetTypeActions/AssetTypeActions_VertexAnimation.h b/Engine/Source/Developer/AssetTools/Private/AssetTypeActions/AssetTypeActions_VertexAnimation.h
--- a/Engine/Source/Developer/AssetTools/Private/AssetTypeActions/AssetTypeActions_VertexAnimation.h
+++ b/Engine/Source/Developer/AssetTools/Private/AssetTypeActions/AssetTypeActions_VertexAnimation.h
@@ -2,6 +2,8 @@
 
 #pragma once
 
+#include "Animation/VertexAnim/VertexAnimation.h"
+
 class FAssetTypeActions_VertexAnimation : public FAssetTypeActions_Base
 {
 public:
diff --git a/Engine/Source/Developer/MovieSceneTools/Private/AnimationTrackEditor.cpp b/Engine/Source/Developer/MovieSceneTools/Private/AnimationTrackEditor.cpp
--- a/Engine/Source/Developer/MovieSceneTools/Private/AnimationTrackEditor.cpp
+++ b/Engine/Source/Developer/MovieSceneTools/Private/AnimationTrackEditor.cpp
@@ -18,6 +18,7 @@
 #include "MovieSceneAnimationSection.h"
 #include "CommonMovieSceneTools.h"
 #include "AssetRegistryModule.h"
+#include "Animation/SkeletalMeshActor.h"
 
 
 
diff --git a/Engine/Source/Editor/AnimGraph/Classes/AnimGraphNode_ApplyAdditive.h b/Engine/Source/Editor/AnimGraph/Classes/AnimGraphNode_ApplyAdditive.h
--- a/Engine/Source/Editor/AnimGraph/Classes/AnimGraphNode_ApplyAdditive.h
+++ b/Engine/Source/Editor/AnimGraph/Classes/AnimGraphNode_ApplyAdditive.h
@@ -2,6 +2,7 @@
 
 #pragma once
 #include "AnimGraphNode_Base.h"
+#include "Animation/AnimNode_ApplyAdditive.h"
 #include "AnimGraphNode_ApplyAdditive.generated.h"
 
 UCLASS(MinimalAPI)
diff --git a/Engine/Source/Editor/AnimGraph/Classes/AnimGraphNode_Base.h b/Engine/Source/Editor/AnimGraph/Classes/AnimGraphNode_Base.h
--- a/Engine/Source/Editor/AnimGraph/Classes/AnimGraphNode_Base.h
+++ b/Engine/Source/Editor/AnimGraph/Classes/AnimGraphNode_Base.h
@@ -2,6 +2,7 @@
 
 #pragma once
 #include "K2Node.h"
+#include "Animation/AnimNodeBase.h"
 #include "AnimGraphNode_Base.generated.h"
 
 //
diff --git a/Engine/Source/Editor/AnimGraph/Classes/AnimGraphNode_BlendListByBool.h b/Engine/Source/Editor/AnimGraph/Classes/AnimGraphNode_BlendListByBool.h
--- a/Engine/Source/Editor/AnimGraph/Classes/AnimGraphNode_BlendListByBool.h
+++ b/Engine/Source/Editor/AnimGraph/Classes/AnimGraphNode_BlendListByBool.h
@@ -2,6 +2,7 @@
 
 #pragma once
 #include "AnimGraphNode_BlendListBase.h"
+#include "Animation/AnimNode_BlendListByBool.h"
 #include "AnimGraphNode_BlendListByBool.generated.h"
 
 UCLASS(MinimalAPI)
diff --git a/Engine/Source/Editor/AnimGraph/Classes/AnimGraphNode_BlendListByEnum.h b/Engine/Source/Editor/AnimGraph/Classes/AnimGraphNode_BlendListByEnum.h
--- a/Engine/Source/Editor/AnimGraph/Classes/AnimGraphNode_BlendListByEnum.h
+++ b/Engine/Source/Editor/AnimGraph/Classes/AnimGraphNode_BlendListByEnum.h
@@ -2,6 +2,7 @@
 
 #pragma once
 #include "AnimGraphNode_BlendListBase.h"
+#include "Animation/AnimNode_BlendListByEnum.h"
 #include "AnimGraphNode_BlendListByEnum.generated.h"
 
 UCLASS(MinimalAPI)
diff --git a/Engine/Source/Editor/AnimGraph/Classes/AnimGraphNode_BlendListByInt.h b/Engine/Source/Editor/AnimGraph/Classes/AnimGraphNode_BlendListByInt.h
--- a/Engine/Source/Editor/AnimGraph/Classes/AnimGraphNode_BlendListByInt.h
+++ b/Engine/Source/Editor/AnimGraph/Classes/AnimGraphNode_BlendListByInt.h
@@ -2,6 +2,7 @@
 
 #pragma once
 #include "AnimGraphNode_BlendListBase.h"
+#include "Animation/AnimNode_BlendListByInt.h"
 #include "AnimGraphNode_BlendListByInt.generated.h"
 
 UCLASS(MinimalAPI)
diff --git a/Engine/Source/Editor/AnimGraph/Classes/AnimGraphNode_BlendSpacePlayer.h b/Engine/Source/Editor/AnimGraph/Classes/AnimGraphNode_BlendSpacePlayer.h
--- a/Engine/Source/Editor/AnimGraph/Classes/AnimGraphNode_BlendSpacePlayer.h
+++ b/Engine/Source/Editor/AnimGraph/Classes/AnimGraphNode_BlendSpacePlayer.h
@@ -2,6 +2,7 @@
 
 #pragma once
 #include "AnimGraphNode_BlendSpaceBase.h"
+#include "Animation/AnimNode_BlendSpacePlayer.h"
 #include "AnimGraphNode_BlendSpacePlayer.generated.h"
 
 UCLASS(MinimalAPI)
diff --git a/Engine/Source/Editor/AnimGraph/Classes/AnimGraphNode_ComponentToLocalSpace.h b/Engine/Source/Editor/AnimGraph/Classes/AnimGraphNode_ComponentToLocalSpace.h
--- a/Engine/Source/Editor/AnimGraph/Classes/AnimGraphNode_ComponentToLocalSpace.h
+++ b/Engine/Source/Editor/AnimGraph/Classes/AnimGraphNode_ComponentToLocalSpace.h
@@ -2,6 +2,7 @@
 
 #pragma once
 #include "AnimGraphNode_Base.h"
+#include "Animation/AnimNodeSpaceConversions.h"
 #include "AnimGraphNode_ComponentToLocalSpace.generated.h"
 
 UCLASS(MinimalAPI)
diff --git a/Engine/Source/Editor/AnimGraph/Classes/AnimGraphNode_CopyBone.h b/Engine/Source/Editor/AnimGraph/Classes/AnimGraphNode_CopyBone.h
--- a/Engine/Source/Editor/AnimGraph/Classes/AnimGraphNode_CopyBone.h
+++ b/Engine/Source/Editor/AnimGraph/Classes/AnimGraphNode_CopyBone.h
@@ -2,6 +2,7 @@
 
 #pragma once
 #include "AnimGraphNode_SkeletalControlBase.h"
+#include "Animation/BoneControllers/AnimNode_CopyBone.h"
 #include "AnimGraphNode_CopyBone.generated.h"
 
 UCLASS(MinimalAPI)
diff --git a/Engine/Source/Editor/AnimGraph/Classes/AnimGraphNode_Fabrik.h b/Engine/Source/Editor/AnimGraph/Classes/AnimGraphNode_Fabrik.h
--- a/Engine/Source/Editor/AnimGraph/Classes/AnimGraphNode_Fabrik.h
+++ b/Engine/Source/Editor/AnimGraph/Classes/AnimGraphNode_Fabrik.h
@@ -2,6 +2,7 @@
 
 #pragma once
 #include "AnimGraphNode_SkeletalControlBase.h"
+#include "Animation/BoneControllers/AnimNode_Fabrik.h"
 #include "AnimGraphNode_Fabrik.generated.h"
 
 // Editor node for FABRIK IK skeletal controller
diff --git a/Engine/Source/Editor/AnimGraph/Classes/AnimGraphNode_HandIKRetargeting.h b/Engine/Source/Editor/AnimGraph/Classes/AnimGraphNode_HandIKRetargeting.h
--- a/Engine/Source/Editor/AnimGraph/Classes/AnimGraphNode_HandIKRetargeting.h
+++ b/Engine/Source/Editor/AnimGraph/Classes/AnimGraphNode_HandIKRetargeting.h
@@ -2,6 +2,7 @@
 
 #pragma once
 #include "AnimGraphNode_SkeletalControlBase.h"
+#include "Animation/BoneControllers/AnimNode_HandIKRetargeting.h"
 #include "AnimGraphNode_HandIKRetargeting.generated.h"
 
 UCLASS(MinimalAPI)
diff --git a/Engine/Source/Editor/AnimGraph/Classes/AnimGraphNode_LayeredBoneBlend.h b/Engine/Source/Editor/AnimGraph/Classes/AnimGraphNode_LayeredBoneBlend.h
--- a/Engine/Source/Editor/AnimGraph/Classes/AnimGraphNode_LayeredBoneBlend.h
+++ b/Engine/Source/Editor/AnimGraph/Classes/AnimGraphNode_LayeredBoneBlend.h
@@ -2,6 +2,7 @@
 
 #pragma once
 #include "AnimGraphNode_BlendListBase.h"
+#include "Animation/AnimNode_LayeredBoneBlend.h"
 #include "AnimGraphNode_LayeredBoneBlend.generated.h"
 
 UCLASS(MinimalAPI)
diff --git a/Engine/Source/Editor/AnimGraph/Classes/AnimGraphNode_LocalToComponentSpace.h b/Engine/Source/Editor/AnimGraph/Classes/AnimGraphNode_LocalToComponentSpace.h
--- a/Engine/Source/Editor/AnimGraph/Classes/AnimGraphNode_LocalToComponentSpace.h
+++ b/Engine/Source/Editor/AnimGraph/Classes/AnimGraphNode_LocalToComponentSpace.h
@@ -2,6 +2,7 @@
 
 #pragma once
 #include "AnimGraphNode_Base.h"
+#include "Animation/AnimNodeSpaceConversions.h"
 #include "AnimGraphNode_LocalToComponentSpace.generated.h"
 
 UCLASS(MinimalAPI)
diff --git a/Engine/Source/Editor/AnimGraph/Classes/AnimGraphNode_LookAt.h b/Engine/Source/Editor/AnimGraph/Classes/AnimGraphNode_LookAt.h
--- a/Engine/Source/Editor/AnimGraph/Classes/AnimGraphNode_LookAt.h
+++ b/Engine/Source/Editor/AnimGraph/Classes/AnimGraphNode_LookAt.h
@@ -3,6 +3,7 @@
 #pragma once
 
 #include "AnimGraphNode_SkeletalControlBase.h"
+#include "Animation/BoneControllers/AnimNode_LookAt.h"
 #include "AnimGraphNode_LookAt.generated.h"
 
 UCLASS(MinimalAPI)
diff --git a/Engine/Source/Editor/AnimGraph/Classes/AnimGraphNode_MeshRefPose.h b/Engine/Source/Editor/AnimGraph/Classes/AnimGraphNode_MeshRefPose.h
--- a/Engine/Source/Editor/AnimGraph/Classes/AnimGraphNode_MeshRefPose.h
+++ b/Engine/Source/Editor/AnimGraph/Classes/AnimGraphNode_MeshRefPose.h
@@ -2,6 +2,7 @@
 
 #pragma once
 #include "AnimGraphNode_Base.h"
+#include "Animation/AnimNode_RefPose.h"
 #include "AnimGraphNode_MeshRefPose.generated.h"
 
 UCLASS(MinimalAPI)
diff --git a/Engine/Source/Editor/AnimGraph/Classes/AnimGraphNode_ModifyBone.h b/Engine/Source/Editor/AnimGraph/Classes/AnimGraphNode_ModifyBone.h
--- a/Engine/Source/Editor/AnimGraph/Classes/AnimGraphNode_ModifyBone.h
+++ b/Engine/Source/Editor/AnimGraph/Classes/AnimGraphNode_ModifyBone.h
@@ -3,6 +3,7 @@
 #pragma once
 
 #include "AnimGraphNode_SkeletalControlBase.h"
+#include "Animation/BoneControllers/AnimNode_ModifyBone.h"
 #include "AnimGraphNode_ModifyBone.generated.h"
 
 UCLASS(MinimalAPI)
diff --git a/Engine/Source/Editor/AnimGraph/Classes/AnimGraphNode_RefPoseBase.h b/Engine/Source/Editor/AnimGraph/Classes/AnimGraphNode_RefPoseBase.h
--- a/Engine/Source/Editor/AnimGraph/Classes/AnimGraphNode_RefPoseBase.h
+++ b/Engine/Source/Editor/AnimGraph/Classes/AnimGraphNode_RefPoseBase.h
@@ -2,6 +2,7 @@
 
 #pragma once
 #include "AnimGraphNode_Base.h"
+#include "Animation/AnimNode_RefPose.h"
 #include "AnimGraphNode_RefPoseBase.generated.h"
 
 UCLASS(MinimalAPI, abstract)
diff --git a/Engine/Source/Editor/AnimGraph/Classes/AnimGraphNode_Root.h b/Engine/Source/Editor/AnimGraph/Classes/AnimGraphNode_Root.h
--- a/Engine/Source/Editor/AnimGraph/Classes/AnimGraphNode_Root.h
+++ b/Engine/Source/Editor/AnimGraph/Classes/AnimGraphNode_Root.h
@@ -2,6 +2,7 @@
 
 #pragma once
 #include "AnimGraphNode_Base.h"
+#include "Animation/AnimNode_Root.h"
 #include "AnimGraphNode_Root.generated.h"
 
 UCLASS(MinimalAPI)
diff --git a/Engine/Source/Editor/AnimGraph/Classes/AnimGraphNode_RotateRootBone.h b/Engine/Source/Editor/AnimGraph/Classes/AnimGraphNode_RotateRootBone.h
--- a/Engine/Source/Editor/AnimGraph/Classes/AnimGraphNode_RotateRootBone.h
+++ b/Engine/Source/Editor/AnimGraph/Classes/AnimGraphNode_RotateRootBone.h
@@ -3,6 +3,7 @@
 
 #pragma once
 #include "AnimGraphNode_Base.h"
+#include "Animation/AnimNode_RotateRootBone.h"
 #include "AnimGraphNode_RotateRootBone.generated.h"
 
 UCLASS(MinimalAPI)
diff --git a/Engine/Source/Editor/AnimGraph/Classes/AnimGraphNode_RotationMultiplier.h b/Engine/Source/Editor/AnimGraph/Classes/AnimGraphNode_RotationMultiplier.h
--- a/Engine/Source/Editor/AnimGraph/Classes/AnimGraphNode_RotationMultiplier.h
+++ b/Engine/Source/Editor/AnimGraph/Classes/AnimGraphNode_RotationMultiplier.h
@@ -2,6 +2,7 @@
 
 #pragma once
 #include "AnimGraphNode_SkeletalControlBase.h"
+#include "Animation/BoneControllers/AnimNode_RotationMultiplier.h"
 #include "AnimGraphNode_RotationMultiplier.generated.h"
 
 UCLASS(MinimalAPI)
diff --git a/Engine/Source/Editor/AnimGraph/Classes/AnimGraphNode_RotationOffsetBlendSpace.h b/Engine/Source/Editor/AnimGraph/Classes/AnimGraphNode_RotationOffsetBlendSpace.h
--- a/Engine/Source/Editor/AnimGraph/Classes/AnimGraphNode_RotationOffsetBlendSpace.h
+++ b/Engine/Source/Editor/AnimGraph/Classes/AnimGraphNode_RotationOffsetBlendSpace.h
@@ -2,6 +2,7 @@
 
 #pragma once
 #include "AnimGraphNode_BlendSpaceBase.h"
+#include "Animation/AnimNode_RotationOffsetBlendSpace.h"
 #include "AnimGraphNode_RotationOffsetBlendSpace.generated.h"
 
 UCLASS(MinimalAPI)
diff --git a/Engine/Source/Editor/AnimGraph/Classes/AnimGraphNode_SaveCachedPose.h b/Engine/Source/Editor/AnimGraph/Classes/AnimGraphNode_SaveCachedPose.h
--- a/Engine/Source/Editor/AnimGraph/Classes/AnimGraphNode_SaveCachedPose.h
+++ b/Engine/Source/Editor/AnimGraph/Classes/AnimGraphNode_SaveCachedPose.h
@@ -2,6 +2,7 @@
 
 #pragma once
 #include "AnimGraphNode_Base.h"
+#include "Animation/AnimNode_SaveCachedPose.h"
 #include "AnimGraphNode_SaveCachedPose.generated.h"
 
 UCLASS(MinimalAPI)
diff --git a/Engine/Source/Editor/AnimGraph/Classes/AnimGraphNode_SequenceEvaluator.h b/Engine/Source/Editor/AnimGraph/Classes/AnimGraphNode_SequenceEvaluator.h
--- a/Engine/Source/Editor/AnimGraph/Classes/AnimGraphNode_SequenceEvaluator.h
+++ b/Engine/Source/Editor/AnimGraph/Classes/AnimGraphNode_SequenceEvaluator.h
@@ -2,6 +2,7 @@
 
 #pragma once
 #include "AnimGraphNode_Base.h"
+#include "Animation/AnimNode_SequenceEvaluator.h"
 #include "AnimGraphNode_SequenceEvaluator.generated.h"
 
 UCLASS(MinimalAPI)
diff --git a/Engine/Source/Editor/AnimGraph/Classes/AnimGraphNode_SequencePlayer.h b/Engine/Source/Editor/AnimGraph/Classes/AnimGraphNode_SequencePlayer.h
--- a/Engine/Source/Editor/AnimGraph/Classes/AnimGraphNode_SequencePlayer.h
+++ b/Engine/Source/Editor/AnimGraph/Classes/AnimGraphNode_SequencePlayer.h
@@ -2,6 +2,7 @@
 
 #pragma once
 #include "AnimGraphNode_Base.h"
+#include "Animation/AnimNode_SequencePlayer.h"
 #include "AnimGraphNode_SequencePlayer.generated.h"
 
 UCLASS(MinimalAPI)
diff --git a/Engine/Source/Editor/AnimGraph/Classes/AnimGraphNode_Slot.h b/Engine/Source/Editor/AnimGraph/Classes/AnimGraphNode_Slot.h
--- a/Engine/Source/Editor/AnimGraph/Classes/AnimGraphNode_Slot.h
+++ b/Engine/Source/Editor/AnimGraph/Classes/AnimGraphNode_Slot.h
@@ -2,6 +2,7 @@
 
 #pragma once
 #include "AnimGraphNode_Base.h"
+#include "Animation/AnimNode_Slot.h"
 #include "AnimGraphNode_Slot.generated.h"
 
 UCLASS(MinimalAPI)
diff --git a/Engine/Source/Editor/AnimGraph/Classes/AnimGraphNode_SpringBone.h b/Engine/Source/Editor/AnimGraph/Classes/AnimGraphNode_SpringBone.h
--- a/Engine/Source/Editor/AnimGraph/Classes/AnimGraphNode_SpringBone.h
+++ b/Engine/Source/Editor/AnimGraph/Classes/AnimGraphNode_SpringBone.h
@@ -2,6 +2,7 @@
 
 #pragma once
 #include "AnimGraphNode_SkeletalControlBase.h"
+#include "Animation/BoneControllers/AnimNode_SpringBone.h"
 #include "AnimGraphNode_SpringBone.generated.h"
 
 UCLASS(MinimalAPI)
diff --git a/Engine/Source/Editor/AnimGraph/Classes/AnimGraphNode_StateMachineBase.h b/Engine/Source/Editor/AnimGraph/Classes/AnimGraphNode_StateMachineBase.h
--- a/Engine/Source/Editor/AnimGraph/Classes/AnimGraphNode_StateMachineBase.h
+++ b/Engine/Source/Editor/AnimGraph/Classes/AnimGraphNode_StateMachineBase.h
@@ -2,6 +2,7 @@
 
 #pragma once
 #include "AnimGraphNode_Base.h"
+#include "Animation/AnimNode_StateMachine.h"
 #include "AnimGraphNode_StateMachineBase.generated.h"
 
 UCLASS(Abstract)
diff --git a/Engine/Source/Editor/AnimGraph/Classes/AnimGraphNode_TransitionPoseEvaluator.h b/Engine/Source/Editor/AnimGraph/Classes/AnimGraphNode_TransitionPoseEvaluator.h
--- a/Engine/Source/Editor/AnimGraph/Classes/AnimGraphNode_TransitionPoseEvaluator.h
+++ b/Engine/Source/Editor/AnimGraph/Classes/AnimGraphNode_TransitionPoseEvaluator.h
@@ -2,6 +2,7 @@
 
 #pragma once
 #include "AnimGraphNode_Base.h"
+#include "Animation/AnimNode_TransitionPoseEvaluator.h"
 #include "AnimGraphNode_TransitionPoseEvaluator.generated.h"
 
 UCLASS(MinimalAPI)
diff --git a/Engine/Source/Editor/AnimGraph/Classes/AnimGraphNode_TransitionResult.h b/Engine/Source/Editor/AnimGraph/Classes/AnimGraphNode_TransitionResult.h
--- a/Engine/Source/Editor/AnimGraph/Classes/AnimGraphNode_TransitionResult.h
+++ b/Engine/Source/Editor/AnimGraph/Classes/AnimGraphNode_TransitionResult.h
@@ -2,6 +2,7 @@
 
 #pragma once
 #include "AnimGraphNode_Base.h"
+#include "Animation/AnimNode_TransitionResult.h"
 #include "AnimGraphNode_TransitionResult.generated.h"
 
 UCLASS(MinimalAPI)
diff --git a/Engine/Source/Editor/AnimGraph/Classes/AnimGraphNode_TwoBoneIK.h b/Engine/Source/Editor/AnimGraph/Classes/AnimGraphNode_TwoBoneIK.h
--- a/Engine/Source/Editor/AnimGraph/Classes/AnimGraphNode_TwoBoneIK.h
+++ b/Engine/Source/Editor/AnimGraph/Classes/AnimGraphNode_TwoBoneIK.h
@@ -2,6 +2,7 @@
 
 #pragma once
 #include "AnimGraphNode_SkeletalControlBase.h"
+#include "Animation/BoneControllers/AnimNode_TwoBoneIK.h"
 #include "AnimGraphNode_TwoBoneIK.generated.h"
 
 UCLASS(MinimalAPI)
diff --git a/Engine/Source/Editor/AnimGraph/Classes/AnimGraphNode_TwoWayBlend.h b/Engine/Source/Editor/AnimGraph/Classes/AnimGraphNode_TwoWayBlend.h
--- a/Engine/Source/Editor/AnimGraph/Classes/AnimGraphNode_TwoWayBlend.h
+++ b/Engine/Source/Editor/AnimGraph/Classes/AnimGraphNode_TwoWayBlend.h
@@ -2,6 +2,7 @@
 
 #pragma once
 #include "AnimGraphNode_Base.h"
+#include "Animation/AnimNode_TwoWayBlend.h"
 #include "AnimGraphNode_TwoWayBlend.generated.h"
 
 UCLASS(MinimalAPI)
diff --git a/Engine/Source/Editor/AnimGraph/Classes/AnimGraphNode_UseCachedPose.h b/Engine/Source/Editor/AnimGraph/Classes/AnimGraphNode_UseCachedPose.h
--- a/Engine/Source/Editor/AnimGraph/Classes/AnimGraphNode_UseCachedPose.h
+++ b/Engine/Source/Editor/AnimGraph/Classes/AnimGraphNode_UseCachedPose.h
@@ -2,6 +2,7 @@
 
 #pragma once
 #include "AnimGraphNode_Base.h"
+#include "Animation/AnimNode_UseCachedPose.h"
 #include "AnimGraphNode_UseCachedPose.generated.h"
 
 UCLASS(MinimalAPI)
diff --git a/Engine/Source/Editor/AnimGraph/Classes/AnimGraphNode_WheelHandler.h b/Engine/Source/Editor/AnimGraph/Classes/AnimGraphNode_WheelHandler.h
--- a/Engine/Source/Editor/AnimGraph/Classes/AnimGraphNode_WheelHandler.h
+++ b/Engine/Source/Editor/AnimGraph/Classes/AnimGraphNode_WheelHandler.h
@@ -3,6 +3,7 @@
 #pragma once
 
 #include "AnimGraphNode_SkeletalControlBase.h"
+#include "Animation/BoneControllers/AnimNode_WheelHandler.h"
 #include "AnimGraphNode_WheelHandler.generated.h"
 
 UCLASS(MinimalAPI)
diff --git a/Engine/Source/Editor/AnimGraph/Classes/AnimPreviewInstance.h b/Engine/Source/Editor/AnimGraph/Classes/AnimPreviewInstance.h
--- a/Engine/Source/Editor/AnimGraph/Classes/AnimPreviewInstance.h
+++ b/Engine/Source/Editor/AnimGraph/Classes/AnimPreviewInstance.h
@@ -5,6 +5,7 @@
 #include "BlueprintGraphDefinitions.h"
 #include "AnimGraphNode_ModifyBone.h"
 #include "Animation/AnimSingleNodeInstance.h"
+#include "Animation/BoneControllers/AnimNode_ModifyBone.h"
 #include "AnimPreviewInstance.generated.h"
 
 /** Enum to know how montage is being played */
diff --git a/Engine/Source/Editor/GraphEditor/Private/AnimationStateNodes/SGraphNodeAnimState.cpp b/Engine/Source/Editor/GraphEditor/Private/AnimationStateNodes/SGraphNodeAnimState.cpp
--- a/Engine/Source/Editor/GraphEditor/Private/AnimationStateNodes/SGraphNodeAnimState.cpp
+++ b/Engine/Source/Editor/GraphEditor/Private/AnimationStateNodes/SGraphNodeAnimState.cpp
@@ -9,6 +9,7 @@
 #include "AnimStateConduitNode.h"
 #include "AnimationStateMachineGraph.h"
 #include "AnimGraphNode_StateMachineBase.h"
+#include "Animation/AnimNode_StateMachine.h"
 
 /////////////////////////////////////////////////////
 // SStateMachineOutputPin
diff --git a/Engine/Source/Editor/KismetCompiler/Private/AnimBlueprintCompiler.h b/Engine/Source/Editor/KismetCompiler/Private/AnimBlueprintCompiler.h
--- a/Engine/Source/Editor/KismetCompiler/Private/AnimBlueprintCompiler.h
+++ b/Engine/Source/Editor/KismetCompiler/Private/AnimBlueprintCompiler.h
@@ -4,6 +4,7 @@
 
 #include "KismetCompiler.h"
 #include "AnimGraphDefinitions.h"
+#include "Animation/AnimNodeBase.h"
 
 //
 // Forward declarations.
diff --git a/Engine/Source/Editor/Matinee/Private/Matinee.cpp b/Engine/Source/Editor/Matinee/Private/Matinee.cpp
--- a/Engine/Source/Editor/Matinee/Private/Matinee.cpp
+++ b/Engine/Source/Editor/Matinee/Private/Matinee.cpp
@@ -16,6 +16,7 @@
 #include "Matinee/InterpTrackVisibility.h"
 #include "Matinee/InterpTrackEvent.h"
 #include "Matinee/InterpGroupDirector.h"
+#include "Matinee/MatineeAnimInterface.h"
 
 #include "Matinee.h"
 #include "MatineeActions.h"
diff --git a/Engine/Source/Editor/Matinee/Private/MatineeDraw.cpp b/Engine/Source/Editor/Matinee/Private/MatineeDraw.cpp
--- a/Engine/Source/Editor/Matinee/Private/MatineeDraw.cpp
+++ b/Engine/Source/Editor/Matinee/Private/MatineeDraw.cpp
@@ -11,6 +11,7 @@
 #include "Runtime/Engine/Public/Interpolation.h"
 #include "Runtime/Engine/Public/InterpolationHitProxy.h"
 #include "Runtime/Engine/Public/Slate/SceneViewport.h"
+
 #include "Matinee/MatineeActor.h"
 #include "Matinee/InterpGroupDirector.h"
 #include "Matinee/InterpGroupInst.h"
@@ -19,8 +20,9 @@
 #include "Matinee/InterpTrackLinearColorBase.h"
 #include "Matinee/InterpTrackEvent.h"
 #include "Matinee/InterpTrackMove.h"
+#include "Matinee/MatineeAnimInterface.h"
 #include "Materials/MaterialInstanceActor.h"
-
+#include "Animation/SkeletalMeshActor.h"
 #include "Particles/Emitter.h"
 
 static const int32 GroupHeadHeight = 24;
diff --git a/Engine/Source/Editor/Matinee/Private/MatineeMenus.cpp b/Engine/Source/Editor/Matinee/Private/MatineeMenus.cpp
--- a/Engine/Source/Editor/Matinee/Private/MatineeMenus.cpp
+++ b/Engine/Source/Editor/Matinee/Private/MatineeMenus.cpp
@@ -36,6 +36,7 @@
 #include "DesktopPlatformModule.h"
 #include "PackageTools.h"
 #include "SoundDefinitions.h"
+#include "Animation/SkeletalMeshActor.h"
 
 #include "WorkspaceMenuStructureModule.h"
 
diff --git a/Engine/Source/Editor/Matinee/Private/MatineeTrackHelpers.cpp b/Engine/Source/Editor/Matinee/Private/MatineeTrackHelpers.cpp
--- a/Engine/Source/Editor/Matinee/Private/MatineeTrackHelpers.cpp
+++ b/Engine/Source/Editor/Matinee/Private/MatineeTrackHelpers.cpp
@@ -21,6 +21,7 @@
 #include "Matinee/InterpTrackVectorProp.h"
 #include "Matinee/InterpTrackLinearColorProp.h"
 #include "Matinee/InterpGroupInst.h"
+#include "Matinee/MatineeAnimInterface.h"
 
 #define LOCTEXT_NAMESPACE "MatineeTrackHelpers"
 
diff --git a/Engine/Source/Editor/Persona/Private/SAnimationSequenceBrowser.cpp b/Engine/Source/Editor/Persona/Private/SAnimationSequenceBrowser.cpp
--- a/Engine/Source/Editor/Persona/Private/SAnimationSequenceBrowser.cpp
+++ b/Engine/Source/Editor/Persona/Private/SAnimationSequenceBrowser.cpp
@@ -3,6 +3,7 @@
 
 #include "PersonaPrivatePCH.h"
 
+#include "Animation/VertexAnim/VertexAnimation.h"
 #include "SAnimationSequenceBrowser.h"
 #include "Persona.h"
 #include "AssetRegistryModule.h"
diff --git a/Engine/Source/Editor/UMGEditor/Private/WidgetBlueprintCompiler.cpp b/Engine/Source/Editor/UMGEditor/Private/WidgetBlueprintCompiler.cpp
--- a/Engine/Source/Editor/UMGEditor/Private/WidgetBlueprintCompiler.cpp
+++ b/Engine/Source/Editor/UMGEditor/Private/WidgetBlueprintCompiler.cpp
@@ -2,6 +2,7 @@
 
 #include "UMGEditorPrivatePCH.h"
 
+#include "Animation/AnimNodeBase.h"
 #include "WidgetBlueprintCompiler.h"
 #include "Kismet2NameValidators.h"
 #include "KismetReinstanceUtilities.h"
diff --git a/Engine/Source/Editor/UMGEditor/Private/WidgetGraphNode_Base.cpp b/Engine/Source/Editor/UMGEditor/Private/WidgetGraphNode_Base.cpp
--- a/Engine/Source/Editor/UMGEditor/Private/WidgetGraphNode_Base.cpp
+++ b/Engine/Source/Editor/UMGEditor/Private/WidgetGraphNode_Base.cpp
@@ -2,6 +2,7 @@
 
 #include "UMGEditorPrivatePCH.h"
 
+#include "Animation/AnimNodeBase.h"
 #include "BlueprintEditorUtils.h"
 #include "WidgetGraphSchema.h"
 #include "K2ActionMenuBuilder.h"
diff --git a/Engine/Source/Editor/UnrealEd/Private/Animation/DebugSkelMeshComponent.cpp b/Engine/Source/Editor/UnrealEd/Private/Animation/DebugSkelMeshComponent.cpp
--- a/Engine/Source/Editor/UnrealEd/Private/Animation/DebugSkelMeshComponent.cpp
+++ b/Engine/Source/Editor/UnrealEd/Private/Animation/DebugSkelMeshComponent.cpp
@@ -5,6 +5,7 @@
 #include "SkeletalRenderPublic.h"
 #include "AnimationRuntime.h"
 #include "AnimPreviewInstance.h"
+#include "Animation/VertexAnim/VertexAnimation.h"
 
 //////////////////////////////////////////////////////////////////////////
 // FDebugSkelMeshSceneProxy
diff --git a/Engine/Source/Editor/UnrealEd/Private/Animation/VertexAnimTools.cpp b/Engine/Source/Editor/UnrealEd/Private/Animation/VertexAnimTools.cpp
--- a/Engine/Source/Editor/UnrealEd/Private/Animation/VertexAnimTools.cpp
+++ b/Engine/Source/Editor/UnrealEd/Private/Animation/VertexAnimTools.cpp
@@ -4,7 +4,7 @@
 #include "FbxImporter.h"
 #include "FbxExporter.h"
 #include "VertexAnimTools.h"
-
+#include "Animation/VertexAnim/VertexAnimation.h"
 
 /** Util to find all meshes with a vertex cache deformer */
 void RecursiveFindVertexAnimMeshes(FbxNode* Node, TArray<FbxMesh*>& OutVertexCacheMeshes)
diff --git a/Engine/Source/Editor/UnrealEd/Private/AssetSelection.cpp b/Engine/Source/Editor/UnrealEd/Private/AssetSelection.cpp
--- a/Engine/Source/Editor/UnrealEd/Private/AssetSelection.cpp
+++ b/Engine/Source/Editor/UnrealEd/Private/AssetSelection.cpp
@@ -21,6 +21,7 @@
 #include "Landscape/Landscape.h"
 
 #include "Particles/Emitter.h"
+#include "Animation/SkeletalMeshActor.h"
 
 extern void OnPlaceStaticMeshActor( AActor* MeshActor, bool bUseSurfaceOrientation );
 
diff --git a/Engine/Source/Editor/UnrealEd/Private/Editor.cpp b/Engine/Source/Editor/UnrealEd/Private/Editor.cpp
--- a/Engine/Source/Editor/UnrealEd/Private/Editor.cpp
+++ b/Engine/Source/Editor/UnrealEd/Private/Editor.cpp
@@ -4,6 +4,8 @@
 #include "UnrealEd.h"
 #include "Matinee/MatineeActor.h"
 #include "Engine/InteractiveFoliageActor.h"
+#include "Animation/SkeletalMeshActor.h"
+#include "Animation/VertexAnim/VertexAnimation.h"
 #include "Engine/WorldComposition.h"
 #include "EditorSupportDelegates.h"
 #include "Factories.h"
diff --git a/Engine/Source/Editor/UnrealEd/Private/EditorActor.cpp b/Engine/Source/Editor/UnrealEd/Private/EditorActor.cpp
--- a/Engine/Source/Editor/UnrealEd/Private/EditorActor.cpp
+++ b/Engine/Source/Editor/UnrealEd/Private/EditorActor.cpp
@@ -14,7 +14,7 @@
 #include "Editor/GeometryMode/Public/EditorGeometry.h"
 #include "ActorEditorUtils.h"
 #include "Foliage/InstancedFoliageActor.h"
-
+#include "Animation/SkeletalMeshActor.h"
 #include "Particles/Emitter.h"
 #include "Particles/ParticleSystemComponent.h"
 
diff --git a/Engine/Source/Editor/UnrealEd/Private/EditorServer.cpp b/Engine/Source/Editor/UnrealEd/Private/EditorServer.cpp
--- a/Engine/Source/Editor/UnrealEd/Private/EditorServer.cpp
+++ b/Engine/Source/Editor/UnrealEd/Private/EditorServer.cpp
@@ -25,6 +25,7 @@
 #include "Animation/AnimSet.h"
 #include "Matinee/InterpTrackAnimControl.h"
 #include "Matinee/InterpData.h"
+#include "Animation/SkeletalMeshActor.h"
 
 #include "Editor/UnrealEd/Public/Kismet2/KismetEditorUtilities.h"
 #include "Editor/PropertyEditor/Public/PropertyEditorModule.h"
diff --git a/Engine/Source/Editor/UnrealEd/Private/Factories/ActorFactory.cpp b/Engine/Source/Editor/UnrealEd/Private/Factories/ActorFactory.cpp
--- a/Engine/Source/Editor/UnrealEd/Private/Factories/ActorFactory.cpp
+++ b/Engine/Source/Editor/UnrealEd/Private/Factories/ActorFactory.cpp
@@ -17,6 +17,8 @@ ActorFactory.cpp:
 #include "Matinee/MatineeActor.h"
 #include "Matinee/InterpData.h"
 #include "Engine/InteractiveFoliageActor.h"
+#include "Animation/SkeletalMeshActor.h"
+#include "Animation/VertexAnim/VertexAnimation.h"
 
 #include "AssetRegistryModule.h"
 
diff --git a/Engine/Source/Editor/UnrealEd/Private/Fbx/FbxAnimationExport.cpp b/Engine/Source/Editor/UnrealEd/Private/Fbx/FbxAnimationExport.cpp
--- a/Engine/Source/Editor/UnrealEd/Private/Fbx/FbxAnimationExport.cpp
+++ b/Engine/Source/Editor/UnrealEd/Private/Fbx/FbxAnimationExport.cpp
@@ -35,7 +35,7 @@
 #include "UnrealEd.h"
 #include "Matinee/InterpData.h"
 #include "Matinee/InterpTrackAnimControl.h"
-
+#include "Animation/SkeletalMeshActor.h"
 
 DEFINE_LOG_CATEGORY_STATIC(LogFbxAnimationExport, Log, All);
 
diff --git a/Engine/Source/Editor/UnrealEd/Private/LevelEditorViewport.cpp b/Engine/Source/Editor/UnrealEd/Private/LevelEditorViewport.cpp
--- a/Engine/Source/Editor/UnrealEd/Private/LevelEditorViewport.cpp
+++ b/Engine/Source/Editor/UnrealEd/Private/LevelEditorViewport.cpp
@@ -3,6 +3,7 @@
 
 #include "UnrealEd.h"
 #include "Materials/MaterialExpressionTextureSample.h"
+#include "Animation/SkeletalMeshActor.h"
 #include "EditorSupportDelegates.h"
 #include "SoundDefinitions.h"
 #include "CameraController.h"
@@ -40,6 +41,7 @@
 #include "LevelEditorActions.h"
 #include "BrushBuilderDragDropOp.h"
 #include "AssetRegistryModule.h"
+#include "Animation/VertexAnim/VertexAnimation.h"
 
 DEFINE_LOG_CATEGORY(LogEditorViewport);
 
diff --git a/Engine/Source/Editor/UnrealEd/Private/ThumbnailHelpers.cpp b/Engine/Source/Editor/UnrealEd/Private/ThumbnailHelpers.cpp
--- a/Engine/Source/Editor/UnrealEd/Private/ThumbnailHelpers.cpp
+++ b/Engine/Source/Editor/UnrealEd/Private/ThumbnailHelpers.cpp
@@ -3,7 +3,7 @@
 #include "UnrealEd.h"
 #include "ThumbnailHelpers.h"
 #include "FXSystem.h"
-
+#include "Animation/SkeletalMeshActor.h"
 #include "Particles/ParticleSystem.h"
 #include "Particles/ParticleSystemComponent.h"
 /*
diff --git a/Engine/Source/Editor/UnrealEd/Private/UnrealEdSrv.cpp b/Engine/Source/Editor/UnrealEd/Private/UnrealEdSrv.cpp
--- a/Engine/Source/Editor/UnrealEd/Private/UnrealEdSrv.cpp
+++ b/Engine/Source/Editor/UnrealEd/Private/UnrealEdSrv.cpp
@@ -25,6 +25,7 @@
 #include "HighResScreenshot.h"
 #include "ActorEditorUtils.h"
 #include "Matinee/InterpData.h"
+#include "Animation/SkeletalMeshActor.h"
 #include "Landscape/LandscapeInfo.h"
 #include "Landscape/LandscapeProxy.h"
 #include "Landscape/LandscapeGizmoActiveActor.h"
diff --git a/Engine/Source/Runtime/Engine/Classes/Animation/AnimNode_BlendListByBool.h b/Engine/Source/Runtime/Engine/Classes/Animation/AnimNode_BlendListByBool.h
--- a/Engine/Source/Runtime/Engine/Classes/Animation/AnimNode_BlendListByBool.h
+++ b/Engine/Source/Runtime/Engine/Classes/Animation/AnimNode_BlendListByBool.h
@@ -3,6 +3,7 @@
 #pragma once
 
 #include "AnimNodeBase.h"
+#include "Animation/AnimNode_BlendListBase.h"
 #include "AnimNode_BlendListByBool.generated.h"
 
 // This node is effectively a 'branch', picking one of two input poses based on an input Boolean value
diff --git a/Engine/Source/Runtime/Engine/Classes/Animation/AnimNode_BlendListByInt.h b/Engine/Source/Runtime/Engine/Classes/Animation/AnimNode_BlendListByInt.h
--- a/Engine/Source/Runtime/Engine/Classes/Animation/AnimNode_BlendListByInt.h
+++ b/Engine/Source/Runtime/Engine/Classes/Animation/AnimNode_BlendListByInt.h
@@ -3,6 +3,7 @@
 #pragma once
 
 #include "AnimNodeBase.h"
+#include "Animation/AnimNode_BlendListBase.h"
 #include "AnimNode_BlendListByInt.generated.h"
 
 // Blend list node; has many children
diff --git a/Engine/Source/Runtime/Engine/Classes/Animation/AnimNode_RefPose.h b/Engine/Source/Runtime/Engine/Classes/Animation/AnimNode_RefPose.h
--- a/Engine/Source/Runtime/Engine/Classes/Animation/AnimNode_RefPose.h
+++ b/Engine/Source/Runtime/Engine/Classes/Animation/AnimNode_RefPose.h
@@ -1,7 +1,7 @@
 // Copyright 1998-2014 Epic Games, Inc. All Rights Reserved.
 
 #pragma once
-
+#include "Animation/AnimNodeBase.h"
 #include "AnimNode_RefPose.generated.h"
 
 UENUM()
diff --git a/Engine/Source/Runtime/Engine/Classes/Animation/AnimNode_Root.h b/Engine/Source/Runtime/Engine/Classes/Animation/AnimNode_Root.h
--- a/Engine/Source/Runtime/Engine/Classes/Animation/AnimNode_Root.h
+++ b/Engine/Source/Runtime/Engine/Classes/Animation/AnimNode_Root.h
@@ -1,7 +1,7 @@
 // Copyright 1998-2014 Epic Games, Inc. All Rights Reserved.
 
 #pragma once
-
+#include "Animation/AnimNodeBase.h"
 #include "AnimNode_Root.generated.h"
 
 // Root node of an animation tree (sink)
diff --git a/Engine/Source/Runtime/Engine/Classes/Animation/AnimNode_SaveCachedPose.h b/Engine/Source/Runtime/Engine/Classes/Animation/AnimNode_SaveCachedPose.h
--- a/Engine/Source/Runtime/Engine/Classes/Animation/AnimNode_SaveCachedPose.h
+++ b/Engine/Source/Runtime/Engine/Classes/Animation/AnimNode_SaveCachedPose.h
@@ -1,7 +1,7 @@
 // Copyright 1998-2014 Epic Games, Inc. All Rights Reserved.
 
 #pragma once
-
+#include "Animation/AnimNodeBase.h"
 #include "AnimNode_SaveCachedPose.generated.h"
 
 USTRUCT()
diff --git a/Engine/Source/Runtime/Engine/Classes/Animation/AnimNode_SequenceEvaluator.h b/Engine/Source/Runtime/Engine/Classes/Animation/AnimNode_SequenceEvaluator.h
--- a/Engine/Source/Runtime/Engine/Classes/Animation/AnimNode_SequenceEvaluator.h
+++ b/Engine/Source/Runtime/Engine/Classes/Animation/AnimNode_SequenceEvaluator.h
@@ -1,7 +1,7 @@
 // Copyright 1998-2014 Epic Games, Inc. All Rights Reserved.
 
 #pragma once
-
+#include "Animation/AnimNodeBase.h"
 #include "AnimNode_SequenceEvaluator.generated.h"
 
 // Evaluates a point in an anim sequence, using a specific time input rather than advancing time internally.
diff --git a/Engine/Source/Runtime/Engine/Classes/Animation/AnimNode_SequencePlayer.h b/Engine/Source/Runtime/Engine/Classes/Animation/AnimNode_SequencePlayer.h
--- a/Engine/Source/Runtime/Engine/Classes/Animation/AnimNode_SequencePlayer.h
+++ b/Engine/Source/Runtime/Engine/Classes/Animation/AnimNode_SequencePlayer.h
@@ -1,6 +1,7 @@
 // Copyright 1998-2014 Epic Games, Inc. All Rights Reserved.
 
 #pragma once
+#include "Animation/AnimNodeBase.h"
 #include "AnimNode_SequencePlayer.generated.h"
 
 #pragma once
diff --git a/Engine/Source/Runtime/Engine/Classes/Animation/AnimNode_Slot.h b/Engine/Source/Runtime/Engine/Classes/Animation/AnimNode_Slot.h
--- a/Engine/Source/Runtime/Engine/Classes/Animation/AnimNode_Slot.h
+++ b/Engine/Source/Runtime/Engine/Classes/Animation/AnimNode_Slot.h
@@ -1,7 +1,7 @@
 // Copyright 1998-2014 Epic Games, Inc. All Rights Reserved.
 
 #pragma once
-
+#include "Animation/AnimNodeBase.h"
 #include "AnimNode_Slot.generated.h"
 
 // An animation slot node normally acts as a passthru, but a montage or PlaySlotAnimation call from
diff --git a/Engine/Source/Runtime/Engine/Classes/Animation/AnimNode_TransitionPoseEvaluator.h b/Engine/Source/Runtime/Engine/Classes/Animation/AnimNode_TransitionPoseEvaluator.h
--- a/Engine/Source/Runtime/Engine/Classes/Animation/AnimNode_TransitionPoseEvaluator.h
+++ b/Engine/Source/Runtime/Engine/Classes/Animation/AnimNode_TransitionPoseEvaluator.h
@@ -1,7 +1,7 @@
 // Copyright 1998-2014 Epic Games, Inc. All Rights Reserved.
 
 #pragma once
-
+#include "Animation/AnimNodeBase.h"
 #include "AnimNode_TransitionPoseEvaluator.generated.h"
 
 // Indicates which state is being evaluated by this node (source or destination)
diff --git a/Engine/Source/Runtime/Engine/Classes/Animation/AnimNode_TransitionResult.h b/Engine/Source/Runtime/Engine/Classes/Animation/AnimNode_TransitionResult.h
--- a/Engine/Source/Runtime/Engine/Classes/Animation/AnimNode_TransitionResult.h
+++ b/Engine/Source/Runtime/Engine/Classes/Animation/AnimNode_TransitionResult.h
@@ -1,7 +1,7 @@
 // Copyright 1998-2014 Epic Games, Inc. All Rights Reserved.
 
 #pragma once
-
+#include "Animation/AnimNodeBase.h"
 #include "AnimNode_TransitionResult.generated.h"
 
 // Root node of a state machine transition graph
diff --git a/Engine/Source/Runtime/Engine/Classes/Animation/AnimNode_TwoWayBlend.h b/Engine/Source/Runtime/Engine/Classes/Animation/AnimNode_TwoWayBlend.h
--- a/Engine/Source/Runtime/Engine/Classes/Animation/AnimNode_TwoWayBlend.h
+++ b/Engine/Source/Runtime/Engine/Classes/Animation/AnimNode_TwoWayBlend.h
@@ -3,6 +3,7 @@
 #pragma once
 
 #include "InputScaleBias.h"
+#include "Animation/AnimNodeBase.h"
 #include "AnimNode_TwoWayBlend.generated.h"
 
 // This represents a baked transition
diff --git a/Engine/Source/Runtime/Engine/Classes/Animation/AnimNode_UseCachedPose.h b/Engine/Source/Runtime/Engine/Classes/Animation/AnimNode_UseCachedPose.h
--- a/Engine/Source/Runtime/Engine/Classes/Animation/AnimNode_UseCachedPose.h
+++ b/Engine/Source/Runtime/Engine/Classes/Animation/AnimNode_UseCachedPose.h
@@ -1,7 +1,7 @@
 // Copyright 1998-2014 Epic Games, Inc. All Rights Reserved.
 
 #pragma once
-
+#include "Animation/AnimNodeBase.h"
 #include "AnimNode_UseCachedPose.generated.h"
 
 USTRUCT()
diff --git a/Engine/Source/Runtime/Engine/Classes/Animation/AnimationTypes.h b/Engine/Source/Runtime/Engine/Classes/Animation/AnimationTypes.h
deleted file mode 100644
diff --git a/Engine/Source/Runtime/Engine/Classes/Animation/BoneControllers/AnimNode_SkeletalControlBase.h b/Engine/Source/Runtime/Engine/Classes/Animation/BoneControllers/AnimNode_SkeletalControlBase.h
--- a/Engine/Source/Runtime/Engine/Classes/Animation/BoneControllers/AnimNode_SkeletalControlBase.h
+++ b/Engine/Source/Runtime/Engine/Classes/Animation/BoneControllers/AnimNode_SkeletalControlBase.h
@@ -11,7 +11,6 @@
 #include "../InputScaleBias.h"
 #include "AnimNode_SkeletalControlBase.generated.h"
 
-DECLARE_LOG_CATEGORY_EXTERN(LogSkeletalControl, Warning, All);
 
 USTRUCT()
 struct ENGINE_API FAnimNode_SkeletalControlBase : public FAnimNode_Base
diff --git a/Engine/Source/Runtime/Engine/Classes/Animation/VertexAnim/MorphTarget.h b/Engine/Source/Runtime/Engine/Classes/Animation/VertexAnim/MorphTarget.h
--- a/Engine/Source/Runtime/Engine/Classes/Animation/VertexAnim/MorphTarget.h
+++ b/Engine/Source/Runtime/Engine/Classes/Animation/VertexAnim/MorphTarget.h
@@ -1,7 +1,7 @@
 // Copyright 1998-2014 Epic Games, Inc. All Rights Reserved.
 
 #pragma once
-
+#include "Animation/VertexAnim/VertexAnimBase.h"
 #include "MorphTarget.generated.h"
 
 class USkeletalMesh;
diff --git a/Engine/Source/Runtime/Engine/Classes/Animation/VertexAnim/VertexAnimation.h b/Engine/Source/Runtime/Engine/Classes/Animation/VertexAnim/VertexAnimation.h
--- a/Engine/Source/Runtime/Engine/Classes/Animation/VertexAnim/VertexAnimation.h
+++ b/Engine/Source/Runtime/Engine/Classes/Animation/VertexAnim/VertexAnimation.h
@@ -2,6 +2,7 @@
 
 
 #pragma once
+#include "Animation/VertexAnim/VertexAnimBase.h"
 #include "VertexAnimation.generated.h"
 
 
diff --git a/Engine/Source/Runtime/Engine/Private/Animation/AnimBlueprintGeneratedClass.cpp b/Engine/Source/Runtime/Engine/Private/Animation/AnimBlueprintGeneratedClass.cpp
--- a/Engine/Source/Runtime/Engine/Private/Animation/AnimBlueprintGeneratedClass.cpp
+++ b/Engine/Source/Runtime/Engine/Private/Animation/AnimBlueprintGeneratedClass.cpp
@@ -5,6 +5,7 @@
 =============================================================================*/ 
 
 #include "EnginePrivate.h"
+#include "Animation/AnimNodeBase.h"
 
 /////////////////////////////////////////////////////
 // FStateMachineDebugData
diff --git a/Engine/Source/Runtime/Engine/Private/Animation/AnimInstance.cpp b/Engine/Source/Runtime/Engine/Private/Animation/AnimInstance.cpp
--- a/Engine/Source/Runtime/Engine/Private/Animation/AnimInstance.cpp
+++ b/Engine/Source/Runtime/Engine/Private/Animation/AnimInstance.cpp
@@ -8,6 +8,7 @@
 #include "AnimationRuntime.h"
 #include "AnimationUtils.h"
 #include "AnimTree.h"
+#include "Animation/AnimNode_StateMachine.h"
 #include "ParticleDefinitions.h"
 #include "DisplayDebugHelpers.h"
 
diff --git a/Engine/Source/Runtime/Engine/Private/Animation/AnimNodeBase.cpp b/Engine/Source/Runtime/Engine/Private/Animation/AnimNodeBase.cpp
--- a/Engine/Source/Runtime/Engine/Private/Animation/AnimNodeBase.cpp
+++ b/Engine/Source/Runtime/Engine/Private/Animation/AnimNodeBase.cpp
@@ -1,6 +1,7 @@
 // Copyright 1998-2014 Epic Games, Inc. All Rights Reserved.
 
 #include "EnginePrivate.h"
+#include "Animation/AnimNodeBase.h"
 
 /////////////////////////////////////////////////////
 // FPoseLinkBase
diff --git a/Engine/Source/Runtime/Engine/Private/Animation/AnimNodeRotationMultiplier.cpp b/Engine/Source/Runtime/Engine/Private/Animation/AnimNodeRotationMultiplier.cpp
--- a/Engine/Source/Runtime/Engine/Private/Animation/AnimNodeRotationMultiplier.cpp
+++ b/Engine/Source/Runtime/Engine/Private/Animation/AnimNodeRotationMultiplier.cpp
@@ -1,6 +1,7 @@
 // Copyright 1998-2014 Epic Games, Inc. All Rights Reserved.
 
 #include "EnginePrivate.h"
+#include "Animation/BoneControllers/AnimNode_RotationMultiplier.h"
 #include "AnimationRuntime.h"
 
 /////////////////////////////////////////////////////
diff --git a/Engine/Source/Runtime/Engine/Private/Animation/AnimNodeSkeletalControlBase.cpp b/Engine/Source/Runtime/Engine/Private/Animation/AnimNodeSkeletalControlBase.cpp
--- a/Engine/Source/Runtime/Engine/Private/Animation/AnimNodeSkeletalControlBase.cpp
+++ b/Engine/Source/Runtime/Engine/Private/Animation/AnimNodeSkeletalControlBase.cpp
@@ -1,6 +1,7 @@
 // Copyright 1998-2014 Epic Games, Inc. All Rights Reserved.
 
 #include "EnginePrivate.h"
+#include "Animation/BoneControllers/AnimNode_SkeletalControlBase.h"
 #include "AnimationUtils.h"
 #include "AnimationRuntime.h"
 
diff --git a/Engine/Source/Runtime/Engine/Private/Animation/AnimNodeSpaceConversions.cpp b/Engine/Source/Runtime/Engine/Private/Animation/AnimNodeSpaceConversions.cpp
--- a/Engine/Source/Runtime/Engine/Private/Animation/AnimNodeSpaceConversions.cpp
+++ b/Engine/Source/Runtime/Engine/Private/Animation/AnimNodeSpaceConversions.cpp
@@ -1,6 +1,7 @@
 // Copyright 1998-2014 Epic Games, Inc. All Rights Reserved.
 
 #include "EnginePrivate.h"
+#include "Animation/AnimNodeSpaceConversions.h"
 
 /////////////////////////////////////////////////////
 // FAnimNode_ConvertComponentToLocalSpace
diff --git a/Engine/Source/Runtime/Engine/Private/Animation/AnimNode_ApplyAdditive.cpp b/Engine/Source/Runtime/Engine/Private/Animation/AnimNode_ApplyAdditive.cpp
--- a/Engine/Source/Runtime/Engine/Private/Animation/AnimNode_ApplyAdditive.cpp
+++ b/Engine/Source/Runtime/Engine/Private/Animation/AnimNode_ApplyAdditive.cpp
@@ -1,6 +1,7 @@
 // Copyright 1998-2014 Epic Games, Inc. All Rights Reserved.
 
 #include "EnginePrivate.h"
+#include "Animation/AnimNode_ApplyAdditive.h"
 
 /////////////////////////////////////////////////////
 // FAnimNode_ApplyAdditive
diff --git a/Engine/Source/Runtime/Engine/Private/Animation/AnimNode_BlendListBase.cpp b/Engine/Source/Runtime/Engine/Private/Animation/AnimNode_BlendListBase.cpp
--- a/Engine/Source/Runtime/Engine/Private/Animation/AnimNode_BlendListBase.cpp
+++ b/Engine/Source/Runtime/Engine/Private/Animation/AnimNode_BlendListBase.cpp
@@ -1,6 +1,7 @@
 // Copyright 1998-2014 Epic Games, Inc. All Rights Reserved.
 
 #include "EnginePrivate.h"
+#include "Animation/AnimNode_BlendListBase.h"
 #include "AnimationRuntime.h"
 #include "AnimTree.h"
 
diff --git a/Engine/Source/Runtime/Engine/Private/Animation/AnimNode_BlendListByBool.cpp b/Engine/Source/Runtime/Engine/Private/Animation/AnimNode_BlendListByBool.cpp
--- a/Engine/Source/Runtime/Engine/Private/Animation/AnimNode_BlendListByBool.cpp
+++ b/Engine/Source/Runtime/Engine/Private/Animation/AnimNode_BlendListByBool.cpp
@@ -1,6 +1,7 @@
 // Copyright 1998-2014 Epic Games, Inc. All Rights Reserved.
 
 #include "EnginePrivate.h"
+#include "Animation/AnimNode_BlendListByBool.h"
 #include "AnimationRuntime.h"
 
 /////////////////////////////////////////////////////
diff --git a/Engine/Source/Runtime/Engine/Private/Animation/AnimNode_BlendListByEnum.cpp b/Engine/Source/Runtime/Engine/Private/Animation/AnimNode_BlendListByEnum.cpp
--- a/Engine/Source/Runtime/Engine/Private/Animation/AnimNode_BlendListByEnum.cpp
+++ b/Engine/Source/Runtime/Engine/Private/Animation/AnimNode_BlendListByEnum.cpp
@@ -2,6 +2,7 @@
 
 #include "EnginePrivate.h"
 #include "AnimationRuntime.h"
+#include "Animation/AnimNode_BlendListByEnum.h"
 
 /////////////////////////////////////////////////////
 // FAnimNode_BlendListByEnum
diff --git a/Engine/Source/Runtime/Engine/Private/Animation/AnimNode_BlendListByInt.cpp b/Engine/Source/Runtime/Engine/Private/Animation/AnimNode_BlendListByInt.cpp
--- a/Engine/Source/Runtime/Engine/Private/Animation/AnimNode_BlendListByInt.cpp
+++ b/Engine/Source/Runtime/Engine/Private/Animation/AnimNode_BlendListByInt.cpp
@@ -1,6 +1,7 @@
 // Copyright 1998-2014 Epic Games, Inc. All Rights Reserved.
 
 #include "EnginePrivate.h"
+#include "Animation/AnimNode_BlendListByInt.h"
 #include "AnimationRuntime.h"
 
 /////////////////////////////////////////////////////
diff --git a/Engine/Source/Runtime/Engine/Private/Animation/AnimNode_BlendSpaceEvaluator.cpp b/Engine/Source/Runtime/Engine/Private/Animation/AnimNode_BlendSpaceEvaluator.cpp
--- a/Engine/Source/Runtime/Engine/Private/Animation/AnimNode_BlendSpaceEvaluator.cpp
+++ b/Engine/Source/Runtime/Engine/Private/Animation/AnimNode_BlendSpaceEvaluator.cpp
@@ -1,6 +1,7 @@
 // Copyright 1998-2014 Epic Games, Inc. All Rights Reserved.
 
 #include "EnginePrivate.h"
+#include "Animation/AnimNode_BlendSpaceEvaluator.h"
 
 /////////////////////////////////////////////////////
 // FAnimNode_BlendSpaceEvaluator
diff --git a/Engine/Source/Runtime/Engine/Private/Animation/AnimNode_BlendSpacePlayer.cpp b/Engine/Source/Runtime/Engine/Private/Animation/AnimNode_BlendSpacePlayer.cpp
--- a/Engine/Source/Runtime/Engine/Private/Animation/AnimNode_BlendSpacePlayer.cpp
+++ b/Engine/Source/Runtime/Engine/Private/Animation/AnimNode_BlendSpacePlayer.cpp
@@ -1,6 +1,7 @@
 // Copyright 1998-2014 Epic Games, Inc. All Rights Reserved.
 
 #include "EnginePrivate.h"
+#include "Animation/AnimNode_BlendSpacePlayer.h"
 
 /////////////////////////////////////////////////////
 // FAnimNode_BlendSpacePlayer
diff --git a/Engine/Source/Runtime/Engine/Private/Animation/AnimNode_CopyBone.cpp b/Engine/Source/Runtime/Engine/Private/Animation/AnimNode_CopyBone.cpp
--- a/Engine/Source/Runtime/Engine/Private/Animation/AnimNode_CopyBone.cpp
+++ b/Engine/Source/Runtime/Engine/Private/Animation/AnimNode_CopyBone.cpp
@@ -1,6 +1,7 @@
 // Copyright 1998-2014 Epic Games, Inc. All Rights Reserved.
 
 #include "EnginePrivate.h"
+#include "Animation/BoneControllers/AnimNode_CopyBone.h"
 
 /////////////////////////////////////////////////////
 // FAnimNode_CopyBone
diff --git a/Engine/Source/Runtime/Engine/Private/Animation/AnimNode_Fabrik.cpp b/Engine/Source/Runtime/Engine/Private/Animation/AnimNode_Fabrik.cpp
--- a/Engine/Source/Runtime/Engine/Private/Animation/AnimNode_Fabrik.cpp
+++ b/Engine/Source/Runtime/Engine/Private/Animation/AnimNode_Fabrik.cpp
@@ -1,6 +1,7 @@
 // Copyright 1998-2013 Epic Games, Inc. All Rights Reserved.
 
 #include "EnginePrivate.h"
+#include "Animation/BoneControllers/AnimNode_Fabrik.h"
 
 /////////////////////////////////////////////////////
 // AnimNode_Fabrik
diff --git a/Engine/Source/Runtime/Engine/Private/Animation/AnimNode_HandIKRetargeting.cpp b/Engine/Source/Runtime/Engine/Private/Animation/AnimNode_HandIKRetargeting.cpp
--- a/Engine/Source/Runtime/Engine/Private/Animation/AnimNode_HandIKRetargeting.cpp
+++ b/Engine/Source/Runtime/Engine/Private/Animation/AnimNode_HandIKRetargeting.cpp
@@ -1,6 +1,7 @@
 // Copyright 1998-2014 Epic Games, Inc. All Rights Reserved.
 
 #include "EnginePrivate.h"
+#include "Animation/BoneControllers/AnimNode_HandIKRetargeting.h"
 
 /////////////////////////////////////////////////////
 // FAnimNode_HandIKRetargeting
diff --git a/Engine/Source/Runtime/Engine/Private/Animation/AnimNode_LayeredBoneBlend.cpp b/Engine/Source/Runtime/Engine/Private/Animation/AnimNode_LayeredBoneBlend.cpp
--- a/Engine/Source/Runtime/Engine/Private/Animation/AnimNode_LayeredBoneBlend.cpp
+++ b/Engine/Source/Runtime/Engine/Private/Animation/AnimNode_LayeredBoneBlend.cpp
@@ -1,6 +1,7 @@
 // Copyright 1998-2014 Epic Games, Inc. All Rights Reserved.
 
 #include "EnginePrivate.h"
+#include "Animation/AnimNode_LayeredBoneBlend.h"
 #include "AnimationRuntime.h"
 #include "AnimTree.h"
 
diff --git a/Engine/Source/Runtime/Engine/Private/Animation/AnimNode_LookAt.cpp b/Engine/Source/Runtime/Engine/Private/Animation/AnimNode_LookAt.cpp
--- a/Engine/Source/Runtime/Engine/Private/Animation/AnimNode_LookAt.cpp
+++ b/Engine/Source/Runtime/Engine/Private/Animation/AnimNode_LookAt.cpp
@@ -1,6 +1,7 @@
 // Copyright 1998-2014 Epic Games, Inc. All Rights Reserved.
 
 #include "EnginePrivate.h"
+#include "Animation/BoneControllers/AnimNode_LookAt.h"
 #include "AnimationRuntime.h"
 
 /////////////////////////////////////////////////////
diff --git a/Engine/Source/Runtime/Engine/Private/Animation/AnimNode_ModifyBone.cpp b/Engine/Source/Runtime/Engine/Private/Animation/AnimNode_ModifyBone.cpp
--- a/Engine/Source/Runtime/Engine/Private/Animation/AnimNode_ModifyBone.cpp
+++ b/Engine/Source/Runtime/Engine/Private/Animation/AnimNode_ModifyBone.cpp
@@ -1,6 +1,7 @@
 // Copyright 1998-2014 Epic Games, Inc. All Rights Reserved.
 
 #include "EnginePrivate.h"
+#include "Animation/BoneControllers/AnimNode_ModifyBone.h"
 #include "AnimationRuntime.h"
 
 /////////////////////////////////////////////////////
diff --git a/Engine/Source/Runtime/Engine/Private/Animation/AnimNode_RefPose.cpp b/Engine/Source/Runtime/Engine/Private/Animation/AnimNode_RefPose.cpp
--- a/Engine/Source/Runtime/Engine/Private/Animation/AnimNode_RefPose.cpp
+++ b/Engine/Source/Runtime/Engine/Private/Animation/AnimNode_RefPose.cpp
@@ -1,6 +1,7 @@
 // Copyright 1998-2014 Epic Games, Inc. All Rights Reserved.
 
 #include "EnginePrivate.h"
+#include "Animation/AnimNode_RefPose.h"
 
 /////////////////////////////////////////////////////
 // FAnimRefPoseNode
diff --git a/Engine/Source/Runtime/Engine/Private/Animation/AnimNode_Root.cpp b/Engine/Source/Runtime/Engine/Private/Animation/AnimNode_Root.cpp
--- a/Engine/Source/Runtime/Engine/Private/Animation/AnimNode_Root.cpp
+++ b/Engine/Source/Runtime/Engine/Private/Animation/AnimNode_Root.cpp
@@ -1,6 +1,7 @@
 // Copyright 1998-2014 Epic Games, Inc. All Rights Reserved.
 
 #include "EnginePrivate.h"
+#include "Animation/AnimNode_Root.h"
 
 /////////////////////////////////////////////////////
 // FAnimNode_Root
diff --git a/Engine/Source/Runtime/Engine/Private/Animation/AnimNode_RotateRootBone.cpp b/Engine/Source/Runtime/Engine/Private/Animation/AnimNode_RotateRootBone.cpp
--- a/Engine/Source/Runtime/Engine/Private/Animation/AnimNode_RotateRootBone.cpp
+++ b/Engine/Source/Runtime/Engine/Private/Animation/AnimNode_RotateRootBone.cpp
@@ -1,6 +1,7 @@
 // Copyright 1998-2014 Epic Games, Inc. All Rights Reserved.
 
 #include "EnginePrivate.h"
+#include "Animation/AnimNode_RotateRootBone.h"
 
 /////////////////////////////////////////////////////
 // FAnimNode_RotateRootBone
diff --git a/Engine/Source/Runtime/Engine/Private/Animation/AnimNode_RotationOffsetBlendSpace.cpp b/Engine/Source/Runtime/Engine/Private/Animation/AnimNode_RotationOffsetBlendSpace.cpp
--- a/Engine/Source/Runtime/Engine/Private/Animation/AnimNode_RotationOffsetBlendSpace.cpp
+++ b/Engine/Source/Runtime/Engine/Private/Animation/AnimNode_RotationOffsetBlendSpace.cpp
@@ -1,6 +1,7 @@
 // Copyright 1998-2014 Epic Games, Inc. All Rights Reserved.
 
 #include "EnginePrivate.h"
+#include "Animation/AnimNode_RotationOffsetBlendSpace.h"
 
 /////////////////////////////////////////////////////
 // FAnimNode_RotationOffsetBlendSpace
diff --git a/Engine/Source/Runtime/Engine/Private/Animation/AnimNode_SaveCachedPose.cpp b/Engine/Source/Runtime/Engine/Private/Animation/AnimNode_SaveCachedPose.cpp
--- a/Engine/Source/Runtime/Engine/Private/Animation/AnimNode_SaveCachedPose.cpp
+++ b/Engine/Source/Runtime/Engine/Private/Animation/AnimNode_SaveCachedPose.cpp
@@ -1,6 +1,7 @@
 // Copyright 1998-2014 Epic Games, Inc. All Rights Reserved.
 
 #include "EnginePrivate.h"
+#include "Animation/AnimNode_SaveCachedPose.h"
 
 /////////////////////////////////////////////////////
 // FAnimNode_SaveCachedPose
diff --git a/Engine/Source/Runtime/Engine/Private/Animation/AnimNode_SequenceEvaluator.cpp b/Engine/Source/Runtime/Engine/Private/Animation/AnimNode_SequenceEvaluator.cpp
--- a/Engine/Source/Runtime/Engine/Private/Animation/AnimNode_SequenceEvaluator.cpp
+++ b/Engine/Source/Runtime/Engine/Private/Animation/AnimNode_SequenceEvaluator.cpp
@@ -1,6 +1,7 @@
 // Copyright 1998-2014 Epic Games, Inc. All Rights Reserved.
 
 #include "EnginePrivate.h"
+#include "Animation/AnimNode_SequenceEvaluator.h"
 
 /////////////////////////////////////////////////////
 // FAnimSequenceEvaluatorNode
diff --git a/Engine/Source/Runtime/Engine/Private/Animation/AnimNode_SequencePlayer.cpp b/Engine/Source/Runtime/Engine/Private/Animation/AnimNode_SequencePlayer.cpp
--- a/Engine/Source/Runtime/Engine/Private/Animation/AnimNode_SequencePlayer.cpp
+++ b/Engine/Source/Runtime/Engine/Private/Animation/AnimNode_SequencePlayer.cpp
@@ -1,6 +1,7 @@
 // Copyright 1998-2014 Epic Games, Inc. All Rights Reserved.
 
 #include "EnginePrivate.h"
+#include "Animation/AnimNode_SequencePlayer.h"
 
 /////////////////////////////////////////////////////
 // FAnimSequencePlayerNode
diff --git a/Engine/Source/Runtime/Engine/Private/Animation/AnimNode_Slot.cpp b/Engine/Source/Runtime/Engine/Private/Animation/AnimNode_Slot.cpp
--- a/Engine/Source/Runtime/Engine/Private/Animation/AnimNode_Slot.cpp
+++ b/Engine/Source/Runtime/Engine/Private/Animation/AnimNode_Slot.cpp
@@ -1,6 +1,7 @@
 // Copyright 1998-2014 Epic Games, Inc. All Rights Reserved.
 
 #include "EnginePrivate.h"
+#include "Animation/AnimNode_Slot.h"
 
 /////////////////////////////////////////////////////
 // FAnimNode_Slot
diff --git a/Engine/Source/Runtime/Engine/Private/Animation/AnimNode_SpringBone.cpp b/Engine/Source/Runtime/Engine/Private/Animation/AnimNode_SpringBone.cpp
--- a/Engine/Source/Runtime/Engine/Private/Animation/AnimNode_SpringBone.cpp
+++ b/Engine/Source/Runtime/Engine/Private/Animation/AnimNode_SpringBone.cpp
@@ -1,6 +1,7 @@
 // Copyright 1998-2014 Epic Games, Inc. All Rights Reserved.
 
 #include "EnginePrivate.h"
+#include "Animation/BoneControllers/AnimNode_SpringBone.h"
 
 /////////////////////////////////////////////////////
 // FAnimNode_SpringBone
diff --git a/Engine/Source/Runtime/Engine/Private/Animation/AnimNode_StateMachine.cpp b/Engine/Source/Runtime/Engine/Private/Animation/AnimNode_StateMachine.cpp
--- a/Engine/Source/Runtime/Engine/Private/Animation/AnimNode_StateMachine.cpp
+++ b/Engine/Source/Runtime/Engine/Private/Animation/AnimNode_StateMachine.cpp
@@ -1,6 +1,9 @@
 // Copyright 1998-2014 Epic Games, Inc. All Rights Reserved.
 
 #include "EnginePrivate.h"
+#include "Animation/AnimNode_StateMachine.h"
+#include "Animation/AnimNode_TransitionResult.h"
+#include "Animation/AnimNode_TransitionPoseEvaluator.h"
 #include "AnimationRuntime.h"
 #include "AnimTree.h"
 
diff --git a/Engine/Source/Runtime/Engine/Private/Animation/AnimNode_TransitionPoseEvaluator.cpp b/Engine/Source/Runtime/Engine/Private/Animation/AnimNode_TransitionPoseEvaluator.cpp
--- a/Engine/Source/Runtime/Engine/Private/Animation/AnimNode_TransitionPoseEvaluator.cpp
+++ b/Engine/Source/Runtime/Engine/Private/Animation/AnimNode_TransitionPoseEvaluator.cpp
@@ -1,6 +1,7 @@
 // Copyright 1998-2014 Epic Games, Inc. All Rights Reserved.
 
 #include "EnginePrivate.h"
+#include "Animation/AnimNode_TransitionPoseEvaluator.h"
 
 /////////////////////////////////////////////////////
 // FAnimNode_TransitionPoseEvaluator
diff --git a/Engine/Source/Runtime/Engine/Private/Animation/AnimNode_TransitionResult.cpp b/Engine/Source/Runtime/Engine/Private/Animation/AnimNode_TransitionResult.cpp
--- a/Engine/Source/Runtime/Engine/Private/Animation/AnimNode_TransitionResult.cpp
+++ b/Engine/Source/Runtime/Engine/Private/Animation/AnimNode_TransitionResult.cpp
@@ -1,6 +1,7 @@
 // Copyright 1998-2014 Epic Games, Inc. All Rights Reserved.
 
 #include "EnginePrivate.h"
+#include "Animation/AnimNode_TransitionResult.h"
 
 /////////////////////////////////////////////////////
 // FAnimNode_TransitionResult
diff --git a/Engine/Source/Runtime/Engine/Private/Animation/AnimNode_TwoBoneIK.cpp b/Engine/Source/Runtime/Engine/Private/Animation/AnimNode_TwoBoneIK.cpp
--- a/Engine/Source/Runtime/Engine/Private/Animation/AnimNode_TwoBoneIK.cpp
+++ b/Engine/Source/Runtime/Engine/Private/Animation/AnimNode_TwoBoneIK.cpp
@@ -1,6 +1,7 @@
 // Copyright 1998-2014 Epic Games, Inc. All Rights Reserved.
 
 #include "EnginePrivate.h"
+#include "Animation/BoneControllers/AnimNode_TwoBoneIK.h"
 #include "AnimationRuntime.h"
 
 /////////////////////////////////////////////////////
diff --git a/Engine/Source/Runtime/Engine/Private/Animation/AnimNode_TwoWayBlend.cpp b/Engine/Source/Runtime/Engine/Private/Animation/AnimNode_TwoWayBlend.cpp
--- a/Engine/Source/Runtime/Engine/Private/Animation/AnimNode_TwoWayBlend.cpp
+++ b/Engine/Source/Runtime/Engine/Private/Animation/AnimNode_TwoWayBlend.cpp
@@ -1,6 +1,7 @@
 // Copyright 1998-2014 Epic Games, Inc. All Rights Reserved.
 
 #include "EnginePrivate.h"
+#include "Animation/AnimNode_TwoWayBlend.h"
 
 /////////////////////////////////////////////////////
 // FAnimationNode_TwoWayBlend
diff --git a/Engine/Source/Runtime/Engine/Private/Animation/AnimNode_UseCachedPose.cpp b/Engine/Source/Runtime/Engine/Private/Animation/AnimNode_UseCachedPose.cpp
--- a/Engine/Source/Runtime/Engine/Private/Animation/AnimNode_UseCachedPose.cpp
+++ b/Engine/Source/Runtime/Engine/Private/Animation/AnimNode_UseCachedPose.cpp
@@ -1,6 +1,7 @@
 // Copyright 1998-2014 Epic Games, Inc. All Rights Reserved.
 
 #include "EnginePrivate.h"
+#include "Animation/AnimNode_UseCachedPose.h"
 
 /////////////////////////////////////////////////////
 // FAnimNode_UseCachedPose
diff --git a/Engine/Source/Runtime/Engine/Private/Animation/AnimNode_WheelHandler.cpp b/Engine/Source/Runtime/Engine/Private/Animation/AnimNode_WheelHandler.cpp
--- a/Engine/Source/Runtime/Engine/Private/Animation/AnimNode_WheelHandler.cpp
+++ b/Engine/Source/Runtime/Engine/Private/Animation/AnimNode_WheelHandler.cpp
@@ -2,6 +2,7 @@
 
 #include "EnginePrivate.h"
 #include "AnimationRuntime.h"
+#include "Animation/BoneControllers/AnimNode_WheelHandler.h"
 
 /////////////////////////////////////////////////////
 // FAnimNode_WheelHandler
diff --git a/Engine/Source/Runtime/Engine/Private/Animation/AnimSingleNodeInstance.cpp b/Engine/Source/Runtime/Engine/Private/Animation/AnimSingleNodeInstance.cpp
--- a/Engine/Source/Runtime/Engine/Private/Animation/AnimSingleNodeInstance.cpp
+++ b/Engine/Source/Runtime/Engine/Private/Animation/AnimSingleNodeInstance.cpp
@@ -6,6 +6,8 @@
 =============================================================================*/ 
 
 #include "EnginePrivate.h"
+#include "Animation/AnimNodeBase.h"
+#include "Animation/VertexAnim/VertexAnimation.h"
 #include "AnimationRuntime.h"
 
 /////////////////////////////////////////////////////
diff --git a/Engine/Source/Runtime/Engine/Private/Animation/AnimationRuntime.cpp b/Engine/Source/Runtime/Engine/Private/Animation/AnimationRuntime.cpp
--- a/Engine/Source/Runtime/Engine/Private/Animation/AnimationRuntime.cpp
+++ b/Engine/Source/Runtime/Engine/Private/Animation/AnimationRuntime.cpp
@@ -9,6 +9,7 @@
 #include "AnimEncoding.h"
 #include "AnimationRuntime.h"
 #include "AnimationUtils.h"
+#include "Animation/AnimData/BoneMaskFilter.h"
 
 DEFINE_LOG_CATEGORY(LogAnimation);
 DEFINE_LOG_CATEGORY(LogRootMotion);
diff --git a/Engine/Source/Runtime/Engine/Private/Animation/AnimationTypes.cpp b/Engine/Source/Runtime/Engine/Private/Animation/AnimationTypes.cpp
deleted file mode 100644
diff --git a/Engine/Source/Runtime/Engine/Private/Animation/InputScaleBias.cpp b/Engine/Source/Runtime/Engine/Private/Animation/InputScaleBias.cpp
--- a/Engine/Source/Runtime/Engine/Private/Animation/InputScaleBias.cpp
+++ b/Engine/Source/Runtime/Engine/Private/Animation/InputScaleBias.cpp
@@ -1,6 +1,7 @@
 // Copyright 1998-2014 Epic Games, Inc. All Rights Reserved.
 
 #include "EnginePrivate.h"
+#include "Animation/InputScaleBias.h"
 
 /////////////////////////////////////////////////////
 // FInputScaleBias
diff --git a/Engine/Source/Runtime/Engine/Private/Animation/SkeletalControl.cpp b/Engine/Source/Runtime/Engine/Private/Animation/SkeletalControl.cpp
--- a/Engine/Source/Runtime/Engine/Private/Animation/SkeletalControl.cpp
+++ b/Engine/Source/Runtime/Engine/Private/Animation/SkeletalControl.cpp
@@ -5,6 +5,7 @@
 =============================================================================*/ 
 
 #include "EnginePrivate.h"
+#include "Animation/AnimData/BoneMaskFilter.h"
 #include "AnimationUtils.h"
 #include "AnimationRuntime.h"
 
diff --git a/Engine/Source/Runtime/Engine/Private/Animation/Skeleton.cpp b/Engine/Source/Runtime/Engine/Private/Animation/Skeleton.cpp
--- a/Engine/Source/Runtime/Engine/Private/Animation/Skeleton.cpp
+++ b/Engine/Source/Runtime/Engine/Private/Animation/Skeleton.cpp
@@ -6,6 +6,7 @@
 
 #include "EnginePrivate.h"
 #include "EngineUtils.h"
+#include "AnimationRuntime.h"
 #include "AssetRegistryModule.h"
 
 #define LOCTEXT_NAMESPACE "Skeleton"
diff --git a/Engine/Source/Runtime/Engine/Private/Animation/VertexAnimation.cpp b/Engine/Source/Runtime/Engine/Private/Animation/VertexAnimation.cpp
--- a/Engine/Source/Runtime/Engine/Private/Animation/VertexAnimation.cpp
+++ b/Engine/Source/Runtime/Engine/Private/Animation/VertexAnimation.cpp
@@ -2,6 +2,7 @@
 
 
 #include "EnginePrivate.h"
+#include "Animation/VertexAnim/VertexAnimation.h"
 
 UVertexAnimation::UVertexAnimation(const class FPostConstructInitializeProperties& PCIP)
 	: Super(PCIP)
diff --git a/Engine/Source/Runtime/Engine/Private/Interpolation.cpp b/Engine/Source/Runtime/Engine/Private/Interpolation.cpp
--- a/Engine/Source/Runtime/Engine/Private/Interpolation.cpp
+++ b/Engine/Source/Runtime/Engine/Private/Interpolation.cpp
@@ -59,6 +59,7 @@
 #include "Matinee/InterpFilter_Classes.h"
 #include "Matinee/InterpFilter_Custom.h"
 #include "Materials/MaterialInstanceActor.h"
+#include "Animation/SkeletalMeshActor.h"
 #include "ParticleDefinitions.h"
 #include "SoundDefinitions.h"
 #include "InterpolationHitProxy.h"
diff --git a/Engine/Source/Runtime/Engine/Private/MatineeAnimInterface.cpp b/Engine/Source/Runtime/Engine/Private/MatineeAnimInterface.cpp
--- a/Engine/Source/Runtime/Engine/Private/MatineeAnimInterface.cpp
+++ b/Engine/Source/Runtime/Engine/Private/MatineeAnimInterface.cpp
@@ -1,6 +1,7 @@
 // Copyright 1998-2014 Epic Games, Inc. All Rights Reserved.
 
 #include "EnginePrivate.h"
+#include "Matinee/MatineeAnimInterface.h"
 
 UMatineeAnimInterface::UMatineeAnimInterface(const class FPostConstructInitializeProperties& PCIP)
 	: Super(PCIP)
diff --git a/Engine/Source/Runtime/Engine/Private/Particles/ParticleModules_Location.cpp b/Engine/Source/Runtime/Engine/Private/Particles/ParticleModules_Location.cpp
--- a/Engine/Source/Runtime/Engine/Private/Particles/ParticleModules_Location.cpp
+++ b/Engine/Source/Runtime/Engine/Private/Particles/ParticleModules_Location.cpp
@@ -28,6 +28,7 @@
 #include "Particles/ParticleModuleRequired.h"
 #include "Particles/ParticleSpriteEmitter.h"
 #include "Particles/ParticleSystemComponent.h"
+#include "Animation/SkeletalMeshActor.h"
 
 UParticleModuleLocationBase::UParticleModuleLocationBase(const class FPostConstructInitializeProperties& PCIP)
 	: Super(PCIP)
diff --git a/Engine/Source/Runtime/Engine/Private/SkeletalMesh.cpp b/Engine/Source/Runtime/Engine/Private/SkeletalMesh.cpp
--- a/Engine/Source/Runtime/Engine/Private/SkeletalMesh.cpp
+++ b/Engine/Source/Runtime/Engine/Private/SkeletalMesh.cpp
@@ -5,6 +5,7 @@
 =============================================================================*/
 
 #include "EnginePrivate.h"
+#include "Animation/SkeletalMeshActor.h"
 #include "EditorSupportDelegates.h"
 #include "GPUSkinVertexFactory.h"
 #include "SkeletalMeshSorting.h"
diff --git a/Engine/Source/Runtime/Engine/Private/SkeletalMeshComponent.cpp b/Engine/Source/Runtime/Engine/Private/SkeletalMeshComponent.cpp
--- a/Engine/Source/Runtime/Engine/Private/SkeletalMeshComponent.cpp
+++ b/Engine/Source/Runtime/Engine/Private/SkeletalMeshComponent.cpp
@@ -14,6 +14,8 @@
 #include "AnimationRuntime.h"
 #include "PhysXASync.h"
 #include "AnimTree.h"
+#include "Animation/AnimNodeBase.h"
+#include "Animation/VertexAnim/VertexAnimation.h"
 
 #include "Particles/ParticleSystemComponent.h"
 
diff --git a/Engine/Source/Runtime/Engine/Private/UnrealEngine.cpp b/Engine/Source/Runtime/Engine/Private/UnrealEngine.cpp
--- a/Engine/Source/Runtime/Engine/Private/UnrealEngine.cpp
+++ b/Engine/Source/Runtime/Engine/Private/UnrealEngine.cpp
@@ -41,6 +41,7 @@
 #include "AudioDevice.h"
 #include "ActiveSound.h"
 #include "DeviceProfiles/DeviceProfileManager.h"
+#include "Animation/SkeletalMeshActor.h"
 
 #include "Particles/Spawn/ParticleModuleSpawn.h"
 #include "Particles/TypeData/ParticleModuleTypeDataMesh.h"
diff --git a/Engine/Source/Runtime/Engine/Public/Engine.h b/Engine/Source/Runtime/Engine/Public/Engine.h
--- a/Engine/Source/Runtime/Engine/Public/Engine.h
+++ b/Engine/Source/Runtime/Engine/Public/Engine.h
@@ -126,7 +126,6 @@
 #include "PhysicsEngine/ConstraintInstance.h"
 #include "PhysicsEngine/PhysicsConstraintComponent.h"
 #include "Components/SkeletalMeshComponent.h"
-#include "Animation/SkeletalMeshActor.h"
 #include "Engine/StaticMeshActor.h"
 #include "Engine/TargetPoint.h"
 #include "Engine/TextRenderActor.h"
@@ -169,7 +168,6 @@
 #include "Components/LineBatchComponent.h"
 #include "Components/MaterialBillboardComponent.h"
 #include "Components/DestructibleComponent.h"
-#include "Animation/AnimationAsset.h"
 #include "Components/PoseableMeshComponent.h"
 #include "Lightmass/LightmassPrimitiveSettingsObject.h"
 #include "Components/StaticMeshComponent.h"
@@ -210,7 +208,6 @@
 #include "Animation/AimOffsetBlendSpace.h"
 #include "Animation/BlendSpace1D.h"
 #include "Animation/AimOffsetBlendSpace1D.h"
-#include "Animation/AnimationTypes.h"
 #include "Animation/AnimSingleNodeInstance.h"
 #include "Vehicles/VehicleAnimInstance.h"
 #include "Animation/AnimNotifies/AnimNotify.h"
@@ -357,40 +354,6 @@
 #include "Sound/SoundMix.h"
 #include "Engine/StaticMesh.h"
 #include "Engine/StaticMeshSocket.h"
-#include "Animation/AnimNodeBase.h"
-#include "Animation/InputScaleBias.h"
-#include "Animation/AnimNode_ApplyAdditive.h"
-#include "Animation/AnimNode_BlendListBase.h"
-#include "Animation/AnimNode_BlendListByBool.h"
-#include "Animation/AnimNode_BlendListByEnum.h"
-#include "Animation/AnimNode_BlendListByInt.h"
-#include "Animation/AnimNode_BlendSpacePlayer.h"
-#include "Animation/AnimNode_BlendSpaceEvaluator.h"
-#include "Animation/BoneControllers/AnimNode_SkeletalControlBase.h"
-#include "Animation/BoneControllers/AnimNode_CopyBone.h"
-#include "Animation/BoneControllers/AnimNode_Fabrik.h"
-#include "Animation/AnimNode_LayeredBoneBlend.h"
-#include "Animation/BoneControllers/AnimNode_LookAt.h"
-#include "Animation/BoneControllers/AnimNode_ModifyBone.h"
-#include "Animation/AnimNode_RefPose.h"
-#include "Animation/AnimNode_Root.h"
-#include "Animation/AnimNode_RotateRootBone.h"
-#include "Animation/BoneControllers/AnimNode_RotationMultiplier.h"
-#include "Animation/AnimNode_RotationOffsetBlendSpace.h"
-#include "Animation/AnimNode_SaveCachedPose.h"
-#include "Animation/AnimNode_SequenceEvaluator.h"
-#include "Animation/AnimNode_SequencePlayer.h"
-#include "Animation/AnimNode_Slot.h"
-#include "Animation/BoneControllers/AnimNode_SpringBone.h"
-#include "Animation/AnimNode_StateMachine.h"
-#include "Animation/AnimNode_TransitionPoseEvaluator.h"
-#include "Animation/AnimNode_TransitionResult.h"
-#include "Animation/BoneControllers/AnimNode_TwoBoneIK.h"
-#include "Animation/AnimNode_TwoWayBlend.h"
-#include "Animation/AnimNode_UseCachedPose.h"
-#include "Animation/BoneControllers/AnimNode_WheelHandler.h"
-#include "Animation/BoneControllers/AnimNode_HandIKRetargeting.h"
-#include "Animation/AnimNodeSpaceConversions.h"
 #include "Camera/CameraStackTypes.h"
 #include "Engine/StreamableManager.h"
 #include "Engine/TextureDefines.h"
@@ -415,9 +378,7 @@
 #include "VectorField/VectorField.h"
 #include "VectorField/VectorFieldAnimated.h"
 #include "VectorField/VectorFieldStatic.h"
-#include "Animation/VertexAnim/VertexAnimBase.h"
 #include "Animation/VertexAnim/MorphTarget.h"
-#include "Animation/VertexAnim/VertexAnimation.h"
 #include "Engine/World.h"
 
 #include "SystemSettings.h"					// Scalability options.
diff --git a/Engine/Source/Runtime/Engine/Public/EngineLogs.h b/Engine/Source/Runtime/Engine/Public/EngineLogs.h
--- a/Engine/Source/Runtime/Engine/Public/EngineLogs.h
+++ b/Engine/Source/Runtime/Engine/Public/EngineLogs.h
@@ -20,6 +20,7 @@ ENGINE_API DECLARE_LOG_CATEGORY_EXTERN(LogNet, Log, All);
 ENGINE_API DECLARE_LOG_CATEGORY_EXTERN(LogNetPlayerMovement, Warning, All);
 ENGINE_API DECLARE_LOG_CATEGORY_EXTERN(LogNetTraffic, Warning, All);
 ENGINE_API DECLARE_LOG_CATEGORY_EXTERN(LogNetDormancy, Warning, All);
+ENGINE_API DECLARE_LOG_CATEGORY_EXTERN(LogSkeletalControl, Warning, All);
 ENGINE_API DECLARE_LOG_CATEGORY_EXTERN(LogSubtitle, Log, All);
 ENGINE_API DECLARE_LOG_CATEGORY_EXTERN(LogTexture, Log, All);
 ENGINE_API DECLARE_LOG_CATEGORY_EXTERN(LogLandscape, Warning, All);
diff --git a/Engine/Source/Runtime/MovieSceneCoreTypes/Private/Instances/MovieSceneAnimationTrackInstance.cpp b/Engine/Source/Runtime/MovieSceneCoreTypes/Private/Instances/MovieSceneAnimationTrackInstance.cpp
--- a/Engine/Source/Runtime/MovieSceneCoreTypes/Private/Instances/MovieSceneAnimationTrackInstance.cpp
+++ b/Engine/Source/Runtime/MovieSceneCoreTypes/Private/Instances/MovieSceneAnimationTrackInstance.cpp
@@ -4,6 +4,7 @@
 #include "MovieSceneAnimationTrackInstance.h"
 #include "IMovieScenePlayer.h"
 #include "MovieSceneAnimationTrack.h"
+#include "Matinee/MatineeAnimInterface.h"
 
 FMovieSceneAnimationTrackInstance::FMovieSceneAnimationTrackInstance( UMovieSceneAnimationTrack& InAnimationTrack )
 {
