Fixing crash on shutdown

TTP#336864 - CRITICAL: CRASH: SHOOTERGAME: Shutdown Assertion in FAsyncAduioDecompressWorker::DoWork

Made ICompressedAudioInfo pointer part of the worker and create it in the constructor so that GEngine and Audio Device aren't checked in a different thread. Seems safe to keep using the raw USoundWave pointer as it's been that way for a long time and the decompress worker is checked in USoundWave::IsReadyForFinishDestroy, which is called during garbage collection on shutdown.

diff --git a/Engine/Source/Runtime/Engine/Private/AudioDecompress.cpp b/Engine/Source/Runtime/Engine/Private/AudioDecompress.cpp
--- a/Engine/Source/Runtime/Engine/Private/AudioDecompress.cpp
+++ b/Engine/Source/Runtime/Engine/Private/AudioDecompress.cpp
@@ -539,50 +539,60 @@ void FOpusAudioInfo::ExpandFile(uint8* DstBuffer, struct FSoundQualityInfo* Qual
 /**
  * Worker for decompression on a separate thread
  */
-void FAsyncAudioDecompressWorker::DoWork( void )
+FAsyncAudioDecompressWorker::FAsyncAudioDecompressWorker(USoundWave* InWave)
+	: Wave(InWave)
+	, AudioInfo(NULL)
 {
-	check(GEngine);
-	check(GEngine->GetAudioDevice());
-	ICompressedAudioInfo* AudioInfo = GEngine->GetAudioDevice()->CreateCompressedAudioInfo(Wave);
-	FSoundQualityInfo	QualityInfo = { 0 };
+	if (GEngine && GEngine->GetAudioDevice())
+	{
+		AudioInfo = GEngine->GetAudioDevice()->CreateCompressedAudioInfo(Wave);
+	}
+}
 
-	// Parse the audio header for the relevant information
-	if (AudioInfo->ReadCompressedInfo(Wave->ResourceData, Wave->ResourceSize, &QualityInfo))
+void FAsyncAudioDecompressWorker::DoWork( void )
+{
+	if (AudioInfo)
 	{
+		FSoundQualityInfo	QualityInfo = { 0 };
 
-#if PLATFORM_ANDROID
-		// Handle resampling
-		if (QualityInfo.SampleRate > 48000)
+		// Parse the audio header for the relevant information
+		if (AudioInfo->ReadCompressedInfo(Wave->ResourceData, Wave->ResourceSize, &QualityInfo))
 		{
-			UE_LOG( LogAudio, Warning, TEXT("Resampling file %s from %d"), *(Wave->GetName()), QualityInfo.SampleRate);
-			UE_LOG( LogAudio, Warning, TEXT("  Size %d"), QualityInfo.SampleDataSize);
-			uint32 SampleCount = QualityInfo.SampleDataSize / (QualityInfo.NumChannels * sizeof(uint16));
-			QualityInfo.SampleRate = QualityInfo.SampleRate / 2;
-			SampleCount /= 2;
-			QualityInfo.SampleDataSize = SampleCount * QualityInfo.NumChannels * sizeof(uint16);
-			AudioInfo->EnableHalfRate(true);
-		}
+
+#if PLATFORM_ANDROID
+			// Handle resampling
+			if (QualityInfo.SampleRate > 48000)
+			{
+				UE_LOG( LogAudio, Warning, TEXT("Resampling file %s from %d"), *(Wave->GetName()), QualityInfo.SampleRate);
+				UE_LOG( LogAudio, Warning, TEXT("  Size %d"), QualityInfo.SampleDataSize);
+				uint32 SampleCount = QualityInfo.SampleDataSize / (QualityInfo.NumChannels * sizeof(uint16));
+				QualityInfo.SampleRate = QualityInfo.SampleRate / 2;
+				SampleCount /= 2;
+				QualityInfo.SampleDataSize = SampleCount * QualityInfo.NumChannels * sizeof(uint16);
+				AudioInfo->EnableHalfRate(true);
+			}
 #endif
-		// Extract the data
-		Wave->SampleRate = QualityInfo.SampleRate;
-		Wave->NumChannels = QualityInfo.NumChannels;
-		Wave->Duration = QualityInfo.Duration;
+			// Extract the data
+			Wave->SampleRate = QualityInfo.SampleRate;
+			Wave->NumChannels = QualityInfo.NumChannels;
+			Wave->Duration = QualityInfo.Duration;
 
-		Wave->RawPCMDataSize = QualityInfo.SampleDataSize;
-		Wave->RawPCMData = ( uint8* )FMemory::Malloc( Wave->RawPCMDataSize );
+			Wave->RawPCMDataSize = QualityInfo.SampleDataSize;
+			Wave->RawPCMData = ( uint8* )FMemory::Malloc( Wave->RawPCMDataSize );
 
-		// Decompress all the sample data into preallocated memory
-		AudioInfo->ExpandFile(Wave->RawPCMData, &QualityInfo);
+			// Decompress all the sample data into preallocated memory
+			AudioInfo->ExpandFile(Wave->RawPCMData, &QualityInfo);
 
-		const SIZE_T ResSize = Wave->GetResourceSize(EResourceSizeMode::Exclusive);
-		INC_DWORD_STAT_BY( STAT_AudioMemorySize, ResSize );
-		INC_DWORD_STAT_BY( STAT_AudioMemory, ResSize );
-	}
+			const SIZE_T ResSize = Wave->GetResourceSize(EResourceSizeMode::Exclusive);
+			INC_DWORD_STAT_BY( STAT_AudioMemorySize, ResSize );
+			INC_DWORD_STAT_BY( STAT_AudioMemory, ResSize );
+		}
 
-	// Delete the compressed data
-	Wave->RemoveAudioResource();
+		// Delete the compressed data
+		Wave->RemoveAudioResource();
 
-	delete AudioInfo;
+		delete AudioInfo;
+	}
 }
 
 // end
diff --git a/Engine/Source/Runtime/Engine/Public/AudioDecompress.h b/Engine/Source/Runtime/Engine/Public/AudioDecompress.h
--- a/Engine/Source/Runtime/Engine/Public/AudioDecompress.h
+++ b/Engine/Source/Runtime/Engine/Public/AudioDecompress.h
@@ -182,16 +182,15 @@ class FAsyncAudioDecompressWorker : public FNonAbandonableTask
 protected:
 	class USoundWave*		Wave;
 
+	ICompressedAudioInfo*	AudioInfo;
+
 public:
 	/**
 	 * Async decompression of audio data
 	 *
 	 * @param	InWave		Wave data to decompress
 	 */
-	FAsyncAudioDecompressWorker(USoundWave* InWave)
-		: Wave(InWave)
-	{
-	}
+	ENGINE_API FAsyncAudioDecompressWorker(USoundWave* InWave);
 
 	/**
 	 * Performs the async audio decompression
