Fixed crash when moving the scrubber within Matinee while playing, if a sound track with no keys.
#ttp 336674 - CRASH: MATINEE: Moving the scrubber within Matinee while the Matinee Fight Scene is playing causes the Editor to crash
#branch UE4
#proj Runtime.Engine
#add Added extra check to early exit from UInterpTrackSound::PreviewUpdateTrack if there are no sound keys in the track.

diff --git a/Engine/Source/Runtime/Engine/Private/Interpolation.cpp b/Engine/Source/Runtime/Engine/Private/Interpolation.cpp
--- a/Engine/Source/Runtime/Engine/Private/Interpolation.cpp
+++ b/Engine/Source/Runtime/Engine/Private/Interpolation.cpp
@@ -7986,7 +7986,13 @@ void UInterpTrackSound::UpdateTrack(float NewPosition, UInterpTrackInst* TrInst,
 
 void UInterpTrackSound::PreviewUpdateTrack(float NewPosition, UInterpTrackInst* TrInst)
 {
-	UInterpGroupInst* GrInst = CastChecked<UInterpGroupInst>( TrInst->GetOuter() );
+	if (Sounds.Num() <= 0)
+	{
+		//UE_LOG(LogMatinee, Warning,TEXT("No sounds for sound track %s"),*GetName());
+		return;
+	}
+
+	UInterpGroupInst* GrInst = CastChecked<UInterpGroupInst>(TrInst->GetOuter());
 	AMatineeActor* MatineeActor = CastChecked<AMatineeActor>( GrInst->GetOuter() );
 	UInterpTrackInstSound* SoundInst = CastChecked<UInterpTrackInstSound>( TrInst );
 	UInterpGroup* Group = CastChecked<UInterpGroup>( GetOuter() );
