Help track non unit quaternions in bone transforms.

diff --git a/Engine/Source/Runtime/Core/Public/Math/Transform.h b/Engine/Source/Runtime/Core/Public/Math/Transform.h
--- a/Engine/Source/Runtime/Core/Public/Math/Transform.h
+++ b/Engine/Source/Runtime/Core/Public/Math/Transform.h
@@ -28,6 +28,15 @@ public:
 	 */
 	static CORE_API const FTransform Identity;
 
+#if ENABLE_UNIT_QUAT_DIAGNOSTIC
+	FORCEINLINE void DiagnosticCheckUnitQuaternion() const
+	{
+		checkf(Rotation.IsNormalized(), TEXT("FTransform Rotation is not Normalized: %s"), *Rotation.ToString());
+	}
+#else
+	FORCEINLINE void DiagnosticCheckUnitQuaternion() const {}
+#endif
+
 #if ENABLE_NAN_DIAGNOSTIC
 	FORCEINLINE void DiagnosticCheckNaN_Scale3D() const
 	{
diff --git a/Engine/Source/Runtime/Core/Public/Math/TransformVectorized.h b/Engine/Source/Runtime/Core/Public/Math/TransformVectorized.h
--- a/Engine/Source/Runtime/Core/Public/Math/TransformVectorized.h
+++ b/Engine/Source/Runtime/Core/Public/Math/TransformVectorized.h
@@ -32,6 +32,15 @@ public:
 	 */
 	static CORE_API const FTransform Identity;
 
+#if ENABLE_UNIT_QUAT_DIAGNOSTIC
+	FORCEINLINE void DiagnosticCheckUnitQuaternion() const
+	{
+		checkf(Rotation.IsNormalized(), TEXT("FTransform Rotation is not Normalized: %s"), *Rotation.ToString());
+	}
+#else
+	FORCEINLINE void DiagnosticCheckUnitQuaternion() const {}
+#endif
+
 #if ENABLE_NAN_DIAGNOSTIC
 	FORCEINLINE void DiagnosticCheckNaN_Scale3D() const
 	{
diff --git a/Engine/Source/Runtime/Core/Public/Math/UnrealMathUtility.h b/Engine/Source/Runtime/Core/Public/Math/UnrealMathUtility.h
--- a/Engine/Source/Runtime/Core/Public/Math/UnrealMathUtility.h
+++ b/Engine/Source/Runtime/Core/Public/Math/UnrealMathUtility.h
@@ -13,6 +13,12 @@
 	#define ENABLE_NAN_DIAGNOSTIC 0
 #endif
 
+// Assert on non unit quaternions.
+#ifndef  ENABLE_UNIT_QUAT_DIAGNOSTIC
+	#define ENABLE_UNIT_QUAT_DIAGNOSTIC 0
+#endif // ! ENABLE_UNIT_QUAT_DIAGNOSTIC
+
+
 /*-----------------------------------------------------------------------------
 	Definitions.
 -----------------------------------------------------------------------------*/
diff --git a/Engine/Source/Runtime/Engine/Classes/Animation/AnimNodeBase.h b/Engine/Source/Runtime/Engine/Classes/Animation/AnimNodeBase.h
--- a/Engine/Source/Runtime/Engine/Classes/Animation/AnimNodeBase.h
+++ b/Engine/Source/Runtime/Engine/Classes/Animation/AnimNodeBase.h
@@ -139,8 +139,35 @@ public:
 		FAnimationRuntime::InitializeTransform(AnimInstance->RequiredBones, Pose.Bones);
 	}
 
-	bool ContainsNaN() const;
-	bool IsNormalized() const;
+#if ENABLE_UNIT_QUAT_DIAGNOSTIC
+	FORCEINLINE void DiagnosticCheckUnitQuaternion() const
+	{
+		check(AnimInstance && AnimInstance->RequiredBones.IsValid());
+		const TArray<FBoneIndexType> & RequiredBoneIndices = AnimInstance->RequiredBones.GetBoneIndicesArray();
+		for (int32 Iter = 0; Iter < RequiredBoneIndices.Num(); ++Iter)
+		{
+			const int32 BoneIndex = RequiredBoneIndices[Iter];
+			Pose.Bones[BoneIndex].DiagnosticCheckUnitQuaternion();
+		}
+	}
+#else
+	FORCEINLINE void DiagnosticCheckUnitQuaternion() const {}
+#endif
+
+#if ENABLE_NAN_DIAGNOSTIC
+	FORCEINLINE void DiagnosticCheckNaN() const
+	{
+		check(AnimInstance && AnimInstance->RequiredBones.IsValid());
+		const TArray<FBoneIndexType> & RequiredBoneIndices = AnimInstance->RequiredBones.GetBoneIndicesArray();
+		for (int32 Iter = 0; Iter < RequiredBoneIndices.Num(); ++Iter)
+		{
+			const int32 BoneIndex = RequiredBoneIndices[Iter];
+			Pose.Bones[BoneIndex].DiagnosticCheckNaN_All();
+		}
+	}
+#else
+	FORCEINLINE void DiagnosticCheckNaN() const {}
+#endif
 };
 
 
@@ -171,8 +198,35 @@ public:
 		Pose.AllocateLocalPoses(AnimInstance->RequiredBones, AnimInstance->RequiredBones.GetRefPoseArray());
 	}
 
-	bool ContainsNaN() const;
-	bool IsNormalized() const;
+#if ENABLE_UNIT_QUAT_DIAGNOSTIC
+	FORCEINLINE void DiagnosticCheckUnitQuaternion() const
+	{
+		check(AnimInstance && AnimInstance->RequiredBones.IsValid());
+		const TArray<FBoneIndexType> & RequiredBoneIndices = AnimInstance->RequiredBones.GetBoneIndicesArray();
+		for (int32 Iter = 0; Iter < RequiredBoneIndices.Num(); ++Iter)
+		{
+			const int32 BoneIndex = RequiredBoneIndices[Iter];
+			Pose.Bones[BoneIndex].DiagnosticCheckUnitQuaternion();
+		}
+	}
+#else
+	FORCEINLINE void DiagnosticCheckUnitQuaternion() const {}
+#endif
+
+#if ENABLE_NAN_DIAGNOSTIC
+	FORCEINLINE void DiagnosticCheckNaN() const
+	{
+		check(AnimInstance && AnimInstance->RequiredBones.IsValid());
+		const TArray<FBoneIndexType> & RequiredBoneIndices = AnimInstance->RequiredBones.GetBoneIndicesArray();
+		for (int32 Iter = 0; Iter < RequiredBoneIndices.Num(); ++Iter)
+		{
+			const int32 BoneIndex = RequiredBoneIndices[Iter];
+			Pose.Bones[BoneIndex].DiagnosticCheckNaN_All();
+		}
+	}
+#else
+	FORCEINLINE void DiagnosticCheckNaN() const {}
+#endif
 };
 
 struct FNodeDebugData
diff --git a/Engine/Source/Runtime/Engine/Private/Animation/AnimNodeBase.cpp b/Engine/Source/Runtime/Engine/Private/Animation/AnimNodeBase.cpp
--- a/Engine/Source/Runtime/Engine/Private/Animation/AnimNodeBase.cpp
+++ b/Engine/Source/Runtime/Engine/Private/Animation/AnimNodeBase.cpp
@@ -119,8 +119,8 @@ void FPoseLink::Evaluate(FPoseContext& Output)
 	}
 
 	// Detect non valid output
-	checkSlow( !Output.ContainsNaN() );
-	checkSlow( Output.IsNormalized() );
+	Output.DiagnosticCheckNaN();
+	Output.DiagnosticCheckUnitQuaternion();
 }
 
 /////////////////////////////////////////////////////
@@ -143,79 +143,17 @@ void FComponentSpacePoseLink::EvaluateComponentSpace(FComponentSpacePoseContext&
 	}
 
 	// Detect non valid output
-	checkSlow( !Output.ContainsNaN() );
-	checkSlow( Output.IsNormalized() );
+	Output.DiagnosticCheckNaN();
+	Output.DiagnosticCheckUnitQuaternion();
 }
 
 /////////////////////////////////////////////////////
 // FPoseContext
 
-bool FPoseContext::ContainsNaN() const
-{
-	checkSlow( AnimInstance && AnimInstance->RequiredBones.IsValid() );
-	const TArray<FBoneIndexType> & RequiredBoneIndices = AnimInstance->RequiredBones.GetBoneIndicesArray();
-	for (int32 Iter = 0; Iter < RequiredBoneIndices.Num(); ++Iter)
-	{
-		const int32 BoneIndex = RequiredBoneIndices[Iter];
-		if (Pose.Bones[BoneIndex].ContainsNaN())
-		{
-			return true;
-		}
-	}
-
-	return false;
-}
-
-bool FPoseContext::IsNormalized() const
-{
-	checkSlow( AnimInstance && AnimInstance->RequiredBones.IsValid() );
-	const TArray<FBoneIndexType> & RequiredBoneIndices = AnimInstance->RequiredBones.GetBoneIndicesArray();
-	for (int32 Iter = 0; Iter < RequiredBoneIndices.Num(); ++Iter)
-	{
-		const int32 BoneIndex = RequiredBoneIndices[Iter];
-		if( !Pose.Bones[BoneIndex].IsRotationNormalized() )
-		{
-			return false;
-		}
-	}
-
-	return true;
-}
 
 /////////////////////////////////////////////////////
 // FComponentSpacePoseContext
 
-bool FComponentSpacePoseContext::ContainsNaN() const
-{
-	checkSlow( AnimInstance && AnimInstance->RequiredBones.IsValid() );
-	const TArray<FBoneIndexType> & RequiredBoneIndices = AnimInstance->RequiredBones.GetBoneIndicesArray();
-	for (int32 Iter = 0; Iter < RequiredBoneIndices.Num(); ++Iter)
-	{
-		const int32 BoneIndex = RequiredBoneIndices[Iter];
-		if (Pose.Bones[BoneIndex].ContainsNaN())
-		{
-			return true;
-		}
-	}
-
-	return false;
-}
-
-bool FComponentSpacePoseContext::IsNormalized() const
-{
-	checkSlow( AnimInstance && AnimInstance->RequiredBones.IsValid() );
-	const TArray<FBoneIndexType> & RequiredBoneIndices = AnimInstance->RequiredBones.GetBoneIndicesArray();
-	for (int32 Iter = 0; Iter < RequiredBoneIndices.Num(); ++Iter)
-	{
-		const int32 BoneIndex = RequiredBoneIndices[Iter];
-		if( !Pose.Bones[BoneIndex].IsRotationNormalized() )
-		{
-			return false;
-		}
-	}
-
-	return true;
-}
 
 /////////////////////////////////////////////////////
 // FNodeDebugData
diff --git a/Engine/Source/Runtime/Engine/Private/Animation/AnimNode_Slot.cpp b/Engine/Source/Runtime/Engine/Private/Animation/AnimNode_Slot.cpp
--- a/Engine/Source/Runtime/Engine/Private/Animation/AnimNode_Slot.cpp
+++ b/Engine/Source/Runtime/Engine/Private/Animation/AnimNode_Slot.cpp
@@ -51,9 +51,6 @@ void FAnimNode_Slot::Evaluate(FPoseContext & Output)
 		}
 
 		Output.AnimInstance->SlotEvaluatePose(SlotName, SourceContext.Pose, Output.Pose, SlotNodeWeight);
-
-		checkSlow(!Output.ContainsNaN());
-		checkSlow(Output.IsNormalized());
 	}
 }
 
diff --git a/Engine/Source/Runtime/Engine/Private/SkeletalMeshComponent.cpp b/Engine/Source/Runtime/Engine/Private/SkeletalMeshComponent.cpp
--- a/Engine/Source/Runtime/Engine/Private/SkeletalMeshComponent.cpp
+++ b/Engine/Source/Runtime/Engine/Private/SkeletalMeshComponent.cpp
@@ -587,8 +587,8 @@ void USkeletalMeshComponent::FillSpaceBases()
 #endif
 		FTransform::Multiply(SpaceBasesData + BoneIndex, LocalTransformsData + BoneIndex, SpaceBasesData + ParentIndex);
 
-		checkSlow( SpaceBases[BoneIndex].IsRotationNormalized() );
-		checkSlow( !SpaceBases[BoneIndex].ContainsNaN() );
+		SpaceBases[BoneIndex].DiagnosticCheckNaN_All();
+		SpaceBases[BoneIndex].DiagnosticCheckUnitQuaternion();
 	}
 }
 
@@ -1216,7 +1216,7 @@ bool USkeletalMeshComponent::IsAnySimulatingPhysics() const
  */
 void USkeletalMeshComponent::DebugDrawBones(UCanvas* Canvas, bool bSimpleBones) const
 {
-	if (GetWorld()->IsGameWorld() && SkeletalMesh && Canvas)
+	if (GetWorld()->IsGameWorld() && SkeletalMesh && Canvas && (SpaceBases.Num() > 0) )
 	{
 		// draw spacebases, we could cache parent bones, but this is mostly debug feature, I'm not caching it right now
 		for ( int32 Index=0; Index<RequiredBones.Num(); ++Index )
