Physics: Changed CalculateMass to use GetSimplePhysicalMaterial with the richer set of override points

diff --git a/Engine/Source/Runtime/Engine/Private/PhysicsEngine/BodySetup.cpp b/Engine/Source/Runtime/Engine/Private/PhysicsEngine/BodySetup.cpp
--- a/Engine/Source/Runtime/Engine/Private/PhysicsEngine/BodySetup.cpp
+++ b/Engine/Source/Runtime/Engine/Private/PhysicsEngine/BodySetup.cpp
@@ -984,7 +984,6 @@ class UPhysicalMaterial* UBodySetup::GetPhysMaterial() const
 float UBodySetup::CalculateMass(const UPrimitiveComponent* Component)
 {
 	FVector ComponentScale(1.0f, 1.0f, 1.0f);
-	UPhysicalMaterial* PhysMat = GetPhysMaterial();
 	const FBodyInstance* BodyInstance = NULL;
 	float MassScale = DefaultInstance.MassScale;
 
@@ -1011,7 +1010,7 @@ float UBodySetup::CalculateMass(const UPrimitiveComponent* Component)
 		BodyInstance = &DefaultInstance;
 	}
 
-	PhysMat = BodyInstance->PhysMaterialOverride != NULL ? BodyInstance->PhysMaterialOverride : PhysMat;
+	UPhysicalMaterial* PhysMat = BodyInstance->GetSimplePhysicalMaterial();
 	MassScale = BodyInstance->MassScale;
 
 	// physical material - nothing can weigh less than hydrogen (0.09 kg/m^3)
