Engine: Cleaning up some whitespace issues

diff --git a/Engine/Source/Runtime/Engine/Classes/Components/PrimitiveComponent.h b/Engine/Source/Runtime/Engine/Classes/Components/PrimitiveComponent.h
--- a/Engine/Source/Runtime/Engine/Classes/Components/PrimitiveComponent.h
+++ b/Engine/Source/Runtime/Engine/Classes/Components/PrimitiveComponent.h
@@ -877,12 +877,12 @@ public:
 	virtual void SetCollisionEnabled(ECollisionEnabled::Type NewType);
 
 	/**  
-	* Set Collision Profile Name
-	* This function is called by constructors when they set ProfileName
-	* This will change current CollisionProfileName to be this, and overwrite Collision Setting
-	* 
-	* @param InCollisionProfileName : New Profile Name
-	*/
+	 * Set Collision Profile Name
+	 * This function is called by constructors when they set ProfileName
+	 * This will change current CollisionProfileName to be this, and overwrite Collision Setting
+	 * 
+	 * @param InCollisionProfileName : New Profile Name
+	 */
 	UFUNCTION(BlueprintCallable, Category="Collision")	
 	virtual void SetCollisionProfileName(FName InCollisionProfileName);
 
@@ -918,12 +918,9 @@ public:
 	// Begin UActorComponent Interface
 	virtual void InvalidateLightingCacheDetailed(bool bInvalidateBuildEnqueuedLighting, bool bTranslationOnly) OVERRIDE;
 	virtual bool IsEditorOnly() const OVERRIDE;
-	// End UActorComponent Interface
-
-	/** @return true if this should create physics state */
 	virtual bool ShouldCreatePhysicsState() const OVERRIDE;
-
 	virtual bool HasValidPhysicsState() const OVERRIDE;
+	// End UActorComponent Interface
 
 	/** @return true if the owner is selected and this component is selectable */
 	virtual bool ShouldRenderSelected() const;
