Physics: Made several FBodyInstance methods const

diff --git a/Engine/Source/Runtime/Engine/Classes/PhysicsEngine/BodyInstance.h b/Engine/Source/Runtime/Engine/Classes/PhysicsEngine/BodyInstance.h
--- a/Engine/Source/Runtime/Engine/Classes/PhysicsEngine/BodyInstance.h
+++ b/Engine/Source/Runtime/Engine/Classes/PhysicsEngine/BodyInstance.h
@@ -33,11 +33,11 @@ struct ENGINE_API FCollisionResponse
 
 	/** Returns the response set on the specified channel */
 	ECollisionResponse GetResponse(ECollisionChannel Channel) const;
-	const FCollisionResponseContainer & GetResponseContainer() const { return ResponseToChannels; }
+	const FCollisionResponseContainer& GetResponseContainer() const { return ResponseToChannels; }
 
 	/** Set all channels from ChannelResponse Array **/
-	void SetCollisionResponseContainer(const FCollisionResponseContainer & InResponseToChannels);
-	void SetResponsesArray(const TArray<FResponseChannel> & InChannelResponses);
+	void SetCollisionResponseContainer(const FCollisionResponseContainer& InResponseToChannels);
+	void SetResponsesArray(const TArray<FResponseChannel>& InChannelResponses);
 	void UpdateResponseContainerFromArray();
 
 private:
@@ -104,7 +104,7 @@ struct ENGINE_API FBodyInstance
 private:
 	/** Collision Profile Name **/
 	UPROPERTY(EditAnywhere, Category=Custom)
-	FName	CollisionProfileName;
+	FName CollisionProfileName;
 
 	/** Type of Collision Enabled 
 	 * 
@@ -121,7 +121,7 @@ private:
 
 	/** Custom Channels for Responses*/
 	UPROPERTY(EditAnywhere, Category=Custom) 
-	struct FCollisionResponse		CollisionResponses;
+	struct FCollisionResponse CollisionResponses;
 
 public:
 
@@ -169,12 +169,11 @@ protected:
 	UPROPERTY(EditAnywhere, AdvancedDisplay, BlueprintReadOnly, Category=Physics, meta=(editcondition="bOverrideWalkableSlopeOnInstance"))
 	struct FWalkableSlopeOverride WalkableSlopeOverride;
 
-public:
-
 	/**	Allows you to override the PhysicalMaterial to use for simple collision on this body. */
 	UPROPERTY(EditAnywhere, AdvancedDisplay, BlueprintReadOnly, Category=Collision)
 	class UPhysicalMaterial* PhysMaterialOverride;
 
+public:
 	/** User specified offset for the center of mass of this object, from the calculated location */
 	UPROPERTY(EditAnywhere, AdvancedDisplay, BlueprintReadWrite, Category=Physics, meta=(DisplayName="Center Of Mass Offset"))
 	FVector COMNudge;
@@ -214,329 +213,330 @@ public:
 
 public:
 #if WITH_PHYSX
-		/** Internal use. Physics-engine representation of this body in the synchronous scene. */
-		physx::PxRigidActor*			RigidActorSync;
+	/** Internal use. Physics-engine representation of this body in the synchronous scene. */
+	physx::PxRigidActor* RigidActorSync;
 
-		/** Internal use. Physics-engine representation of this body in the asynchronous scene. */
-		physx::PxRigidActor*			RigidActorAsync;
+	/** Internal use. Physics-engine representation of this body in the asynchronous scene. */
+	physx::PxRigidActor* RigidActorAsync;
 
-		/** Internal use. Physics-engine representation of a PxAggregate for this body, in case it has alot of shapes. */
-		physx::PxAggregate*				BodyAggregate;
+	/** Internal use. Physics-engine representation of a PxAggregate for this body, in case it has alot of shapes. */
+	physx::PxAggregate* BodyAggregate;
 
-		TSharedPtr<TArray<ANSICHAR>>	CharDebugName;
+	TSharedPtr<TArray<ANSICHAR>> CharDebugName;
 #endif	//WITH_PHYSX
 
-		/** PrimitiveComponent containing this body.   */
-		TWeakObjectPtr<class UPrimitiveComponent>		OwnerComponent;
+	/** PrimitiveComponent containing this body.   */
+	TWeakObjectPtr<class UPrimitiveComponent> OwnerComponent;
 
-		/** BodySetup pointer that this instance is initialized from */
-		TWeakObjectPtr<class UBodySetup>				BodySetup;
+	/** BodySetup pointer that this instance is initialized from */
+	TWeakObjectPtr<class UBodySetup> BodySetup;
 
-		/** Constructor **/
-		FBodyInstance();
+	/** Constructor **/
+	FBodyInstance();
 
-		// BodyInstance interface
+	// BodyInstance interface
 
-		/** Update CollisionEnabled in old assets (before VER_UE4_CHANGE_BENABLECOLLISION_TO_COLLISIONENABLED) using bEnableCollision_DEPRECATED */ 
-		void UpdateFromDeprecatedEnableCollision();
+	/** Update CollisionEnabled in old assets (before VER_UE4_CHANGE_BENABLECOLLISION_TO_COLLISIONENABLED) using bEnableCollision_DEPRECATED */ 
+	void UpdateFromDeprecatedEnableCollision();
 
-		/**  
-		 * Update profile data if required
-		 * 
-		 * @param : bVerifyProfile - if true, it makes sure it has correct set up with current profile, if false, it overwrites from profile data
-		 *								(for backward compatibility)
-		 * 
-		 **/
-		void LoadProfileData(bool bVerifyProfile);
+	/**  
+	 * Update profile data if required
+	 * 
+	 * @param : bVerifyProfile - if true, it makes sure it has correct set up with current profile, if false, it overwrites from profile data
+	 *								(for backward compatibility)
+	 * 
+	 **/
+	void LoadProfileData(bool bVerifyProfile);
 
 #if WITH_PHYSX
-		void InitBody(class UBodySetup* Setup, const FTransform& Transform, class UPrimitiveComponent* PrimComp,class FPhysScene* InRBScene, physx::PxAggregate* InAggregate = NULL);
+	void InitBody(class UBodySetup* Setup, const FTransform& Transform, class UPrimitiveComponent* PrimComp,class FPhysScene* InRBScene, physx::PxAggregate* InAggregate = NULL);
 #endif	//WITH_PHYSX
 
-		void TermBody();
+	void TermBody();
 
 
 #if WITH_BODY_WELDING
-		/** 
-		 * Takes two body instances and welds them together to create a single simulated rigid body
-		 */
-		void Weld(FBodyInstance * Body, const FTransform & RelativeTM);
+	/** 
+	 * Takes two body instances and welds them together to create a single simulated rigid body
+	 */
+	void Weld(FBodyInstance* Body, const FTransform& RelativeTM);
 #endif
 
-		/**
-		 * Update Body Scale
-		 * @param inScale3D - new Scale3D. If that's different from previous Scale3D, it will update Body scale
-		 * @return true if succeed
-		 */
-		bool UpdateBodyScale(const FVector & inScale3D);
+	/**
+	 * Update Body Scale
+	 * @param inScale3D - new Scale3D. If that's different from previous Scale3D, it will update Body scale
+	 * @return true if succeed
+	 */
+	bool UpdateBodyScale(const FVector& InScale3D);
 
-		FVector GetCOMPosition();
-		void DrawCOMPosition( class FPrimitiveDrawInterface* PDI, float COMRenderSize, const FColor& COMRenderColor );
+	FVector GetCOMPosition() const;
+	void DrawCOMPosition( class FPrimitiveDrawInterface* PDI, float COMRenderSize, const FColor& COMRenderColor );
 
-		/** Utility for copying properties from one BodyInstance to another. */
-		void CopyBodyInstancePropertiesFrom(const FBodyInstance* FromInst);
+	/** Utility for copying properties from one BodyInstance to another. */
+	void CopyBodyInstancePropertiesFrom(const FBodyInstance* FromInst);
 
-		/** Find the correct PhysicalMaterial for simple geometry on this body */
-		UPhysicalMaterial* GetSimplePhysicalMaterial();
+	/** Find the correct PhysicalMaterial for simple geometry on this body */
+	UPhysicalMaterial* GetSimplePhysicalMaterial() const;
 
-		/** Get the complex PhysicalMaterial array for this body */
-		TArray<UPhysicalMaterial*> GetComplexPhysicalMaterials();
+	/** Get the complex PhysicalMaterial array for this body */
+	TArray<UPhysicalMaterial*> GetComplexPhysicalMaterials() const;
 
-		/** Returns the slope override struct for this instance. If we don't have our own custom setting, it will return the setting from the body setup. */
-		const struct FWalkableSlopeOverride& GetWalkableSlopeOverride() const;
+	/** Returns the slope override struct for this instance. If we don't have our own custom setting, it will return the setting from the body setup. */
+	const struct FWalkableSlopeOverride& GetWalkableSlopeOverride() const;
 
-		/** Returns whether this body wants (and can) use the async scene. */
-		bool UseAsyncScene() const;
+	/** Returns whether this body wants (and can) use the async scene. */
+	bool UseAsyncScene() const;
 
 #if WITH_PHYSX
-		/**
-		 * Return the PxRigidActor from the given scene (see EPhysicsSceneType), if SceneType is in the range [0, PST_MAX).
-		 * If SceneType < 0, the PST_Sync actor is returned if it is not NULL, otherwise the PST_Async actor is returned.
-		 * Invalid scene types will cause NULL to be returned.
-		 */
-		physx::PxRigidActor* GetPxRigidActor(int32 SceneType = -1) const;
-		/** Return the PxRigidDynamic if it exists in one of the scenes (NULL otherwise).  Currently a PxRigidDynamic can exist in only one of the two scenes. */
-		physx::PxRigidDynamic* GetPxRigidDynamic() const;
-
-		/** 
-		 *	Utility to get all the shapes from a FBodyInstance 
-		 *	Shapes belonging to sync actor are first, then async. Number of shapes belonging to sync actor is returned.
-		 */
-		TArray<physx::PxShape*> GetAllShapes(int32& OutNumSyncShapes) const;
+	/**
+	 * Return the PxRigidActor from the given scene (see EPhysicsSceneType), if SceneType is in the range [0, PST_MAX).
+	 * If SceneType < 0, the PST_Sync actor is returned if it is not NULL, otherwise the PST_Async actor is returned.
+	 * Invalid scene types will cause NULL to be returned.
+	 */
+	physx::PxRigidActor* GetPxRigidActor(int32 SceneType = -1) const;
+	/** Return the PxRigidDynamic if it exists in one of the scenes (NULL otherwise).  Currently a PxRigidDynamic can exist in only one of the two scenes. */
+	physx::PxRigidDynamic* GetPxRigidDynamic() const;
+
+	/** 
+	 *	Utility to get all the shapes from a FBodyInstance 
+	 *	Shapes belonging to sync actor are first, then async. Number of shapes belonging to sync actor is returned.
+	 */
+	TArray<physx::PxShape*> GetAllShapes(int32& OutNumSyncShapes) const;
 #endif	//WITH_PHYSX
 
-		/** Returns the body's mass */
-		float GetBodyMass() const;
-		/** Return bounds of physics representation */
-		FBox GetBodyBounds() const;
-
-
-		/** Set this body to be fixed (kinematic) or not. */
-		void SetInstanceSimulatePhysics(bool bSimulate, bool bMaintainPhysicsBlending=false, bool bIgnoreOwner=false);
-		/** Makes sure the current kinematic state matches the simulate flag */
-		void UpdateInstanceSimulatePhysics(bool bIgnoreOwner = false);
-		/** Returns true if this body is simulating, false if it is fixed (kinematic) */
-		bool IsInstanceSimulatingPhysics(bool bIgnoreOwner = false);
-		/** Should Simulate Physics **/
-		bool ShouldInstanceSimulatingPhysics(bool bIgnoreOwner=false);
-		/** Returns whether this body is awake */
-		bool IsInstanceAwake();
-		/** Wake this body */
-		void WakeInstance();
-		/** Force this body to sleep */
-		void PutInstanceToSleep();
-		/** Add a force to this body */
-		void AddForce(const FVector& Force, bool bAllowSubstepping = true);
-		/** Add a force at a particular world position to this body */
-		void AddForceAtPosition(const FVector& Force, const FVector& Position, bool bAllowSubstepping = true);
-		/** Add a torque to this body */
-		void AddTorque(const FVector& Torque, bool bAllowSubstepping = true);
-		/** Add an impulse to this body */
-		void AddImpulse(const FVector& Impulse, bool bVelChange);
-		/** Add an impulse to this body and a particular world position */
-		void AddImpulseAtPosition(const FVector& Impulse, const FVector& Position);
-		/** Set the linear velocity of this body */
-		void SetLinearVelocity(const FVector& NewVel, bool bAddToCurrent);
-		/** Set the angular velocity of this body */
-		void SetAngularVelocity(const FVector& NewAngVel, bool bAddToCurrent);
-		/** Set the maximum angular velocity of this body */
-		void SetMaxAngularVelocity(float NewMaxAngVel, bool bAddToCurrent);
-		/** Set whether we should get a notification about physics collisions */
-		void SetInstanceNotifyRBCollision(bool bNewNotifyCollision);
-		/** Enables/disables whether this body is affected by gravity. */
-		void SetEnableGravity(bool bGravityEnabled);
-
-		/** See if this body is valid. */
-		bool IsValidBodyInstance() const;
-
-		/** Get current transform in world space from physics body. */
-		FTransform GetUnrealWorldTransform() const;
-
-		/** 
-		 *	Move the physics body to a new pose. 
-		 *	@param	bTeleport	If true, no velocity is inferred on the kinematic body from this movement, but it moves right away.
-		 */
-		void SetBodyTransform(const FTransform& NewTransform, bool bTeleport);
-
-		/** Get current velocity in world space from physics body. */
-		FVector GetUnrealWorldVelocity();
-
-		/** Get current angular velocity in world space from physics body. */
-		FVector GetUnrealWorldAngularVelocity();
-
-		/** Get current velocity of a point on this physics body, in world space. Point is specified in world space. */
-		FVector GetUnrealWorldVelocityAtPoint(const FVector& Point);
-
-		/** Set physical material override for this body */
-		void SetPhysMaterialOverride( class UPhysicalMaterial* NewPhysMaterial );
-
-		/** Set a new contact report force threhold.  Threshold < 0 disables this feature. */
-		void SetContactReportForceThreshold( float Threshold );
-
-		/** Set the collision response of this body to a particular channel */
-		void SetResponseToChannel(ECollisionChannel Channel, ECollisionResponse NewResponse);
-
-		/** Get the collision response of this body to a particular channel */
-		ECollisionResponse GetResponseToChannel(ECollisionChannel Channel) const;
-
-		/** Set the response of this body to all channels */
-		void SetResponseToAllChannels(ECollisionResponse NewResponse);
-
-		/** Set the response of this body to the supplied settings */
-		void SetResponseToChannels(const FCollisionResponseContainer& NewReponses);
-
-		/** Get Collision ResponseToChannels container for this component **/
-		const FCollisionResponseContainer & GetResponseToChannels() const;
-
-		/** Set the movement channel of this body to the one supplied */
-		void SetObjectType(ECollisionChannel Channel);
-
-		/** Get the movement channel of this body **/
-		ECollisionChannel GetObjectType() const;
-
-		/** Controls what kind of collision is enabled for this body and allows optional disable physics rebuild */
-		void SetCollisionEnabled(ECollisionEnabled::Type NewType, bool bUpdatePhysicsFilterData = true);
-
-		/** Get the current type of collision enabled */
-		ECollisionEnabled::Type GetCollisionEnabled() const;
-
-		/**  
-		 * Set Collision Profile Name
-		 * This function is called by constructors when they set ProfileName
-		 * This will change current CollisionProfileName to be this, and overwrite Collision Setting
-		 * 
-		 * @param InCollisionProfileName : New Profile Name
-		 */
-		void SetCollisionProfileName(FName InCollisionProfileName);
-
-		/** Get the current collision profile assigned to this body */
-		FName GetCollisionProfileName() const;
-
-		/** return true if it uses Collision Profile System. False otherwise*/
-		bool DoesUseCollisionProfile() const;
-
-		/** Update instance's mass properties (mass, inertia and center-of-mass offset) based on MassScale, InstanceMassScale and COMNudge. */
-		void UpdateMassProperties();
-
-		/** Update instance's linear and angular damping */
-		void UpdateDampingProperties();
-
-		/** Update the instance's material properties (friction, restitution) */
-		void UpdatePhysicalMaterials();
-
-		/** Update the instances collision filtering data */
-		void UpdatePhysicsFilterData();
-
-		friend FArchive& operator<<(FArchive& Ar,FBodyInstance& BodyInst);
-
-		/** Get the name for this body, for use in debugging */
-		FString GetBodyDebugName();
-
-		/** 
-		 *  Trace a ray against just this bodyinstance
-		 *  @param  OutHit					Information about hit against this component, if true is returned
-		 *  @param  Start					Start location of the ray
-		 *  @param  End						End location of the ray
-		 *	@param	bTraceComplex			Should we trace against complex or simple collision of this body
-		 *  @param bReturnPhysicalMaterial	Fill in the PhysMaterial field of OutHit
-		 *  @return true if a hit is found
-		 */
-		bool LineTrace(struct FHitResult& OutHit, const FVector& Start, const FVector& End, bool bTraceComplex, bool bReturnPhysicalMaterial = false);
+	/** Returns the body's mass */
+	float GetBodyMass() const;
+	/** Return bounds of physics representation */
+	FBox GetBodyBounds() const;
+
+
+	/** Set this body to be fixed (kinematic) or not. */
+	void SetInstanceSimulatePhysics(bool bSimulate, bool bMaintainPhysicsBlending=false, bool bIgnoreOwner=false);
+	/** Makes sure the current kinematic state matches the simulate flag */
+	void UpdateInstanceSimulatePhysics(bool bIgnoreOwner = false);
+	/** Returns true if this body is simulating, false if it is fixed (kinematic) */
+	bool IsInstanceSimulatingPhysics(bool bIgnoreOwner = false);
+	/** Should Simulate Physics **/
+	bool ShouldInstanceSimulatingPhysics(bool bIgnoreOwner=false);
+	/** Returns whether this body is awake */
+	bool IsInstanceAwake();
+	/** Wake this body */
+	void WakeInstance();
+	/** Force this body to sleep */
+	void PutInstanceToSleep();
+	/** Add a force to this body */
+	void AddForce(const FVector& Force, bool bAllowSubstepping = true);
+	/** Add a force at a particular world position to this body */
+	void AddForceAtPosition(const FVector& Force, const FVector& Position, bool bAllowSubstepping = true);
+	/** Add a torque to this body */
+	void AddTorque(const FVector& Torque, bool bAllowSubstepping = true);
+	/** Add an impulse to this body */
+	void AddImpulse(const FVector& Impulse, bool bVelChange);
+	/** Add an impulse to this body and a particular world position */
+	void AddImpulseAtPosition(const FVector& Impulse, const FVector& Position);
+	/** Set the linear velocity of this body */
+	void SetLinearVelocity(const FVector& NewVel, bool bAddToCurrent);
+	/** Set the angular velocity of this body */
+	void SetAngularVelocity(const FVector& NewAngVel, bool bAddToCurrent);
+	/** Set the maximum angular velocity of this body */
+	void SetMaxAngularVelocity(float NewMaxAngVel, bool bAddToCurrent);
+	/** Set whether we should get a notification about physics collisions */
+	void SetInstanceNotifyRBCollision(bool bNewNotifyCollision);
+	/** Enables/disables whether this body is affected by gravity. */
+	void SetEnableGravity(bool bGravityEnabled);
+
+	/** See if this body is valid. */
+	bool IsValidBodyInstance() const;
+
+	/** Get current transform in world space from physics body. */
+	FTransform GetUnrealWorldTransform() const;
+
+	/** 
+		*	Move the physics body to a new pose. 
+		*	@param	bTeleport	If true, no velocity is inferred on the kinematic body from this movement, but it moves right away.
+		*/
+	void SetBodyTransform(const FTransform& NewTransform, bool bTeleport);
+
+	/** Get current velocity in world space from physics body. */
+	FVector GetUnrealWorldVelocity() const;
+
+	/** Get current angular velocity in world space from physics body. */
+	FVector GetUnrealWorldAngularVelocity() const;
+
+	/** Get current velocity of a point on this physics body, in world space. Point is specified in world space. */
+	FVector GetUnrealWorldVelocityAtPoint(const FVector& Point) const;
+
+	/** Set physical material override for this body */
+	void SetPhysMaterialOverride(class UPhysicalMaterial* NewPhysMaterial);
+
+	/** Set a new contact report force threhold.  Threshold < 0 disables this feature. */
+	void SetContactReportForceThreshold(float Threshold);
+
+	/** Set the collision response of this body to a particular channel */
+	void SetResponseToChannel(ECollisionChannel Channel, ECollisionResponse NewResponse);
+
+	/** Get the collision response of this body to a particular channel */
+	ECollisionResponse GetResponseToChannel(ECollisionChannel Channel) const;
+
+	/** Set the response of this body to all channels */
+	void SetResponseToAllChannels(ECollisionResponse NewResponse);
+
+	/** Set the response of this body to the supplied settings */
+	void SetResponseToChannels(const FCollisionResponseContainer& NewReponses);
+
+	/** Get Collision ResponseToChannels container for this component **/
+	const FCollisionResponseContainer& GetResponseToChannels() const;
+
+	/** Set the movement channel of this body to the one supplied */
+	void SetObjectType(ECollisionChannel Channel);
+
+	/** Get the movement channel of this body **/
+	ECollisionChannel GetObjectType() const;
+
+	/** Controls what kind of collision is enabled for this body and allows optional disable physics rebuild */
+	void SetCollisionEnabled(ECollisionEnabled::Type NewType, bool bUpdatePhysicsFilterData = true);
+
+	/** Get the current type of collision enabled */
+	ECollisionEnabled::Type GetCollisionEnabled() const;
+
+	/**  
+	 * Set Collision Profile Name
+	 * This function is called by constructors when they set ProfileName
+	 * This will change current CollisionProfileName to be this, and overwrite Collision Setting
+	 * 
+	 * @param InCollisionProfileName : New Profile Name
+	 */
+	void SetCollisionProfileName(FName InCollisionProfileName);
+
+	/** Get the current collision profile assigned to this body */
+	FName GetCollisionProfileName() const;
+
+	/** return true if it uses Collision Profile System. False otherwise*/
+	bool DoesUseCollisionProfile() const;
+
+	/** Update instance's mass properties (mass, inertia and center-of-mass offset) based on MassScale, InstanceMassScale and COMNudge. */
+	void UpdateMassProperties();
+
+	/** Update instance's linear and angular damping */
+	void UpdateDampingProperties();
+
+	/** Update the instance's material properties (friction, restitution) */
+	void UpdatePhysicalMaterials();
+
+	/** Update the instances collision filtering data */
+	void UpdatePhysicsFilterData();
+
+	friend FArchive& operator<<(FArchive& Ar,FBodyInstance& BodyInst);
+
+	/** Get the name for this body, for use in debugging */
+	FString GetBodyDebugName() const;
+
+	/** 
+	 *  Trace a ray against just this bodyinstance
+	 *  @param  OutHit					Information about hit against this component, if true is returned
+	 *  @param  Start					Start location of the ray
+	 *  @param  End						End location of the ray
+	 *	@param	bTraceComplex			Should we trace against complex or simple collision of this body
+	 *  @param bReturnPhysicalMaterial	Fill in the PhysMaterial field of OutHit
+	 *  @return true if a hit is found
+	 */
+	bool LineTrace(struct FHitResult& OutHit, const FVector& Start, const FVector& End, bool bTraceComplex, bool bReturnPhysicalMaterial = false);
+
+	/** 
+	 *  Trace a box against just this bodyinstance
+	 *  @param  OutHit          Information about hit against this component, if true is returned
+	 *  @param  Start           Start location of the box
+	 *  @param  End             End location of the box
+	 *  @param  CollisionShape	Collision Shape
+	 *	@param	bTraceComplex	Should we trace against complex or simple collision of this body
+	 *  @return true if a hit is found
+	 */
+	bool Sweep(struct FHitResult& OutHit, const FVector& Start, const FVector& End, const FCollisionShape& Shape, bool bTraceComplex);
 
-		/** 
-		 *  Trace a box against just this bodyinstance
-		 *  @param  OutHit          Information about hit against this component, if true is returned
-		 *  @param  Start           Start location of the box
-		 *  @param  End             End location of the box
-		 *  @param  CollisionShape	Collision Shape
-		 *	@param	bTraceComplex	Should we trace against complex or simple collision of this body
-		 *  @return true if a hit is found
-		 */
-		bool Sweep(struct FHitResult& OutHit, const FVector& Start, const FVector& End, const FCollisionShape & Shape, bool bTraceComplex);
-
-		/** 
-		 *  Test if the bodyinstance overlaps with the geometry in the Pos/Rot
-		 *
-		 *	@param	PGeom			Geometry it would like to test
-		 *  @param  ShapePose       Transform information in world. Use U2PTransform to convert from FTransform
-		 *  @return true if PrimComp overlaps this component at the specified location/rotation
-		 */
 #if WITH_PHYSX
-		bool Overlap(const physx::PxGeometry& Geom, const physx::PxTransform&  ShapePose);
+	/**
+	 *  Test if the bodyinstance overlaps with the geometry in the Pos/Rot
+	 *
+	 *	@param	PGeom			Geometry it would like to test
+	 *  @param  ShapePose       Transform information in world. Use U2PTransform to convert from FTransform
+	 *  @return true if PrimComp overlaps this component at the specified location/rotation
+	 */
+	bool Overlap(const physx::PxGeometry& Geom, const physx::PxTransform&  ShapePose);
 #endif	//WITH_PHYSX
-		/**
-		 * Add an impulse to this bodyinstance, radiating out from the specified position.
-		 *
-		 * @param Origin		Point of origin for the radial impulse blast, in world space
-		 * @param Radius		Size of radial impulse. Beyond this distance from Origin, there will be no affect.
-		 * @param Strength		Maximum strength of impulse applied to body.
-		 * @param Falloff		Allows you to control the strength of the impulse as a function of distance from Origin.
-		 * @param bVelChange	If true, the Strength is taken as a change in velocity instead of an impulse (ie. mass will have no affect).
-		 */
-		void AddRadialImpulseToBody(const FVector& Origin, float Radius, float Strength, uint8 Falloff, bool bVelChange);
-
-		/**
-		 *	Add a force to this bodyinstance, originating from the supplied world-space location.
-		 *
-		 *	@param Origin		Origin of force in world space.
-		 *	@param Radius		Radius within which to apply the force.
-		 *	@param Strength		Strength of force to apply.
-		 *  @param Falloff		Allows you to control the strength of the force as a function of distance from Origin.
-		 */
-		void AddRadialForceToBody(const FVector& Origin, float Radius, float Strength, uint8 Falloff);
-
-		/**
-		 * Get distance to the body surface if available
-		 * It is only valid if BodyShape is convex
-		 * If point is inside or shape is not convex, it will return 0.f
-		 *
-		 * @param Point				Point in world space
-		 * @param OutPointOnBody	Point on the surface of body closest to Point
-		 */
-		float GetDistanceToBody(const FVector& Point, FVector& OutPointOnBody) const;
-
-		/** 
-		 * Returns memory used by resources allocated for this body instance ( ex. Physx resources )
-		 **/
-		SIZE_T GetBodyInstanceResourceSize(EResourceSizeMode::Type Mode) const;
-
-		/**
-		 * UObject notification by OwningComponent
-		 */
-		void FixupData(class UObject * Loader);
-
-		const FCollisionResponse & GetCollisionResponse() { return CollisionResponses; }
+
+	/**
+	 * Add an impulse to this bodyinstance, radiating out from the specified position.
+	 *
+	 * @param Origin		Point of origin for the radial impulse blast, in world space
+	 * @param Radius		Size of radial impulse. Beyond this distance from Origin, there will be no affect.
+	 * @param Strength		Maximum strength of impulse applied to body.
+	 * @param Falloff		Allows you to control the strength of the impulse as a function of distance from Origin.
+	 * @param bVelChange	If true, the Strength is taken as a change in velocity instead of an impulse (ie. mass will have no affect).
+	 */
+	void AddRadialImpulseToBody(const FVector& Origin, float Radius, float Strength, uint8 Falloff, bool bVelChange);
+
+	/**
+	 *	Add a force to this bodyinstance, originating from the supplied world-space location.
+	 *
+	 *	@param Origin		Origin of force in world space.
+	 *	@param Radius		Radius within which to apply the force.
+	 *	@param Strength		Strength of force to apply.
+	 *  @param Falloff		Allows you to control the strength of the force as a function of distance from Origin.
+	 */
+	void AddRadialForceToBody(const FVector& Origin, float Radius, float Strength, uint8 Falloff);
+
+	/**
+	 * Get distance to the body surface if available
+	 * It is only valid if BodyShape is convex
+	 * If point is inside or shape is not convex, it will return 0.f
+	 *
+	 * @param Point				Point in world space
+	 * @param OutPointOnBody	Point on the surface of body closest to Point
+	 */
+	float GetDistanceToBody(const FVector& Point, FVector& OutPointOnBody) const;
+
+	/** 
+	 * Returns memory used by resources allocated for this body instance ( ex. Physx resources )
+	 **/
+	SIZE_T GetBodyInstanceResourceSize(EResourceSizeMode::Type Mode) const;
+
+	/**
+	 * UObject notification by OwningComponent
+	 */
+	void FixupData(class UObject* Loader);
+
+	const FCollisionResponse& GetCollisionResponse() const { return CollisionResponses; }
 
 private:
 #if WITH_PHYSX
-		FPhysxUserData PhysxUserData;
-
-		/**
-		*  Trace a box against just this bodyinstance
-		*  @param  OutHit          Information about hit against this component, if true is returned
-		*  @param  Start           Start location of the box
-		*  @param  End             End location of the box
-		*  @param  CollisionShape	Collision Shape
-		*  @param	bTraceComplex	Should we trace against complex or simple collision of this body
-		*  @param	PGeom			Geometry that will sweep
-		*  @return true if a hit is found
-		*/
-		bool InternalSweep(struct FHitResult& OutHit, const FVector& Start, const FVector& End, const FCollisionShape & Shape, bool bTraceComplex, const physx::PxRigidActor* RigidBody, const physx::PxGeometry * Geometry);
+	FPhysxUserData PhysxUserData;
+
+	/**
+	 *  Trace a box against just this bodyinstance
+	 *  @param  OutHit          Information about hit against this component, if true is returned
+	 *  @param  Start           Start location of the box
+	 *  @param  End             End location of the box
+	 *  @param  CollisionShape	Collision Shape
+	 *  @param	bTraceComplex	Should we trace against complex or simple collision of this body
+	 *  @param	PGeom			Geometry that will sweep
+	 *  @return true if a hit is found
+	 */
+	bool InternalSweep(struct FHitResult& OutHit, const FVector& Start, const FVector& End, const FCollisionShape& Shape, bool bTraceComplex, const physx::PxRigidActor* RigidBody, const physx::PxGeometry* Geometry);
 #endif 
-		/**
-		 * Invalidate Collision Profile Name
-		 * This gets called when it invalidates the reason of Profile Name
-		 * for example, they would like to re-define CollisionEnabled or ObjectType or ResponseChannels
-		 */
-		void InvalidateCollisionProfileName();
+	/**
+	 * Invalidate Collision Profile Name
+	 * This gets called when it invalidates the reason of Profile Name
+	 * for example, they would like to re-define CollisionEnabled or ObjectType or ResponseChannels
+	 */
+	void InvalidateCollisionProfileName();
 		
-		/**
-		 * Return true if the collision profile name is valid
-		 */
-		static bool IsValidCollisionProfileName(FName InCollisionProfileName);
-
-		friend class UCollisionProfile;
-		friend class FBodyInstanceCustomization;
-		friend class ULandscapeHeightfieldCollisionComponent;
-		friend class ULandscapeMeshCollisionComponent;
+	/**
+	 * Return true if the collision profile name is valid
+	 */
+	static bool IsValidCollisionProfileName(FName InCollisionProfileName);
+
+	friend class UCollisionProfile;
+	friend class FBodyInstanceCustomization;
+	friend class ULandscapeHeightfieldCollisionComponent;
+	friend class ULandscapeMeshCollisionComponent;
 };
diff --git a/Engine/Source/Runtime/Engine/Private/PhysicsEngine/BodyInstance.cpp b/Engine/Source/Runtime/Engine/Private/PhysicsEngine/BodyInstance.cpp
--- a/Engine/Source/Runtime/Engine/Private/PhysicsEngine/BodyInstance.cpp
+++ b/Engine/Source/Runtime/Engine/Private/PhysicsEngine/BodyInstance.cpp
@@ -66,7 +66,7 @@ ECollisionResponse FCollisionResponse::GetResponse(ECollisionChannel Channel) co
 }
 
 /** Set all channels from ChannelResponse Array **/
-void FCollisionResponse::SetCollisionResponseContainer(const FCollisionResponseContainer & InResponseToChannels)
+void FCollisionResponse::SetCollisionResponseContainer(const FCollisionResponseContainer& InResponseToChannels)
 {
 	ResponseToChannels = InResponseToChannels;
 #if 1// @hack until PostLoad is disabled for CDO of BP WITH_EDITOR
@@ -75,7 +75,7 @@ void FCollisionResponse::SetCollisionResponseContainer(const FCollisionResponseC
 #endif
 }
 
-void FCollisionResponse::SetResponsesArray(const TArray<FResponseChannel> & InChannelResponses)
+void FCollisionResponse::SetResponsesArray(const TArray<FResponseChannel>& InChannelResponses)
 {
 #if DO_GUARD_SLOW
 	// verify if the name is overlapping, if so, ensure, do not remove in debug becuase it will cause inconsistent bug between debug/release
@@ -139,7 +139,7 @@ void FCollisionResponse::UpdateArrayFromResponseContainer()
 {
 	ResponseArray.Empty();
 
-	const FCollisionResponseContainer & DefaultResponse = FCollisionResponseContainer::GetDefaultResponseContainer();
+	const FCollisionResponseContainer& DefaultResponse = FCollisionResponseContainer::GetDefaultResponseContainer();
 
 	for(int32 i=0; i<ARRAY_COUNT(ResponseToChannels.EnumArray); i++)
 	{
@@ -166,7 +166,7 @@ void FCollisionResponse::UpdateResponseContainerFromArray()
 
 	for (auto Iter = ResponseArray.CreateIterator(); Iter; ++Iter)
 	{
-		FResponseChannel & Response = *Iter;
+		FResponseChannel& Response = *Iter;
 
 		int32 EnumIndex = UCollisionProfile::Get()->ReturnContainerIndexFromChannelName(Response.Channel);
 		if ( EnumIndex != INDEX_NONE )
@@ -391,7 +391,7 @@ void FBodyInstance::SetResponseToChannel(ECollisionChannel Channel, ECollisionRe
 	UpdatePhysicsFilterData();
 }
 
-const FCollisionResponseContainer & FBodyInstance::GetResponseToChannels() const
+const FCollisionResponseContainer& FBodyInstance::GetResponseToChannels() const
 {
 	return CollisionResponses.GetResponseContainer();
 }
@@ -528,7 +528,7 @@ void FBodyInstance::UpdatePhysicsFilterData()
 	// if no collision, but if world wants to enable trace collision for components, allow it
 	if(UseCollisionEnabled == ECollisionEnabled::NoCollision && Owner && Owner->IsA(AVolume::StaticClass())==false )
 	{
-		UWorld * World = Owner->GetWorld();
+		UWorld* World = Owner->GetWorld();
 		UPrimitiveComponent* PrimComp = OwnerComponent.Get();
 		if (World && World->bEnableTraceCollision && 
 			(PrimComp->IsA(UStaticMeshComponent::StaticClass()) || PrimComp->IsA(USkeletalMeshComponent::StaticClass()) || PrimComp->IsA(UBrushComponent::StaticClass())))
@@ -1005,7 +1005,7 @@ void FBodyInstance::InitBody(UBodySetup* Setup, const FTransform& Transform, UPr
 
 //helper function for TermBody to avoid code duplication between scenes
 #if WITH_PHYSX
-void TermBodyHelper(int32 & SceneIndex, PxRigidActor *& PRigidActor, FBodyInstance * BodyInstance)
+void TermBodyHelper(int32& SceneIndex, PxRigidActor*& PRigidActor, FBodyInstance* BodyInstance)
 {
 	if (SceneIndex)
 	{
@@ -1019,7 +1019,7 @@ void TermBodyHelper(int32 & SceneIndex, PxRigidActor *& PRigidActor, FBodyInstan
 			if (PRigidActor)
 			{
 				// Let FPhysScene know
-				FPhysScene * PhysScene = FPhysxUserData::Get<FPhysScene>(PScene->userData);
+				FPhysScene* PhysScene = FPhysxUserData::Get<FPhysScene>(PScene->userData);
 				if (PhysScene)
 				{
 					PhysScene->TermBody(BodyInstance);
@@ -1078,20 +1078,20 @@ void FBodyInstance::TermBody()
 }
 
 #if WITH_BODY_WELDING
-void FBodyInstance::Weld(FBodyInstance * TheirBody, const FTransform & RelativeTM)
+void FBodyInstance::Weld(FBodyInstance* TheirBody, const FTransform& RelativeTM)
 {
 	check(TheirBody);
 #if WITH_PHYSX
 	
-	//UBodySetup * LocalSpaceSetup = TheirBody->BodySetup->CreateSpaceCopy(RelativeTM);
+	//UBodySetup* LocalSpaceSetup = TheirBody->BodySetup->CreateSpaceCopy(RelativeTM);
 
 	//child body gets placed into the same scenes as parent body
-	if (PxRigidActor * MyBody = RigidActorSync)
+	if (PxRigidActor* MyBody = RigidActorSync)
 	{
 		TheirBody->BodySetup->AddShapesToRigidActor(MyBody, Scale3D, &RelativeTM);
 	}
 
-	if (PxRigidActor * MyBody = RigidActorAsync)
+	if (PxRigidActor* MyBody = RigidActorAsync)
 	{
 		TheirBody->BodySetup->AddShapesToRigidActor(MyBody, Scale3D, &RelativeTM);
 	}
@@ -1145,7 +1145,7 @@ namespace EScaleMode
 }
 
 //computes the relative scaling vectors based on scale mode used
-void ComputeScalingVectors(EScaleMode::Type ScaleMode, const FVector & NewScale3D, const FVector & OldScale3D, FVector & RelativeScale3D, FVector & RelativeScale3DAbs, FVector & OutScale3D)
+void ComputeScalingVectors(EScaleMode::Type ScaleMode, const FVector& NewScale3D, const FVector& OldScale3D, FVector& RelativeScale3D, FVector& RelativeScale3DAbs, FVector& OutScale3D)
 {
 	FVector NewScale3DAbs = NewScale3D.GetAbs();
 	FVector OldScale3DAbs = OldScale3D.GetAbs();
@@ -1210,7 +1210,7 @@ void ComputeScalingVectors(EScaleMode::Type ScaleMode, const FVector & NewScale3
 	}
 }
 
-bool FBodyInstance::UpdateBodyScale(const FVector & InScale3D)
+bool FBodyInstance::UpdateBodyScale(const FVector& InScale3D)
 {
 	FVector InScale3DAdjusted = InScale3D;
 #if WITH_PHYSX
@@ -1252,7 +1252,7 @@ bool FBodyInstance::UpdateBodyScale(const FVector & InScale3D)
 
 	for (int32 ShapeIdx = 0; ShapeIdx < PShapes.Num(); ++ShapeIdx)
 	{
-		PxShape * PShape = PShapes[ShapeIdx];
+		PxShape* PShape = PShapes[ShapeIdx];
 		PxGeometryType::Enum GeomType = PShape->getGeometryType();
 		
 		if (GeomType == PxGeometryType::eSPHERE)
@@ -1281,7 +1281,7 @@ bool FBodyInstance::UpdateBodyScale(const FVector & InScale3D)
 
 	for(int32 ShapeIdx=0; ShapeIdx<PShapes.Num(); ShapeIdx++)
 	{
-		PxGeometry * UpdatedGeometry = NULL;
+		PxGeometry* UpdatedGeometry = NULL;
 		PxShape* PShape = PShapes[ShapeIdx];
 		PxScene* PScene = PShape->getActor()->getScene();
 		SCENE_LOCK_READ(PScene);
@@ -1649,8 +1649,8 @@ void FBodyInstance::SetBodyTransform(const FTransform& NewTransform, bool bTelep
 		// If kinematic and not teleporting, set kinematic target
 		if(!IsRigidDynamicNonKinematic(PRigidDynamic) && !bTeleport)
 		{
-			const PxScene * PScene = PRigidDynamic->getScene();
-			FPhysScene * PhysScene = FPhysxUserData::Get<FPhysScene>(PScene->userData);
+			const PxScene* PScene = PRigidDynamic->getScene();
+			FPhysScene* PhysScene = FPhysxUserData::Get<FPhysScene>(PScene->userData);
 			PhysScene->SetKinematicTarget(this, NewTransform, true);
 		}
 		// Otherwise, set global pose
@@ -1678,7 +1678,7 @@ void FBodyInstance::SetBodyTransform(const FTransform& NewTransform, bool bTelep
 #endif  // WITH_PHYSX
 }
 
-FVector FBodyInstance::GetUnrealWorldVelocity()
+FVector FBodyInstance::GetUnrealWorldVelocity() const
 {
 	FVector LinVel(0,0,0);
 #if WITH_PHYSX
@@ -1693,7 +1693,7 @@ FVector FBodyInstance::GetUnrealWorldVelocity()
 }
 
 /** Note: returns angular velocity in degrees per second. */
-FVector FBodyInstance::GetUnrealWorldAngularVelocity()
+FVector FBodyInstance::GetUnrealWorldAngularVelocity() const
 {
 	FVector AngVel(0,0,0);
 #if WITH_PHYSX
@@ -1707,7 +1707,7 @@ FVector FBodyInstance::GetUnrealWorldAngularVelocity()
 	return AngVel;
 }
 
-FVector FBodyInstance::GetUnrealWorldVelocityAtPoint(const FVector& Point)
+FVector FBodyInstance::GetUnrealWorldVelocityAtPoint(const FVector& Point) const
 {
 	FVector LinVel(0,0,0);
 
@@ -1725,7 +1725,7 @@ FVector FBodyInstance::GetUnrealWorldVelocityAtPoint(const FVector& Point)
 }
 
 
-FVector FBodyInstance::GetCOMPosition()
+FVector FBodyInstance::GetCOMPosition() const
 {
 #if WITH_PHYSX
 	PxRigidDynamic* PRigidDynamic = GetPxRigidDynamic();
@@ -1887,14 +1887,14 @@ void FBodyInstance::SetPhysMaterialOverride( UPhysicalMaterial* NewPhysMaterial
 
 
 
-UPhysicalMaterial* FBodyInstance::GetSimplePhysicalMaterial()
+UPhysicalMaterial* FBodyInstance::GetSimplePhysicalMaterial() const
 {
 	check( GEngine->DefaultPhysMaterial != NULL );
 
 	// Find the PhysicalMaterial we need to apply to the physics bodies.
 	// (LOW priority) Engine Mat, Material PhysMat, BodySetup Mat, Component Override, Body Override (HIGH priority)
 
-	UPhysicalMaterial * ReturnPhysMaterial = NULL;
+	UPhysicalMaterial* ReturnPhysMaterial = NULL;
 
 	// BodyInstance override
 	if( PhysMaterialOverride != NULL)	
@@ -1944,7 +1944,7 @@ UPhysicalMaterial* FBodyInstance::GetSimplePhysicalMaterial()
 	return ReturnPhysMaterial;
 }
 
-TArray<UPhysicalMaterial*> FBodyInstance::GetComplexPhysicalMaterials()
+TArray<UPhysicalMaterial*> FBodyInstance::GetComplexPhysicalMaterials() const
 {
 	TArray<UPhysicalMaterial*> PhysMaterials;
 
@@ -2134,7 +2134,7 @@ void FBodyInstance::SetMaxAngularVelocity(float NewMaxAngVel, bool bAddToCurrent
 {
 	//NewMaxAngVel = 400;
 #if WITH_PHYSX
-	PxRigidDynamic * PRigidDynamic = GetPxRigidDynamic();
+	PxRigidDynamic* PRigidDynamic = GetPxRigidDynamic();
 	if (PRigidDynamic)
 	{
 		float RadNewMaxAngVel = FMath::DegreesToRadians(NewMaxAngVel);
@@ -2163,8 +2163,8 @@ void FBodyInstance::AddForce(const FVector& Force, bool bAllowSubstepping)
 	PxRigidDynamic* PRigidDynamic = GetPxRigidDynamic();
 	if(IsRigidDynamicNonKinematic(PRigidDynamic))
 	{
-		const PxScene * PScene = PRigidDynamic->getScene();
-		FPhysScene * PhysScene = FPhysxUserData::Get<FPhysScene>(PScene->userData);
+		const PxScene* PScene = PRigidDynamic->getScene();
+		FPhysScene* PhysScene = FPhysxUserData::Get<FPhysScene>(PScene->userData);
 		PhysScene->AddForce(this, Force, bAllowSubstepping);
 		
 	}
@@ -2177,8 +2177,8 @@ void FBodyInstance::AddForceAtPosition(const FVector& Force, const FVector& Posi
 	PxRigidDynamic* PRigidDynamic = GetPxRigidDynamic();
 	if(IsRigidDynamicNonKinematic(PRigidDynamic))
 	{
-		const PxScene * PScene = PRigidDynamic->getScene();
-		FPhysScene * PhysScene = FPhysxUserData::Get<FPhysScene>(PScene->userData);
+		const PxScene* PScene = PRigidDynamic->getScene();
+		FPhysScene* PhysScene = FPhysxUserData::Get<FPhysScene>(PScene->userData);
 		PhysScene->AddForceAtPosition(this, Force, Position, bAllowSubstepping);
 	}
 #endif // WITH_PHYSX
@@ -2190,8 +2190,8 @@ void FBodyInstance::AddTorque(const FVector& Torque, bool bAllowSubstepping)
 	PxRigidDynamic* PRigidDynamic = GetPxRigidDynamic();
 	if(IsRigidDynamicNonKinematic(PRigidDynamic))
 	{
-		const PxScene * PScene = PRigidDynamic->getScene();
-		FPhysScene * PhysScene = FPhysxUserData::Get<FPhysScene>(PScene->userData);
+		const PxScene* PScene = PRigidDynamic->getScene();
+		FPhysScene* PhysScene = FPhysxUserData::Get<FPhysScene>(PScene->userData);
 		PhysScene->AddTorque(this, Torque, bAllowSubstepping);
 	}
 #endif // WITH_PHYSX
@@ -2272,7 +2272,7 @@ void FBodyInstance::AddRadialForceToBody(const FVector& Origin, float Radius, fl
 }
 
 
-FString FBodyInstance::GetBodyDebugName()
+FString FBodyInstance::GetBodyDebugName() const
 {
 	FString DebugName;
 
@@ -2421,7 +2421,7 @@ bool FBodyInstance::Sweep(struct FHitResult& OutHit, const FVector& Start, const
 }
 
 #if WITH_PHYSX
-bool FBodyInstance::InternalSweep(struct FHitResult& OutHit, const FVector& Start, const FVector& End, const FCollisionShape & Shape, bool bTraceComplex, const PxRigidActor* RigidBody, const PxGeometry * Geometry)
+bool FBodyInstance::InternalSweep(struct FHitResult& OutHit, const FVector& Start, const FVector& End, const FCollisionShape& Shape, bool bTraceComplex, const PxRigidActor* RigidBody, const PxGeometry* Geometry)
 {
 	const FVector Delta = End - Start;
 	const float DeltaMag = Delta.Size();
@@ -2456,7 +2456,7 @@ bool FBodyInstance::InternalSweep(struct FHitResult& OutHit, const FVector& Star
 			if((bTraceComplex && bShapeIsComplex) || (!bTraceComplex && bShapeIsSimple))
 			{
 				GeometryFromShapeStorage GeomStorage;
-				PxGeometry * PGeom = GetGeometryFromShape(GeomStorage, PShape, true);
+				PxGeometry* PGeom = GetGeometryFromShape(GeomStorage, PShape, true);
 				PxTransform PGlobalPose = PCompTM.transform(PShape->getLocalPose());
 
 				if (PGeom)
@@ -2651,7 +2651,7 @@ SIZE_T FBodyInstance::GetBodyInstanceResourceSize(EResourceSizeMode::Type Mode)
 	return ResSize;
 }
 
-void FBodyInstance::FixupData(class UObject * Loader)
+void FBodyInstance::FixupData(class UObject* Loader)
 {
 	check (Loader);
 
