#ue4 - ACharacter::SetBase() calls CharacterMovement->SaveBaseLocation() rather than using different logic for setting RelativeMovement.Rotation.

diff --git a/Engine/Source/Runtime/Engine/Private/Character.cpp b/Engine/Source/Runtime/Engine/Private/Character.cpp
--- a/Engine/Source/Runtime/Engine/Private/Character.cpp
+++ b/Engine/Source/Runtime/Engine/Private/Character.cpp
@@ -485,17 +485,14 @@ void ACharacter::SetBase( UPrimitiveComponent* NewBaseComponent, bool bNotifyPaw
 		MovementBase = NewBaseComponent;
 
 		// Update relative location/rotation
+		// Only do this for non-simulated proxies, because we don't want to overwrite RelativeMovement for those (it could about to be updated by a network update).
 		if ( Role == ROLE_Authority || Role == ROLE_AutonomousProxy )
 		{
 			UE_LOG(LogCharacter, Verbose, TEXT("Setting base on %s for '%s' to '%s'"), Role == ROLE_Authority ? TEXT("Server") : TEXT("AutoProxy"), *GetName(), *GetFullNameSafe(NewBaseComponent));
 			RelativeMovement.MovementBase = NewBaseComponent;
 			RelativeMovement.bServerHasBaseComponent = (NewBaseComponent != NULL); // Flag whether the server had a non-null base.
 
-			if (MovementBaseUtility::UseRelativePosition(MovementBase))
-			{
-				RelativeMovement.Location = GetActorLocation() - MovementBase->GetComponentLocation();
-				RelativeMovement.Rotation = (FRotationMatrix(GetActorRotation()) * FRotationMatrix(MovementBase->GetComponentRotation()).GetTransposed()).Rotator();
-			}
+			CharacterMovement->SaveBaseLocation();
 		}
 		else
 		{
