#ue4 - CharacterMovementComponent::ComputeFloorDist() uses an even shorter capsule when the initial test starts in penetration.

diff --git a/Engine/Source/Runtime/Engine/Private/CharacterMovementComponent.cpp b/Engine/Source/Runtime/Engine/Private/CharacterMovementComponent.cpp
--- a/Engine/Source/Runtime/Engine/Private/CharacterMovementComponent.cpp
+++ b/Engine/Source/Runtime/Engine/Private/CharacterMovementComponent.cpp
@@ -2203,7 +2203,7 @@ void UCharacterMovementComponent::ApplyVelocityBraking(float DeltaTime, float Fr
 		Velocity = Velocity + ((-2.f * Friction) * Velocity + RevAccel) * dt ; 
 		if( (Velocity | OldVel) > 0.f )
 		{
-			SumVel += dt * Velocity/DeltaTime;
+			SumVel += Velocity * (dt / DeltaTime);
 		}
 		else
 		{
@@ -3752,7 +3752,7 @@ void UCharacterMovementComponent::ComputeFloorDist(const FVector& CapsuleLocatio
 		// Use a shorter height to avoid sweeps giving weird results if we start on a surface.
 		// This also allows us to adjust out of penetrations.
 		const float ShrinkScale = 0.9f;
-		const float ShrinkScaleOverlap = 0.6f;
+		const float ShrinkScaleOverlap = 0.1f;
 		float ShrinkHeight = (PawnHalfHeight - PawnRadius) * (1.f - ShrinkScale);
 		float TraceDist = SweepDistance + ShrinkHeight;
 
@@ -3781,7 +3781,7 @@ void UCharacterMovementComponent::ComputeFloorDist(const FVector& CapsuleLocatio
 				ShrinkHeight = (PawnHalfHeight - PawnRadius) * (1.f - ShrinkScaleOverlap);
 				TraceDist = SweepDistance + ShrinkHeight;
 				CapsuleShape.Capsule.Radius = FMath::Max(0.f, CapsuleShape.Capsule.Radius - SWEEP_EDGE_REJECT_DISTANCE - KINDA_SMALL_NUMBER);
-				CapsuleShape.Capsule.HalfHeight = FMath::Max(PawnHalfHeight - ShrinkHeight, 0.1f);
+				CapsuleShape.Capsule.HalfHeight = FMath::Max(PawnHalfHeight - ShrinkHeight, CapsuleShape.Capsule.Radius);
 				Hit.Reset(1.f, false);
 
 				if(bUseFlatBaseForFloorChecks)
