Fix crash in FNavTestSceneProxy DrawDynamicElements when dragging and dropping a blueprint that is using the ANavigationTestingActor

diff --git a/Engine/Source/Runtime/Engine/Private/AI/Navigation/NavTestRenderingComponent.cpp b/Engine/Source/Runtime/Engine/Private/AI/Navigation/NavTestRenderingComponent.cpp
--- a/Engine/Source/Runtime/Engine/Private/AI/Navigation/NavTestRenderingComponent.cpp
+++ b/Engine/Source/Runtime/Engine/Private/AI/Navigation/NavTestRenderingComponent.cpp
@@ -73,27 +73,30 @@ public:
 
 	virtual void DrawDynamicElements(FPrimitiveDrawInterface* PDI,const FSceneView* View) OVERRIDE
 	{
-		//DrawArc(PDI, Link.Left, Link.Right, 0.4f, NavMeshColors[Link.AreaID], SDPG_World, 3.5f);
-		//const FVector VOffset(0,0,FVector::Dist(Link.Left, Link.Right)*1.333f);
-		//DrawArrowHead(PDI, Link.Right, Link.Left+VOffset, 30.f, NavMeshColors[Link.AreaID], SDPG_World, 3.5f);
-
-		const FVector ActorLocation = NavTestActor->GetActorLocation();
-		const FVector ProjectedLocation = NavTestActor->ProjectedLocation + NavMeshDrawOffset;
-		const FColor ProjectedColor = NavTestActor->bProjectedLocationValid ? FColor(0,255,0,120) : FColor(255,0,0,120);
-		const FVector BoxExtent(20,20,20);
-
-		FMaterialRenderProxy* const ColoredMeshInstance = new(FMemStack::Get()) FColoredMaterialRenderProxy(GEngine->DebugMeshMaterial->GetRenderProxy(false), ProjectedColor);
-		//DrawBox(PDI, FTransform(ProjectedLocation).ToMatrixNoScale(),BoxExtent, ColoredMeshInstance, SDPG_World);
-		DrawSphere(PDI, ProjectedLocation, BoxExtent, 10, 7, ColoredMeshInstance, SDPG_World);
-
-		//DrawWireBox(PDI, FBox(ProjectedLocation-BoxExtent, ProjectedLocation+BoxExtent), ProjectedColor, false);
-		DrawWireBox(PDI, FBox(ActorLocation-BoxExtent, ActorLocation+BoxExtent), FColor::White, false);			
-		const FVector LineEnd = ProjectedLocation - (ProjectedLocation-ActorLocation).SafeNormal()*BoxExtent.X;
-		PDI->DrawLine(LineEnd, ActorLocation, ProjectedColor, SDPG_World, 2.5);
-		DrawArrowHead(PDI, LineEnd, ActorLocation, 20.f, ProjectedColor, SDPG_World, 2.5f);
-
-		// draw query extent
-		DrawWireBox(PDI, FBox(ActorLocation-NavTestActor->QueryingExtent, ActorLocation+NavTestActor->QueryingExtent), FColor::Blue, false);
+		if (NavTestActor)
+		{
+			//DrawArc(PDI, Link.Left, Link.Right, 0.4f, NavMeshColors[Link.AreaID], SDPG_World, 3.5f);
+			//const FVector VOffset(0,0,FVector::Dist(Link.Left, Link.Right)*1.333f);
+			//DrawArrowHead(PDI, Link.Right, Link.Left+VOffset, 30.f, NavMeshColors[Link.AreaID], SDPG_World, 3.5f);
+
+			const FVector ActorLocation = NavTestActor->GetActorLocation();
+			const FVector ProjectedLocation = NavTestActor->ProjectedLocation + NavMeshDrawOffset;
+			const FColor ProjectedColor = NavTestActor->bProjectedLocationValid ? FColor(0, 255, 0, 120) : FColor(255, 0, 0, 120);
+			const FVector BoxExtent(20, 20, 20);
+
+			FMaterialRenderProxy* const ColoredMeshInstance = new(FMemStack::Get()) FColoredMaterialRenderProxy(GEngine->DebugMeshMaterial->GetRenderProxy(false), ProjectedColor);
+			//DrawBox(PDI, FTransform(ProjectedLocation).ToMatrixNoScale(),BoxExtent, ColoredMeshInstance, SDPG_World);
+			DrawSphere(PDI, ProjectedLocation, BoxExtent, 10, 7, ColoredMeshInstance, SDPG_World);
+
+			//DrawWireBox(PDI, FBox(ProjectedLocation-BoxExtent, ProjectedLocation+BoxExtent), ProjectedColor, false);
+			DrawWireBox(PDI, FBox(ActorLocation - BoxExtent, ActorLocation + BoxExtent), FColor::White, false);
+			const FVector LineEnd = ProjectedLocation - (ProjectedLocation - ActorLocation).SafeNormal()*BoxExtent.X;
+			PDI->DrawLine(LineEnd, ActorLocation, ProjectedColor, SDPG_World, 2.5);
+			DrawArrowHead(PDI, LineEnd, ActorLocation, 20.f, ProjectedColor, SDPG_World, 2.5f);
+
+			// draw query extent
+			DrawWireBox(PDI, FBox(ActorLocation - NavTestActor->QueryingExtent, ActorLocation + NavTestActor->QueryingExtent), FColor::Blue, false);
+		}
 
 		// draw path
 		if (!bShowBestPath || !NodeDebug.Num())
