Duplicate includes into EnginePrivate.h, so headers can stop being included globally in Engine, but can still be included automaticallly externally (via Engine.h)
Stop including Engine.h from anywhere inside Engine module
Add missing '#pragma once' to Interface_PostProcessVolume.h
Stop including EngineModule.h in EnginePrivate.h

diff --git a/Engine/Source/Runtime/Engine/Classes/Interfaces/Interface_PostProcessVolume.h b/Engine/Source/Runtime/Engine/Classes/Interfaces/Interface_PostProcessVolume.h
--- a/Engine/Source/Runtime/Engine/Classes/Interfaces/Interface_PostProcessVolume.h
+++ b/Engine/Source/Runtime/Engine/Classes/Interfaces/Interface_PostProcessVolume.h
@@ -1,4 +1,6 @@
 // Copyright 1998-2013 Epic Games, Inc. All Rights Reserved.
+
+#pragma once
 #include "Interface_PostProcessVolume.generated.h"
 
 /** Interface for general PostProcessVolume access **/
diff --git a/Engine/Source/Runtime/Engine/Private/Engine.cpp b/Engine/Source/Runtime/Engine/Private/Engine.cpp
--- a/Engine/Source/Runtime/Engine/Private/Engine.cpp
+++ b/Engine/Source/Runtime/Engine/Private/Engine.cpp
@@ -5,6 +5,7 @@
 =============================================================================*/
 
 #include "EnginePrivate.h"
+#include "EngineModule.h"
 #include "EditorSupportDelegates.h"
 
 /*-----------------------------------------------------------------------------
diff --git a/Engine/Source/Runtime/Engine/Private/EnginePrivate.h b/Engine/Source/Runtime/Engine/Private/EnginePrivate.h
--- a/Engine/Source/Runtime/Engine/Private/EnginePrivate.h
+++ b/Engine/Source/Runtime/Engine/Private/EnginePrivate.h
@@ -4,13 +4,399 @@
 	EnginePrivate.h: Unreal engine private header file.
 =============================================================================*/
 
-/*----------------------------------------------------------------------------
-	Engine public includes.
-----------------------------------------------------------------------------*/
+#pragma once
+
+
+#include "Core.h"
+#include "CoreUObject.h"
+#include "InputCore.h"
+#include "EngineDefines.h"
+#include "EngineSettings.h"
+#include "EngineStats.h"
+#include "EngineLogs.h"
+#include "EngineGlobals.h"
+
+#include "Engine/EngineBaseTypes.h"
+#include "Camera/CameraTypes.h"
+#include "Engine/EngineTypes.h"
+#include "Components/ActorComponent.h"
+#include "GameFramework/Actor.h"
+#include "Sound/AmbientSound.h"
+#include "Engine/Brush.h"
+#include "GameFramework/Volume.h"
+#include "Engine/BlockingVolume.h"
+#include "GameFramework/CameraBlockingVolume.h"
+#include "Engine/CullDistanceVolume.h"
+#include "Engine/LevelStreamingVolume.h"
+#include "AI/Navigation/NavMeshBoundsVolume.h"
+#include "AI/Navigation/NavRelevantActorInterface.h"
+#include "AI/Navigation/NavModifierVolume.h"
+#include "GameFramework/PhysicsVolume.h"
+#include "GameFramework/DefaultPhysicsVolume.h"
+#include "GameFramework/KillZVolume.h"
+#include "GameFramework/PainCausingVolume.h"
+#include "Interfaces/Interface_PostProcessVolume.h"
+#include "Engine/PostProcessVolume.h"
+#include "Sound/ReverbVolume.h"
+#include "Engine/TriggerVolume.h"
+#include "Camera/CameraActor.h"
+#include "AI/Navigation/NavAgentInterface.h"
+#include "GameFramework/Controller.h"
+#include "AI/Navigation/NavLinkDefinition.h"
+#include "AI/Navigation/NavigationTypes.h"
+#include "Kismet/BlueprintFunctionLibrary.h"
+#include "AI/Navigation/NavigationSystem.h"
+#include "AI/Navigation/NavCollision.h"
+#include "GameFramework/PlayerMuteList.h"
+#include "Camera/PlayerCameraManager.h"
+#include "GameFramework/Info.h"
+#include "GameFramework/GameMode.h"
+#include "Components/InputComponent.h"
+#include "Curves/CurveBase.h"
+#include "Curves/CurveFloat.h"
+#include "GameFramework/ForceFeedbackEffect.h"
+#include "GameFramework/PlayerController.h"
+#include "Engine/DebugCameraController.h"
+#include "Engine/DecalActor.h"
+#include "PhysicsEngine/DestructibleActor.h"
+#include "Engine/Canvas.h"
+#include "GameFramework/HUD.h"
+#include "Debug/GameplayDebuggingHUDComponent.h"
+#include "Atmosphere/AtmosphericFog.h"
+#include "Engine/ExponentialHeightFog.h"
+#include "GameFramework/GameSession.h"
+#include "GameFramework/GameState.h"
+#include "GameFramework/PlayerState.h"
+#include "Engine/SkyLight.h"
+#include "Engine/WindDirectionalSource.h"
+#include "GameFramework/MusicTrackDataStructures.h"
+#include "GameFramework/WorldSettings.h"
+#include "PhysicsEngine/BodyInstance.h"
+#include "Components/SceneComponent.h"
+#include "Components/LightComponentBase.h"
+#include "Components/LightComponent.h"
+#include "Engine/LevelBounds.h"
+#include "Engine/LevelScriptActor.h"
+#include "Engine/Light.h"
+#include "Engine/DirectionalLight.h"
+#include "Engine/PointLight.h"
+#include "Engine/SpotLight.h"
+#include "Engine/GeneratedMeshAreaLight.h"
+#include "AI/Navigation/NavigationData.h"
+#include "AI/Navigation/NavigationGraph.h"
+#include "AI/Navigation/NavAreas/NavArea.h"
+#include "AI/Navigation/RecastNavMesh.h"
+#include "AI/Navigation/NavigationGraphNode.h"
+#include "Engine/NavigationObjectBase.h"
+#include "GameFramework/PlayerStart.h"
+#include "Engine/PlayerStartPIE.h"
+#include "AI/Navigation/NavPathObserverInterface.h"
+#include "AI/Navigation/NavigationTestingActor.h"
+#include "AI/Navigation/NavLinkHostInterface.h"
+#include "AI/Navigation/NavLinkProxy.h"
+#include "Engine/Note.h"
+#include "Components/PrimitiveComponent.h"
+#include "GameFramework/MovementComponent.h"
+#include "GameFramework/NavMovementComponent.h"
+#include "GameFramework/PawnMovementComponent.h"
+#include "GameFramework/Pawn.h"
+#include "Interfaces/NetworkPredictionInterface.h"
+#include "GameFramework/CharacterMovementComponent.h"
+#include "GameFramework/Character.h"
+#include "GameFramework/FloatingPawnMovement.h"
+#include "GameFramework/DefaultPawn.h"
+#include "GameFramework/SpectatorPawn.h"
+#include "GameFramework/WheeledVehicle.h"
+#include "Engine/ReflectionCapture.h"
+#include "Engine/BoxReflectionCapture.h"
+#include "Engine/PlaneReflectionCapture.h"
+#include "Engine/SphereReflectionCapture.h"
+#include "PhysicsEngine/RigidBodyBase.h"
+#include "PhysicsEngine/PhysicsConstraintActor.h"
+#include "PhysicsEngine/PhysicsThruster.h"
+#include "PhysicsEngine/RadialForceActor.h"
+#include "Engine/SceneCapture.h"
+#include "Engine/SceneCapture2D.h"
+#include "Engine/SceneCaptureCube.h"
+#include "Components/MeshComponent.h"
+#include "Components/SkinnedMeshComponent.h"
+#include "PhysicsEngine/ConstraintInstance.h"
+#include "PhysicsEngine/PhysicsConstraintComponent.h"
+#include "Components/SkeletalMeshComponent.h"
+#include "Engine/StaticMeshActor.h"
+#include "Engine/TargetPoint.h"
+#include "Engine/TextRenderActor.h"
+#include "Engine/TriggerBase.h"
+#include "Engine/TriggerBox.h"
+#include "Engine/TriggerCapsule.h"
+#include "Engine/TriggerSphere.h"
+#include "VectorField/VectorFieldVolume.h"
+#include "Components/ApplicationLifecycleComponent.h"
+#include "Engine/DataAsset.h"
+#include "Debug/GameplayDebuggingControllerComponent.h"
+#include "GameFramework/SpectatorPawnMovement.h"
+#include "Vehicles/VehicleWheel.h"
+#include "Vehicles/WheeledVehicleMovementComponent.h"
+#include "Vehicles/WheeledVehicleMovementComponent4W.h"
+#include "GameFramework/ProjectileMovementComponent.h"
+#include "GameFramework/RotatingMovementComponent.h"
+#include "Components/PawnNoiseEmitterComponent.h"
+#include "PhysicsEngine/PhysicsHandleComponent.h"
+#include "Atmosphere/AtmosphericFogComponent.h"
+#include "Sound/SoundAttenuation.h"
+#include "Components/AudioComponent.h"
+#include "Camera/CameraComponent.h"
+#include "Components/ChildActorComponent.h"
+#include "Components/DecalComponent.h"
+#include "Components/ExponentialHeightFogComponent.h"
+#include "Components/DirectionalLightComponent.h"
+#include "Components/PointLightComponent.h"
+#include "Components/SpotLightComponent.h"
+#include "Components/SkyLightComponent.h"
+#include "AI/Navigation/NavigationGraphNodeComponent.h"
+#include "PhysicsEngine/PhysicsThrusterComponent.h"
+#include "Components/PostProcessComponent.h"
+#include "Components/ArrowComponent.h"
+#include "Components/BillboardComponent.h"
+#include "Components/BrushComponent.h"
+#include "Components/DrawFrustumComponent.h"
+#include "Debug/DebugDrawService.h"
+#include "Debug/GameplayDebuggingComponent.h"
+#include "Components/LineBatchComponent.h"
+#include "Components/MaterialBillboardComponent.h"
+#include "Components/DestructibleComponent.h"
+#include "Components/PoseableMeshComponent.h"
+#include "Lightmass/LightmassPrimitiveSettingsObject.h"
+#include "Components/StaticMeshComponent.h"
+#include "Components/InstancedStaticMeshComponent.h"
+#include "Interfaces/Interface_CollisionDataProvider.h"
+#include "Components/SplineMeshComponent.h"
+#include "Components/ModelComponent.h"
+#include "AI/Navigation/NavLinkRenderingComponent.h"
+#include "AI/Navigation/NavMeshRenderingComponent.h"
+#include "AI/Navigation/NavTestRenderingComponent.h"
+#include "Components/ShapeComponent.h"
+#include "Components/BoxComponent.h"
+#include "Components/CapsuleComponent.h"
+#include "Components/SphereComponent.h"
+#include "Components/DrawSphereComponent.h"
+#include "Components/TextRenderComponent.h"
+#include "Components/VectorFieldComponent.h"
+#include "PhysicsEngine/RadialForceComponent.h"
+#include "Components/ReflectionCaptureComponent.h"
+#include "Components/BoxReflectionCaptureComponent.h"
+#include "Components/PlaneReflectionCaptureComponent.h"
+#include "Components/SphereReflectionCaptureComponent.h"
+#include "Components/SceneCaptureComponent.h"
+#include "Components/SceneCaptureComponent2D.h"
+#include "Components/SceneCaptureComponentCube.h"
+#include "GameFramework/SpringArmComponent.h"
+#include "Components/WindDirectionalSourceComponent.h"
+#include "Components/TimelineComponent.h"
+#include "Animation/AnimSequenceBase.h"
+#include "Animation/AnimCompositeBase.h"
+#include "Animation/AnimComposite.h"
+#include "Animation/AnimMontage.h"
+#include "Animation/AnimSequence.h"
+#include "Animation/BlendSpaceBase.h"
+#include "Animation/AnimStateMachineTypes.h"
+#include "Animation/AnimInstance.h"
+#include "Animation/BlendSpace.h"
+#include "Animation/AimOffsetBlendSpace.h"
+#include "Animation/BlendSpace1D.h"
+#include "Animation/AimOffsetBlendSpace1D.h"
+#include "Animation/AnimSingleNodeInstance.h"
+#include "Vehicles/VehicleAnimInstance.h"
+#include "Animation/AnimNotifies/AnimNotify.h"
+#include "Animation/AnimNotifies/AnimNotifyState.h"
+#include "Animation/AnimNotifies/AnimNotifyState_TimedParticleEffect.h"
+#include "Animation/AnimNotifies/AnimNotifyState_Trail.h"
+#include "EditorFramework/AssetImportData.h"
+#include "Engine/AssetUserData.h"
+#include "Sound/AudioSettings.h"
+#include "AI/Navigation/AvoidanceManager.h"
+#include "Engine/BlendableInterface.h"
+#include "Engine/BlueprintCore.h"
+#include "Engine/Blueprint.h"
+#include "Animation/AnimBlueprint.h"
+#include "Sound/DialogueTypes.h"
+#include "Kismet/GameplayStatics.h"
+#include "Kismet/HeadMountedDisplayFunctionLibrary.h"
+#include "Kismet/KismetArrayLibrary.h"
+#include "Kismet/KismetInputLibrary.h"
+#include "Kismet/KismetMaterialLibrary.h"
+#include "Kismet/KismetMathLibrary.h"
+#include "Kismet/KismetNodeHelperLibrary.h"
+#include "Kismet/KismetStringLibrary.h"
+#include "Kismet/KismetSystemLibrary.h"
+#include "Kismet/KismetTextLibrary.h"
+#include "Engine/BlueprintGeneratedClass.h"
+#include "Animation/AnimBlueprintGeneratedClass.h"
+#include "PhysicsEngine/BodySetup.h"
+#include "Slate/ButtonStyleAsset.h"
+#include "Camera/CameraAnim.h"
+#include "Camera/CameraAnimInst.h"
+#include "Camera/CameraModifier.h"
+#include "Camera/CameraShake.h"
+#include "Camera/CameraModifier_CameraShake.h"
+#include "GameFramework/CheatManager.h"
+#include "Slate/CheckboxStyleAsset.h"
+#include "Engine/CollisionProfile.h"
+#include "Curves/CurveLinearColor.h"
+#include "Curves/CurveVector.h"
+#include "Curves/CurveEdPresetCurve.h"
+#include "Engine/CurveTable.h"
+#include "GameFramework/DamageType.h"
+#include "Vehicles/TireType.h"
+#include "Engine/DataTable.h"
+#include "Sound/DialogueVoice.h"
+#include "Sound/DialogueWave.h"
+#include "Distributions/Distribution.h"
+#include "Distributions/DistributionFloat.h"
+#include "Distributions/DistributionFloatConstant.h"
+#include "Distributions/DistributionFloatParameterBase.h"
+#include "Distributions/DistributionFloatParticleParameter.h"
+#include "Distributions/DistributionFloatSoundParameter.h"
+#include "Distributions/DistributionFloatConstantCurve.h"
+#include "Distributions/DistributionFloatUniform.h"
+#include "Distributions/DistributionFloatUniformCurve.h"
+#include "Distributions/DistributionVector.h"
+#include "Distributions/DistributionVectorConstant.h"
+#include "Distributions/DistributionVectorParameterBase.h"
+#include "Distributions/DistributionVectorParticleParameter.h"
+#include "Distributions/DistributionVectorConstantCurve.h"
+#include "Distributions/DistributionVectorUniform.h"
+#include "Distributions/DistributionVectorUniformCurve.h"
+#include "Engine/DynamicBlueprintBinding.h"
+#include "Engine/ComponentDelegateBinding.h"
+#include "Engine/InputDelegateBinding.h"
+#include "Engine/InputActionDelegateBinding.h"
+#include "Engine/InputAxisDelegateBinding.h"
+#include "Engine/InputAxisKeyDelegateBinding.h"
+#include "Engine/InputVectorAxisDelegateBinding.h"
+#include "Engine/InputKeyDelegateBinding.h"
+#include "Engine/InputTouchDelegateBinding.h"
+#include "Engine/Engine.h"
+#include "Engine/GameEngine.h"
+#include "Exporters/Exporter.h"
+#include "Engine/FontImportOptions.h"
+#include "Engine/Font.h"
+#include "GameFramework/GameUserSettings.h"
+#include "GameFramework/PlayerInput.h"
+#include "GameFramework/InputSettings.h"
+#include "Engine/InterpCurveEdSetup.h"
+#include "Engine/IntSerialization.h"
+#include "Layers/Layer.h"
+#include "Engine/LevelBase.h"
+#include "Engine/Level.h"
+#include "Engine/LevelStreaming.h"
+#include "Engine/LevelStreamingAlwaysLoaded.h"
+#include "Engine/LevelStreamingBounds.h"
+#include "Engine/LevelStreamingKismet.h"
+#include "Engine/LevelStreamingPersistent.h"
+#include "Lightmass/LightmappedSurfaceCollection.h"
+#include "GameFramework/LocalMessage.h"
+#include "GameFramework/EngineMessage.h"
+#include "Materials/MaterialInterface.h"
+#include "Materials/Material.h"
+#include "Materials/MaterialInstance.h"
+#include "Materials/MaterialInstanceConstant.h"
+#include "Materials/MaterialInstanceDynamic.h"
+#include "Materials/MaterialParameterCollection.h"
+#include "Materials/MaterialParameterCollectionInstance.h"
+#include "Engine/NetDriver.h"
+#include "Engine/ObjectLibrary.h"
+#include "Engine/ObjectReferencer.h"
+#include "GameFramework/OnlineSession.h"
+#include "Engine/PackageMapClient.h"
+#include "Engine/PendingNetGame.h"
+#include "PhysicalMaterials/PhysicalMaterial.h"
+#include "PhysicalMaterials/PhysicalMaterialPropertyBase.h"
+#include "PhysicsEngine/PhysicsAsset.h"
+#include "PhysicsEngine/PhysicsCollisionHandler.h"
+#include "PhysicsEngine/PhysicsConstraintTemplate.h"
+#include "PhysicsEngine/PhysicsSettings.h"
+#include "Engine/PlatformInterfaceBase.h"
+#include "Engine/CloudStorageBase.h"
+#include "Engine/InGameAdManager.h"
+#include "Engine/MicroTransactionBase.h"
+#include "Engine/TwitterIntegrationBase.h"
+#include "Engine/PlatformInterfaceWebResponse.h"
+#include "Engine/Player.h"
+#include "Engine/LocalPlayer.h"
+#include "Engine/NetConnection.h"
+#include "Engine/ChildConnection.h"
+#include "Engine/Polys.h"
+#include "Engine/RendererSettings.h"
+#include "Sound/ReverbEffect.h"
+#include "MovieScene/RuntimeMovieScenePlayerInterface.h"
+#include "GameFramework/SaveGame.h"
+#include "Engine/SaveGameSummary.h"
+#include "Engine/SCS_Node.h"
+#include "Engine/Selection.h"
+#include "Engine/SimpleConstructionScript.h"
+#include "Animation/PreviewAssetAttachComponent.h"
+#include "Engine/SkeletalMesh.h"
+#include "Engine/DestructibleMesh.h"
+#include "Engine/SkeletalMeshSocket.h"
+#include "Animation/Skeleton.h"
+#include "Slate/SlateBrushAsset.h"
+#include "Sound/SoundBase.h"
+#include "Sound/DialogueSoundWaveProxy.h"
+#include "Sound/SoundClass.h"
+#include "Sound/SoundCue.h"
+#include "Sound/SoundGroups.h"
+#include "Sound/SoundWave.h"
+#include "Sound/SoundWaveStreaming.h"
+#include "Sound/SoundMix.h"
+#include "Engine/StaticMesh.h"
+#include "Engine/StaticMeshSocket.h"
+#include "Camera/CameraStackTypes.h"
+#include "Engine/StreamableManager.h"
+#include "Engine/TextureDefines.h"
+#include "Tests/TextPropertyTestObject.h"
+#include "Engine/Texture.h"
+#include "Engine/Texture2D.h"
+#include "Engine/LightMapTexture2D.h"
+#include "Engine/ShadowMapTexture2D.h"
+#include "Engine/TextureLightProfile.h"
+#include "Engine/Texture2DDynamic.h"
+#include "Engine/TextureCube.h"
+#include "Engine/TextureMovie.h"
+#include "Engine/TextureRenderTarget.h"
+#include "Engine/TextureRenderTarget2D.h"
+#include "Engine/CanvasRenderTarget2D.h"
+#include "Engine/TextureRenderTargetCube.h"
+#include "EditorFramework/ThumbnailInfo.h"
+#include "Engine/TimelineTemplate.h"
+#include "GameFramework/TouchInterface.h"
+#include "Engine/UserDefinedEnum.h"
+#include "Engine/UserDefinedStruct.h"
+#include "VectorField/VectorField.h"
+#include "VectorField/VectorFieldAnimated.h"
+#include "VectorField/VectorFieldStatic.h"
+#include "Animation/VertexAnim/MorphTarget.h"
+#include "Engine/World.h"
+
+#include "SystemSettings.h"					// Scalability options.
+#include "SceneManagement.h"				// Scene management.
+
+#include "NavDataGenerator.h"
+#include "NavLinkRenderingProxy.h"
+#include "NavigationModifier.h"
+#include "NavigationOctree.h"
+#include "ContentStreaming.h"				// Content streaming class definitions.
+#include "ComponentReregisterContext.h"	
+#include "DrawDebugHelpers.h"
+#include "UnrealEngine.h"					// Unreal engine helpers.
+#include "Canvas.h"							// Canvas.
+#include "EngineUtils.h"
+#include "InstancedFoliage.h"				// Instanced foliage.
+#include "UnrealExporter.h"					// Exporter definition.
+#include "TimerManager.h"					// Game play timers
+#include "EngineService.h"
+#include "NavigationSystemHelpers.h"
+#include "HardwareInfo.h"
 
-#include "Engine.h"
 
-/*-----------------------------------------------------------------------------
-	Engine private includes.
------------------------------------------------------------------------------*/
 
diff --git a/Engine/Source/Runtime/Engine/Private/GameEngine.cpp b/Engine/Source/Runtime/Engine/Private/GameEngine.cpp
--- a/Engine/Source/Runtime/Engine/Private/GameEngine.cpp
+++ b/Engine/Source/Runtime/Engine/Private/GameEngine.cpp
@@ -13,6 +13,7 @@
 #include "MallocProfiler.h"
 #include "Net/NetworkProfiler.h"
 #include "ConfigCacheIni.h"
+#include "EngineModule.h"
 
 #include "AVIWriter.h"
 
diff --git a/Engine/Source/Runtime/Engine/Private/GameViewportClient.cpp b/Engine/Source/Runtime/Engine/Private/GameViewportClient.cpp
--- a/Engine/Source/Runtime/Engine/Private/GameViewportClient.cpp
+++ b/Engine/Source/Runtime/Engine/Private/GameViewportClient.cpp
@@ -15,6 +15,7 @@
 #include "IHeadMountedDisplay.h"
 #include "SVirtualJoystick.h"
 #include "SceneViewport.h"
+#include "EngineModule.h"
 
 #include "HighResScreenshot.h"
 
diff --git a/Engine/Source/Runtime/Engine/Private/Materials/MaterialShared.cpp b/Engine/Source/Runtime/Engine/Private/Materials/MaterialShared.cpp
--- a/Engine/Source/Runtime/Engine/Private/Materials/MaterialShared.cpp
+++ b/Engine/Source/Runtime/Engine/Private/Materials/MaterialShared.cpp
@@ -15,6 +15,7 @@
 #include "HLSLMaterialTranslator.h"
 #include "MaterialUniformExpressions.h"
 #include "Developer/TargetPlatform/Public/TargetPlatform.h"
+#include "EngineModule.h"
 
 DEFINE_LOG_CATEGORY(LogMaterial);
 
diff --git a/Engine/Source/Runtime/Engine/Private/PrimitiveDrawingUtils.cpp b/Engine/Source/Runtime/Engine/Private/PrimitiveDrawingUtils.cpp
--- a/Engine/Source/Runtime/Engine/Private/PrimitiveDrawingUtils.cpp
+++ b/Engine/Source/Runtime/Engine/Private/PrimitiveDrawingUtils.cpp
@@ -4,6 +4,7 @@
 #include "DynamicMeshBuilder.h"
 #include "SceneManagement.h"
 #include "StaticMeshResources.h"
+#include "EngineModule.h"
 
 /** Emits draw events for a given FMeshBatch and the FPrimitiveSceneProxy corresponding to that mesh element. */
 void EmitMeshDrawEvents(const FPrimitiveSceneProxy* PrimitiveSceneProxy, const FMeshBatch& Mesh)
diff --git a/Engine/Source/Runtime/Engine/Private/ReflectionCaptureComponent.cpp b/Engine/Source/Runtime/Engine/Private/ReflectionCaptureComponent.cpp
--- a/Engine/Source/Runtime/Engine/Private/ReflectionCaptureComponent.cpp
+++ b/Engine/Source/Runtime/Engine/Private/ReflectionCaptureComponent.cpp
@@ -7,6 +7,7 @@
 #include "EnginePrivate.h"
 #include "DerivedDataCacheInterface.h"
 #include "TargetPlatform.h"
+#include "EngineModule.h"
 
 /** 
  * Size of all reflection captures.
diff --git a/Engine/Source/Runtime/Engine/Private/RendererSupport.cpp b/Engine/Source/Runtime/Engine/Private/RendererSupport.cpp
--- a/Engine/Source/Runtime/Engine/Private/RendererSupport.cpp
+++ b/Engine/Source/Runtime/Engine/Private/RendererSupport.cpp
@@ -10,6 +10,7 @@
 #include "RenderCore.h"
 #include "GlobalShader.h"
 #include "Slate.h"
+#include "EngineModule.h"
 
 /** Clears and optionally backs up all references to renderer module classes in other modules, particularly engine. */
 void ClearReferencesToRendererModuleClasses(
diff --git a/Engine/Source/Runtime/Engine/Private/SceneView.cpp b/Engine/Source/Runtime/Engine/Private/SceneView.cpp
--- a/Engine/Source/Runtime/Engine/Private/SceneView.cpp
+++ b/Engine/Source/Runtime/Engine/Private/SceneView.cpp
@@ -6,7 +6,7 @@
 
 #include "EnginePrivate.h"
 #include "SceneManagement.h"
-#include "EnginePrivate.h"
+#include "EngineModule.h"
 #include "StereoRendering.h"
 #include "HighResScreenshot.h"
 #include "Slate/SceneViewport.h"
diff --git a/Engine/Source/Runtime/Engine/Private/ShaderCompiler/ShaderCompiler.cpp b/Engine/Source/Runtime/Engine/Private/ShaderCompiler/ShaderCompiler.cpp
--- a/Engine/Source/Runtime/Engine/Private/ShaderCompiler/ShaderCompiler.cpp
+++ b/Engine/Source/Runtime/Engine/Private/ShaderCompiler/ShaderCompiler.cpp
@@ -10,6 +10,7 @@
 #include "GlobalShader.h"
 #include "TargetPlatform.h"
 #include "DerivedDataCacheInterface.h"
+#include "EngineModule.h"
 
 DEFINE_LOG_CATEGORY(LogShaderCompilers);
 
diff --git a/Engine/Source/Runtime/Engine/Private/TileRendering.cpp b/Engine/Source/Runtime/Engine/Private/TileRendering.cpp
--- a/Engine/Source/Runtime/Engine/Private/TileRendering.cpp
+++ b/Engine/Source/Runtime/Engine/Private/TileRendering.cpp
@@ -6,6 +6,7 @@
 
 #include "EnginePrivate.h"
 #include "TileRendering.h"
+#include "EngineModule.h"
 
 /** 
 * vertex data for a screen quad 
diff --git a/Engine/Source/Runtime/Engine/Public/EngineModule.h b/Engine/Source/Runtime/Engine/Public/EngineModule.h
--- a/Engine/Source/Runtime/Engine/Public/EngineModule.h
+++ b/Engine/Source/Runtime/Engine/Public/EngineModule.h
@@ -2,6 +2,8 @@
 // 
 // Engine module class
 
+#pragma once
+
 /** Implements the engine module. */
 class FEngineModule : public FDefaultModuleImpl
 {
diff --git a/Engine/Source/Runtime/Engine/Public/GlobalShader.h b/Engine/Source/Runtime/Engine/Public/GlobalShader.h
--- a/Engine/Source/Runtime/Engine/Public/GlobalShader.h
+++ b/Engine/Source/Runtime/Engine/Public/GlobalShader.h
@@ -6,7 +6,6 @@
 
 #pragma once
 
-#include "Engine.h"
 #include "Shader.h"
 #include "SceneManagement.h"
 #include "ShaderParameters.h"
diff --git a/Engine/Source/Runtime/Engine/Public/MaterialCompiler.h b/Engine/Source/Runtime/Engine/Public/MaterialCompiler.h
--- a/Engine/Source/Runtime/Engine/Public/MaterialCompiler.h
+++ b/Engine/Source/Runtime/Engine/Public/MaterialCompiler.h
@@ -6,7 +6,6 @@
 
 #pragma once
 
-#include "Engine.h"
 #include "Materials/MaterialExpressionWorldPosition.h"
 #include "Materials/MaterialExpressionTextureSample.h"
 #include "Materials/MaterialExpressionSpeedTree.h"
diff --git a/Engine/Source/Runtime/Engine/Public/MaterialShader.h b/Engine/Source/Runtime/Engine/Public/MaterialShader.h
--- a/Engine/Source/Runtime/Engine/Public/MaterialShader.h
+++ b/Engine/Source/Runtime/Engine/Public/MaterialShader.h
@@ -6,7 +6,6 @@
 
 #pragma once
 
-#include "Engine.h"
 #include "Shader.h"
 #include "ShaderParameters.h"
 #include "MaterialShared.h"
diff --git a/Engine/Source/Runtime/Engine/Public/MaterialShared.h b/Engine/Source/Runtime/Engine/Public/MaterialShared.h
--- a/Engine/Source/Runtime/Engine/Public/MaterialShared.h
+++ b/Engine/Source/Runtime/Engine/Public/MaterialShared.h
@@ -6,7 +6,6 @@
 
 #pragma once
 
-#include "Engine.h"
 #include "RefCounting.h"
 #include "ShaderCore.h"
 #include "Shader.h"
diff --git a/Engine/Source/Runtime/Engine/Public/MeshMaterialShader.h b/Engine/Source/Runtime/Engine/Public/MeshMaterialShader.h
--- a/Engine/Source/Runtime/Engine/Public/MeshMaterialShader.h
+++ b/Engine/Source/Runtime/Engine/Public/MeshMaterialShader.h
@@ -6,7 +6,6 @@
 
 #pragma once
 
-#include "Engine.h"
 #include "Shader.h"
 #include "MaterialShader.h"
 
diff --git a/Engine/Source/Runtime/Engine/Public/Model.h b/Engine/Source/Runtime/Engine/Public/Model.h
--- a/Engine/Source/Runtime/Engine/Public/Model.h
+++ b/Engine/Source/Runtime/Engine/Public/Model.h
@@ -6,7 +6,6 @@
 
 #pragma once
 
-#include "Engine.h"
 #include "RawIndexBuffer.h"
 #include "LocalVertexFactory.h"
 #include "RenderUtils.h"
diff --git a/Engine/Source/Runtime/Engine/Public/OneColorShader.h b/Engine/Source/Runtime/Engine/Public/OneColorShader.h
--- a/Engine/Source/Runtime/Engine/Public/OneColorShader.h
+++ b/Engine/Source/Runtime/Engine/Public/OneColorShader.h
@@ -3,7 +3,6 @@
 
 #pragma once
 
-#include "Engine.h"
 #include "GlobalShader.h"
 #include "ShaderParameters.h"
 
diff --git a/Engine/Source/Runtime/Engine/Public/ParticleDefinitions.h b/Engine/Source/Runtime/Engine/Public/ParticleDefinitions.h
--- a/Engine/Source/Runtime/Engine/Public/ParticleDefinitions.h
+++ b/Engine/Source/Runtime/Engine/Public/ParticleDefinitions.h
@@ -6,6 +6,5 @@
 
 #pragma once
 
-#include "Engine.h"
 #include "ParticleHelper.h"
 #include "ParticleEmitterInstances.h"
diff --git a/Engine/Source/Runtime/Engine/Public/ParticleEmitterInstances.h b/Engine/Source/Runtime/Engine/Public/ParticleEmitterInstances.h
--- a/Engine/Source/Runtime/Engine/Public/ParticleEmitterInstances.h
+++ b/Engine/Source/Runtime/Engine/Public/ParticleEmitterInstances.h
@@ -7,7 +7,6 @@
 
 #pragma once
 
-#include "Engine.h"
 #include "Distributions.h"
 #include "ParticleHelper.h"
 
diff --git a/Engine/Source/Runtime/Engine/Public/ParticleHelper.h b/Engine/Source/Runtime/Engine/Public/ParticleHelper.h
--- a/Engine/Source/Runtime/Engine/Public/ParticleHelper.h
+++ b/Engine/Source/Runtime/Engine/Public/ParticleHelper.h
@@ -6,7 +6,6 @@
 
 #pragma once
 
-#include "Engine.h"
 #include "ParticleVertexFactory.h"
 #include "ParticleBeamTrailVertexFactory.h"
 #include "MeshParticleVertexFactory.h"
diff --git a/Engine/Source/Runtime/Engine/Public/ParticleVertexFactory.h b/Engine/Source/Runtime/Engine/Public/ParticleVertexFactory.h
--- a/Engine/Source/Runtime/Engine/Public/ParticleVertexFactory.h
+++ b/Engine/Source/Runtime/Engine/Public/ParticleVertexFactory.h
@@ -6,7 +6,6 @@
 
 #pragma once
 
-#include "Engine.h"
 #include "VertexFactory.h"
 #include "UniformBuffer.h"
 
diff --git a/Engine/Source/Runtime/Engine/Public/RawIndexBuffer.h b/Engine/Source/Runtime/Engine/Public/RawIndexBuffer.h
--- a/Engine/Source/Runtime/Engine/Public/RawIndexBuffer.h
+++ b/Engine/Source/Runtime/Engine/Public/RawIndexBuffer.h
@@ -6,7 +6,6 @@
 
 #pragma once
 
-#include "Engine.h"
 #include "RenderResource.h"
 
 class FRawIndexBuffer : public FIndexBuffer
diff --git a/Engine/Source/Runtime/Engine/Public/ScreenRendering.h b/Engine/Source/Runtime/Engine/Public/ScreenRendering.h
--- a/Engine/Source/Runtime/Engine/Public/ScreenRendering.h
+++ b/Engine/Source/Runtime/Engine/Public/ScreenRendering.h
@@ -6,7 +6,6 @@
 
 #pragma once
 
-#include "Engine.h"
 #include "Shader.h"
 #include "ShaderParameters.h"
 #include "GlobalShader.h"
diff --git a/Engine/Source/Runtime/Engine/Public/StaticLighting.h b/Engine/Source/Runtime/Engine/Public/StaticLighting.h
--- a/Engine/Source/Runtime/Engine/Public/StaticLighting.h
+++ b/Engine/Source/Runtime/Engine/Public/StaticLighting.h
@@ -6,7 +6,6 @@
 
 #pragma once
 
-#include "Engine.h"
 #include "VertexFactory.h"
 
 // Forward declarations.
diff --git a/Engine/Source/Runtime/Engine/Public/Tests/AutomationCommon.h b/Engine/Source/Runtime/Engine/Public/Tests/AutomationCommon.h
--- a/Engine/Source/Runtime/Engine/Public/Tests/AutomationCommon.h
+++ b/Engine/Source/Runtime/Engine/Public/Tests/AutomationCommon.h
@@ -2,8 +2,6 @@
 
 #pragma once
 
-#include "Engine.h"
-
 ///////////////////////////////////////////////////////////////////////
 // Common Latent commands which are used across test type. I.e. Engine, Network, etc...
 
diff --git a/Engine/Source/Runtime/Engine/Public/UpdateTextureShaders.h b/Engine/Source/Runtime/Engine/Public/UpdateTextureShaders.h
--- a/Engine/Source/Runtime/Engine/Public/UpdateTextureShaders.h
+++ b/Engine/Source/Runtime/Engine/Public/UpdateTextureShaders.h
@@ -3,7 +3,6 @@
 
 #pragma once
 
-#include "Engine.h"
 #include "GlobalShader.h"
 #include "ShaderParameters.h"
 
