Remove Volume headers from EnginePrivate.h

diff --git a/Engine/Source/Runtime/Engine/Classes/AI/Navigation/NavLinkProxy.h b/Engine/Source/Runtime/Engine/Classes/AI/Navigation/NavLinkProxy.h
--- a/Engine/Source/Runtime/Engine/Classes/AI/Navigation/NavLinkProxy.h
+++ b/Engine/Source/Runtime/Engine/Classes/AI/Navigation/NavLinkProxy.h
@@ -2,6 +2,7 @@
 
 #pragma once
 #include "NavigationTypes.h"
+#include "Ai/Navigation/NavRelevantActorInterface.h"
 #include "NavLinkProxy.generated.h"
 
 DECLARE_DYNAMIC_MULTICAST_DELEGATE_TwoParams( FSmartLinkReachedSignature, class AActor*, MovingActor, const FVector&, DestinationPoint );
diff --git a/Engine/Source/Runtime/Engine/Classes/AI/Navigation/NavMeshBoundsVolume.h b/Engine/Source/Runtime/Engine/Classes/AI/Navigation/NavMeshBoundsVolume.h
--- a/Engine/Source/Runtime/Engine/Classes/AI/Navigation/NavMeshBoundsVolume.h
+++ b/Engine/Source/Runtime/Engine/Classes/AI/Navigation/NavMeshBoundsVolume.h
@@ -6,6 +6,7 @@
  */
 
 #pragma once
+#include "GameFramework/Volume.h"
 #include "NavMeshBoundsVolume.generated.h"
 
 UCLASS(MinimalAPI)
diff --git a/Engine/Source/Runtime/Engine/Classes/AI/Navigation/NavModifierVolume.h b/Engine/Source/Runtime/Engine/Classes/AI/Navigation/NavModifierVolume.h
--- a/Engine/Source/Runtime/Engine/Classes/AI/Navigation/NavModifierVolume.h
+++ b/Engine/Source/Runtime/Engine/Classes/AI/Navigation/NavModifierVolume.h
@@ -1,6 +1,7 @@
 // Copyright 1998-2014 Epic Games, Inc. All Rights Reserved.
 
 #pragma once 
+#include "GameFramework/Volume.h"
 #include "NavModifierVolume.generated.h"
 
 /** 
diff --git a/Engine/Source/Runtime/Engine/Classes/Components/PostProcessComponent.h b/Engine/Source/Runtime/Engine/Classes/Components/PostProcessComponent.h
--- a/Engine/Source/Runtime/Engine/Classes/Components/PostProcessComponent.h
+++ b/Engine/Source/Runtime/Engine/Classes/Components/PostProcessComponent.h
@@ -1,8 +1,8 @@
 // Copyright 1998-2013 Epic Games, Inc. All Rights Reserved.
 
 #pragma once
+#include "Interfaces/Interface_PostProcessVolume.h"
 #include "PostProcessComponent.generated.h"
-
 /**
  *  PostProcessComponent. Enables Post process controls for blueprints.
  *	Will use a parent UShapeComponent to provide volume data if available.
diff --git a/Engine/Source/Runtime/Engine/Classes/Engine/CullDistanceVolume.h b/Engine/Source/Runtime/Engine/Classes/Engine/CullDistanceVolume.h
--- a/Engine/Source/Runtime/Engine/Classes/Engine/CullDistanceVolume.h
+++ b/Engine/Source/Runtime/Engine/Classes/Engine/CullDistanceVolume.h
@@ -2,6 +2,7 @@
 
 
 #pragma once
+#include "GameFramework/Volume.h"
 #include "CullDistanceVolume.generated.h"
 
 /**
diff --git a/Engine/Source/Runtime/Engine/Classes/Engine/LevelStreamingVolume.h b/Engine/Source/Runtime/Engine/Classes/Engine/LevelStreamingVolume.h
--- a/Engine/Source/Runtime/Engine/Classes/Engine/LevelStreamingVolume.h
+++ b/Engine/Source/Runtime/Engine/Classes/Engine/LevelStreamingVolume.h
@@ -6,6 +6,7 @@
  */
 
 #pragma once
+#include "GameFramework/Volume.h"
 #include "LevelStreamingVolume.generated.h"
 
 /** Enum for different usage cases of level streaming volumes. */
diff --git a/Engine/Source/Runtime/Engine/Classes/Engine/PostProcessVolume.h b/Engine/Source/Runtime/Engine/Classes/Engine/PostProcessVolume.h
--- a/Engine/Source/Runtime/Engine/Classes/Engine/PostProcessVolume.h
+++ b/Engine/Source/Runtime/Engine/Classes/Engine/PostProcessVolume.h
@@ -6,6 +6,7 @@
 //=============================================================================
 
 #pragma once
+#include "GameFramework/Volume.h"
 #include "PostProcessVolume.generated.h"
 
 	// for FPostprocessSettings
diff --git a/Engine/Source/Runtime/Engine/Classes/Engine/StaticMeshActor.h b/Engine/Source/Runtime/Engine/Classes/Engine/StaticMeshActor.h
--- a/Engine/Source/Runtime/Engine/Classes/Engine/StaticMeshActor.h
+++ b/Engine/Source/Runtime/Engine/Classes/Engine/StaticMeshActor.h
@@ -1,6 +1,7 @@
 // Copyright 1998-2014 Epic Games, Inc. All Rights Reserved.
 
 #pragma once
+#include "Ai/Navigation/NavRelevantActorInterface.h"
 #include "StaticMeshActor.generated.h"
 
 /**
diff --git a/Engine/Source/Runtime/Engine/Classes/Engine/TriggerVolume.h b/Engine/Source/Runtime/Engine/Classes/Engine/TriggerVolume.h
--- a/Engine/Source/Runtime/Engine/Classes/Engine/TriggerVolume.h
+++ b/Engine/Source/Runtime/Engine/Classes/Engine/TriggerVolume.h
@@ -2,6 +2,7 @@
 
 
 #pragma once
+#include "GameFramework/Volume.h"
 #include "TriggerVolume.generated.h"
 
 UCLASS()
diff --git a/Engine/Source/Runtime/Engine/Classes/GameFramework/DefaultPhysicsVolume.h b/Engine/Source/Runtime/Engine/Classes/GameFramework/DefaultPhysicsVolume.h
--- a/Engine/Source/Runtime/Engine/Classes/GameFramework/DefaultPhysicsVolume.h
+++ b/Engine/Source/Runtime/Engine/Classes/GameFramework/DefaultPhysicsVolume.h
@@ -6,6 +6,7 @@
 //=============================================================================
 
 #pragma once
+#include "GameFramework/PhysicsVolume.h"
 #include "DefaultPhysicsVolume.generated.h"
 
 UCLASS(notplaceable, transient, MinimalAPI)
diff --git a/Engine/Source/Runtime/Engine/Classes/GameFramework/KillZVolume.h b/Engine/Source/Runtime/Engine/Classes/GameFramework/KillZVolume.h
--- a/Engine/Source/Runtime/Engine/Classes/GameFramework/KillZVolume.h
+++ b/Engine/Source/Runtime/Engine/Classes/GameFramework/KillZVolume.h
@@ -1,6 +1,7 @@
 // Copyright 1998-2014 Epic Games, Inc. All Rights Reserved.
 
 #pragma once
+#include "GameFramework/PhysicsVolume.h"
 #include "KillZVolume.generated.h"
 
 UCLASS()
diff --git a/Engine/Source/Runtime/Engine/Classes/GameFramework/PainCausingVolume.h b/Engine/Source/Runtime/Engine/Classes/GameFramework/PainCausingVolume.h
--- a/Engine/Source/Runtime/Engine/Classes/GameFramework/PainCausingVolume.h
+++ b/Engine/Source/Runtime/Engine/Classes/GameFramework/PainCausingVolume.h
@@ -6,6 +6,7 @@
 //=============================================================================
 
 #pragma once
+#include "GameFramework/PhysicsVolume.h"
 #include "PainCausingVolume.generated.h"
 
 
diff --git a/Engine/Source/Runtime/Engine/Classes/GameFramework/PhysicsVolume.h b/Engine/Source/Runtime/Engine/Classes/GameFramework/PhysicsVolume.h
--- a/Engine/Source/Runtime/Engine/Classes/GameFramework/PhysicsVolume.h
+++ b/Engine/Source/Runtime/Engine/Classes/GameFramework/PhysicsVolume.h
@@ -1,6 +1,7 @@
 // Copyright 1998-2014 Epic Games, Inc. All Rights Reserved.
 
 #pragma once
+#include "GameFramework/Volume.h"
 #include "PhysicsVolume.generated.h"
 
 /**
diff --git a/Engine/Source/Runtime/Engine/Classes/GameFramework/Volume.h b/Engine/Source/Runtime/Engine/Classes/GameFramework/Volume.h
--- a/Engine/Source/Runtime/Engine/Classes/GameFramework/Volume.h
+++ b/Engine/Source/Runtime/Engine/Classes/GameFramework/Volume.h
@@ -5,6 +5,7 @@
 //=============================================================================
 
 #pragma once
+#include "Engine/Brush.h"
 #include "Volume.generated.h"
 
 UCLASS(MinimalAPI, showcategories=Collision, hidecategories=(Brush, Physics), abstract)
diff --git a/Engine/Source/Runtime/Engine/Classes/GameFramework/WorldSettings.h b/Engine/Source/Runtime/Engine/Classes/GameFramework/WorldSettings.h
--- a/Engine/Source/Runtime/Engine/Classes/GameFramework/WorldSettings.h
+++ b/Engine/Source/Runtime/Engine/Classes/GameFramework/WorldSettings.h
@@ -1,6 +1,7 @@
 // Copyright 1998-2014 Epic Games, Inc. All Rights Reserved.
 
 #pragma once
+#include "Sound/ReverbVolume.h"
 #include "WorldSettings.generated.h"
 
 UENUM()
diff --git a/Engine/Source/Runtime/Engine/Classes/Sound/ReverbVolume.h b/Engine/Source/Runtime/Engine/Classes/Sound/ReverbVolume.h
--- a/Engine/Source/Runtime/Engine/Classes/Sound/ReverbVolume.h
+++ b/Engine/Source/Runtime/Engine/Classes/Sound/ReverbVolume.h
@@ -6,6 +6,7 @@
  */
 
 #pragma once
+#include "GameFramework/Volume.h"
 #include "ReverbVolume.generated.h"
 
 /**
diff --git a/Engine/Source/Runtime/Engine/Private/AI/Navigation/RecastNavMeshGenerator.cpp b/Engine/Source/Runtime/Engine/Private/AI/Navigation/RecastNavMeshGenerator.cpp
--- a/Engine/Source/Runtime/Engine/Private/AI/Navigation/RecastNavMeshGenerator.cpp
+++ b/Engine/Source/Runtime/Engine/Private/AI/Navigation/RecastNavMeshGenerator.cpp
@@ -11,6 +11,7 @@
 #include "RecastNavMeshGenerator.h"
 #include "PImplRecastNavMesh.h"
 #include "SurfaceIterators.h"
+#include "Ai/Navigation/NavMeshBoundsVolume.h"
 
 // recast includes
 #include "Recast.h"
diff --git a/Engine/Source/Runtime/Engine/Private/BlockingVolume.cpp b/Engine/Source/Runtime/Engine/Private/BlockingVolume.cpp
--- a/Engine/Source/Runtime/Engine/Private/BlockingVolume.cpp
+++ b/Engine/Source/Runtime/Engine/Private/BlockingVolume.cpp
@@ -1,6 +1,7 @@
 // Copyright 1998-2014 Epic Games, Inc. All Rights Reserved.
 
 #include "EnginePrivate.h"
+#include "Engine/BlockingVolume.h"
 
 static FName InvisibleWall_NAME(TEXT("InvisibleWall"));
 
diff --git a/Engine/Source/Runtime/Engine/Private/CameraBlockingVolume.cpp b/Engine/Source/Runtime/Engine/Private/CameraBlockingVolume.cpp
--- a/Engine/Source/Runtime/Engine/Private/CameraBlockingVolume.cpp
+++ b/Engine/Source/Runtime/Engine/Private/CameraBlockingVolume.cpp
@@ -1,6 +1,7 @@
 // Copyright 1998-2014 Epic Games, Inc. All Rights Reserved.
 
 #include "EnginePrivate.h"
+#include "GameFramework/CameraBlockingVolume.h"
 
 ACameraBlockingVolume::ACameraBlockingVolume(const class FPostConstructInitializeProperties& PCIP)
 	: Super(PCIP)
diff --git a/Engine/Source/Runtime/Engine/Private/CharacterMovementComponent.cpp b/Engine/Source/Runtime/Engine/Private/CharacterMovementComponent.cpp
--- a/Engine/Source/Runtime/Engine/Private/CharacterMovementComponent.cpp
+++ b/Engine/Source/Runtime/Engine/Private/CharacterMovementComponent.cpp
@@ -6,6 +6,7 @@
 =============================================================================*/
 
 #include "EnginePrivate.h"
+#include "GameFramework/PhysicsVolume.h"
 #include "GameFramework/GameNetworkManager.h"
 // @todo this is here only due to circular dependency to AIModule. To be removed
 #include "Navigation/PathFollowingComponent.h"
diff --git a/Engine/Source/Runtime/Engine/Private/CullDistanceVolume.cpp b/Engine/Source/Runtime/Engine/Private/CullDistanceVolume.cpp
--- a/Engine/Source/Runtime/Engine/Private/CullDistanceVolume.cpp
+++ b/Engine/Source/Runtime/Engine/Private/CullDistanceVolume.cpp
@@ -1,6 +1,7 @@
 // Copyright 1998-2014 Epic Games, Inc. All Rights Reserved.
 
 #include "EnginePrivate.h"
+#include "Engine/CullDistanceVolume.h"
 
 ACullDistanceVolume::ACullDistanceVolume(const class FPostConstructInitializeProperties& PCIP)
 	: Super(PCIP)
diff --git a/Engine/Source/Runtime/Engine/Private/DefaultPhysicsVolume.cpp b/Engine/Source/Runtime/Engine/Private/DefaultPhysicsVolume.cpp
--- a/Engine/Source/Runtime/Engine/Private/DefaultPhysicsVolume.cpp
+++ b/Engine/Source/Runtime/Engine/Private/DefaultPhysicsVolume.cpp
@@ -1,6 +1,7 @@
 // Copyright 1998-2014 Epic Games, Inc. All Rights Reserved.
 
 #include "EnginePrivate.h"
+#include "GameFramework/DefaultPhysicsVolume.h"
 
 DEFINE_LOG_CATEGORY_STATIC(LogVolume, Log, All);
 
diff --git a/Engine/Source/Runtime/Engine/Private/EnginePrivate.h b/Engine/Source/Runtime/Engine/Private/EnginePrivate.h
--- a/Engine/Source/Runtime/Engine/Private/EnginePrivate.h
+++ b/Engine/Source/Runtime/Engine/Private/EnginePrivate.h
@@ -23,22 +23,6 @@
 #include "GameFramework/Actor.h"
 #include "Sound/AmbientSound.h"
 #include "Engine/Brush.h"
-#include "GameFramework/Volume.h"
-#include "Engine/BlockingVolume.h"
-#include "GameFramework/CameraBlockingVolume.h"
-#include "Engine/CullDistanceVolume.h"
-#include "Engine/LevelStreamingVolume.h"
-#include "AI/Navigation/NavMeshBoundsVolume.h"
-#include "AI/Navigation/NavRelevantActorInterface.h"
-#include "AI/Navigation/NavModifierVolume.h"
-#include "GameFramework/PhysicsVolume.h"
-#include "GameFramework/DefaultPhysicsVolume.h"
-#include "GameFramework/KillZVolume.h"
-#include "GameFramework/PainCausingVolume.h"
-#include "Interfaces/Interface_PostProcessVolume.h"
-#include "Engine/PostProcessVolume.h"
-#include "Sound/ReverbVolume.h"
-#include "Engine/TriggerVolume.h"
 #include "Camera/CameraActor.h"
 #include "AI/Navigation/NavAgentInterface.h"
 #include "GameFramework/Controller.h"
@@ -130,7 +114,6 @@
 #include "Engine/TriggerBox.h"
 #include "Engine/TriggerCapsule.h"
 #include "Engine/TriggerSphere.h"
-#include "VectorField/VectorFieldVolume.h"
 #include "Components/ApplicationLifecycleComponent.h"
 #include "Engine/DataAsset.h"
 #include "Debug/GameplayDebuggingControllerComponent.h"
diff --git a/Engine/Source/Runtime/Engine/Private/ErrorChecking.cpp b/Engine/Source/Runtime/Engine/Private/ErrorChecking.cpp
--- a/Engine/Source/Runtime/Engine/Private/ErrorChecking.cpp
+++ b/Engine/Source/Runtime/Engine/Private/ErrorChecking.cpp
@@ -11,6 +11,9 @@
 #include "UObjectToken.h"
 #include "MapErrors.h"
 #include "Model.h"
+#include "GameFramework/PainCausingVolume.h"
+#include "GameFramework/DefaultPhysicsVolume.h"
+#include "Engine/LevelStreamingVolume.h"
 
 #define LOCTEXT_NAMESPACE "ErrorChecking"
 
diff --git a/Engine/Source/Runtime/Engine/Private/KillZVolume.cpp b/Engine/Source/Runtime/Engine/Private/KillZVolume.cpp
--- a/Engine/Source/Runtime/Engine/Private/KillZVolume.cpp
+++ b/Engine/Source/Runtime/Engine/Private/KillZVolume.cpp
@@ -1,6 +1,7 @@
 // Copyright 1998-2014 Epic Games, Inc. All Rights Reserved.
 
 #include "EnginePrivate.h"
+#include "GameFramework/KillZVolume.h"
 
 AKillZVolume::AKillZVolume(const class FPostConstructInitializeProperties& PCIP)
 	: Super(PCIP)
diff --git a/Engine/Source/Runtime/Engine/Private/LevelStreaming.cpp b/Engine/Source/Runtime/Engine/Private/LevelStreaming.cpp
--- a/Engine/Source/Runtime/Engine/Private/LevelStreaming.cpp
+++ b/Engine/Source/Runtime/Engine/Private/LevelStreaming.cpp
@@ -1,6 +1,7 @@
 // Copyright 1998-2014 Epic Games, Inc. All Rights Reserved.
 
 #include "EnginePrivate.h"
+#include "Engine/LevelStreamingVolume.h"
 #include "LevelUtils.h"
 #if WITH_EDITOR
 	#include "Slate.h"
diff --git a/Engine/Source/Runtime/Engine/Private/LevelStreamingVolume.cpp b/Engine/Source/Runtime/Engine/Private/LevelStreamingVolume.cpp
--- a/Engine/Source/Runtime/Engine/Private/LevelStreamingVolume.cpp
+++ b/Engine/Source/Runtime/Engine/Private/LevelStreamingVolume.cpp
@@ -1,6 +1,7 @@
 // Copyright 1998-2014 Epic Games, Inc. All Rights Reserved.
 
 #include "EnginePrivate.h"
+#include "Engine/LevelStreamingVolume.h"
 
 ALevelStreamingVolume::ALevelStreamingVolume(const class FPostConstructInitializeProperties& PCIP)
 	: Super(PCIP)
diff --git a/Engine/Source/Runtime/Engine/Private/LevelTick.cpp b/Engine/Source/Runtime/Engine/Private/LevelTick.cpp
--- a/Engine/Source/Runtime/Engine/Private/LevelTick.cpp
+++ b/Engine/Source/Runtime/Engine/Private/LevelTick.cpp
@@ -5,6 +5,7 @@
 =============================================================================*/
 
 #include "EnginePrivate.h"
+#include "Engine/LevelStreamingVolume.h"
 #include "Net/UnrealNetwork.h"
 #include "Collision.h"
 #include "PhysicsPublic.h"
diff --git a/Engine/Source/Runtime/Engine/Private/MovementComponent.cpp b/Engine/Source/Runtime/Engine/Private/MovementComponent.cpp
--- a/Engine/Source/Runtime/Engine/Private/MovementComponent.cpp
+++ b/Engine/Source/Runtime/Engine/Private/MovementComponent.cpp
@@ -2,6 +2,7 @@
 
 
 #include "EnginePrivate.h"
+#include "GameFramework/PhysicsVolume.h"
 #include "MessageLog.h"
 
 #define LOCTEXT_NAMESPACE "MovementComponent"
diff --git a/Engine/Source/Runtime/Engine/Private/NavMeshBoundsVolume.cpp b/Engine/Source/Runtime/Engine/Private/NavMeshBoundsVolume.cpp
--- a/Engine/Source/Runtime/Engine/Private/NavMeshBoundsVolume.cpp
+++ b/Engine/Source/Runtime/Engine/Private/NavMeshBoundsVolume.cpp
@@ -1,6 +1,7 @@
 // Copyright 1998-2014 Epic Games, Inc. All Rights Reserved.
 
 #include "EnginePrivate.h"
+#include "AI/Navigation/NavMeshBoundsVolume.h"
 
 ANavMeshBoundsVolume::ANavMeshBoundsVolume(const class FPostConstructInitializeProperties& PCIP)
 	: Super(PCIP)
diff --git a/Engine/Source/Runtime/Engine/Private/NavigationObjectBase.cpp b/Engine/Source/Runtime/Engine/Private/NavigationObjectBase.cpp
--- a/Engine/Source/Runtime/Engine/Private/NavigationObjectBase.cpp
+++ b/Engine/Source/Runtime/Engine/Private/NavigationObjectBase.cpp
@@ -1,6 +1,7 @@
 // Copyright 1998-2014 Epic Games, Inc. All Rights Reserved.
 
 #include "EnginePrivate.h"
+#include "GameFramework/PhysicsVolume.h"
 
 DEFINE_LOG_CATEGORY_STATIC(LogNavigationPoint, Log, All);
 
diff --git a/Engine/Source/Runtime/Engine/Private/PainCausingVolume.cpp b/Engine/Source/Runtime/Engine/Private/PainCausingVolume.cpp
--- a/Engine/Source/Runtime/Engine/Private/PainCausingVolume.cpp
+++ b/Engine/Source/Runtime/Engine/Private/PainCausingVolume.cpp
@@ -1,6 +1,7 @@
 // Copyright 1998-2014 Epic Games, Inc. All Rights Reserved.
 
 #include "EnginePrivate.h"
+#include "GameFramework/PainCausingVolume.h"
 
 APainCausingVolume::APainCausingVolume(const class FPostConstructInitializeProperties& PCIP)
 	: Super(PCIP)
diff --git a/Engine/Source/Runtime/Engine/Private/PhysicsVolume.cpp b/Engine/Source/Runtime/Engine/Private/PhysicsVolume.cpp
--- a/Engine/Source/Runtime/Engine/Private/PhysicsVolume.cpp
+++ b/Engine/Source/Runtime/Engine/Private/PhysicsVolume.cpp
@@ -1,6 +1,7 @@
 // Copyright 1998-2014 Epic Games, Inc. All Rights Reserved.
 
 #include "EnginePrivate.h"
+#include "GameFramework/PhysicsVolume.h"
 
 APhysicsVolume::APhysicsVolume(const class FPostConstructInitializeProperties& PCIP)
 	: Super(PCIP)
diff --git a/Engine/Source/Runtime/Engine/Private/PostProcessVolume.cpp b/Engine/Source/Runtime/Engine/Private/PostProcessVolume.cpp
--- a/Engine/Source/Runtime/Engine/Private/PostProcessVolume.cpp
+++ b/Engine/Source/Runtime/Engine/Private/PostProcessVolume.cpp
@@ -1,6 +1,7 @@
 // Copyright 1998-2014 Epic Games, Inc. All Rights Reserved.
 
 #include "EnginePrivate.h"
+#include "Engine/PostProcessVolume.h"
 
 APostProcessVolume::APostProcessVolume(const class FPostConstructInitializeProperties& PCIP)
 	: Super(PCIP)
diff --git a/Engine/Source/Runtime/Engine/Private/PrimitiveComponent.cpp b/Engine/Source/Runtime/Engine/Private/PrimitiveComponent.cpp
--- a/Engine/Source/Runtime/Engine/Private/PrimitiveComponent.cpp
+++ b/Engine/Source/Runtime/Engine/Private/PrimitiveComponent.cpp
@@ -5,12 +5,13 @@
 =============================================================================*/
 
 #include "EnginePrivate.h"
+#include "GameFramework/PhysicsVolume.h"
 #include "PhysicsPublic.h"
 #include "LevelUtils.h"
 #if WITH_PHYSX
 #include "PhysicsEngine/PhysXSupport.h"
 #include "Collision/PhysXCollision.h"
-#endif
+#endif // WITH_PHYSX
 #include "Collision/CollisionDebugDrawing.h"
 #include "MessageLog.h"
 #include "UObjectToken.h"
diff --git a/Engine/Source/Runtime/Engine/Private/SceneComponent.cpp b/Engine/Source/Runtime/Engine/Private/SceneComponent.cpp
--- a/Engine/Source/Runtime/Engine/Private/SceneComponent.cpp
+++ b/Engine/Source/Runtime/Engine/Private/SceneComponent.cpp
@@ -8,6 +8,7 @@
 #include "PhysicsPublic.h"
 #include "MessageLog.h"
 #include "Net/UnrealNetwork.h"
+#include "GameFramework/PhysicsVolume.h"
 
 #define LOCTEXT_NAMESPACE "SceneComponent"
 
diff --git a/Engine/Source/Runtime/Engine/Private/TriggerVolume.cpp b/Engine/Source/Runtime/Engine/Private/TriggerVolume.cpp
--- a/Engine/Source/Runtime/Engine/Private/TriggerVolume.cpp
+++ b/Engine/Source/Runtime/Engine/Private/TriggerVolume.cpp
@@ -1,6 +1,7 @@
 // Copyright 1998-2014 Epic Games, Inc. All Rights Reserved.
 
 #include "EnginePrivate.h"
+#include "Engine/TriggerVolume.h"
 
 ATriggerVolume::ATriggerVolume(const class FPostConstructInitializeProperties& PCIP)
 	: Super(PCIP)
diff --git a/Engine/Source/Runtime/Engine/Private/VectorFieldVolume.cpp b/Engine/Source/Runtime/Engine/Private/VectorFieldVolume.cpp
--- a/Engine/Source/Runtime/Engine/Private/VectorFieldVolume.cpp
+++ b/Engine/Source/Runtime/Engine/Private/VectorFieldVolume.cpp
@@ -1,6 +1,7 @@
 // Copyright 1998-2014 Epic Games, Inc. All Rights Reserved.
 
 #include "EnginePrivate.h"
+#include "VectorField/VectorFieldVolume.h"
 
 AVectorFieldVolume::AVectorFieldVolume(const class FPostConstructInitializeProperties& PCIP)
 	: Super(PCIP)
