Anim Trails can now be used via Blueprints.

TTP 333673

Simply create a PSC and call BeginTrails/EndTrails on it.

Refactored things a bit to be nicer too. E.g Trails emitters no longer know or care about who own's them. Identifying a notify with a trail is done at the component level via the components Outer.

diff --git a/Engine/Source/Runtime/Engine/Classes/Animation/AnimNotifies/AnimNotifyState_Trail.h b/Engine/Source/Runtime/Engine/Classes/Animation/AnimNotifies/AnimNotifyState_Trail.h
--- a/Engine/Source/Runtime/Engine/Classes/Animation/AnimNotifies/AnimNotifyState_Trail.h
+++ b/Engine/Source/Runtime/Engine/Classes/Animation/AnimNotifies/AnimNotifyState_Trail.h
@@ -5,15 +5,6 @@
 
 DECLARE_LOG_CATEGORY_EXTERN(LogAnimTrails, Log, All);
 
-/** The way that the width scale property affects the trail */
-UENUM()
-enum ETrailWidthMode
-{
-	ETrailWidthMode_FromCentre UMETA(DisplayName = "From Centre"),
-	ETrailWidthMode_FromFirst UMETA(DisplayName = "From First Socket"),
-	ETrailWidthMode_FromSecond UMETA(DisplayName = "From Second Socket"),
-};
-
 UCLASS(editinlinenew, Blueprintable, const, hidecategories = Object, collapsecategories, MinimalAPI, meta = (ShowHiddenSelfPins, DisplayName = "Trail"))
 class UAnimNotifyState_Trail : public UAnimNotifyState
 {
diff --git a/Engine/Source/Runtime/Engine/Classes/Particles/ParticleSystemComponent.h b/Engine/Source/Runtime/Engine/Classes/Particles/ParticleSystemComponent.h
--- a/Engine/Source/Runtime/Engine/Classes/Particles/ParticleSystemComponent.h
+++ b/Engine/Source/Runtime/Engine/Classes/Particles/ParticleSystemComponent.h
@@ -20,6 +20,15 @@ struct FDynamicEmitterDataBase;
 struct FDynamicEmitterReplayDataBase;
 struct FParticleAnimTrailEmitterInstance;
 
+/** Controls the way that the width scale property affects anim trails */
+UENUM()
+enum ETrailWidthMode
+{
+	ETrailWidthMode_FromCentre UMETA(DisplayName = "From Centre"),
+	ETrailWidthMode_FromFirst UMETA(DisplayName = "From First Socket"),
+	ETrailWidthMode_FromSecond UMETA(DisplayName = "From Second Socket"),
+};
+
 /** Enum for specifying type of a name instance parameter. */
 UENUM()
 enum EParticleSysParamType
@@ -627,11 +636,37 @@ public:
 	
 	/**
 	* Fills the passed array with all trail emitters associated with a particular AnimNotifyState.
-	* @param InAnimNotifyState	The AnimNotifyState collecting the trails.
-	* @param OutTrailEmitters	The array to fill with pointers to
-	* @param bIncludeUnassociated If true, trails that have not yet been associated with an AnimNotifyState are included in OutTrailEmitters.
+	* @param OutTrailEmitters	The array to fill with pointers to the trail emitters.
+	*/
+	virtual void GetTrailEmitters(TArray< struct FParticleAnimTrailEmitterInstance* >& OutTrailEmitters);
+	
+	/**
+	* Begins all trail emitters in this component.
+	*
+	* @param	InFirstSocketName	The name of the first socket for the trail.
+	* @param	InSecondSocketName	The name of the second socket for the trail.
+	* @param	InWidthMode			How the width value is applied to the trail.
+	* @param	InWidth				The width of the trail.
+	*/
+	UFUNCTION(BlueprintCallable, Category = "Effects|Particles|Trails")
+	void BeginTrails(FName InFirstSocketName, FName InSecondSocketName, ETrailWidthMode InWidthMode, float InWidth);
+
+	/**
+	* Ends all trail emitters in this component.
+	*/
+	UFUNCTION(BlueprintCallable, Category = "Effects|Particles|Trails")
+	void EndTrails();
+
+	/**
+	* Sets the defining data for all trails in this component.
+	*
+	* @param	InFirstSocketName	The name of the first socket for the trail.
+	* @param	InSecondSocketName	The name of the second socket for the trail.
+	* @param	InWidthMode			How the width value is applied to the trail.
+	* @param	InWidth				The width of the trail.
 	*/
-	virtual void GetTrailEmitters(UAnimNotifyState* InAnimNotifyState, TArray< struct FParticleAnimTrailEmitterInstance* >& OutTrailEmitters, bool bIncludeUnassociated = false);
+	UFUNCTION(BlueprintCallable, Category = "Effects|Particles|Trails")
+	void SetTrailSourceData(FName InFirstSocketName, FName InSecondSocketName, ETrailWidthMode InWidthMode, float InWidth);
 
 public:
 	TArray<struct FParticleEmitterInstance*> EmitterInstances;
diff --git a/Engine/Source/Runtime/Engine/Private/Animation/AnimNotifyState_Trail.cpp b/Engine/Source/Runtime/Engine/Private/Animation/AnimNotifyState_Trail.cpp
--- a/Engine/Source/Runtime/Engine/Private/Animation/AnimNotifyState_Trail.cpp
+++ b/Engine/Source/Runtime/Engine/Private/Animation/AnimNotifyState_Trail.cpp
@@ -51,10 +51,12 @@ void UAnimNotifyState_Trail::NotifyBegin(class USkeletalMeshComponent * MeshComp
 	bool bFoundExistingTrail = false;
 	for (USceneComponent* Comp : Children)
 	{
-		if (UParticleSystemComponent* ParticleComp = Cast<UParticleSystemComponent>(Comp))
+		UParticleSystemComponent* ParticleComp = Cast<UParticleSystemComponent>(Comp);
+		if (ParticleComp && ParticleComp->GetOuter() == this)
 		{
+			bFoundExistingTrail = true;
 			TArray< FParticleAnimTrailEmitterInstance* > TrailEmitters;
-			ParticleComp->GetTrailEmitters(this, TrailEmitters);
+			ParticleComp->GetTrailEmitters(TrailEmitters);
 
 			//If there are any trails, ensure the template hasn't been changed. Also destroy the component if there are erros.
 			if (TrailEmitters.Num() > 0 && (bError || PSTemplate != ParticleComp->Template))
@@ -66,8 +68,7 @@ void UAnimNotifyState_Trail::NotifyBegin(class USkeletalMeshComponent * MeshComp
 			{
 				for (FParticleAnimTrailEmitterInstance* Trail : TrailEmitters)
 				{
-					bFoundExistingTrail = true;
-					Trail->BeginTrail(this);
+					Trail->BeginTrail();
 					Trail->SetTrailSourceData(FirstSocketName, SecondSocketName, WidthScaleMode, Width);
 
 #if WITH_EDITORONLY_DATA
@@ -80,20 +81,32 @@ void UAnimNotifyState_Trail::NotifyBegin(class USkeletalMeshComponent * MeshComp
 
 	if (!bFoundExistingTrail && !bError)
 	{
-		if (UParticleSystemComponent* NewParticleComp = UGameplayStatics::SpawnEmitterAttached(PSTemplate, MeshComp))
-		{
-			TArray< FParticleAnimTrailEmitterInstance* > TrailEmitters;
-			NewParticleComp->GetTrailEmitters(this, TrailEmitters, true);
+		//Spawn a new component from PSTemplate. This notify is made the outer so that the component can be identified later.
+		UParticleSystemComponent* NewParticleComp = ConstructObject<UParticleSystemComponent>(UParticleSystemComponent::StaticClass(), this);
+		NewParticleComp->bAutoDestroy = true;
+		NewParticleComp->SecondsBeforeInactive = 0.0f;
+		NewParticleComp->bAutoActivate = false;
+		NewParticleComp->SetTemplate(PSTemplate);
+		NewParticleComp->bOverrideLODMethod = false;
 
-			for (FParticleAnimTrailEmitterInstance* Trail : TrailEmitters)
-			{
-				Trail->BeginTrail(this);
-				Trail->SetTrailSourceData(FirstSocketName, SecondSocketName, WidthScaleMode, Width);
+		NewParticleComp->RegisterComponentWithWorld(MeshComp->GetWorld());
+		NewParticleComp->AttachTo(MeshComp, NAME_None);
+		//NewParticleComp->SetWorldLocationAndRotation(Location, Rotation);
+
+		NewParticleComp->SetRelativeScale3D(FVector(1.f));
+		NewParticleComp->ActivateSystem(true);
+
+		TArray< FParticleAnimTrailEmitterInstance* > TrailEmitters;
+		NewParticleComp->GetTrailEmitters(TrailEmitters);
+
+		for (FParticleAnimTrailEmitterInstance* Trail : TrailEmitters)
+		{
+			Trail->BeginTrail();
+			Trail->SetTrailSourceData(FirstSocketName, SecondSocketName, WidthScaleMode, Width);
 
 #if WITH_EDITORONLY_DATA
-				Trail->SetTrailDebugData(bRenderGeometry, bRenderSpawnPoints, bRenderTessellation, bRenderTangents);
+			Trail->SetTrailDebugData(bRenderGeometry, bRenderSpawnPoints, bRenderTessellation, bRenderTangents);
 #endif
-			}
 		}
 	}
 
@@ -117,10 +130,11 @@ void UAnimNotifyState_Trail::NotifyTick(class USkeletalMeshComponent * MeshComp,
 	bool bFoundExistingTrail = false;
 	for (USceneComponent* Comp : Children)
 	{
-		if (UParticleSystemComponent* ParticleComp = Cast<UParticleSystemComponent>(Comp))
+		UParticleSystemComponent* ParticleComp = Cast<UParticleSystemComponent>(Comp);
+		if ( ParticleComp && ParticleComp->GetOuter() == this )
 		{
 			TArray< FParticleAnimTrailEmitterInstance* > TrailEmitters;
-			ParticleComp->GetTrailEmitters(this, TrailEmitters);
+			ParticleComp->GetTrailEmitters(TrailEmitters);
 			if (bError && TrailEmitters.Num() > 0)
 			{
 				Comp->DestroyComponent();
@@ -150,14 +164,10 @@ void UAnimNotifyState_Trail::NotifyEnd(class USkeletalMeshComponent * MeshComp,
 	bool bFoundExistingTrail = false;
 	for (USceneComponent* Comp : Children)
 	{
-		if (UParticleSystemComponent* ParticleComp = Cast<UParticleSystemComponent>(Comp))
+		UParticleSystemComponent* ParticleComp = Cast<UParticleSystemComponent>(Comp);
+		if (ParticleComp && ParticleComp->GetOuter() == this)
 		{
-			TArray< FParticleAnimTrailEmitterInstance* > TrailEmitters;
-			ParticleComp->GetTrailEmitters(this, TrailEmitters);
-			for (FParticleAnimTrailEmitterInstance* Trail : TrailEmitters)
-			{
-				Trail->EndTrail();
-			}
+			ParticleComp->EndTrails();
 		}
 	}
 
diff --git a/Engine/Source/Runtime/Engine/Private/Particles/ParticleComponents.cpp b/Engine/Source/Runtime/Engine/Private/Particles/ParticleComponents.cpp
--- a/Engine/Source/Runtime/Engine/Private/Particles/ParticleComponents.cpp
+++ b/Engine/Source/Runtime/Engine/Private/Particles/ParticleComponents.cpp
@@ -4763,19 +4763,47 @@ int32 UParticleSystemComponent::GetNumActiveParticles() const
 	return NumParticles;
 }
 
-void UParticleSystemComponent::GetTrailEmitters(UAnimNotifyState* InAnimNotifyState, TArray< FParticleAnimTrailEmitterInstance* >& OutTrailEmitters, bool bIncludeUnassociated)
+void UParticleSystemComponent::GetTrailEmitters(TArray< FParticleAnimTrailEmitterInstance* >& OutTrailEmitters)
 {
-	int32 NumEmitters = EmitterInstances.Num();
-	for (int32 i = 0; i < NumEmitters; ++i)
+	for (FParticleEmitterInstance* Inst : EmitterInstances)
 	{
-		FParticleEmitterInstance* Instance = EmitterInstances[i];
-		if (Instance && Instance->IsTrailEmitter())
+		if (Inst && Inst->IsTrailEmitter())
 		{
-			FParticleAnimTrailEmitterInstance* TrailInstance = (FParticleAnimTrailEmitterInstance*)(Instance);
-			if ((bIncludeUnassociated && TrailInstance->AnimNotifyState == NULL) || TrailInstance->AnimNotifyState == InAnimNotifyState)
-			{
-				OutTrailEmitters.Add(TrailInstance);
-			}
+			OutTrailEmitters.Add((FParticleAnimTrailEmitterInstance*)(Inst));
+		}
+	}
+}
+
+void UParticleSystemComponent::BeginTrails(FName InFirstSocketName, FName InSecondSocketName, ETrailWidthMode InWidthMode, float InWidth)
+{
+	for (FParticleEmitterInstance* Inst : EmitterInstances)
+	{
+		if (Inst)
+		{
+			Inst->BeginTrail();
+			Inst->SetTrailSourceData(InFirstSocketName, InSecondSocketName, InWidthMode, InWidth);
+		}
+	}
+}
+
+void UParticleSystemComponent::EndTrails()
+{
+	for (FParticleEmitterInstance* Inst : EmitterInstances)
+	{
+		if (Inst)
+		{
+			Inst->EndTrail();
+		}
+	}
+}
+
+void UParticleSystemComponent::SetTrailSourceData(FName InFirstSocketName, FName InSecondSocketName, ETrailWidthMode InWidthMode, float InWidth)
+{
+	for (FParticleEmitterInstance* Inst : EmitterInstances)
+	{
+		if (Inst)
+		{
+			Inst->SetTrailSourceData(InFirstSocketName, InSecondSocketName, InWidthMode, InWidth);
 		}
 	}
 }
diff --git a/Engine/Source/Runtime/Engine/Private/Particles/ParticleTrail2EmitterInstance.cpp b/Engine/Source/Runtime/Engine/Private/Particles/ParticleTrail2EmitterInstance.cpp
--- a/Engine/Source/Runtime/Engine/Private/Particles/ParticleTrail2EmitterInstance.cpp
+++ b/Engine/Source/Runtime/Engine/Private/Particles/ParticleTrail2EmitterInstance.cpp
@@ -2831,7 +2831,6 @@ FParticleAnimTrailEmitterInstance::FParticleAnimTrailEmitterInstance() :
 	, SecondSocketName(NAME_None)
 	, Width(1.0f)
 	, WidthMode( ETrailWidthMode_FromCentre )
-	, AnimNotifyState(NULL)
 	, bTagTrailAsDead(false)
 	, bTrailEnabled(false)
 #if WITH_EDITORONLY_DATA
@@ -4058,16 +4057,8 @@ SIZE_T FParticleAnimTrailEmitterInstance::GetResourceSize(EResourceSizeMode::Typ
 	return ResSize;
 }
 
-void FParticleAnimTrailEmitterInstance::BeginTrail(class UAnimNotifyState* InAnimNotifyState)
+void FParticleAnimTrailEmitterInstance::BeginTrail()
 {
-	if( !AnimNotifyState )
-	{
-		//This is the first time we've been triggered.
-		AnimNotifyState = InAnimNotifyState;
-	}
-
-	check(AnimNotifyState == InAnimNotifyState);
-
 	//Mark any existing trails as dead.
 	for (int32 FindTrailIdx = 0; FindTrailIdx < ActiveParticles; FindTrailIdx++)
 	{
@@ -4088,8 +4079,6 @@ void FParticleAnimTrailEmitterInstance::BeginTrail(class UAnimNotifyState* InAni
 
 void FParticleAnimTrailEmitterInstance::EndTrail()
 {
-	check(AnimNotifyState);
-
 	FirstSocketName = NAME_None;
 	SecondSocketName = NAME_None;
 	bTagTrailAsDead = true;
@@ -4098,7 +4087,6 @@ void FParticleAnimTrailEmitterInstance::EndTrail()
 
 void FParticleAnimTrailEmitterInstance::SetTrailSourceData(FName InFirstSocketName, FName InSecondSocketName, ETrailWidthMode InWidthMode, float InWidth )
 {
-	check(AnimNotifyState);
 	check(!InFirstSocketName.IsNone());
 	check(!InSecondSocketName.IsNone());
 
diff --git a/Engine/Source/Runtime/Engine/Public/ParticleEmitterInstances.h b/Engine/Source/Runtime/Engine/Public/ParticleEmitterInstances.h
--- a/Engine/Source/Runtime/Engine/Public/ParticleEmitterInstances.h
+++ b/Engine/Source/Runtime/Engine/Public/ParticleEmitterInstances.h
@@ -645,6 +645,26 @@ public:
 	virtual void ApplyWorldOffset(FVector InOffset, bool bWorldShift);
 		
 	virtual bool IsTrailEmitter()const{ return false; }
+
+	/**
+	* Begins the trail.
+	*/
+	virtual void BeginTrail(){}
+
+	/**
+	* Ends the trail.
+	*/
+	virtual void EndTrail(){}
+
+	/**
+	* Sets the data that defines this trail.
+	*
+	* @param	InFirstSocketName	The name of the first socket for the trail.
+	* @param	InSecondSocketName	The name of the second socket for the trail.
+	* @param	InWidthMode			How the width value is applied to the trail.
+	* @param	InWidth				The width of the trail.
+	*/
+	virtual void SetTrailSourceData(FName InFirstSocketName, FName InSecondSocketName, ETrailWidthMode InWidthMode, float InWidth){}
 protected:
 
 	/**
@@ -1389,11 +1409,6 @@ struct FParticleAnimTrailEmitterInstance : public FParticleTrailsEmitterInstance
 	*/
 	ETrailWidthMode WidthMode;
 
-	/**
-	 *	The anim notify state that spawned this emitter.
-	 */
-	class UAnimNotifyState* AnimNotifyState;
-
 	/**
 	*	When set, the current trail will be marked as dead in the next tick.
 	*/
@@ -1505,17 +1520,13 @@ struct FParticleAnimTrailEmitterInstance : public FParticleTrailsEmitterInstance
 
 	/**
 	 * Begins the trail.
-	 *
-	 * @param	InAnimNotifyState	The AnimNotifyState that created (or is creating) this trail.
 	 */
-	void BeginTrail(UAnimNotifyState* InAnimNotifyState);	
-	
+	virtual void BeginTrail();
+
 	/**
 	 * Ends the trail.
-	 *
-	 * @return  True if this was a trail ended by this notify, false otherwise.
 	 */
-	void EndTrail();
+	virtual void EndTrail();
 
 	/**
 	* Sets the date that defines this trail.
@@ -1525,7 +1536,7 @@ struct FParticleAnimTrailEmitterInstance : public FParticleTrailsEmitterInstance
 	* @param	InWidthMode			How the width value is applied to the trail.
 	* @param	InWidth				The width of the trail.
 	*/
-	void SetTrailSourceData( FName InFirstSocketName, FName InSecondSocketName, ETrailWidthMode InWidthMode, float InWidth );
+	virtual void SetTrailSourceData(FName InFirstSocketName, FName InSecondSocketName, ETrailWidthMode InWidthMode, float InWidth);
 
 	bool IsTrailActive()const;
 
