Allow animation evaluation to be run as a graph task, spreading the load across worker threads

#TTP 317273 - ANIM: Multi threading support

#proj UE4

diff --git a/Engine/Source/Editor/UnrealEd/Classes/Animation/DebugSkelMeshComponent.h b/Engine/Source/Editor/UnrealEd/Classes/Animation/DebugSkelMeshComponent.h
--- a/Engine/Source/Editor/UnrealEd/Classes/Animation/DebugSkelMeshComponent.h
+++ b/Engine/Source/Editor/UnrealEd/Classes/Animation/DebugSkelMeshComponent.h
@@ -157,7 +157,7 @@ class UDebugSkelMeshComponent : public USkeletalMeshComponent
 	// Begin SkinnedMeshComponent Interface
 	virtual bool ShouldCPUSkin() OVERRIDE;
 	virtual void PostInitMeshObject(class FSkeletalMeshObject* MeshObject) OVERRIDE;
-	virtual void RefreshBoneTransforms() OVERRIDE;
+	virtual void RefreshBoneTransforms(FActorComponentTickFunction* TickFunction = NULL) OVERRIDE;
 	// End SkinnedMeshComponent Interface
 
 	// Begin SkeletalMeshComponent Interface
diff --git a/Engine/Source/Editor/UnrealEd/Private/Animation/DebugSkelMeshComponent.cpp b/Engine/Source/Editor/UnrealEd/Private/Animation/DebugSkelMeshComponent.cpp
--- a/Engine/Source/Editor/UnrealEd/Private/Animation/DebugSkelMeshComponent.cpp
+++ b/Engine/Source/Editor/UnrealEd/Private/Animation/DebugSkelMeshComponent.cpp
@@ -350,10 +350,10 @@ void UDebugSkelMeshComponent::SetShowBoneWeight(bool bNewShowBoneWeight)
 	bDrawBoneInfluences = bNewShowBoneWeight;
 }
 
-void UDebugSkelMeshComponent::RefreshBoneTransforms()
+void UDebugSkelMeshComponent::RefreshBoneTransforms(FActorComponentTickFunction* TickFunction)
 {
 	// Run regular update first so we get RequiredBones up to date.
-	Super::RefreshBoneTransforms();
+	Super::RefreshBoneTransforms(NULL); // Pass NULL so we force non threaded work
 
 	// Non retargeted pose.
 	NonRetargetedSpaceBases.Empty();
@@ -362,7 +362,7 @@ void UDebugSkelMeshComponent::RefreshBoneTransforms()
 		TArray<FTransform> BackupSpaceBases = SpaceBases;
 
 		AnimScriptInstance->RequiredBones.SetDisableRetargeting(true);
-		Super::RefreshBoneTransforms();
+		Super::RefreshBoneTransforms(NULL);
 		AnimScriptInstance->RequiredBones.SetDisableRetargeting(false);
 
 		NonRetargetedSpaceBases = SpaceBases;
@@ -378,13 +378,13 @@ void UDebugSkelMeshComponent::RefreshBoneTransforms()
 		if( AnimScriptInstance && AnimScriptInstance->RequiredBones.IsValid() )
 		{
 			AnimScriptInstance->RequiredBones.SetUseRAWData(true);
-			Super::RefreshBoneTransforms();
+			Super::RefreshBoneTransforms(NULL);
 			AnimScriptInstance->RequiredBones.SetUseRAWData(false);
 		}
 		// Otherwise we'll just get ref pose.
 		else
 		{
-			Super::RefreshBoneTransforms();
+			Super::RefreshBoneTransforms(NULL);
 		}
 	}
 	else
diff --git a/Engine/Source/Runtime/Engine/Classes/Components/DestructibleComponent.h b/Engine/Source/Runtime/Engine/Classes/Components/DestructibleComponent.h
--- a/Engine/Source/Runtime/Engine/Classes/Components/DestructibleComponent.h
+++ b/Engine/Source/Runtime/Engine/Classes/Components/DestructibleComponent.h
@@ -131,7 +131,7 @@ public:
 
 	// Begin SkinnedMeshComponent interface.
 	virtual bool ShouldUpdateTransform(bool bLODHasChanged) const OVERRIDE;
-	virtual void RefreshBoneTransforms() OVERRIDE;
+	virtual void RefreshBoneTransforms(FActorComponentTickFunction* TickFunction = NULL) OVERRIDE;
 	virtual void SetSkeletalMesh(USkeletalMesh* InSkelMesh) OVERRIDE;
 	virtual FTransform GetSocketTransform(FName InSocketName, ERelativeTransformSpace TransformSpace = RTS_World) const OVERRIDE;
 	// End SkinnedMeshComponent interface.
diff --git a/Engine/Source/Runtime/Engine/Classes/Components/PoseableMeshComponent.h b/Engine/Source/Runtime/Engine/Classes/Components/PoseableMeshComponent.h
--- a/Engine/Source/Runtime/Engine/Classes/Components/PoseableMeshComponent.h
+++ b/Engine/Source/Runtime/Engine/Classes/Components/PoseableMeshComponent.h
@@ -59,7 +59,7 @@ class ENGINE_API UPoseableMeshComponent : public USkinnedMeshComponent
 	void ResetBoneTransformByName(FName BoneName);
 
 	// Begin USkinnedMeshComponent Interface
-	virtual void RefreshBoneTransforms() OVERRIDE;
+	virtual void RefreshBoneTransforms(FActorComponentTickFunction* TickFunction = NULL) OVERRIDE;
 	virtual bool AllocateTransformData() OVERRIDE;
 	// End USkinnedMeshComponent Interface
 
diff --git a/Engine/Source/Runtime/Engine/Classes/Components/SkeletalMeshComponent.h b/Engine/Source/Runtime/Engine/Classes/Components/SkeletalMeshComponent.h
--- a/Engine/Source/Runtime/Engine/Classes/Components/SkeletalMeshComponent.h
+++ b/Engine/Source/Runtime/Engine/Classes/Components/SkeletalMeshComponent.h
@@ -647,14 +647,14 @@ public:
 	void FreezeClothSection(bool bFreeze);
 
 	/** Evaluate Anim System **/
-	void EvaluateAnimation(/*FTransform & ExtractedRootMotionDelta, int32 & bHasRootMotion*/);
+	void EvaluateAnimation( TArray<FTransform>& OutLocalAtoms, TArray<struct FActiveVertexAnim>& OutVertexAnims, FVector& OutRootBoneTranslation ) const;
 
 	/**
-	 * Take the LocalAtoms array (translation vector, rotation quaternion and scale vector) and update the array of component-space bone transformation matrices (SpaceBases).
+	 * Take the SourceAtoms array (translation vector, rotation quaternion and scale vector) and update the array of component-space bone transformation matrices (DestSpaceBases).
 	 * It will work down hierarchy multiplying the component-space transform of the parent by the relative transform of the child.
 	 * This code also applies any per-bone rotators etc. as part of the composition process
 	 */
-	void FillSpaceBases();
+	void FillSpaceBases( const TArray<FTransform>& SourceAtoms, TArray<FTransform>& DestSpaceBases ) const;
 
 	/** 
 	 * Recalculates the RequiredBones array in this SkeletalMeshComponent based on current SkeletalMesh, LOD and PhysicsAsset.
@@ -745,7 +745,7 @@ public:
 
 	// Begin USkinnedMeshComponent interface
 	virtual bool UpdateLODStatus() OVERRIDE;
-	virtual void RefreshBoneTransforms() OVERRIDE;
+	virtual void RefreshBoneTransforms( FActorComponentTickFunction* TickFunction = NULL ) OVERRIDE;
 	virtual void TickPose( float DeltaTime ) OVERRIDE;
 	virtual void UpdateSlaveComponent() OVERRIDE;
 	virtual bool ShouldUpdateTransform(bool bLODHasChanged) const OVERRIDE;
@@ -763,6 +763,17 @@ public:
 	void UpdateRBJointMotors();
 
 
+	/**
+	* Runs the animation evaluation for the current pose into the supplied variables
+	*
+	* @param	OutSpaceBases			Component space bone transforms
+	* @param	OutLocalAtoms			Local space bone transforms
+	* @param	OutVertexAnims			Active vertex animations
+	* @param	OutRootBoneTranslation	Calculated root bone translation
+	*/
+	void PerformAnimationEvaluation( TArray<FTransform>& OutSpaceBases, TArray<FTransform>& OutLocalAtoms, TArray<FActiveVertexAnim>& OutVertexAnims, FVector& OutRootBoneTranslation ) const;
+	void PostAnimEvaluation( bool bDoInterpolation, bool bDuplicateToCacheBones );
+
 	/**
 	 * Blend of Physics Bones with PhysicsWeight and Animated Bones with (1-PhysicsWeight)
 	 *
@@ -1006,6 +1017,34 @@ private:
 	bool ShouldBlendPhysicsBones();	
 	void ClearAnimScriptInstance();
 
+	// Reference to our animation evaluation tick event
+	FGraphEventRef EvaluationTickEvent;
+	FGraphEventRef TickCompletionEvent;
+
+	//Handle parallel evaluation of animation
+	TArray<FTransform> PTSpaceBases;
+	TArray<FTransform> PTLocalAtoms;
+	TArray<FActiveVertexAnim> PTVertexAnims;
+	FVector PTRootBoneTranslation;
+	bool bPTDoInterpolation;
+	bool bPTDuplicateToCacheBones;
+
+	//Parallel Task Delegates
+	FSimpleDelegateGraphTask::FDelegate ParallelEvaluationDelegate;
+	FSimpleDelegateGraphTask::FDelegate ParallelCompletionDelegate;
+
+	// Parallel evaluation wrappers
+	void ParallelAnimationEvaluation() { PerformAnimationEvaluation(PTSpaceBases, PTLocalAtoms, PTVertexAnims, PTRootBoneTranslation); }
+	void CompleteParallelAnimationEvaluation()
+	{
+		Exchange(PTSpaceBases, bPTDoInterpolation ? CachedSpaceBases : SpaceBases);
+		Exchange(PTLocalAtoms, bPTDoInterpolation ? CachedLocalAtoms : LocalAtoms);
+		Exchange(PTVertexAnims, ActiveVertexAnims);
+		Exchange(PTRootBoneTranslation, RootBoneTranslation);
+
+		PostAnimEvaluation(bPTDoInterpolation, bPTDuplicateToCacheBones);
+	}
+
 	friend class FSkeletalMeshComponentDetails;
 
 private:
diff --git a/Engine/Source/Runtime/Engine/Classes/Components/SkinnedMeshComponent.h b/Engine/Source/Runtime/Engine/Classes/Components/SkinnedMeshComponent.h
--- a/Engine/Source/Runtime/Engine/Classes/Components/SkinnedMeshComponent.h
+++ b/Engine/Source/Runtime/Engine/Classes/Components/SkinnedMeshComponent.h
@@ -473,10 +473,12 @@ public:
 	/** 
 	 * Refresh Bone Transform (SpaceBases)
 	 * Each class will need to implement this function
-	 * Ideally this function should be atomic (not relying on Tick or any other update.) 
+	 * Ideally this function should be atomic (not relying on Tick or any other update.)
+	 *
+	 * @param TickFunction Supplied as non null if we are running in a tick, allows us to create graph tasks for parallelism
 	 * 
 	 */
-	virtual void RefreshBoneTransforms() PURE_VIRTUAL(USkinnedMeshComponent::RefreshBoneTransforms,);
+	virtual void RefreshBoneTransforms(FActorComponentTickFunction* TickFunction = NULL) PURE_VIRTUAL(USkinnedMeshComponent::RefreshBoneTransforms, );
 
 	/**
 	 * Tick Pose, this function ticks and do whatever it needs to do in this frame, should be called before RefreshBoneTransforms
@@ -513,7 +515,7 @@ public:
 	/** 
 	 * Combine CurveKeys (that reference morph targets by name) and ActiveAnims (that reference vertex anims by reference) into the ActiveVertexAnims array.
 	 */
-	void UpdateActiveVertexAnims(const TMap<FName, float>& MorphCurveAnims, const TArray<FActiveVertexAnim>& ActiveAnims);
+	TArray<struct FActiveVertexAnim> UpdateActiveVertexAnims(const TMap<FName, float>& MorphCurveAnims, const TArray<FActiveVertexAnim>& ActiveAnims) const;
 
 	/**
 	 * Checks/updates material usage on proxy based on current morph target usage
@@ -746,7 +748,7 @@ public:
 	 *
 	 * @return Pointer to found MorphTarget. Returns NULL if could not find target with that name.
 	 */
-	virtual class UMorphTarget* FindMorphTarget( FName MorphTargetName );
+	virtual class UMorphTarget* FindMorphTarget( FName MorphTargetName ) const;
 
 	/**	
 	 * Find all bones by name within given radius 
diff --git a/Engine/Source/Runtime/Engine/Classes/Engine/SkeletalMesh.h b/Engine/Source/Runtime/Engine/Classes/Engine/SkeletalMesh.h
--- a/Engine/Source/Runtime/Engine/Classes/Engine/SkeletalMesh.h
+++ b/Engine/Source/Runtime/Engine/Classes/Engine/SkeletalMesh.h
@@ -748,7 +748,7 @@ public:
 	 *
 	 * @return Pointer to found MorphTarget. Returns NULL if could not find target with that name.
 	 */
-	ENGINE_API UMorphTarget* FindMorphTarget( FName MorphTargetName );
+	ENGINE_API UMorphTarget* FindMorphTarget( FName MorphTargetName ) const;
 
 	/** if name conflicts, it will overwrite the reference */
 	ENGINE_API void RegisterMorphTarget(UMorphTarget* MorphTarget);
diff --git a/Engine/Source/Runtime/Engine/Private/Animation/AnimInstance.cpp b/Engine/Source/Runtime/Engine/Private/Animation/AnimInstance.cpp
--- a/Engine/Source/Runtime/Engine/Private/Animation/AnimInstance.cpp
+++ b/Engine/Source/Runtime/Engine/Private/Animation/AnimInstance.cpp
@@ -44,6 +44,8 @@ DEFINE_STAT(STAT_InterpolateSkippedFrames);
 DEFINE_STAT(STAT_AnimTickTime);
 DEFINE_STAT(STAT_SkinnedMeshCompTick);
 DEFINE_STAT(STAT_TickUpdateRate);
+DEFINE_STAT(STAT_PerformAnimEvaluation);
+DEFINE_STAT(STAT_PostAnimEvaluation);
 
 DEFINE_STAT(STAT_AnimStateMachineUpdate);
 DEFINE_STAT(STAT_AnimStateMachineFindTransition);
diff --git a/Engine/Source/Runtime/Engine/Private/DestructibleComponent.cpp b/Engine/Source/Runtime/Engine/Private/DestructibleComponent.cpp
--- a/Engine/Source/Runtime/Engine/Private/DestructibleComponent.cpp
+++ b/Engine/Source/Runtime/Engine/Private/DestructibleComponent.cpp
@@ -635,7 +635,7 @@ void UDestructibleComponent::OnDamageEvent(const NxApexDamageEventReportData& In
 }
 #endif // WITH_APEX
 
-void UDestructibleComponent::RefreshBoneTransforms()
+void UDestructibleComponent::RefreshBoneTransforms(FActorComponentTickFunction* TickFunction)
 {
 #if WITH_APEX
 	if(ApexDestructibleActor != NULL && SkeletalMesh)
diff --git a/Engine/Source/Runtime/Engine/Private/PoseableMeshComponent.cpp b/Engine/Source/Runtime/Engine/Private/PoseableMeshComponent.cpp
--- a/Engine/Source/Runtime/Engine/Private/PoseableMeshComponent.cpp
+++ b/Engine/Source/Runtime/Engine/Private/PoseableMeshComponent.cpp
@@ -42,7 +42,7 @@ bool UPoseableMeshComponent::AllocateTransformData()
 	return false;
 }
 
-void UPoseableMeshComponent::RefreshBoneTransforms()
+void UPoseableMeshComponent::RefreshBoneTransforms(FActorComponentTickFunction* TickFunction)
 {
 	SCOPE_CYCLE_COUNTER(STAT_RefreshBoneTransforms);
 
diff --git a/Engine/Source/Runtime/Engine/Private/SkeletalMesh.cpp b/Engine/Source/Runtime/Engine/Private/SkeletalMesh.cpp
--- a/Engine/Source/Runtime/Engine/Private/SkeletalMesh.cpp
+++ b/Engine/Source/Runtime/Engine/Private/SkeletalMesh.cpp
@@ -2812,7 +2812,7 @@ void USkeletalMesh::InitMorphTargets()
 }
 
 
-UMorphTarget* USkeletalMesh::FindMorphTarget( FName MorphTargetName )
+UMorphTarget* USkeletalMesh::FindMorphTarget( FName MorphTargetName ) const
 {
 	if( MorphTargetName != NAME_None )
 	{
diff --git a/Engine/Source/Runtime/Engine/Private/SkeletalMeshComponent.cpp b/Engine/Source/Runtime/Engine/Private/SkeletalMeshComponent.cpp
--- a/Engine/Source/Runtime/Engine/Private/SkeletalMeshComponent.cpp
+++ b/Engine/Source/Runtime/Engine/Private/SkeletalMeshComponent.cpp
@@ -23,6 +23,8 @@
 	#error EXPERIMENTAL_PARALLEL_CODE must be defined as either zero or one
 #endif
 
+TAutoConsoleVariable<int32> CVarUseParallelAnimationEvaluation(TEXT("ParallelAnimEvaluation"), 0, TEXT("If true, animation evaluation will be run across the task graph system as opposed to purely on the game thread"));
+
 USkeletalMeshComponent::USkeletalMeshComponent(const class FPostConstructInitializeProperties& PCIP)
 	: Super(PCIP)
 {
@@ -72,6 +74,9 @@ USkeletalMeshComponent::USkeletalMeshComponent(const class FPostConstructInitial
 #endif
 
 	bTickInEditor = true;
+
+	ParallelEvaluationDelegate = FSimpleDelegateGraphTask::FDelegate::CreateUObject(this, &USkeletalMeshComponent::ParallelAnimationEvaluation);
+	ParallelCompletionDelegate = FSimpleDelegateGraphTask::FDelegate::CreateUObject(this, &USkeletalMeshComponent::CompleteParallelAnimationEvaluation);
 }
 
 
@@ -532,7 +537,7 @@ static void IntersectBoneIndexArrays(TArray<FBoneIndexType>& Output, const TArra
 }
 
 
-void USkeletalMeshComponent::FillSpaceBases()
+void USkeletalMeshComponent::FillSpaceBases(const TArray<FTransform>& SourceAtoms, TArray<FTransform>& DestSpaceBases) const
 {
 	SCOPE_CYCLE_COUNTER(STAT_SkelComposeTime);
 
@@ -541,12 +546,11 @@ void USkeletalMeshComponent::FillSpaceBases()
 		return;
 	}
 
-	// right now all this does is to convert to SpaceBases
-	check( SkeletalMesh->RefSkeleton.GetNum() == LocalAtoms.Num() );
-	check( SkeletalMesh->RefSkeleton.GetNum() == SpaceBases.Num() );
-	check( SkeletalMesh->RefSkeleton.GetNum() == BoneVisibilityStates.Num() );
+	// right now all this does is populate DestSpaceBases
+	check(SkeletalMesh->RefSkeleton.GetNum() == SourceAtoms.Num());
+	check( SkeletalMesh->RefSkeleton.GetNum() == DestSpaceBases.Num() );
 
-	const int32 NumBones = LocalAtoms.Num();
+	const int32 NumBones = SourceAtoms.Num();
 
 #if (UE_BUILD_DEBUG || UE_BUILD_DEVELOPMENT)
 	/** Keep track of which bones have been processed for fast look up */
@@ -554,13 +558,13 @@ void USkeletalMeshComponent::FillSpaceBases()
 	BoneProcessed.AddZeroed(NumBones);
 #endif
 
-	FTransform * LocalTransformsData = LocalAtoms.GetTypedData(); 
-	FTransform * SpaceBasesData = SpaceBases.GetTypedData();
+	const FTransform * LocalTransformsData = SourceAtoms.GetTypedData();
+	FTransform * SpaceBasesData = DestSpaceBases.GetTypedData();
 
 	// First bone is always root bone, and it doesn't have a parent.
 	{
 		check( RequiredBones[0] == 0 );
-		SpaceBases[0] = LocalAtoms[0];
+		DestSpaceBases[0] = SourceAtoms[0];
 
 #if (UE_BUILD_DEBUG || UE_BUILD_DEVELOPMENT)
 		// Mark bone as processed
@@ -587,8 +591,8 @@ void USkeletalMeshComponent::FillSpaceBases()
 #endif
 		FTransform::Multiply(SpaceBasesData + BoneIndex, LocalTransformsData + BoneIndex, SpaceBasesData + ParentIndex);
 
-		SpaceBases[BoneIndex].DiagnosticCheckNaN_All();
-		SpaceBases[BoneIndex].DiagnosticCheckUnitQuaternion();
+		DestSpaceBases[BoneIndex].DiagnosticCheckNaN_All();
+		DestSpaceBases[BoneIndex].DiagnosticCheckUnitQuaternion();
 	}
 }
 
@@ -686,7 +690,7 @@ void USkeletalMeshComponent::RecalcRequiredBones(int32 LODIndex)
 	// TODO UE4
 
 	// Purge invisible bones and their children
-	// this has to be done before mirror table check/phsysics body checks
+	// this has to be done before mirror table check/physics body checks
 	// mirror table/phys body ones has to be calculated
 	if (ShouldUpdateBoneVisibility())
 	{
@@ -749,7 +753,7 @@ void USkeletalMeshComponent::RecalcRequiredBones(int32 LODIndex)
 	CachedSpaceBases.Empty();
 }
 
-void USkeletalMeshComponent::EvaluateAnimation(/*FTransform & ExtractedRootMotionDelta, int32 & bHasRootMotion*/)
+void USkeletalMeshComponent::EvaluateAnimation(TArray<FTransform>& OutLocalAtoms, TArray<FActiveVertexAnim>& OutVertexAnims, FVector& OutRootBoneTranslation) const
 {
 	SCOPE_CYCLE_COUNTER(STAT_AnimBlendTime);
 
@@ -776,34 +780,34 @@ void USkeletalMeshComponent::EvaluateAnimation(/*FTransform & ExtractedRootMotio
 			// can we avoid that copy?
 			if( EvaluationContext.Pose.Bones.Num() > 0 )
 			{
-				LocalAtoms = EvaluationContext.Pose.Bones;
+				OutLocalAtoms = EvaluationContext.Pose.Bones;
 			}
 			else
 			{
-				FAnimationRuntime::FillWithRefPose(LocalAtoms, AnimScriptInstance->RequiredBones);
+				FAnimationRuntime::FillWithRefPose(OutLocalAtoms, AnimScriptInstance->RequiredBones);
 			}
 		}
 		else
 		{
-			FAnimationRuntime::FillWithRefPose(LocalAtoms, AnimScriptInstance->RequiredBones);
+			FAnimationRuntime::FillWithRefPose(OutLocalAtoms, AnimScriptInstance->RequiredBones);
 		}
 
-		UpdateActiveVertexAnims(AnimScriptInstance->MorphTargetCurves, AnimScriptInstance->VertexAnims);
+		OutVertexAnims = UpdateActiveVertexAnims(AnimScriptInstance->MorphTargetCurves, AnimScriptInstance->VertexAnims);
 	}
 	else
 	{
-		LocalAtoms = SkeletalMesh->RefSkeleton.GetRefBonePose();
+		OutLocalAtoms = SkeletalMesh->RefSkeleton.GetRefBonePose();
 
 		// if it's only morph, there is no reason to blend
 		if ( MorphTargetCurves.Num() > 0 )
 		{
 			TArray<struct FActiveVertexAnim> EmptyVertexAnims;
-			UpdateActiveVertexAnims(MorphTargetCurves, EmptyVertexAnims);
+			OutVertexAnims = UpdateActiveVertexAnims(MorphTargetCurves, EmptyVertexAnims);
 		}
 	}
 
 	// Remember the root bone's translation so we can move the bounds.
-	RootBoneTranslation = LocalAtoms[0].GetTranslation() - SkeletalMesh->RefSkeleton.GetRefBonePose()[0].GetTranslation();
+	OutRootBoneTranslation = OutLocalAtoms[0].GetTranslation() - SkeletalMesh->RefSkeleton.GetRefBonePose()[0].GetTranslation();
 }
 
 void USkeletalMeshComponent::UpdateSlaveComponent()
@@ -816,105 +820,143 @@ void USkeletalMeshComponent::UpdateSlaveComponent()
 
 		if ( MasterSMC->AnimScriptInstance )
 		{
-			UpdateActiveVertexAnims(MasterSMC->AnimScriptInstance->MorphTargetCurves, MasterSMC->AnimScriptInstance->VertexAnims);
+			ActiveVertexAnims = UpdateActiveVertexAnims(MasterSMC->AnimScriptInstance->MorphTargetCurves, MasterSMC->AnimScriptInstance->VertexAnims);
 		}
 	}
 
 	Super::UpdateSlaveComponent();
 }
 
-void USkeletalMeshComponent::RefreshBoneTransforms()
+void USkeletalMeshComponent::PerformAnimationEvaluation(TArray<FTransform>& OutSpaceBases, TArray<FTransform>& OutLocalAtoms, TArray<FActiveVertexAnim>& OutVertexAnims, FVector& OutRootBoneTranslation) const
 {
-	SCOPE_CYCLE_COUNTER(STAT_RefreshBoneTransforms);
-
+	SCOPE_CYCLE_COUNTER(STAT_PerformAnimEvaluation);
 	// Can't do anything without a SkeletalMesh
 	// Do nothing more if no bones in skeleton.
-	if( !SkeletalMesh || SpaceBases.Num() == 0 )
+	if (!SkeletalMesh || OutSpaceBases.Num() == 0)
 	{
 		return;
 	}
-	
 	AActor * Owner = GetOwner();
 	const FAnimUpdateRateParameters  & UpdateRateParams = Owner ? Owner->AnimUpdateRateParams : FAnimUpdateRateParameters();
 	UE_LOG(LogAnimation, Verbose, TEXT("RefreshBoneTransforms(%s)"), *GetNameSafe(Owner));
 
+	// evaluate pure animations, and fill up LocalAtoms
+	EvaluateAnimation(OutLocalAtoms, OutVertexAnims, OutRootBoneTranslation);
+	// Fill SpaceBases from LocalAtoms
+	FillSpaceBases(OutLocalAtoms, OutSpaceBases);
+}
+
+void USkeletalMeshComponent::RefreshBoneTransforms(FActorComponentTickFunction* TickFunction)
+{
+	SCOPE_CYCLE_COUNTER(STAT_RefreshBoneTransforms);
+
+	AActor * Owner = GetOwner();
+	UE_LOG(LogAnimation, Verbose, TEXT("RefreshBoneTransforms(%s)"), *GetNameSafe(Owner));
+
+	// Recalculate the RequiredBones array, if necessary
+	if (!bRequiredBonesUpToDate)
 	{
-		FScopeLockPhysXWriter LockPhysXForWriting;
+		RecalcRequiredBones(PredictedLODLevel);
+	}
 
-		// Recalculate the RequiredBones array, if necessary
-		if( !bRequiredBonesUpToDate )
-		{
-			RecalcRequiredBones(PredictedLODLevel);
-		}
+	//Handle update rate optimization setup
+	const FAnimUpdateRateParameters  & UpdateRateParams = Owner ? Owner->AnimUpdateRateParams : FAnimUpdateRateParameters();
+	
+	const bool bDoUpdateRateOptimization = bEnableUpdateRateOptimizations && (UpdateRateParams.GetEvaluationRate() > 1);
+	//Dont mark cache as invalid if we aren't performing optimization anyway
+	const bool bInvalidCachedBones = bDoUpdateRateOptimization &&
+									( (LocalAtoms.Num() != SkeletalMesh->RefSkeleton.GetNum())
+									  || (LocalAtoms.Num() != CachedLocalAtoms.Num())
+									  || (SpaceBases.Num() != CachedSpaceBases.Num()) );
+
+	const bool bDoEvaluation = !bDoUpdateRateOptimization || bInvalidCachedBones || !UpdateRateParams.ShouldSkipEvaluation();
+	
+	bPTDoInterpolation = bDoUpdateRateOptimization && !bInvalidCachedBones && UpdateRateParams.ShouldInterpolateSkippedFrames();
+	bPTDuplicateToCacheBones = bInvalidCachedBones || (bDoUpdateRateOptimization && bDoEvaluation && !bPTDoInterpolation);
 
-		// Update rate turned off, evaluate every frame.
-		if( !bEnableUpdateRateOptimizations || (UpdateRateParams.GetEvaluationRate() <= 1) )
+	if (!bDoUpdateRateOptimization)
+	{
+		//If we aren't optimizing clear the cached local atoms
+		CachedLocalAtoms.Empty();
+		CachedSpaceBases.Empty();
+	}
+
+	const bool bDoPAE = !!CVarUseParallelAnimationEvaluation.GetValueOnGameThread();
+	if (bDoEvaluation && TickFunction && bDoPAE)
+	{
+		if (SkeletalMesh->RefSkeleton.GetNum() != PTLocalAtoms.Num())
 		{
-			// evaluate pure animations, and fill up LocalAtoms
-			EvaluateAnimation();
-			// We need the mesh space bone transforms now for renderer to get delta from ref pose:
-			FillSpaceBases();
-			// Invalidate cached bones.
-			CachedLocalAtoms.Empty();
-			CachedSpaceBases.Empty();
+			// Initialize Parallel Task arrays
+			PTLocalAtoms = LocalAtoms;
+			PTSpaceBases = SpaceBases;
+			PTVertexAnims = ActiveVertexAnims;
 		}
-		else
+
+		// start parallel work
+		FGraphEventRef EvaluationTickEvent = FSimpleDelegateGraphTask::CreateAndDispatchWhenReady
+			(
+			ParallelEvaluationDelegate
+			, TEXT("USkeletalMeshComponent::ParallelAnimationEvaluation")
+			);
+
+		// set up a task to run on the game thread to accept the results
+		TickCompletionEvent = FSimpleDelegateGraphTask::CreateAndDispatchWhenReady
+			(
+			ParallelCompletionDelegate
+			, TEXT("USkeletalMeshComponent::CompleteParallelAnimationEvaluation")
+			, EvaluationTickEvent
+			, ENamedThreads::GameThread
+			);
+
+		TickFunction->GetCompletionHandle()->DontCompleteUntil(TickCompletionEvent);
+
+	}
+	else
+	{
+		if (bDoEvaluation)
 		{
-			// figure out if our cache is invalid.
-			const bool bInvalidCachedBones = (LocalAtoms.Num() != SkeletalMesh->RefSkeleton.GetNum()) 
-				|| (LocalAtoms.Num() != CachedLocalAtoms.Num())
-				|| (SpaceBases.Num() != CachedSpaceBases.Num());
-
-			// If cache is invalid, we need to rebuild it. And we can't interpolate.
-			// (same path if we're not interpolating and not skipping a frame).
-			if( bInvalidCachedBones || (!UpdateRateParams.ShouldInterpolateSkippedFrames() && !UpdateRateParams.ShouldSkipEvaluation()) )
+			if (bPTDoInterpolation)
 			{
-				// evaluate pure animations, and fill up LocalAtoms
-				EvaluateAnimation();
-				// Fill SpaceBases from LocalAtoms
-				FillSpaceBases();
-
-				// Cache bones
-				CachedLocalAtoms = LocalAtoms;
-				CachedSpaceBases = SpaceBases;
+				PerformAnimationEvaluation(CachedSpaceBases, CachedLocalAtoms, ActiveVertexAnims, RootBoneTranslation);
 			}
 			else
 			{
-				// No interpolation, just copy
-				// @todo: if we don't blend any physics, we could even skip the copy.
-				if( !UpdateRateParams.ShouldInterpolateSkippedFrames() )
-				{
-					LocalAtoms = CachedLocalAtoms;
-					SpaceBases = CachedSpaceBases;
-				}
-				else
-				{
-					// If we are not skipping evaluation this frame, refresh cache.
-					if( !UpdateRateParams.ShouldSkipEvaluation() )
-					{
-						// Preserve LocalAtoms and SpaceBases, so we can keep interpolation.
-						Exchange(LocalAtoms, CachedLocalAtoms);
-						Exchange(SpaceBases, CachedSpaceBases);
+				PerformAnimationEvaluation(SpaceBases, LocalAtoms, ActiveVertexAnims, RootBoneTranslation);
+			}
+		}
+		else if (!bPTDoInterpolation)
+		{
+			LocalAtoms = CachedLocalAtoms;
+			SpaceBases = CachedSpaceBases;
+		}
 
-						// evaluate pure animations, and fill up LocalAtoms
-						EvaluateAnimation();
-						// Fill SpaceBases from LocalAtoms
-						FillSpaceBases();
+		PostAnimEvaluation(bPTDoInterpolation, bPTDuplicateToCacheBones);
+	}
 
-						Exchange(LocalAtoms, CachedLocalAtoms);
-						Exchange(SpaceBases, CachedSpaceBases);
-					}
+}
 
-					// Interpolate
-					{
-						SCOPE_CYCLE_COUNTER(STAT_InterpolateSkippedFrames);
-						const float Alpha = 0.25f + (1.f / float(FMath::Max(UpdateRateParams.GetEvaluationRate(), 2) * 2));
-						FAnimationRuntime::LerpBoneTransforms(LocalAtoms, CachedLocalAtoms, Alpha, RequiredBones);
-						FAnimationRuntime::LerpBoneTransforms(SpaceBases, CachedSpaceBases, Alpha, RequiredBones);
-					}
-				}
-			}
-		}
+void USkeletalMeshComponent::PostAnimEvaluation(bool bDoInterpolation, bool bDuplicateToCacheBones)
+{
+	SCOPE_CYCLE_COUNTER(STAT_PostAnimEvaluation);
+	if( bDuplicateToCacheBones )
+	{
+		CachedSpaceBases = SpaceBases;
+		CachedLocalAtoms = LocalAtoms;
+	}
+
+	if( bDoInterpolation )
+	{
+		SCOPE_CYCLE_COUNTER(STAT_InterpolateSkippedFrames);
+		AActor * Owner = GetOwner();
+		const FAnimUpdateRateParameters  & UpdateRateParams = Owner ? Owner->AnimUpdateRateParams : FAnimUpdateRateParameters();
+
+		const float Alpha = 0.25f + (1.f / float(FMath::Max(UpdateRateParams.GetEvaluationRate(), 2) * 2));
+		FAnimationRuntime::LerpBoneTransforms(LocalAtoms, CachedLocalAtoms, Alpha, RequiredBones);
+		FAnimationRuntime::LerpBoneTransforms(SpaceBases, CachedSpaceBases, Alpha, RequiredBones);
+	}
+
+	{
+		FScopeLockPhysXWriter LockPhysXForWriting;
 
 		// Transforms updated, cached local bounds are now out of date.
 		InvalidateCachedBounds();
@@ -935,9 +977,9 @@ void USkeletalMeshComponent::RefreshBoneTransforms()
 		if( !IsSimulatingPhysics() )
 		{
 			SCOPE_CYCLE_COUNTER(STAT_UpdateLocalToWorldAndOverlaps);
-			
+
 			// New bone positions need to be sent to render thread
-			UpdateComponentToWorld();	
+			UpdateComponentToWorld();
 
 			// animation often change overlap. 
 			UpdateOverlaps();
diff --git a/Engine/Source/Runtime/Engine/Private/SkinnedMeshComponent.cpp b/Engine/Source/Runtime/Engine/Private/SkinnedMeshComponent.cpp
--- a/Engine/Source/Runtime/Engine/Private/SkinnedMeshComponent.cpp
+++ b/Engine/Source/Runtime/Engine/Private/SkinnedMeshComponent.cpp
@@ -317,7 +317,7 @@ void USkinnedMeshComponent::TickComponent(float DeltaTime, enum ELevelTick TickT
 		}
 		else 
 		{
-			RefreshBoneTransforms(); 
+			RefreshBoneTransforms(ThisTickFunction);
 		}
 	}
 }
@@ -896,7 +896,7 @@ void USkinnedMeshComponent::SetForceWireframe(bool InForceWireframe)
 }
 
 
-UMorphTarget* USkinnedMeshComponent::FindMorphTarget( FName MorphTargetName )
+UMorphTarget* USkinnedMeshComponent::FindMorphTarget( FName MorphTargetName ) const
 {
 	if( SkeletalMesh != NULL )
 	{
@@ -1566,10 +1566,9 @@ int32 FindVertexAnim(const TArray<FActiveVertexAnim>& ActiveAnims, UVertexAnimBa
 	return INDEX_NONE;
 }
 
-void USkinnedMeshComponent::UpdateActiveVertexAnims(const TMap<FName, float>& MorphCurveAnims, const TArray<FActiveVertexAnim>& ActiveAnims)
+TArray<FActiveVertexAnim> USkinnedMeshComponent::UpdateActiveVertexAnims(const TMap<FName, float>& MorphCurveAnims, const TArray<FActiveVertexAnim>& ActiveAnims) const
 {
-	// Empty the existing ActiveVertexAnims array
-	ActiveVertexAnims.Empty();
+	TArray<struct FActiveVertexAnim> OutVertexAnims;
 
 	// First copy ActiveAnims
 	for(int32 AnimIdx=0; AnimIdx < ActiveAnims.Num(); AnimIdx++)
@@ -1580,7 +1579,7 @@ void USkinnedMeshComponent::UpdateActiveVertexAnims(const TMap<FName, float>& Mo
 			ActiveAnim.VertAnim != NULL &&
 			ActiveAnim.VertAnim->BaseSkelMesh == SkeletalMesh )
 		{
-			ActiveVertexAnims.Add(ActiveAnim);
+			OutVertexAnims.Add(ActiveAnim);
 		}
 		// @TODO Need to check for duplicates here?
 	}
@@ -1599,20 +1598,22 @@ void USkinnedMeshComponent::UpdateActiveVertexAnims(const TMap<FName, float>& Mo
 			if(Target != NULL)				
 			{
 				// See if this morph target already has an entry
-				int32 AnimIndex = FindVertexAnim(ActiveVertexAnims, Target);
+				int32 AnimIndex = FindVertexAnim(OutVertexAnims, Target);
 				// If not, add it
 				if(AnimIndex == INDEX_NONE)
 				{
-					ActiveVertexAnims.Add(FActiveVertexAnim(Target, Weight));
+					OutVertexAnims.Add(FActiveVertexAnim(Target, Weight));
 				}
 				// If it does, use the max weight
 				else
 				{
 					const float CurrentWeight = ActiveVertexAnims[AnimIndex].Weight;
-					ActiveVertexAnims[AnimIndex].Weight = FMath::Max<float>(CurrentWeight, Weight);
+					OutVertexAnims[AnimIndex].Weight = FMath::Max<float>(CurrentWeight, Weight);
 				}
 			}
 		}
 	}
+
+	return OutVertexAnims;
 }
 
diff --git a/Engine/Source/Runtime/Engine/Public/AnimTree.h b/Engine/Source/Runtime/Engine/Public/AnimTree.h
--- a/Engine/Source/Runtime/Engine/Public/AnimTree.h
+++ b/Engine/Source/Runtime/Engine/Public/AnimTree.h
@@ -17,6 +17,8 @@ DECLARE_CYCLE_STAT_EXTERN(TEXT("StateMachine Update"), STAT_AnimStateMachineUpda
 DECLARE_CYCLE_STAT_EXTERN(TEXT("StateMachine Find Transition"), STAT_AnimStateMachineFindTransition, STATGROUP_Anim, );
 DECLARE_CYCLE_STAT_EXTERN(TEXT("RefreshBoneTransforms"), STAT_RefreshBoneTransforms, STATGROUP_Anim, );
 DECLARE_CYCLE_STAT_EXTERN(TEXT("Evaluate Anim"), STAT_AnimBlendTime, STATGROUP_Anim, );
+DECLARE_CYCLE_STAT_EXTERN(TEXT("Perform Anim Evaluation"), STAT_PerformAnimEvaluation, STATGROUP_Anim, );
+DECLARE_CYCLE_STAT_EXTERN(TEXT("Post Anim Evaluation"), STAT_PostAnimEvaluation, STATGROUP_Anim, );
 DECLARE_CYCLE_STAT_EXTERN(TEXT("GraphTime"), STAT_AnimGraphEvaluate, STATGROUP_Anim, );
 DECLARE_CYCLE_STAT_EXTERN(TEXT("Eval Pose"), STAT_AnimNativeEvaluatePoses, STATGROUP_Anim, );
 DECLARE_CYCLE_STAT_EXTERN(TEXT("Anim Decompression"), STAT_GetAnimationPose, STATGROUP_Anim, );
