Another CIS fix to World.cpp

diff --git a/Engine/Source/Runtime/Engine/Private/World.cpp b/Engine/Source/Runtime/Engine/Private/World.cpp
--- a/Engine/Source/Runtime/Engine/Private/World.cpp
+++ b/Engine/Source/Runtime/Engine/Private/World.cpp
@@ -6,6 +6,7 @@
 
 #include "EnginePrivate.h"
 #include "Engine/LevelScriptBlueprint.h"
+#include "Engine/CullDistanceVolume.h"
 #include "Engine/Console.h"
 #include "Engine/WorldComposition.h"
 #include "Matinee/MatineeActor.h"
@@ -18,6 +19,7 @@
 #include "PrecomputedLightVolume.h"
 #include "UObjectAnnotation.h"
 #include "RenderCore.h"
+#include "EngineModule.h"
 #include "ParticleHelper.h"
 #include "TickTaskManagerInterface.h"
 #include "FXSystem.h"
