Remove some headers from EnginePrivate.h (HUD, vehicles, movement components, pawn/character, splinemesh)

diff --git a/Engine/Source/Runtime/Engine/Classes/Debug/GameplayDebuggingHUDComponent.h b/Engine/Source/Runtime/Engine/Classes/Debug/GameplayDebuggingHUDComponent.h
--- a/Engine/Source/Runtime/Engine/Classes/Debug/GameplayDebuggingHUDComponent.h
+++ b/Engine/Source/Runtime/Engine/Classes/Debug/GameplayDebuggingHUDComponent.h
@@ -1,7 +1,7 @@
 // Copyright 1998-2014 Epic Games, Inc. All Rights Reserved.
 
 #pragma once
-
+#include "GameFramework/HUD.h"
 #include "GameplayDebuggingHUDComponent.generated.h"
 
 UCLASS(notplaceable)
diff --git a/Engine/Source/Runtime/Engine/Classes/GameFramework/Character.h b/Engine/Source/Runtime/Engine/Classes/GameFramework/Character.h
--- a/Engine/Source/Runtime/Engine/Classes/GameFramework/Character.h
+++ b/Engine/Source/Runtime/Engine/Classes/GameFramework/Character.h
@@ -1,7 +1,7 @@
 // Copyright 1998-2014 Epic Games, Inc. All Rights Reserved.
 
 #pragma once
-
+#include "GameFramework/CharacterMovementComponent.h"
 #include "Character.generated.h"
 
 DECLARE_DYNAMIC_MULTICAST_DELEGATE_ThreeParams(FMovementModeChangedSignature, class ACharacter*, Character, EMovementMode, PrevMovementMode, uint8, PreviousCustomMode);
diff --git a/Engine/Source/Runtime/Engine/Classes/GameFramework/NavMovementComponent.h b/Engine/Source/Runtime/Engine/Classes/GameFramework/NavMovementComponent.h
--- a/Engine/Source/Runtime/Engine/Classes/GameFramework/NavMovementComponent.h
+++ b/Engine/Source/Runtime/Engine/Classes/GameFramework/NavMovementComponent.h
@@ -6,6 +6,7 @@
 
 #pragma once
 #include "AI/Navigation/NavAgentInterface.h"
+#include "GameFramework/MovementComponent.h"
 #include "NavMovementComponent.generated.h"
 
 UCLASS(abstract, dependson=(UNavigationSystem, UPrimitiveComponent, INavAgentInterface))
diff --git a/Engine/Source/Runtime/Engine/Classes/GameFramework/PawnMovementComponent.h b/Engine/Source/Runtime/Engine/Classes/GameFramework/PawnMovementComponent.h
--- a/Engine/Source/Runtime/Engine/Classes/GameFramework/PawnMovementComponent.h
+++ b/Engine/Source/Runtime/Engine/Classes/GameFramework/PawnMovementComponent.h
@@ -5,6 +5,7 @@
  */
 
 #pragma once
+#include "GameFramework/NavMovementComponent.h"
 #include "PawnMovementComponent.generated.h"
 
 UCLASS(abstract)
diff --git a/Engine/Source/Runtime/Engine/Classes/GameFramework/SpectatorPawn.h b/Engine/Source/Runtime/Engine/Classes/GameFramework/SpectatorPawn.h
--- a/Engine/Source/Runtime/Engine/Classes/GameFramework/SpectatorPawn.h
+++ b/Engine/Source/Runtime/Engine/Classes/GameFramework/SpectatorPawn.h
@@ -7,6 +7,7 @@
 //=============================================================================
 
 #pragma once
+#include "GameFramework/DefaultPawn.h"
 #include "SpectatorPawn.generated.h"
 
 
diff --git a/Engine/Source/Runtime/Engine/Classes/GameFramework/SpectatorPawnMovement.h b/Engine/Source/Runtime/Engine/Classes/GameFramework/SpectatorPawnMovement.h
--- a/Engine/Source/Runtime/Engine/Classes/GameFramework/SpectatorPawnMovement.h
+++ b/Engine/Source/Runtime/Engine/Classes/GameFramework/SpectatorPawnMovement.h
@@ -5,6 +5,7 @@
  */
 
 #pragma once
+#include "GameFramework/FloatingPawnMovement.h"
 #include "SpectatorPawnMovement.generated.h"
 
 UCLASS(DependsOn=UFloatingPawnMovement)
diff --git a/Engine/Source/Runtime/Engine/Classes/Vehicles/WheeledVehicleMovementComponent4W.h b/Engine/Source/Runtime/Engine/Classes/Vehicles/WheeledVehicleMovementComponent4W.h
--- a/Engine/Source/Runtime/Engine/Classes/Vehicles/WheeledVehicleMovementComponent4W.h
+++ b/Engine/Source/Runtime/Engine/Classes/Vehicles/WheeledVehicleMovementComponent4W.h
@@ -4,6 +4,7 @@
  * Base VehicleSim for the 4W PhysX vehicle class
  */
 #pragma once
+#include "Vehicles/WheeledVehicleMovementComponent.h"
 #include "WheeledVehicleMovementComponent4W.generated.h"
 
 UENUM()
diff --git a/Engine/Source/Runtime/Engine/Private/AI/Navigation/AvoidanceManager.cpp b/Engine/Source/Runtime/Engine/Private/AI/Navigation/AvoidanceManager.cpp
--- a/Engine/Source/Runtime/Engine/Private/AI/Navigation/AvoidanceManager.cpp
+++ b/Engine/Source/Runtime/Engine/Private/AI/Navigation/AvoidanceManager.cpp
@@ -2,6 +2,7 @@
 
 #include "EnginePrivate.h"
 #include "AI/Navigation/AvoidanceManager.h"
+#include "GameFramework/Character.h"
 
 DEFINE_STAT(STAT_AI_ObstacleAvoidance);
 
diff --git a/Engine/Source/Runtime/Engine/Private/Animation/AnimInstance.cpp b/Engine/Source/Runtime/Engine/Private/Animation/AnimInstance.cpp
--- a/Engine/Source/Runtime/Engine/Private/Animation/AnimInstance.cpp
+++ b/Engine/Source/Runtime/Engine/Private/Animation/AnimInstance.cpp
@@ -9,6 +9,7 @@
 #include "AnimationUtils.h"
 #include "AnimTree.h"
 #include "Animation/AnimNode_StateMachine.h"
+#include "GameFramework/Character.h"
 #include "ParticleDefinitions.h"
 #include "DisplayDebugHelpers.h"
 
diff --git a/Engine/Source/Runtime/Engine/Private/Animation/AnimMontage.cpp b/Engine/Source/Runtime/Engine/Private/Animation/AnimMontage.cpp
--- a/Engine/Source/Runtime/Engine/Private/Animation/AnimMontage.cpp
+++ b/Engine/Source/Runtime/Engine/Private/Animation/AnimMontage.cpp
@@ -7,6 +7,7 @@
 #include "EnginePrivate.h"
 #include "AnimationUtils.h"
 #include "AnimationRuntime.h"
+#include "GameFramework/CharacterMovementComponent.h"
 
 ///////////////////////////////////////////////////////////////////////////
 //
diff --git a/Engine/Source/Runtime/Engine/Private/Animation/AnimNode_WheelHandler.cpp b/Engine/Source/Runtime/Engine/Private/Animation/AnimNode_WheelHandler.cpp
--- a/Engine/Source/Runtime/Engine/Private/Animation/AnimNode_WheelHandler.cpp
+++ b/Engine/Source/Runtime/Engine/Private/Animation/AnimNode_WheelHandler.cpp
@@ -3,6 +3,8 @@
 #include "EnginePrivate.h"
 #include "AnimationRuntime.h"
 #include "Animation/BoneControllers/AnimNode_WheelHandler.h"
+#include "GameFramework/WheeledVehicle.h"
+#include "Vehicles/WheeledVehicleMovementComponent.h"
 
 /////////////////////////////////////////////////////
 // FAnimNode_WheelHandler
diff --git a/Engine/Source/Runtime/Engine/Private/Character.cpp b/Engine/Source/Runtime/Engine/Private/Character.cpp
--- a/Engine/Source/Runtime/Engine/Private/Character.cpp
+++ b/Engine/Source/Runtime/Engine/Private/Character.cpp
@@ -6,6 +6,7 @@
 =============================================================================*/
 
 #include "EnginePrivate.h"
+#include "GameFramework/Character.h"
 #include "Net/UnrealNetwork.h"
 #include "DisplayDebugHelpers.h"
 
diff --git a/Engine/Source/Runtime/Engine/Private/CharacterMovementComponent.cpp b/Engine/Source/Runtime/Engine/Private/CharacterMovementComponent.cpp
--- a/Engine/Source/Runtime/Engine/Private/CharacterMovementComponent.cpp
+++ b/Engine/Source/Runtime/Engine/Private/CharacterMovementComponent.cpp
@@ -8,6 +8,8 @@
 #include "EnginePrivate.h"
 #include "GameFramework/PhysicsVolume.h"
 #include "GameFramework/GameNetworkManager.h"
+#include "GameFramework/Character.h"
+
 // @todo this is here only due to circular dependency to AIModule. To be removed
 #include "Navigation/PathFollowingComponent.h"
 
diff --git a/Engine/Source/Runtime/Engine/Private/Controller.cpp b/Engine/Source/Runtime/Engine/Private/Controller.cpp
--- a/Engine/Source/Runtime/Engine/Private/Controller.cpp
+++ b/Engine/Source/Runtime/Engine/Private/Controller.cpp
@@ -9,6 +9,7 @@
 #include "ConfigCacheIni.h"
 #include "NetworkingDistanceConstants.h"
 #include "VisualLog.h"
+#include "GameFramework/Character.h"
 
 // @todo this is here only due to circular dependency to AIModule. To be removed
 #include "Navigation/PathFollowingComponent.h"
diff --git a/Engine/Source/Runtime/Engine/Private/Debug/GameplayDebuggingHUDComponent.cpp b/Engine/Source/Runtime/Engine/Private/Debug/GameplayDebuggingHUDComponent.cpp
--- a/Engine/Source/Runtime/Engine/Private/Debug/GameplayDebuggingHUDComponent.cpp
+++ b/Engine/Source/Runtime/Engine/Private/Debug/GameplayDebuggingHUDComponent.cpp
@@ -6,6 +6,7 @@
 
 #include "EnginePrivate.h"
 #include "Debug/GameplayDebuggingComponent.h"
+#include "Debug/GameplayDebuggingHUDComponent.h"
 #include "Net/UnrealNetwork.h"
 
 DEFINE_LOG_CATEGORY_STATIC(LogHUD, Log, All);
diff --git a/Engine/Source/Runtime/Engine/Private/DebugCameraController.cpp b/Engine/Source/Runtime/Engine/Private/DebugCameraController.cpp
--- a/Engine/Source/Runtime/Engine/Private/DebugCameraController.cpp
+++ b/Engine/Source/Runtime/Engine/Private/DebugCameraController.cpp
@@ -6,6 +6,8 @@
 =============================================================================*/
 
 #include "EnginePrivate.h"
+#include "GameFramework/SpectatorPawn.h"
+#include "GameFramework/SpectatorPawnMovement.h"
 #include "Engine/DebugCameraHUD.h"
 
 /** The currently selected actor. */
diff --git a/Engine/Source/Runtime/Engine/Private/DebugCameraHUD.cpp b/Engine/Source/Runtime/Engine/Private/DebugCameraHUD.cpp
--- a/Engine/Source/Runtime/Engine/Private/DebugCameraHUD.cpp
+++ b/Engine/Source/Runtime/Engine/Private/DebugCameraHUD.cpp
@@ -5,6 +5,7 @@
 =============================================================================*/
 
 #include "EnginePrivate.h"
+#include "GameFramework/SpectatorPawn.h"
 #include "Engine/DebugCameraHUD.h"
 
 // ------------------
diff --git a/Engine/Source/Runtime/Engine/Private/DefaultPawn.cpp b/Engine/Source/Runtime/Engine/Private/DefaultPawn.cpp
--- a/Engine/Source/Runtime/Engine/Private/DefaultPawn.cpp
+++ b/Engine/Source/Runtime/Engine/Private/DefaultPawn.cpp
@@ -2,6 +2,7 @@
 
 
 #include "EnginePrivate.h"
+#include "GameFramework/DefaultPawn.h"
 
 FName ADefaultPawn::MovementComponentName(TEXT("MovementComponent0"));
 FName ADefaultPawn::CollisionComponentName(TEXT("CollisionComponent0"));
diff --git a/Engine/Source/Runtime/Engine/Private/DrawDebugHelpers.cpp b/Engine/Source/Runtime/Engine/Private/DrawDebugHelpers.cpp
--- a/Engine/Source/Runtime/Engine/Private/DrawDebugHelpers.cpp
+++ b/Engine/Source/Runtime/Engine/Private/DrawDebugHelpers.cpp
@@ -5,7 +5,7 @@
  */
 
 #include "EnginePrivate.h"
-
+#include "GameFramework/HUD.h"
 
 void FlushPersistentDebugLines( const UWorld* InWorld )
 {
diff --git a/Engine/Source/Runtime/Engine/Private/EnginePrivate.h b/Engine/Source/Runtime/Engine/Private/EnginePrivate.h
--- a/Engine/Source/Runtime/Engine/Private/EnginePrivate.h
+++ b/Engine/Source/Runtime/Engine/Private/EnginePrivate.h
@@ -44,8 +44,6 @@
 #include "Engine/DecalActor.h"
 #include "PhysicsEngine/DestructibleActor.h"
 #include "Engine/Canvas.h"
-#include "GameFramework/HUD.h"
-#include "Debug/GameplayDebuggingHUDComponent.h"
 #include "Atmosphere/AtmosphericFog.h"
 #include "Engine/ExponentialHeightFog.h"
 #include "GameFramework/GameSession.h"
@@ -80,17 +78,8 @@
 #include "AI/Navigation/NavLinkProxy.h"
 #include "Engine/Note.h"
 #include "Components/PrimitiveComponent.h"
-#include "GameFramework/MovementComponent.h"
-#include "GameFramework/NavMovementComponent.h"
-#include "GameFramework/PawnMovementComponent.h"
 #include "GameFramework/Pawn.h"
 #include "Interfaces/NetworkPredictionInterface.h"
-#include "GameFramework/CharacterMovementComponent.h"
-#include "GameFramework/Character.h"
-#include "GameFramework/FloatingPawnMovement.h"
-#include "GameFramework/DefaultPawn.h"
-#include "GameFramework/SpectatorPawn.h"
-#include "GameFramework/WheeledVehicle.h"
 #include "Engine/ReflectionCapture.h"
 #include "Engine/BoxReflectionCapture.h"
 #include "Engine/PlaneReflectionCapture.h"
@@ -117,12 +106,6 @@
 #include "Components/ApplicationLifecycleComponent.h"
 #include "Engine/DataAsset.h"
 #include "Debug/GameplayDebuggingControllerComponent.h"
-#include "GameFramework/SpectatorPawnMovement.h"
-#include "Vehicles/VehicleWheel.h"
-#include "Vehicles/WheeledVehicleMovementComponent.h"
-#include "Vehicles/WheeledVehicleMovementComponent4W.h"
-#include "GameFramework/ProjectileMovementComponent.h"
-#include "GameFramework/RotatingMovementComponent.h"
 #include "Components/PawnNoiseEmitterComponent.h"
 #include "PhysicsEngine/PhysicsHandleComponent.h"
 #include "Atmosphere/AtmosphericFogComponent.h"
@@ -153,7 +136,6 @@
 #include "Components/StaticMeshComponent.h"
 #include "Components/InstancedStaticMeshComponent.h"
 #include "Interfaces/Interface_CollisionDataProvider.h"
-#include "Components/SplineMeshComponent.h"
 #include "Components/ModelComponent.h"
 #include "AI/Navigation/NavLinkRenderingComponent.h"
 #include "AI/Navigation/NavMeshRenderingComponent.h"
@@ -189,7 +171,6 @@
 #include "Animation/BlendSpace1D.h"
 #include "Animation/AimOffsetBlendSpace1D.h"
 #include "Animation/AnimSingleNodeInstance.h"
-#include "Vehicles/VehicleAnimInstance.h"
 #include "Animation/AnimNotifies/AnimNotify.h"
 #include "Animation/AnimNotifies/AnimNotifyState.h"
 #include "Animation/AnimNotifies/AnimNotifyState_TimedParticleEffect.h"
@@ -230,7 +211,6 @@
 #include "Curves/CurveEdPresetCurve.h"
 #include "Engine/CurveTable.h"
 #include "GameFramework/DamageType.h"
-#include "Vehicles/TireType.h"
 #include "Engine/DataTable.h"
 #include "Sound/DialogueVoice.h"
 #include "Sound/DialogueWave.h"
diff --git a/Engine/Source/Runtime/Engine/Private/FloatingPawnMovement.cpp b/Engine/Source/Runtime/Engine/Private/FloatingPawnMovement.cpp
--- a/Engine/Source/Runtime/Engine/Private/FloatingPawnMovement.cpp
+++ b/Engine/Source/Runtime/Engine/Private/FloatingPawnMovement.cpp
@@ -1,7 +1,7 @@
 // Copyright 1998-2014 Epic Games, Inc. All Rights Reserved.
 
 #include "EnginePrivate.h"
-
+#include "GameFramework/FloatingPawnMovement.h"
 
 UFloatingPawnMovement::UFloatingPawnMovement(const class FPostConstructInitializeProperties& PCIP)
 	: Super(PCIP)
diff --git a/Engine/Source/Runtime/Engine/Private/GameViewportClient.cpp b/Engine/Source/Runtime/Engine/Private/GameViewportClient.cpp
--- a/Engine/Source/Runtime/Engine/Private/GameViewportClient.cpp
+++ b/Engine/Source/Runtime/Engine/Private/GameViewportClient.cpp
@@ -2,12 +2,12 @@
 
 #include "EnginePrivate.h"
 #include "Engine/Console.h"
+#include "GameFramework/HUD.h"
 #include "ParticleDefinitions.h"
 #include "SoundDefinitions.h"
 #include "FXSystem.h"
 #include "SubtitleManager.h"
 #include "ImageUtils.h"
-
 #include "RenderCore.h"
 #include "ColorList.h"
 #include "Slate.h"
diff --git a/Engine/Source/Runtime/Engine/Private/GameplayStatics.cpp b/Engine/Source/Runtime/Engine/Private/GameplayStatics.cpp
--- a/Engine/Source/Runtime/Engine/Private/GameplayStatics.cpp
+++ b/Engine/Source/Runtime/Engine/Private/GameplayStatics.cpp
@@ -8,6 +8,7 @@
 #include "IForceFeedbackSystem.h"
 #include "Slate.h"
 #include "Particles/ParticleSystemComponent.h"
+#include "GameFramework/Character.h"
 
 //////////////////////////////////////////////////////////////////////////
 // UGameplayStatics
diff --git a/Engine/Source/Runtime/Engine/Private/HUD.cpp b/Engine/Source/Runtime/Engine/Private/HUD.cpp
--- a/Engine/Source/Runtime/Engine/Private/HUD.cpp
+++ b/Engine/Source/Runtime/Engine/Private/HUD.cpp
@@ -5,6 +5,7 @@
 =============================================================================*/
 
 #include "EnginePrivate.h"
+#include "GameFramework/HUD.h"
 #include "Net/UnrealNetwork.h"
 #include "MessageLog.h"
 #include "UObjectToken.h"
diff --git a/Engine/Source/Runtime/Engine/Private/Pawn.cpp b/Engine/Source/Runtime/Engine/Private/Pawn.cpp
--- a/Engine/Source/Runtime/Engine/Private/Pawn.cpp
+++ b/Engine/Source/Runtime/Engine/Private/Pawn.cpp
@@ -7,6 +7,7 @@
 =============================================================================*/
 
 #include "EnginePrivate.h"
+#include "GameFramework/Character.h"
 #include "Net/UnrealNetwork.h"
 #include "ConfigCacheIni.h"
 #include "ParticleDefinitions.h"
diff --git a/Engine/Source/Runtime/Engine/Private/PhysicsEngine/PhysicalMaterial.cpp b/Engine/Source/Runtime/Engine/Private/PhysicsEngine/PhysicalMaterial.cpp
--- a/Engine/Source/Runtime/Engine/Private/PhysicsEngine/PhysicalMaterial.cpp
+++ b/Engine/Source/Runtime/Engine/Private/PhysicsEngine/PhysicalMaterial.cpp
@@ -6,6 +6,7 @@
 
 #include "EnginePrivate.h"
 #include "PhysicsPublic.h"
+#include "Vehicles/TireType.h"
 
 #if WITH_PHYSX
 	#include "PhysXSupport.h"
diff --git a/Engine/Source/Runtime/Engine/Private/Player.cpp b/Engine/Source/Runtime/Engine/Private/Player.cpp
--- a/Engine/Source/Runtime/Engine/Private/Player.cpp
+++ b/Engine/Source/Runtime/Engine/Private/Player.cpp
@@ -8,6 +8,7 @@
 
 #include "SubtitleManager.h"
 #include "Net/UnrealNetwork.h"
+#include "GameFramework/HUD.h"
 
 #include "RenderCore.h"
 #include "ColorList.h"
diff --git a/Engine/Source/Runtime/Engine/Private/ProjectileMovementComponent.cpp b/Engine/Source/Runtime/Engine/Private/ProjectileMovementComponent.cpp
--- a/Engine/Source/Runtime/Engine/Private/ProjectileMovementComponent.cpp
+++ b/Engine/Source/Runtime/Engine/Private/ProjectileMovementComponent.cpp
@@ -1,6 +1,7 @@
 // Copyright 1998-2014 Epic Games, Inc. All Rights Reserved.
 
 #include "EnginePrivate.h"
+#include "GameFramework/ProjectileMovementComponent.h"
 
 DEFINE_LOG_CATEGORY_STATIC(LogProjectileMovement, Log, All);
 
diff --git a/Engine/Source/Runtime/Engine/Private/RotatingMovementComponent.cpp b/Engine/Source/Runtime/Engine/Private/RotatingMovementComponent.cpp
--- a/Engine/Source/Runtime/Engine/Private/RotatingMovementComponent.cpp
+++ b/Engine/Source/Runtime/Engine/Private/RotatingMovementComponent.cpp
@@ -1,7 +1,7 @@
 // Copyright 1998-2014 Epic Games, Inc. All Rights Reserved.
 
 #include "EnginePrivate.h"
-
+#include "GameFramework/RotatingMovementComponent.h"
 
 URotatingMovementComponent::URotatingMovementComponent(const class FPostConstructInitializeProperties& PCIP)
 	: Super(PCIP)
diff --git a/Engine/Source/Runtime/Engine/Private/SkeletalMeshComponent.cpp b/Engine/Source/Runtime/Engine/Private/SkeletalMeshComponent.cpp
--- a/Engine/Source/Runtime/Engine/Private/SkeletalMeshComponent.cpp
+++ b/Engine/Source/Runtime/Engine/Private/SkeletalMeshComponent.cpp
@@ -16,7 +16,7 @@
 #include "AnimTree.h"
 #include "Animation/AnimNodeBase.h"
 #include "Animation/VertexAnim/VertexAnimation.h"
-
+#include "GameFramework/Character.h"
 #include "Particles/ParticleSystemComponent.h"
 
 #ifndef EXPERIMENTAL_PARALLEL_CODE  
diff --git a/Engine/Source/Runtime/Engine/Private/SpectatorPawn.cpp b/Engine/Source/Runtime/Engine/Private/SpectatorPawn.cpp
--- a/Engine/Source/Runtime/Engine/Private/SpectatorPawn.cpp
+++ b/Engine/Source/Runtime/Engine/Private/SpectatorPawn.cpp
@@ -2,7 +2,8 @@
 
 
 #include "EnginePrivate.h"
-
+#include "GameFramework/SpectatorPawn.h"
+#include "GameFramework/SpectatorPawnMovement.h"
 
 ASpectatorPawn::ASpectatorPawn(const class FPostConstructInitializeProperties& PCIP)
 	: Super(PCIP
diff --git a/Engine/Source/Runtime/Engine/Private/SpectatorPawnMovement.cpp b/Engine/Source/Runtime/Engine/Private/SpectatorPawnMovement.cpp
--- a/Engine/Source/Runtime/Engine/Private/SpectatorPawnMovement.cpp
+++ b/Engine/Source/Runtime/Engine/Private/SpectatorPawnMovement.cpp
@@ -6,7 +6,7 @@
 =============================================================================*/
 
 #include "EnginePrivate.h"
-
+#include "GameFramework/SpectatorPawnMovement.h"
 
 USpectatorPawnMovement::USpectatorPawnMovement(const class FPostConstructInitializeProperties& PCIP)
 	: Super(PCIP)
diff --git a/Engine/Source/Runtime/Engine/Private/StaticMeshComponent.cpp b/Engine/Source/Runtime/Engine/Private/StaticMeshComponent.cpp
--- a/Engine/Source/Runtime/Engine/Private/StaticMeshComponent.cpp
+++ b/Engine/Source/Runtime/Engine/Private/StaticMeshComponent.cpp
@@ -2,6 +2,7 @@
 
 #include "EnginePrivate.h"
 #include "StaticMeshResources.h"
+#include "Components/SplineMeshComponent.h"
 #include "Net/UnrealNetwork.h"
 #include "MessageLog.h"
 #include "UObjectToken.h"
diff --git a/Engine/Source/Runtime/Engine/Private/Vehicles/PhysXVehicleManager.cpp b/Engine/Source/Runtime/Engine/Private/Vehicles/PhysXVehicleManager.cpp
--- a/Engine/Source/Runtime/Engine/Private/Vehicles/PhysXVehicleManager.cpp
+++ b/Engine/Source/Runtime/Engine/Private/Vehicles/PhysXVehicleManager.cpp
@@ -2,6 +2,8 @@
 
 #include "EnginePrivate.h"
 #include "PhysicsPublic.h"
+#include "Vehicles/TireType.h"
+#include "Vehicles/WheeledVehicleMovementComponent.h"
 
 #if WITH_VEHICLE
 
diff --git a/Engine/Source/Runtime/Engine/Private/Vehicles/TireType.cpp b/Engine/Source/Runtime/Engine/Private/Vehicles/TireType.cpp
--- a/Engine/Source/Runtime/Engine/Private/Vehicles/TireType.cpp
+++ b/Engine/Source/Runtime/Engine/Private/Vehicles/TireType.cpp
@@ -1,6 +1,7 @@
 // Copyright 1998-2014 Epic Games, Inc. All Rights Reserved.
 
 #include "EnginePrivate.h"
+#include "Vehicles/TireType.h"
 
 #if WITH_PHYSX
 	#include "PhysXVehicleManager.h"
diff --git a/Engine/Source/Runtime/Engine/Private/Vehicles/VehicleAnimInstance.cpp b/Engine/Source/Runtime/Engine/Private/Vehicles/VehicleAnimInstance.cpp
--- a/Engine/Source/Runtime/Engine/Private/Vehicles/VehicleAnimInstance.cpp
+++ b/Engine/Source/Runtime/Engine/Private/Vehicles/VehicleAnimInstance.cpp
@@ -6,6 +6,8 @@
 =============================================================================*/ 
 
 #include "EnginePrivate.h"
+#include "Vehicles/VehicleAnimInstance.h"
+#include "GameFramework/WheeledVehicle.h"
 #include "AnimationRuntime.h"
 
 /////////////////////////////////////////////////////
diff --git a/Engine/Source/Runtime/Engine/Private/Vehicles/VehicleWheel.cpp b/Engine/Source/Runtime/Engine/Private/Vehicles/VehicleWheel.cpp
--- a/Engine/Source/Runtime/Engine/Private/Vehicles/VehicleWheel.cpp
+++ b/Engine/Source/Runtime/Engine/Private/Vehicles/VehicleWheel.cpp
@@ -2,6 +2,9 @@
 
 #include "EnginePrivate.h"
 #include "PhysicsPublic.h"
+#include "Vehicles/VehicleWheel.h"
+#include "Vehicles/WheeledVehicleMovementComponent.h"
+#include "Vehicles/TireType.h"
 
 #if WITH_PHYSX
 #include "../PhysicsEngine/PhysXSupport.h"
diff --git a/Engine/Source/Runtime/Engine/Private/Vehicles/WheeledVehicleMovementComponent.cpp b/Engine/Source/Runtime/Engine/Private/Vehicles/WheeledVehicleMovementComponent.cpp
--- a/Engine/Source/Runtime/Engine/Private/Vehicles/WheeledVehicleMovementComponent.cpp
+++ b/Engine/Source/Runtime/Engine/Private/Vehicles/WheeledVehicleMovementComponent.cpp
@@ -4,6 +4,9 @@
 #include "PhysicsPublic.h"
 #include "Net/UnrealNetwork.h"
 #include "MessageLog.h"
+#include "Vehicles/VehicleWheel.h"
+#include "Vehicles/WheeledVehicleMovementComponent.h"
+#include "Vehicles/TireType.h"
 
 #if WITH_PHYSX
 #include "../PhysicsEngine/PhysXSupport.h"
diff --git a/Engine/Source/Runtime/Engine/Private/Vehicles/WheeledVehicleMovementComponent4W.cpp b/Engine/Source/Runtime/Engine/Private/Vehicles/WheeledVehicleMovementComponent4W.cpp
--- a/Engine/Source/Runtime/Engine/Private/Vehicles/WheeledVehicleMovementComponent4W.cpp
+++ b/Engine/Source/Runtime/Engine/Private/Vehicles/WheeledVehicleMovementComponent4W.cpp
@@ -2,6 +2,7 @@
 
 #include "EnginePrivate.h"
 #include "PhysicsPublic.h"
+#include "Vehicles/WheeledVehicleMovementComponent4W.h"
 
 #if WITH_VEHICLE
 #include "../PhysicsEngine/PhysXSupport.h"
diff --git a/Engine/Source/Runtime/Engine/Private/WheeledVehicle.cpp b/Engine/Source/Runtime/Engine/Private/WheeledVehicle.cpp
--- a/Engine/Source/Runtime/Engine/Private/WheeledVehicle.cpp
+++ b/Engine/Source/Runtime/Engine/Private/WheeledVehicle.cpp
@@ -6,6 +6,8 @@
 =============================================================================*/
 
 #include "EnginePrivate.h"
+#include "GameFramework/WheeledVehicle.h"
+#include "Vehicles/WheeledVehicleMovementComponent4W.h"
 #include "DisplayDebugHelpers.h"
 
 FName AWheeledVehicle::VehicleMovementComponentName(TEXT("MovementComp"));
