AIModule interface introduced. UWorld now contains a UAISystemBase property. #UE4

- includes further removal of AI dependencies from Engine code

diff --git a/Engine/Config/BaseEngine.ini b/Engine/Config/BaseEngine.ini
--- a/Engine/Config/BaseEngine.ini
+++ b/Engine/Config/BaseEngine.ini
@@ -32,8 +32,6 @@ LocalPlayerClassName=/Script/Engine.LocalPlayer
 ; If you change to your WorldSettings class, the previous levels won't replace first actor - WorldSettings
 WorldSettingsClassName=/Script/Engine.WorldSettings
 NavigationSystemClassName=/Script/Engine.NavigationSystem
-BehaviorTreeManagerClassName=/Script/Engine.BehaviorTreeManager
-EnvironmentQueryManagerClassName=/Script/Engine.EnvQueryManager
 AvoidanceManagerClassName=/Script/Engine.AvoidanceManager
 PhysicsCollisionHandlerClassName=/Script/Engine.PhysicsCollisionHandler
 LevelActorClassName=/Script/Engine.LevelActor
@@ -2152,3 +2150,7 @@ TilesStreamingTimeThreshold=1.0
 [/Script/WindowsTargetPlatform.WindowsTargetSettings]
 +TargetedRHIs=PCD3D_SM5
 +TargetedRHIs=PCD3D_SM4
+
+[/Script/Engine.AISystemBase]
+AISystemModuleName=AIModule
+AISystemClassName=/Script/AIModule.AISystem
\ No newline at end of file
diff --git a/Engine/Source/Runtime/AIModule/Classes/AIController.h b/Engine/Source/Runtime/AIModule/Classes/AIController.h
--- a/Engine/Source/Runtime/AIModule/Classes/AIController.h
+++ b/Engine/Source/Runtime/AIModule/Classes/AIController.h
@@ -64,7 +64,7 @@ struct FFocusKnowledge
 	TArray<struct FFocusItem> Priorities;
 };
 
-UCLASS(BlueprintType, Blueprintable, dependson=(UPathFollowingComponent))
+UCLASS(BlueprintType, Blueprintable)
 class AIMODULE_API AAIController : public AController
 {
 	GENERATED_UCLASS_BODY()
diff --git a/Engine/Source/Runtime/AIModule/Classes/AISystem.h b/Engine/Source/Runtime/AIModule/Classes/AISystem.h
--- a/Engine/Source/Runtime/AIModule/Classes/AISystem.h
+++ b/Engine/Source/Runtime/AIModule/Classes/AISystem.h
@@ -1,16 +1,79 @@
 // Copyright 1998-2014 Epic Games, Inc. All Rights Reserved.
 
 #pragma once
-
+#include "Misc/CoreMisc.h"
+#include "AI/AISystemBase.h"
 #include "AISystem.generated.h"
 
+class UWorld;
+class UBehaviorTreeManager;
+class UEnvQueryManager;
+
 UCLASS(config=Engine)
-class AIMODULE_API UAISystem : public UBlueprintFunctionLibrary
+class AIMODULE_API UAISystem : public UAISystemBase
 {
 	GENERATED_UCLASS_BODY()
 
+	/** Behavior tree manager used by game */
+	UPROPERTY(Transient)
+	UBehaviorTreeManager* BehaviorTreeManager;
+
+	/** Environment query manager used by game */
+	UPROPERTY(Transient)
+	UEnvQueryManager* EnvironmentQueryManager;
+
+	// IAISystemInterface begin	
+	//virtual IAISystemInterface* CreateInstanceForWorldImpl(UWorld* World);
+	virtual void InitializeActorsForPlay(bool bTimeGotReset) OVERRIDE;
+	virtual void WorldOriginChanged(FIntPoint OldOrigin, FIntPoint NewOrigin) OVERRIDE;
+	virtual void CleanupWorld(bool bSessionEnded = true, bool bCleanupResources = true, UWorld* NewWorld = NULL) OVERRIDE;
+	// IAISystemInterface end
+	
+	/** Behavior tree manager getter */
+	FORCEINLINE UBehaviorTreeManager* GetBehaviorTreeManager() { return BehaviorTreeManager; }
+	/** Behavior tree manager const getter */
+	FORCEINLINE const UBehaviorTreeManager* GetBehaviorTreeManager() const { return BehaviorTreeManager; }
+
+	/** Behavior tree manager getter */
+	FORCEINLINE UEnvQueryManager* GetEnvironmentQueryManager() { return EnvironmentQueryManager; }
+	/** Behavior tree manager const getter */
+	FORCEINLINE const UEnvQueryManager* GetEnvironmentQueryManager() const { return EnvironmentQueryManager; }
+
+	FORCEINLINE static UAISystem* GetCurrent(UWorld* World) 
+	{ 
+		check(World != NULL)
+		check(World->GetAISystem() == NULL || Cast<UAISystem>(World->GetAISystem()) != NULL);
+		return (UAISystem*)(World->GetAISystem());
+	}
+
+	FORCEINLINE static UBehaviorTreeManager* GetCurrentBTManager(UWorld* World)
+	{
+		UAISystem* AISys = GetCurrent(World);
+		// if AI system is NULL you're probably running your AI code on a client
+		check(AISys != NULL);
+		return AISys->GetBehaviorTreeManager();
+	}
+
+	FORCEINLINE static UEnvQueryManager* GetCurrentEQSManager(UWorld* World)
+	{
+		UAISystem* AISys = GetCurrent(World);
+		// if AI system is NULL you're probably running your AI code on a client
+		check(AISys != NULL);
+		return AISys->GetEnvironmentQueryManager();
+	}
+
+	FORCEINLINE UWorld* GetOuterWorld() { return Cast<UWorld>(GetOuter()); }
+	
 	//----------------------------------------------------------------------//
-	// BP functions
+	// cheats
 	//----------------------------------------------------------------------//
-	
+	UFUNCTION(exec)
+	virtual void AIIgnorePlayers();
+
+	UFUNCTION(exec)
+	virtual void AILoggingVerbose();
+
+	/** insta-runs EQS query for GameplayDebugComponent selected AI */
+	UFUNCTION(exec)
+	void RunEQS(const FString& QueryName);
 };
\ No newline at end of file
diff --git a/Engine/Source/Runtime/AIModule/Classes/BehaviorTree/BehaviorTree.h b/Engine/Source/Runtime/AIModule/Classes/BehaviorTree/BehaviorTree.h
--- a/Engine/Source/Runtime/AIModule/Classes/BehaviorTree/BehaviorTree.h
+++ b/Engine/Source/Runtime/AIModule/Classes/BehaviorTree/BehaviorTree.h
@@ -2,9 +2,10 @@
 #pragma once
 
 #include "BTCompositeNode.h"
+#include "BTTaskNode.h"
 #include "BehaviorTree.generated.h"
 
-UCLASS(DependsOn = (UBTCompositeNode, UBTTaskNode), BlueprintType)
+UCLASS(BlueprintType)
 class AIMODULE_API UBehaviorTree : public UObject
 {
 	GENERATED_UCLASS_BODY()
diff --git a/Engine/Source/Runtime/AIModule/Classes/Blueprint/AIAsyncTaskBlueprintProxy.h b/Engine/Source/Runtime/AIModule/Classes/Blueprint/AIAsyncTaskBlueprintProxy.h
--- a/Engine/Source/Runtime/AIModule/Classes/Blueprint/AIAsyncTaskBlueprintProxy.h
+++ b/Engine/Source/Runtime/AIModule/Classes/Blueprint/AIAsyncTaskBlueprintProxy.h
@@ -1,6 +1,7 @@
 // Copyright 1998-2014 Epic Games, Inc. All Rights Reserved.
 #pragma once
 
+#include "AI/Navigation/NavLinkDefinition.h"
 #include "Navigation/PathFollowingComponent.h"
 #include "AIAsyncTaskBlueprintProxy.generated.h"
 
diff --git a/Engine/Source/Runtime/AIModule/Classes/EnvironmentQuery/EQSTestingPawn.h b/Engine/Source/Runtime/AIModule/Classes/EnvironmentQuery/EQSTestingPawn.h
--- a/Engine/Source/Runtime/AIModule/Classes/EnvironmentQuery/EQSTestingPawn.h
+++ b/Engine/Source/Runtime/AIModule/Classes/EnvironmentQuery/EQSTestingPawn.h
@@ -4,12 +4,13 @@
 
 #include "GameFramework/Character.h"
 #include "EnvironmentQuery/EQSQueryResultSourceInterface.h"
+#include "EnvironmentQuery/EnvQueryTypes.h"
 #include "EQSTestingPawn.generated.h"
 
 /** this class is abstract even though it's perfectly functional on its own.
  *	The reason is to stop it from showing as valid player pawn type when configuring 
  *	project's game mode. */
-UCLASS(abstract, hidecategories=(Advanced, Attachment, Collision, Animation),DependsOn=(UEnvQueryTypes))
+UCLASS(abstract, hidecategories=(Advanced, Attachment, Collision, Animation))
 class AIMODULE_API AEQSTestingPawn : public ACharacter, public IEQSQueryResultSourceInterface
 {
 	GENERATED_UCLASS_BODY()
diff --git a/Engine/Source/Runtime/AIModule/Classes/Navigation/NavigationComponent.h b/Engine/Source/Runtime/AIModule/Classes/Navigation/NavigationComponent.h
--- a/Engine/Source/Runtime/AIModule/Classes/Navigation/NavigationComponent.h
+++ b/Engine/Source/Runtime/AIModule/Classes/Navigation/NavigationComponent.h
@@ -1,10 +1,13 @@
 // Copyright 1998-2014 Epic Games, Inc. All Rights Reserved.
 
 #pragma once
+#include "AI/Navigation/NavAgentInterface.h"
+#include "AI/Navigation/NavPathObserverInterface.h"
 #include "AI/Navigation/NavigationPathGenerator.h"
+#include "AI/Navigation/NavigationSystem.h"
 #include "NavigationComponent.generated.h"
 
-UCLASS(dependson=(UNavAgentInterface, UNavPathObserverInterface, UNavigationSystem), config=Engine)
+UCLASS(config=Engine)
 class AIMODULE_API UNavigationComponent : public UActorComponent, public INavigationPathGenerator
 {
 	GENERATED_UCLASS_BODY()
diff --git a/Engine/Source/Runtime/AIModule/Private/AIModule.cpp b/Engine/Source/Runtime/AIModule/Private/AIModule.cpp
--- a/Engine/Source/Runtime/AIModule/Private/AIModule.cpp
+++ b/Engine/Source/Runtime/AIModule/Private/AIModule.cpp
@@ -1,12 +1,15 @@
 // Copyright 1998-2014 Epic Games, Inc. All Rights Reserved.
 
 #include "AIModulePrivate.h"
+#include "AISystem.h"
 
 class FAIModule : public IAIModule
 {
 	// Begin IModuleInterface
 	virtual void StartupModule() OVERRIDE;
 	virtual void ShutdownModule() OVERRIDE;
+	
+	virtual UAISystemBase* CreateAISystemInstance(UWorld* World) OVERRIDE;
 	// End IModuleInterface
 };
 
@@ -22,3 +25,9 @@ void FAIModule::ShutdownModule()
 	// This function may be called during shutdown to clean up your module.  For modules that support dynamic reloading,
 	// we call this function before unloading the module.
 }
+
+UAISystemBase* FAIModule::CreateAISystemInstance(UWorld* World)
+{
+	TSubclassOf<class UAISystemBase> AISystemClass = LoadClass<UAISystemBase>(NULL, *UAISystem::GetAISystemClassName().ClassName, NULL, LOAD_None, NULL);
+	return ConstructObject<UAISystemBase>(AISystemClass, World);
+}
\ No newline at end of file
diff --git a/Engine/Source/Runtime/AIModule/Private/AIModulePrivate.h b/Engine/Source/Runtime/AIModule/Private/AIModulePrivate.h
--- a/Engine/Source/Runtime/AIModule/Private/AIModulePrivate.h
+++ b/Engine/Source/Runtime/AIModule/Private/AIModulePrivate.h
@@ -2,13 +2,6 @@
 
 #pragma once
 
-
-
-#if WITH_EDITOR
-#include "UnrealEd.h"
-#endif
-
-#include "Core.h"
 //#include "CoreUObject.h"
 //#include "Engine/EngineBaseTypes.h"
 //#include "Engine/EngineTypes.h"
@@ -21,6 +14,7 @@
 #include "AIModule.h"
 
 #include "AITypes.h"
+#include "AISystem.h"
 #include "EnvironmentQuery/EnvQueryTypes.h"
 #include "BehaviorTree/Blackboard/BlackboardKeyType.h"
 #include "AIController.h"
diff --git a/Engine/Source/Runtime/AIModule/Private/AISystem.cpp b/Engine/Source/Runtime/AIModule/Private/AISystem.cpp
--- a/Engine/Source/Runtime/AIModule/Private/AISystem.cpp
+++ b/Engine/Source/Runtime/AIModule/Private/AISystem.cpp
@@ -1,8 +1,109 @@
 // Copyright 1998-2014 Epic Games, Inc. All Rights Reserved.
 
 #include "AIModulePrivate.h"
+#include "BehaviorTree/BehaviorTreeManager.h"
+#include "EnvironmentQuery/EnvQueryManager.h"
 #include "AISystem.h"
 
-UAISystem::UAISystem(const class FPostConstructInitializeProperties& PCIP) : Super(PCIP)
+UAISystem::UAISystem(const FPostConstructInitializeProperties& PCIP) : Super(PCIP)
 {
+	if (HasAnyFlags(RF_ClassDefaultObject) == false)
+	{
+		UWorld* WorldOuter = Cast<UWorld>(GetOuter());
+		UObject* ManagersOuter = WorldOuter != NULL ? (UObject*)WorldOuter : (UObject*)this;
+		BehaviorTreeManager = NewObject<UBehaviorTreeManager>(ManagersOuter);
+		EnvironmentQueryManager = NewObject<UEnvQueryManager>(ManagersOuter);
+	}
+}
+
+void UAISystem::InitializeActorsForPlay(bool bTimeGotReset)
+{
+
+}
+
+void UAISystem::WorldOriginChanged(FIntPoint OldOrigin, FIntPoint NewOrigin)
+{
+
+}
+
+void UAISystem::CleanupWorld(bool bSessionEnded, bool bCleanupResources, UWorld* NewWorld)
+{
+
+}
+
+void UAISystem::AIIgnorePlayers()
+{
+	AAIController::ToggleAIIgnorePlayers();
+}
+
+void UAISystem::AILoggingVerbose()
+{
+	UWorld* OuterWorld = GetOuterWorld();
+	if (OuterWorld)
+	{
+		APlayerController* PC = OuterWorld->GetFirstPlayerController();
+		if (PC)
+		{
+			PC->ConsoleCommand(TEXT("log lognavigation verbose | log logpathfollowing verbose | log LogCharacter verbose | log LogBehaviorTree verbose | log LogPawnAction verbose|"));
+		}
+	}
+}
+
+void UAISystem::RunEQS(const FString& QueryName)
+{
+	UWorld* OuterWorld = GetOuterWorld();
+	if (OuterWorld == NULL)
+	{
+		return;
+	}
+
+	APlayerController* MyPC = OuterWorld->GetFirstPlayerController();
+	UEnvQueryManager* EQS = GetEnvironmentQueryManager();
+	if (MyPC && EQS)
+	{
+		AActor* Target = NULL;
+		UGameplayDebuggingControllerComponent* DebugComp = MyPC->FindComponentByClass<UGameplayDebuggingControllerComponent>();
+		if (DebugComp)
+		{
+			Target = DebugComp->GetCurrentDebugTarget();
+		}
+
+//#if WITH_EDITOR
+//		if (Target == NULL && GEditor != NULL)
+//		{
+//			Target = GEditor->GetSelectedObjects()->GetTop<AActor>();
+//
+//			// this part should not be needed, but is due to gameplay debugging messed up design
+//			if (Target == NULL)
+//			{
+//				for (FObjectIterator It(UGameplayDebuggingComponent::StaticClass()); It && (Target == NULL); ++It)
+//				{
+//					UGameplayDebuggingComponent* Comp = ((UGameplayDebuggingComponent*)*It);
+//					if (Comp->IsSelected())
+//					{
+//						Target = Comp->GetOwner();
+//					}
+//				}
+//			}
+//		}
+//#endif // WITH_EDITOR
+
+		if (Target)
+		{
+			UEnvQuery* QueryTemplate = EQS->FindQueryTemplate(QueryName);
+
+			if (QueryTemplate)
+			{
+				EQS->RunInstantQuery(FEnvQueryRequest(QueryTemplate, Target), EEnvQueryRunMode::AllMatching);
+			}
+			else
+			{
+				MyPC->ClientMessage(FString::Printf(TEXT("Unable to fing query template \'%s\'"), *QueryName));
+			}
+		}
+		else
+		{
+			MyPC->ClientMessage(TEXT("No debugging target"));
+		}
+	}
 }
\ No newline at end of file
diff --git a/Engine/Source/Runtime/AIModule/Private/AISystemExec.cpp b/Engine/Source/Runtime/AIModule/Private/AISystemExec.cpp
--- a/Engine/Source/Runtime/AIModule/Private/AISystemExec.cpp
+++ b/Engine/Source/Runtime/AIModule/Private/AISystemExec.cpp
@@ -19,24 +19,28 @@ FAISystemExec AISystemExecInstance;
 
 bool FAISystemExec::Exec(UWorld* Inworld, const TCHAR* Cmd, FOutputDevice& Ar)
 {
-	bool bHandled = false;
-
-	if (FParse::Command(&Cmd, TEXT("AIIgnorePlayers")))
+	if (Inworld == NULL)
 	{
-		AAIController::ToggleAIIgnorePlayers();
-		//bHandled = true; 
+		return false;
 	}
-	else if (FParse::Command(&Cmd, TEXT("AILoggingVerbose")))
+
+	UAISystem* AISys = UAISystem::GetCurrent(Inworld);	
+	bool bHandled = false;
+
+	if (AISys != NULL)
 	{
-		if (Inworld)
+		if (FParse::Command(&Cmd, TEXT("AIIgnorePlayers")))
+		{
+			AISys->AIIgnorePlayers();
+			bHandled = true;
+		}
+		else if (FParse::Command(&Cmd, TEXT("AILoggingVerbose")))
 		{
-			APlayerController* PC = Inworld->GetFirstPlayerController();
-			if (PC)
-			{
-				PC->ConsoleCommand(TEXT("log lognavigation verbose | log logpathfollowing verbose | log LogCharacter verbose | log LogBehaviorTree verbose | log LogPawnAction verbose|"));
-			}
+			AISys->AILoggingVerbose();
+			bHandled = true;
 		}
-		//bHandled = true;
+
+		// void UAISystem::RunEQS(const FString& QueryName)
 	}
 
 	return bHandled;
diff --git a/Engine/Source/Runtime/AIModule/Private/BehaviorTree/BehaviorTreeComponent.cpp b/Engine/Source/Runtime/AIModule/Private/BehaviorTree/BehaviorTreeComponent.cpp
--- a/Engine/Source/Runtime/AIModule/Private/BehaviorTree/BehaviorTreeComponent.cpp
+++ b/Engine/Source/Runtime/AIModule/Private/BehaviorTree/BehaviorTreeComponent.cpp
@@ -1205,7 +1205,7 @@ bool UBehaviorTreeComponent::PushInstance(class UBehaviorTree* TreeAsset)
 		return false;
 	}
 
-	UBehaviorTreeManager* BTManager = GetWorld()->GetBehaviorTreeManager();
+	UBehaviorTreeManager* BTManager = Cast<UAISystem>(GetWorld()->GetAISystem())->GetBehaviorTreeManager();
 	if (BTManager == NULL)
 	{
 		UE_VLOG(GetOwner(), LogBehaviorTree, Warning, TEXT("Failed to execute tree %s: behavior tree manager not found!"), *GetNameSafe(TreeAsset));
diff --git a/Engine/Source/Runtime/AIModule/Private/BehaviorTree/Tasks/BTTask_RunEQSQuery.cpp b/Engine/Source/Runtime/AIModule/Private/BehaviorTree/Tasks/BTTask_RunEQSQuery.cpp
--- a/Engine/Source/Runtime/AIModule/Private/BehaviorTree/Tasks/BTTask_RunEQSQuery.cpp
+++ b/Engine/Source/Runtime/AIModule/Private/BehaviorTree/Tasks/BTTask_RunEQSQuery.cpp
@@ -47,7 +47,8 @@ EBTNodeResult::Type UBTTask_RunEQSQuery::ExecuteTask(class UBehaviorTreeComponen
 EBTNodeResult::Type UBTTask_RunEQSQuery::AbortTask(class UBehaviorTreeComponent* OwnerComp, uint8* NodeMemory)
 {
 	UWorld* MyWorld = OwnerComp->GetWorld();
-	UEnvQueryManager* QueryManager = MyWorld->GetEnvironmentQueryManager();
+	UEnvQueryManager* QueryManager = UAISystem::GetCurrentEQSManager(MyWorld);
+	
 	if (QueryManager)
 	{
 		FBTEnvQueryTaskMemory* MyMemory = (FBTEnvQueryTaskMemory*)NodeMemory;
diff --git a/Engine/Source/Runtime/AIModule/Private/EnvironmentQuery/EQSTestingPawn.cpp b/Engine/Source/Runtime/AIModule/Private/EnvironmentQuery/EQSTestingPawn.cpp
--- a/Engine/Source/Runtime/AIModule/Private/EnvironmentQuery/EQSTestingPawn.cpp
+++ b/Engine/Source/Runtime/AIModule/Private/EnvironmentQuery/EQSTestingPawn.cpp
@@ -26,41 +26,43 @@ AEQSTestingPawn::AEQSTestingPawn(const class FPostConstructInitializeProperties&
 	CapsuleComponent->SetCollisionProfileName(CollisionProfileName);
 
 #if WITH_EDITORONLY_DATA
-	UArrowComponent* ArrowComponent = FindComponentByClass<UArrowComponent>();
-	if (ArrowComponent != NULL)
+	EdRenderComp = PCIP.CreateEditorOnlyDefaultSubobject<UEQSRenderingComponent>(this, TEXT("EQSRender"));
+	if (HasAnyFlags(RF_ClassDefaultObject) == false)
 	{
-		ArrowComponent->SetRelativeScale3D(FVector(2,2,2));
-		ArrowComponent->bIsScreenSizeScaled = true;
-	}
+		UArrowComponent* ArrowComponent = FindComponentByClass<UArrowComponent>();
+		if (ArrowComponent != NULL)
+		{
+			ArrowComponent->SetRelativeScale3D(FVector(2, 2, 2));
+			ArrowComponent->bIsScreenSizeScaled = true;
+		}
 
-	TSubobjectPtr<UBillboardComponent> SpriteComponent = PCIP.CreateEditorOnlyDefaultSubobject<UBillboardComponent>(this, TEXT("Sprite"));
-	if (!IsRunningCommandlet() && (SpriteComponent != nullptr))
-	{
-		struct FConstructorStatics
+		TSubobjectPtr<UBillboardComponent> SpriteComponent = PCIP.CreateEditorOnlyDefaultSubobject<UBillboardComponent>(this, TEXT("Sprite"));
+		if (!IsRunningCommandlet() && (SpriteComponent != nullptr))
 		{
-			ConstructorHelpers::FObjectFinderOptional<UTexture2D> TextureObject;
-			FName ID_Misc;
-			FText NAME_Misc;
-			FConstructorStatics()
-				: TextureObject(TEXT("/Engine/EditorResources/S_Pawn"))
-				, ID_Misc(TEXT("Misc"))
-				, NAME_Misc(NSLOCTEXT("SpriteCategory", "Misc", "Misc"))
+			struct FConstructorStatics
 			{
-			}
-		};
-		static FConstructorStatics ConstructorStatics;
-
-		SpriteComponent->Sprite = ConstructorStatics.TextureObject.Get();
-		SpriteComponent->RelativeScale3D = FVector(0.5f, 0.5f, 0.5f);
-		SpriteComponent->bHiddenInGame = true;
-		//SpriteComponent->Mobility = EComponentMobility::Static;
-		SpriteComponent->SpriteInfo.Category = ConstructorStatics.ID_Misc;
-		SpriteComponent->SpriteInfo.DisplayName = ConstructorStatics.NAME_Misc;
-		SpriteComponent->AttachParent = RootComponent;
-		SpriteComponent->bIsScreenSizeScaled = true;
+				ConstructorHelpers::FObjectFinderOptional<UTexture2D> TextureObject;
+				FName ID_Misc;
+				FText NAME_Misc;
+				FConstructorStatics()
+					: TextureObject(TEXT("/Engine/EditorResources/S_Pawn"))
+					, ID_Misc(TEXT("Misc"))
+					, NAME_Misc(NSLOCTEXT("SpriteCategory", "Misc", "Misc"))
+				{
+				}
+			};
+			static FConstructorStatics ConstructorStatics;
+
+			SpriteComponent->Sprite = ConstructorStatics.TextureObject.Get();
+			SpriteComponent->RelativeScale3D = FVector(0.5f, 0.5f, 0.5f);
+			SpriteComponent->bHiddenInGame = true;
+			//SpriteComponent->Mobility = EComponentMobility::Static;
+			SpriteComponent->SpriteInfo.Category = ConstructorStatics.ID_Misc;
+			SpriteComponent->SpriteInfo.DisplayName = ConstructorStatics.NAME_Misc;
+			SpriteComponent->AttachParent = RootComponent;
+			SpriteComponent->bIsScreenSizeScaled = true;
+		}
 	}
-
-	EdRenderComp = PCIP.CreateEditorOnlyDefaultSubobject<UEQSRenderingComponent>(this, TEXT("EQSRender"));
 #endif // WITH_EDITORONLY_DATA
 
 	// Default to no tick function, but if we set 'never ticks' to false (so there is a tick function) it is enabled by default
@@ -117,7 +119,7 @@ void AEQSTestingPawn::PostLoad()
 		SpriteComponent->bHiddenInGame = !bShouldBeVisibleInGame;
 	}
 
-	UEnvQueryManager* EQS = GetWorld() ? GetWorld()->GetEnvironmentQueryManager() : NULL;
+	UEnvQueryManager* EQS = GetWorld() ? UAISystem::GetCurrentEQSManager(GetWorld()) : NULL;
 	if (EQS)
 	{
 		RunEQSQuery();
@@ -154,7 +156,7 @@ void AEQSTestingPawn::Reset()
 
 void AEQSTestingPawn::MakeOneStep()
 {
-	UEnvQueryManager* EQS = GetWorld()->GetEnvironmentQueryManager();
+	UEnvQueryManager* EQS = UAISystem::GetCurrentEQSManager(GetWorld());
 	if (EQS == NULL)
 	{
 		return;
diff --git a/Engine/Source/Runtime/AIModule/Private/EnvironmentQuery/EnvQueryManager.cpp b/Engine/Source/Runtime/AIModule/Private/EnvironmentQuery/EnvQueryManager.cpp
--- a/Engine/Source/Runtime/AIModule/Private/EnvironmentQuery/EnvQueryManager.cpp
+++ b/Engine/Source/Runtime/AIModule/Private/EnvironmentQuery/EnvQueryManager.cpp
@@ -57,7 +57,7 @@ int32 FEnvQueryRequest::Execute(EEnvQueryRunMode::Type RunMode, FQueryFinishedSi
 		}
 	}
 
-	UEnvQueryManager* EnvQueryManager = World->GetEnvironmentQueryManager();
+	UEnvQueryManager* EnvQueryManager = UAISystem::GetCurrentEQSManager(World);
 	if (EnvQueryManager == NULL)
 	{
 		UE_LOG(LogEQS, Warning, TEXT("Missing EQS manager!"));
@@ -91,13 +91,13 @@ void UEnvQueryManager::FinishDestroy()
 
 UEnvQueryManager* UEnvQueryManager::GetCurrent(UWorld* World)
 {
-	return World ? World->GetEnvironmentQueryManager() : NULL;
+	return World ? UAISystem::GetCurrentEQSManager(World) : NULL;
 }
 
 UEnvQueryManager* UEnvQueryManager::GetCurrent(class UObject* WorldContextObject)
 {
 	UWorld* World = WorldContextObject ? GEngine->GetWorldFromContextObject(WorldContextObject) : NULL;
-	return World ? World->GetEnvironmentQueryManager() : NULL;
+	return World ? UAISystem::GetCurrentEQSManager(World) : NULL;
 }
 
 #if WITH_EDITOR
@@ -111,7 +111,7 @@ void UEnvQueryManager::NotifyAssetUpdate(UEnvQuery* Query)
 	UWorld* World = GEditor->GetEditorWorldContext().World();
 	if (World)
 	{
-		UEnvQueryManager* EQS = World->GetEnvironmentQueryManager();
+		UEnvQueryManager* EQS = UAISystem::GetCurrentEQSManager(World);
 		if (EQS)
 		{
 			EQS->InstanceCache.Reset();
diff --git a/Engine/Source/Runtime/AIModule/Private/Navigation/PathFollowingComponent.cpp b/Engine/Source/Runtime/AIModule/Private/Navigation/PathFollowingComponent.cpp
--- a/Engine/Source/Runtime/AIModule/Private/Navigation/PathFollowingComponent.cpp
+++ b/Engine/Source/Runtime/AIModule/Private/Navigation/PathFollowingComponent.cpp
@@ -7,6 +7,7 @@
 #include "Navigation/NavigationComponent.h"
 #include "AI/Navigation/SmartNavLinkComponent.h"
 #include "VisualLog.h"
+#include "Navigation/PathFollowingComponent.h"
 
 #define USE_PHYSIC_FOR_VISIBILITY_TESTS 1 // Physic will be used for visibility tests if set or only raycasts on navmesh if not
 
diff --git a/Engine/Source/Runtime/AIModule/Public/AIModule.h b/Engine/Source/Runtime/AIModule/Public/AIModule.h
--- a/Engine/Source/Runtime/AIModule/Public/AIModule.h
+++ b/Engine/Source/Runtime/AIModule/Public/AIModule.h
@@ -2,12 +2,12 @@
 
 #pragma once
 
-#include "ModuleManager.h"
+#include "AI/AISystemBase.h"
 
 /**
  * The public interface to this module
  */
-class IAIModule : public IModuleInterface
+class IAIModule : public IAISystemModule
 {
 
 public:
diff --git a/Engine/Source/Runtime/Engine/Classes/AI/Navigation/NavFilters/NavigationQueryFilter.h b/Engine/Source/Runtime/Engine/Classes/AI/Navigation/NavFilters/NavigationQueryFilter.h
--- a/Engine/Source/Runtime/Engine/Classes/AI/Navigation/NavFilters/NavigationQueryFilter.h
+++ b/Engine/Source/Runtime/Engine/Classes/AI/Navigation/NavFilters/NavigationQueryFilter.h
@@ -86,6 +86,113 @@ struct ENGINE_API FNavigationFilterFlags
 #endif
 };
 
+class INavigationQueryFilterInterface
+{
+public:
+	virtual ~INavigationQueryFilterInterface(){}
+
+	virtual void Reset() = 0;
+
+	virtual void SetAreaCost(uint8 AreaType, float Cost) = 0;
+	virtual void SetFixedAreaEnteringCost(uint8 AreaType, float Cost) = 0;
+	virtual void SetExcludedArea(uint8 AreaType) = 0;
+	virtual void SetAllAreaCosts(const float* CostArray, const int32 Count) = 0;
+	virtual void GetAllAreaCosts(float* CostArray, float* FixedCostArray, const int32 Count) const = 0;
+	virtual void SetBacktrackingEnabled(const bool bBacktracking) = 0;
+	virtual bool IsBacktrackingEnabled() const = 0;
+	virtual bool IsEqual(const INavigationQueryFilterInterface* Other) const = 0;
+	virtual void SetIncludeFlags(uint16 Flags) = 0;
+	virtual uint16 GetIncludeFlags() const = 0;
+	virtual void SetExcludeFlags(uint16 Flags) = 0;
+	virtual uint16 GetExcludeFlags() const = 0;
+
+	virtual class INavigationQueryFilterInterface* CreateCopy() const = 0;
+};
+
+struct ENGINE_API FNavigationQueryFilter : public TSharedFromThis<FNavigationQueryFilter>
+{
+	FNavigationQueryFilter() : QueryFilterImpl(NULL), MaxSearchNodes(DefaultMaxSearchNodes) {}
+private:
+	FNavigationQueryFilter(const FNavigationQueryFilter& Source);
+	FNavigationQueryFilter(const FNavigationQueryFilter* Source);
+	FNavigationQueryFilter(const TSharedPtr<FNavigationQueryFilter> Source);
+	FNavigationQueryFilter& operator=(const FNavigationQueryFilter& Source);
+public:
+
+	/** set travel cost for area */
+	void SetAreaCost(uint8 AreaType, float Cost);
+
+	/** set entering cost for area */
+	void SetFixedAreaEnteringCost(uint8 AreaType, float Cost);
+
+	/** mark area as excluded from path finding */
+	void SetExcludedArea(uint8 AreaType);
+
+	/** set travel cost for all areas */
+	void SetAllAreaCosts(const TArray<float>& CostArray);
+	void SetAllAreaCosts(const float* CostArray, const int32 Count);
+
+	/** get travel & entering costs for all areas */
+	void GetAllAreaCosts(float* CostArray, float* FixedCostArray, const int32 Count) const;
+
+	/** set required flags of navigation nodes */
+	void SetIncludeFlags(uint16 Flags);
+
+	/** get required flags of navigation nodes */
+	uint16 GetIncludeFlags() const;
+
+	/** set forbidden flags of navigation nodes */
+	void SetExcludeFlags(uint16 Flags);
+
+	/** get forbidden flags of navigation nodes */
+	uint16 GetExcludeFlags() const;
+
+	/** set node limit for A* loop */
+	void SetMaxSearchNodes(const uint32 MaxNodes) { MaxSearchNodes = MaxNodes; }
+
+	/** get node limit for A* loop */
+	FORCEINLINE uint32 GetMaxSearchNodes() const { return MaxSearchNodes; }
+
+	/** mark filter as backtracking - parse directional links in opposite direction
+	*  (find path from End to Start, but all links works like on path from Start to End) */
+	void SetBacktrackingEnabled(const bool bBacktracking) { QueryFilterImpl->SetBacktrackingEnabled(bBacktracking); }
+
+	/** get backtracking status */
+	bool IsBacktrackingEnabled() const { return QueryFilterImpl->IsBacktrackingEnabled(); }
+
+	template<typename FilterType>
+	void SetFilterType()
+	{
+		QueryFilterImpl = MakeShareable(new FilterType());
+	}
+
+	FORCEINLINE_DEBUGGABLE void SetFilterImplementation(const class INavigationQueryFilterInterface* InQueryFilterImpl)
+	{
+		QueryFilterImpl = MakeShareable(InQueryFilterImpl->CreateCopy());
+	}
+
+	FORCEINLINE const INavigationQueryFilterInterface* GetImplementation() const { return QueryFilterImpl.Get(); }
+	FORCEINLINE INavigationQueryFilterInterface* GetImplementation() { return QueryFilterImpl.Get(); }
+	void Reset() { GetImplementation()->Reset(); }
+
+	TSharedPtr<FNavigationQueryFilter> GetCopy() const;
+
+	FORCEINLINE bool operator==(const FNavigationQueryFilter& Other) const
+	{
+		const INavigationQueryFilterInterface* Impl0 = GetImplementation();
+		const INavigationQueryFilterInterface* Impl1 = Other.GetImplementation();
+		return Impl0 && Impl1 && Impl0->IsEqual(Impl1);
+	}
+
+	static const uint32 DefaultMaxSearchNodes;
+
+protected:
+	void Assign(const FNavigationQueryFilter& Source);
+
+	TSharedPtr<INavigationQueryFilterInterface, ESPMode::ThreadSafe> QueryFilterImpl;
+	uint32 MaxSearchNodes;
+};
+
 /** Class containing definition of a navigation query filter */
 UCLASS(Abstract, Blueprintable)
 class ENGINE_API UNavigationQueryFilter : public UObject
diff --git a/Engine/Source/Runtime/Engine/Classes/AI/Navigation/NavigationData.h b/Engine/Source/Runtime/Engine/Classes/AI/Navigation/NavigationData.h
--- a/Engine/Source/Runtime/Engine/Classes/AI/Navigation/NavigationData.h
+++ b/Engine/Source/Runtime/Engine/Classes/AI/Navigation/NavigationData.h
@@ -1,8 +1,12 @@
 // Copyright 1998-2014 Epic Games, Inc. All Rights Reserved.
 
 #pragma once
+#include "AI/Navigation/NavigationTypes.h"
+#include "AI/NavDataGenerator.h"
 #include "NavigationData.generated.h"
 
+class FNavDataGenerator; 
+
 USTRUCT()
 struct FSupportedAreaData
 {
@@ -22,7 +26,7 @@ struct FSupportedAreaData
  *	Represents abstract Navigation Data (sub-classed as NavMesh, NavGraph, etc)
  *	Used as a common interface for all navigation types handled by NavigationSystem
  */
-UCLASS(config=Engine, defaultconfig, dependson=UNavigationSystem, NotBlueprintable, abstract)
+UCLASS(config=Engine, defaultconfig, NotBlueprintable, abstract)
 class ENGINE_API ANavigationData : public AActor
 {
 	GENERATED_UCLASS_BODY()
@@ -162,7 +166,7 @@ public:
 	// Querying                                                                
 	//----------------------------------------------------------------------//
 	FORCEINLINE TSharedPtr<const FNavigationQueryFilter> GetDefaultQueryFilter() const { return DefaultQueryFilter; }
-	FORCEINLINE const INavigationQueryFilterInterface* GetDefaultQueryFilterImpl() const { return DefaultQueryFilter->GetImplementation(); }	
+	FORCEINLINE const class INavigationQueryFilterInterface* GetDefaultQueryFilterImpl() const { return DefaultQueryFilter->GetImplementation(); }	
 	FORCEINLINE FVector GetDefaultQueryExtent() const { return NavDataConfig.DefaultQueryExtent; }
 
 	/** 
@@ -323,7 +327,7 @@ protected:
 	virtual int32 GetNewAreaID(const UClass* AreaClass) const;
 	
 #if WITH_NAVIGATION_GENERATOR
-	virtual class FNavDataGenerator* ConstructGenerator(const FNavAgentProperties& AgentProps) PURE_VIRTUAL(ANavigationData::ConstructGenerator, return NULL; );
+	virtual FNavDataGenerator* ConstructGenerator(const FNavAgentProperties& AgentProps) PURE_VIRTUAL(ANavigationData::ConstructGenerator, return NULL; );
 #endif // WITH_NAVIGATION_GENERATOR
 
 protected:
diff --git a/Engine/Source/Runtime/Engine/Classes/AI/Navigation/NavigationGraph.h b/Engine/Source/Runtime/Engine/Classes/AI/Navigation/NavigationGraph.h
--- a/Engine/Source/Runtime/Engine/Classes/AI/Navigation/NavigationGraph.h
+++ b/Engine/Source/Runtime/Engine/Classes/AI/Navigation/NavigationGraph.h
@@ -1,8 +1,11 @@
 // Copyright 1998-2014 Epic Games, Inc. All Rights Reserved.
 
 #pragma once
+#include "AI/Navigation/NavigationData.h"
 #include "NavigationGraph.generated.h"
 
+class UNavigationSystem;
+
 USTRUCT()
 struct FNavGraphEdge
 {
@@ -50,13 +53,16 @@ struct FNavGraphNode
 };
 
 /** currently abstract since it's not full implemented */
-UCLASS(config=Engine, dependson=UNavigationSystem, MinimalAPI, abstract)
+UCLASS(config=Engine, MinimalAPI, abstract)
 class ANavigationGraph : public ANavigationData
 {
 	GENERATED_UCLASS_BODY()
 
+protected:
 #if WITH_NAVIGATION_GENERATOR
-	virtual class FNavDataGenerator* ConstructGenerator(const FNavAgentProperties& AgentProps) OVERRIDE;
+	virtual FNavDataGenerator* ConstructGenerator(const FNavAgentProperties& AgentProps) OVERRIDE;
 #endif // WITH_NAVIGATION_GENERATOR
-	static ANavigationData* CreateNavigationInstances(class UNavigationSystem* NavSys);
+
+public:
+	static ANavigationData* CreateNavigationInstances(UNavigationSystem* NavSys);
 };
diff --git a/Engine/Source/Runtime/Engine/Classes/AI/Navigation/NavigationGraphNodeComponent.h b/Engine/Source/Runtime/Engine/Classes/AI/Navigation/NavigationGraphNodeComponent.h
--- a/Engine/Source/Runtime/Engine/Classes/AI/Navigation/NavigationGraphNodeComponent.h
+++ b/Engine/Source/Runtime/Engine/Classes/AI/Navigation/NavigationGraphNodeComponent.h
@@ -1,9 +1,10 @@
 // Copyright 1998-2014 Epic Games, Inc. All Rights Reserved.
 
 #pragma once
+#include "AI/Navigation/NavigationGraph.h"
 #include "NavigationGraphNodeComponent.generated.h"
 
-UCLASS(config=Engine, dependson=ANavigationGraph, MinimalAPI, HideCategories=(Mobility))
+UCLASS(config=Engine, MinimalAPI, HideCategories=(Mobility))
 class UNavigationGraphNodeComponent : public USceneComponent
 {
 	GENERATED_UCLASS_BODY()
diff --git a/Engine/Source/Runtime/Engine/Classes/AI/Navigation/NavigationSystem.h b/Engine/Source/Runtime/Engine/Classes/AI/Navigation/NavigationSystem.h
--- a/Engine/Source/Runtime/Engine/Classes/AI/Navigation/NavigationSystem.h
+++ b/Engine/Source/Runtime/Engine/Classes/AI/Navigation/NavigationSystem.h
@@ -6,34 +6,18 @@
 #include "AI/Navigation/NavigationTypes.h"
 #include "Kismet/BlueprintFunctionLibrary.h"
 #include "GenericOctreePublic.h"
-//#include "AI/Navigation/NavigationData.h"
+#include "AI/Navigation/NavigationData.h"
 #include "NavigationSystem.generated.h"
 
 #define NAVSYS_DEBUG (0 && UE_BUILD_DEBUG)
-/** Whether to compile in navigation data generation - should be compiled at least when WITH_EDITOR is true */
-#define WITH_NAVIGATION_GENERATOR (WITH_RECAST || WITH_EDITOR)
+
 // if we'll be rebuilding navigation at runtime
 #define WITH_RUNTIME_NAVIGATION_BUILDING (1 && WITH_NAVIGATION_GENERATOR)
 
 #define NAVOCTREE_CONTAINS_COLLISION_DATA (1 && WITH_RECAST)
 
-#define DEFAULT_NAV_QUERY_EXTENT_HORIZONTAL 50.f
-#define DEFAULT_NAV_QUERY_EXTENT_VERTICAL 100.f
-
 #define NAV_USE_MAIN_NAVIGATION_DATA NULL
 
-UENUM()
-namespace ENavigationQueryResult
-{
-	enum Type
-	{
-		Invalid,
-		Error,
-		Fail,
-		Success
-	};
-}
-
 DECLARE_LOG_CATEGORY_EXTERN(LogNavigation, Warning, All);
 
 /** delegate to let interested parties know that new nav area class has been registered */
@@ -197,119 +181,8 @@ struct FNavigationPortalEdge
 	FORCEINLINE FVector GetMiddlePoint() const { return Left + (Right - Left) / 2; }
 };
 
-/** 
- *	Delegate used to communicate that path finding query has been finished. 
- *	@param uint32 unique Query ID of given query 
- *	@param ENavigationQueryResult enum expressed query result. 
- *	@param FNavPathSharedPtr resulting path. Valid only for ENavigationQueryResult == ENavigationQueryResult::Fail 
- *		(may contain path leading as close to destination as possible) 
- *		and ENavigationQueryResult == ENavigationQueryResult::Success
- */
-DECLARE_DELEGATE_ThreeParams( FNavPathQueryDelegate, uint32, ENavigationQueryResult::Type, FNavPathSharedPtr );
-
-struct FPathFindingResult
-{
-	FNavPathSharedPtr Path;
-	ENavigationQueryResult::Type Result;
-
-	FPathFindingResult(ENavigationQueryResult::Type InResult = ENavigationQueryResult::Invalid) : Result(InResult)
-	{}
-
-	FORCEINLINE bool IsSuccessful() const { return Result == ENavigationQueryResult::Success; }
-	FORCEINLINE bool IsPartial() const { return Result == ENavigationQueryResult::Fail && Path.IsValid() && Path->IsPartial(); }
-};
-
-namespace EPathFindingMode
-{
-	enum Type
-	{
-		Regular,
-		Hierarchical,
-	};
-};
-
-USTRUCT()
-struct FMovementProperties
-{
-	GENERATED_USTRUCT_BODY()
-
-	UPROPERTY(EditAnywhere, BlueprintReadWrite, Category=MovementProperties)
-	uint32 bCanCrouch:1;    // if true, this pawn is capable of crouching
-
-	// movement capabilities - used by AI for reachability tests
-	UPROPERTY(EditAnywhere, BlueprintReadWrite, Category=MovementProperties)
-	uint32 bCanJump:1;
-
-	UPROPERTY(EditAnywhere, BlueprintReadWrite, Category=MovementProperties)
-	uint32 bCanWalk:1;
-
-	UPROPERTY(EditAnywhere, BlueprintReadWrite, Category=MovementProperties)
-	uint32 bCanSwim:1;
-
-	UPROPERTY(EditAnywhere, BlueprintReadWrite, Category=MovementProperties)
-	uint32 bCanFly:1;
-
-	FMovementProperties()
-		: bCanCrouch(false)
-		, bCanJump(false)
-		, bCanWalk(false)
-		, bCanSwim(false)
-		, bCanFly(false)
-	{
-	}
-};
-
-USTRUCT()
-struct FNavAgentProperties : public FMovementProperties
-{
-	GENERATED_USTRUCT_BODY()
-
-	UPROPERTY(EditAnywhere, BlueprintReadWrite, Category=MovementProperties)
-	float AgentRadius;
-
-	UPROPERTY(EditAnywhere, BlueprintReadWrite, Category=MovementProperties)
-	float AgentHeight;
-
-	FNavAgentProperties(float Radius = -1.f, float Height = -1.f)
-		: AgentRadius(Radius)
-		, AgentHeight(Height)
-	{
-	}
-
-	void UpdateWithCollisionComponent(class UShapeComponent* CollisionComponent);
-
-	FORCEINLINE bool IsValid() const { return AgentRadius >= 0 && AgentHeight >= 0; }
-
-	FORCEINLINE bool IsEquivalent(const FNavAgentProperties& Other, float Precision = 5.f) const 
-	{
-		return FGenericPlatformMath::Abs(AgentRadius - Other.AgentRadius) < Precision && FGenericPlatformMath::Abs(AgentHeight - Other.AgentHeight) < Precision;
-	}
-
-	bool operator==(const FNavAgentProperties& Other) const
-	{
-		return IsEquivalent(Other);
-	}
-
-	FVector GetExtent() const
-	{
-		return FVector(AgentRadius,AgentRadius,AgentHeight/2);
-	}
-};
-
-inline uint32 GetTypeHash( const FNavAgentProperties& A )
-{
-	return (int16(A.AgentRadius) << 16) | int16(A.AgentHeight);
-}
-
 namespace FNavigationSystem
 {
-	enum ECreateIfEmpty 
-	{
-		Invalid = -1,
-		DontCreate = 0,
-		Create = 1,
-	};
-
 	enum EMode
 	{
 		InvalidMode = -1,
@@ -318,39 +191,8 @@ namespace FNavigationSystem
 		SimulationMode,
 		PIEMode,
 	};
-
-	/** used as a fallback value for navigation agent radius, when none specified via UNavigationSystem::SupportedAgents */
-	extern const float FallbackAgentRadius;
-
-	/** used as a fallback value for navigation agent height, when none specified via UNavigationSystem::SupportedAgents */
-	extern const float FallbackAgentHeight;
-
-	static const FBox InvalidBoundingBox(0);
 }
 
-USTRUCT()
-struct ENGINE_API FNavDataConfig : public FNavAgentProperties
-{
-	GENERATED_USTRUCT_BODY()
-
-	UPROPERTY(EditAnywhere, BlueprintReadWrite, Category=Display)
-	FName Name;
-
-	UPROPERTY(EditAnywhere, BlueprintReadWrite, Category=Display)
-	FColor Color;
-
-	UPROPERTY(EditAnywhere, BlueprintReadOnly, Category=Querying, config)
-	FVector DefaultQueryExtent;
-
-	FNavDataConfig(float Radius = FNavigationSystem::FallbackAgentRadius, float Height = FNavigationSystem::FallbackAgentHeight)
-		: FNavAgentProperties(Radius, Height)
-		, Name(TEXT("Default"))
-		, Color(140,255,0,164)
-		, DefaultQueryExtent(DEFAULT_NAV_QUERY_EXTENT_HORIZONTAL, DEFAULT_NAV_QUERY_EXTENT_HORIZONTAL, DEFAULT_NAV_QUERY_EXTENT_VERTICAL)
-	{
-	}	
-};
-
 namespace FNavigationSystem
 {
 	/** 
@@ -359,162 +201,6 @@ namespace FNavigationSystem
 	typedef class ANavigationData* (*FNavigationDataInstanceCreator)(class UWorld*, const FNavDataConfig&);
 }
 
-class INavigationQueryFilterInterface
-{
-public:
-	virtual ~INavigationQueryFilterInterface(){}
-
-	virtual void Reset() = 0;
-
-	virtual void SetAreaCost(uint8 AreaType, float Cost) = 0;
-	virtual void SetFixedAreaEnteringCost(uint8 AreaType, float Cost) = 0;
-	virtual void SetExcludedArea(uint8 AreaType) = 0;
-	virtual void SetAllAreaCosts(const float* CostArray, const int32 Count) = 0;
-	virtual void GetAllAreaCosts(float* CostArray, float* FixedCostArray, const int32 Count) const = 0;
-	virtual void SetBacktrackingEnabled(const bool bBacktracking) = 0;
-	virtual bool IsBacktrackingEnabled() const = 0;
-	virtual bool IsEqual(const INavigationQueryFilterInterface* Other) const = 0;
-	virtual void SetIncludeFlags(uint16 Flags) = 0;
-	virtual uint16 GetIncludeFlags() const = 0;
-	virtual void SetExcludeFlags(uint16 Flags) = 0;
-	virtual uint16 GetExcludeFlags() const = 0;
-
-	virtual class INavigationQueryFilterInterface* CreateCopy() const = 0;
-};
-
-struct ENGINE_API FNavigationQueryFilter : public TSharedFromThis<FNavigationQueryFilter>
-{
-	FNavigationQueryFilter() : QueryFilterImpl(NULL), MaxSearchNodes(DefaultMaxSearchNodes) {}
-private:
-	FNavigationQueryFilter(const FNavigationQueryFilter& Source);
-	FNavigationQueryFilter(const FNavigationQueryFilter* Source);
-	FNavigationQueryFilter(const TSharedPtr<FNavigationQueryFilter> Source);
-	FNavigationQueryFilter& operator=(const FNavigationQueryFilter& Source);
-public:
-
-	/** set travel cost for area */
-	void SetAreaCost(uint8 AreaType, float Cost);
-
-	/** set entering cost for area */
-	void SetFixedAreaEnteringCost(uint8 AreaType, float Cost);
-
-	/** mark area as excluded from path finding */
-	void SetExcludedArea(uint8 AreaType);
-
-	/** set travel cost for all areas */
-	void SetAllAreaCosts(const TArray<float>& CostArray);
-	void SetAllAreaCosts(const float* CostArray, const int32 Count);
-
-	/** get travel & entering costs for all areas */
-	void GetAllAreaCosts(float* CostArray, float* FixedCostArray, const int32 Count) const;
-
-	/** set required flags of navigation nodes */
-	void SetIncludeFlags(uint16 Flags);
-
-	/** get required flags of navigation nodes */
-	uint16 GetIncludeFlags() const;
-
-	/** set forbidden flags of navigation nodes */
-	void SetExcludeFlags(uint16 Flags);
-
-	/** get forbidden flags of navigation nodes */
-	uint16 GetExcludeFlags() const;
-
-	/** set node limit for A* loop */
-	void SetMaxSearchNodes(const uint32 MaxNodes) { MaxSearchNodes = MaxNodes; }
-
-	/** get node limit for A* loop */
-	FORCEINLINE uint32 GetMaxSearchNodes() const { return MaxSearchNodes; }
-
-	/** mark filter as backtracking - parse directional links in opposite direction
-	 *  (find path from End to Start, but all links works like on path from Start to End) */
-	void SetBacktrackingEnabled(const bool bBacktracking) {	QueryFilterImpl->SetBacktrackingEnabled(bBacktracking);	}
-	
-	/** get backtracking status */
-	bool IsBacktrackingEnabled() const { return QueryFilterImpl->IsBacktrackingEnabled(); }
-
-	template<typename FilterType>
-	void SetFilterType()
-	{
-		QueryFilterImpl = MakeShareable(new FilterType());
-	}
-
-	FORCEINLINE_DEBUGGABLE void SetFilterImplementation(const class INavigationQueryFilterInterface* InQueryFilterImpl)
-	{
-		QueryFilterImpl = MakeShareable(InQueryFilterImpl->CreateCopy());
-	}
-
-	FORCEINLINE const INavigationQueryFilterInterface* GetImplementation() const { return QueryFilterImpl.Get(); }
-	FORCEINLINE INavigationQueryFilterInterface* GetImplementation() { return QueryFilterImpl.Get(); }
-	void Reset() { GetImplementation()->Reset(); }
-
-	TSharedPtr<FNavigationQueryFilter> GetCopy() const;
-
-	FORCEINLINE bool operator==(const FNavigationQueryFilter& Other) const
-	{
-		const INavigationQueryFilterInterface* Impl0 = GetImplementation();
-		const INavigationQueryFilterInterface* Impl1 = Other.GetImplementation();
-		return Impl0 && Impl1 && Impl0->IsEqual(Impl1);
-	}
-
-	static const uint32 DefaultMaxSearchNodes;
-	
-protected:
-	void Assign(const FNavigationQueryFilter& Source);
-
-	TSharedPtr<INavigationQueryFilterInterface, ESPMode::ThreadSafe> QueryFilterImpl;
-	uint32 MaxSearchNodes;
-};
-
-struct ENGINE_API FPathFindingQuery
-{
-	TWeakObjectPtr<const class ANavigationData> NavData;
-	TWeakObjectPtr<const UObject> Owner;
-	FVector StartLocation;
-	FVector EndLocation;
-	TSharedPtr<const FNavigationQueryFilter> QueryFilter;
-
-	/** additional flags passed to navigation data handling request */
-	int32 NavDataFlags;
-
-	FPathFindingQuery()
-		: NavData(NULL)
-		, Owner(NULL)
-		, StartLocation(FVector::ZeroVector)
-		, EndLocation(FVector::ZeroVector)
-		, NavDataFlags(0)
-	{
-	}
-
-	FPathFindingQuery(const FPathFindingQuery& Source);
-
-	FPathFindingQuery(const UObject* InOwner, const class ANavigationData* InNavData, const FVector& Start, const FVector& End, TSharedPtr<const FNavigationQueryFilter> SourceQueryFilter = NULL);
-};
-
-struct FAsyncPathFindingQuery : public FPathFindingQuery
-{
-	const uint32 QueryID;
-	const FNavPathQueryDelegate OnDoneDelegate;
-	const TEnumAsByte<EPathFindingMode::Type> Mode;
-	FPathFindingResult Result;
-
-	FAsyncPathFindingQuery()
-		: QueryID(INVALID_NAVQUERYID)
-	{
-	}
-
-	FAsyncPathFindingQuery(const UObject* InOwner, const class ANavigationData* InNavData, const FVector& Start, const FVector& End, const FNavPathQueryDelegate& Delegate, TSharedPtr<const FNavigationQueryFilter> SourceQueryFilter);
-	FAsyncPathFindingQuery(const FPathFindingQuery& Query, const FNavPathQueryDelegate& Delegate, const EPathFindingMode::Type QueryMode);
-	
-protected:
-	FORCEINLINE static uint32 GetUniqueID() 
-	{
-		return ++LastPathFindingUniqueID;
-	}
-
-	static uint32 LastPathFindingUniqueID;
-};
-
 struct FNavigationSystemExec: public FSelfRegisteringExec
 {
 	FNavigationSystemExec()
diff --git a/Engine/Source/Runtime/Engine/Classes/AI/Navigation/NavigationTypes.h b/Engine/Source/Runtime/Engine/Classes/AI/Navigation/NavigationTypes.h
--- a/Engine/Source/Runtime/Engine/Classes/AI/Navigation/NavigationTypes.h
+++ b/Engine/Source/Runtime/Engine/Classes/AI/Navigation/NavigationTypes.h
@@ -11,9 +11,33 @@
 #define INVALID_NAVDATA uint32(0)
 #define INVALID_NAVEXTENT (FVector::ZeroVector)
 
+#define DEFAULT_NAV_QUERY_EXTENT_HORIZONTAL 50.f
+#define DEFAULT_NAV_QUERY_EXTENT_VERTICAL 100.f
+
+/** Whether to compile in navigation data generation - should be compiled at least when WITH_EDITOR is true */
+#define WITH_NAVIGATION_GENERATOR (WITH_RECAST || WITH_EDITOR)
+
 /** uniform identifier type for navigation data elements may it be a polygon or graph node */
 typedef uint64 NavNodeRef;
 
+namespace FNavigationSystem
+{
+	/** used as a fallback value for navigation agent radius, when none specified via UNavigationSystem::SupportedAgents */
+	extern const float FallbackAgentRadius;
+
+	/** used as a fallback value for navigation agent height, when none specified via UNavigationSystem::SupportedAgents */
+	extern const float FallbackAgentHeight;
+
+	static const FBox InvalidBoundingBox(0);
+
+	enum ECreateIfEmpty
+	{
+		Invalid = -1,
+		DontCreate = 0,
+		Create = 1,
+	};
+}
+
 //////////////////////////////////////////////////////////////////////////
 // Navigation data generation
 
@@ -278,6 +302,194 @@ protected:
 	TWeakObjectPtr<class ANavigationData> Owner;
 };
 
+USTRUCT()
+struct FMovementProperties
+{
+	GENERATED_USTRUCT_BODY()
+
+	UPROPERTY(EditAnywhere, BlueprintReadWrite, Category=MovementProperties)
+	uint32 bCanCrouch:1;    // if true, this pawn is capable of crouching
+
+	// movement capabilities - used by AI for reachability tests
+	UPROPERTY(EditAnywhere, BlueprintReadWrite, Category=MovementProperties)
+	uint32 bCanJump:1;
+
+	UPROPERTY(EditAnywhere, BlueprintReadWrite, Category=MovementProperties)
+	uint32 bCanWalk:1;
+
+	UPROPERTY(EditAnywhere, BlueprintReadWrite, Category=MovementProperties)
+	uint32 bCanSwim:1;
+
+	UPROPERTY(EditAnywhere, BlueprintReadWrite, Category=MovementProperties)
+	uint32 bCanFly:1;
+
+	FMovementProperties()
+		: bCanCrouch(false)
+		, bCanJump(false)
+		, bCanWalk(false)
+		, bCanSwim(false)
+		, bCanFly(false)
+	{
+	}
+};
+
+USTRUCT()
+struct FNavAgentProperties : public FMovementProperties
+{
+	GENERATED_USTRUCT_BODY()
+
+	UPROPERTY(EditAnywhere, BlueprintReadWrite, Category = MovementProperties)
+	float AgentRadius;
+
+	UPROPERTY(EditAnywhere, BlueprintReadWrite, Category = MovementProperties)
+	float AgentHeight;
+
+	FNavAgentProperties(float Radius = -1.f, float Height = -1.f)
+		: AgentRadius(Radius)
+		, AgentHeight(Height)
+	{
+	}
+
+	void UpdateWithCollisionComponent(class UShapeComponent* CollisionComponent);
+
+	FORCEINLINE bool IsValid() const { return AgentRadius >= 0 && AgentHeight >= 0; }
+
+	FORCEINLINE bool IsEquivalent(const FNavAgentProperties& Other, float Precision = 5.f) const
+	{
+		return FGenericPlatformMath::Abs(AgentRadius - Other.AgentRadius) < Precision && FGenericPlatformMath::Abs(AgentHeight - Other.AgentHeight) < Precision;
+	}
+
+	bool operator==(const FNavAgentProperties& Other) const
+	{
+		return IsEquivalent(Other);
+	}
+
+	FVector GetExtent() const
+	{
+		return FVector(AgentRadius, AgentRadius, AgentHeight / 2);
+	}
+};
+
+inline uint32 GetTypeHash(const FNavAgentProperties& A)
+{
+	return (int16(A.AgentRadius) << 16) | int16(A.AgentHeight);
+}
+
+USTRUCT()
+struct ENGINE_API FNavDataConfig : public FNavAgentProperties
+{
+	GENERATED_USTRUCT_BODY()
+
+	UPROPERTY(EditAnywhere, BlueprintReadWrite, Category=Display)
+	FName Name;
+
+	UPROPERTY(EditAnywhere, BlueprintReadWrite, Category=Display)
+	FColor Color;
+
+	UPROPERTY(EditAnywhere, BlueprintReadOnly, Category=Querying, config)
+	FVector DefaultQueryExtent;
+
+	FNavDataConfig(float Radius = FNavigationSystem::FallbackAgentRadius, float Height = FNavigationSystem::FallbackAgentHeight)
+		: FNavAgentProperties(Radius, Height)
+		, Name(TEXT("Default"))
+		, Color(140,255,0,164)
+		, DefaultQueryExtent(DEFAULT_NAV_QUERY_EXTENT_HORIZONTAL, DEFAULT_NAV_QUERY_EXTENT_HORIZONTAL, DEFAULT_NAV_QUERY_EXTENT_VERTICAL)
+	{
+	}	
+};
+
+UENUM()
+namespace ENavigationQueryResult
+{
+	enum Type
+	{
+		Invalid,
+		Error,
+		Fail,
+		Success
+	};
+}
+
+struct FPathFindingResult
+{
+	FNavPathSharedPtr Path;
+	ENavigationQueryResult::Type Result;
+
+	FPathFindingResult(ENavigationQueryResult::Type InResult = ENavigationQueryResult::Invalid) : Result(InResult)
+	{}
+
+	FORCEINLINE bool IsSuccessful() const { return Result == ENavigationQueryResult::Success; }
+	FORCEINLINE bool IsPartial() const { return Result == ENavigationQueryResult::Fail && Path.IsValid() && Path->IsPartial(); }
+};
+
+struct ENGINE_API FPathFindingQuery
+{
+	TWeakObjectPtr<const class ANavigationData> NavData;
+	TWeakObjectPtr<const UObject> Owner;
+	FVector StartLocation;
+	FVector EndLocation;
+	TSharedPtr<const FNavigationQueryFilter> QueryFilter;
+
+	/** additional flags passed to navigation data handling request */
+	int32 NavDataFlags;
+
+	FPathFindingQuery()
+		: NavData(NULL)
+		, Owner(NULL)
+		, StartLocation(FVector::ZeroVector)
+		, EndLocation(FVector::ZeroVector)
+		, NavDataFlags(0)
+	{
+	}
+
+	FPathFindingQuery(const FPathFindingQuery& Source);
+
+	FPathFindingQuery(const UObject* InOwner, const class ANavigationData* InNavData, const FVector& Start, const FVector& End, TSharedPtr<const FNavigationQueryFilter> SourceQueryFilter = NULL);
+};
+
+namespace EPathFindingMode
+{
+	enum Type
+	{
+		Regular,
+		Hierarchical,
+	};
+};
+
+/**
+*	Delegate used to communicate that path finding query has been finished.
+*	@param uint32 unique Query ID of given query
+*	@param ENavigationQueryResult enum expressed query result.
+*	@param FNavPathSharedPtr resulting path. Valid only for ENavigationQueryResult == ENavigationQueryResult::Fail
+*		(may contain path leading as close to destination as possible)
+*		and ENavigationQueryResult == ENavigationQueryResult::Success
+*/
+DECLARE_DELEGATE_ThreeParams(FNavPathQueryDelegate, uint32, ENavigationQueryResult::Type, FNavPathSharedPtr);
+
+struct FAsyncPathFindingQuery : public FPathFindingQuery
+{
+	const uint32 QueryID;
+	const FNavPathQueryDelegate OnDoneDelegate;
+	const TEnumAsByte<EPathFindingMode::Type> Mode;
+	FPathFindingResult Result;
+
+	FAsyncPathFindingQuery()
+		: QueryID(INVALID_NAVQUERYID)
+	{
+	}
+
+	FAsyncPathFindingQuery(const UObject* InOwner, const class ANavigationData* InNavData, const FVector& Start, const FVector& End, const FNavPathQueryDelegate& Delegate, TSharedPtr<const FNavigationQueryFilter> SourceQueryFilter);
+	FAsyncPathFindingQuery(const FPathFindingQuery& Query, const FNavPathQueryDelegate& Delegate, const EPathFindingMode::Type QueryMode);
+
+protected:
+	FORCEINLINE static uint32 GetUniqueID()
+	{
+		return ++LastPathFindingUniqueID;
+	}
+
+	static uint32 LastPathFindingUniqueID;
+};
+
 //////////////////////////////////////////////////////////////////////////
 // Custom path following data
 
diff --git a/Engine/Source/Runtime/Engine/Classes/AI/Navigation/RecastNavMesh.h b/Engine/Source/Runtime/Engine/Classes/AI/Navigation/RecastNavMesh.h
--- a/Engine/Source/Runtime/Engine/Classes/AI/Navigation/RecastNavMesh.h
+++ b/Engine/Source/Runtime/Engine/Classes/AI/Navigation/RecastNavMesh.h
@@ -4,6 +4,7 @@
 
 #include "NavAreas/NavArea.h"
 #include "Tickable.h"
+#include "AI/Navigation/NavigationData.h"
 #include "RecastNavMesh.generated.h"
 
 /** Initial checkin. */
@@ -30,6 +31,12 @@
 #define RECAST_STRAIGHTPATH_OFFMESH_CONNECTION 0x04
 #define RECAST_UNWALKABLE_POLY_COST	FLT_MAX
 
+class FNavDataGenerator;
+class FPImplRecastNavMesh;
+class FRecastQueryFilter;
+class USmartNavLinkComponent;
+class UNavArea;
+
 UENUM()
 namespace ERecastPartitioning
 {
@@ -654,13 +661,13 @@ public:
 
 protected:
 #if WITH_NAVIGATION_GENERATOR
-	virtual class FNavDataGenerator* ConstructGenerator(const FNavAgentProperties& AgentProps) OVERRIDE;
+	virtual FNavDataGenerator* ConstructGenerator(const FNavAgentProperties& AgentProps) OVERRIDE;
 	virtual void SortNavigationTiles();
 #endif // WITH_NAVIGATION_GENERATOR
 	// End ANavigationData Interface
 
 	/** Serialization helper. */
-	void SerializeRecastNavMesh(FArchive& Ar, class FPImplRecastNavMesh*& NavMesh);
+	void SerializeRecastNavMesh(FArchive& Ar, FPImplRecastNavMesh*& NavMesh);
 
 public:
 	/** Returns bounding box for the whole navmesh. */
@@ -737,13 +744,13 @@ public:
 	// Smart links
 	//----------------------------------------------------------------------//
 
-	virtual void UpdateSmartLink(class USmartNavLinkComponent* LinkComp) OVERRIDE;
+	virtual void UpdateSmartLink(USmartNavLinkComponent* LinkComp) OVERRIDE;
 
 	/** update area class and poly flags for all offmesh links with given UserId */
-	void UpdateNavigationLinkArea(int32 UserId, TSubclassOf<class UNavArea> AreaClass) const;
+	void UpdateNavigationLinkArea(int32 UserId, TSubclassOf<UNavArea> AreaClass) const;
 
 	/** update area class and poly flags for all offmesh segment links with given UserId */
-	void UpdateSegmentLinkArea(int32 UserId, TSubclassOf<class UNavArea> AreaClass) const;
+	void UpdateSegmentLinkArea(int32 UserId, TSubclassOf<UNavArea> AreaClass) const;
 
 	//----------------------------------------------------------------------//
 	// Batch processing (important with async rebuilding)
@@ -804,7 +811,7 @@ public:
 	bool GetClustersWithinPathingDistance(FVector const& StartLoc, const float PathingDistance, TArray<NavNodeRef>& FoundClusters, bool bBackTracking = false) const;
 
 	/** Filters nav polys in PolyRefs with Filter */
-	bool FilterPolys(TArray<NavNodeRef>& PolyRefs, const class FRecastQueryFilter* Filter, const UObject* Querier = NULL) const;
+	bool FilterPolys(TArray<NavNodeRef>& PolyRefs, const FRecastQueryFilter* Filter, const UObject* Querier = NULL) const;
 
 	/** Get all polys from tile */
 	bool GetPolysInTile(int32 TileIndex, TArray<FNavPoly>& Polys) const;
@@ -834,7 +841,7 @@ public:
 	/** Checks whether this instance of navmesh supports given Agent type */
 	virtual bool DoesSupportAgent(const FNavAgentProperties& AgentProps) const OVERRIDE;
 
-	static const class FRecastQueryFilter* GetNamedFilter(ERecastNamedFilter::Type FilterType);
+	static const FRecastQueryFilter* GetNamedFilter(ERecastNamedFilter::Type FilterType);
 	FORCEINLINE static FNavPolyFlags GetNavLinkFlag() { return NavLinkFlag; }
 	
 	/** initialize tiles set with given size */
@@ -891,7 +898,7 @@ private:
 	 *	@NOTE: if we switch over to C++11 this should be unique_ptr
 	 *	@TODO since it's no secret we're using recast there's no point in having separate implementation class. FPImplRecastNavMesh should be merged into ARecastNavMesh
 	 */
-	class FPImplRecastNavMesh* RecastNavMeshImpl;
+	FPImplRecastNavMesh* RecastNavMeshImpl;
 
 	/** array containing tile references of all navmesh tiles recently rebuild. Handled and cleared in Tick function */
 	TArray<uint32> FreshTiles;
diff --git a/Engine/Source/Runtime/Engine/Classes/Engine/World.h b/Engine/Source/Runtime/Engine/Classes/Engine/World.h
--- a/Engine/Source/Runtime/Engine/Classes/Engine/World.h
+++ b/Engine/Source/Runtime/Engine/Classes/Engine/World.h
@@ -570,12 +570,9 @@ private:
 	UPROPERTY(Transient)
 	class AGameMode*							AuthorityGameMode;
 
-	/** Behavior tree manager used by game */	
-	TWeakObjectPtr<class UBehaviorTreeManager> BehaviorTreeManager;
-
-	/** Environment query manager used by game */
-	TWeakObjectPtr<class UEnvQueryManager>		EnvironmentQueryManager;
-
+	UPROPERTY(Transient)
+	class UAISystemBase*						AISystem;
+	
 	/** RVO avoidance manager used by game */
 	UPROPERTY(Transient)
 	class UAvoidanceManager*					AvoidanceManager;
@@ -1415,19 +1412,14 @@ public:
 	/** NavigationSystem const getter */
 	FORCEINLINE const UNavigationSystem* GetNavigationSystem() const { return NavigationSystem; }
 
-	/** Behavior tree manager getter */
-	FORCEINLINE class UBehaviorTreeManager* GetBehaviorTreeManager() { return BehaviorTreeManager.Get(); }
-	/** Behavior tree manager const getter */
-	FORCEINLINE const class UBehaviorTreeManager* GetBehaviorTreeManager() const { return BehaviorTreeManager.Get(); }
-
-	/** Behavior tree manager getter */
-	FORCEINLINE class UEnvQueryManager* GetEnvironmentQueryManager() { return EnvironmentQueryManager.Get(); }
-	/** Behavior tree manager const getter */
-	FORCEINLINE const class UEnvQueryManager* GetEnvironmentQueryManager() const { return EnvironmentQueryManager.Get(); }
+	/** AISystem getter */
+	FORCEINLINE UAISystemBase* GetAISystem() { return AISystem; }
+	/** AISystem const getter */
+	FORCEINLINE const UAISystemBase* GetAISystem() const { return AISystem; }
 
-	/** Behavior tree manager getter */
+	/** Avoidance manager getter */
 	FORCEINLINE class UAvoidanceManager* GetAvoidanceManager() { return AvoidanceManager; }
-	/** Behavior tree manager const getter */
+	/** Avoidance manager getter */
 	FORCEINLINE const class UAvoidanceManager* GetAvoidanceManager() const { return AvoidanceManager; }
 
 	/** Returns an iterator for the controller list. */
diff --git a/Engine/Source/Runtime/Engine/Classes/GameFramework/CheatManager.h b/Engine/Source/Runtime/Engine/Classes/GameFramework/CheatManager.h
--- a/Engine/Source/Runtime/Engine/Classes/GameFramework/CheatManager.h
+++ b/Engine/Source/Runtime/Engine/Classes/GameFramework/CheatManager.h
@@ -297,10 +297,6 @@ class ENGINE_API UCheatManager : public UObject
 	UFUNCTION(exec)
 	void ToggleGameplayDebugView(const FString& ViewName);
 
-	/** insta-runs EQS query for GameplayDebugComponent selected AI */
-	UFUNCTION(exec)
-	void RunEQS(const FString& QueryName);
-
 	/**
 	 * This will move the player and set their rotation to the passed in values.
 	 * This actually does the location / rotation setting.  Additionally it will set you as ghost as the level may have
diff --git a/Engine/Source/Runtime/Engine/Private/AI/Navigation/NavFilters/NavigationQueryFilter.cpp b/Engine/Source/Runtime/Engine/Private/AI/Navigation/NavFilters/NavigationQueryFilter.cpp
--- a/Engine/Source/Runtime/Engine/Private/AI/Navigation/NavFilters/NavigationQueryFilter.cpp
+++ b/Engine/Source/Runtime/Engine/Private/AI/Navigation/NavFilters/NavigationQueryFilter.cpp
@@ -5,6 +5,111 @@
 #include "Navigation/NavigationComponent.h"
 #include "AI/Navigation/NavFilters/NavigationQueryFilter.h"
 
+//----------------------------------------------------------------------//
+// FNavigationQueryFilter
+//----------------------------------------------------------------------//
+FNavigationQueryFilter::FNavigationQueryFilter(const FNavigationQueryFilter& Source)
+{
+	Assign(Source);
+}
+
+FNavigationQueryFilter::FNavigationQueryFilter(const FNavigationQueryFilter* Source)
+: MaxSearchNodes(DefaultMaxSearchNodes)
+{
+	if (Source != NULL)
+	{
+		Assign(*Source);
+	}
+}
+
+FNavigationQueryFilter::FNavigationQueryFilter(const TSharedPtr<FNavigationQueryFilter> Source)
+: MaxSearchNodes(DefaultMaxSearchNodes)
+{
+	if (Source.IsValid())
+	{
+		SetFilterImplementation(Source->GetImplementation());
+	}
+}
+
+FNavigationQueryFilter& FNavigationQueryFilter::operator=(const FNavigationQueryFilter& Source)
+{
+	Assign(Source);
+	return *this;
+}
+
+void FNavigationQueryFilter::Assign(const FNavigationQueryFilter& Source)
+{
+	if (Source.GetImplementation() != NULL)
+	{
+		QueryFilterImpl = Source.QueryFilterImpl;
+	}
+	MaxSearchNodes = Source.GetMaxSearchNodes();
+}
+
+TSharedPtr<FNavigationQueryFilter> FNavigationQueryFilter::GetCopy() const
+{
+	TSharedPtr<FNavigationQueryFilter> Copy = MakeShareable(new FNavigationQueryFilter());
+	Copy->QueryFilterImpl = MakeShareable(QueryFilterImpl->CreateCopy());
+	Copy->MaxSearchNodes = MaxSearchNodes;
+
+	return Copy;
+}
+
+void FNavigationQueryFilter::SetAreaCost(uint8 AreaType, float Cost)
+{
+	check(QueryFilterImpl.IsValid());
+	QueryFilterImpl->SetAreaCost(AreaType, Cost);
+}
+
+void FNavigationQueryFilter::SetFixedAreaEnteringCost(uint8 AreaType, float Cost)
+{
+	check(QueryFilterImpl.IsValid());
+	QueryFilterImpl->SetFixedAreaEnteringCost(AreaType, Cost);
+}
+
+void FNavigationQueryFilter::SetExcludedArea(uint8 AreaType)
+{
+	QueryFilterImpl->SetExcludedArea(AreaType);
+}
+
+void FNavigationQueryFilter::SetAllAreaCosts(const TArray<float>& CostArray)
+{
+	SetAllAreaCosts(CostArray.GetTypedData(), CostArray.Num());
+}
+
+void FNavigationQueryFilter::SetAllAreaCosts(const float* CostArray, const int32 Count)
+{
+	QueryFilterImpl->SetAllAreaCosts(CostArray, Count);
+}
+
+void FNavigationQueryFilter::GetAllAreaCosts(float* CostArray, float* FixedCostArray, const int32 Count) const
+{
+	QueryFilterImpl->GetAllAreaCosts(CostArray, FixedCostArray, Count);
+}
+
+void FNavigationQueryFilter::SetIncludeFlags(uint16 Flags)
+{
+	QueryFilterImpl->SetIncludeFlags(Flags);
+}
+
+uint16 FNavigationQueryFilter::GetIncludeFlags() const
+{
+	return QueryFilterImpl->GetIncludeFlags();
+}
+
+void FNavigationQueryFilter::SetExcludeFlags(uint16 Flags)
+{
+	QueryFilterImpl->SetExcludeFlags(Flags);
+}
+
+uint16 FNavigationQueryFilter::GetExcludeFlags() const
+{
+	return QueryFilterImpl->GetExcludeFlags();
+}
+
+//----------------------------------------------------------------------//
+// UNavigationQueryFilter
+//----------------------------------------------------------------------//
 UNavigationQueryFilter::UNavigationQueryFilter(const class FPostConstructInitializeProperties& PCIP) : Super(PCIP)
 {
 	IncludeFlags.Packed = 0xffff;
diff --git a/Engine/Source/Runtime/Engine/Private/AI/Navigation/NavigationData.cpp b/Engine/Source/Runtime/Engine/Private/AI/Navigation/NavigationData.cpp
--- a/Engine/Source/Runtime/Engine/Private/AI/Navigation/NavigationData.cpp
+++ b/Engine/Source/Runtime/Engine/Private/AI/Navigation/NavigationData.cpp
@@ -5,6 +5,57 @@
 #include "AI/Navigation/NavAreas/NavAreaMeta.h"
 #include "AI/Navigation/NavigationData.h"
 
+//----------------------------------------------------------------------//
+// FPathFindingQuery
+//----------------------------------------------------------------------//
+FPathFindingQuery::FPathFindingQuery(const UObject* InOwner, const class ANavigationData* InNavData, const FVector& Start, const FVector& End, TSharedPtr<const FNavigationQueryFilter> SourceQueryFilter)
+: NavData(InNavData)
+, Owner(InOwner)
+, StartLocation(Start)
+, EndLocation(End)
+, QueryFilter(SourceQueryFilter)
+{
+	if (SourceQueryFilter.IsValid() == false && NavData.IsValid() == true)
+	{
+		QueryFilter = NavData->GetDefaultQueryFilter();
+	}
+}
+
+FPathFindingQuery::FPathFindingQuery(const FPathFindingQuery& Source)
+: NavData(Source.NavData)
+, Owner(Source.Owner)
+, StartLocation(Source.StartLocation)
+, EndLocation(Source.EndLocation)
+, QueryFilter(Source.QueryFilter)
+, NavDataFlags(Source.NavDataFlags)
+{
+	if (Source.QueryFilter.IsValid() == false && NavData.IsValid() == true)
+	{
+		QueryFilter = NavData->GetDefaultQueryFilter();
+	}
+}
+
+//----------------------------------------------------------------------//
+// FAsyncPathFindingQuery
+//----------------------------------------------------------------------//
+uint32 FAsyncPathFindingQuery::LastPathFindingUniqueID = INVALID_NAVQUERYID;
+
+FAsyncPathFindingQuery::FAsyncPathFindingQuery(const UObject* InOwner, const class ANavigationData* InNavData, const FVector& Start, const FVector& End, const FNavPathQueryDelegate& Delegate, TSharedPtr<const FNavigationQueryFilter> SourceQueryFilter)
+: FPathFindingQuery(InOwner, InNavData, Start, End, SourceQueryFilter)
+, QueryID(GetUniqueID())
+, OnDoneDelegate(Delegate)
+{
+
+}
+
+FAsyncPathFindingQuery::FAsyncPathFindingQuery(const FPathFindingQuery& Query, const FNavPathQueryDelegate& Delegate, const EPathFindingMode::Type QueryMode)
+: FPathFindingQuery(Query)
+, QueryID(GetUniqueID())
+, OnDoneDelegate(Delegate)
+, Mode(QueryMode)
+{
+
+}
 //----------------------------------------------------------------------//
 // FNavDataGenerator
 //----------------------------------------------------------------------//
diff --git a/Engine/Source/Runtime/Engine/Private/AI/Navigation/NavigationSystem.cpp b/Engine/Source/Runtime/Engine/Private/AI/Navigation/NavigationSystem.cpp
--- a/Engine/Source/Runtime/Engine/Private/AI/Navigation/NavigationSystem.cpp
+++ b/Engine/Source/Runtime/Engine/Private/AI/Navigation/NavigationSystem.cpp
@@ -26,9 +26,9 @@
 static const uint32 INITIAL_ASYNC_QUERIES_SIZE = 32;
 static const uint32 REGISTRATION_QUEUE_SIZE = 16;	// and we'll not reallocate
 #if WITH_RECAST
-static const uint32 MAX_SEARCH_NODES = RECAST_MAX_SEARCH_NODES;
+static const uint32 MAX_NAV_SEARCH_NODES = RECAST_MAX_SEARCH_NODES;
 #else // WITH_RECAST
-static const uint32 MAX_SEARCH_NODES = 2048;
+static const uint32 MAX_NAV_SEARCH_NODES = 2048;
 #endif // WITH_RECAST
 
 DEFINE_LOG_CATEGORY(LogNavigation);
@@ -73,7 +73,10 @@ DEFINE_STAT(STAT_Navigation_PathVisibilityOptimisation);
 
 //----------------------------------------------------------------------//
 // consts
-//----------------------------------------------------------	------------//
+//----------------------------------------------------------------------//
+
+const uint32 FNavigationQueryFilter::DefaultMaxSearchNodes = MAX_NAV_SEARCH_NODES;
+
 namespace FNavigationSystem
 {
 	// these are totally arbitrary values, and it should haven happen these are ever used.
@@ -95,164 +98,6 @@ namespace NavigationDebugDrawing
 	const FVector PathNodeBoxExtent(16.f);
 }
 
-//----------------------------------------------------------------------//
-// FNavigationQueryFilter
-//----------------------------------------------------------------------//
-const uint32 FNavigationQueryFilter::DefaultMaxSearchNodes = MAX_SEARCH_NODES;
-
-FNavigationQueryFilter::FNavigationQueryFilter(const FNavigationQueryFilter& Source)
-{	
-	Assign(Source);
-}
-
-FNavigationQueryFilter::FNavigationQueryFilter(const FNavigationQueryFilter* Source)
-	: MaxSearchNodes(DefaultMaxSearchNodes)
-{	
-	if (Source != NULL)
-	{
-		Assign(*Source);
-	}
-}
-
-FNavigationQueryFilter::FNavigationQueryFilter(const TSharedPtr<FNavigationQueryFilter> Source)
-	: MaxSearchNodes(DefaultMaxSearchNodes)
-{
-	if (Source.IsValid())
-	{
-		SetFilterImplementation(Source->GetImplementation());
-	}
-}
-
-FNavigationQueryFilter& FNavigationQueryFilter::operator=(const FNavigationQueryFilter& Source)
-{
-	Assign(Source);
-	return *this;
-}
-
-void FNavigationQueryFilter::Assign(const FNavigationQueryFilter& Source)
-{
-	if (Source.GetImplementation() != NULL)
-	{
-		QueryFilterImpl = Source.QueryFilterImpl;
-	}
-	MaxSearchNodes = Source.GetMaxSearchNodes();
-}
-
-TSharedPtr<FNavigationQueryFilter> FNavigationQueryFilter::GetCopy() const
-{
-	TSharedPtr<FNavigationQueryFilter> Copy = MakeShareable(new FNavigationQueryFilter());
-	Copy->QueryFilterImpl = MakeShareable(QueryFilterImpl->CreateCopy());
-	Copy->MaxSearchNodes = MaxSearchNodes;
-
-	return Copy;
-}
-
-void FNavigationQueryFilter::SetAreaCost(uint8 AreaType, float Cost)
-{
-	check(QueryFilterImpl.IsValid());
-	QueryFilterImpl->SetAreaCost(AreaType, Cost);
-}
-
-void FNavigationQueryFilter::SetFixedAreaEnteringCost(uint8 AreaType, float Cost)
-{
-	check(QueryFilterImpl.IsValid());
-	QueryFilterImpl->SetFixedAreaEnteringCost(AreaType, Cost);
-}
-
-void FNavigationQueryFilter::SetExcludedArea(uint8 AreaType)
-{
-	QueryFilterImpl->SetExcludedArea(AreaType);
-}
-
-void FNavigationQueryFilter::SetAllAreaCosts(const TArray<float>& CostArray)
-{
-	SetAllAreaCosts(CostArray.GetTypedData(), CostArray.Num());
-}
-
-void FNavigationQueryFilter::SetAllAreaCosts(const float* CostArray, const int32 Count)
-{
-	QueryFilterImpl->SetAllAreaCosts(CostArray, Count);
-}
-
-void FNavigationQueryFilter::GetAllAreaCosts(float* CostArray, float* FixedCostArray, const int32 Count) const
-{
-	QueryFilterImpl->GetAllAreaCosts(CostArray, FixedCostArray, Count);
-}
-
-void FNavigationQueryFilter::SetIncludeFlags(uint16 Flags)
-{
-	QueryFilterImpl->SetIncludeFlags(Flags);
-}
-
-uint16 FNavigationQueryFilter::GetIncludeFlags() const
-{
-	return QueryFilterImpl->GetIncludeFlags();
-}
-
-void FNavigationQueryFilter::SetExcludeFlags(uint16 Flags)
-{
-	QueryFilterImpl->SetExcludeFlags(Flags);
-}
-
-uint16 FNavigationQueryFilter::GetExcludeFlags() const
-{
-	return QueryFilterImpl->GetExcludeFlags();
-}
-
-
-//----------------------------------------------------------------------//
-// 
-//----------------------------------------------------------------------//
-FPathFindingQuery::FPathFindingQuery(const UObject* InOwner, const class ANavigationData* InNavData, const FVector& Start, const FVector& End, TSharedPtr<const FNavigationQueryFilter> SourceQueryFilter)
-	: NavData(InNavData)
-	, Owner(InOwner)
-	, StartLocation(Start)
-	, EndLocation(End)
-	, QueryFilter(SourceQueryFilter)
-{
-	if (SourceQueryFilter.IsValid() == false && NavData.IsValid() == true)
-	{
-		QueryFilter = NavData->GetDefaultQueryFilter();
-	}
-}
-
-FPathFindingQuery::FPathFindingQuery(const FPathFindingQuery& Source)
-	: NavData(Source.NavData)
-	, Owner(Source.Owner)
-	, StartLocation(Source.StartLocation)
-	, EndLocation(Source.EndLocation)
-	, QueryFilter(Source.QueryFilter)
-	, NavDataFlags(Source.NavDataFlags)
-{
-	if (Source.QueryFilter.IsValid() == false && NavData.IsValid() == true)
-	{
-		QueryFilter = NavData->GetDefaultQueryFilter();
-	}
-}
-
-//----------------------------------------------------------------------//
-// FAsyncPathFindingQuery
-//----------------------------------------------------------------------//
-uint32 FAsyncPathFindingQuery::LastPathFindingUniqueID = INVALID_NAVQUERYID;
-
-FAsyncPathFindingQuery::FAsyncPathFindingQuery(const UObject* InOwner, const class ANavigationData* InNavData, const FVector& Start, const FVector& End, const FNavPathQueryDelegate& Delegate, TSharedPtr<const FNavigationQueryFilter> SourceQueryFilter)
-	: FPathFindingQuery(InOwner, InNavData, Start, End, SourceQueryFilter)
-	, QueryID(GetUniqueID())
-	, OnDoneDelegate(Delegate)
-{
-
-}
-
-FAsyncPathFindingQuery::FAsyncPathFindingQuery(const FPathFindingQuery& Query, const FNavPathQueryDelegate& Delegate, const EPathFindingMode::Type QueryMode)
-	: FPathFindingQuery(Query)
-	, QueryID(GetUniqueID())
-	, OnDoneDelegate(Delegate)
-	, Mode(QueryMode)
-{
-
-}
-
-
 //----------------------------------------------------------------------//
 // UNavigationSystem                                                                
 //----------------------------------------------------------------------//
diff --git a/Engine/Source/Runtime/Engine/Private/CheatManager.cpp b/Engine/Source/Runtime/Engine/Private/CheatManager.cpp
--- a/Engine/Source/Runtime/Engine/Private/CheatManager.cpp
+++ b/Engine/Source/Runtime/Engine/Private/CheatManager.cpp
@@ -7,10 +7,6 @@
 #include "NavDataGenerator.h"
 #include "OnlineSubsystemUtils.h"
 #include "VisualLog.h"
-// @todo this is here only due to circular dependency to AIModule. To be removed
-#include "AIController.h"
-#include "EnvironmentQuery/EnvQuery.h"
-#include "EnvironmentQuery/EnvQueryManager.h"
 
 #if WITH_EDITOR
 #include "UnrealEd.h"
@@ -1102,59 +1098,6 @@ void UCheatManager::ToggleGameplayDebugView(const FString& InViewName)
 	}
 }
 
-void UCheatManager::RunEQS(const FString& QueryName)
-{
-	APlayerController* MyPC = GetOuterAPlayerController();
-	UEnvQueryManager* EQS = GetWorld()->GetEnvironmentQueryManager();
-	if (MyPC && EQS)
-	{
-		AActor* Target = NULL;
-		UGameplayDebuggingControllerComponent* DebugComp = MyPC->FindComponentByClass<UGameplayDebuggingControllerComponent>();
-		if (DebugComp)
-		{
-			Target = DebugComp->GetCurrentDebugTarget();
-		}
-
-#if WITH_EDITOR
-		if (Target == NULL && GEditor != NULL)
-		{
-			Target = GEditor->GetSelectedObjects()->GetTop<AActor>();
-
-			// this part should not be needed, but is due to gameplay debugging messed up design
-			if (Target == NULL)
-			{
-				for (FObjectIterator It(UGameplayDebuggingComponent::StaticClass()); It && (Target == NULL); ++It)
-				{
-					UGameplayDebuggingComponent* Comp = ((UGameplayDebuggingComponent*)*It);
-					if (Comp->IsSelected())
-					{
-						Target = Comp->GetOwner();
-					}
-				}
-			}
-		}
-#endif // WITH_EDITOR
-
-		if (Target)
-		{
-			UEnvQuery* QueryTemplate = EQS->FindQueryTemplate(QueryName);
-			
-			if (QueryTemplate)
-			{
-				EQS->RunInstantQuery(FEnvQueryRequest(QueryTemplate, Target), EEnvQueryRunMode::AllMatching);
-			}
-			else
-			{
-				GetOuterAPlayerController()->ClientMessage(FString::Printf(TEXT("Unable to fing query template \'%s\'"), *QueryName));
-			}
-		}
-		else
-		{
-			GetOuterAPlayerController()->ClientMessage(TEXT("No debugging target"));
-		}
-	}
-}
-
 void UCheatManager::LogOutBugItGoToLogFile( const FString& InScreenShotDesc, const FString& InGoString, const FString& InLocString )
 {
 #if ALLOW_DEBUG_FILES
diff --git a/Engine/Source/Runtime/Engine/Private/World.cpp b/Engine/Source/Runtime/Engine/Private/World.cpp
--- a/Engine/Source/Runtime/Engine/Private/World.cpp
+++ b/Engine/Source/Runtime/Engine/Private/World.cpp
@@ -27,6 +27,7 @@
 #include "VisualLog.h"
 #include "LevelUtils.h"
 #include "PhysicsPublic.h"
+#include "AI/AISystemBase.h"
 #if WITH_EDITOR
 	#include "Editor/UnrealEd/Public/Kismet2/KismetEditorUtilities.h"
 	#include "Editor/UnrealEd/Public/Kismet2/BlueprintEditorUtils.h"
@@ -38,10 +39,6 @@
 #include "Landscape/LandscapeInfo.h"
 #include "Particles/ParticleEventManager.h"
 
-// @todo this is here only due to circular dependency to AIModule. To be removed
-#include "BehaviorTree/BehaviorTreeManager.h"
-#include "EnvironmentQuery/EnvQueryManager.h"
-
 DEFINE_LOG_CATEGORY_STATIC(LogWorld, Log, All);
 DEFINE_LOG_CATEGORY(LogSpawn);
 
@@ -126,8 +123,6 @@ void UWorld::Serialize( FArchive& Ar )
 		[REDACTED_EMAIL]: How to handle w/out having a cooker/pakcager?
 #endif	//#if WITH_EDITORONLY_DATA
 
-		Ar << BehaviorTreeManager;
-		Ar << EnvironmentQueryManager;
 		Ar << NavigationSystem;
 		Ar << AvoidanceManager;
 	}
@@ -208,11 +203,6 @@ void UWorld::AddReferencedObjects(UObject* InThis, FReferenceCollector& Collecto
 	}
 #endif
 
-	UBehaviorTreeManager* BTManager = This->BehaviorTreeManager.GetEvenIfUnreachable();
-	Collector.AddReferencedObject(BTManager, This);
-	UEnvQueryManager* EQSManager = This->EnvironmentQueryManager.GetEvenIfUnreachable();
-	Collector.AddReferencedObject(EQSManager, This);
-
 	Super::AddReferencedObjects( This, Collector );
 }
 
@@ -679,9 +669,16 @@ void UWorld::InitWorld(const InitializationValues IVS)
 		UNavigationSystem::CreateNavigationSystem(this);
 	}
 
-	BehaviorTreeManager = NewObject<UBehaviorTreeManager>(this);
-	EnvironmentQueryManager = NewObject<UEnvQueryManager>(this);
-
+	auto AISystemModule = FModuleManager::LoadModulePtr<IAISystemModule>(UAISystemBase::GetAISystemModuleName());
+	if (AISystemModule)
+	{
+		AISystem = AISystemModule->CreateAISystemInstance(this);
+		if (AISystem == NULL)
+		{
+			UE_LOG(LogWorld, Error, TEXT("Failed to create AISystem instance of class %s!"), *UAISystemBase::GetAISystemClassName().ClassName);
+		}
+	}
+	
 	if (GEngine->AvoidanceManagerClass != NULL)
 	{
 		AvoidanceManager = NewObject<UAvoidanceManager>(this, GEngine->AvoidanceManagerClass);
@@ -2720,6 +2717,11 @@ void UWorld::InitializeActorsForPlay(const FURL& InURL, bool bResetTime)
 		NavigationSystem->OnInitializeActors();
 	}
 
+	if (AISystem != NULL)
+	{
+		AISystem->InitializeActorsForPlay(bResetTime);
+	}
+
 	for( int32 LevelIndex=0; LevelIndex<Levels.Num(); LevelIndex++ )
 	{
 		ULevel*	Level = Levels[LevelIndex];
@@ -2756,6 +2758,11 @@ void UWorld::CleanupWorld(bool bSessionEnded, bool bCleanupResources, UWorld* Ne
 
 	FWorldDelegates::OnWorldCleanup.Broadcast(this, bSessionEnded, bCleanupResources);
 
+	if (AISystem != NULL)
+	{
+		AISystem->CleanupWorld(bSessionEnded, bCleanupResources, NewWorld);
+	}
+
 	if (bCleanupResources == true)
 	{
 		if (NavigationSystem != NULL)
@@ -3884,6 +3891,11 @@ bool UWorld::SetNewWorldOrigin(const FIntPoint& InNewOrigin)
 		}
 	}
 
+	if (AISystem != NULL)
+	{
+		AISystem->WorldOriginChanged(PreviosWorldOrigin, GlobalOriginOffset);
+	}
+
 	// Broadcast that have finished world shifting
 	FCoreDelegates::PostWorldOriginOffset.Broadcast(this, PreviosWorldOrigin, GlobalOriginOffset);
 
