Moved an ETrailWidthMode enum to ParticleHelper.h to get around a compile error.

diff --git a/Engine/Source/Runtime/Engine/Classes/Animation/AnimNotifies/AnimNotifyState_Trail.h b/Engine/Source/Runtime/Engine/Classes/Animation/AnimNotifies/AnimNotifyState_Trail.h
--- a/Engine/Source/Runtime/Engine/Classes/Animation/AnimNotifies/AnimNotifyState_Trail.h
+++ b/Engine/Source/Runtime/Engine/Classes/Animation/AnimNotifies/AnimNotifyState_Trail.h
@@ -1,6 +1,7 @@
 // Copyright 1998-2014 Epic Games, Inc. All Rights Reserved.
 
 #pragma once
+#include "ParticleHelper.h"
 #include "AnimNotifyState_Trail.generated.h"
 
 DECLARE_LOG_CATEGORY_EXTERN(LogAnimTrails, Log, All);
diff --git a/Engine/Source/Runtime/Engine/Classes/Particles/ParticleSystemComponent.h b/Engine/Source/Runtime/Engine/Classes/Particles/ParticleSystemComponent.h
--- a/Engine/Source/Runtime/Engine/Classes/Particles/ParticleSystemComponent.h
+++ b/Engine/Source/Runtime/Engine/Classes/Particles/ParticleSystemComponent.h
@@ -20,15 +20,6 @@ struct FDynamicEmitterDataBase;
 struct FDynamicEmitterReplayDataBase;
 struct FParticleAnimTrailEmitterInstance;
 
-/** Controls the way that the width scale property affects anim trails */
-UENUM()
-enum ETrailWidthMode
-{
-	ETrailWidthMode_FromCentre UMETA(DisplayName = "From Centre"),
-	ETrailWidthMode_FromFirst UMETA(DisplayName = "From First Socket"),
-	ETrailWidthMode_FromSecond UMETA(DisplayName = "From Second Socket"),
-};
-
 /** Enum for specifying type of a name instance parameter. */
 UENUM()
 enum EParticleSysParamType
diff --git a/Engine/Source/Runtime/Engine/Public/ParticleHelper.h b/Engine/Source/Runtime/Engine/Public/ParticleHelper.h
--- a/Engine/Source/Runtime/Engine/Public/ParticleHelper.h
+++ b/Engine/Source/Runtime/Engine/Public/ParticleHelper.h
@@ -10,6 +10,8 @@
 #include "ParticleBeamTrailVertexFactory.h"
 #include "MeshParticleVertexFactory.h"
 
+#include "ParticleHelper.generated.h"
+
 #define _ENABLE_PARTICLE_LOD_INGAME_
 
 DECLARE_LOG_CATEGORY_EXTERN(LogParticles, Log, All);
@@ -2867,3 +2869,11 @@ public:
 /** The global null color vertex buffer, which is set with a stride of 0 on meshes without a color component. */
 extern TGlobalResource<FNullDynamicParameterVertexBuffer> GNullDynamicParameterVertexBuffer;
 
+/** Controls the way that the width scale property affects anim trails */
+UENUM()
+enum ETrailWidthMode
+{
+	ETrailWidthMode_FromCentre UMETA(DisplayName = "From Centre"),
+	ETrailWidthMode_FromFirst UMETA(DisplayName = "From First Socket"),
+	ETrailWidthMode_FromSecond UMETA(DisplayName = "From Second Socket"),
+};
\ No newline at end of file
