Files missing in prev submission #UE4

diff --git a/Engine/Source/Runtime/Engine/Classes/AI/AISystemBase.h b/Engine/Source/Runtime/Engine/Classes/AI/AISystemBase.h
new file mode 100644
--- /dev/null
+++ b/Engine/Source/Runtime/Engine/Classes/AI/AISystemBase.h
@@ -0,0 +1,54 @@
+// Copyright 1998-2014 Epic Games, Inc. All Rights Reserved.
+
+#pragma once
+#include "Math/IntPoint.h"
+#include "ModuleManager.h"
+#include "AISystemBase.generated.h"
+
+class UWorld;
+
+UCLASS(abstract, config = Engine, defaultconfig)
+class ENGINE_API UAISystemBase : public UObject
+{
+	GENERATED_UCLASS_BODY()
+
+	virtual ~UAISystemBase(){}
+		
+	/** 
+	 * Called when world initializes all actors and prepares them to start gameplay
+	 * @param bTimeGotReset whether the WorldSettings's TimeSeconds are reset to zero
+	 */
+	virtual void InitializeActorsForPlay(bool bTimeGotReset) PURE_VIRTUAL(UAISystemBase::InitializeActorsForPlay, );
+
+	/**
+	 * Event called on world origin changes
+	 *
+	 * @param	OldOrigin			Previous world origin position
+	 * @param	NewOrigin			New world origin position
+	 */
+	virtual void WorldOriginChanged(FIntPoint OldOrigin, FIntPoint NewOrigin) PURE_VIRTUAL(UAISystemBase::WorldOriginChanged, );
+
+	/**
+	 * Called by UWorld::CleanupWorld.
+	 * @param bSessionEnded whether to notify the viewport that the game session has ended
+	 * @param NewWorld Optional new world that will be loaded after this world is cleaned up. Specify a new world to prevent it and it's sublevels from being GCed during map transitions.
+	 */
+	virtual void CleanupWorld(bool bSessionEnded = true, bool bCleanupResources = true, UWorld* NewWorld = NULL) PURE_VIRTUAL(UAISystemBase::CleanupWorld, );
+
+private:
+	UPROPERTY(globalconfig, noclear, meta = (MetaClass = "AISystem", DisplayName = "AISystem Class"))
+	FStringClassReference AISystemClassName;
+
+	UPROPERTY(globalconfig, noclear, meta = (MetaClass = "AISystem", DisplayName = "AISystem Module"))
+	FName AISystemModuleName;
+	
+public:
+	static FStringClassReference GetAISystemClassName();
+	static FName GetAISystemModuleName();
+};
+
+class IAISystemModule : public IModuleInterface
+{
+public:
+	virtual UAISystemBase* CreateAISystemInstance(UWorld* World) = 0;
+};
\ No newline at end of file
diff --git a/Engine/Source/Runtime/Engine/Private/AI/AISystemBase.cpp b/Engine/Source/Runtime/Engine/Private/AI/AISystemBase.cpp
new file mode 100644
--- /dev/null
+++ b/Engine/Source/Runtime/Engine/Private/AI/AISystemBase.cpp
@@ -0,0 +1,21 @@
+// Copyright 1998-2014 Epic Games, Inc. All Rights Reserved.
+
+#include "EnginePrivate.h"
+#include "AI/AISystemBase.h"
+
+UAISystemBase::UAISystemBase(const class FPostConstructInitializeProperties& PCIP)
+: Super(PCIP)
+{
+}
+
+FName UAISystemBase::GetAISystemModuleName()
+{
+	UAISystemBase* AISystemDefaultObject = Cast<UAISystemBase>(StaticClass()->GetDefaultObject());
+	return AISystemDefaultObject != NULL ? AISystemDefaultObject->AISystemModuleName : TEXT("");
+}
+
+FStringClassReference UAISystemBase::GetAISystemClassName()
+{
+	UAISystemBase* AISystemDefaultObject = Cast<UAISystemBase>(StaticClass()->GetDefaultObject());
+	return AISystemDefaultObject != NULL ? AISystemDefaultObject->AISystemClassName : FStringClassReference();
+}
\ No newline at end of file
