Particles spawned with module Skel Vert/Surf Location can now inherit their color from the corresponding vert(s).

#ttp 333609 - UE4: CASCADE: Allow Target Mesh Vertex Colors to be accessed by the Particle System

diff --git a/Engine/Source/Runtime/Engine/Classes/Components/SkinnedMeshComponent.h b/Engine/Source/Runtime/Engine/Classes/Components/SkinnedMeshComponent.h
--- a/Engine/Source/Runtime/Engine/Classes/Components/SkinnedMeshComponent.h
+++ b/Engine/Source/Runtime/Engine/Classes/Components/SkinnedMeshComponent.h
@@ -465,6 +465,13 @@ public:
 	 * @param VertexIndex Vertex Index. If compressed, this will be slow. 
 	 */
 	FVector GetSkinnedVertexPosition(int32 VertexIndex) const;
+	
+	/**
+	* Returns color of the vertex.
+	*
+	* @param VertexIndex Vertex Index. If compressed, this will be slow.
+	*/
+	FColor GetVertexColor(int32 VertexIndex) const;
 
 	/**
 	 * Update functions
diff --git a/Engine/Source/Runtime/Engine/Classes/Particles/Location/ParticleModuleLocationSkelVertSurface.h b/Engine/Source/Runtime/Engine/Classes/Particles/Location/ParticleModuleLocationSkelVertSurface.h
--- a/Engine/Source/Runtime/Engine/Classes/Particles/Location/ParticleModuleLocationSkelVertSurface.h
+++ b/Engine/Source/Runtime/Engine/Classes/Particles/Location/ParticleModuleLocationSkelVertSurface.h
@@ -83,6 +83,9 @@ class UParticleModuleLocationSkelVertSurface : public UParticleModuleLocationBas
 	UPROPERTY(EditAnywhere, Category=VertSurface)
 	TArray<int32> ValidMaterialIndices;
 
+	/** If true, particles inherit the associated vertex color on spawn. This feature is not supported for GPU particles. */
+	UPROPERTY(EditAnywhere, Category = VertSurface)
+	uint32 bInheritVertexColor : 1;
 
 	// Begin UObject Interface
 	virtual void PostLoad() OVERRIDE;
@@ -108,6 +111,7 @@ class UParticleModuleLocationSkelVertSurface : public UParticleModuleLocationBas
 	virtual int32 GetNumberOfCustomMenuOptions() const OVERRIDE;
 	virtual bool GetCustomMenuEntryDisplayString(int32 InEntryIndex, FString& OutDisplayString) const OVERRIDE;
 	virtual bool PerformCustomMenuEntry(int32 InEntryIndex) OVERRIDE;
+	virtual bool IsValidForLODLevel(UParticleLODLevel* LODLevel, FString& OutErrorString) OVERRIDE;
 #endif
 	//End UParticleModule Interface
 
diff --git a/Engine/Source/Runtime/Engine/Private/Particles/ParticleModules_Location.cpp b/Engine/Source/Runtime/Engine/Private/Particles/ParticleModules_Location.cpp
--- a/Engine/Source/Runtime/Engine/Private/Particles/ParticleModules_Location.cpp
+++ b/Engine/Source/Runtime/Engine/Private/Particles/ParticleModules_Location.cpp
@@ -2093,6 +2093,8 @@ void UParticleModuleLocationSkelVertSurface::Spawn(FParticleEmitterInstance* Own
 		return;
 	}
 
+	FStaticLODModel& LODModel = SkelMeshResource->LODModels[0];
+
 	// Determine the bone/socket to spawn at
 	int32 SourceIndex = -1;
 	int32 ActiveBoneIndex = -1;
@@ -2118,7 +2120,6 @@ void UParticleModuleLocationSkelVertSurface::Spawn(FParticleEmitterInstance* Own
 	}
 	else if(SourceType == VERTSURFACESOURCE_Surface)
 	{
-		FStaticLODModel& LODModel = SkelMeshResource->LODModels[0];
 		int32 SectionCount = LODModel.Sections.Num();
 		int32 RandomSection = FMath::RoundToInt(FMath::SRand() * ((float)SectionCount-1));
 
@@ -2182,6 +2183,33 @@ void UParticleModuleLocationSkelVertSurface::Spawn(FParticleEmitterInstance* Own
 					Particle.BaseVelocity = InstancePayload->BoneVelocities[VelocityIndex];
 				}
 			}
+			
+			if (bInheritVertexColor)
+			{
+				FColor UseColor;
+				if (SourceType == VERTSURFACESOURCE_Vert)
+				{
+					UseColor = SourceComponent->GetVertexColor(SourceIndex);
+				}
+				else if (SourceType == VERTSURFACESOURCE_Surface)
+				{
+					int32 VertIndex[3];
+					FColor VertColors[3];
+
+					VertIndex[0] = LODModel.MultiSizeIndexContainer.GetIndexBuffer()->Get(SourceIndex);
+					VertIndex[1] = LODModel.MultiSizeIndexContainer.GetIndexBuffer()->Get(SourceIndex + 1);
+					VertIndex[2] = LODModel.MultiSizeIndexContainer.GetIndexBuffer()->Get(SourceIndex + 2);
+					VertColors[0] = SourceComponent->GetVertexColor(VertIndex[0]);
+					VertColors[1] = SourceComponent->GetVertexColor(VertIndex[1]);
+					VertColors[2] = SourceComponent->GetVertexColor(VertIndex[2]);
+					UseColor.R = (VertColors[0].R + VertColors[1].R + VertColors[2].R) / 3;
+					UseColor.G = (VertColors[0].G + VertColors[1].G + VertColors[2].G) / 3;
+					UseColor.B = (VertColors[0].B + VertColors[1].B + VertColors[2].B) / 3;
+					UseColor.A = (VertColors[0].A + VertColors[1].A + VertColors[2].A) / 3;
+				}
+				Particle.Color = UseColor;
+				Particle.BaseColor = UseColor;
+			}
 
 			if (bMeshRotationActive)
 			{
@@ -2470,6 +2498,18 @@ bool UParticleModuleLocationSkelVertSurface::PerformCustomMenuEntry(int32 InEntr
 	}
 	return false;
 }
+
+bool UParticleModuleLocationSkelVertSurface::IsValidForLODLevel(UParticleLODLevel* LODLevel, FString& OutErrorString)
+{
+	if ( bInheritVertexColor && LODLevel && LODLevel->TypeDataModule && LODLevel->TypeDataModule->IsA<UParticleModuleTypeDataGpu>() )
+	{
+		OutErrorString = NSLOCTEXT("UnrealEd", "Module_LocationSkelVertSurface_InheritVertexColorOnGPUError", "Inherit Vertex Color is not supported on GPU emitters.").ToString();
+		return false;
+	}
+
+	return true;
+}
+
 #endif
 
 USkeletalMeshComponent* UParticleModuleLocationSkelVertSurface::GetSkeletalMeshComponentSource(FParticleEmitterInstance* Owner)
diff --git a/Engine/Source/Runtime/Engine/Private/SkinnedMeshComponent.cpp b/Engine/Source/Runtime/Engine/Private/SkinnedMeshComponent.cpp
--- a/Engine/Source/Runtime/Engine/Private/SkinnedMeshComponent.cpp
+++ b/Engine/Source/Runtime/Engine/Private/SkinnedMeshComponent.cpp
@@ -1425,6 +1425,31 @@ FVector USkinnedMeshComponent::GetSkinnedVertexPosition(int32 VertexIndex) const
 		: GetTypedSkinnedVertexPosition<false>(Chunk, Model.VertexBufferGPUSkin, VertIndex, bSoftVertex);
 }
 
+FColor USkinnedMeshComponent::GetVertexColor(int32 VertexIndex) const
+{
+	// Fail if no mesh
+	if (!SkeletalMesh || !MeshObject)
+	{
+		return FColor(255, 255, 255, 255);
+	}
+
+	FStaticLODModel& Model = MeshObject->GetSkeletalMeshResource().LODModels[0];
+
+	// Find the chunk and vertex within that chunk, and skinning type, for this vertex.
+	int32 ChunkIndex;
+	int32 VertIndex;
+	bool bSoftVertex;
+	bool bHasExtraBoneInfluences;
+	Model.GetChunkAndSkinType(VertexIndex, ChunkIndex, VertIndex, bSoftVertex, bHasExtraBoneInfluences);
+
+	check(ChunkIndex < Model.Chunks.Num());
+	const FSkelMeshChunk& Chunk = Model.Chunks[ChunkIndex];
+	
+	int32 VertexBase = bSoftVertex ? Chunk.GetSoftVertexBufferIndex() : Chunk.GetRigidVertexBufferIndex();
+
+	return Model.ColorVertexBuffer.VertexColor(VertexBase + VertIndex);
+}
+
 
 void USkinnedMeshComponent::HideBone( int32 BoneIndex, EPhysBodyOp PhysBodyOption)
 {
