TTP 335882 - LIVE: CASCADE: The Kill Box doesn't kill the particles inside the box when set to true, but kills them when set to false

Actually only happens if lower left corner values are bigger than upper right corner. Solved by component-wise Min/Max'ing the corners before testing.

diff --git a/Engine/Source/Runtime/Engine/Private/Particles/ParticleModules.cpp b/Engine/Source/Runtime/Engine/Private/Particles/ParticleModules.cpp
--- a/Engine/Source/Runtime/Engine/Private/Particles/ParticleModules.cpp
+++ b/Engine/Source/Runtime/Engine/Private/Particles/ParticleModules.cpp
@@ -3195,8 +3195,10 @@ void UParticleModuleKillBox::Update(FParticleEmitterInstance* Owner, int32 Offse
 {
 	UParticleLODLevel* LODLevel	= Owner->SpriteTemplate->GetCurrentLODLevel(Owner);
 
-	FVector CheckLL = LowerLeftCorner.GetValue(Owner->EmitterTime, Owner->Component);
-	FVector CheckUR = UpperRightCorner.GetValue(Owner->EmitterTime, Owner->Component);
+	FVector LL = LowerLeftCorner.GetValue(Owner->EmitterTime, Owner->Component);
+	FVector UR = UpperRightCorner.GetValue(Owner->EmitterTime, Owner->Component);
+	FVector CheckLL = LL.ComponentMin(UR);
+	FVector CheckUR = UR.ComponentMax(LL);
 	if (bAbsolute == false)
 	{
 		CheckLL += Owner->Component->GetComponentLocation();
