Add missing HUD header (fix CIS)

diff --git a/Engine/Source/Runtime/Engine/Private/Debug/GameplayDebuggingControllerComponent.cpp b/Engine/Source/Runtime/Engine/Private/Debug/GameplayDebuggingControllerComponent.cpp
--- a/Engine/Source/Runtime/Engine/Private/Debug/GameplayDebuggingControllerComponent.cpp
+++ b/Engine/Source/Runtime/Engine/Private/Debug/GameplayDebuggingControllerComponent.cpp
@@ -2,6 +2,7 @@
 
 #include "EnginePrivate.h"
 #include "Net/UnrealNetwork.h"
+#include "GameFramework/HUD.h"
 #include "BehaviorTreeDelegates.h"
 #include "VisualLog.h"
 
