Fix a non-unity compile error.

diff --git a/Engine/Source/Runtime/Engine/Classes/EdGraph/EdGraph.h b/Engine/Source/Runtime/Engine/Classes/EdGraph/EdGraph.h
--- a/Engine/Source/Runtime/Engine/Classes/EdGraph/EdGraph.h
+++ b/Engine/Source/Runtime/Engine/Classes/EdGraph/EdGraph.h
@@ -2,6 +2,7 @@
 
 #pragma once
 #include "BlueprintUtilities.h"
+#include "EdGraphSchema.h"
 #include "EdGraph.generated.h"
 
 USTRUCT()
@@ -57,7 +58,7 @@ class ENGINE_API UEdGraph : public UObject
 
 	/** The schema that this graph obeys */
 	UPROPERTY()
-	TSubclassOf<class UEdGraphSchema>  Schema;
+	TSubclassOf<UEdGraphSchema>  Schema;
 
 	/** Set of all nodes in this graph */
 	UPROPERTY()
