Audio: do not decompress in realtime if device does not support it.

Fixes TTP #336826 where DTYPE_Realtime soundwaves were silently failing in OpenAL under Linux (and perhaps HTML5, too). May also help other platforms (IOS?) where DTYPE_Realtime is not implemented.

diff --git a/Engine/Source/Runtime/Android/AndroidAudio/Public/AndroidAudioDevice.h b/Engine/Source/Runtime/Android/AndroidAudio/Public/AndroidAudioDevice.h
--- a/Engine/Source/Runtime/Android/AndroidAudio/Public/AndroidAudioDevice.h
+++ b/Engine/Source/Runtime/Android/AndroidAudio/Public/AndroidAudioDevice.h
@@ -243,6 +243,11 @@ public:
 
 	virtual bool HasCompressedAudioInfoClass(USoundWave* SoundWave) OVERRIDE;
 
+	virtual bool SupportsRealtimeDecompression() const OVERRIDE
+	{
+		return true;
+	}
+
 	virtual class ICompressedAudioInfo* CreateCompressedAudioInfo(USoundWave* SoundWave) OVERRIDE;
 
 	/** Starts up any platform specific hardware/APIs */
diff --git a/Engine/Source/Runtime/Engine/Private/AudioDevice.cpp b/Engine/Source/Runtime/Engine/Private/AudioDevice.cpp
--- a/Engine/Source/Runtime/Engine/Private/AudioDevice.cpp
+++ b/Engine/Source/Runtime/Engine/Private/AudioDevice.cpp
@@ -2279,7 +2279,8 @@ void FAudioDevice::Precache(USoundWave* SoundWave, bool bSynchronous, bool bTrac
 		const FSoundGroup& SoundGroup = GetDefault<USoundGroups>()->GetSoundGroup(SoundWave->SoundGroup);
 
 		// handle audio decompression
-		if (bDisableAudioCaching || (!SoundGroup.bAlwaysDecompressOnLoad && SoundWave->Duration > SoundGroup.DecompressedDuration))
+		if (SupportsRealtimeDecompression() && 
+			(bDisableAudioCaching || (!SoundGroup.bAlwaysDecompressOnLoad && SoundWave->Duration > SoundGroup.DecompressedDuration)))
 		{
 			// Store as compressed data and decompress in realtime
 			SoundWave->DecompressionType = DTYPE_RealTime;
diff --git a/Engine/Source/Runtime/Engine/Public/AudioDevice.h b/Engine/Source/Runtime/Engine/Public/AudioDevice.h
--- a/Engine/Source/Runtime/Engine/Public/AudioDevice.h
+++ b/Engine/Source/Runtime/Engine/Public/AudioDevice.h
@@ -476,6 +476,12 @@ public:
 	/** Whether this SoundWave has an associated info class to decompress it */
 	virtual bool HasCompressedAudioInfoClass(USoundWave* SoundWave) { return false; }
 
+	/** Whether this device supports realtime decompression of sound waves (i.e. DTYPE_RealTime) */
+	virtual bool SupportsRealtimeDecompression() const
+	{ 
+		return false;	// assume no support by default
+	}
+
 	/** Creates a Compressed audio info class suitable for decompressing this SoundWave */
 	virtual class ICompressedAudioInfo* CreateCompressedAudioInfo(USoundWave* SoundWave) { return NULL; }
 
diff --git a/Engine/Source/Runtime/Mac/CoreAudio/Public/CoreAudioDevice.h b/Engine/Source/Runtime/Mac/CoreAudio/Public/CoreAudioDevice.h
--- a/Engine/Source/Runtime/Mac/CoreAudio/Public/CoreAudioDevice.h
+++ b/Engine/Source/Runtime/Mac/CoreAudio/Public/CoreAudioDevice.h
@@ -362,6 +362,11 @@ class FCoreAudioDevice : public FAudioDevice
 
 	virtual bool HasCompressedAudioInfoClass(USoundWave* SoundWave) OVERRIDE;
 
+	virtual bool SupportsRealtimeDecompression() const OVERRIDE
+	{ 
+		return true; 
+	}
+
 	virtual class ICompressedAudioInfo* CreateCompressedAudioInfo(USoundWave* SoundWave) OVERRIDE;
 
 	FORCEINLINE uint32 CalculateMatrixElement( uint32 InputNum, uint32 OutputNum )
diff --git a/Engine/Source/Runtime/Windows/XAudio2/Public/XAudio2Device.h b/Engine/Source/Runtime/Windows/XAudio2/Public/XAudio2Device.h
--- a/Engine/Source/Runtime/Windows/XAudio2/Public/XAudio2Device.h
+++ b/Engine/Source/Runtime/Windows/XAudio2/Public/XAudio2Device.h
@@ -104,6 +104,11 @@ class FXAudio2Device : public FAudioDevice
 
 	virtual bool HasCompressedAudioInfoClass(USoundWave* SoundWave) OVERRIDE;
 
+	virtual bool SupportsRealtimeDecompression() const OVERRIDE
+	{
+		return true;
+	}
+
 	virtual class ICompressedAudioInfo* CreateCompressedAudioInfo(USoundWave* SoundWave) OVERRIDE;
 
 	/** 
