* fixed game console closing on the enter key (to how it was before)

diff --git a/Engine/Source/Runtime/Engine/Private/UserInterface/Console.cpp b/Engine/Source/Runtime/Engine/Private/UserInterface/Console.cpp
--- a/Engine/Source/Runtime/Engine/Private/UserInterface/Console.cpp
+++ b/Engine/Source/Runtime/Engine/Private/UserInterface/Console.cpp
@@ -633,7 +633,12 @@ bool UConsole::InputKey_InputLine( int32 ControllerId, FKey Key, EInputEvent Eve
 				//OutputText( Localize("Errors","Exec","Core") );
 
 				OutputText( TEXT("") );
-				FakeGotoState(NAME_None);
+
+				if(ConsoleState == NAME_Typing)
+				{
+					// close after each command when in typing mode (single line)
+					FakeGotoState(NAME_None);
+				}
 
 				UpdateCompleteIndices();
 			}
