Fix parsing of shorthand IPv6 URLs.
Since they can start with the characters "[:" the FURL constructor was seeing the ':' and assuming a file path.

diff --git a/Engine/Source/Runtime/Engine/Private/URL.cpp b/Engine/Source/Runtime/Engine/Private/URL.cpp
--- a/Engine/Source/Runtime/Engine/Private/URL.cpp
+++ b/Engine/Source/Runtime/Engine/Private/URL.cpp
@@ -221,7 +221,7 @@ FURL::FURL( FURL* Base, const TCHAR* TextURL, ETravelType Type )
 		// Handle pure filenames & Posix paths.
 		bool FarHost=0;
 		bool FarMap=0;
-		if( FCString::Strlen(URL)>2 && (URL[1]==':' || (URL[0]=='/' && !FPackageName::IsValidLongPackageName(URL, true))) )
+		if( FCString::Strlen(URL)>2 && ((URL[0] != '[' && URL[1]==':') || (URL[0]=='/' && !FPackageName::IsValidLongPackageName(URL, true))) )
 		{
 			// Pure filename.
 			Protocol = UrlConfig.DefaultProtocol;
