Fix UseComplexAsSimpleCollision destroying simple convex geom

diff --git a/Engine/Source/Runtime/Engine/Private/PhysicsEngine/BodySetup.cpp b/Engine/Source/Runtime/Engine/Private/PhysicsEngine/BodySetup.cpp
--- a/Engine/Source/Runtime/Engine/Private/PhysicsEngine/BodySetup.cpp
+++ b/Engine/Source/Runtime/Engine/Private/PhysicsEngine/BodySetup.cpp
@@ -104,26 +104,29 @@ void UBodySetup::CreatePhysicsMeshes()
 		check( !bGenerateNonMirroredCollision || CookedDataReader.ConvexMeshes.Num() == 0 || CookedDataReader.ConvexMeshes.Num() == AggGeom.ConvexElems.Num() );
 		check( !bGenerateMirroredCollision || CookedDataReader.ConvexMeshesNegX.Num() == 0 || CookedDataReader.ConvexMeshesNegX.Num() == AggGeom.ConvexElems.Num() );
 
-		//If the cooked data no longer has convex meshes, make sure to empty AggGeom.ConvexElems - otherwise we leave NULLS which cause issues, and we also read past the end of CookedDataReader.ConvexMeshes
-		if( (bGenerateNonMirroredCollision && CookedDataReader.ConvexMeshes.Num() == 0) || (bGenerateMirroredCollision && CookedDataReader.ConvexMeshesNegX.Num() == 0))
+		if (CollisionTraceFlag != CTF_UseComplexAsSimple)
 		{
-			AggGeom.ConvexElems.Empty();
-		}
-
-		for( int32 ElementIndex = 0; ElementIndex < AggGeom.ConvexElems.Num(); ElementIndex++ )
-		{
-			FKConvexElem& ConvexElem = AggGeom.ConvexElems[ElementIndex];
-
-			if(bGenerateNonMirroredCollision)
+			//If the cooked data no longer has convex meshes, make sure to empty AggGeom.ConvexElems - otherwise we leave NULLS which cause issues, and we also read past the end of CookedDataReader.ConvexMeshes
+			if ((bGenerateNonMirroredCollision && CookedDataReader.ConvexMeshes.Num() == 0) || (bGenerateMirroredCollision && CookedDataReader.ConvexMeshesNegX.Num() == 0))
 			{
-				ConvexElem.ConvexMesh = CookedDataReader.ConvexMeshes[ ElementIndex ];
-				FPhysxSharedData::Get().Add(ConvexElem.ConvexMesh);
+				AggGeom.ConvexElems.Empty();
 			}
 
-			if(bGenerateMirroredCollision)
+			for (int32 ElementIndex = 0; ElementIndex < AggGeom.ConvexElems.Num(); ElementIndex++)
 			{
-				ConvexElem.ConvexMeshNegX = CookedDataReader.ConvexMeshesNegX[ ElementIndex ];
-				FPhysxSharedData::Get().Add(ConvexElem.ConvexMeshNegX);
+				FKConvexElem& ConvexElem = AggGeom.ConvexElems[ElementIndex];
+
+				if (bGenerateNonMirroredCollision)
+				{
+					ConvexElem.ConvexMesh = CookedDataReader.ConvexMeshes[ElementIndex];
+					FPhysxSharedData::Get().Add(ConvexElem.ConvexMesh);
+				}
+
+				if (bGenerateMirroredCollision)
+				{
+					ConvexElem.ConvexMeshNegX = CookedDataReader.ConvexMeshesNegX[ElementIndex];
+					FPhysxSharedData::Get().Add(ConvexElem.ConvexMeshNegX);
+				}
 			}
 		}
 
