Naming cleanup to prepare for larger changes

* Cleanup naming
* Allow NetGUID to change for the caller when we auto-assign for default guids

diff --git a/Engine/Source/Runtime/CoreUObject/Public/UObject/CoreNet.h b/Engine/Source/Runtime/CoreUObject/Public/UObject/CoreNet.h
--- a/Engine/Source/Runtime/CoreUObject/Public/UObject/CoreNet.h
+++ b/Engine/Source/Runtime/CoreUObject/Public/UObject/CoreNet.h
@@ -167,7 +167,7 @@ class COREUOBJECT_API UPackageMap : public UObject
 	FClassNetCache* GetClassNetCache(UClass* Class);
 	void	ClearClassNetCache();
 
-	virtual UObject * NetGUIDAssign(FNetworkGUID NetGUID, FString Path, UObject *ObjOuter = NULL) { return NULL; }
+	virtual UObject * ResolvePathAndAssignNetGUID( FNetworkGUID & InOutNetGUID, FString Path, UObject * ObjOuter ) { return NULL; }
 	
 	virtual void HandleUnAssignedObject(const UObject* Obj) { }
 
diff --git a/Engine/Source/Runtime/Engine/Classes/Engine/PackageMapClient.h b/Engine/Source/Runtime/Engine/Classes/Engine/PackageMapClient.h
--- a/Engine/Source/Runtime/Engine/Classes/Engine/PackageMapClient.h
+++ b/Engine/Source/Runtime/Engine/Classes/Engine/PackageMapClient.h
@@ -54,7 +54,7 @@ class UPackageMapClient : public UPackageMap
 	virtual void LogDebugInfo(FOutputDevice& Ar) OVERRIDE;
 	virtual bool SupportsObject( const UObject* Object ) OVERRIDE;
 	
-	virtual UObject * NetGUIDAssign(FNetworkGUID NetGUID, FString Path, UObject *ObjOuter = NULL) OVERRIDE;
+	virtual UObject * ResolvePathAndAssignNetGUID( FNetworkGUID & InOutNetGUID, FString Path, UObject * ObjOuter ) OVERRIDE;
 
 	virtual bool WriteObject( FArchive& Ar, UObject* Outer, FNetworkGUID NetGUID, FString ObjName ) OVERRIDE;
 
diff --git a/Engine/Source/Runtime/Engine/Private/PackageMapClient.cpp b/Engine/Source/Runtime/Engine/Private/PackageMapClient.cpp
--- a/Engine/Source/Runtime/Engine/Private/PackageMapClient.cpp
+++ b/Engine/Source/Runtime/Engine/Private/PackageMapClient.cpp
@@ -568,11 +568,11 @@ bool UPackageMapClient::InternalIsMapped(UObject* Object, FNetworkGUID &NetGUID)
  *	Finds objects from PathName and ObjOuter, and assigns it a NetGUID. Called on the client after serialziing a
  *	<NetGUID, Path> pair.
  */
-UObject * UPackageMapClient::NetGUIDAssign(FNetworkGUID NetGUID, FString PathName, UObject *ObjOuter)
+UObject * UPackageMapClient::ResolvePathAndAssignNetGUID( FNetworkGUID & InOutNetGUID, FString PathName, UObject *ObjOuter)
 {
-	if (NetGUID.IsDynamic())
+	if ( InOutNetGUID.IsDynamic() )
 	{
-		UE_LOG(LogNetPackageMap, Log, TEXT("UPackageMapClient::NetGUIDAssign. Trying to load dynamic NetGUID by path: <%s,%s>. Outer: %s"), *NetGUID.ToString(), *PathName, ObjOuter ? *ObjOuter->GetName() : TEXT("NULL"));
+		UE_LOG( LogNetPackageMap, Warning, TEXT( "UPackageMapClient::ResolvePathAndAssignNetGUID. Trying to load dynamic NetGUID by path: <%s,%s>. Outer: %s" ), *InOutNetGUID.ToString(), *PathName, ObjOuter ? *ObjOuter->GetName() : TEXT( "NULL" ) );
 	}
 
 	// We don't have the object mapped yet
@@ -585,12 +585,12 @@ UObject * UPackageMapClient::NetGUIDAssign(FNetworkGUID NetGUID, FString PathNam
 			// This could be garbage or a security risk
 			// Another possibility is in dynamic property replication if the server reads the previously serialized state
 			// that has a now destroyed actor in it.
-			UE_LOG( LogNetPackageMap, Warning, TEXT("Could not find Object for: NetGUID <%s, %s> (and IsNetGUIDAuthority())"), *NetGUID.ToString(), *PathName );
+			UE_LOG( LogNetPackageMap, Warning, TEXT("Could not find Object for: NetGUID <%s, %s> (and IsNetGUIDAuthority())"), *InOutNetGUID.ToString(), *PathName );
 			return NULL;
 		}
 		else
 		{
-			UE_LOG(LogNetPackageMap, Log, TEXT("Could not find Object for: NetGUID <%s, %s>"), *NetGUID.ToString(), *PathName );
+			UE_LOG(LogNetPackageMap, Log, TEXT("Could not find Object for: NetGUID <%s, %s>"), *InOutNetGUID.ToString(), *PathName );
 		}
 		
 
@@ -605,7 +605,7 @@ UObject * UPackageMapClient::NetGUIDAssign(FNetworkGUID NetGUID, FString PathNam
 			if (!ObjOuter || Package)
 			{
 				Object = LoadPackage(Package, *PathName, LOAD_None);
-				UE_LOG(LogNetPackageMap, Log, TEXT("UPackageMapClient::NetGUIDAssign  LoadPackage %s. Found: %s"), *PathName, Object ? *Object->GetName() : TEXT("NULL") );
+				UE_LOG(LogNetPackageMap, Log, TEXT("UPackageMapClient::ResolvePathAndAssignNetGUID  LoadPackage %s. Found: %s"), *PathName, Object ? *Object->GetName() : TEXT("NULL") );
 			}
 			else
 			{
@@ -614,9 +614,9 @@ UObject * UPackageMapClient::NetGUIDAssign(FNetworkGUID NetGUID, FString PathNam
 		}
 	}
 
-	if (Object)
+	if ( Object )
 	{
-		AssignNetGUID(Object, NetGUID);
+		InOutNetGUID = AssignNetGUID( Object, InOutNetGUID );
 	}
 
 	return Object;
@@ -837,7 +837,7 @@ bool UPackageMapClient::InternalLoadObjectPath( FArchive& Ar, UObject*& Object,
 	// ----------------
 	if ( !Object && !OuterFailure )
 	{
-		Object = NetGUIDAssign( NetGUID, PathName, ObjOuter );
+		Object = ResolvePathAndAssignNetGUID( NetGUID, PathName, ObjOuter );
 	}
 
 	// ----------------	
diff --git a/Engine/Source/Runtime/Engine/Private/PendingNetGame.cpp b/Engine/Source/Runtime/Engine/Private/PendingNetGame.cpp
--- a/Engine/Source/Runtime/Engine/Private/PendingNetGame.cpp
+++ b/Engine/Source/Runtime/Engine/Private/PendingNetGame.cpp
@@ -325,7 +325,7 @@ void UPendingNetGame::NotifyControlMessage(UNetConnection* Connection, uint8 Mes
 			FNetworkGUID NetGUID;
 			FString Path;
 			FNetControlMessage<NMT_NetGUIDAssign>::Receive(Bunch, NetGUID, Path);
-			NetDriver->ServerConnection->PackageMap->NetGUIDAssign(NetGUID, Path, NULL);
+			NetDriver->ServerConnection->PackageMap->ResolvePathAndAssignNetGUID(NetGUID, Path, NULL);
 			break;
 		}
 		default:
diff --git a/Engine/Source/Runtime/Engine/Private/World.cpp b/Engine/Source/Runtime/Engine/Private/World.cpp
--- a/Engine/Source/Runtime/Engine/Private/World.cpp
+++ b/Engine/Source/Runtime/Engine/Private/World.cpp
@@ -3366,7 +3366,7 @@ void UWorld::NotifyControlMessage(UNetConnection* Connection, uint8 MessageType,
 				FNetControlMessage<NMT_NetGUIDAssign>::Receive(Bunch, NetGUID, Path);
 
 				UE_LOG(LogNet, Verbose, TEXT("NMT_NetGUIDAssign  NetGUID %s. Path: %s. "), *NetGUID.ToString(), *Path );
-				Connection->PackageMap->NetGUIDAssign(NetGUID, Path, NULL);
+				Connection->PackageMap->ResolvePathAndAssignNetGUID(NetGUID, Path, NULL);
 				break;
 			}
 		}
