Support for deferring the resolve of network GUIDs, first pass

diff --git a/Engine/Source/Runtime/Engine/Classes/Engine/PackageMapClient.h b/Engine/Source/Runtime/Engine/Classes/Engine/PackageMapClient.h
--- a/Engine/Source/Runtime/Engine/Classes/Engine/PackageMapClient.h
+++ b/Engine/Source/Runtime/Engine/Classes/Engine/PackageMapClient.h
@@ -22,6 +22,15 @@
 #include "Net/DataBunch.h"
 #include "PackageMapClient.generated.h"
 
+class FDeferredResolvePath
+{
+public:
+	FDeferredResolvePath() {}
+
+	FName			PathName;
+	FNetworkGUID	OuterGUID;
+};
+
 UCLASS(transient)
 class UPackageMapClient : public UPackageMap
 {
@@ -95,8 +104,11 @@ protected:
 	void	InternalWriteObject( FArchive& Ar, FNetworkGUID NetGUID, const UObject * Object, FString ObjectPathName, UObject* ObjectOuter );	
 	void	InternalWriteObjectPath( FArchive& Ar, FNetworkGUID NetGUID, const UObject * Object, FString ObjectPathName, UObject* ObjectOuter );
 
-	FNetworkGUID	InternalLoadObject( FArchive& Ar, UObject*& Object, int InternalLoadObjectRecursionCount);
-	bool			InternalLoadObjectPath( FArchive& Ar, UObject*& Object, FNetworkGUID NetGUID, int InternalLoadObjectRecursionCount);
+	FNetworkGUID	InternalLoadObject( FArchive & Ar, UObject *& Object, const bool bIsOuterLevel, const int InternalLoadObjectRecursionCount );
+
+	void		ResolvePathAndAssignNetGUID_Deferred( const FNetworkGUID & NetGUID, const FString & PathName, const FNetworkGUID & OuterGUID, const bool bIsPackage );
+	UObject *	ResolvePathAndAssignNetGUID( FNetworkGUID & InOutNetGUID, const FString & PathName, const FNetworkGUID & OuterGUID, const bool bIsPackage, const bool bNoLoad );
+	UObject *	ResolvePathAndAssignNetGUID( FNetworkGUID & InOutNetGUID, const FString & PathName, UObject * ObjOuter, const bool bIsPackage, const bool bNoLoad );
 
 	bool	ShouldSendFullPath(const UObject* Object, const FNetworkGUID &NetGUID);
 
@@ -127,4 +139,6 @@ protected:
 	FOutBunch *						CurrentExportBunch;
 
 	int32							ExportNetGUIDCount;
+
+	TMap< FNetworkGUID, FDeferredResolvePath >	DeferredResolvePaths;
 };
diff --git a/Engine/Source/Runtime/Engine/Private/PackageMapClient.cpp b/Engine/Source/Runtime/Engine/Private/PackageMapClient.cpp
--- a/Engine/Source/Runtime/Engine/Private/PackageMapClient.cpp
+++ b/Engine/Source/Runtime/Engine/Private/PackageMapClient.cpp
@@ -94,7 +94,7 @@ bool UPackageMapClient::SerializeObject( FArchive& Ar, UClass* Class, UObject*&
 		// ----------------	
 		// Read NetGUID from stream and resolve object
 		// ----------------	
-		FNetworkGUID NetGUID = InternalLoadObject(Ar, Object, 0);
+		FNetworkGUID NetGUID = InternalLoadObject( Ar, Object, false, 0 );
 
 		// Write out NetGUID to caller if necessary
 		if (OutNetGUID)
@@ -570,56 +570,7 @@ bool UPackageMapClient::InternalIsMapped(UObject* Object, FNetworkGUID &NetGUID)
  */
 UObject * UPackageMapClient::ResolvePathAndAssignNetGUID( FNetworkGUID & InOutNetGUID, FString PathName, UObject *ObjOuter)
 {
-	if ( InOutNetGUID.IsDynamic() )
-	{
-		UE_LOG( LogNetPackageMap, Warning, TEXT( "UPackageMapClient::ResolvePathAndAssignNetGUID. Trying to load dynamic NetGUID by path: <%s,%s>. Outer: %s" ), *InOutNetGUID.ToString(), *PathName, ObjOuter ? *ObjOuter->GetName() : TEXT( "NULL" ) );
-	}
-
-	// We don't have the object mapped yet
-	UObject * Object = StaticFindObject( UObject::StaticClass(), ObjOuter, *PathName, false );
-	if (!Object)
-	{
-		if (IsNetGUIDAuthority())
-		{
-			// The authority shouldn't be loading resources on demand, at least for now.
-			// This could be garbage or a security risk
-			// Another possibility is in dynamic property replication if the server reads the previously serialized state
-			// that has a now destroyed actor in it.
-			UE_LOG( LogNetPackageMap, Warning, TEXT("Could not find Object for: NetGUID <%s, %s> (and IsNetGUIDAuthority())"), *InOutNetGUID.ToString(), *PathName );
-			return NULL;
-		}
-		else
-		{
-			UE_LOG(LogNetPackageMap, Log, TEXT("Could not find Object for: NetGUID <%s, %s>"), *InOutNetGUID.ToString(), *PathName );
-		}
-		
-
-		Object = StaticLoadObject( UObject::StaticClass(), ObjOuter, *PathName, NULL, LOAD_NoWarn );
-		if (Object)
-		{
-			UE_LOG(LogNetPackageMap, Log, TEXT("  StaticLoadObject. Found: %s"), Object ? *Object->GetName() : TEXT("NULL") );
-		}
-		else
-		{
-			UPackage* Package = Cast<UPackage>(ObjOuter);
-			if (!ObjOuter || Package)
-			{
-				Object = LoadPackage(Package, *PathName, LOAD_None);
-				UE_LOG(LogNetPackageMap, Log, TEXT("UPackageMapClient::ResolvePathAndAssignNetGUID  LoadPackage %s. Found: %s"), *PathName, Object ? *Object->GetName() : TEXT("NULL") );
-			}
-			else
-			{
-				UE_LOG(LogNetPackageMap, Warning, TEXT("Outer %s is not a package for object %s"), ObjOuter ? *ObjOuter->GetName() : TEXT("NULL"), *PathName );
-			}
-		}
-	}
-
-	if ( Object )
-	{
-		InOutNetGUID = AssignNetGUID( Object, InOutNetGUID );
-	}
-
-	return Object;
+	return ResolvePathAndAssignNetGUID( InOutNetGUID, PathName, ObjOuter, false, false );
 }
 
 //--------------------------------------------------------------------
@@ -628,40 +579,64 @@ UObject * UPackageMapClient::ResolvePathAndAssignNetGUID( FNetworkGUID & InOutNe
 //
 //--------------------------------------------------------------------
 
+struct FExportFlags
+{
+	union
+	{
+		struct
+		{
+			uint8 bHasPath		: 1;
+			uint8 bIsLevel		: 1;
+			uint8 bIsPackage	: 1;
+		};
+
+		uint8	Value;
+	};
+
+	FExportFlags()
+	{
+		Value = 0;
+	}
+};
+
 /** Writes an object NetGUID given the NetGUID and either the object itself, or FString full name of the object. Appends full name/path if necessary */
 void UPackageMapClient::InternalWriteObject( FArchive & Ar, FNetworkGUID NetGUID, const UObject * Object, FString ObjectPathName, UObject * ObjectOuter )
 {
 	Ar << NetGUID;
-	NET_CHECKSUM(Ar);
+	NET_CHECKSUM( Ar );
 
-	if (!NetGUID.IsValid())
+	if ( !NetGUID.IsValid() )
 	{
+		// We're done writing
 		return;
 	}
 
-	// Write byte to say if we are including fullpath or not
-	//   note: Default NetGUID is implied to always include it
-	uint8 SendingFullPath = 0;
-	
+	// Write export flags
+	//   note: Default NetGUID is implied to always send path
+	FExportFlags ExportFlags;
+
 	if ( NetGUID.IsDefault() )
 	{
-		SendingFullPath = 1;
+		ExportFlags.bHasPath = 1;
 	}		
 	else if ( IsExportingNetGUIDBunch )
 	{
 		if ( Object != NULL )
 		{
-			SendingFullPath = ShouldSendFullPath( Object, NetGUID ) ? 1 : 0;
+			ExportFlags.bHasPath = ShouldSendFullPath( Object, NetGUID ) ? 1 : 0;
 		}
 		else
 		{
-			SendingFullPath = ObjectPathName.IsEmpty() ? 0 : 1;
+			ExportFlags.bHasPath = ObjectPathName.IsEmpty() ? 0 : 1;
 		}
 
-		Ar << SendingFullPath;
+		ExportFlags.bIsLevel	= Cast< const ULevel >( Object ) != NULL ? 1 : 0;
+		ExportFlags.bIsPackage	= Cast< const UPackage >( Object ) != NULL ? 1 : 0;
+
+		Ar << ExportFlags.Value;
 	}
 
-	if ( SendingFullPath )
+	if ( ExportFlags.bHasPath )
 	{
 		InternalWriteObjectPath( Ar, NetGUID, Object, ObjectPathName, ObjectOuter );
 	}
@@ -679,7 +654,7 @@ void UPackageMapClient::InternalWriteObjectPath( FArchive & Ar, FNetworkGUID Net
 		// If the object isn't NULL, expect an empty path name, then fill it out with the actual info
 		check( ObjectOuter == NULL );
 		check( ObjectPathName.IsEmpty() );
-		ObjectPathName = SerializeOuter ? Object->GetName() : Object->GetPathName();	// Full path if we didn't serialize the outer, else just the obj name
+		ObjectPathName = Object->GetName();
 		ObjectOuter = Object->GetOuter();
 	}
 	else
@@ -689,17 +664,14 @@ void UPackageMapClient::InternalWriteObjectPath( FArchive & Ar, FNetworkGUID Net
 		check( !ObjectPathName.IsEmpty() );
 	}
 
-	if ( SerializeOuter )
-	{
-		// Serialize reference to outer. This is basically a form of compression.
-		FNetworkGUID OuterNetGUID = InternalGetOrAssignNetGUID( ObjectOuter );
-		InternalWriteObject( Ar, OuterNetGUID, ObjectOuter, TEXT( "" ), NULL );
-	}
+	// Serialize reference to outer. This is basically a form of compression.
+	FNetworkGUID OuterNetGUID = InternalGetOrAssignNetGUID( ObjectOuter );
 
-	// Serialize Name of object
+	InternalWriteObject( Ar, OuterNetGUID, ObjectOuter, TEXT( "" ), NULL );
 
 	GEngine->NetworkRemapPath(Connection->Driver->GetWorld(), ObjectPathName, false);
 
+	// Serialize Name of object
 	Ar << ObjectPathName;
 
 	if ( IsExportingNetGUIDBunch )
@@ -709,22 +681,37 @@ void UPackageMapClient::InternalWriteObjectPath( FArchive & Ar, FNetworkGUID Net
 		int32& Count = NetGUIDExportCountMap.FindOrAdd(NetGUID);
 		Count++;
 	}
-	
+
 	UE_CLOG(!bSuppressLogs, LogNetPackageMap, Log, TEXT("UPackageMapClient::InternalWriteObjectPath Serialized Object %s as <%s,%s>)"), *ObjectPathName, *NetGUID.ToString(), *ObjectPathName );
 }
 
-
 //--------------------------------------------------------------------
 //
 //	Loading
 //
 //--------------------------------------------------------------------
 
+static void SanityCheckExport( const UObject * Object, const FNetworkGUID & NetGUID, const FString & ExpectedPathName, const UObject * ExpectedOuter )
+{
+	if ( Object->GetName() != ExpectedPathName )
+	{
+		UE_LOG( LogNetPackageMap, Warning, TEXT( "SanityCheckExport: Name mismatch. NetGUID: %s, Object: %s, Expected: %s" ), *NetGUID.ToString(), *Object->GetPathName(), *ExpectedPathName );
+	}
+
+	if ( Object->GetOuter() != ExpectedOuter )
+	{
+		const FString CurrentOuterName	= Object->GetOuter() != NULL ? *Object->GetOuter()->GetName() : TEXT( "NULL" );
+		const FString ExpectedOuterName = ExpectedOuter != NULL ? *ExpectedOuter->GetName() : TEXT( "NULL" );
+		UE_LOG( LogNetPackageMap, Warning, TEXT( "SanityCheckExport: Outer mismatch. Object: %s, NetGUID: %s, Current: %s, Expected: %s" ), *Object->GetPathName(), *NetGUID.ToString(), *CurrentOuterName, *ExpectedOuterName );
+	}
+}
+
 /** Loads a UObject from an FArchive stream. Reads object path if there, and tries to load object if its not already loaded */
-FNetworkGUID UPackageMapClient::InternalLoadObject( FArchive& Ar, UObject*& Object, const int InternalLoadObjectRecursionCount )
+FNetworkGUID UPackageMapClient::InternalLoadObject( FArchive & Ar, UObject *& Object, const bool bIsOuterLevel, const int InternalLoadObjectRecursionCount )
 {
-	if(InternalLoadObjectRecursionCount > INTERNAL_LOAD_OBJECT_RECURSION_LIMIT) {
-		UE_LOG(LogNetPackageMap, Warning, TEXT("InternalLoadObject hit recursion limit."));
+	if ( InternalLoadObjectRecursionCount > INTERNAL_LOAD_OBJECT_RECURSION_LIMIT ) 
+	{
+		UE_LOG( LogNetPackageMap, Warning, TEXT( "InternalLoadObject: Hit recursion limit." ) );
 		Ar.SetError(); 
 		Object = NULL;
 		return FNetworkGUID(); 
@@ -735,15 +722,15 @@ FNetworkGUID UPackageMapClient::InternalLoadObject( FArchive& Ar, UObject*& Obje
 	// ----------------	
 	FNetworkGUID NetGUID;
 	Ar << NetGUID;
-	NET_CHECKSUM_OR_END(Ar);
+	NET_CHECKSUM_OR_END( Ar );
 
-	if (Ar.IsError())
+	if ( Ar.IsError() )
 	{
 		Object = NULL;
 		return NetGUID;
 	}
 
-	if (!NetGUID.IsValid())
+	if ( !NetGUID.IsValid() )
 	{
 		Object = NULL;
 		return NetGUID;
@@ -752,24 +739,25 @@ FNetworkGUID UPackageMapClient::InternalLoadObject( FArchive& Ar, UObject*& Obje
 	// ----------------	
 	// Try to resolve NetGUID
 	// ----------------	
-	if (NetGUID.IsValid() && !NetGUID.IsDefault())
+	if ( NetGUID.IsValid() && !NetGUID.IsDefault() )
 	{
 		Object = GetObjectFromNetGUID( NetGUID );
 
-		UE_CLOG(!bSuppressLogs, LogNetPackageMap, Log, TEXT("InternalLoadObject loaded %s from NetGUID <%s>"), Object ? *Object->GetFullName() : TEXT("NULL"), *NetGUID.ToString() );
+		UE_CLOG( !bSuppressLogs, LogNetPackageMap, Log, TEXT( "InternalLoadObject loaded %s from NetGUID <%s>" ), Object ? *Object->GetFullName() : TEXT( "NULL" ), *NetGUID.ToString() );
 	}
 
 	// ----------------	
 	// Read the full if its there
 	// ----------------	
-	uint8 ReceivingFullPath = 0;
+	FExportFlags ExportFlags;
+
 	if ( NetGUID.IsDefault() )
 	{
-		ReceivingFullPath = 1;
+		ExportFlags.bHasPath = 1;
 	}
 	else if ( IsExportingNetGUIDBunch )
 	{
-		Ar << ReceivingFullPath;
+		Ar << ExportFlags.Value;
 
 		if ( Ar.IsError() )
 		{
@@ -778,94 +766,156 @@ FNetworkGUID UPackageMapClient::InternalLoadObject( FArchive& Ar, UObject*& Obje
 		}
 	}
 
-	if ( ReceivingFullPath )
+	if ( ExportFlags.bHasPath )
 	{
-		InternalLoadObjectPath(Ar, Object, NetGUID, InternalLoadObjectRecursionCount);
+		UObject * ObjOuter = NULL;
+
+		FNetworkGUID OuterGUID = InternalLoadObject( Ar, ObjOuter, ExportFlags.bIsLevel || bIsOuterLevel, InternalLoadObjectRecursionCount + 1 );
+
+		FString PathName;
+		Ar << PathName;
 
 		if ( Ar.IsError() )
 		{
+			UE_LOG( LogNetPackageMap, Error, TEXT( "InternalLoadObject: Failed to load path name" ) );
 			Object = NULL;
 			return NetGUID;
 		}
+
+		GEngine->NetworkRemapPath( Connection->Driver->GetWorld(), PathName, true );
+
+		if ( Object == NULL )
+		{
+			// Try to resolve the name (we may have already loaded it), but DON'T load yet if we haven't already
+			Object = ResolvePathAndAssignNetGUID( NetGUID, PathName, OuterGUID, ExportFlags.bIsPackage, false );
+		}
+
+		if ( Object != NULL )
+		{
+			// If we already have the object, just do some sanity checking and return
+			SanityCheckExport( Object, NetGUID, PathName, ObjOuter );
+			return NetGUID;
+		}
+
+		// If this is a level or our outer is being deferred, then defer this guid
+		const bool bCanDefer = false;//!NetGUID.IsDefault() && ( ExportFlags.bIsLevel || bIsOuterLevel || DeferredResolvePaths.Contains( OuterGUID ) );
+
+		if ( bCanDefer )
+		{
+			UE_LOG( LogNetPackageMap, Log, TEXT( "InternalLoadObject: Deferring resolve. Path: %s, NetGUID: %s, Outer: %s" ), *PathName, *NetGUID.ToString(), ObjOuter != NULL ? *ObjOuter->GetFullName() : TEXT( "NULL" ) );
+			ResolvePathAndAssignNetGUID_Deferred( NetGUID, PathName, OuterGUID, ExportFlags.bIsPackage );
+		}
+		else
+		{
+			Object = ResolvePathAndAssignNetGUID( NetGUID, PathName, OuterGUID, ExportFlags.bIsPackage, false );
+		}
 	}
 
-	if ( !Object && NetGUID.IsStatic() )
+	if ( Object == NULL && NetGUID.IsStatic() && !DeferredResolvePaths.Contains( NetGUID ) )
 	{
-		UE_LOG(LogNetPackageMap, Warning, TEXT("InternalLoadObject unable to resolve reference from NetGUID <%s> (received full path: %d)"),  *NetGUID.ToString(), ReceivingFullPath );
+		UE_LOG( LogNetPackageMap, Warning, TEXT( "InternalLoadObject: Unable to resolve object. NetGUID: %s, HasPath: %s" ), *NetGUID.ToString(), ExportFlags.bHasPath ? TEXT( "TRUE" ) : TEXT( "FALSE" ) );
 	}
 
 	return NetGUID;
 }
 
-/** Actually does the dynamic loading */
-bool UPackageMapClient::InternalLoadObjectPath( FArchive& Ar, UObject*& Object, FNetworkGUID NetGUID, const int InternalLoadObjectRecursionCount )
+void UPackageMapClient::ResolvePathAndAssignNetGUID_Deferred( const FNetworkGUID & NetGUID, const FString & PathName, const FNetworkGUID & OuterGUID, const bool bIsPackage )
 {
-	UObject* ObjOuter = NULL;
+	check( !NetGUID.IsDefault() );		// We can't defer default guids, doesn't make sense
 
-	FNetworkGUID OuterGUID;
-
-	if ( SerializeOuter )
+	if ( !DeferredResolvePaths.Contains( NetGUID ) )
 	{
-		OuterGUID = InternalLoadObject( Ar, ObjOuter, InternalLoadObjectRecursionCount + 1 );
-	}
+		FDeferredResolvePath DeferredResolvePath;
 
-	const bool OuterFailure = ( SerializeOuter && OuterGUID.IsValid() && ObjOuter == NULL );
+		DeferredResolvePath.OuterGUID	= OuterGUID;
+		DeferredResolvePath.PathName	= FName( *PathName );
 
-	FString PathName;
-	Ar << PathName;
-	
-	if ( OuterFailure )
+		DeferredResolvePaths.Add( NetGUID, DeferredResolvePath );
+	}
+	else
 	{
-		// Early out now if we failed to resolve the outer. This may happen in normal circumstances and is not an error.
-		// (for example, an object was GC'd on the server right before the client sent a reference to the object)
-		UE_LOG( LogNetPackageMap, Log, TEXT( "InternalLoadObjectPath: Failed to load outer. OuterGUID: %s. Path: %s"), *OuterGUID.ToString(), *PathName );
-		Object = NULL;
-		return false;
+		// Sanity check
+		FDeferredResolvePath & DeferredResolvePath = DeferredResolvePaths.FindChecked( NetGUID );
+
+		if ( DeferredResolvePath.OuterGUID != OuterGUID )
+		{
+			UE_LOG( LogNetPackageMap, Warning, TEXT( "ResolvePathAndAssignNetGUID_Deferred: Deferred Outer mismatch. NetGUID: %s, Path: %s, Current: %s, Expected: %s" ), *NetGUID.ToString(), *PathName, *DeferredResolvePath.OuterGUID.ToString(), *OuterGUID.ToString() );
+		}
 	}
+}
 
-	if ( Ar.IsError() )
+UObject * UPackageMapClient::ResolvePathAndAssignNetGUID( FNetworkGUID & InOutNetGUID, const FString & PathName, UObject * ObjOuter, const bool bIsPackage, const bool bNoLoad )
+{
+	if ( InOutNetGUID.IsDynamic() )
 	{
-		UE_LOG( LogNetPackageMap, Error, TEXT( "InternalLoadObjectPath: Failed to load path name") );
-		Object = NULL;
-		return false;
+		UE_LOG( LogNetPackageMap, Warning, TEXT( "UPackageMapClient::ResolvePathAndAssignNetGUID. Trying to load dynamic NetGUID by path: <%s,%s>. Outer: %s" ), *InOutNetGUID.ToString(), *PathName, ObjOuter ? *ObjOuter->GetName() : TEXT( "NULL" ) );
 	}
 
-	GEngine->NetworkRemapPath(Connection->Driver->GetWorld(), PathName, true);
+	// We don't have the object mapped yet
+	UObject * Object = StaticFindObject( UObject::StaticClass(), ObjOuter, *PathName, false );
 
-	// ----------------
-	//	If we didn't find it and got the full path, resolve the full path and assign the NetGUID to the object we find
-	// ----------------
-	if ( !Object && !OuterFailure )
+	if ( !Object && !bNoLoad )
 	{
-		Object = ResolvePathAndAssignNetGUID( NetGUID, PathName, ObjOuter );
+		if ( IsNetGUIDAuthority() )
+		{
+			// The authority shouldn't be loading resources on demand, at least for now.
+			// This could be garbage or a security risk
+			// Another possibility is in dynamic property replication if the server reads the previously serialized state
+			// that has a now destroyed actor in it.
+			UE_LOG( LogNetPackageMap, Warning, TEXT( "Could not find Object for: NetGUID <%s, %s> (and IsNetGUIDAuthority())" ), *InOutNetGUID.ToString(), *PathName );
+			return NULL;
+		}
+		else
+		{
+			UE_LOG( LogNetPackageMap, Log, TEXT( "Could not find Object for: NetGUID <%s, %s>" ), *InOutNetGUID.ToString(), *PathName );
+		}
+
+		Object = StaticLoadObject( UObject::StaticClass(), ObjOuter, *PathName, NULL, LOAD_NoWarn );
+
+		if ( Object )
+		{
+			UE_LOG(LogNetPackageMap, Log, TEXT("  StaticLoadObject. Found: %s" ), Object ? *Object->GetName() : TEXT("NULL") );
+		}
+		else
+		{
+			if ( bIsPackage )
+			{
+				Object = LoadPackage( Cast< UPackage >( ObjOuter ), *PathName, LOAD_None );
+				UE_LOG( LogNetPackageMap, Log, TEXT( "UPackageMapClient::ResolvePathAndAssignNetGUID  LoadPackage %s. Found: %s" ), *PathName, Object ? *Object->GetName() : TEXT( "NULL" ) );
+			}
+			else
+			{
+				UE_LOG( LogNetPackageMap, Warning, TEXT( "Outer %s is not a package for object %s" ), ObjOuter ? *ObjOuter->GetName() : TEXT( "NULL" ), *PathName );
+			}
+		}
 	}
 
-	// ----------------	
-	// Logging
-	// ----------------	
-	
-	if ( NetGUID.IsStatic() && Object == NULL )
+	if ( Object )
 	{
-		// This is bad:: we were sent a static NetGUID without a path, but were not able to resolve it. This should not happen.
-		UE_LOG(LogNetPackageMap, Warning, TEXT("Unable to resolve static NetGUID <%s> from Path: %s"), *NetGUID.ToString(), *PathName );
+		InOutNetGUID = AssignNetGUID( Object, InOutNetGUID );
 	}
-	else
+
+	return Object;
+}
+
+UObject * UPackageMapClient::ResolvePathAndAssignNetGUID( FNetworkGUID & InOutNetGUID, const FString & PathName, const FNetworkGUID & OuterGUID, const bool bIsPackage, const bool bNoLoad )
+{
+	UObject * ObjOuter = OuterGUID.IsValid() ? GetObjectFromNetGUID( OuterGUID ) : NULL;
+
+	if ( OuterGUID.IsValid() && ObjOuter == NULL )
 	{
-		UE_CLOG(!bSuppressLogs, LogNetPackageMap, Log, TEXT("Read Serialized Object %s as <%s,%s>"), Object ? *Object->GetPathName() : TEXT("NULL"), *NetGUID.ToString(), *PathName );
+		UE_LOG( LogNetPackageMap, Warning, TEXT( "ResolvePathAndAssignNetGUID: Unable resolve outer. Path: %s, NetGUID: %s, OuterGUID: %s " ), *PathName, *InOutNetGUID.ToString(), *OuterGUID.ToString() );
+		return NULL;
 	}
 
-#if 0
-	// This isn't safe to do, since we could delete objects that the client has in flight
-	// We'll need to find another way to sandbox the clients
-	if ( NetGUID.IsValid() && Object == NULL && IsNetGUIDAuthority() )
+	UObject * Object = ResolvePathAndAssignNetGUID( InOutNetGUID, PathName, ObjOuter, bIsPackage, bNoLoad );
+
+	if ( Object == NULL )
 	{
-		UE_LOG( LogNetPackageMap, Error, TEXT( "Unable to resolve static NetGUID <%s> from Path: %s"), *NetGUID.ToString(), *PathName );
-		Ar.SetError();
-		return false;
+		UE_LOG( LogNetPackageMap, Warning, TEXT( "ResolvePathAndAssignNetGUID: Unable to resolve object. Path: %s, NetGUID: %s, Outer: %s" ), *PathName, *InOutNetGUID.ToString(), ObjOuter != NULL ? *ObjOuter->GetPathName() : TEXT( "NULL" ) );
 	}
-#endif
 
-	return true;
+	return Object;
 }
 
 //--------------------------------------------------------------------
@@ -1041,7 +1091,7 @@ void UPackageMapClient::ReceiveNetGUIDBunch( FInBunch &InBunch )
 	while( NumGUIDsRead < NumGUIDsInBunch )
 	{
 		UObject *Obj;
-		InternalLoadObject( InBunch, Obj, 0 );
+		InternalLoadObject( InBunch, Obj, false, 0 );
 		if ( InBunch.IsError() )
 		{
 			UE_LOG( LogNetPackageMap, Error, TEXT( "UPackageMapClient::ReceiveNetGUIDBunch: InBunch.IsError() after InternalLoadObject" ) );
