makes cloth physics property sliders drag smoothly in Presona

diff --git a/Engine/Source/Runtime/Engine/Private/SkeletalMesh.cpp b/Engine/Source/Runtime/Engine/Private/SkeletalMesh.cpp
--- a/Engine/Source/Runtime/Engine/Private/SkeletalMesh.cpp
+++ b/Engine/Source/Runtime/Engine/Private/SkeletalMesh.cpp
@@ -2149,6 +2149,14 @@ void USkeletalMesh::PreEditChange(UProperty* PropertyAboutToChange)
 {
 	Super::PreEditChange(PropertyAboutToChange);
 
+	if (GIsEditor &&
+		PropertyAboutToChange &&
+		PropertyAboutToChange->GetOuterUField() &&
+		PropertyAboutToChange->GetOuterUField()->GetFName() == FName(TEXT("ClothPhysicsProperties")))
+	{
+		// if this is a member property of ClothPhysicsProperties, don't release render resources to drag sliders smoothly 
+		return;
+	}
 	FlushRenderState();
 }
 
@@ -2157,6 +2165,13 @@ void USkeletalMesh::PostEditChangeProperty(FPropertyChangedEvent& PropertyChange
 	bool bFullPrecisionUVsReallyChanged = false;
 
 	UProperty* PropertyThatChanged = PropertyChangedEvent.Property;
+	
+	// if this is a member property of ClothPhysicsProperties, skip RestartRenderState to drag ClothPhysicsProperties sliders smoothly
+	bool bSkipRestartRenderState = 
+		GIsEditor &&
+		PropertyThatChanged &&
+		PropertyThatChanged->GetOuterUField() &&
+		PropertyThatChanged->GetOuterUField()->GetFName() == FName(TEXT("ClothPhysicsProperties"));	
 
 	if( GIsEditor &&
 		PropertyThatChanged &&
@@ -2186,7 +2201,10 @@ void USkeletalMesh::PostEditChangeProperty(FPropertyChangedEvent& PropertyChange
 		}
 	}
 	
-	RestartRenderState();
+	if (!bSkipRestartRenderState)
+	{
+		RestartRenderState();
+	}
 
 	if ( GIsEditor && PropertyThatChanged && PropertyThatChanged->GetName() == TEXT("StreamingDistanceMultiplier") )
 	{
