makes particles emit from simulating cloth vertices/surfaces

diff --git a/Engine/Source/Runtime/Engine/Classes/Components/SkeletalMeshComponent.h b/Engine/Source/Runtime/Engine/Classes/Components/SkeletalMeshComponent.h
--- a/Engine/Source/Runtime/Engine/Classes/Components/SkeletalMeshComponent.h
+++ b/Engine/Source/Runtime/Engine/Classes/Components/SkeletalMeshComponent.h
@@ -756,6 +756,7 @@ public:
 	virtual void UnHideBone( int32 BoneIndex ) OVERRIDE;
 	virtual void SetPhysicsAsset(class UPhysicsAsset* NewPhysicsAsset,bool bForceReInit = false) OVERRIDE;
 	virtual void SetSkeletalMesh(class USkeletalMesh* NewMesh) OVERRIDE;
+	virtual FVector GetSkinnedVertexPosition(int32 VertexIndex) const OVERRIDE;
 	// End USkinnedMeshComponent interface
 	/** 
 	 *	Iterate over each joint in the physics for this mesh, setting its AngularPositionTarget based on the animation information.
@@ -966,6 +967,12 @@ public:
 	/** only check whether there are valid clothing actors or not */
 	bool HasValidClothingActors();
 
+	/** get root bone matrix by the root bone index which Apex cloth asset is holding */
+	void GetClothRootBoneMatrix(int32 AssetIndex, FMatrix& OutRootBoneMatrix) const;
+
+	/** if the vertex index is valid for simulated vertices, returns the position in world space */
+	bool GetClothSimulatedPosition(int32 AssetIndex, int32 VertexIndex, FVector& OutSimulPos) const;
+
 #if WITH_CLOTH_COLLISION_DETECTION
 
 	/** draws currently intersected collisions */
diff --git a/Engine/Source/Runtime/Engine/Classes/Components/SkinnedMeshComponent.h b/Engine/Source/Runtime/Engine/Classes/Components/SkinnedMeshComponent.h
--- a/Engine/Source/Runtime/Engine/Classes/Components/SkinnedMeshComponent.h
+++ b/Engine/Source/Runtime/Engine/Classes/Components/SkinnedMeshComponent.h
@@ -464,8 +464,8 @@ public:
 	 *
 	 * @param VertexIndex Vertex Index. If compressed, this will be slow. 
 	 */
-	FVector GetSkinnedVertexPosition(int32 VertexIndex) const;
-	
+	virtual FVector GetSkinnedVertexPosition(int32 VertexIndex) const;
+
 	/**
 	* Returns color of the vertex.
 	*
@@ -480,8 +480,8 @@ public:
 	/** 
 	 * Refresh Bone Transform (SpaceBases)
 	 * Each class will need to implement this function
-	 * Ideally this function should be atomic (not relying on Tick or any other update.)
-	 *
+	 * Ideally this function should be atomic (not relying on Tick or any other update.) 
+	 * 
 	 * @param TickFunction Supplied as non null if we are running in a tick, allows us to create graph tasks for parallelism
 	 * 
 	 */
diff --git a/Engine/Source/Runtime/Engine/Private/SkeletalMeshComponentPhysics.cpp b/Engine/Source/Runtime/Engine/Private/SkeletalMeshComponentPhysics.cpp
--- a/Engine/Source/Runtime/Engine/Private/SkeletalMeshComponentPhysics.cpp
+++ b/Engine/Source/Runtime/Engine/Private/SkeletalMeshComponentPhysics.cpp
@@ -1,4 +1,4 @@
-// Copyright 1998-2014 Epic Games, Inc. All Rights Reserved.
+ï»¿// Copyright 1998-2014 Epic Games, Inc. All Rights Reserved.
 
 #include "EnginePrivate.h"
 #include "PhysicsPublic.h"
@@ -1449,6 +1449,94 @@ void USkeletalMeshComponent::UpdatePhysicsToRBChannels()
 	}
 }
 
+FVector USkeletalMeshComponent::GetSkinnedVertexPosition(int32 VertexIndex) const
+{
+#if WITH_APEX_CLOTHING
+	// only if this component has clothing and is showing simulated results	
+	if (SkeletalMesh &&
+		SkeletalMesh->ClothingAssets.Num() > 0 &&
+		MeshObject &&
+		!bDisableClothSimulation &&
+		ClothBlendWeight > 0.0f // if cloth blend weight is 0.0, only showing skinned vertices regardless of simulation positions
+		)
+	{
+		FStaticLODModel& Model = MeshObject->GetSkeletalMeshResource().LODModels[0];
+
+		// Find the chunk and vertex within that chunk, and skinning type, for this vertex.
+		int32 ChunkIndex;
+		int32 VertIndexInChunk;
+		bool bSoftVertex;
+		bool bHasExtraBoneInfluences;
+		Model.GetChunkAndSkinType(VertexIndex, ChunkIndex, VertIndexInChunk, bSoftVertex, bHasExtraBoneInfluences);
+
+		bool bClothVertex = false;
+		int32 ClothAssetIndex = -1;
+
+		// if this chunk has cloth data
+		if (Model.Chunks[ChunkIndex].HasApexClothData())
+		{
+			bClothVertex = true;
+			ClothAssetIndex = Model.Chunks[ChunkIndex].CorrespondClothAssetIndex;
+		}
+		else
+		{
+			// if this chunk corresponds to a cloth section, returns corresponding cloth section's info instead
+			for (int32 SectionIdx = 0; SectionIdx < Model.Sections.Num(); SectionIdx++)
+			{
+				const FSkelMeshSection& Section = Model.Sections[SectionIdx];
+
+				// find a section which has this chunk index
+				if (Section.ChunkIndex == ChunkIndex)
+				{
+					// if current section is disabled and the corresponding cloth section is visible
+					if (Section.bDisabled && Section.CorrespondClothSectionIndex >= 0)
+					{
+						bClothVertex = true;
+
+						const FSkelMeshSection& ClothSection = Model.Sections[Section.CorrespondClothSectionIndex];
+						const FSkelMeshChunk& ClothChunk = Model.Chunks[ClothSection.ChunkIndex];
+
+						ClothAssetIndex = ClothChunk.CorrespondClothAssetIndex;
+
+						// the index can exceed the range because this vertex index is based on the corresponding original section
+						// the number of cloth chunk's vertices is not always same as the corresponding one 
+						// cloth chunk has only soft vertices
+						if (VertIndexInChunk >= ClothChunk.GetNumSoftVertices())
+						{
+							// if the index exceeds, re-assign a random vertex index for this chunk
+							VertIndexInChunk = FMath::TruncToInt(FMath::SRand() * (ClothChunk.GetNumSoftVertices() - 1));
+						}
+					}
+					// if found, quit this loop quickly
+					break;
+				}
+			}
+		}
+
+		if (bClothVertex)
+		{
+			FVector SimulatedPos;
+			if (GetClothSimulatedPosition(ClothAssetIndex, VertIndexInChunk, SimulatedPos))
+			{
+				// a simulated position is in world space and convert this to local space
+				// because SkinnedMeshComponent::GetSkinnedVertexPosition() returns the position in local space
+				SimulatedPos = ComponentToWorld.InverseTransformPosition(SimulatedPos);
+
+				// if blend weight is 1.0, doesn't need to blend with a skinned position
+				if (ClothBlendWeight < 1.0f) 
+				{
+					// blend with a skinned position
+					FVector SkinnedPos = Super::GetSkinnedVertexPosition(VertexIndex);
+					SimulatedPos = SimulatedPos*ClothBlendWeight + SkinnedPos*(1.0f - ClothBlendWeight);
+				}
+				return SimulatedPos;
+			}
+		}
+	}
+#endif // #if WITH_APEX_CLOTHING
+	return Super::GetSkinnedVertexPosition(VertexIndex);
+}
+
 //////////////////////////////////////////////////////////////////////////
 // COLLISION
 
@@ -2911,7 +2999,7 @@ void USkeletalMeshComponent::CheckClothTeleport(float DeltaTime)
 		// Rotation matrix's transpose means an inverse but can't use a transpose because this matrix includes scales
 		FMatrix AInvB = CurRootBoneMat * PrevRootBoneMatrix.Inverse();
 		float Trace = AInvB.M[0][0] + AInvB.M[1][1] + AInvB.M[2][2];
-		float CosineTheta = (Trace - 1.0f) / 2.0f; // trace = 1+2cos(¥è) for a 3¡¿3 matrix
+		float CosineTheta = (Trace - 1.0f) / 2.0f; // trace = 1+2cos(theta) for a 3x3 matrix
 
 		if ( CosineTheta < ClothTeleportCosineThresholdInRad ) // has the root bone rotated too much
 		{
@@ -3027,6 +3115,65 @@ void USkeletalMeshComponent::UpdateClothState(float DeltaTime)
 	// reset to Continuous
 	ClothTeleportMode = FClothingActor::Continuous;
 }
+
+void USkeletalMeshComponent::GetClothRootBoneMatrix(int32 AssetIndex, FMatrix& OutRootBoneMatrix) const
+{
+	if (IsValidClothingActor(AssetIndex))
+	{
+		TSharedPtr<FClothingAssetWrapper> Asset = ClothingActors[AssetIndex].ParentClothingAsset;
+
+		check(Asset.IsValid());
+
+		const NxParameterized::Interface* AssetParams = Asset->GetAsset()->getAssetNxParameterized();
+		uint32 InternalRootBoneIndex;
+		verify(NxParameterized::getParamU32(*AssetParams, "rootBoneIndex", InternalRootBoneIndex));
+		check(InternalRootBoneIndex >= 0);
+		FName BoneName = Asset->GetConvertedBoneName(InternalRootBoneIndex);
+		int32 BoneIndex = GetBoneIndex(BoneName);
+		check(BoneIndex >= 0);
+		OutRootBoneMatrix = GetBoneMatrix(BoneIndex);
+	}
+	else
+	{
+		OutRootBoneMatrix = GetBoneMatrix(0);
+	}
+}
+
+bool USkeletalMeshComponent::GetClothSimulatedPosition(int32 AssetIndex, int32 VertexIndex, FVector& OutSimulPos) const
+{
+	bool bSucceed = false;
+
+	if (IsValidClothingActor(AssetIndex))
+	{
+		NxClothingActor* ApexClothingActor = ClothingActors[AssetIndex].ApexClothingActor;
+
+		if (ApexClothingActor)
+		{
+			uint32 NumSimulVertices = ApexClothingActor->getNumSimulationVertices();
+
+			// handles only simulated vertices, indices of fixed vertices are bigger than # of simulated vertices
+			if ((uint32)VertexIndex < NumSimulVertices)
+			{
+				bSucceed = true;
+				const physx::PxVec3* Vertices = ApexClothingActor->getSimulationPositions();
+
+				if (bLocalSpaceSimulation)
+				{
+					FMatrix ClothRootBoneMatrix;
+					GetClothRootBoneMatrix(AssetIndex, ClothRootBoneMatrix);
+					OutSimulPos = ClothRootBoneMatrix.TransformPosition(P2UVector(Vertices[VertexIndex]));
+				}
+				else
+				{
+					OutSimulPos = P2UVector(Vertices[VertexIndex]);
+				}
+			}
+		}
+	}
+
+	return bSucceed;
+}
+
 #endif// #if WITH_APEX_CLOTHING
 
 void USkeletalMeshComponent::TickClothing(float DeltaTime)
@@ -3104,17 +3251,8 @@ void USkeletalMeshComponent::GetUpdateClothSimulationData(TArray<FClothSimulData
 				// In case of Local space simulation, need to transform simulated positions with the internal bone matrix
 				if (bLocalSpaceSimulation)
 				{
-					check(ClothingActors[ActorIndex].ParentClothingAsset.IsValid());
-
-					const NxParameterized::Interface* AssetParams = ClothingActors[ActorIndex].ParentClothingAsset->GetAsset()->getAssetNxParameterized();
-					uint32 InternalRootBoneIndex;
-					NxParameterized::getParamU32(*AssetParams, "rootBoneIndex", InternalRootBoneIndex);
-					check(InternalRootBoneIndex >= 0);
-					FName BoneName = ClothingActors[ActorIndex].ParentClothingAsset->GetConvertedBoneName(InternalRootBoneIndex);					
-					int32 BoneIndex = GetBoneIndex(BoneName);
-					check(BoneIndex >= 0);
-					FMatrix ClothRootBoneMatrix = GetBoneMatrix(BoneIndex);
-
+					FMatrix ClothRootBoneMatrix;
+					GetClothRootBoneMatrix(ActorIndex, ClothRootBoneMatrix);
 					for (uint32 VertexIndex = 0; VertexIndex < NumSimulVertices; VertexIndex++)
 					{						
 						ClothData.ClothSimulPositions[VertexIndex] = ClothRootBoneMatrix.TransformPosition(P2UVector(Vertices[VertexIndex]));
