Fix non-unity (missing ParticleModuleTypeDataGpu.h)

diff --git a/Engine/Source/Runtime/Engine/Private/Particles/ParticleModules_Location.cpp b/Engine/Source/Runtime/Engine/Private/Particles/ParticleModules_Location.cpp
--- a/Engine/Source/Runtime/Engine/Private/Particles/ParticleModules_Location.cpp
+++ b/Engine/Source/Runtime/Engine/Private/Particles/ParticleModules_Location.cpp
@@ -24,6 +24,7 @@
 #include "Particles/Location/ParticleModuleLocationWorldOffset.h"
 #include "Particles/Location/ParticleModuleLocationWorldOffset_Seeded.h"
 #include "Particles/Location/ParticleModuleLocation_Seeded.h"
+#include "Particles/TypeData/ParticleModuleTypeDataGpu.h"
 #include "Particles/ParticleLODLevel.h"
 #include "Particles/ParticleModuleRequired.h"
 #include "Particles/ParticleSpriteEmitter.h"
