PIE in New Window respects resolution request if  r.FullScreen is 1

#ttp 336153 - PIE in New Window doesn't respect request resolution if  r.FullScreen is 1

#branch UE4

#change Implementation of CL#2079140 from PeterK

diff --git a/Engine/Source/Runtime/Engine/Private/GameEngine.cpp b/Engine/Source/Runtime/Engine/Private/GameEngine.cpp
--- a/Engine/Source/Runtime/Engine/Private/GameEngine.cpp
+++ b/Engine/Source/Runtime/Engine/Private/GameEngine.cpp
@@ -81,16 +81,21 @@ int32 GetBoundFullScreenModeCVar()
 // depending on WindowMode and the console variable r.FullScreenMode
 EWindowMode::Type GetWindowModeType(EWindowMode::Type WindowMode)
 {
-	EWindowMode::Type CurrentWindowMode = EWindowMode::ConvertIntToWindowMode(GetBoundFullScreenModeCVar());
-
 	if (FPlatformProperties::SupportsWindowedMode())
 	{
-		if (WindowMode != EWindowMode::Windowed && GEngine && GEngine->HMDDevice.IsValid())
+		if (WindowMode == EWindowMode::Windowed)
+		{
+			return WindowMode;
+		}
+
+		if (GEngine && GEngine->HMDDevice.IsValid())
 		{
 			return EWindowMode::Fullscreen;
 		}
-		
-		return CurrentWindowMode;
+
+		// Figure out which full screen mode we should be
+		EWindowMode::Type DesiredFullScreenWindowMode = EWindowMode::ConvertIntToWindowMode(GetBoundFullScreenModeCVar());
+		return DesiredFullScreenWindowMode;
 	}
 
 	return EWindowMode::Fullscreen;
