Changed AISystem to be created only for Editor and Server games #UE4

diff --git a/Engine/Source/Editor/UnrealEd/Private/EditorServer.cpp b/Engine/Source/Editor/UnrealEd/Private/EditorServer.cpp
--- a/Engine/Source/Editor/UnrealEd/Private/EditorServer.cpp
+++ b/Engine/Source/Editor/UnrealEd/Private/EditorServer.cpp
@@ -2173,6 +2173,7 @@ bool UEditorEngine::Map_Load(const TCHAR* Str, FOutputDevice& Ar)
 					}
 
 					UNavigationSystem::InitializeForWorld(Context.World(), FNavigationSystem::EditorMode);
+					Context.World()->CreateAISystem();
 
 					// Assign stationary light channels for previewing
 					ULightComponent::ReassignStationaryLightChannels(Context.World(), false);
diff --git a/Engine/Source/Editor/UnrealEd/Private/PlayLevel.cpp b/Engine/Source/Editor/UnrealEd/Private/PlayLevel.cpp
--- a/Engine/Source/Editor/UnrealEd/Private/PlayLevel.cpp
+++ b/Engine/Source/Editor/UnrealEd/Private/PlayLevel.cpp
@@ -2437,6 +2437,7 @@ UWorld* UEditorEngine::CreatePlayInEditorWorld(FWorldContext &PieWorldContext, b
 	}
 	
 	UNavigationSystem::InitializeForWorld(PlayWorld, PieWorldContext.GamePlayers.Num() > 0 ? FNavigationSystem::PIEMode : FNavigationSystem::SimulationMode);
+	PlayWorld->CreateAISystem();
 
 	EditorWorld->TransferBlueprintDebugReferences(PlayWorld);
 
diff --git a/Engine/Source/Runtime/AIModule/Classes/AISystem.h b/Engine/Source/Runtime/AIModule/Classes/AISystem.h
--- a/Engine/Source/Runtime/AIModule/Classes/AISystem.h
+++ b/Engine/Source/Runtime/AIModule/Classes/AISystem.h
@@ -22,8 +22,9 @@ class AIMODULE_API UAISystem : public UAISystemBase
 	UPROPERTY(Transient)
 	UEnvQueryManager* EnvironmentQueryManager;
 
-	// IAISystemInterface begin	
-	//virtual IAISystemInterface* CreateInstanceForWorldImpl(UWorld* World);
+	~UAISystem();
+
+	// IAISystemInterface begin		
 	virtual void InitializeActorsForPlay(bool bTimeGotReset) OVERRIDE;
 	virtual void WorldOriginChanged(FIntPoint OldOrigin, FIntPoint NewOrigin) OVERRIDE;
 	virtual void CleanupWorld(bool bSessionEnded = true, bool bCleanupResources = true, UWorld* NewWorld = NULL) OVERRIDE;
diff --git a/Engine/Source/Runtime/AIModule/Private/AIModule.cpp b/Engine/Source/Runtime/AIModule/Private/AIModule.cpp
--- a/Engine/Source/Runtime/AIModule/Private/AIModule.cpp
+++ b/Engine/Source/Runtime/AIModule/Private/AIModule.cpp
@@ -3,6 +3,8 @@
 #include "AIModulePrivate.h"
 #include "AISystem.h"
 
+DEFINE_LOG_CATEGORY_STATIC(LogAIModule, Log, All);
+
 class FAIModule : public IAIModule
 {
 	// Begin IModuleInterface
@@ -28,6 +30,8 @@ void FAIModule::ShutdownModule()
 
 UAISystemBase* FAIModule::CreateAISystemInstance(UWorld* World)
 {
+	UE_LOG(LogAIModule, Log, TEXT("Creating AISystem for world %s"), *GetNameSafe(World));
+	
 	TSubclassOf<class UAISystemBase> AISystemClass = LoadClass<UAISystemBase>(NULL, *UAISystem::GetAISystemClassName().ClassName, NULL, LOAD_None, NULL);
 	return ConstructObject<UAISystemBase>(AISystemClass, World);
 }
\ No newline at end of file
diff --git a/Engine/Source/Runtime/AIModule/Private/AISystem.cpp b/Engine/Source/Runtime/AIModule/Private/AISystem.cpp
--- a/Engine/Source/Runtime/AIModule/Private/AISystem.cpp
+++ b/Engine/Source/Runtime/AIModule/Private/AISystem.cpp
@@ -16,6 +16,11 @@ UAISystem::UAISystem(const FPostConstructInitializeProperties& PCIP) : Super(PCI
 	}
 }
 
+UAISystem::~UAISystem()
+{
+	CleanupWorld(true, true, NULL);
+}
+
 void UAISystem::InitializeActorsForPlay(bool bTimeGotReset)
 {
 
diff --git a/Engine/Source/Runtime/AIModule/Private/BehaviorTree/BehaviorTreeComponent.cpp b/Engine/Source/Runtime/AIModule/Private/BehaviorTree/BehaviorTreeComponent.cpp
--- a/Engine/Source/Runtime/AIModule/Private/BehaviorTree/BehaviorTreeComponent.cpp
+++ b/Engine/Source/Runtime/AIModule/Private/BehaviorTree/BehaviorTreeComponent.cpp
@@ -1205,7 +1205,7 @@ bool UBehaviorTreeComponent::PushInstance(class UBehaviorTree* TreeAsset)
 		return false;
 	}
 
-	UBehaviorTreeManager* BTManager = Cast<UAISystem>(GetWorld()->GetAISystem())->GetBehaviorTreeManager();
+	UBehaviorTreeManager* BTManager = UAISystem::GetCurrentBTManager(GetWorld());
 	if (BTManager == NULL)
 	{
 		UE_VLOG(GetOwner(), LogBehaviorTree, Warning, TEXT("Failed to execute tree %s: behavior tree manager not found!"), *GetNameSafe(TreeAsset));
diff --git a/Engine/Source/Runtime/Engine/Classes/Engine/World.h b/Engine/Source/Runtime/Engine/Classes/Engine/World.h
--- a/Engine/Source/Runtime/Engine/Classes/Engine/World.h
+++ b/Engine/Source/Runtime/Engine/Classes/Engine/World.h
@@ -1417,6 +1417,11 @@ public:
 	/** NavigationSystem const getter */
 	FORCEINLINE const UNavigationSystem* GetNavigationSystem() const { return NavigationSystem; }
 
+	/** AISystem getter. if AISystem is missing it tries to create one and returns the result.
+	 *	@NOTE the result can be NULL, for example on client games or if no AI module or AISystem class have not been specified
+	 *	@see UAISystemBase::AISystemClassName and UAISystemBase::AISystemModuleName*/
+	UAISystemBase* CreateAISystem();
+
 	/** AISystem getter */
 	FORCEINLINE UAISystemBase* GetAISystem() { return AISystem; }
 	/** AISystem const getter */
@@ -1784,6 +1789,7 @@ public:
 			, bRequiresHitProxies(true)
 			, bCreatePhysicsScene(true)
 			, bCreateNavigation(false)
+			, bCreateAISystem(false)
 			, bShouldSimulatePhysics(true)
 			, bEnableTraceCollision(false)
 			, bTransactional(true)
@@ -1795,6 +1801,7 @@ public:
 		uint32 bRequiresHitProxies:1;
 		uint32 bCreatePhysicsScene:1;
 		uint32 bCreateNavigation:1;
+		uint32 bCreateAISystem:1;
 		uint32 bShouldSimulatePhysics:1;
 		uint32 bEnableTraceCollision:1;
 		uint32 bTransactional:1;
@@ -1805,6 +1812,7 @@ public:
 		InitializationValues& RequiresHitProxies(const bool bRequires) { bRequiresHitProxies = bRequires; return *this; }
 		InitializationValues& CreatePhysicsScene(const bool bCreate) { bCreatePhysicsScene = bCreate; return *this; }
 		InitializationValues& CreateNavigation(const bool bCreate) { bCreateNavigation = bCreate; return *this; }
+		InitializationValues& CreateAISystem(const bool bCreate) { bCreateAISystem = bCreate; return *this; }
 		InitializationValues& ShouldSimulatePhysics(const bool bInShouldSimulatePhysics) { bShouldSimulatePhysics = bInShouldSimulatePhysics; return *this; }
 		InitializationValues& EnableTraceCollision(const bool bInEnableTraceCollision) { bEnableTraceCollision = bInEnableTraceCollision; return *this; }
 		InitializationValues& SetTransactional(const bool bInTransactional) { bTransactional = bInTransactional; return *this; }
diff --git a/Engine/Source/Runtime/Engine/Private/World.cpp b/Engine/Source/Runtime/Engine/Private/World.cpp
--- a/Engine/Source/Runtime/Engine/Private/World.cpp
+++ b/Engine/Source/Runtime/Engine/Private/World.cpp
@@ -629,6 +629,26 @@ void UWorld::UpdateParameterCollectionInstances(bool bUpdateInstanceUniformBuffe
 	}
 }
 
+UAISystemBase* UWorld::CreateAISystem()
+{
+#if WITH_SERVER_CODE || WITH_EDITOR
+	// create navigation system for editor and server targets, but remove it from game clients
+	if (AISystem == NULL && GetNetMode() != NM_Client)
+	{
+		auto AISystemModule = FModuleManager::LoadModulePtr<IAISystemModule>(UAISystemBase::GetAISystemModuleName());
+		if (AISystemModule)
+		{
+			AISystem = AISystemModule->CreateAISystemInstance(this);
+			if (AISystem == NULL)
+			{
+				UE_LOG(LogWorld, Error, TEXT("Failed to create AISystem instance of class %s!"), *UAISystemBase::GetAISystemClassName().ClassName);
+			}
+		}
+	}
+#endif
+	return AISystem; 
+}
+
 void UWorld::InitWorld(const InitializationValues IVS)
 {
 	if (!ensure(!bIsWorldInitialized))
@@ -669,14 +689,9 @@ void UWorld::InitWorld(const InitializationValues IVS)
 		UNavigationSystem::CreateNavigationSystem(this);
 	}
 
-	auto AISystemModule = FModuleManager::LoadModulePtr<IAISystemModule>(UAISystemBase::GetAISystemModuleName());
-	if (AISystemModule)
+	if (IVS.bCreateAISystem)
 	{
-		AISystem = AISystemModule->CreateAISystemInstance(this);
-		if (AISystem == NULL)
-		{
-			UE_LOG(LogWorld, Error, TEXT("Failed to create AISystem instance of class %s!"), *UAISystemBase::GetAISystemClassName().ClassName);
-		}
+		CreateAISystem();
 	}
 	
 	if (GEngine->AvoidanceManagerClass != NULL)
@@ -953,7 +968,7 @@ UWorld* UWorld::CreateWorld( const EWorldType::Type InWorldType, bool bInformEng
 	const FString WorldNameString = (WorldName != NAME_None) ? WorldName.ToString() : TEXT("NewWorld");
 	UWorld* NewWorld = new( WorldPackage				, *WorldNameString			) UWorld(FPostConstructInitializeProperties(),FURL(NULL));
 	NewWorld->WorldType = InWorldType;
-	NewWorld->InitializeNewWorld(UWorld::InitializationValues().ShouldSimulatePhysics(false).EnableTraceCollision(true).CreateNavigation(InWorldType == EWorldType::Editor));
+	NewWorld->InitializeNewWorld(UWorld::InitializationValues().ShouldSimulatePhysics(false).EnableTraceCollision(true).CreateNavigation(InWorldType == EWorldType::Editor).CreateAISystem(InWorldType == EWorldType::Editor));
 
 	// Clear the dirty flag set during SpawnActor and UpdateLevelComponents
 	WorldPackage->SetDirtyFlag(false);
