Further removal of AIModule dependencies from the Engine #UE4

diff --git a/Engine/Source/Runtime/Engine/Private/Animation/AnimInstance.cpp b/Engine/Source/Runtime/Engine/Private/Animation/AnimInstance.cpp
--- a/Engine/Source/Runtime/Engine/Private/Animation/AnimInstance.cpp
+++ b/Engine/Source/Runtime/Engine/Private/Animation/AnimInstance.cpp
@@ -15,16 +15,8 @@
 
 #include "MessageLog.h"
 
-
-// @todo this is here only due to circular dependency to AIModule. To be removed
-#include "AIController.h"
-#include "Navigation/PathFollowingComponent.h"
-#include "Navigation/NavigationComponent.h"
-#include "BrainComponent.h"
-
 /** Anim stats */
 
-
 DEFINE_STAT(STAT_UpdateSkelMeshBounds);
 DEFINE_STAT(STAT_MeshObjectUpdate);
 DEFINE_STAT(STAT_BlendInPhysics);
@@ -1950,42 +1942,12 @@ void UAnimInstance::AddReferencedObjects(UObject* InThis, FReferenceCollector& C
 // 
 void UAnimInstance::LockAIResources(bool bLockMovement, bool LockAILogic)
 {
-	APawn* PawnOwner = TryGetPawnOwner();
-	if (PawnOwner)
-	{
-		AAIController* OwningAI = Cast<AAIController>(PawnOwner->Controller);
-		if( OwningAI )
-		{
-			if (bLockMovement && OwningAI->PathFollowingComponent)
-			{
-				OwningAI->PathFollowingComponent->LockResource(EAILockSource::Animation);
-			}
-			if (LockAILogic && OwningAI->BrainComponent)
-			{
-				OwningAI->BrainComponent->LockResource(EAILockSource::Animation);
-			}			
-		}
-	}
+	UE_LOG(LogAnimation, Error, TEXT("%s: LockAIResources is no longer supported. Please use LockAIResourcesWithAnimation instead."), *GetName());	
 }
 
 void UAnimInstance::UnlockAIResources(bool bUnlockMovement, bool UnlockAILogic)
 {
-	APawn* PawnOwner = TryGetPawnOwner();
-	if (PawnOwner)
-	{
-		AAIController* OwningAI = Cast<AAIController>(PawnOwner->Controller);
-		if( OwningAI )
-		{
-			if (bUnlockMovement && OwningAI->PathFollowingComponent)
-			{
-				OwningAI->PathFollowingComponent->ClearResourceLock(EAILockSource::Animation);
-			}
-			if (UnlockAILogic && OwningAI->BrainComponent)
-			{
-				OwningAI->BrainComponent->ClearResourceLock(EAILockSource::Animation);
-			}			
-		}
-	}
+	UE_LOG(LogAnimation, Error, TEXT("%s: UnlockAIResources is no longer supported. Please use UnlockAIResourcesWithAnimation instead."), *GetName());
 }
 
 #undef LOCTEXT_NAMESPACE 
diff --git a/Engine/Source/Runtime/Engine/Private/PlayerState.cpp b/Engine/Source/Runtime/Engine/Private/PlayerState.cpp
--- a/Engine/Source/Runtime/Engine/Private/PlayerState.cpp
+++ b/Engine/Source/Runtime/Engine/Private/PlayerState.cpp
@@ -7,8 +7,6 @@
 #include "EnginePrivate.h"
 #include "Net/UnrealNetwork.h"
 #include "OnlineSubsystemUtils.h"
-// @todo this is here only due to circular dependency to AIModule. To be removed
-#include "AIController.h"
 
 APlayerState::APlayerState(const class FPostConstructInitializeProperties& PCIP)
 	: Super(PCIP
@@ -110,9 +108,10 @@ void APlayerState::PostInitializeComponents()
 		return;
 	}
 
-	if (Cast<AAIController>(GetOwner()) != NULL)
+	AController* OwningController = Cast<AController>(GetOwner());
+	if (OwningController != NULL)
 	{
-		bIsABot = true;
+		bIsABot = (Cast<APlayerController>(OwningController) == NULL);
 	}
 
 	if (World->GameState)
