Add missing DefaultPhysicsVolume.h to World.cpp

diff --git a/Engine/Source/Runtime/Engine/Private/World.cpp b/Engine/Source/Runtime/Engine/Private/World.cpp
--- a/Engine/Source/Runtime/Engine/Private/World.cpp
+++ b/Engine/Source/Runtime/Engine/Private/World.cpp
@@ -7,6 +7,7 @@
 #include "EnginePrivate.h"
 #include "Engine/LevelScriptBlueprint.h"
 #include "Engine/CullDistanceVolume.h"
+#include "Engine/DefaultPhysicsVolume.h"
 #include "Engine/Console.h"
 #include "Engine/WorldComposition.h"
 #include "Matinee/MatineeActor.h"
