Fix SlateGameResources trying to open the message log in game

diff --git a/Engine/Source/Runtime/Engine/Private/Slate/SlateGameResources.cpp b/Engine/Source/Runtime/Engine/Private/Slate/SlateGameResources.cpp
--- a/Engine/Source/Runtime/Engine/Private/Slate/SlateGameResources.cpp
+++ b/Engine/Source/Runtime/Engine/Private/Slate/SlateGameResources.cpp
@@ -146,23 +146,29 @@ void FSlateGameResources::Log( ISlateStyle::EStyleMessageSeverity Severity, cons
 	case ISlateStyle::EStyleMessageSeverity::Info: EngineMessageSeverity = EMessageSeverity::Info; break;
 	}
 
-	FMessageLog SlateStyleLog( "SlateStyleLog" );
-	SlateStyleLog.AddMessage( FTokenizedMessage::Create( EngineMessageSeverity, Message ) );
-
-	if ( EngineMessageSeverity <= EMessageSeverity::Warning )
+	if( GIsEditor )
 	{
-		SlateStyleLog.Open();
+		FMessageLog SlateStyleLog("SlateStyleLog");
+		SlateStyleLog.AddMessage(FTokenizedMessage::Create(EngineMessageSeverity, Message));
+
+		if (EngineMessageSeverity <= EMessageSeverity::Warning)
+		{
+			SlateStyleLog.Open();
+		}
 	}
 }
 
 void FSlateGameResources::Log( const TSharedRef< class FTokenizedMessage >& Message ) const
 {
-	FMessageLog SlateStyleLog( "SlateStyleLog" );
-	SlateStyleLog.AddMessage( Message );
-
-	if ( Message->GetSeverity() <= EMessageSeverity::Warning )
+	if (GIsEditor)
 	{
-		SlateStyleLog.Open();
+		FMessageLog SlateStyleLog("SlateStyleLog");
+		SlateStyleLog.AddMessage(Message);
+
+		if (Message->GetSeverity() <= EMessageSeverity::Warning)
+		{
+			SlateStyleLog.Open();
+		}
 	}
 }
 
