TTP 336440: SpawnPerUnit and InitialLocation modules together break ribbon emitters; the ribbon emitter is a special beast.  Moved position setting of SpawnPerUnit spawned particles before the spawn module processing, so spawn modules can modify the position.

diff --git a/Engine/Source/Runtime/Engine/Private/Particles/ParticleTrail2EmitterInstance.cpp b/Engine/Source/Runtime/Engine/Private/Particles/ParticleTrail2EmitterInstance.cpp
--- a/Engine/Source/Runtime/Engine/Private/Particles/ParticleTrail2EmitterInstance.cpp
+++ b/Engine/Source/Runtime/Engine/Private/Particles/ParticleTrail2EmitterInstance.cpp
@@ -1818,6 +1818,11 @@ bool FParticleRibbonEmitterInstance::Spawn_Source(float DeltaTime)
 
 				Component->ComponentToWorld = FTransform(CurrRotation, CurrPosition);
 
+				// Moved before the spawn module iteration so the modules have a chance to modify the position; this may change behavior slightly, but is necessary 
+				// for things like InitalLocation modules to not break in combination with SpawnPerUnit
+				Particle->Location = CurrPosition;
+				Particle->OldLocation = CurrPosition;
+
 				// Standard spawn setup
 				PreSpawn(Particle, Location, FVector::ZeroVector);
 				for (int32 SpawnModuleIdx = 0; SpawnModuleIdx < LODLevel->SpawnModules.Num(); SpawnModuleIdx++)
@@ -1856,11 +1861,6 @@ bool FParticleRibbonEmitterInstance::Spawn_Source(float DeltaTime)
 				TrailData->Tangent = CurrTangent * InvCount;
 				TrailData->SpawnTime = ElapsedTime - StoredSpawnTime;
 				TrailData->SpawnDelta = TrueSpawnTime;
-				// Set the location and up vectors
-				//@todo. Need to add support for offsetting from the source...
-				Particle->Location = CurrPosition;
-				Particle->OldLocation = CurrPosition;
-
 				TrailData->Up = CurrUp;
 
 				TrailData->bMovementSpawned = true;
