#ue4 Adding FindSubstring function from branch to main

diff --git a/Engine/Source/Runtime/Engine/Classes/Kismet/KismetStringLibrary.h b/Engine/Source/Runtime/Engine/Classes/Kismet/KismetStringLibrary.h
--- a/Engine/Source/Runtime/Engine/Classes/Kismet/KismetStringLibrary.h
+++ b/Engine/Source/Runtime/Engine/Classes/Kismet/KismetStringLibrary.h
@@ -232,6 +232,18 @@ class UKismetStringLibrary : public UBlueprintFunctionLibrary
 	UFUNCTION(BlueprintPure, meta=(StartIndex="0", Length="1"), Category="Utilities|String")
 	static FString GetSubstring(const FString& SourceString, int32 StartIndex, int32 Length);
 
+	/**
+	 * Finds the starting index of a substring in the a specified string
+	 * @param SearchIn The string to search within
+	 * @param Substring The string to look for in the SearchIn string
+	 * @param bUseCase Whether or not to be case-sensitive
+	 * @param bSearchFromEnd Whether or not to start the search from the end of the string instead of the beginning
+	 * @param StartPosition The position to start the search from
+	 * @return The index (starting from 0 if bSearchFromEnd is false) of the first occurence of the substring
+	 */
+	UFUNCTION(BlueprintPure, Category = "String")
+	static int32 FindSubstring(const FString& SearchIn, const FString& Substring, bool bUseCase = false, bool bSearchFromEnd = false, int32 StartPosition = -1);
+
 	/** 
 	 * Gets a single character from the string (as an integer)
 	 * @param SourceString - The string to convert
diff --git a/Engine/Source/Runtime/Engine/Private/KismetStringLibrary.cpp b/Engine/Source/Runtime/Engine/Private/KismetStringLibrary.cpp
--- a/Engine/Source/Runtime/Engine/Private/KismetStringLibrary.cpp
+++ b/Engine/Source/Runtime/Engine/Private/KismetStringLibrary.cpp
@@ -256,6 +256,14 @@ FString UKismetStringLibrary::GetSubstring(const FString& SourceString, int32 St
 	return SourceString.Mid(StartIndex, Length);
 }
 
+int32 UKismetStringLibrary::FindSubstring(const FString& SearchIn, const FString& Substring, bool bUseCase, bool bSearchFromEnd, int32 StartPosition)
+{
+	ESearchCase::Type Case = bUseCase ? ESearchCase::CaseSensitive : ESearchCase::IgnoreCase;
+	ESearchDir::Type Dir = bSearchFromEnd ? ESearchDir::FromEnd : ESearchDir::FromStart;
+
+	return SearchIn.Find(Substring, Case, Dir, StartPosition);
+}
+
 int32 UKismetStringLibrary::GetCharacterAsNumber(const FString& SourceString, int32 Index)
 {
 	if ((Index >= 0) && (Index < SourceString.Len()))
