non-unity fix (missing GlobalShader.h)

diff --git a/Engine/Source/Runtime/Engine/Private/World.cpp b/Engine/Source/Runtime/Engine/Private/World.cpp
--- a/Engine/Source/Runtime/Engine/Private/World.cpp
+++ b/Engine/Source/Runtime/Engine/Private/World.cpp
@@ -20,6 +20,7 @@
 #include "PrecomputedLightVolume.h"
 #include "UObjectAnnotation.h"
 #include "RenderCore.h"
+#include "GlobalShader.h"
 #include "EngineModule.h"
 #include "ParticleHelper.h"
 #include "TickTaskManagerInterface.h"
