[336738] No longer mark a child blueprint dirty on load if it contains a component that's attached to a parent blueprint or native class component.

- Modified USCS_Node::SetParent() to not call Modify() on the SCS node unless the attachment is actually being modified.

diff --git a/Engine/Source/Runtime/Engine/Private/SCS_Node.cpp b/Engine/Source/Runtime/Engine/Private/SCS_Node.cpp
--- a/Engine/Source/Runtime/Engine/Private/SCS_Node.cpp
+++ b/Engine/Source/Runtime/Engine/Private/SCS_Node.cpp
@@ -281,22 +281,40 @@ void USCS_Node::SetParent(USCS_Node* InParentNode)
 	check(InParentNode->GetSCS()->GetBlueprint() != NULL);
 	check(InParentNode->GetSCS()->GetBlueprint()->GeneratedClass != NULL);
 
-	Modify();
+	const FName NewParentComponentOrVariableName = InParentNode->VariableName;
+	const FName NewParentComponentOwnerClassName = InParentNode->GetSCS()->GetBlueprint()->GeneratedClass->GetFName();
 
-	bIsParentComponentNative = false;
-	ParentComponentOrVariableName = InParentNode->VariableName;
-	ParentComponentOwnerClassName = InParentNode->GetSCS()->GetBlueprint()->GeneratedClass->GetFName();
+	// Only modify if it differs from current
+	if(bIsParentComponentNative
+		|| ParentComponentOrVariableName != NewParentComponentOrVariableName
+		|| ParentComponentOwnerClassName != NewParentComponentOwnerClassName)
+	{
+		Modify();
+
+		bIsParentComponentNative = false;
+		ParentComponentOrVariableName = NewParentComponentOrVariableName;
+		ParentComponentOwnerClassName = NewParentComponentOwnerClassName;
+	}
 }
 
 void USCS_Node::SetParent(USceneComponent* InParentComponent)
 {
 	check(InParentComponent != NULL);
 
-	Modify();
+	const FName NewParentComponentOrVariableName = InParentComponent->GetFName();
+	const FName NewParentComponentOwnerClassName = NAME_None;
 
-	bIsParentComponentNative = true;
-	ParentComponentOrVariableName = InParentComponent->GetFName();
-	ParentComponentOwnerClassName = NAME_None;
+	// Only modify if it differs from current
+	if(bIsParentComponentNative
+		|| ParentComponentOrVariableName != NewParentComponentOrVariableName
+		|| ParentComponentOwnerClassName != NewParentComponentOwnerClassName)
+	{
+		Modify();
+
+		bIsParentComponentNative = true;
+		ParentComponentOrVariableName = NewParentComponentOrVariableName;
+		ParentComponentOwnerClassName = NewParentComponentOwnerClassName;
+	}
 }
 
 USceneComponent* USCS_Node::GetParentComponentTemplate(UBlueprint* InBlueprint) const
