Paper2D: Add mobility / simulation setting synchronization on sprite components in the same manner as static mesh components

diff --git a/Engine/Plugins/Experimental/Paper2D/Source/Paper2D/Classes/PaperRenderComponent.h b/Engine/Plugins/Experimental/Paper2D/Source/Paper2D/Classes/PaperRenderComponent.h
--- a/Engine/Plugins/Experimental/Paper2D/Source/Paper2D/Classes/PaperRenderComponent.h
+++ b/Engine/Plugins/Experimental/Paper2D/Source/Paper2D/Classes/PaperRenderComponent.h
@@ -21,8 +21,6 @@ protected:
 	// The color of the sprite (passed to the sprite material as a vertex color)
 	UPROPERTY(BlueprintReadOnly, Interp, Category=Sprite)
 	FLinearColor SpriteColor;
-protected:
-	friend class FPaperRenderSceneProxy;
 
 public:
 	/** Change the PaperSprite used by this instance. */
@@ -37,6 +35,16 @@ public:
 	UFUNCTION(BlueprintCallable, Category="Sprite")
 	void SetSpriteColor(FLinearColor NewColor);
 
+	// Returns the wireframe color to use for this component.
+	FLinearColor GetWireframeColor() const;
+
+public:
+	// UObject interface
+#if WITH_EDITOR
+	virtual void PostEditChangeProperty(FPropertyChangedEvent& PropertyChangedEvent) OVERRIDE;
+#endif
+	// End of UObject interface
+
 	// UActorComponent interface
 	virtual void SendRenderDynamicData_Concurrent() OVERRIDE;
 	virtual const UObject* AdditionalStatObject() const OVERRIDE;
@@ -54,6 +62,6 @@ public:
 	virtual class UBodySetup* GetBodySetup() OVERRIDE;
 	// End of UPrimitiveComponent interface
 
-	// Returns the wireframe color to use for this component.
-	FLinearColor GetWireframeColor() const;
+protected:
+	friend class FPaperRenderSceneProxy;
 };
diff --git a/Engine/Plugins/Experimental/Paper2D/Source/Paper2D/Private/PaperRenderComponent.cpp b/Engine/Plugins/Experimental/Paper2D/Source/Paper2D/Private/PaperRenderComponent.cpp
--- a/Engine/Plugins/Experimental/Paper2D/Source/Paper2D/Private/PaperRenderComponent.cpp
+++ b/Engine/Plugins/Experimental/Paper2D/Source/Paper2D/Private/PaperRenderComponent.cpp
@@ -11,8 +11,6 @@
 UPaperRenderComponent::UPaperRenderComponent(const FPostConstructInitializeProperties& PCIP)
 	: Super(PCIP)
 {
-	BodyInstance.SetCollisionEnabled(ECollisionEnabled::NoCollision);
-
 	SetCollisionProfileName(UCollisionProfile::BlockAllDynamic_ProfileName);
 
 	MaterialOverride = nullptr;
@@ -22,6 +20,15 @@ UPaperRenderComponent::UPaperRenderComponent(const FPostConstructInitializePrope
 	Mobility = EComponentMobility::Static;
 }
 
+#if WITH_EDITOR
+void UPaperRenderComponent::PostEditChangeProperty(FPropertyChangedEvent& PropertyChangedEvent)
+{
+	FBodyInstanceEditorHelpers::EnsureConsistentMobilitySimulationSettingsOnPostEditChange(this, PropertyChangedEvent);
+
+	Super::PostEditChangeProperty(PropertyChangedEvent);
+}
+#endif
+
 FPrimitiveSceneProxy* UPaperRenderComponent::CreateSceneProxy()
 {
 	FPaperRenderSceneProxy* NewProxy = new FPaperRenderSceneProxy(this);
diff --git a/Engine/Source/Runtime/Engine/Classes/Components/StaticMeshComponent.h b/Engine/Source/Runtime/Engine/Classes/Components/StaticMeshComponent.h
--- a/Engine/Source/Runtime/Engine/Classes/Components/StaticMeshComponent.h
+++ b/Engine/Source/Runtime/Engine/Classes/Components/StaticMeshComponent.h
@@ -208,6 +208,8 @@ public:
 	virtual FBoxSphereBounds CalcBounds(const FTransform & LocalToWorld) const OVERRIDE;
 	virtual bool HasAnySockets() const OVERRIDE;
 	virtual void QuerySupportedSockets(TArray<FComponentSocketDescription>& OutSockets) const OVERRIDE;
+	virtual FTransform GetSocketTransform(FName InSocketName, ERelativeTransformSpace TransformSpace = RTS_World) const OVERRIDE;
+	virtual bool DoesSocketExist(FName InSocketName) const OVERRIDE;
 	virtual bool ShouldCollideWhenPlacing() const OVERRIDE
 	{
 		// Current Method of collision does not work with non-capsule shapes, enable when it works with static meshes
@@ -392,11 +394,6 @@ public:
 	 */
 	virtual bool SetStaticLightingMapping(bool bTextureMapping, int32 ResolutionToUse);
 
-	/** Socket support overrides. */
-	virtual FTransform GetSocketTransform(FName InSocketName, ERelativeTransformSpace TransformSpace = RTS_World) const OVERRIDE;
-
-	virtual bool DoesSocketExist(FName InSocketName) const OVERRIDE;
-
 	/**
 	 * Returns the named socket on the static mesh component.
 	 *
diff --git a/Engine/Source/Runtime/Engine/Classes/PhysicsEngine/BodyInstance.h b/Engine/Source/Runtime/Engine/Classes/PhysicsEngine/BodyInstance.h
--- a/Engine/Source/Runtime/Engine/Classes/PhysicsEngine/BodyInstance.h
+++ b/Engine/Source/Runtime/Engine/Classes/PhysicsEngine/BodyInstance.h
@@ -554,3 +554,16 @@ protected:
 
 #endif	//WITH_BOX2D
 };
+
+
+#if WITH_EDITOR
+
+// Helper methods for classes with body instances
+struct FBodyInstanceEditorHelpers
+{
+	ENGINE_API static void EnsureConsistentMobilitySimulationSettingsOnPostEditChange(UPrimitiveComponent* Component, FPropertyChangedEvent& PropertyChangedEvent);
+
+private:
+	FBodyInstanceEditorHelpers() {}
+};
+#endif
\ No newline at end of file
diff --git a/Engine/Source/Runtime/Engine/Private/PhysicsEngine/BodyInstance.cpp b/Engine/Source/Runtime/Engine/Private/PhysicsEngine/BodyInstance.cpp
--- a/Engine/Source/Runtime/Engine/Private/PhysicsEngine/BodyInstance.cpp
+++ b/Engine/Source/Runtime/Engine/Private/PhysicsEngine/BodyInstance.cpp
@@ -2871,4 +2871,58 @@ bool FBodyInstance::UseAsyncScene() const
 	return bUseAsyncScene && bHasAsyncScene;
 }
 
+////////////////////////////////////////////////////////////////////////////
+// FBodyInstanceEditorHelpers
+
+void FBodyInstanceEditorHelpers::EnsureConsistentMobilitySimulationSettingsOnPostEditChange(UPrimitiveComponent* Component, FPropertyChangedEvent& PropertyChangedEvent)
+{
+	if (UProperty* PropertyThatChanged = PropertyChangedEvent.Property)
+	{
+		const FName PropertyName = PropertyThatChanged->GetFName();
+
+		// Automatically change collision profile based on mobility and physics settings (if it is currently one of the default profiles)
+		const bool bMobilityChanged = PropertyName == GET_MEMBER_NAME_CHECKED(USceneComponent, Mobility);
+		const bool bSimulatePhysicsChanged = PropertyName == GET_MEMBER_NAME_CHECKED(FBodyInstance, bSimulatePhysics);
+
+		if (bMobilityChanged || bSimulatePhysicsChanged)
+		{
+			// If we enabled physics simulation, but we are not marked movable, do that for them
+			if (bSimulatePhysicsChanged && Component->BodyInstance.bSimulatePhysics && (Component->Mobility != EComponentMobility::Movable))
+			{
+				Component->SetMobility(EComponentMobility::Movable);
+			}
+			// If we made the component no longer movable, but simulation was enabled, disable that for them
+			else if (bMobilityChanged && (Component->Mobility != EComponentMobility::Movable) && Component->BodyInstance.bSimulatePhysics)
+			{
+				Component->BodyInstance.bSimulatePhysics = false;
+			}
+
+			// If the collision profile is one of the 'default' ones for a StaticMeshActor, make sure it is the correct one
+			// If user has changed it to something else, don't touch it
+			const FName CurrentProfileName = Component->BodyInstance.GetCollisionProfileName();
+			if ((CurrentProfileName == UCollisionProfile::BlockAll_ProfileName) ||
+				(CurrentProfileName == UCollisionProfile::BlockAllDynamic_ProfileName) ||
+				(CurrentProfileName == UCollisionProfile::PhysicsActor_ProfileName))
+			{
+				if (Component->Mobility == EComponentMobility::Movable)
+				{
+					if (Component->BodyInstance.bSimulatePhysics)
+					{
+						Component->BodyInstance.SetCollisionProfileName(UCollisionProfile::PhysicsActor_ProfileName);
+					}
+					else
+					{
+						Component->BodyInstance.SetCollisionProfileName(UCollisionProfile::BlockAllDynamic_ProfileName);
+					}
+				}
+				else
+				{
+					Component->BodyInstance.SetCollisionProfileName(UCollisionProfile::BlockAll_ProfileName);
+				}
+			}
+		}
+	}
+}
+
+
 #undef LOCTEXT_NAMESPACE
diff --git a/Engine/Source/Runtime/Engine/Private/StaticMeshComponent.cpp b/Engine/Source/Runtime/Engine/Private/StaticMeshComponent.cpp
--- a/Engine/Source/Runtime/Engine/Private/StaticMeshComponent.cpp
+++ b/Engine/Source/Runtime/Engine/Private/StaticMeshComponent.cpp
@@ -1060,49 +1060,10 @@ void UStaticMeshComponent::PostEditChangeProperty(FPropertyChangedEvent& Propert
 		{
 			GEngine->TriggerStreamingDataRebuild();
 		}
-
-		// Automatically change collision profile based on mobility and physics settings (if it is currently one of the default profiles)
-		bool bMobilityChanged = PropertyThatChanged->GetName().Contains(TEXT("Mobility"));
-		bool bSimulatePhysicsChanged = PropertyThatChanged->GetName().Contains(TEXT("bSimulatePhysics"));
-		if( bMobilityChanged || bSimulatePhysicsChanged )
-		{
-			// If we enabled physics simulation, but we are not marked movable, do that for them
-			if(bSimulatePhysicsChanged && BodyInstance.bSimulatePhysics && (Mobility != EComponentMobility::Movable))
-			{
-				SetMobility(EComponentMobility::Movable);
-			}
-			// If we made the component no longer movable, but simulation was enabled, disable that for them
-			else if(bMobilityChanged && (Mobility != EComponentMobility::Movable) && BodyInstance.bSimulatePhysics)
-			{
-				BodyInstance.bSimulatePhysics = false;
-			}
-
-			// If the collision profile is one of the 'default' ones for a StaticMeshActor, make sure it is the correct one
-			// If user has changed it to something else, don't touch it
-			FName CurrentProfileName = BodyInstance.GetCollisionProfileName();
-			if(	CurrentProfileName == UCollisionProfile::BlockAll_ProfileName ||
-				CurrentProfileName == UCollisionProfile::BlockAllDynamic_ProfileName ||
-				CurrentProfileName == UCollisionProfile::PhysicsActor_ProfileName )
-			{
-				if(Mobility == EComponentMobility::Movable)
-				{
-					if(BodyInstance.bSimulatePhysics)
-					{
-						BodyInstance.SetCollisionProfileName(UCollisionProfile::PhysicsActor_ProfileName);
-					}
-					else
-					{
-						BodyInstance.SetCollisionProfileName(UCollisionProfile::BlockAllDynamic_ProfileName);
-					}
-				}
-				else
-				{
-					BodyInstance.SetCollisionProfileName(UCollisionProfile::BlockAll_ProfileName);
-				}
-			}
-		}
 	}
 
+	FBodyInstanceEditorHelpers::EnsureConsistentMobilitySimulationSettingsOnPostEditChange(this, PropertyChangedEvent);
+
 	LightmassSettings.EmissiveBoost = FMath::Max(LightmassSettings.EmissiveBoost, 0.0f);
 	LightmassSettings.DiffuseBoost = FMath::Max(LightmassSettings.DiffuseBoost, 0.0f);
 
