Physics: Fixed non-editor build failure

diff --git a/Engine/Source/Runtime/Engine/Private/PhysicsEngine/BodyInstance.cpp b/Engine/Source/Runtime/Engine/Private/PhysicsEngine/BodyInstance.cpp
--- a/Engine/Source/Runtime/Engine/Private/PhysicsEngine/BodyInstance.cpp
+++ b/Engine/Source/Runtime/Engine/Private/PhysicsEngine/BodyInstance.cpp
@@ -2874,6 +2874,8 @@ bool FBodyInstance::UseAsyncScene() const
 ////////////////////////////////////////////////////////////////////////////
 // FBodyInstanceEditorHelpers
 
+#if WITH_EDITOR
+
 void FBodyInstanceEditorHelpers::EnsureConsistentMobilitySimulationSettingsOnPostEditChange(UPrimitiveComponent* Component, FPropertyChangedEvent& PropertyChangedEvent)
 {
 	if (UProperty* PropertyThatChanged = PropertyChangedEvent.Property)
@@ -2924,5 +2926,6 @@ void FBodyInstanceEditorHelpers::EnsureConsistentMobilitySimulationSettingsOnPos
 	}
 }
 
+#endif // WITH_EDITOR
 
 #undef LOCTEXT_NAMESPACE
