Physics: Fix DrawCollisionSides collision

diff --git a/Engine/Source/Runtime/Engine/Private/PhysicsEngine2D/AggregateGeometry2D.cpp b/Engine/Source/Runtime/Engine/Private/PhysicsEngine2D/AggregateGeometry2D.cpp
--- a/Engine/Source/Runtime/Engine/Private/PhysicsEngine2D/AggregateGeometry2D.cpp
+++ b/Engine/Source/Runtime/Engine/Private/PhysicsEngine2D/AggregateGeometry2D.cpp
@@ -3,7 +3,7 @@
 #include "EnginePrivate.h"
 #include "PhysicsEngine/AggregateGeometry2D.h"
 
-static const int32 DrawCollisionSides = 16;
+static const int32 DrawCollisionSides2D = 16;
 
 //@TODO: PAPER2D: These should be renamed, and set when the Paper2D ones are loaded
 FVector PaperAxisX(1.0f, 0.0f, 0.0f);
@@ -25,7 +25,7 @@ void FCircleElement2D::DrawElemWire(FPrimitiveDrawInterface* PDI, const FTransfo
 	const FVector X = ElemTM.TransformVector(PaperAxisX);
 	const FVector Y = ElemTM.TransformVector(PaperAxisY);
 
-	DrawCircle(PDI, ElemCenter, X, Y, Color, Scale * Radius, DrawCollisionSides, SDPG_World);
+	DrawCircle(PDI, ElemCenter, X, Y, Color, Scale * Radius, DrawCollisionSides2D, SDPG_World);
 }
 
 void FCircleElement2D::DrawElemSolid(class FPrimitiveDrawInterface* PDI, const FTransform& ElemTM, float Scale, const FMaterialRenderProxy* MaterialRenderProxy)
@@ -34,7 +34,7 @@ void FCircleElement2D::DrawElemSolid(class FPrimitiveDrawInterface* PDI, const F
 	const FVector X = ElemTM.TransformVector(PaperAxisX);
 	const FVector Y = ElemTM.TransformVector(PaperAxisY);
 
-	DrawDisc(PDI, ElemCenter, X, Y, FColor::White, Scale * Radius, DrawCollisionSides, MaterialRenderProxy, SDPG_World);
+	DrawDisc(PDI, ElemCenter, X, Y, FColor::White, Scale * Radius, DrawCollisionSides2D, MaterialRenderProxy, SDPG_World);
 }
 
 //////////////////////////////////////////////////////////////////////////
