Fix crash when trying to simulate car without any wheel names

diff --git a/Engine/Source/Runtime/Engine/Private/Vehicles/WheeledVehicleMovementComponent4W.cpp b/Engine/Source/Runtime/Engine/Private/Vehicles/WheeledVehicleMovementComponent4W.cpp
--- a/Engine/Source/Runtime/Engine/Private/Vehicles/WheeledVehicleMovementComponent4W.cpp
+++ b/Engine/Source/Runtime/Engine/Private/Vehicles/WheeledVehicleMovementComponent4W.cpp
@@ -251,6 +251,15 @@ void UWheeledVehicleMovementComponent4W::SetupVehicle()
 		return;
 	}
 
+	for (int32 WheelIdx = 0; WheelIdx < WheelSetups.Num(); ++WheelIdx)
+	{
+		const FWheelSetup & WheelSetup = WheelSetups[WheelIdx];
+		if (WheelSetup.BoneName == NAME_None)
+		{
+			return;
+		}
+	}
+
 	// Setup the chassis and wheel shapes
 	SetupVehicleShapes();
 
