Don't explicitly require packages to be known when exporting guids

* Undo a change that was made that requires the server to identify packages explicitly when exporting guids
* Cleanup

diff --git a/Engine/Source/Runtime/CoreUObject/Public/UObject/CoreNet.h b/Engine/Source/Runtime/CoreUObject/Public/UObject/CoreNet.h
--- a/Engine/Source/Runtime/CoreUObject/Public/UObject/CoreNet.h
+++ b/Engine/Source/Runtime/CoreUObject/Public/UObject/CoreNet.h
@@ -167,7 +167,7 @@ class COREUOBJECT_API UPackageMap : public UObject
 	FClassNetCache* GetClassNetCache(UClass* Class);
 	void	ClearClassNetCache();
 
-	virtual UObject * ResolvePathAndAssignNetGUID( FNetworkGUID & InOutNetGUID, FString Path, UObject * ObjOuter ) { return NULL; }
+	virtual UObject * ResolvePathAndAssignNetGUID( FNetworkGUID & InOutNetGUID, const FString & PathName, UObject * ObjOuter, const bool bNoLoad = false ) { return NULL; }
 	
 	virtual void HandleUnAssignedObject(const UObject* Obj) { }
 
diff --git a/Engine/Source/Runtime/Engine/Classes/Engine/PackageMapClient.h b/Engine/Source/Runtime/Engine/Classes/Engine/PackageMapClient.h
--- a/Engine/Source/Runtime/Engine/Classes/Engine/PackageMapClient.h
+++ b/Engine/Source/Runtime/Engine/Classes/Engine/PackageMapClient.h
@@ -63,8 +63,6 @@ class UPackageMapClient : public UPackageMap
 	virtual void LogDebugInfo(FOutputDevice& Ar) OVERRIDE;
 	virtual bool SupportsObject( const UObject* Object ) OVERRIDE;
 	
-	virtual UObject * ResolvePathAndAssignNetGUID( FNetworkGUID & InOutNetGUID, FString Path, UObject * ObjOuter ) OVERRIDE;
-
 	virtual bool WriteObject( FArchive& Ar, UObject* Outer, FNetworkGUID NetGUID, FString ObjName ) OVERRIDE;
 
 	virtual void ResetPackageMap() OVERRIDE;
@@ -106,9 +104,10 @@ protected:
 
 	FNetworkGUID	InternalLoadObject( FArchive & Ar, UObject *& Object, const bool bIsOuterLevel, const int InternalLoadObjectRecursionCount );
 
-	void		ResolvePathAndAssignNetGUID_Deferred( const FNetworkGUID & NetGUID, const FString & PathName, const FNetworkGUID & OuterGUID, const bool bIsPackage );
-	UObject *	ResolvePathAndAssignNetGUID( FNetworkGUID & InOutNetGUID, const FString & PathName, const FNetworkGUID & OuterGUID, const bool bIsPackage, const bool bNoLoad );
-	UObject *	ResolvePathAndAssignNetGUID( FNetworkGUID & InOutNetGUID, const FString & PathName, UObject * ObjOuter, const bool bIsPackage, const bool bNoLoad );
+	void		ResolvePathAndAssignNetGUID_Deferred( const FNetworkGUID & NetGUID, const FString & PathName, const FNetworkGUID & OuterGUID );
+	UObject *	ResolvePathAndAssignNetGUID( FNetworkGUID & InOutNetGUID, const FString & PathName, const FNetworkGUID & OuterGUID, const bool bNoLoad = false );
+	
+	virtual UObject * ResolvePathAndAssignNetGUID( FNetworkGUID & InOutNetGUID, const FString & PathName, UObject * ObjOuter, const bool bNoLoad = false ) OVERRIDE;
 
 	bool	ShouldSendFullPath(const UObject* Object, const FNetworkGUID &NetGUID);
 
diff --git a/Engine/Source/Runtime/Engine/Private/PackageMapClient.cpp b/Engine/Source/Runtime/Engine/Private/PackageMapClient.cpp
--- a/Engine/Source/Runtime/Engine/Private/PackageMapClient.cpp
+++ b/Engine/Source/Runtime/Engine/Private/PackageMapClient.cpp
@@ -564,15 +564,6 @@ bool UPackageMapClient::InternalIsMapped(UObject* Object, FNetworkGUID &NetGUID)
 	return Mapped;
 }
 
-/**	
- *	Finds objects from PathName and ObjOuter, and assigns it a NetGUID. Called on the client after serialziing a
- *	<NetGUID, Path> pair.
- */
-UObject * UPackageMapClient::ResolvePathAndAssignNetGUID( FNetworkGUID & InOutNetGUID, FString PathName, UObject *ObjOuter)
-{
-	return ResolvePathAndAssignNetGUID( InOutNetGUID, PathName, ObjOuter, false, false );
-}
-
 //--------------------------------------------------------------------
 //
 //	Writing
@@ -787,7 +778,7 @@ FNetworkGUID UPackageMapClient::InternalLoadObject( FArchive & Ar, UObject *& Ob
 		if ( Object == NULL )
 		{
 			// Try to resolve the name (we may have already loaded it), but DON'T load yet if we haven't already
-			Object = ResolvePathAndAssignNetGUID( NetGUID, PathName, OuterGUID, ExportFlags.bIsPackage, false );
+			Object = ResolvePathAndAssignNetGUID( NetGUID, PathName, OuterGUID, true );
 		}
 
 		if ( Object != NULL )
@@ -803,11 +794,11 @@ FNetworkGUID UPackageMapClient::InternalLoadObject( FArchive & Ar, UObject *& Ob
 		if ( bCanDefer )
 		{
 			UE_LOG( LogNetPackageMap, Log, TEXT( "InternalLoadObject: Deferring resolve. Path: %s, NetGUID: %s, Outer: %s" ), *PathName, *NetGUID.ToString(), ObjOuter != NULL ? *ObjOuter->GetFullName() : TEXT( "NULL" ) );
-			ResolvePathAndAssignNetGUID_Deferred( NetGUID, PathName, OuterGUID, ExportFlags.bIsPackage );
+			ResolvePathAndAssignNetGUID_Deferred( NetGUID, PathName, OuterGUID );
 		}
 		else
 		{
-			Object = ResolvePathAndAssignNetGUID( NetGUID, PathName, OuterGUID, ExportFlags.bIsPackage, false );
+			Object = ResolvePathAndAssignNetGUID( NetGUID, PathName, OuterGUID );
 		}
 	}
 
@@ -819,7 +810,7 @@ FNetworkGUID UPackageMapClient::InternalLoadObject( FArchive & Ar, UObject *& Ob
 	return NetGUID;
 }
 
-void UPackageMapClient::ResolvePathAndAssignNetGUID_Deferred( const FNetworkGUID & NetGUID, const FString & PathName, const FNetworkGUID & OuterGUID, const bool bIsPackage )
+void UPackageMapClient::ResolvePathAndAssignNetGUID_Deferred( const FNetworkGUID & NetGUID, const FString & PathName, const FNetworkGUID & OuterGUID )
 {
 	check( !NetGUID.IsDefault() );		// We can't defer default guids, doesn't make sense
 
@@ -844,7 +835,7 @@ void UPackageMapClient::ResolvePathAndAssignNetGUID_Deferred( const FNetworkGUID
 	}
 }
 
-UObject * UPackageMapClient::ResolvePathAndAssignNetGUID( FNetworkGUID & InOutNetGUID, const FString & PathName, UObject * ObjOuter, const bool bIsPackage, const bool bNoLoad )
+UObject * UPackageMapClient::ResolvePathAndAssignNetGUID( FNetworkGUID & InOutNetGUID, const FString & PathName, UObject * ObjOuter, const bool bNoLoad )
 {
 	if ( InOutNetGUID.IsDynamic() )
 	{
@@ -862,30 +853,37 @@ UObject * UPackageMapClient::ResolvePathAndAssignNetGUID( FNetworkGUID & InOutNe
 			// This could be garbage or a security risk
 			// Another possibility is in dynamic property replication if the server reads the previously serialized state
 			// that has a now destroyed actor in it.
-			UE_LOG( LogNetPackageMap, Warning, TEXT( "Could not find Object for: NetGUID <%s, %s> (and IsNetGUIDAuthority())" ), *InOutNetGUID.ToString(), *PathName );
+			UE_LOG( LogNetPackageMap, Warning, TEXT( "UPackageMapClient::ResolvePathAndAssignNetGUID: Could not find Object for: NetGUID <%s, %s> (and IsNetGUIDAuthority() == false)" ), *InOutNetGUID.ToString(), *PathName );
 			return NULL;
 		}
 		else
 		{
-			UE_LOG( LogNetPackageMap, Log, TEXT( "Could not find Object for: NetGUID <%s, %s>" ), *InOutNetGUID.ToString(), *PathName );
+			UE_LOG( LogNetPackageMap, Log, TEXT( "UPackageMapClient::ResolvePathAndAssignNetGUID: Could not find Object for: NetGUID <%s, %s>" ), *InOutNetGUID.ToString(), *PathName );
 		}
 
 		Object = StaticLoadObject( UObject::StaticClass(), ObjOuter, *PathName, NULL, LOAD_NoWarn );
 
 		if ( Object )
 		{
-			UE_LOG(LogNetPackageMap, Log, TEXT("  StaticLoadObject. Found: %s" ), Object ? *Object->GetName() : TEXT("NULL") );
+			UE_LOG( LogNetPackageMap, Log, TEXT( "UPackageMapClient::ResolvePathAndAssignNetGUID: StaticLoadObject. Found: %s" ), Object ? *Object->GetName() : TEXT( "NULL" ) );
 		}
 		else
 		{
-			if ( bIsPackage )
+			//
+			// If we failed to load it as an object, try to load it as a package
+			//
+
+			UPackage * PackageOuter = Cast< UPackage >( ObjOuter );
+
+			// If we have an outer at this point, it only makes sense for it to be a package (or NULL)
+			if ( PackageOuter != NULL || ObjOuter == NULL )
 			{
-				Object = LoadPackage( Cast< UPackage >( ObjOuter ), *PathName, LOAD_None );
-				UE_LOG( LogNetPackageMap, Log, TEXT( "UPackageMapClient::ResolvePathAndAssignNetGUID  LoadPackage %s. Found: %s" ), *PathName, Object ? *Object->GetName() : TEXT( "NULL" ) );
-			}
+				Object = LoadPackage( PackageOuter, *PathName, LOAD_None );
+				UE_LOG( LogNetPackageMap, Log, TEXT( "UPackageMapClient::ResolvePathAndAssignNetGUID: LoadPackage: %s. Found: %s" ), *PathName, Object ? *Object->GetName() : TEXT( "NULL" ) );
+			} 
 			else
 			{
-				UE_LOG( LogNetPackageMap, Warning, TEXT( "Outer %s is not a package for object %s" ), ObjOuter ? *ObjOuter->GetName() : TEXT( "NULL" ), *PathName );
+				UE_LOG( LogNetPackageMap, Warning, TEXT( "UPackageMapClient::ResolvePathAndAssignNetGUID: Outer %s is not a package for object %s" ), ObjOuter ? *ObjOuter->GetName() : TEXT( "NULL" ), *PathName );
 			}
 		}
 	}
@@ -898,7 +896,7 @@ UObject * UPackageMapClient::ResolvePathAndAssignNetGUID( FNetworkGUID & InOutNe
 	return Object;
 }
 
-UObject * UPackageMapClient::ResolvePathAndAssignNetGUID( FNetworkGUID & InOutNetGUID, const FString & PathName, const FNetworkGUID & OuterGUID, const bool bIsPackage, const bool bNoLoad )
+UObject * UPackageMapClient::ResolvePathAndAssignNetGUID( FNetworkGUID & InOutNetGUID, const FString & PathName, const FNetworkGUID & OuterGUID, const bool bNoLoad )
 {
 	UObject * ObjOuter = OuterGUID.IsValid() ? GetObjectFromNetGUID( OuterGUID ) : NULL;
 
@@ -908,9 +906,9 @@ UObject * UPackageMapClient::ResolvePathAndAssignNetGUID( FNetworkGUID & InOutNe
 		return NULL;
 	}
 
-	UObject * Object = ResolvePathAndAssignNetGUID( InOutNetGUID, PathName, ObjOuter, bIsPackage, bNoLoad );
+	UObject * Object = ResolvePathAndAssignNetGUID( InOutNetGUID, PathName, ObjOuter, bNoLoad );
 
-	if ( Object == NULL )
+	if ( Object == NULL && !bNoLoad )
 	{
 		UE_LOG( LogNetPackageMap, Warning, TEXT( "ResolvePathAndAssignNetGUID: Unable to resolve object. Path: %s, NetGUID: %s, Outer: %s" ), *PathName, *InOutNetGUID.ToString(), ObjOuter != NULL ? *ObjOuter->GetPathName() : TEXT( "NULL" ) );
 	}
