Physics: Add 2D implementations of several methods in FBodyInstance and improve 2D InitBody feature parity

diff --git a/Engine/Source/Runtime/Engine/Classes/PhysicsEngine/BodyInstance.h b/Engine/Source/Runtime/Engine/Classes/PhysicsEngine/BodyInstance.h
--- a/Engine/Source/Runtime/Engine/Classes/PhysicsEngine/BodyInstance.h
+++ b/Engine/Source/Runtime/Engine/Classes/PhysicsEngine/BodyInstance.h
@@ -61,7 +61,7 @@ private:
 
 	/** Custom Channels for Responses */
 	UPROPERTY()
-	TArray<FResponseChannel>			ResponseArray;
+	TArray<FResponseChannel> ResponseArray;
 
 	friend struct FBodyInstance;
 };
@@ -250,9 +250,9 @@ public:
 	void LoadProfileData(bool bVerifyProfile);
 
 #if WITH_PHYSX
-typedef physx::PxAggregate* PhysXAggregateType;
+	typedef physx::PxAggregate* PhysXAggregateType;
 #elif WITH_BOX2D
-typedef void* PhysXAggregateType;
+	typedef void* PhysXAggregateType;
 #endif
 
 #if WITH_PHYSX || WITH_BOX2D
@@ -276,8 +276,11 @@ typedef void* PhysXAggregateType;
 	 */
 	bool UpdateBodyScale(const FVector& InScale3D);
 
+	/** Returns the center of mass of this body (in world space) */
 	FVector GetCOMPosition() const;
-	void DrawCOMPosition( class FPrimitiveDrawInterface* PDI, float COMRenderSize, const FColor& COMRenderColor );
+
+	/** Draws the center of mass as a wire star */
+	void DrawCOMPosition(class FPrimitiveDrawInterface* PDI, float COMRenderSize, const FColor& COMRenderColor);
 
 	/** Utility for copying properties from one BodyInstance to another. */
 	void CopyBodyInstancePropertiesFrom(const FBodyInstance* FromInst);
@@ -311,6 +314,9 @@ typedef void* PhysXAggregateType;
 	TArray<physx::PxShape*> GetAllShapes(int32& OutNumSyncShapes) const;
 #endif	//WITH_PHYSX
 
+	/** Returns true if the body is not static */
+	bool IsDynamic() const;
+
 	/** Returns the body's mass */
 	float GetBodyMass() const;
 	/** Return bounds of physics representation */
@@ -326,7 +332,7 @@ typedef void* PhysXAggregateType;
 	/** Should Simulate Physics **/
 	bool ShouldInstanceSimulatingPhysics(bool bIgnoreOwner=false);
 	/** Returns whether this body is awake */
-	bool IsInstanceAwake();
+	bool IsInstanceAwake() const;
 	/** Wake this body */
 	void WakeInstance();
 	/** Force this body to sleep */
@@ -358,10 +364,10 @@ typedef void* PhysXAggregateType;
 	/** Get current transform in world space from physics body. */
 	FTransform GetUnrealWorldTransform() const;
 
-	/** 
-		*	Move the physics body to a new pose. 
-		*	@param	bTeleport	If true, no velocity is inferred on the kinematic body from this movement, but it moves right away.
-		*/
+	/**
+	 *	Move the physics body to a new pose.
+	 *	@param	bTeleport	If true, no velocity is inferred on the kinematic body from this movement, but it moves right away.
+	 */
 	void SetBodyTransform(const FTransform& NewTransform, bool bTeleport);
 
 	/** Get current velocity in world space from physics body. */
diff --git a/Engine/Source/Runtime/Engine/Private/PhysicsEngine/BodyInstance.cpp b/Engine/Source/Runtime/Engine/Private/PhysicsEngine/BodyInstance.cpp
--- a/Engine/Source/Runtime/Engine/Private/PhysicsEngine/BodyInstance.cpp
+++ b/Engine/Source/Runtime/Engine/Private/PhysicsEngine/BodyInstance.cpp
@@ -187,7 +187,7 @@ void FCollisionResponse::UpdateResponseContainerFromArray()
 
 FBodyInstance::FBodyInstance()
 : InstanceBodyIndex(INDEX_NONE)
-, Scale3D(1.0)
+, Scale3D(1.0f)
 , SceneIndexSync(0)
 , SceneIndexAsync(0)
 , bEnableCollision_DEPRECATED(true)
@@ -311,8 +311,8 @@ TArray<PxShape*> FBodyInstance::GetAllShapes(int32& OutNumSyncShapes) const
 
 void FBodyInstance::UpdatePhysicalMaterials()
 {
+	//@TODO: BOX2D: Implement UpdatePhysicalMaterials
 #if WITH_PHYSX
-
 	UPhysicalMaterial* SimplePhysMat = GetSimplePhysicalMaterial();
 	check(SimplePhysMat != NULL);
 	PxMaterial* PSimpleMat = SimplePhysMat->GetPhysXMaterial();
@@ -430,7 +430,7 @@ ECollisionChannel FBodyInstance::GetObjectType() const
 
 void FBodyInstance::SetCollisionProfileName(FName InCollisionProfileName)
 {
-	if ( CollisionProfileName!=InCollisionProfileName )
+	if (CollisionProfileName != InCollisionProfileName)
 	{
 		CollisionProfileName = InCollisionProfileName;
 		// now Load ProfileData
@@ -467,7 +467,7 @@ ECollisionEnabled::Type FBodyInstance::GetCollisionEnabled() const
 {
 	// Check actor override
 	AActor* Owner = OwnerComponent.IsValid() ? OwnerComponent.Get()->GetOwner() : NULL;
-	if(Owner != NULL && !Owner->GetActorEnableCollision())
+	if ((Owner != NULL) && !Owner->GetActorEnableCollision())
 	{
 		return ECollisionEnabled::NoCollision;
 	}
@@ -481,9 +481,8 @@ ECollisionEnabled::Type FBodyInstance::GetCollisionEnabled() const
 /** Update the filter data on the physics shapes, based on the owning component flags. */
 void FBodyInstance::UpdatePhysicsFilterData()
 {
-#if WITH_PHYSX
 	// Do nothing if no physics actor
-	if(GetPxRigidActor() == NULL)
+	if (!IsValidBodyInstance())
 	{
 		return;
 	}
@@ -506,8 +505,7 @@ void FBodyInstance::UpdatePhysicsFilterData()
 
 	// Get skelmeshcomp ID
 	uint32 SkelMeshCompID = 0;
-	USkeletalMeshComponent* SkelMeshComp = Cast<USkeletalMeshComponent>(OwnerComponent.Get());
-	if(SkelMeshComp != NULL)
+	if (USkeletalMeshComponent* SkelMeshComp = Cast<USkeletalMeshComponent>(OwnerComponent.Get()))
 	{
 		SkelMeshCompID = SkelMeshComp->GetUniqueID();
 
@@ -531,7 +529,7 @@ void FBodyInstance::UpdatePhysicsFilterData()
 
 #if WITH_EDITOR
 	// if no collision, but if world wants to enable trace collision for components, allow it
-	if(UseCollisionEnabled == ECollisionEnabled::NoCollision && Owner && Owner->IsA(AVolume::StaticClass())==false )
+	if ((UseCollisionEnabled == ECollisionEnabled::NoCollision) && Owner && (Owner->IsA(AVolume::StaticClass()) == false))
 	{
 		UWorld* World = Owner->GetWorld();
 		UPrimitiveComponent* PrimComp = OwnerComponent.Get();
@@ -547,161 +545,313 @@ void FBodyInstance::UpdatePhysicsFilterData()
 	}
 #endif
 
-	// Is the target a static actor
-	bool bDestStatic = (GetPxRigidActor()->isRigidStatic() != NULL);
-	bool bUseComplexAsSimple = (BodySetup.Get()->CollisionTraceFlag==CTF_UseComplexAsSimple);
-	bool bUseSimpleAsComplex = (BodySetup.Get()->CollisionTraceFlag==CTF_UseSimpleAsComplex);
+	const bool bUseComplexAsSimple = (BodySetup.Get()->CollisionTraceFlag == CTF_UseComplexAsSimple);
+	const bool bUseSimpleAsComplex = (BodySetup.Get()->CollisionTraceFlag == CTF_UseSimpleAsComplex);
 
-	// Create the filterdata structs
-	PxFilterData PQueryFilterData, PSimFilterData;
-	PxFilterData PSimpleQueryData, PComplexQueryData; 
-	if (UseCollisionEnabled != ECollisionEnabled::NoCollision)
+#if WITH_PHYSX
+	if (PxRigidActor* PActor = GetPxRigidActor())
 	{
-		CreateShapeFilterData(ObjectType, OwnerID, UseResponse, SkelMeshCompID, InstanceBodyIndex, PQueryFilterData, PSimFilterData, bUseCCD && !bPhysicsStatic, bUseNotifyRBCollision, bPhysicsStatic);
-		PSimpleQueryData = PQueryFilterData;
-		PComplexQueryData = PQueryFilterData;	
-		
-		// Build filterdata variations for complex and simple
-		PSimpleQueryData.word3 |= EPDF_SimpleCollision;
-		if(bUseSimpleAsComplex)
+		// Create the filterdata structs
+		PxFilterData PQueryFilterData, PSimFilterData;
+		PxFilterData PSimpleQueryData, PComplexQueryData;
+		if (UseCollisionEnabled != ECollisionEnabled::NoCollision)
 		{
-			PSimpleQueryData.word3 |= EPDF_ComplexCollision;
-		}
+			CreateShapeFilterData(ObjectType, OwnerID, UseResponse, SkelMeshCompID, InstanceBodyIndex, PQueryFilterData, PSimFilterData, bUseCCD && !bPhysicsStatic, bUseNotifyRBCollision, bPhysicsStatic);
+			PSimpleQueryData = PQueryFilterData;
+			PComplexQueryData = PQueryFilterData;
 
-		PComplexQueryData.word3 |= EPDF_ComplexCollision;
-		if(bUseComplexAsSimple)
-		{
-			PComplexQueryData.word3 |= EPDF_SimpleCollision;
-		}		
-	}
+			// Build filterdata variations for complex and simple
+			PSimpleQueryData.word3 |= EPDF_SimpleCollision;
+			if (bUseSimpleAsComplex)
+			{
+				PSimpleQueryData.word3 |= EPDF_ComplexCollision;
+			}
 
-	// Iterate over all shapes and assign filterdata
-	int32 NumSyncShapes = 0;
-	TArray<PxShape*> AllShapes = GetAllShapes(NumSyncShapes);
+			PComplexQueryData.word3 |= EPDF_ComplexCollision;
+			if (bUseComplexAsSimple)
+			{
+				PComplexQueryData.word3 |= EPDF_SimpleCollision;
+			}
+		}
 
-	bool bUpdateMassProperties = false;
+		// Iterate over all shapes and assign filterdata
+		int32 NumSyncShapes = 0;
+		TArray<PxShape*> AllShapes = GetAllShapes(NumSyncShapes);
 
-	// Only perform scene queries in the synchronous scene for static shapes
-	const int32 SceneQueryShapeNumMax = bDestStatic ? NumSyncShapes : AllShapes.Num();
+		// Is the target a static actor
+		const bool bDestStatic = PActor->isRigidStatic() != NULL;
 
-	PxScene* AsyncScene = (RigidActorAsync != NULL) ? RigidActorAsync->getScene() : NULL;
-	SCENE_LOCK_WRITE(AsyncScene);
+		// Only perform scene queries in the synchronous scene for static shapes
+		const int32 SceneQueryShapeNumMax = bDestStatic ? NumSyncShapes : AllShapes.Num();
 
-	for(int32 ShapeIdx=0; ShapeIdx<AllShapes.Num(); ShapeIdx++)
-	{
-		PxShape* PShape = AllShapes[ShapeIdx];
-		PShape->setSimulationFilterData(PSimFilterData);
+		PxScene* AsyncScene = (RigidActorAsync != NULL) ? RigidActorAsync->getScene() : NULL;
+		SCENE_LOCK_WRITE(AsyncScene);
 
-		// If query collision is enabled..
-		if(UseCollisionEnabled == ECollisionEnabled::QueryAndPhysics || UseCollisionEnabled == ECollisionEnabled::QueryOnly)
+		bool bUpdateMassProperties = false;
+		for (int32 ShapeIdx = 0; ShapeIdx < AllShapes.Num(); ShapeIdx++)
 		{
-			// Only perform scene queries in the synchronous scene for static shapes
-			if(bPhysicsStatic)
-			{
-				bool bIsSyncShape = (ShapeIdx < NumSyncShapes);
-				PShape->setFlag(PxShapeFlag::eSCENE_QUERY_SHAPE, bIsSyncShape);
-			}
-			// If non-static, always enable scene queries
-			else
-			{
-				PShape->setFlag(PxShapeFlag::eSCENE_QUERY_SHAPE, true);
-			}
-
-			// See if we want physics collision
-			bool bSimCollision = (UseCollisionEnabled == ECollisionEnabled::QueryAndPhysics);
+			PxShape* PShape = AllShapes[ShapeIdx];
+			PShape->setSimulationFilterData(PSimFilterData);
 
-			// Triangle mesh is 'complex' geom
-			if(PShape->getGeometryType() == PxGeometryType::eTRIANGLEMESH)
+			// If query collision is enabled..
+			if (UseCollisionEnabled == ECollisionEnabled::QueryAndPhysics || UseCollisionEnabled == ECollisionEnabled::QueryOnly)
 			{
-				PShape->setQueryFilterData(PComplexQueryData);
-
-				// on dynamic objects and objects which don't use complex as simple, tri mesh not used for sim
-				if(!bSimCollision || !bDestStatic || !bUseComplexAsSimple)
+				// Only perform scene queries in the synchronous scene for static shapes
+				if (bPhysicsStatic)
 				{
-					PShape->setFlag(PxShapeFlag::eSIMULATION_SHAPE, false); 
+					bool bIsSyncShape = (ShapeIdx < NumSyncShapes);
+					PShape->setFlag(PxShapeFlag::eSCENE_QUERY_SHAPE, bIsSyncShape);
 				}
+				// If non-static, always enable scene queries
 				else
 				{
-					PShape->setFlag(PxShapeFlag::eSIMULATION_SHAPE, true); 
+					PShape->setFlag(PxShapeFlag::eSCENE_QUERY_SHAPE, true);
 				}
 
-				if(OwnerComponent == NULL || !OwnerComponent->IsA(UModelComponent::StaticClass()))
-				{
-					PShape->setFlag(PxShapeFlag::eVISUALIZATION, false); // dont draw the tri mesh, we can see it anyway, and its slow
-				}
-			}
-			// Everything else is 'simple'
-			else
-			{
-				PShape->setQueryFilterData(PSimpleQueryData);
+				// See if we want physics collision
+				bool bSimCollision = (UseCollisionEnabled == ECollisionEnabled::QueryAndPhysics);
 
-				// See if we currently have sim collision
-				bool bCurrentSimCollision = (PShape->getFlags() & PxShapeFlag::eSIMULATION_SHAPE);
-				// Enable sim collision
-				if (bSimCollision && !bCurrentSimCollision)
-				{
-					bUpdateMassProperties = true;
-					PShape->setFlag(PxShapeFlag::eSIMULATION_SHAPE, true); 
-				}
-				// Disable sim collision
-				else if(!bSimCollision && bCurrentSimCollision)
+				// Triangle mesh is 'complex' geom
+				if (PShape->getGeometryType() == PxGeometryType::eTRIANGLEMESH)
 				{
-					bUpdateMassProperties = true;
-					PShape->setFlag(PxShapeFlag::eSIMULATION_SHAPE, false);
-				}
-							
-				// enable swept bounds for CCD for this shape
-				PxRigidBody* PBody = GetPxRigidActor()->is<PxRigidBody>();
-				if(bSimCollision && !bPhysicsStatic && bUseCCD && PBody)
-				{
-					PBody->setRigidBodyFlag(PxRigidBodyFlag::eENABLE_CCD, true);
+					PShape->setQueryFilterData(PComplexQueryData);
+
+					// on dynamic objects and objects which don't use complex as simple, tri mesh not used for sim
+					if (!bSimCollision || !bDestStatic || !bUseComplexAsSimple)
+					{
+						PShape->setFlag(PxShapeFlag::eSIMULATION_SHAPE, false);
+					}
+					else
+					{
+						PShape->setFlag(PxShapeFlag::eSIMULATION_SHAPE, true);
+					}
+
+					if (OwnerComponent == NULL || !OwnerComponent->IsA(UModelComponent::StaticClass()))
+					{
+						PShape->setFlag(PxShapeFlag::eVISUALIZATION, false); // dont draw the tri mesh, we can see it anyway, and its slow
+					}
 				}
-				else if(PBody)
+				// Everything else is 'simple'
+				else
 				{
-					
-					PBody->setRigidBodyFlag(PxRigidBodyFlag::eENABLE_CCD, false);
+					PShape->setQueryFilterData(PSimpleQueryData);
+
+					// See if we currently have sim collision
+					bool bCurrentSimCollision = (PShape->getFlags() & PxShapeFlag::eSIMULATION_SHAPE);
+					// Enable sim collision
+					if (bSimCollision && !bCurrentSimCollision)
+					{
+						bUpdateMassProperties = true;
+						PShape->setFlag(PxShapeFlag::eSIMULATION_SHAPE, true);
+					}
+					// Disable sim collision
+					else if (!bSimCollision && bCurrentSimCollision)
+					{
+						bUpdateMassProperties = true;
+						PShape->setFlag(PxShapeFlag::eSIMULATION_SHAPE, false);
+					}
+
+					// enable swept bounds for CCD for this shape
+					PxRigidBody* PBody = GetPxRigidActor()->is<PxRigidBody>();
+					if (bSimCollision && !bPhysicsStatic && bUseCCD && PBody)
+					{
+						PBody->setRigidBodyFlag(PxRigidBodyFlag::eENABLE_CCD, true);
+					}
+					else if (PBody)
+					{
+
+						PBody->setRigidBodyFlag(PxRigidBodyFlag::eENABLE_CCD, false);
+					}
 				}
 			}
+			// No collision enabled
+			else
+			{
+				PShape->setFlag(PxShapeFlag::eSIMULATION_SHAPE, false);
+				PShape->setFlag(PxShapeFlag::eSCENE_QUERY_SHAPE, false);
+			}
 		}
-		// No collision enabled
-		else
+
+		if (bUpdateMassProperties)
 		{
-			PShape->setFlag(PxShapeFlag::eSIMULATION_SHAPE, false); 
-			PShape->setFlag(PxShapeFlag::eSCENE_QUERY_SHAPE, false); 	
+			UpdateMassProperties();
 		}
+
+		SCENE_UNLOCK_WRITE(AsyncScene);
 	}
+#endif
+
+#if WITH_BOX2D
 
-	if (bUpdateMassProperties)
+	if (BodyInstancePtr != nullptr)
 	{
-		UpdateMassProperties();
-	}
+		//@TODO: BOX2D: Implement UpdatePhysicsFilterData
 
-	SCENE_UNLOCK_WRITE(AsyncScene);
-#endif // WITH_PHYSX
+		bool bUpdateMassProperties = false;
+
+		if (bUpdateMassProperties)
+		{
+			UpdateMassProperties();
+		}
+	}
+#endif
 }
 
 #if WITH_PHYSX || WITH_BOX2D
 void FBodyInstance::InitBody(UBodySetup* Setup, const FTransform& Transform, UPrimitiveComponent* PrimComp, FPhysScene* InRBScene, PxAggregate* InAggregate)
 {
+	check(Setup);
+	AActor* Owner = PrimComp ? PrimComp->GetOwner() : NULL;
+
+	// Make the debug name for this geometry...
+	FString DebugName(TEXT(""));
+#if (WITH_EDITORONLY_DATA || UE_BUILD_DEBUG || LOOKING_FOR_PERF_ISSUES) && !(UE_BUILD_SHIPPING || UE_BUILD_TEST) && !NO_LOGGING
+	UStaticMeshComponent* SMComp = Cast<UStaticMeshComponent>(PrimComp);
+	if (Owner)
+	{
+		DebugName += FString::Printf(TEXT("Actor: %s "), *Owner->GetPathName());
+	}
+
+	if (PrimComp)
+	{
+		DebugName += FString::Printf(TEXT("Component: %s "), *PrimComp->GetName());
+	}
+
+	if (SMComp)
+	{
+		DebugName += FString::Printf(TEXT("StaticMesh: %s"), *SMComp->StaticMesh->GetPathName());
+	}
+
+	if (Setup->BoneName != NAME_None)
+	{
+		DebugName += FString::Printf(TEXT("Bone: %s "), *Setup->BoneName.ToString());
+	}
+
+#if WITH_PHYSX
+	// Convert to char* for PhysX
+	CharDebugName = MakeShareable(new TArray<ANSICHAR>(StringToArray<ANSICHAR>(*DebugName, DebugName.Len() + 1)));
+#endif
+#endif
+
+	if (Transform.GetScale3D().IsNearlyZero())
+	{
+		UE_LOG(LogPhysics, Warning, TEXT("FBodyInstance::InitBody : Scale3D is (nearly) zero: %s"), *DebugName);
+		return;
+	}
+
+	// Check we support mirroring/non-mirroring
+	const float TransformDet = Transform.GetDeterminant();
+	if (TransformDet < 0.f && !Setup->bGenerateMirroredCollision)
+	{
+		UE_LOG(LogPhysics, Warning, TEXT("FBodyInstance::InitBody : Body is mirrored but bGenerateMirroredCollision == false: %s"), *DebugName);
+		return;
+	}
+
+	if (TransformDet > 0.f && !Setup->bGenerateNonMirroredCollision)
+	{
+		UE_LOG(LogPhysics, Warning, TEXT("FBodyInstance::InitBody : Body is not mirrored but bGenerateNonMirroredCollision == false: %s"), *DebugName);
+		return;
+	}
+
+	if (Transform.ContainsNaN())
+	{
+		UE_LOG(LogPhysics, Warning, TEXT("InitBody: Bad transform - %s %s\n%s"), *DebugName, *Setup->BoneName.ToString(), *Transform.ToString());
+		return;
+	}
+
+	// remember my owning component
+	OwnerComponent = PrimComp;
+	BodySetup = Setup;
+	Scale3D = Transform.GetScale3D();
+
+	// whenever update filter, this check will trigger, it seems it's best to know when we initialize body instead. 
+	check(BodySetup.IsValid());
+
+	USkeletalMeshComponent* SkelMeshComp = Cast<USkeletalMeshComponent>(OwnerComponent.Get());
+	if (SkelMeshComp != nullptr)
+	{
+		// if not dedicated and if set up is done to use physics, enable flag on
+		if ((BodySetup->PhysicsType == PhysType_Simulated) || (BodySetup->PhysicsType == PhysType_Default))
+		{
+			bool bEnableSimulation = (SkelMeshComp && IsRunningDedicatedServer()) ? SkelMeshComp->bEnablePhysicsOnDedicatedServer : true;
+			bEnableSimulation &= ((BodySetup->PhysicsType == PhysType_Simulated) || (SkelMeshComp->BodyInstance.bSimulatePhysics));	//if unfixed enable. If default look at parent
+			if (bEnableSimulation)
+			{
+				// set simulate to true if using physics
+				bSimulatePhysics = true;
+				if (BodySetup->PhysicsType == PhysType_Simulated)
+				{
+					PhysicsBlendWeight = 1.f;
+				}
+
+			}
+			else
+			{
+				bSimulatePhysics = false;
+				if (BodySetup->PhysicsType == PhysType_Simulated)
+				{
+					PhysicsBlendWeight = 0.f;
+				}
+			}
+		}
+	}
+
+	// See if we are 'static'
+	const bool bPhysicsStatic = (OwnerComponent == NULL) || (OwnerComponent->Mobility != EComponentMobility::Movable);
+
+	// In skeletal case, we need both our bone and skelcomponent flag to be true.
+	// This might be 'and'ing us with ourself, but thats fine.
+	const bool bUseSimulate = ShouldInstanceSimulatingPhysics();
+
+	// Determine if we should start out awake or sleeping
+	bool bShouldStartAwake = bStartAwake;
+	FVector InitialLinVel(EForceInit::ForceInitToZero);
+
+	if (SkelMeshComp != NULL)
+	{
+		bShouldStartAwake = bShouldStartAwake && SkelMeshComp->BodyInstance.bStartAwake;
+	}
+
+	if (Owner != NULL)
+	{
+		InitialLinVel = Owner->GetVelocity();
+
+		if (InitialLinVel.Size() > KINDA_SMALL_NUMBER)
+		{
+			bShouldStartAwake = true;
+		}
+	}
+
 #if WITH_BOX2D
 	if (UBodySetup2D* BodySetup2D = Cast<UBodySetup2D>(Setup))
 	{
-		BodySetup = Setup;
-		Scale3D = Transform.GetScale3D();
-
-		if (b2World* World = FPhysicsIntegration2D::FindAssociatedWorld(PrimComp->GetWorld()))
+		if (b2World* BoxWorld = FPhysicsIntegration2D::FindAssociatedWorld(PrimComp->GetWorld()))
 		{
-			const FVector Scale3D = Transform.GetScale3D();
 			const b2Vec2 Scale2D = FPhysicsIntegration2D::ConvertUnrealVectorToBox(Scale3D);
 
 			if (Setup != nullptr)
 			{
 				OwnerComponent = PrimComp;
 
+				// Create the body definition
 				b2BodyDef BodyDefinition;
-				BodyDefinition.type = bSimulatePhysics ? b2_dynamicBody : b2_kinematicBody;
+				if (bPhysicsStatic)
+				{
+					BodyDefinition.type = b2_staticBody;
+				}
+				else
+				{
+					BodyDefinition.type = bUseSimulate ? b2_dynamicBody : b2_kinematicBody;
+				}
+				BodyDefinition.awake = bShouldStartAwake;
+
+				if (bUseSimulate)
+				{
+					BodyDefinition.linearVelocity = FPhysicsIntegration2D::ConvertUnrealVectorToBox(InitialLinVel);
+				}
 
-				BodyInstancePtr = World->CreateBody(&BodyDefinition);
+				// Create the body
+				BodyInstancePtr = BoxWorld->CreateBody(&BodyDefinition);
 				BodyInstancePtr->SetUserData(this);
 
 				// Circles
@@ -750,7 +900,7 @@ void FBodyInstance::InitBody(UBodySetup* Setup, const FTransform& Transform, UPr
 						for (int32 VertexIndex = 0; VertexIndex < Convex.VertexData.Num(); ++VertexIndex)
 						{
 							const FVector2D SourceVert = Convex.VertexData[VertexIndex];
-							new (Verts)b2Vec2(SourceVert.X * Scale2D.x, SourceVert.Y * Scale2D.y);
+							new (Verts) b2Vec2(SourceVert.X * Scale2D.x, SourceVert.Y * Scale2D.y);
 						}
 
 						b2PolygonShape ConvexPoly;
@@ -769,129 +919,82 @@ void FBodyInstance::InitBody(UBodySetup* Setup, const FTransform& Transform, UPr
 					}
 				}
 
-				SetBodyTransform(Transform, /*bTeleport=*/ true);
+				// Make sure it contained at least one shape
+				if (BodyInstancePtr->GetFixtureList() == nullptr)
+				{
+					if (DebugName.Len())
+					{
+						UE_LOG(LogPhysics, Log, TEXT("InitBody: failed - no shapes: %s"), *DebugName);
+					}
+
+					BodyInstancePtr->SetUserData(nullptr);
+					BoxWorld->DestroyBody(BodyInstancePtr);
+					BodyInstancePtr = nullptr;
+
+					//clear Owner and Setup info as well to properly clean up the BodyInstance.
+					OwnerComponent = NULL;
+					BodySetup = NULL;
+
+					return;
+				}
+				else
+				{
+					// Position the body
+					SetBodyTransform(Transform, /*bTeleport=*/ true);
+				}
 			}
 		}
 		else
 		{
-			//@TODO: Error out?
+			//@TODO: BOX2D: Error out if the world scene mapping wasn't found?
 		}
 
-		return;
-	}
-#endif
-
-
-#if WITH_PHYSX
-	PhysxUserData = FPhysxUserData(this);
-
-	check(Setup);
-	AActor* Owner = PrimComp ? PrimComp->GetOwner() : NULL;
+		// Apply correct physical materials to shape we created.
+		UpdatePhysicalMaterials();
 
-	// If there is already a body instanced, or there is no scene to create it into, do nothing.
-	if (GetPxRigidActor() != NULL || !InRBScene)
-	{
-		return;
-	}
+		// Set the filter data on the shapes (call this after setting BodyData because it uses that pointer)
+		UpdatePhysicsFilterData();
 
-	// Make the debug name for this geometry...
-	FString DebugName(TEXT(""));
-#if (WITH_EDITORONLY_DATA || UE_BUILD_DEBUG || LOOKING_FOR_PERF_ISSUES) && !(UE_BUILD_SHIPPING || UE_BUILD_TEST) && !NO_LOGGING
-	UStaticMeshComponent* SMComp = Cast<UStaticMeshComponent>(PrimComp);
-	if(Owner)
-	{
-		DebugName += FString::Printf( TEXT("Actor: %s "), *Owner->GetPathName() );
-	}
+		if (!bPhysicsStatic)
+		{
+			// Compute mass (call this after setting BodyData because it uses that pointer)
+			UpdateMassProperties();
 
-	if(PrimComp)
-	{
-		DebugName += FString::Printf( TEXT("Component: %s "), *PrimComp->GetName() );
-	}
+			// Update damping
+			UpdateDampingProperties();
 
-	if(SMComp)
-	{
-		DebugName += FString::Printf( TEXT("StaticMesh: %s"), *SMComp->StaticMesh->GetPathName() );
-	}
+			SetMaxAngularVelocity(MaxAngularVelocity, false);
 
-	if(Setup->BoneName != NAME_None)
-	{
-		DebugName += FString::Printf( TEXT("Bone: %s "), *Setup->BoneName.ToString() );
-	}
-	// Convert to char* for PhysX
-	CharDebugName = MakeShareable(new TArray<ANSICHAR>(StringToArray<ANSICHAR>(*DebugName, DebugName.Len() + 1)));
+			//@TODO: BOX2D: Determine if sleep threshold and solver settings can be configured per-body or not
+#if 0
+			// Set the parameters for determining when to put the object to sleep.
+			float SleepEnergyThresh = PNewDynamic->getSleepThreshold();
+			if (SleepFamily == SF_Sensitive)
+			{
+				SleepEnergyThresh /= 20.f;
+			}
+			PNewDynamic->setSleepThreshold(SleepEnergyThresh);
+			// set solver iteration count 
+			int32 PositionIterCount = FMath::Clamp(PositionSolverIterationCount, 1, 255);
+			int32 VelocityIterCount = FMath::Clamp(VelocitySolverIterationCount, 1, 255);
+			PNewDynamic->setSolverIterationCounts(PositionIterCount, VelocityIterCount);
 #endif
+		}
 
-	if(Transform.GetScale3D().IsNearlyZero())
-	{
-		UE_LOG(LogPhysics, Warning, TEXT("FBodyInstance::InitBody : Scale3D is (nearly) zero: %s"), *DebugName);
-		return;
-	}
-
-	// Check we support mirroring/non-mirroring
-	const float TransformDet = Transform.GetDeterminant();
-	if(TransformDet < 0.f && !Setup->bGenerateMirroredCollision)
-	{
-		UE_LOG(LogPhysics, Warning, TEXT("FBodyInstance::InitBody : Body is mirrored but bGenerateMirroredCollision == false: %s"), *DebugName);
+		// Skip over the rest of the PhysX-specific setup in this function, since it was a 2D object
 		return;
 	}
+#endif
 
-	if(TransformDet > 0.f && !Setup->bGenerateNonMirroredCollision)
-	{
-		UE_LOG(LogPhysics, Warning, TEXT("FBodyInstance::InitBody : Body is not mirrored but bGenerateNonMirroredCollision == false: %s"), *DebugName);
-		return;
-	}
+#if WITH_PHYSX
+	PhysxUserData = FPhysxUserData(this);
 
-	if(Transform.ContainsNaN())
+	// If there is already a body instanced, or there is no scene to create it into, do nothing.
+	if (GetPxRigidActor() != NULL || !InRBScene)
 	{
-		UE_LOG(LogPhysics, Warning, TEXT("InitBody: Bad transform - %s %s\n%s"), *DebugName, *Setup->BoneName.ToString(), *Transform.ToString());
 		return;
 	}
 
-	// remember my owning component
-	OwnerComponent	= PrimComp;
-	BodySetup		= Setup;
-
-	// whenever update filter, this check will trigger, it seems it's best to know when we initialize body instead. 
-	check(BodySetup.IsValid());
-
-	USkeletalMeshComponent* SkelMeshComp = NULL;
-	if (OwnerComponent.IsValid() && OwnerComponent.Get()->IsA(USkeletalMeshComponent::StaticClass()))
-	{
-		SkelMeshComp = Cast<USkeletalMeshComponent>(OwnerComponent.Get());
-
-		// if not dedicated and if set up is done to use physics, enable flag on
-		if (BodySetup->PhysicsType == PhysType_Simulated || BodySetup->PhysicsType == PhysType_Default)
-		{
-			bool bEnableSimulation = (SkelMeshComp && IsRunningDedicatedServer())? SkelMeshComp->bEnablePhysicsOnDedicatedServer: true;
-			bEnableSimulation &= ((BodySetup->PhysicsType == PhysType_Simulated) || (SkelMeshComp->BodyInstance.bSimulatePhysics));	//if unfixed enable. If default look at parent
-			if ( bEnableSimulation )
-			{
-				// set simulate to true if using physics
-				bSimulatePhysics = true;
-				if (BodySetup->PhysicsType == PhysType_Simulated)
-				{
-					PhysicsBlendWeight = 1.f;
-				}
-				
-			}
-	 		else
-	 		{
-	 			bSimulatePhysics = false;
-				if (BodySetup->PhysicsType == PhysType_Simulated)
-				{
-					PhysicsBlendWeight = 0.f;
-				}
-	 		}
-		}
-	}
-
-	// See if we are 'static'
-	const bool bPhysicsStatic = (OwnerComponent == NULL) || (OwnerComponent->Mobility != EComponentMobility::Movable);
-
-	// In skeletal case, we need both our bone and skelcomponent flag to be true.
-	// This might be 'and'ing us with ourself, but thats fine.
-	const bool bUseSimulate = ShouldInstanceSimulatingPhysics();
-
 	PxScene* PSceneForNewDynamic = NULL;	// TBD
 	PxScene* PSceneSync = InRBScene->GetPhysXScene(PST_Sync);
 	PxScene* PSceneAsync = InRBScene->HasAsyncScene() ? InRBScene->GetPhysXScene(PST_Async) : NULL;
@@ -947,11 +1050,6 @@ void FBodyInstance::InitBody(UBodySetup* Setup, const FTransform& Transform, UPr
 		PNewDynamic->setRigidDynamicFlag(PxRigidDynamicFlag::eUSE_KINEMATIC_TARGET_FOR_SCENE_QUERIES, true);
 	}
 
-
-	// save Scale3D 
-	Scale3D = Transform.GetScale3D();
-
-
 	// Copy geom from template and scale
 
 	// Sync:
@@ -994,8 +1092,8 @@ void FBodyInstance::InitBody(UBodySetup* Setup, const FTransform& Transform, UPr
 		}
 
 		//clear Owner and Setup info as well to properly clean up the BodyInstance.
-		OwnerComponent	= NULL;
-		BodySetup		= NULL;
+		OwnerComponent = NULL;
+		BodySetup = NULL;
 
 		return;
 	}
@@ -1049,13 +1147,7 @@ void FBodyInstance::InitBody(UBodySetup* Setup, const FTransform& Transform, UPr
 
 
 	// Set initial velocities
-	FVector InitialLinVel(0.f);
-	if(Owner != NULL)
-	{	
-		InitialLinVel = Owner->GetVelocity();
-	}
-
-	if(PNewDynamic != NULL)
+	if (PNewDynamic != NULL)
 	{
 		// Compute mass (call this after setting BodyData because it uses that pointer)
 		UpdateMassProperties();
@@ -1082,18 +1174,11 @@ void FBodyInstance::InitBody(UBodySetup* Setup, const FTransform& Transform, UPr
 		int32 VelocityIterCount = FMath::Clamp(VelocitySolverIterationCount, 1, 255);
 		PNewDynamic->setSolverIterationCounts(PositionIterCount, VelocityIterCount);
 
-		bool bShouldStartAwake = bStartAwake;
-
-		if (SkelMeshComp != NULL)
-		{
-			bShouldStartAwake = bShouldStartAwake && SkelMeshComp->BodyInstance.bStartAwake;
-		}
-
 		// wakeUp and putToSleep will issue warnings on kinematic actors
-		if(IsRigidDynamicNonKinematic(PNewDynamic))
+		if (IsRigidDynamicNonKinematic(PNewDynamic))
 		{
 		    // Sleep/wake up as appropriate
-		    if(bShouldStartAwake || (Owner && Owner->GetVelocity().Size() > KINDA_SMALL_NUMBER))
+		    if (bShouldStartAwake)
 		    {
 			    // Wake up bodies that are part of a moving actor.
 			    PNewDynamic->wakeUp();
@@ -1205,6 +1290,9 @@ void FBodyInstance::TermBody()
 void FBodyInstance::Weld(FBodyInstance* TheirBody, const FTransform& RelativeTM)
 {
 	check(TheirBody);
+
+	//@TODO: BOX2D: Implement Weld
+
 #if WITH_PHYSX
 	
 	//UBodySetup* LocalSpaceSetup = TheirBody->BodySetup->CreateSpaceCopy(RelativeTM);
@@ -1337,8 +1425,8 @@ void ComputeScalingVectors(EScaleMode::Type ScaleMode, const FVector& NewScale3D
 bool FBodyInstance::UpdateBodyScale(const FVector& InScale3D)
 {
 	FVector InScale3DAdjusted = InScale3D;
-#if WITH_PHYSX
-	if (BodySetup == NULL || GetPxRigidActor() == NULL)
+
+	if (IsValidBodyInstance())
 	{
 		//UE_LOG(LogPhysics, Log, TEXT("Body hasn't been initialized. Call InitBody to initialize."));
 		return false;
@@ -1350,7 +1438,7 @@ bool FBodyInstance::UpdateBodyScale(const FVector& InScale3D)
 		return false;
 	}
 
-	bool bSuccess=false;
+	bool bSuccess = false;
 
 #if !(UE_BUILD_SHIPPING || UE_BUILD_TEST)
 	ensure ( !Scale3D.ContainsNaN() && !InScale3D.ContainsNaN() );
@@ -1368,12 +1456,14 @@ bool FBodyInstance::UpdateBodyScale(const FVector& InScale3D)
 	OldScale3D.Y = OldScale3D.Y < 0.1f ? 0.1f : OldScale3D.Y;
 	OldScale3D.Z = OldScale3D.Z < 0.1f ? 0.1f : OldScale3D.Z;
 
+	// Determine the scaling mode
+	EScaleMode::Type ScaleMode = EScaleMode::Free;
+
+#if WITH_PHYSX
 	//Get all shapes
 	int32 TotalShapeCount = 0;
 	TArray<PxShape *> PShapes = GetAllShapes(TotalShapeCount);
 
-	EScaleMode::Type ScaleMode = EScaleMode::Free;
-
 	for (int32 ShapeIdx = 0; ShapeIdx < PShapes.Num(); ++ShapeIdx)
 	{
 		PxShape* PShape = PShapes[ShapeIdx];
@@ -1389,13 +1479,21 @@ bool FBodyInstance::UpdateBodyScale(const FVector& InScale3D)
 			ScaleMode = EScaleMode::LockedXY;
 		}
 	}
+#endif
+#if WITH_BOX2D
+	if (BodyInstancePtr)
+	{
+		//@TODO: BOX2D: UpdateBodyScale is not implemented yet
+	}
+#endif
 
 	FVector RelativeScale3D;
 	FVector RelativeScale3DAbs;
 	FVector NewScale3D;
 	ComputeScalingVectors(ScaleMode, InScale3DAdjusted, OldScale3D, RelativeScale3D, RelativeScale3DAbs, NewScale3D);
 
-
+	// Apply scaling
+#if WITH_PHYSX
 	//we need to allocate all of these here because PhysX insists on using the stack. This is wasteful, but reduces a lot of code duplication
 	PxSphereGeometry PSphereGeom;
 	PxBoxGeometry PBoxGeom;
@@ -1568,9 +1666,17 @@ bool FBodyInstance::UpdateBodyScale(const FVector& InScale3D)
 			FMessageLog("PIE").Warning(FText::Format(LOCTEXT("PhysicsInvalidScale", "Scale ''{0}'' is not valid on object '{1}'."), FText::FromString(InScale3DAdjusted.ToString()), FText::FromString(GetBodyDebugName())));
 		}
 	}
+#endif
+
+#if WITH_BOX2D
+	if (BodyInstancePtr)
+	{
+		//@TODO: BOX2D: UpdateBodyScale is not implemented yet
+	}
+#endif
 
 	// if success, overwrite old Scale3D, otherwise, just don't do it. It will have invalid scale next time
-	if ( bSuccess )
+	if (bSuccess)
 	{
 		Scale3D = NewScale3D;
 
@@ -1582,9 +1688,6 @@ bool FBodyInstance::UpdateBodyScale(const FVector& InScale3D)
 	}
 
 	return bSuccess;
-#else
-	return false;
-#endif
 }
 
 void FBodyInstance::UpdateInstanceSimulatePhysics(bool bIgnoreParent)
@@ -1625,30 +1728,45 @@ void FBodyInstance::UpdateInstanceSimulatePhysics(bool bIgnoreParent)
 	bSimulatePhysics = bUseSimulate;
 }
 
+bool FBodyInstance::IsDynamic() const
+{
+#if WITH_PHYSX
+	if (PxRigidActor const* const PRigidActor = GetPxRigidActor())
+	{
+		return PRigidActor->isRigidDynamic() != nullptr;
+	}
+#endif // WITH_PHYSX
+
+#if WITH_BOX2D
+	if (BodyInstancePtr != nullptr)
+	{
+		return BodyInstancePtr->GetType() != b2_staticBody;
+	}
+#endif // WITH_BOX2D
+
+	return false;
+}
 
 void FBodyInstance::SetInstanceSimulatePhysics(bool bSimulate, bool bMaintainPhysicsBlending, bool bIgnoreOwner)
 {
-#if WITH_PHYSX
 #if !(UE_BUILD_SHIPPING || UE_BUILD_TEST)
 	if (bSimulate && OwnerComponent.IsValid())
 	{
-		PxRigidActor const* const PRigidActor = GetPxRigidActor();
-		if(PRigidActor == NULL)
+		if (IsValidBodyInstance())
 		{
-			FMessageLog("PIE").Warning(FText::Format(LOCTEXT("SimPhysNoBody", "Trying to simulate physics on ''{0}'' but no physics body."), 
+			FMessageLog("PIE").Warning(FText::Format(LOCTEXT("SimPhysNoBody", "Trying to simulate physics on ''{0}'' but no physics body."),
 				FText::FromString(GetPathNameSafe(OwnerComponent.Get()))));
 		}
-		else if(!PRigidActor->isRigidDynamic())
+		else if (!IsDynamic())
 		{
-			FMessageLog("PIE").Warning(FText::Format(LOCTEXT("SimPhysStatic", "Trying to simulate physics on ''{0}'' but it is static."), 
+			FMessageLog("PIE").Warning(FText::Format(LOCTEXT("SimPhysStatic", "Trying to simulate physics on ''{0}'' but it is static."),
 				FText::FromString(GetPathNameSafe(OwnerComponent.Get()))));
 		}
 	}
 #endif
-#endif // WITH_PHYSX
 
 	// If we are enabling simulation, and we are the root body of our component, we detach the component 
-	if(bSimulate == true && OwnerComponent.IsValid() && OwnerComponent->IsRegistered() && OwnerComponent->AttachParent != NULL && OwnerComponent->GetBodyInstance() == this)
+	if ((bSimulate == true) && OwnerComponent.IsValid() && OwnerComponent->IsRegistered() && (OwnerComponent->AttachParent != NULL) && (OwnerComponent->GetBodyInstance() == this))
 	{
 		OwnerComponent->DetachFromParent(true);
 	}
@@ -1694,27 +1812,23 @@ bool FBodyInstance::ShouldInstanceSimulatingPhysics(bool bIgnoreOwner)
 
 bool FBodyInstance::IsValidBodyInstance() const
 {
-	bool Retval = false;
-
 #if WITH_PHYSX
-	PxRigidActor* PActor = GetPxRigidActor();
-	if(PActor != NULL)
+	if (PxRigidActor* PActor = GetPxRigidActor())
 	{
-		Retval = true;
+		return true;
 	}
 #endif // WITH_PHYSX
 
 #if WITH_BOX2D
 	if (BodyInstancePtr != nullptr)
 	{
-		Retval = true;
+		return true;
 	}
 #endif
 
-	return Retval;
+	return false;
 }
 
-
 FTransform FBodyInstance::GetUnrealWorldTransform() const
 {
 #if WITH_PHYSX
@@ -1734,12 +1848,9 @@ FTransform FBodyInstance::GetUnrealWorldTransform() const
 	{
 		const b2Vec2 Pos2D = BodyInstancePtr->GetPosition();
 		const float RotationInRadians = BodyInstancePtr->GetAngle();
-		//@TODO: Should I use get transform?
-		//@TODO: What about scale?
-
+		
 		const FVector Translation3D(FPhysicsIntegration2D::ConvertBoxVectorToUnreal(Pos2D));
 		const FRotator Rotation3D(FMath::RadiansToDegrees(RotationInRadians), 0.0f, 0.0f);
-		const FVector Scale3D(1.0f, 1.0f, 1.0f);
 
 		return FTransform(Rotation3D, Translation3D, Scale3D);
 	}
@@ -1748,16 +1859,13 @@ FTransform FBodyInstance::GetUnrealWorldTransform() const
 	return FTransform::Identity;
 }
 
-#if !(UE_BUILD_SHIPPING || UE_BUILD_TEST)
-extern bool GShouldLogOutAFrameOfSetBodyTransform;
-#endif // !(UE_BUILD_SHIPPING || UE_BUILD_TEST)
-
 void FBodyInstance::SetBodyTransform(const FTransform& NewTransform, bool bTeleport)
 {
 	SCOPE_CYCLE_COUNTER(STAT_SetBodyTransform);
 
 #if !(UE_BUILD_SHIPPING || UE_BUILD_TEST)
-	if(GShouldLogOutAFrameOfSetBodyTransform == true)
+	extern bool GShouldLogOutAFrameOfSetBodyTransform;
+	if (GShouldLogOutAFrameOfSetBodyTransform == true)
 	{
 		UE_LOG(LogPhysics, Log, TEXT("SetBodyTransform: %s"), *GetBodyDebugName());
 	}
@@ -1834,7 +1942,7 @@ void FBodyInstance::SetBodyTransform(const FTransform& NewTransform, bool bTelep
 		const FVector NewLocation = NewTransform.GetLocation();
 		const b2Vec2 NewLocation2D(FPhysicsIntegration2D::ConvertUnrealVectorToBox(NewLocation));
 
-		//@TODO: What about scale?
+		//@TODO: BOX2D: SetBodyTransform: What about scale?
 		const FRotator NewRotation3D(NewTransform.GetRotation());
 		const float NewAngle = FMath::DegreesToRadians(NewRotation3D.Pitch);
 
@@ -1845,22 +1953,31 @@ void FBodyInstance::SetBodyTransform(const FTransform& NewTransform, bool bTelep
 
 FVector FBodyInstance::GetUnrealWorldVelocity() const
 {
-	FVector LinVel(0,0,0);
+	FVector LinVel(EForceInit::ForceInitToZero);
+
 #if WITH_PHYSX
-	PxRigidDynamic* PRigidDynamic = GetPxRigidDynamic();
-	if(PRigidDynamic != NULL)
+	if (PxRigidDynamic* PRigidDynamic = GetPxRigidDynamic())
 	{
 		SCOPED_SCENE_READ_LOCK(PRigidDynamic->getScene());
 		LinVel = P2UVector(PRigidDynamic->getLinearVelocity());
 	}
 #endif // WITH_PHYSX
+
+#if WITH_BOX2D
+	if (BodyInstancePtr != nullptr)
+	{
+		LinVel = FPhysicsIntegration2D::ConvertBoxVectorToUnreal(BodyInstancePtr->GetLinearVelocity());
+	}
+#endif
+
 	return LinVel;
 }
 
 /** Note: returns angular velocity in degrees per second. */
 FVector FBodyInstance::GetUnrealWorldAngularVelocity() const
 {
-	FVector AngVel(0,0,0);
+	FVector AngVel(EForceInit::ForceInitToZero);
+
 #if WITH_PHYSX
 	PxRigidDynamic* PRigidDynamic = GetPxRigidDynamic();
 	if(PRigidDynamic != NULL)
@@ -1869,16 +1986,23 @@ FVector FBodyInstance::GetUnrealWorldAngularVelocity() const
 		AngVel = FMath::RadiansToDegrees( P2UVector(PRigidDynamic->getAngularVelocity()) );
 	}
 #endif // WITH_PHYSX
+
+#if WITH_BOX2D
+	if (BodyInstancePtr != nullptr)
+	{
+		AngVel = FPhysicsIntegration2D::ConvertBoxAngularVelocityToUnreal(BodyInstancePtr->GetAngularVelocity());
+	}
+#endif
+
 	return AngVel;
 }
 
 FVector FBodyInstance::GetUnrealWorldVelocityAtPoint(const FVector& Point) const
 {
-	FVector LinVel(0,0,0);
+	FVector LinVel(EForceInit::ForceInitToZero);
 
 #if WITH_PHYSX
-	PxRigidDynamic* PRigidDynamic = GetPxRigidDynamic();
-	if(PRigidDynamic != NULL)
+	if (PxRigidDynamic* PRigidDynamic = GetPxRigidDynamic())
 	{
 		SCOPED_SCENE_READ_LOCK(PRigidDynamic->getScene());
 		PxVec3 PPoint = U2PVector(Point);
@@ -1886,15 +2010,21 @@ FVector FBodyInstance::GetUnrealWorldVelocityAtPoint(const FVector& Point) const
 	}
 #endif // WITH_PHYSX
 
+#if WITH_BOX2D
+	if (BodyInstancePtr != nullptr)
+	{
+		const b2Vec2 BoxPoint = FPhysicsIntegration2D::ConvertUnrealVectorToBox(Point);
+		LinVel = FPhysicsIntegration2D::ConvertBoxVectorToUnreal(BodyInstancePtr->GetLinearVelocityFromWorldPoint(BoxPoint));
+	}
+#endif
+
 	return LinVel;
 }
 
-
 FVector FBodyInstance::GetCOMPosition() const
 {
 #if WITH_PHYSX
-	PxRigidDynamic* PRigidDynamic = GetPxRigidDynamic();
-	if(PRigidDynamic)
+	if (PxRigidDynamic* PRigidDynamic = GetPxRigidDynamic())
 	{
 		SCOPED_SCENE_READ_LOCK(PRigidDynamic->getScene());
 		PxTransform PLocalCOM = PRigidDynamic->getCMassLocalPose();
@@ -1905,9 +2035,16 @@ FVector FBodyInstance::GetCOMPosition() const
 	{
 		return FVector::ZeroVector;
 	}
-#else
-	return FVector::ZeroVector;
 #endif // WITH_PHYSX
+
+#if WITH_BOX2D
+	if (BodyInstancePtr != nullptr)
+	{
+		return FPhysicsIntegration2D::ConvertBoxVectorToUnreal(BodyInstancePtr->GetWorldCenter());
+	}
+#endif
+
+	return FVector::ZeroVector;
 }
 
 float FBodyInstance::GetBodyMass() const
@@ -1915,17 +2052,19 @@ float FBodyInstance::GetBodyMass() const
 	float Retval = 0.f;
 
 #if WITH_PHYSX
-	PxRigidDynamic* PRigidDynamic = GetPxRigidDynamic();
-	if(PRigidDynamic)
+	if (PxRigidDynamic* PRigidDynamic = GetPxRigidDynamic())
 	{
 		SCOPED_SCENE_READ_LOCK(PRigidDynamic->getScene());
 		Retval = PRigidDynamic->getMass();
 	}
-	else
+#endif // WITH_PHYSX
+
+#if WITH_BOX2D
+	if (BodyInstancePtr != nullptr)
 	{
-		Retval = 0.0f;
+		Retval = BodyInstancePtr->GetMass();
 	}
-#endif // WITH_PHYSX
+#endif
 
 	return Retval;
 }
@@ -1947,12 +2086,14 @@ FBox FBodyInstance::GetBodyBounds() const
 	}
 #endif // WITH_PHYSX
 
+	//@TODO: BOX2D: Implement GetBodyBounds
+
 	return Bounds;
 }
 
-void FBodyInstance::DrawCOMPosition( FPrimitiveDrawInterface* PDI, float COMRenderSize, const FColor& COMRenderColor )
+void FBodyInstance::DrawCOMPosition(FPrimitiveDrawInterface* PDI, float COMRenderSize, const FColor& COMRenderColor)
 {
-	if(IsValidBodyInstance())
+	if (IsValidBodyInstance())
 	{
 		DrawWireStar(PDI, GetCOMPosition(), COMRenderSize, COMRenderColor, SDPG_World);
 	}
@@ -2210,31 +2351,46 @@ void FBodyInstance::UpdateMassProperties()
 	}
 #endif
 
+	//@TODO: BOX2D: Implement UpdateMassProperties
 }
 
 void FBodyInstance::UpdateDampingProperties()
 {
 #if WITH_PHYSX
-	PxRigidDynamic* PRigidDynamic = GetPxRigidDynamic();
-	if (PRigidDynamic != NULL)
+	if (PxRigidDynamic* PRigidDynamic = GetPxRigidDynamic())
 	{
 		SCOPED_SCENE_WRITE_LOCK(PRigidDynamic->getScene());
 		PRigidDynamic->setLinearDamping(LinearDamping);
 		PRigidDynamic->setAngularDamping(AngularDamping);
 	}
 #endif
+
+#if WITH_BOX2D
+	if (BodyInstancePtr)
+	{
+		BodyInstancePtr->SetLinearDamping(LinearDamping);
+		BodyInstancePtr->SetAngularDamping(AngularDamping);
+	}
+#endif
 }
 
-bool FBodyInstance::IsInstanceAwake()
+bool FBodyInstance::IsInstanceAwake() const
 {
 #if WITH_PHYSX
-	PxRigidDynamic* PRigidDynamic = GetPxRigidDynamic();
-	if (PRigidDynamic != NULL)
+	if (PxRigidDynamic* PRigidDynamic = GetPxRigidDynamic())
 	{
 		SCOPED_SCENE_READ_LOCK(PRigidDynamic->getScene());
 		return !PRigidDynamic->isSleeping();
 	}
 #endif
+
+#if WITH_BOX2D
+	if (BodyInstancePtr)
+	{
+		return BodyInstancePtr->IsAwake();
+	}
+#endif
+
 	return false;
 }
 
@@ -2248,6 +2404,13 @@ void FBodyInstance::WakeInstance()
 		PRigidDynamic->wakeUp();
 	}
 #endif
+
+#if WITH_BOX2D
+	if (BodyInstancePtr)
+	{
+		BodyInstancePtr->SetAwake(true);
+	}
+#endif
 }
 
 void FBodyInstance::PutInstanceToSleep()
@@ -2259,50 +2422,84 @@ void FBodyInstance::PutInstanceToSleep()
 		PRigidDynamic->putToSleep();
 	}
 #endif //WITH_PHYSX
+
+#if WITH_BOX2D
+	if (BodyInstancePtr)
+	{
+		BodyInstancePtr->SetAwake(false);
+	}
+#endif
 }
 
 void FBodyInstance::SetLinearVelocity(const FVector& NewVel, bool bAddToCurrent)
 {
 #if WITH_PHYSX
-	PxRigidDynamic* PRigidDynamic = GetPxRigidDynamic();
-	if (PRigidDynamic != NULL)
+	if (PxRigidDynamic* PRigidDynamic = GetPxRigidDynamic())
 	{
 		PxVec3 PNewVel = U2PVector(NewVel);
 
-		if(bAddToCurrent)
+		if (bAddToCurrent)
 		{
-			PxVec3 POldVel = PRigidDynamic->getLinearVelocity();
+			const PxVec3 POldVel = PRigidDynamic->getLinearVelocity();
 			PNewVel += POldVel;
 		}
 
 		PRigidDynamic->setLinearVelocity(PNewVel);
 	}
 #endif // WITH_PHYSX
+
+#if WITH_BOX2D
+	if (BodyInstancePtr)
+	{
+		b2Vec2 BoxNewVel = FPhysicsIntegration2D::ConvertUnrealVectorToBox(NewVel);
+
+		if (bAddToCurrent)
+		{
+			const b2Vec2 BoxOldVel = BodyInstancePtr->GetLinearVelocity();
+			BoxNewVel += BoxOldVel;
+		}
+
+		BodyInstancePtr->SetLinearVelocity(BoxNewVel);
+	}
+#endif
 }
 
 /** Note NewAngVel is in degrees per second */
 void FBodyInstance::SetAngularVelocity(const FVector& NewAngVel, bool bAddToCurrent)
 {
 #if WITH_PHYSX
-	PxRigidDynamic* PRigidDynamic = GetPxRigidDynamic();
-	if (PRigidDynamic != NULL)
+	if (PxRigidDynamic* PRigidDynamic = GetPxRigidDynamic())
 	{
 		PxVec3 PNewAngVel = U2PVector( FMath::DegreesToRadians(NewAngVel) );
 
-		if(bAddToCurrent)
+		if (bAddToCurrent)
 		{
-			PxVec3 POldAngVel = PRigidDynamic->getAngularVelocity();
+			const PxVec3 POldAngVel = PRigidDynamic->getAngularVelocity();
 			PNewAngVel += POldAngVel;
 		}
 
 		PRigidDynamic->setAngularVelocity(PNewAngVel);
 	}
 #endif //WITH_PHYSX
+
+#if WITH_BOX2D
+	if (BodyInstancePtr)
+	{
+		float BoxNewAngVel = FPhysicsIntegration2D::ConvertUnrealAngularVelocityToBox(NewAngVel);
+
+		if (bAddToCurrent)
+		{
+			const float BoxOldAngVel = BodyInstancePtr->GetAngularVelocity();
+			BoxNewAngVel += BoxOldAngVel;
+		}
+
+		BodyInstancePtr->SetAngularVelocity(BoxNewAngVel);
+	}
+#endif
 }
 
 void FBodyInstance::SetMaxAngularVelocity(float NewMaxAngVel, bool bAddToCurrent)
 {
-	//NewMaxAngVel = 400;
 #if WITH_PHYSX
 	PxRigidDynamic* PRigidDynamic = GetPxRigidDynamic();
 	if (PRigidDynamic)
@@ -2324,6 +2521,8 @@ void FBodyInstance::SetMaxAngularVelocity(float NewMaxAngVel, bool bAddToCurrent
 		MaxAngularVelocity = NewMaxAngVel;	//doesn't really matter since we are not dynamic, but makes sense that we update this anyway
 	}
 #endif
+
+	//@TODO: BOX2D: Implement SetMaxAngularVelocity
 }
 
 
@@ -2339,6 +2538,8 @@ void FBodyInstance::AddForce(const FVector& Force, bool bAllowSubstepping)
 		
 	}
 #endif // WITH_PHYSX
+
+	//@TODO: BOX2D: Implement AddForce
 }
 
 void FBodyInstance::AddForceAtPosition(const FVector& Force, const FVector& Position, bool bAllowSubstepping)
@@ -2352,6 +2553,8 @@ void FBodyInstance::AddForceAtPosition(const FVector& Force, const FVector& Posi
 		PhysScene->AddForceAtPosition(this, Force, Position, bAllowSubstepping);
 	}
 #endif // WITH_PHYSX
+
+	//@TODO: BOX2D: Implement AddForceAtPosition
 }
 
 void FBodyInstance::AddTorque(const FVector& Torque, bool bAllowSubstepping)
@@ -2365,6 +2568,8 @@ void FBodyInstance::AddTorque(const FVector& Torque, bool bAllowSubstepping)
 		PhysScene->AddTorque(this, Torque, bAllowSubstepping);
 	}
 #endif // WITH_PHYSX
+
+	//@TODO: BOX2D: Implement AddTorque
 }
 
 void FBodyInstance::AddImpulse(const FVector& Impulse, bool bVelChange)
@@ -2377,6 +2582,8 @@ void FBodyInstance::AddImpulse(const FVector& Impulse, bool bVelChange)
 		PRigidDynamic->addForce(U2PVector(Impulse), Mode, true);
 	}
 #endif // WITH_PHYSX
+
+	//@TODO: BOX2D: Implement AddImpulse
 }
 
 void FBodyInstance::AddImpulseAtPosition(const FVector& Impulse, const FVector& Position)
@@ -2389,6 +2596,8 @@ void FBodyInstance::AddImpulseAtPosition(const FVector& Impulse, const FVector&
 		PxRigidBodyExt::addForceAtPos(*PRigidDynamic, U2PVector(Impulse), U2PVector(Position), Mode, true);
 	}
 #endif // WITH_PHYSX
+
+	//@TODO: BOX2D: Implement AddImpulseAtPosition
 }
 
 void FBodyInstance::SetInstanceNotifyRBCollision(bool bNewNotifyCollision)
@@ -2398,19 +2607,26 @@ void FBodyInstance::SetInstanceNotifyRBCollision(bool bNewNotifyCollision)
 }
 
 
-void FBodyInstance::SetEnableGravity( bool bInGravityEnabled )
+void FBodyInstance::SetEnableGravity(bool bInGravityEnabled)
 {
 	if (bEnableGravity != bInGravityEnabled)
 	{
 		bEnableGravity = bInGravityEnabled;
+
 #if WITH_PHYSX
-		PxRigidDynamic* PActor = GetPxRigidDynamic();
-		if (PActor != NULL)
+		if (PxRigidDynamic* PActor = GetPxRigidDynamic())
 		{
 			PActor->setActorFlag( PxActorFlag::eDISABLE_GRAVITY, !bEnableGravity );
 		}
 #endif // WITH_PHYSX
 
+#if WITH_BOX2D
+		if (BodyInstancePtr)
+		{
+			BodyInstancePtr->SetGravityScale(bEnableGravity ? 1.0f : 0.0f);
+		}
+#endif
+
 		if (bEnableGravity)
 		{
 			WakeInstance();
@@ -2428,6 +2644,8 @@ void FBodyInstance::AddRadialImpulseToBody(const FVector& Origin, float Radius,
 		AddRadialImpulseToPxRigidDynamic(*PRigidDynamic, Origin, Radius, Strength, Falloff, bVelChange);
 	}
 #endif // WITH_PHYSX
+
+//@TODO: BOX2D: Implement AddRadialImpulseToBody
 }
 
 void FBodyInstance::AddRadialForceToBody(const FVector& Origin, float Radius, float Strength, uint8 Falloff)
@@ -2439,6 +2657,8 @@ void FBodyInstance::AddRadialForceToBody(const FVector& Origin, float Radius, fl
 		AddRadialForceToPxRigidDynamic(*PRigidDynamic, Origin, Radius, Strength, Falloff);
 	}
 #endif // WITH_PHYSX
+
+	//@TODO: BOX2D: Implement AddradialForceToBody
 }
 
 
@@ -2542,6 +2762,9 @@ bool FBodyInstance::LineTrace(struct FHitResult& OutHit,const FVector& Start,con
 		}
 	}
 #endif //WITH_PHYSX
+
+	//@TODO: BOX2D: Implement LineTrace
+
 	return false;
 }
 
@@ -2587,6 +2810,9 @@ bool FBodyInstance::Sweep(struct FHitResult& OutHit, const FVector& Start, const
 		}
 	}
 #endif //WITH_PHYSX
+
+	//@TODO: BOX2D: Implement Sweep
+
 	return false;
 }
 
@@ -2714,6 +2940,8 @@ float FBodyInstance::GetDistanceToBody(const FVector& Point, FVector& OutPointOn
 	// returns -1 if shape is not found
 	return -1.f;
 #endif //WITH_PHYSX
+
+	//@TODO: BOX2D: Implement DistanceToBody
 }
 
 #if WITH_PHYSX
@@ -2749,7 +2977,7 @@ bool FBodyInstance::Overlap(const PxGeometry& PGeom, const PxTransform&  ShapePo
 
 bool FBodyInstance::IsValidCollisionProfileName(FName InCollisionProfileName)
 {
-	return (InCollisionProfileName != NAME_None && InCollisionProfileName != UCollisionProfile::CustomCollisionProfileName);
+	return (InCollisionProfileName != NAME_None) && (InCollisionProfileName != UCollisionProfile::CustomCollisionProfileName);
 }
 
 void FBodyInstance::LoadProfileData(bool bVerifyProfile)
diff --git a/Engine/Source/Runtime/Engine/Private/PhysicsEngine2D/Box2DIntegration.h b/Engine/Source/Runtime/Engine/Private/PhysicsEngine2D/Box2DIntegration.h
--- a/Engine/Source/Runtime/Engine/Private/PhysicsEngine2D/Box2DIntegration.h
+++ b/Engine/Source/Runtime/Engine/Private/PhysicsEngine2D/Box2DIntegration.h
@@ -81,7 +81,7 @@ private:
 
 public:
 
-	//@TODO: These are both fixed on a XZ coordinate system
+	//@TODO: BOX2D: These conversion functions ignore the Paper2D coordinate system (they are fixed on a XZ coordinate system)
 	static inline FVector ConvertBoxVectorToUnreal(const b2Vec2& Vector)
 	{
 		return FVector(Vector.x * UnrealUnitsPerMeter, 0.0f, Vector.y * UnrealUnitsPerMeter);
@@ -92,6 +92,16 @@ public:
 		return b2Vec2(Vector.X / UnrealUnitsPerMeter, Vector.Z / UnrealUnitsPerMeter);
 	}
 
+	static inline float ConvertUnrealAngularVelocityToBox(const FVector& AngularVelocity)
+	{
+		return FMath::DegreesToRadians<float>(AngularVelocity.Y);
+	}
+
+	static inline FVector ConvertBoxAngularVelocityToUnreal(const float AngularVelocity)
+	{
+		return FVector(0.0f, FMath::RadiansToDegrees<float>(AngularVelocity), 0.0f);
+	}
+
 	// Finds the scene associated with the specified UWorld
 	static TSharedPtr<FPhysicsScene2D> FindAssociatedScene(UWorld* Source);
 
