Physics: Fix typo in UpdateBodyScale that prevents it from working on valid bodies

diff --git a/Engine/Source/Runtime/Engine/Private/PhysicsEngine/BodyInstance.cpp b/Engine/Source/Runtime/Engine/Private/PhysicsEngine/BodyInstance.cpp
--- a/Engine/Source/Runtime/Engine/Private/PhysicsEngine/BodyInstance.cpp
+++ b/Engine/Source/Runtime/Engine/Private/PhysicsEngine/BodyInstance.cpp
@@ -1426,7 +1426,7 @@ bool FBodyInstance::UpdateBodyScale(const FVector& InScale3D)
 {
 	FVector InScale3DAdjusted = InScale3D;
 
-	if (IsValidBodyInstance())
+	if (!IsValidBodyInstance())
 	{
 		//UE_LOG(LogPhysics, Log, TEXT("Body hasn't been initialized. Call InitBody to initialize."));
 		return false;
