Physics: Fix another IsValidBodyInstance sense inversion

diff --git a/Engine/Source/Runtime/Engine/Private/PhysicsEngine/BodyInstance.cpp b/Engine/Source/Runtime/Engine/Private/PhysicsEngine/BodyInstance.cpp
--- a/Engine/Source/Runtime/Engine/Private/PhysicsEngine/BodyInstance.cpp
+++ b/Engine/Source/Runtime/Engine/Private/PhysicsEngine/BodyInstance.cpp
@@ -1752,7 +1752,7 @@ void FBodyInstance::SetInstanceSimulatePhysics(bool bSimulate, bool bMaintainPhy
 #if !(UE_BUILD_SHIPPING || UE_BUILD_TEST)
 	if (bSimulate && OwnerComponent.IsValid())
 	{
-		if (IsValidBodyInstance())
+		if (!IsValidBodyInstance())
 		{
 			FMessageLog("PIE").Warning(FText::Format(LOCTEXT("SimPhysNoBody", "Trying to simulate physics on ''{0}'' but no physics body."),
 				FText::FromString(GetPathNameSafe(OwnerComponent.Get()))));
@@ -1850,7 +1850,7 @@ FTransform FBodyInstance::GetUnrealWorldTransform() const
 		const float RotationInRadians = BodyInstancePtr->GetAngle();
 		
 		const FVector Translation3D(FPhysicsIntegration2D::ConvertBoxVectorToUnreal(Pos2D));
-		const FRotator Rotation3D(FMath::RadiansToDegrees(RotationInRadians), 0.0f, 0.0f);
+		const FRotator Rotation3D(FMath::RadiansToDegrees(RotationInRadians), 0.0f, 0.0f); //@TODO: BOX2D: Should be moved to FPhysicsIntegration2D
 
 		return FTransform(Rotation3D, Translation3D, Scale3D);
 	}
