Physics: Fixing monolithic build issue

diff --git a/Engine/Source/Runtime/Engine/Private/PhysicsEngine2D/AggregateGeometry2D.cpp b/Engine/Source/Runtime/Engine/Private/PhysicsEngine2D/AggregateGeometry2D.cpp
--- a/Engine/Source/Runtime/Engine/Private/PhysicsEngine2D/AggregateGeometry2D.cpp
+++ b/Engine/Source/Runtime/Engine/Private/PhysicsEngine2D/AggregateGeometry2D.cpp
@@ -5,10 +5,10 @@
 
 static const int32 DrawCollisionSides2D = 16;
 
-//@TODO: PAPER2D: These should be renamed, and set when the Paper2D ones are loaded
-FVector PaperAxisX(1.0f, 0.0f, 0.0f);
-FVector PaperAxisY(0.0f, 0.0f, 1.0f);
-FVector PaperAxisZ(0.0f, -1.0f, 0.0f);
+//@TODO: PAPER2D: These axes should be renamed, and set when the Paper2D ones are loaded
+FVector Box2D_AxisX(1.0f, 0.0f, 0.0f);
+FVector Box2D_AxisY(0.0f, 0.0f, 1.0f);
+FVector Box2D_AxisZ(0.0f, -1.0f, 0.0f);
 
 //////////////////////////////////////////////////////////////////////////
 // FCircleElement2D
@@ -22,8 +22,8 @@ float FCircleElement2D::GetArea(const FVector& Scale) const
 void FCircleElement2D::DrawElemWire(FPrimitiveDrawInterface* PDI, const FTransform& ElemTM, float Scale, const FColor Color)
 {
 	const FVector ElemCenter = ElemTM.GetLocation();
-	const FVector X = ElemTM.TransformVector(PaperAxisX);
-	const FVector Y = ElemTM.TransformVector(PaperAxisY);
+	const FVector X = ElemTM.TransformVector(Box2D_AxisX);
+	const FVector Y = ElemTM.TransformVector(Box2D_AxisY);
 
 	DrawCircle(PDI, ElemCenter, X, Y, Color, Scale * Radius, DrawCollisionSides2D, SDPG_World);
 }
@@ -31,8 +31,8 @@ void FCircleElement2D::DrawElemWire(FPrimitiveDrawInterface* PDI, const FTransfo
 void FCircleElement2D::DrawElemSolid(class FPrimitiveDrawInterface* PDI, const FTransform& ElemTM, float Scale, const FMaterialRenderProxy* MaterialRenderProxy)
 {
 	const FVector ElemCenter = ElemTM.GetLocation();
-	const FVector X = ElemTM.TransformVector(PaperAxisX);
-	const FVector Y = ElemTM.TransformVector(PaperAxisY);
+	const FVector X = ElemTM.TransformVector(Box2D_AxisX);
+	const FVector Y = ElemTM.TransformVector(Box2D_AxisY);
 
 	DrawDisc(PDI, ElemCenter, X, Y, FColor::White, Scale * Radius, DrawCollisionSides2D, MaterialRenderProxy, SDPG_World);
 }
