Physics: Implement 2D versions of UpdatePhysicalMaterials and UpdateMassProperties

diff --git a/Engine/Source/Runtime/Engine/Private/PhysicsEngine/BodyInstance.cpp b/Engine/Source/Runtime/Engine/Private/PhysicsEngine/BodyInstance.cpp
--- a/Engine/Source/Runtime/Engine/Private/PhysicsEngine/BodyInstance.cpp
+++ b/Engine/Source/Runtime/Engine/Private/PhysicsEngine/BodyInstance.cpp
@@ -311,10 +311,10 @@ TArray<PxShape*> FBodyInstance::GetAllShapes(int32& OutNumSyncShapes) const
 
 void FBodyInstance::UpdatePhysicalMaterials()
 {
-	//@TODO: BOX2D: Implement UpdatePhysicalMaterials
-#if WITH_PHYSX
 	UPhysicalMaterial* SimplePhysMat = GetSimplePhysicalMaterial();
 	check(SimplePhysMat != NULL);
+
+#if WITH_PHYSX
 	PxMaterial* PSimpleMat = SimplePhysMat->GetPhysXMaterial();
 	check(PSimpleMat != NULL);
 
@@ -376,6 +376,19 @@ void FBodyInstance::UpdatePhysicalMaterials()
 		SCENE_UNLOCK_WRITE(RigidActorAsync->getScene());
 	}
 #endif
+
+#if WITH_BOX2D
+	if (BodyInstancePtr)
+	{
+		for (b2Fixture* Fixture = BodyInstancePtr->GetFixtureList(); Fixture; Fixture = Fixture->GetNext())
+		{
+			Fixture->SetFriction(SimplePhysMat->Friction);
+			Fixture->SetRestitution(SimplePhysMat->Restitution);
+
+			//@TODO: BOX2D: Determine if it's feasible to add support for FrictionCombineMode to Box2D
+		}
+	}
+#endif
 }
 
 void FBodyInstance::InvalidateCollisionProfileName()
@@ -864,9 +877,6 @@ void FBodyInstance::InitBody(UBodySetup* Setup, const FTransform& Transform, UPr
 
 					b2FixtureDef FixtureDef;
 					FixtureDef.shape = &CircleShape;
-					FixtureDef.density = 1.0f; //@TODO:
-					FixtureDef.friction = 0.3f;//@TODO:
-					//FixtureDef.restitution = ;
 
 					BodyInstancePtr->CreateFixture(&FixtureDef);
 				}
@@ -882,8 +892,6 @@ void FBodyInstance::InitBody(UBodySetup* Setup, const FTransform& Transform, UPr
 
 					b2FixtureDef FixtureDef;
 					FixtureDef.shape = &DynamicBox;
-					FixtureDef.density = 1.0f; //@TODO:
-					FixtureDef.friction = 0.3f;//@TODO:
 
 					BodyInstancePtr->CreateFixture(&FixtureDef);
 				}
@@ -908,8 +916,6 @@ void FBodyInstance::InitBody(UBodySetup* Setup, const FTransform& Transform, UPr
 
 						b2FixtureDef FixtureDef;
 						FixtureDef.shape = &ConvexPoly;
-						FixtureDef.density = 1.0f; //@TODO:
-						FixtureDef.friction = 0.3f;//@TODO:
 
 						BodyInstancePtr->CreateFixture(&FixtureDef);
 					}
@@ -992,6 +998,10 @@ void FBodyInstance::InitBody(UBodySetup* Setup, const FTransform& Transform, UPr
 	// If there is already a body instanced, or there is no scene to create it into, do nothing.
 	if (GetPxRigidActor() != NULL || !InRBScene)
 	{
+		// clear Owner and Setup info as well to properly clean up the BodyInstance.
+		OwnerComponent = NULL;
+		BodySetup = NULL;
+
 		return;
 	}
 
@@ -1019,7 +1029,7 @@ void FBodyInstance::InitBody(UBodySetup* Setup, const FTransform& Transform, UPr
 		PNewDynamic = GPhysXSDK->createRigidDynamic(PTransform);
 
 		// Put the dynamic actor in one scene or the other
-		if( !UseAsyncScene() )
+		if (!UseAsyncScene())
 		{
 			PNewActorSync = PNewDynamic;
 			PSceneForNewDynamic = PSceneSync;
@@ -1031,7 +1041,7 @@ void FBodyInstance::InitBody(UBodySetup* Setup, const FTransform& Transform, UPr
 		}
 
 		// Set kinematic if desired
-		if(bUseSimulate)
+		if (bUseSimulate)
 		{
 			PNewDynamic->setRigidDynamicFlag(PxRigidDynamicFlag::eKINEMATIC, false);
 		}
@@ -1053,7 +1063,7 @@ void FBodyInstance::InitBody(UBodySetup* Setup, const FTransform& Transform, UPr
 	// Copy geom from template and scale
 
 	// Sync:
-	if( PNewActorSync != NULL )
+	if (PNewActorSync != NULL)
 	{
 		Setup->AddShapesToRigidActor(PNewActorSync, Scale3D);
 		PNewActorSync->userData = &PhysxUserData; // Store pointer to owning bodyinstance.
@@ -1063,7 +1073,7 @@ void FBodyInstance::InitBody(UBodySetup* Setup, const FTransform& Transform, UPr
 	}
 
 	// Async:
-	if( PNewActorAsync != NULL )
+	if (PNewActorAsync != NULL)
 	{
 		check(PSceneAsync);
 
@@ -1075,18 +1085,18 @@ void FBodyInstance::InitBody(UBodySetup* Setup, const FTransform& Transform, UPr
 	}
 
 	// If we added no shapes, generate warning, destroy actor and bail out (don't add to scene).
-	if((PNewActorSync && PNewActorSync->getNbShapes() == 0) || ((PNewActorAsync && PNewActorAsync->getNbShapes() == 0)))
+	if ((PNewActorSync && PNewActorSync->getNbShapes() == 0) || ((PNewActorAsync && PNewActorAsync->getNbShapes() == 0)))
 	{
 		if (DebugName.Len())
 		{
 			UE_LOG(LogPhysics, Log, TEXT("InitBody: failed - no shapes: %s"), *DebugName);
 		}
 		// else if probably a level with no bsp
-		if(PNewActorSync)
+		if (PNewActorSync)
 		{
 			PNewActorSync->release();
 		}
-		if(PNewActorAsync)
+		if (PNewActorAsync)
 		{
 			PNewActorAsync->release();
 		}
@@ -2310,13 +2320,14 @@ int32 GetNumSimShapes(PxRigidDynamic* PRigidDynamic)
 
 void FBodyInstance::UpdateMassProperties()
 {
+	UPhysicalMaterial* PhysMat = GetSimplePhysicalMaterial();
+	check(PhysMat);
+
 #if WITH_PHYSX
-	check(BodySetup != NULL);
 	PxRigidDynamic* PRigidDynamic = GetPxRigidDynamic();
-	if((PRigidDynamic != NULL) && (GetNumSimShapes(PRigidDynamic) > 0))
+	if ((PRigidDynamic != NULL) && (GetNumSimShapes(PRigidDynamic) > 0))
 	{
 		// First, reset mass to default
-		UPhysicalMaterial* PhysMat  = GetSimplePhysicalMaterial();
 
 		// physical material - nothing can weigh less than hydrogen (0.09 kg/m^3)
 		float DensityKGPerCubicUU = FMath::Max(0.00009f, PhysMat->Density * 0.001f);
@@ -2341,7 +2352,7 @@ void FBodyInstance::UpdateMassProperties()
 		PRigidDynamic->setMass(NewMass);
 
 		// Apply the COMNudge
-		if(!COMNudge.IsZero())
+		if (!COMNudge.IsZero())
 		{
 			PxVec3 PCOMNudge = U2PVector(COMNudge);
 			PxTransform PCOMTransform = PRigidDynamic->getCMassLocalPose();
@@ -2351,7 +2362,28 @@ void FBodyInstance::UpdateMassProperties()
 	}
 #endif
 
-	//@TODO: BOX2D: Implement UpdateMassProperties
+#if WITH_BOX2D
+	if (BodyInstancePtr)
+	{
+		//@TODO: BOX2D: Implement COMNudge, Unreal 'funky' mass algorithm, etc... for UpdateMassProperties (if we don't update the formula, we need to update the displayed mass in the details panel)
+
+		// Unreal material density is in g/cm^3, and Box2D density is in kg/m^2
+		// physical material - nothing can weigh less than hydrogen (0.09 kg/m^3)
+		const float DensityKGPerCubicCM = FMath::Max(0.00009f, PhysMat->Density * 0.001f);
+		const float DensityKGPerCubicM = DensityKGPerCubicCM * 1000.0f;
+		const float DensityKGPerSquareM = DensityKGPerCubicM * 0.1f; //@TODO: BOX2D: Should there be a thickness property for mass calculations?
+		const float MassScaledDensity = DensityKGPerSquareM * FMath::Clamp<float>(MassScale, 0.01f, 100.0f);
+
+		// Apply the density
+		for (b2Fixture* Fixture = BodyInstancePtr->GetFixtureList(); Fixture; Fixture = Fixture->GetNext())
+		{
+			Fixture->SetDensity(MassScaledDensity);
+		}
+
+		// Recalculate the body mass / COM / etc... based on the updated density
+		BodyInstancePtr->ResetMassData();
+	}
+#endif
 }
 
 void FBodyInstance::UpdateDampingProperties()
