Physics: Generalized debug text for FBodyInstance to provide the associated asset for any kind of component instead of just UStaticMeshComponent

diff --git a/Engine/Source/Runtime/Engine/Private/PhysicsEngine/BodyInstance.cpp b/Engine/Source/Runtime/Engine/Private/PhysicsEngine/BodyInstance.cpp
--- a/Engine/Source/Runtime/Engine/Private/PhysicsEngine/BodyInstance.cpp
+++ b/Engine/Source/Runtime/Engine/Private/PhysicsEngine/BodyInstance.cpp
@@ -719,20 +719,9 @@ void FBodyInstance::InitBody(UBodySetup* Setup, const FTransform& Transform, UPr
 	// Make the debug name for this geometry...
 	FString DebugName(TEXT(""));
 #if (WITH_EDITORONLY_DATA || UE_BUILD_DEBUG || LOOKING_FOR_PERF_ISSUES) && !(UE_BUILD_SHIPPING || UE_BUILD_TEST) && !NO_LOGGING
-	UStaticMeshComponent* SMComp = Cast<UStaticMeshComponent>(PrimComp);
-	if (Owner)
-	{
-		DebugName += FString::Printf(TEXT("Actor: %s "), *Owner->GetPathName());
-	}
-
 	if (PrimComp)
 	{
-		DebugName += FString::Printf(TEXT("Component: %s "), *PrimComp->GetName());
-	}
-
-	if (SMComp)
-	{
-		DebugName += FString::Printf(TEXT("StaticMesh: %s"), *SMComp->StaticMesh->GetPathName());
+		DebugName += FString::Printf(TEXT("Component: %s "), *PrimComp->GetReadableName());
 	}
 
 	if (Setup->BoneName != NAME_None)
@@ -2698,14 +2687,19 @@ FString FBodyInstance::GetBodyDebugName() const
 {
 	FString DebugName;
 
-	if(OwnerComponent != NULL)
+	if (OwnerComponent != NULL)
 	{
 		DebugName = OwnerComponent->GetPathName();
+		if (const* StatObject = OwnerComponent->AdditionalStatObject())
+		{
+			DebugName += TEXT(" ");
+			StatObject->AppendName(DebugName);
+		}
 	}
 
-	if(BodySetup != NULL && BodySetup->BoneName != NAME_None)
+	if ((BodySetup != NULL) && (BodySetup->BoneName != NAME_None))
 	{
-		DebugName += FString(TEXT(" ")) + BodySetup->BoneName.ToString();
+		DebugName += FString(TEXT(" Bone: ")) + BodySetup->BoneName.ToString();
 	}
 
 	return DebugName;
