Physics: Fixed typo

diff --git a/Engine/Source/Runtime/Engine/Private/PhysicsEngine/BodyInstance.cpp b/Engine/Source/Runtime/Engine/Private/PhysicsEngine/BodyInstance.cpp
--- a/Engine/Source/Runtime/Engine/Private/PhysicsEngine/BodyInstance.cpp
+++ b/Engine/Source/Runtime/Engine/Private/PhysicsEngine/BodyInstance.cpp
@@ -2690,7 +2690,7 @@ FString FBodyInstance::GetBodyDebugName() const
 	if (OwnerComponent != NULL)
 	{
 		DebugName = OwnerComponent->GetPathName();
-		if (const* StatObject = OwnerComponent->AdditionalStatObject())
+		if (const UObject* StatObject = OwnerComponent->AdditionalStatObject())
 		{
 			DebugName += TEXT(" ");
 			StatObject->AppendName(DebugName);
