Changed AActor::RegisterAllComponents to not dirty outer package.

diff --git a/Engine/Source/Runtime/Engine/Private/Actor.cpp b/Engine/Source/Runtime/Engine/Private/Actor.cpp
--- a/Engine/Source/Runtime/Engine/Private/Actor.cpp
+++ b/Engine/Source/Runtime/Engine/Private/Actor.cpp
@@ -2970,7 +2970,7 @@ void AActor::RegisterAllComponents()
 
 		//Before we register our component, save it to our transaction buffer so if "undone" it will return to an unregistered state.
 		//This should prevent unwanted components hanging around when undoing a copy/paste or duplication action.
-		RootComponent->Modify();
+		RootComponent->Modify(false);
 
 		check(GetWorld());
 		RootComponent->RegisterComponentWithWorld(GetWorld());
@@ -2986,7 +2986,7 @@ void AActor::RegisterAllComponents()
 		{
 			//Before we register our component, save it to our transaction buffer so if "undone" it will return to an unregistered state.
 			//This should prevent unwanted components hanging around when undoing a copy/paste or duplication action.
-			Component->Modify();
+			Component->Modify(false);
 
 			check(GetWorld());
 			Component->RegisterComponentWithWorld(GetWorld());
