Don't include Slate.h/SlateCore.h in Engine.h/EnginePrivate.h

diff --git a/Engine/Source/Developer/AssetTools/Private/AssetTypeActions/AssetTypeActions_SlateBrush.h b/Engine/Source/Developer/AssetTools/Private/AssetTypeActions/AssetTypeActions_SlateBrush.h
--- a/Engine/Source/Developer/AssetTools/Private/AssetTypeActions/AssetTypeActions_SlateBrush.h
+++ b/Engine/Source/Developer/AssetTools/Private/AssetTypeActions/AssetTypeActions_SlateBrush.h
@@ -2,6 +2,8 @@
 
 #pragma once
 
+#include "Slate/SlateBrushAsset.h"
+
 class FAssetTypeActions_SlateBrush : public FAssetTypeActions_Base
 {
 public:
diff --git a/Engine/Source/Developer/AssetTools/Private/AssetTypeActions/AssetTypeActions_Texture2D.cpp b/Engine/Source/Developer/AssetTools/Private/AssetTypeActions/AssetTypeActions_Texture2D.cpp
--- a/Engine/Source/Developer/AssetTools/Private/AssetTypeActions/AssetTypeActions_Texture2D.cpp
+++ b/Engine/Source/Developer/AssetTools/Private/AssetTypeActions/AssetTypeActions_Texture2D.cpp
@@ -1,6 +1,7 @@
 // Copyright 1998-2014 Epic Games, Inc. All Rights Reserved.
 
 #include "AssetToolsPrivatePCH.h"
+#include "Slate/SlateBrushAsset.h"
 #include "ContentBrowserModule.h"
 
 #define LOCTEXT_NAMESPACE "AssetTypeActions"
diff --git a/Engine/Source/Developer/MeshUtilities/Private/MeshUtilities.cpp b/Engine/Source/Developer/MeshUtilities/Private/MeshUtilities.cpp
--- a/Engine/Source/Developer/MeshUtilities/Private/MeshUtilities.cpp
+++ b/Engine/Source/Developer/MeshUtilities/Private/MeshUtilities.cpp
@@ -13,6 +13,7 @@
 #include "Landscape/LandscapeDataAccess.h"
 #include "ImageUtils.h"
 #include "MaterialExportUtils.h"
+#include "Textures/TextureAtlas.h"
 
 /*------------------------------------------------------------------------------
 	MeshUtilities module.
diff --git a/Engine/Source/Editor/UnrealEd/Private/Factories/SlateBrushAssetFactory.cpp b/Engine/Source/Editor/UnrealEd/Private/Factories/SlateBrushAssetFactory.cpp
--- a/Engine/Source/Editor/UnrealEd/Private/Factories/SlateBrushAssetFactory.cpp
+++ b/Engine/Source/Editor/UnrealEd/Private/Factories/SlateBrushAssetFactory.cpp
@@ -1,7 +1,7 @@
 // Copyright 1998-2014 Epic Games, Inc. All Rights Reserved.
 
 #include "UnrealEd.h"
-
+#include "Slate/SlateBrushAsset.h"
 #include "Factories/SlateBrushAssetFactory.h"
 
 #define LOCTEXT_NAMESPACE "SlateBrushAssetFactory"
diff --git a/Engine/Source/Editor/UnrealEd/Private/ThumbnailRendering/SlateBrushThumbnailRenderer.cpp b/Engine/Source/Editor/UnrealEd/Private/ThumbnailRendering/SlateBrushThumbnailRenderer.cpp
--- a/Engine/Source/Editor/UnrealEd/Private/ThumbnailRendering/SlateBrushThumbnailRenderer.cpp
+++ b/Engine/Source/Editor/UnrealEd/Private/ThumbnailRendering/SlateBrushThumbnailRenderer.cpp
@@ -1,7 +1,7 @@
 // Copyright 1998-2014 Epic Games, Inc. All Rights Reserved.
 
 #include "UnrealEd.h"
-
+#include "Slate/SlateBrushAsset.h"
 // FPreviewScene derived helpers for rendering
 #include "ThumbnailHelpers.h"
 
diff --git a/Engine/Source/Runtime/Engine/Classes/EdGraph/EdGraphSchema.h b/Engine/Source/Runtime/Engine/Classes/EdGraph/EdGraphSchema.h
--- a/Engine/Source/Runtime/Engine/Classes/EdGraph/EdGraphSchema.h
+++ b/Engine/Source/Runtime/Engine/Classes/EdGraph/EdGraphSchema.h
@@ -4,6 +4,8 @@
 
 #include "EdGraphSchema.generated.h"
 
+class FSlateRect;
+
 /** Distinguishes between different graph types. Graphs can have different properties; for example: functions have one entry point, ubergraphs can have multiples. */
 UENUM()
 enum EGraphType
diff --git a/Engine/Source/Runtime/Engine/Classes/GameFramework/TouchInterface.h b/Engine/Source/Runtime/Engine/Classes/GameFramework/TouchInterface.h
--- a/Engine/Source/Runtime/Engine/Classes/GameFramework/TouchInterface.h
+++ b/Engine/Source/Runtime/Engine/Classes/GameFramework/TouchInterface.h
@@ -2,10 +2,10 @@
 
 #pragma once
 
-#include "Slate.h"
-#include "SVirtualJoystick.h"
 #include "TouchInterface.generated.h"
 
+class SVirtualJoystick;
+
 USTRUCT()
 struct FTouchInputControl
 {
diff --git a/Engine/Source/Runtime/Engine/Private/EnginePrivate.h b/Engine/Source/Runtime/Engine/Private/EnginePrivate.h
--- a/Engine/Source/Runtime/Engine/Private/EnginePrivate.h
+++ b/Engine/Source/Runtime/Engine/Private/EnginePrivate.h
@@ -197,14 +197,12 @@
 #include "Engine/BlueprintGeneratedClass.h"
 #include "Animation/AnimBlueprintGeneratedClass.h"
 #include "PhysicsEngine/BodySetup.h"
-#include "Slate/ButtonStyleAsset.h"
 #include "Camera/CameraAnim.h"
 #include "Camera/CameraAnimInst.h"
 #include "Camera/CameraModifier.h"
 #include "Camera/CameraShake.h"
 #include "Camera/CameraModifier_CameraShake.h"
 #include "GameFramework/CheatManager.h"
-#include "Slate/CheckboxStyleAsset.h"
 #include "Engine/CollisionProfile.h"
 #include "Curves/CurveLinearColor.h"
 #include "Curves/CurveVector.h"
@@ -303,7 +301,6 @@
 #include "Engine/DestructibleMesh.h"
 #include "Engine/SkeletalMeshSocket.h"
 #include "Animation/Skeleton.h"
-#include "Slate/SlateBrushAsset.h"
 #include "Sound/SoundBase.h"
 #include "Sound/DialogueSoundWaveProxy.h"
 #include "Sound/SoundClass.h"
diff --git a/Engine/Source/Runtime/Engine/Private/GameFramework/TouchInterface.cpp b/Engine/Source/Runtime/Engine/Private/GameFramework/TouchInterface.cpp
--- a/Engine/Source/Runtime/Engine/Private/GameFramework/TouchInterface.cpp
+++ b/Engine/Source/Runtime/Engine/Private/GameFramework/TouchInterface.cpp
@@ -2,6 +2,7 @@
 
 #include "EnginePrivate.h"
 #include "GameFramework/TouchInterface.h"
+#include "Slate.h"
 
 UTouchInterface::UTouchInterface(const class FPostConstructInitializeProperties& PCIP)
 	: Super(PCIP)
diff --git a/Engine/Source/Runtime/Engine/Private/Slate/SlateGameResources.cpp b/Engine/Source/Runtime/Engine/Private/Slate/SlateGameResources.cpp
--- a/Engine/Source/Runtime/Engine/Private/Slate/SlateGameResources.cpp
+++ b/Engine/Source/Runtime/Engine/Private/Slate/SlateGameResources.cpp
@@ -2,6 +2,7 @@
 
 #include "EnginePrivate.h"
 #include "Slate.h"
+#include "Slate/SlateBrushAsset.h"
 #include "AssetRegistryModule.h"
 #include "SlateGameResources.h"
 #include "MessageLog.h"
diff --git a/Engine/Source/Runtime/Engine/Private/Tests/AutomationCommon.cpp b/Engine/Source/Runtime/Engine/Private/Tests/AutomationCommon.cpp
--- a/Engine/Source/Runtime/Engine/Private/Tests/AutomationCommon.cpp
+++ b/Engine/Source/Runtime/Engine/Private/Tests/AutomationCommon.cpp
@@ -1,6 +1,7 @@
 // Copyright 1998-2014 Epic Games, Inc. All Rights Reserved.
 
 #include "EnginePrivate.h"
+#include "Slate.h"
 #include "AutomationCommon.h"
 #include "ImageUtils.h"
 
diff --git a/Engine/Source/Runtime/Engine/Private/UserInterface/Canvas.cpp b/Engine/Source/Runtime/Engine/Private/UserInterface/Canvas.cpp
--- a/Engine/Source/Runtime/Engine/Private/UserInterface/Canvas.cpp
+++ b/Engine/Source/Runtime/Engine/Private/UserInterface/Canvas.cpp
@@ -5,6 +5,7 @@
 =============================================================================*/
 
 #include "EnginePrivate.h"
+#include "Slate.h"
 #include "TileRendering.h"
 #include "RHIStaticStates.h"
 #include "WordWrapper.h"
diff --git a/Engine/Source/Runtime/Engine/Public/Engine.h b/Engine/Source/Runtime/Engine/Public/Engine.h
--- a/Engine/Source/Runtime/Engine/Public/Engine.h
+++ b/Engine/Source/Runtime/Engine/Public/Engine.h
@@ -236,14 +236,12 @@
 #include "Engine/BlueprintGeneratedClass.h"
 #include "Animation/AnimBlueprintGeneratedClass.h"
 #include "PhysicsEngine/BodySetup.h"
-#include "Slate/ButtonStyleAsset.h"
 #include "Camera/CameraAnim.h"
 #include "Camera/CameraAnimInst.h"
 #include "Camera/CameraModifier.h"
 #include "Camera/CameraShake.h"
 #include "Camera/CameraModifier_CameraShake.h"
 #include "GameFramework/CheatManager.h"
-#include "Slate/CheckboxStyleAsset.h"
 #include "Engine/CollisionProfile.h"
 #include "Curves/CurveLinearColor.h"
 #include "Curves/CurveVector.h"
@@ -343,7 +341,6 @@
 #include "Engine/DestructibleMesh.h"
 #include "Engine/SkeletalMeshSocket.h"
 #include "Animation/Skeleton.h"
-#include "Slate/SlateBrushAsset.h"
 #include "Sound/SoundBase.h"
 #include "Sound/DialogueSoundWaveProxy.h"
 #include "Sound/SoundClass.h"
diff --git a/Engine/Source/Runtime/Engine/Public/Slate/SceneViewport.h b/Engine/Source/Runtime/Engine/Public/Slate/SceneViewport.h
--- a/Engine/Source/Runtime/Engine/Public/Slate/SceneViewport.h
+++ b/Engine/Source/Runtime/Engine/Public/Slate/SceneViewport.h
@@ -4,6 +4,8 @@
 #ifndef __SLATE_VIEWPORT_H__
 #define __SLATE_VIEWPORT_H__
 
+#include "SlateCore.h"
+
 /** Called in FSceneViewport::ResizeFrame after ResizeViewport*/
 DECLARE_DELEGATE_OneParam( FOnSceneViewportResize, FVector2D );
 
diff --git a/Engine/Source/Runtime/HeadMountedDisplay/Public/IHeadMountedDisplay.h b/Engine/Source/Runtime/HeadMountedDisplay/Public/IHeadMountedDisplay.h
--- a/Engine/Source/Runtime/HeadMountedDisplay/Public/IHeadMountedDisplay.h
+++ b/Engine/Source/Runtime/HeadMountedDisplay/Public/IHeadMountedDisplay.h
@@ -3,6 +3,7 @@
 #pragma once
 
 #include "StereoRendering.h"
+#include "Layout/SlateRect.h"
 
 // depending on your kit and SDK, you may want to use this.
 // new distortion handling still in development.
diff --git a/Engine/Source/Runtime/UMG/Private/Components/Border.cpp b/Engine/Source/Runtime/UMG/Private/Components/Border.cpp
--- a/Engine/Source/Runtime/UMG/Private/Components/Border.cpp
+++ b/Engine/Source/Runtime/UMG/Private/Components/Border.cpp
@@ -1,6 +1,7 @@
 // Copyright 1998-2014 Epic Games, Inc. All Rights Reserved.
 
 #include "UMGPrivatePCH.h"
+#include "Slate/SlateBrushAsset.h"
 
 #define LOCTEXT_NAMESPACE "UMG"
 
diff --git a/Engine/Source/Runtime/UMG/Private/Components/CircularThrobberComponent.cpp b/Engine/Source/Runtime/UMG/Private/Components/CircularThrobberComponent.cpp
--- a/Engine/Source/Runtime/UMG/Private/Components/CircularThrobberComponent.cpp
+++ b/Engine/Source/Runtime/UMG/Private/Components/CircularThrobberComponent.cpp
@@ -1,6 +1,7 @@
 // Copyright 1998-2014 Epic Games, Inc. All Rights Reserved.
 
 #include "UMGPrivatePCH.h"
+#include "Slate/SlateBrushAsset.h"
 
 #define LOCTEXT_NAMESPACE "UMG"
 
diff --git a/Engine/Source/Runtime/UMG/Private/Components/Image.cpp b/Engine/Source/Runtime/UMG/Private/Components/Image.cpp
--- a/Engine/Source/Runtime/UMG/Private/Components/Image.cpp
+++ b/Engine/Source/Runtime/UMG/Private/Components/Image.cpp
@@ -1,6 +1,7 @@
 // Copyright 1998-2014 Epic Games, Inc. All Rights Reserved.
 
 #include "UMGPrivatePCH.h"
+#include "Slate/SlateBrushAsset.h"
 
 #define LOCTEXT_NAMESPACE "UMG"
 
diff --git a/Engine/Source/Runtime/UMG/Private/Components/ThrobberComponent.cpp b/Engine/Source/Runtime/UMG/Private/Components/ThrobberComponent.cpp
--- a/Engine/Source/Runtime/UMG/Private/Components/ThrobberComponent.cpp
+++ b/Engine/Source/Runtime/UMG/Private/Components/ThrobberComponent.cpp
@@ -1,6 +1,7 @@
 // Copyright 1998-2014 Epic Games, Inc. All Rights Reserved.
 
 #include "UMGPrivatePCH.h"
+#include "Slate/SlateBrushAsset.h"
 
 #define LOCTEXT_NAMESPACE "UMG"
 
diff --git a/Engine/Source/Runtime/UMG/Private/WidgetBlueprintLibrary.cpp b/Engine/Source/Runtime/UMG/Private/WidgetBlueprintLibrary.cpp
--- a/Engine/Source/Runtime/UMG/Private/WidgetBlueprintLibrary.cpp
+++ b/Engine/Source/Runtime/UMG/Private/WidgetBlueprintLibrary.cpp
@@ -1,7 +1,7 @@
 // Copyright 1998-2014 Epic Games, Inc. All Rights Reserved.
 
 #include "UMGPrivatePCH.h"
-
+#include "Slate/SlateBrushAsset.h"
 #include "WidgetBlueprintLibrary.h"
 
 #define LOCTEXT_NAMESPACE "UMG"
diff --git a/Engine/Source/Runtime/UMG/Public/Components/Border.h b/Engine/Source/Runtime/UMG/Public/Components/Border.h
--- a/Engine/Source/Runtime/UMG/Public/Components/Border.h
+++ b/Engine/Source/Runtime/UMG/Public/Components/Border.h
@@ -4,6 +4,8 @@
 
 #include "Border.generated.h"
 
+class USlateBrushAsset;
+
 /**
  * A border is a container widget that can contain one child widget, providing an opportunity 
  * to surround it with a border image and adjustable padding.
diff --git a/Engine/Source/Runtime/UMG/Public/Components/CircularThrobberComponent.h b/Engine/Source/Runtime/UMG/Public/Components/CircularThrobberComponent.h
--- a/Engine/Source/Runtime/UMG/Public/Components/CircularThrobberComponent.h
+++ b/Engine/Source/Runtime/UMG/Public/Components/CircularThrobberComponent.h
@@ -4,6 +4,8 @@
 
 #include "CircularThrobberComponent.generated.h"
 
+class USlateBrushAsset;
+
 /**  A throbber widget that orients images in a spinning circle. */
 UCLASS(meta=(BlueprintSpawnableComponent), ClassGroup=UserInterface)
 class UMG_API UCircularThrobberComponent : public UWidget
diff --git a/Engine/Source/Runtime/UMG/Public/Components/Image.h b/Engine/Source/Runtime/UMG/Public/Components/Image.h
--- a/Engine/Source/Runtime/UMG/Public/Components/Image.h
+++ b/Engine/Source/Runtime/UMG/Public/Components/Image.h
@@ -4,6 +4,8 @@
 
 #include "Image.generated.h"
 
+class USlateBrushAsset;
+
 /** Image widget */
 UCLASS(meta=(BlueprintSpawnableComponent), ClassGroup=UserInterface)
 class UMG_API UImage : public UWidget
diff --git a/Engine/Source/Runtime/UMG/Public/Components/ThrobberComponent.h b/Engine/Source/Runtime/UMG/Public/Components/ThrobberComponent.h
--- a/Engine/Source/Runtime/UMG/Public/Components/ThrobberComponent.h
+++ b/Engine/Source/Runtime/UMG/Public/Components/ThrobberComponent.h
@@ -4,6 +4,8 @@
 
 #include "ThrobberComponent.generated.h"
 
+class USlateBrushAsset;
+
 /** A Throbber widget that shows several zooming circles in a row. */
 UCLASS(meta=(BlueprintSpawnableComponent), ClassGroup=UserInterface)
 class UMG_API UThrobberComponent : public UWidget
