Fix GetInputKeyDown times being incorrect
TTP #337324

diff --git a/Engine/Source/Runtime/Engine/Private/UserInterface/PlayerInput.cpp b/Engine/Source/Runtime/Engine/Private/UserInterface/PlayerInput.cpp
--- a/Engine/Source/Runtime/Engine/Private/UserInterface/PlayerInput.cpp
+++ b/Engine/Source/Runtime/Engine/Private/UserInterface/PlayerInput.cpp
@@ -70,7 +70,7 @@ void UPlayerInput::FlushPressedKeys()
 
 	UWorld* World = GetWorld();
 	check(World);
-	float TimeSeconds = World->GetTimeSeconds(); 
+	float TimeSeconds = World->GetRealTimeSeconds();
 	for (TMap<FKey,FKeyState>::TIterator It(KeyStateMap); It; ++It)
 	{
 		FKeyState& KeyState = It.Value();
@@ -1149,6 +1149,7 @@ void UPlayerInput::DisplayDebug(class UCanvas* Canvas, const FDebugDisplayInfo&
 
 		UWorld* World = GetWorld();
 		check(World);
+		const float WorldRealTimeSeconds = World->GetRealTimeSeconds();
 		for (TMap<FKey,FKeyState>::TIterator It(KeyStateMap); It; ++It)
 		{
 			FKeyState const* const KeyState = &It.Value();
@@ -1161,7 +1162,7 @@ void UPlayerInput::DisplayDebug(class UCanvas* Canvas, const FDebugDisplayInfo&
 			{
 				if (!Key.IsFloatAxis())
 				{
-					Str += FString::Printf(TEXT(" time: %.2f"), World->TimeSince(KeyState->LastUpDownTransitionTime));
+					Str += FString::Printf(TEXT(" time: %.2f"), WorldRealTimeSeconds - KeyState->LastUpDownTransitionTime);
 				}
 				Canvas->SetDrawColor(180,255,180);
 				Canvas->DrawText(RenderFont, Str,4.0f, YPos);
@@ -1238,7 +1239,7 @@ float UPlayerInput::GetTimeDown( FKey InKey ) const
 		FKeyState const* const KeyState = KeyStateMap.Find(InKey);
 		if (KeyState && KeyState->bDown)
 		{
-			DownTime = (World->GetTimeSeconds() - KeyState->LastUpDownTransitionTime);
+			DownTime = (World->GetRealTimeSeconds() - KeyState->LastUpDownTransitionTime);
 		}
 	}
 
