FastArray serialization: rebuild client replication map if MarkArrayDirty is called. Allows clients to add elements to array without affecting replication.

diff --git a/Engine/Source/Runtime/Engine/Classes/Engine/NetSerialization.h b/Engine/Source/Runtime/Engine/Classes/Engine/NetSerialization.h
--- a/Engine/Source/Runtime/Engine/Classes/Engine/NetSerialization.h
+++ b/Engine/Source/Runtime/Engine/Classes/Engine/NetSerialization.h
@@ -355,6 +355,7 @@ struct FFastArraySerializer
 	/** This must be called if you just remove something from the array */
 	void MarkArrayDirty()
 	{
+		ItemMap.Reset();		// This allows to clients to add predictive elements to arrays without affecting replication.
 		ArrayReplicationKey++;
 		if (ArrayReplicationKey == INDEX_NONE)
 			ArrayReplicationKey++;
@@ -1073,4 +1074,4 @@ struct TStructOpsTypeTraits< FVector_NetQuantizeNormal > : public TStructOpsType
 	};
 };
 
-// --------------------------------------------------------------
\ No newline at end of file
+// --------------------------------------------------------------
