Fix sweepTest to use sweepSingle since sweepAny returns block on overlap
Fix rayTest to use bSingle for cases with multiple volumes

diff --git a/Engine/Source/Runtime/Engine/Private/Collision/PhysXCollision.cpp b/Engine/Source/Runtime/Engine/Private/Collision/PhysXCollision.cpp
--- a/Engine/Source/Runtime/Engine/Private/Collision/PhysXCollision.cpp
+++ b/Engine/Source/Runtime/Engine/Private/Collision/PhysXCollision.cpp
@@ -326,6 +326,7 @@ bool RaycastTest(const UWorld * World, const FVector Start, const FVector End, E
 		PxSceneQueryFilterData PQueryFilterData(PFilter, PxSceneQueryFilterFlag::eSTATIC | PxSceneQueryFilterFlag::eDYNAMIC | PxSceneQueryFilterFlag::ePREFILTER);
 		PxSceneQueryFlags POutputFlags = PxHitFlags();
 		FPxQueryFilterCallback PQueryCallback(Params.IgnoreActors);
+		PQueryCallback.bSingleQuery = true;
 
 		FPhysScene* PhysScene = World->GetPhysicsScene();
 
@@ -602,6 +603,7 @@ bool GeomSweepTest(const UWorld * World, const PxGeometry& PGeom, const PxQuat&
 
 		FPxQueryFilterCallbackSweep PQueryCallbackSweep(Params.IgnoreActors);
 		PQueryCallbackSweep.DiscardInitialOverlaps = !Params.bFindInitialOverlaps;
+		PQueryCallbackSweep.bSingleQuery = true;
 
 		PxTransform PStartTM(U2PVector(Start), PGeomRot);
 		PxVec3 PDir = U2PVector(Delta/DeltaMag);
@@ -612,16 +614,16 @@ bool GeomSweepTest(const UWorld * World, const PxGeometry& PGeom, const PxQuat&
 		// Enable scene locks, in case they are required
 		SCOPED_SCENE_READ_LOCK(SyncScene);
 
-		PxSceneQueryHit PQueryHit;
-		bHaveBlockingHit = SyncScene->sweepAny(PGeom, PStartTM, PDir, DeltaMag, POutputFlags, PQueryHit, PQueryFilterData, &PQueryCallbackSweep);
+		PxSweepHit PQueryHit;
+		bHaveBlockingHit = SyncScene->sweepSingle(PGeom, PStartTM, PDir, DeltaMag, POutputFlags, PQueryHit, PQueryFilterData, &PQueryCallbackSweep);
 
 		// Test async scene if async tests are requested and there was no blocking hit was found in the sync scene (since no hit info other than a boolean yes/no is recorded)
 		if( !bHaveBlockingHit && Params.bTraceAsyncScene && PhysScene->HasAsyncScene())
 		{
 			PxScene* AsyncScene = PhysScene->GetPhysXScene(PST_Async);
 			SCOPED_SCENE_READ_LOCK(AsyncScene);
-
-			bHaveBlockingHit = AsyncScene->sweepAny(PGeom, PStartTM, PDir, DeltaMag, POutputFlags, PQueryHit, PQueryFilterData, &PQueryCallbackSweep);
+			PxSweepHit PTestHit;
+			bHaveBlockingHit = AsyncScene->sweepSingle(PGeom, PStartTM, PDir, DeltaMag, POutputFlags, PTestHit, PQueryFilterData, &PQueryCallbackSweep);
 		}
 	}
 
