#ue4 - Fix FName being used improperly in a call to ensureMsgf()

diff --git a/Engine/Source/Runtime/Engine/Private/Slate/SlateGameResources.cpp b/Engine/Source/Runtime/Engine/Private/Slate/SlateGameResources.cpp
--- a/Engine/Source/Runtime/Engine/Private/Slate/SlateGameResources.cpp
+++ b/Engine/Source/Runtime/Engine/Private/Slate/SlateGameResources.cpp
@@ -47,7 +47,7 @@ const FSlateBrush* FSlateGameResources::GetBrush( const FName PropertyName, cons
 	if(Resource)
 	{
 		const USlateBrushAsset* BrushAsset = Cast<USlateBrushAsset>(*Resource);
-		ensureMsgf(BrushAsset, TEXT("Could not find resource '%s'"), PropertyName);
+		ensureMsgf(BrushAsset, TEXT("Could not find resource '%s'"), *PropertyName.ToString());
 		return BrushAsset ? &BrushAsset->Brush : GetDefaultBrush();
 	}
 	return FSlateStyleSet::GetBrush(PropertyName, Specifier);
@@ -60,7 +60,7 @@ const FSlateBrush* FSlateGameResources::GetOptionalBrush(const FName PropertyNam
 	if(Resource)
 	{
 		const USlateBrushAsset* BrushAsset = Cast<USlateBrushAsset>(*Resource);
-		ensureMsgf(BrushAsset, TEXT("Could not find resource '%s'"), PropertyName);
+		ensureMsgf(BrushAsset, TEXT("Could not find resource '%s'"), *PropertyName.ToString());
 		return BrushAsset ? &BrushAsset->Brush : DefaultBrush;
 	}
 	return FSlateStyleSet::GetOptionalBrush(PropertyName, Specifier, DefaultBrush);
