#ue4 - Make sure to abort overlap reporting if previous events mark the actor or component for destruction.

Prevents incorrect overlaps reported in cases such as a fast moving projectile that explodes early in a sweep, and which should not report overlaps with objects later in the sweep in the same frame (since it was destroyed earlier in time).

diff --git a/Engine/Source/Runtime/Engine/Private/PrimitiveComponent.cpp b/Engine/Source/Runtime/Engine/Private/PrimitiveComponent.cpp
--- a/Engine/Source/Runtime/Engine/Private/PrimitiveComponent.cpp
+++ b/Engine/Source/Runtime/Engine/Private/PrimitiveComponent.cpp
@@ -1828,6 +1828,12 @@ static bool IsPrimCompValidAndAlive(UPrimitiveComponent* PrimComp)
 // @todo, don't need to pass in Other actor?
 void UPrimitiveComponent::BeginComponentOverlap(const FOverlapInfo& OtherOverlap, bool bDoNotifies)
 {
+	// If pending kill, we should not generate any new overlaps
+	if (IsPendingKill())
+	{
+		return;
+	}
+
 	//	UE_LOG(LogActor, Log, TEXT("BEGIN OVERLAP! Self=%s SelfComp=%s, Other=%s, OtherComp=%s"), *GetNameSafe(this), *GetNameSafe(MyComp), *GetNameSafe(OtherComp->GetOwner()), *GetNameSafe(OtherComp));
 
 	UPrimitiveComponent* OtherComp = OtherOverlap.OverlapInfo.Component.Get();
@@ -2072,36 +2078,40 @@ void UPrimitiveComponent::UpdateOverlaps(TArray<FOverlapInfo> const* PendingOver
 			// determine what changed
 			TArray<FOverlapInfo> NewOverlappingComponents;
 
-			// Might be able to avoid testing for new overlaps at the end location.
-			if (OverlapsAtEndLocation != NULL && CVarAllowCachedOverlaps->GetInt())
+			// If pending kill, we should not generate any new overlaps
+			if (!IsPendingKill())
 			{
-				UE_LOG(LogPrimitiveComponent, VeryVerbose, TEXT("%s Skipping overlap test!"), *GetName());
-				NewOverlappingComponents = *OverlapsAtEndLocation;
-			}
-			else
-			{
-				UE_LOG(LogPrimitiveComponent, VeryVerbose, TEXT("%s Performing overlaps!"), *GetName());
-				UWorld* const MyWorld = MyActor->GetWorld();
-				TArray<FOverlapResult> Overlaps;
-				static FName NAME_UpdateOverlaps = FName(TEXT("UpdateOverlaps"));
-				// note this will include overlaps with components in the same actor.  
-				FComponentQueryParams Params(NAME_UpdateOverlaps);
-				Params.bTraceAsyncScene = bCheckAsyncSceneOnMove;
-				Params.AddIgnoredActors(MoveIgnoreActors);
-				MyWorld->ComponentOverlapMulti(Overlaps, this, GetComponentLocation(), GetComponentRotation(), Params);
-
-				for( int32 ResultIdx=0; ResultIdx<Overlaps.Num(); ResultIdx++ )
+				// Might be able to avoid testing for new overlaps at the end location.
+				if (OverlapsAtEndLocation != NULL && CVarAllowCachedOverlaps->GetInt())
 				{
-					const FOverlapResult& Result = Overlaps[ResultIdx];
-
-					UPrimitiveComponent const* const HitComp = Result.Component.Get();
-					if (!Result.bBlockingHit && HitComp  && HitComp->bGenerateOverlapEvents)
+					UE_LOG(LogPrimitiveComponent, VeryVerbose, TEXT("%s Skipping overlap test!"), *GetName());
+					NewOverlappingComponents = *OverlapsAtEndLocation;
+				}
+				else
+				{
+					UE_LOG(LogPrimitiveComponent, VeryVerbose, TEXT("%s Performing overlaps!"), *GetName());
+					UWorld* const MyWorld = MyActor->GetWorld();
+					TArray<FOverlapResult> Overlaps;
+					static FName NAME_UpdateOverlaps = FName(TEXT("UpdateOverlaps"));
+					// note this will include overlaps with components in the same actor.  
+					FComponentQueryParams Params(NAME_UpdateOverlaps);
+					Params.bTraceAsyncScene = bCheckAsyncSceneOnMove;
+					Params.AddIgnoredActors(MoveIgnoreActors);
+					MyWorld->ComponentOverlapMulti(Overlaps, this, GetComponentLocation(), GetComponentRotation(), Params);
+
+					for( int32 ResultIdx=0; ResultIdx<Overlaps.Num(); ResultIdx++ )
 					{
-						AActor* const ResultActor = Result.GetActor();
-						if ( (ResultActor != NULL) && (HitComp != this) && !ResultActor->IsAttachedTo(MyActor) && !MyActor->IsAttachedTo(ResultActor) && ResultActor != MyWorld->GetWorldSettings() && ResultActor->bActorInitialized)
+						const FOverlapResult& Result = Overlaps[ResultIdx];
+
+						UPrimitiveComponent const* const HitComp = Result.Component.Get();
+						if (!Result.bBlockingHit && HitComp  && HitComp->bGenerateOverlapEvents)
 						{
-							NewOverlappingComponents.Add(FOverlapInfo(Result.Component.Get(), Result.ItemIndex));		// don't need to add unique unless the overlap check can return dupes
-						}						
+							AActor* const ResultActor = Result.GetActor();
+							if ( (ResultActor != NULL) && (HitComp != this) && !ResultActor->IsAttachedTo(MyActor) && !MyActor->IsAttachedTo(ResultActor) && ResultActor != MyWorld->GetWorldSettings() && ResultActor->bActorInitialized)
+							{
+								NewOverlappingComponents.Add(FOverlapInfo(Result.Component.Get(), Result.ItemIndex));		// don't need to add unique unless the overlap check can return dupes
+							}						
+						}
 					}
 				}
 			}
