#ue4 - Fix ToggleDebugCamera

diff --git a/Engine/Source/Runtime/Engine/Private/DebugCameraController.cpp b/Engine/Source/Runtime/Engine/Private/DebugCameraController.cpp
--- a/Engine/Source/Runtime/Engine/Private/DebugCameraController.cpp
+++ b/Engine/Source/Runtime/Engine/Private/DebugCameraController.cpp
@@ -243,7 +243,7 @@ void ADebugCameraController::PostInitializeComponents()
 	SpawnInfo.bNoCollisionFail = true;
 	MyHUD = GetWorld()->SpawnActor<ADebugCameraHUD>( SpawnInfo );
 
-	ChangeState(NAME_Spectating);
+	ChangeState(NAME_Inactive);
 }
 
 void ADebugCameraController::OnActivate( APlayerController* OriginalPC )
