Fixed texture streaming bug

diff --git a/Engine/Source/Runtime/Engine/Private/SceneManagement.cpp b/Engine/Source/Runtime/Engine/Private/SceneManagement.cpp
--- a/Engine/Source/Runtime/Engine/Private/SceneManagement.cpp
+++ b/Engine/Source/Runtime/Engine/Private/SceneManagement.cpp
@@ -99,7 +99,7 @@ float ComputeBoundsScreenSize( const FVector4& Origin, const float SphereRadius,
 
 int8 ComputeStaticMeshLOD( const FStaticMeshRenderData* RenderData, const FVector4& Origin, const float SphereRadius, const FSceneView& View, float FactorScale )
 {
-	const int32 NumLODs = MAX_STATIC_MESH_LODS+1;
+	const int32 NumLODs = MAX_STATIC_MESH_LODS;
 
 	const float ScreenSize = ComputeBoundsScreenSize(Origin, SphereRadius, View) * FactorScale;
 
diff --git a/Engine/Source/Runtime/Engine/Private/StaticMesh.cpp b/Engine/Source/Runtime/Engine/Private/StaticMesh.cpp
--- a/Engine/Source/Runtime/Engine/Private/StaticMesh.cpp
+++ b/Engine/Source/Runtime/Engine/Private/StaticMesh.cpp
@@ -553,7 +553,7 @@ FStaticMeshRenderData::FStaticMeshRenderData()
 	: MaxStreamingTextureFactor(0.0f)
 	, bLODsShareStaticLighting(false)
 {
-	for (int32 LODIndex = 0; LODIndex < MAX_STATIC_MESH_LODS+1; ++LODIndex)
+	for (int32 LODIndex = 0; LODIndex < MAX_STATIC_MESH_LODS; ++LODIndex)
 	{
 		ScreenSize[LODIndex] = 0.0f;
 	}
@@ -592,35 +592,9 @@ void FStaticMeshRenderData::Serialize(FArchive& Ar, UStaticMesh* Owner, bool bCo
 
 	if (bCooked)
 	{
-		for (int32 LODIndex = 0; LODIndex < MAX_STATIC_MESH_LODS+1; ++LODIndex)
+		for (int32 LODIndex = 0; LODIndex < MAX_STATIC_MESH_LODS; ++LODIndex)
 		{
-			if(Ar.UE4Ver() < VER_UE4_STATIC_MESH_SCREEN_SIZE_LODS)
-			{
-				// We're loading in old data; which will be a LOD distance so we need
-				// to make an appropriate display factor from the distance;
-				float LODDistance = 0.0f;
-				Ar << LODDistance;
-
-				// Assuming an FOV of 90 and a screen size of 1920x1080 to estimate an appropriate display factor.
-				const float HalfFOV = PI / 4.0f;
-				const float ScreenWidth = 1920.0f;
-				const float ScreenHeight = 1080.0f;
-				const FVector4 PointToTest(0.0f, 0.0f, LODDistance, 1.0f);
-				FPerspectiveMatrix ProjMatrix(HalfFOV, ScreenWidth, ScreenHeight, 1.0f);
-
-				const float Divisor =  Dot3(-PointToTest, ProjMatrix.GetColumn(2));
-
-				FVector4 ScreenPosition = ProjMatrix.TransformFVector4(PointToTest);
-				const float ScreenMultiple = ScreenWidth / 2.0f * ProjMatrix.M[0][0];
-				const float ScreenRadius = ScreenMultiple * Bounds.SphereRadius / FMath::Max(ScreenPosition.W, 1.0f);
-				const float ScreenArea = ScreenWidth * ScreenHeight;
-				const float BoundsArea = PI * ScreenRadius * ScreenRadius;
-				ScreenSize[LODIndex] = FMath::Clamp(BoundsArea / ScreenArea, 0.0f, 1.0f);
-			}
-			else
-			{
-				Ar << ScreenSize[LODIndex];
-			}
+			Ar << ScreenSize[LODIndex];
 		}
 	}
 }
@@ -728,7 +702,7 @@ void FStaticMeshRenderData::ResolveSectionInfo(UStaticMesh* Owner)
 			ScreenSize[LODIndex] = FMath::Clamp(AutoDisplayFactor, 0.0f, ScreenSize[LODIndex-1] - Tolerance);
 		}
 	}
-	for (; LODIndex < MAX_STATIC_MESH_LODS + 1; ++LODIndex)
+	for (; LODIndex < MAX_STATIC_MESH_LODS; ++LODIndex)
 	{
 		ScreenSize[LODIndex] = 0.0f;
 	}
