Fixed world origin shifting for a foliage editor data

diff --git a/Engine/Source/Runtime/Engine/Private/InstancedFoliage.cpp b/Engine/Source/Runtime/Engine/Private/InstancedFoliage.cpp
--- a/Engine/Source/Runtime/Engine/Private/InstancedFoliage.cpp
+++ b/Engine/Source/Runtime/Engine/Private/InstancedFoliage.cpp
@@ -1633,6 +1633,13 @@ void AInstancedFoliageActor::ApplyWorldOffset(const FVector& InOffset, bool bWor
 				// Rehash instance location
 				MeshInfo.InstanceHash->InsertInstance(Instance.Location, InstanceIdx);
 			}
+
+			for (auto It = MeshInfo.ComponentHash.CreateIterator(); It; ++It)
+			{
+				// We assume here that component we painted foliage on will be shifted by same value
+				FFoliageComponentHashInfo& Info = It.Value();
+				Info.CachedLocation+= InOffset;
+			}
 		}
 	}
 }
