Fixed splitscreen coordinate adjustment for QA HUD clicking.

Changed coordinates (which were -ve) were being subtracted not added. Changed the code to add it.

TTP 318601 - QAGAME: SPLITSCREEN: SplitScreen player is unable to click the HUD in the QA-ClickHUD.umap.

diff --git a/Engine/Source/Runtime/Engine/Private/HUD.cpp b/Engine/Source/Runtime/Engine/Private/HUD.cpp
--- a/Engine/Source/Runtime/Engine/Private/HUD.cpp
+++ b/Engine/Source/Runtime/Engine/Private/HUD.cpp
@@ -159,7 +159,7 @@ void AHUD::PostRender()
 			{
 				for (FVector2D& ContactPoint : ContactPoints)
 				{
-					ContactPoint -= ContactPointOffset;
+					ContactPoint += ContactPointOffset;
 				}
 			}
 			UpdateHitBoxCandidates( ContactPoints );
@@ -802,7 +802,7 @@ const TArray<FIntRect>& AHUD::GetUIBlurRectangles() const
 
 bool AHUD::UpdateAndDispatchHitBoxClickEvents(FVector2D ClickLocation, const EInputEvent InEventType, const bool bDispatchOverOutEvent)
 {
-	ClickLocation -= GetCoordinateOffset();
+	ClickLocation += GetCoordinateOffset();
 
 	bool bHit = false;
 	for (FHUDHitBox* HitBoxHit : HitBoxHits)
