Allow TActorIteractor to take actor class to look for

diff --git a/Engine/Source/Runtime/Engine/Public/EngineUtils.h b/Engine/Source/Runtime/Engine/Public/EngineUtils.h
--- a/Engine/Source/Runtime/Engine/Public/EngineUtils.h
+++ b/Engine/Source/Runtime/Engine/Public/EngineUtils.h
@@ -407,8 +407,8 @@ public:
 	/**
 	 * Constructor, inits the starting position for iteration
 	 */
-	TActorIterator( UWorld* InWorld )
-		: TActorIteratorBase< FActorFilter >( InWorld, ActorType::StaticClass() )
+	TActorIterator( UWorld* InWorld, TSubclassOf<ActorType> InClass = ActorType::StaticClass() )
+		: TActorIteratorBase< FActorFilter >( InWorld, InClass )
 	{
 		++(*this);
 	}
