Canvas/HUD World Space Selection Rectangle (Marquee)
https://github.com/EpicGames/UnrealEngine/pull/127 ([REDACTED_PROFILE_URL])

diff --git a/Engine/Source/Runtime/Core/Public/Math/Box2D.h b/Engine/Source/Runtime/Core/Public/Math/Box2D.h
--- a/Engine/Source/Runtime/Core/Public/Math/Box2D.h
+++ b/Engine/Source/Runtime/Core/Public/Math/Box2D.h
@@ -283,22 +283,6 @@ public:
 	 */
 	FORCEINLINE bool Intersect( const FBox2D & other ) const;
 
-	/**
-	 * Test for intersection of bounding box and this bounding box in the X plane
-	 *
-	 * @param other bounding box to test intersection
-	 * @return true if the other box intersects this box in the X Plane, otherwise false
-	 */
-	FORCEINLINE bool IntersectX( const FBox2D& other ) const;
-	
-	/**
-	 * Test for intersection of bounding box and this bounding box in the Y plane
-	 *
-	 * @param other bounding box to test intersection
-	 * @return true if the other box intersects this box in the Y Plane, otherwise false
-	 */
-	FORCEINLINE bool IntersectY( const FBox2D& other ) const;
-	
 	/**
 	 * Test if point is inside this box
 	 *
@@ -434,26 +418,6 @@ FORCEINLINE bool FBox2D::Intersect( const FBox2D & Other ) const
 	return true;
 }
 
-
-FORCEINLINE bool FBox2D::IntersectX( const FBox2D& Other ) const
-{
-	if ((Min.X > Other.Max.X) || (Other.Min.X > Max.X))
-	{
-		return false;
-	}
-
-	return true;
-}
-FORCEINLINE bool FBox2D::IntersectY( const FBox2D& Other ) const
-{
-	if ((Min.Y > Other.Max.Y) || (Other.Min.Y > Max.Y))
-	{
-		return false;
-	}
-
-	return true;
-}
-
 FORCEINLINE FString FBox2D::ToString( ) const
 {
 	return FString::Printf(TEXT("bIsValid=%s, Min=(%s), Max=(%s)"), bIsValid ? TEXT("true") : TEXT("false"), *Min.ToString(), *Max.ToString());
diff --git a/Engine/Source/Runtime/Engine/Classes/GameFramework/HUD.h b/Engine/Source/Runtime/Engine/Classes/GameFramework/HUD.h
--- a/Engine/Source/Runtime/Engine/Classes/GameFramework/HUD.h
+++ b/Engine/Source/Runtime/Engine/Classes/GameFramework/HUD.h
@@ -401,6 +401,64 @@ public:
 	UFUNCTION(BlueprintPure, Category=HUD)
 	void Deproject(float ScreenX, float ScreenY, FVector& WorldPosition, FVector& WorldDirection);
 
+
+	/*
+	 * Returns the array of actors inside a selection rectangle, with a class filter.
+	 *
+	 * Sample usage:
+	 *
+	 *       TArray<AStaticMeshActor*> ActorsInSelectionRect;
+	 * 		Canvas->GetActorsInSelectionRectangle<AStaticMeshActor>(FirstPoint,SecondPoint,ActorsInSelectionRect);
+	 *
+	 *
+	 * @param FirstPoint					The first point, or anchor of the marquee box. Where the dragging of the marquee started in screen space.
+	 * @param SecondPoint					The second point, where the mouse cursor currently is / the other point of the box selection, in screen space.
+	 * @return OutActors					The actors that are within the selection box according to selection rule
+	 * @param bActorMustBeFullyEnclosed  	The Selection rule: whether the selection box can partially intersect Actor, or must fully enclose the Actor.
+	 *
+	 * returns false if selection could not occur. Make sure template class is extending AActor.
+	 */
+	template <typename ClassFilter>
+	bool GetActorsInSelectionRectangle(const FVector2D& FirstPoint, const FVector2D& SecondPoint, TArray<ClassFilter*>& OutActors, bool bActorMustBeFullyEnclosed = false)
+	{
+		//Is Actor subclass?
+		if (!ClassFilter::StaticClass()->IsChildOf(AActor::StaticClass()))
+		{
+			return false;
+		}
+
+		//Run Inner Function, output to Base AActor Array
+		TArray<AActor*> OutActorsBaseArray;
+		GetActorsInSelectionRectangle(ClassFilter, FirstPoint, SecondPoint, OutActorsBaseArray, bActorMustBeFullyEnclosed);
+
+		//Construct casted template type array
+		for (AActor* EachActor : OutActorsBaseArray)
+		{
+			OutActors.Add(CastChecked<ClassFilter>(EachActor));
+		}
+
+		return true;
+	};
+
+	/**
+	 * Returns the array of actors inside a selection rectangle, with a class filter.
+	 *
+	 * Sample usage:
+	 *
+	 *       TArray<AStaticMeshActor*> ActorsInSelectionRect;
+	 * 		Canvas->GetActorsInSelectionRectangle<AStaticMeshActor>(FirstPoint,SecondPoint,ActorsInSelectionRect);
+	 *
+	 *
+	 * @param FirstPoint					The first point, or anchor of the marquee box. Where the dragging of the marquee started in screen space.
+	 * @param SecondPoint					The second point, where the mouse cursor currently is / the other point of the box selection, in screen space.
+	 * @return OutActors					The actors that are within the selection box according to selection rule
+	 * @param bActorMustBeFullyEnclosed  	The Selection rule: whether the selection box can partially intersect Actor, or must fully enclose the Actor.
+	 *
+	 */
+	UFUNCTION(BlueprintPure, Category=HUD)
+	void GetActorsInSelectionRectangle(TSubclassOf<class AActor> ClassFilter, const FVector2D& FirstPoint, const FVector2D& SecondPoint, TArray<AActor*>& OutActors, bool bActorMustBeFullyEnclosed = false);
+
+
 	/**
 	 * Add a hitbox to the hud
 	 * @param Position			Coordinates of the top left of the hit box.
diff --git a/Engine/Source/Runtime/Engine/Private/HUD.cpp b/Engine/Source/Runtime/Engine/Private/HUD.cpp
--- a/Engine/Source/Runtime/Engine/Private/HUD.cpp
+++ b/Engine/Source/Runtime/Engine/Private/HUD.cpp
@@ -631,6 +631,76 @@ void AHUD::Deproject(float ScreenX, float ScreenY, FVector& WorldPosition, FVect
 	}
 }
 
+
+void AHUD::GetActorsInSelectionRectangle(TSubclassOf<class AActor> ClassFilter, const FVector2D& FirstPoint, const FVector2D& SecondPoint, TArray<AActor*>& OutActors, bool bActorMustBeFullyEnclosed)
+{
+	//Create Selection Rectangle from Points
+	FBox2D SelectionRectangle;
+
+	//This method ensures that an appropriate rectangle is generated, 
+	//		no matter what the coordinates of first and second point actually are.
+	SelectionRectangle += FirstPoint;
+	SelectionRectangle += SecondPoint;
+
+
+	//The Actor Bounds Point Mapping
+	const FVector BoundsPointMapping[8] =
+	{
+		FVector(1, 1, 1),
+		FVector(1, 1, -1),
+		FVector(1, -1, 1),
+		FVector(1, -1, -1),
+		FVector(-1, 1, 1),
+		FVector(-1, 1, -1),
+		FVector(-1, -1, 1),
+		FVector(-1, -1, -1)
+	};
+
+	//~~~
+
+	//For Each Actor of the Class Filter Type
+	for (TActorIterator<AActor> Itr(GetWorld(), ClassFilter); Itr; ++Itr)
+	{
+		AActor* EachActor = *Itr;
+
+		//Get Actor Bounds				//casting to base class, checked by template in the .h
+		const FBox EachActorBounds = Cast<AActor>(EachActor)->GetComponentsBoundingBox(true); /* All Components */
+
+		//Center
+		const FVector BoxCenter = EachActorBounds.GetCenter();
+
+		//Extents
+		const FVector BoxExtents = EachActorBounds.GetExtent();
+
+		// Build 2D bounding box of actor in screen space
+		FBox2D ActorBox2D(0);
+		for (uint8 BoundsPointItr = 0; BoundsPointItr < 8; BoundsPointItr++)
+		{
+			// Project vert into screen space.
+			const FVector ProjectedWorldLocation = Project(BoxCenter + (BoundsPointMapping[BoundsPointItr] * BoxExtents));
+			// Add to 2D bounding box
+			ActorBox2D += FVector2D(ProjectedWorldLocation.X, ProjectedWorldLocation.Y);
+		}
+
+		//Selection Box must fully enclose the Projected Actor Bounds
+		if (bActorMustBeFullyEnclosed)
+		{
+			if(SelectionRectangle.IsInside(ActorBox2D))
+			{
+				OutActors.Add(Cast<AActor>(EachActor));
+			}
+		}
+		//Partial Intersection with Projected Actor Bounds
+		else
+		{
+			if (SelectionRectangle.Intersect(ActorBox2D))
+			{
+				OutActors.Add(Cast<AActor>(EachActor));
+			}
+		}
+	}
+}
+
 void AHUD::DrawRect(FLinearColor Color, float ScreenX, float ScreenY, float Width, float Height)
 {
 	if (IsCanvasValid_WarnIfNot())	
