Add AddInstanceWorldSpace
https://github.com/EpicGames/UnrealEngine/pull/192 ([REDACTED_PROFILE_URL])

diff --git a/Engine/Source/Runtime/Engine/Classes/Components/InstancedStaticMeshComponent.h b/Engine/Source/Runtime/Engine/Classes/Components/InstancedStaticMeshComponent.h
--- a/Engine/Source/Runtime/Engine/Classes/Components/InstancedStaticMeshComponent.h
+++ b/Engine/Source/Runtime/Engine/Classes/Components/InstancedStaticMeshComponent.h
@@ -96,9 +96,14 @@ class ENGINE_API UInstancedStaticMeshComponent : public UStaticMeshComponent
 	UPROPERTY(EditAnywhere, BlueprintReadOnly, Category=Culling)
 	int32 InstanceEndCullDistance;
 
+	/** Add an instance to this component. Transform is given in local space of this component.  */
 	UFUNCTION(BlueprintCallable, Category="Components|InstancedStaticMesh")
 	void AddInstance(const FTransform& InstanceTransform);
 
+	/** Add an instance to this component. Transform is given in world space. */
+	UFUNCTION(BlueprintCallable, Category = "Components|InstancedStaticMesh")
+	 void AddInstanceWorldSpace(const FTransform& WorldTransform);
+
 	virtual bool ShouldCreatePhysicsState() const;
 
 	/** Clear all instances being rendered by this component */
diff --git a/Engine/Source/Runtime/Engine/Private/InstancedStaticMesh.cpp b/Engine/Source/Runtime/Engine/Private/InstancedStaticMesh.cpp
--- a/Engine/Source/Runtime/Engine/Private/InstancedStaticMesh.cpp
+++ b/Engine/Source/Runtime/Engine/Private/InstancedStaticMesh.cpp
@@ -1434,6 +1434,13 @@ void UInstancedStaticMeshComponent::AddInstance(const FTransform& InstanceTransf
 	MarkRenderStateDirty();
 }
 
+void UInstancedStaticMeshComponent::AddInstanceWorldSpace(const FTransform& WorldTransform)
+ {
+	// Transform from world space to local space
+	FTransform RelativeTM = WorldTransform.GetRelativeTransform(ComponentToWorld);
+	AddInstance(RelativeTM);
+}
+
 bool UInstancedStaticMeshComponent::ShouldCreatePhysicsState() const
 {
 	return IsRegistered() && (bAlwaysCreatePhysicsState || IsCollisionEnabled());
