Add PointsAreCoplanar (C++ and BP)
https://github.com/EpicGames/UnrealEngine/pull/193 ([REDACTED_PROFILE_URL])

diff --git a/Engine/Source/Runtime/Core/Private/Math/UnrealMath.cpp b/Engine/Source/Runtime/Core/Private/Math/UnrealMath.cpp
--- a/Engine/Source/Runtime/Core/Private/Math/UnrealMath.cpp
+++ b/Engine/Source/Runtime/Core/Private/Math/UnrealMath.cpp
@@ -1779,6 +1779,30 @@ bool FMath::GetDistanceWithinConeSegment(FVector Point, FVector ConeStartPoint,
 	return true;
 }
 
+bool FMath::PointsAreCoplanar(const TArray<FVector>& Points, const float Tolerance)
+{
+	//less than 4 points = coplanar
+	if (Points.Num() < 4)
+	{
+		return true;
+	}
+
+	//Get the Normal for plane determined by first 3 points
+	const FVector Normal = FVector::CrossProduct(Points[2] - Points[0], Points[1] - Points[0]).SafeNormal();
+
+	const int32 Total = Points.Num();
+	for (int32 v = 3; v < Total; v++)
+	{
+		//Abs of PointPlaneDist, dist should be 0
+		if (FMath::Abs(FVector::PointPlaneDist(Points[v], Points[0], Normal)) > Tolerance)
+		{
+			return false;
+		}
+	}
+
+	return true;
+}
+
 bool FMath::GetDotDistance
 ( 
 			FVector2D	&OutDotDist, 
diff --git a/Engine/Source/Runtime/Core/Public/Math/UnrealMathUtility.h b/Engine/Source/Runtime/Core/Public/Math/UnrealMathUtility.h
--- a/Engine/Source/Runtime/Core/Public/Math/UnrealMathUtility.h
+++ b/Engine/Source/Runtime/Core/Public/Math/UnrealMathUtility.h
@@ -819,6 +819,15 @@ struct FMath : public FPlatformMath
 	 */
 	static CORE_API bool GetDistanceWithinConeSegment(FVector Point, FVector ConeStartPoint, FVector ConeLine, float RadiusAtStart, float RadiusAtEnd, float &PercentageOut);
 
+	/**
+	 * Determines whether a given set of points are coplanar, with a tolerance. Any three points or less are always coplanar.
+	 *
+	 * @param Points - The set of points to determine coplanarity for.
+	 * @param Tolerance - Larger numbers means more variance is allowed.
+	 *
+	 * @return Whether the points are relatively coplanar, based on the tolerance
+	 */
+	static CORE_API bool PointsAreCoplanar(const TArray<FVector>& Points, const float Tolerance = 0.1f);
 
 	// Formatting functions
 
diff --git a/Engine/Source/Runtime/Engine/Classes/Kismet/KismetMathLibrary.h b/Engine/Source/Runtime/Engine/Classes/Kismet/KismetMathLibrary.h
--- a/Engine/Source/Runtime/Engine/Classes/Kismet/KismetMathLibrary.h
+++ b/Engine/Source/Runtime/Engine/Classes/Kismet/KismetMathLibrary.h
@@ -1012,6 +1012,17 @@ class UKismetMathLibrary : public UBlueprintFunctionLibrary
 	UFUNCTION(BlueprintCallable, BlueprintAuthorityOnly, Category="Math|Geometry", meta=(HidePin="WorldContextObject", DefaultToSelf="WorldContextObject"))
 	static void MinimumAreaRectangle(UObject* WorldContextObject, const TArray<FVector>& InVerts, const FVector& SampleSurfaceNormal, FVector& OutRectCenter, FRotator& OutRectRotation, float& OutSideLengthX, float& OutSideLengthY, bool bDebugDraw = false);
 
+	/**
+	 * Determines whether a given set of points are coplanar, with a tolerance. Any three points or less are always coplanar.
+	 *
+	 * @param Points - The set of points to determine coplanarity for.
+	 * @param Tolerance - Larger numbers means more variance is allowed.
+	 *
+	 * @return Whether the points are relatively coplanar, based on the tolerance
+	 */
+	UFUNCTION(BlueprintCallable, Category = "Math|Geometry")
+	static bool PointsAreCoplanar(const TArray<FVector>& Points, float Tolerance = 0.1f);
+
 	//
 	// Intersection
 	//
diff --git a/Engine/Source/Runtime/Engine/Private/KismetMathLibrary.cpp b/Engine/Source/Runtime/Engine/Private/KismetMathLibrary.cpp
--- a/Engine/Source/Runtime/Engine/Private/KismetMathLibrary.cpp
+++ b/Engine/Source/Runtime/Engine/Private/KismetMathLibrary.cpp
@@ -1402,6 +1402,11 @@ void UKismetMathLibrary::MinimumAreaRectangle(class UObject* WorldContextObject,
 	}
 }
 
+bool UKismetMathLibrary::PointsAreCoplanar(const TArray<FVector>& Points, float Tolerance)
+{
+	return FMath::PointsAreCoplanar(Points, Tolerance);
+}
+
 bool UKismetMathLibrary::LinePlaneIntersection(const FVector& LineStart, const FVector& LineEnd, const FPlane& APlane, float& T, FVector& Intersection)
 {
 	FVector RayDir = LineEnd - LineStart;
