- Supports convex collision volumes for interaction between the character and children clothing
- Added visualization for convex collision volumes

diff --git a/Engine/Source/Runtime/Engine/Private/SkeletalMeshComponentPhysics.cpp b/Engine/Source/Runtime/Engine/Private/SkeletalMeshComponentPhysics.cpp
--- a/Engine/Source/Runtime/Engine/Private/SkeletalMeshComponentPhysics.cpp
+++ b/Engine/Source/Runtime/Engine/Private/SkeletalMeshComponentPhysics.cpp
@@ -78,11 +78,14 @@ void USkeletalMesh::LoadClothCollisionVolumes(int32 AssetIndex, NxClothingAsset*
 	physx::PxI32 NumBoneActors;
 	verify(NxParameterized::getParamArraySize(*AssetParams, "boneActors", NumBoneActors));
 
+	// convexes are constructed with bone vertices
+	physx::PxI32 NumBoneVertices;
+	verify(NxParameterized::getParamArraySize(*AssetParams, "boneVertices", NumBoneVertices));
+
 	Asset.ClothCollisionVolumes.Empty(NumBoneActors);
 
 	char ParameterName[MAX_SPRINTF];
 
-	physx::PxMat44 PxLocalPose;
 	for(int32 i=0; i<NumBoneActors; i++)
 	{
 		FApexClothCollisionVolumeData& CollisionData = Asset.ClothCollisionVolumes[Asset.ClothCollisionVolumes.AddZeroed()];
@@ -93,8 +96,20 @@ void USkeletalMesh::LoadClothCollisionVolumes(int32 AssetIndex, NxClothingAsset*
 		verify(NxParameterized::getParamU32(*AssetParams, ParameterName, CollisionData.ConvexVerticesCount));
 		if(CollisionData.ConvexVerticesCount > 0)
 		{
+			CollisionData.BoneVertices.Empty(CollisionData.ConvexVerticesCount);
 			FCStringAnsi::Sprintf(ParameterName, "boneActors[%d].convexVerticesStart", i);
 			verify(NxParameterized::getParamU32(*AssetParams, ParameterName, CollisionData.ConvexVerticesStart));
+			// read vertex data which compose a convex
+			int32 NumMaxVertIndex = CollisionData.ConvexVerticesStart + CollisionData.ConvexVerticesCount;
+			check( NumMaxVertIndex <= NumBoneVertices );
+
+			for(int32 VertIdx=CollisionData.ConvexVerticesStart; VertIdx < NumMaxVertIndex; VertIdx++)
+			{
+				FCStringAnsi::Sprintf(ParameterName, "boneVertices[%d]", VertIdx);
+				physx::PxVec3 BoneVertex;
+				verify(NxParameterized::getParamVec3(*AssetParams, ParameterName, BoneVertex));
+				CollisionData.BoneVertices.Add(P2UVector(BoneVertex));
+			}
 		}
 		else
 		{
@@ -102,12 +117,13 @@ void USkeletalMesh::LoadClothCollisionVolumes(int32 AssetIndex, NxClothingAsset*
 			verify(NxParameterized::getParamF32(*AssetParams, ParameterName, CollisionData.CapsuleRadius));
 			FCStringAnsi::Sprintf(ParameterName, "boneActors[%d].capsuleHeight", i);
 			verify(NxParameterized::getParamF32(*AssetParams, ParameterName, CollisionData.CapsuleHeight));
+			// local pose is only used for a capsule
+			physx::PxMat44 PxLocalPose;
+		    FCStringAnsi::Sprintf(ParameterName, "boneActors[%d].localPose", i);
+		    verify(NxParameterized::getParamMat34(*AssetParams, ParameterName, PxLocalPose));
+    
+		    CollisionData.LocalPose = P2UMatrix(PxLocalPose);
 		}
-
-		FCStringAnsi::Sprintf(ParameterName, "boneActors[%d].localPose", i);
-		verify(NxParameterized::getParamMat34(*AssetParams, ParameterName, PxLocalPose));
-
-		CollisionData.LocalPose = P2UMatrix(PxLocalPose);
 	}
 
 	// load convex data
@@ -2443,7 +2459,7 @@ bool USkeletalMeshComponent::GetClothCollisionDataFromStaticMesh(UPrimitiveCompo
 
 void USkeletalMeshComponent::FindClothCollisions(TArray<FApexClothCollisionVolumeData>& OutCollisions)
 {
-	if(!SkeletalMesh)
+	if (!SkeletalMesh)
 	{
 		return;
 	}
@@ -2451,16 +2467,17 @@ void USkeletalMeshComponent::FindClothCollisions(TArray<FApexClothCollisionVolum
 	int32 NumAssets = SkeletalMesh->ClothingAssets.Num();
 
 	// find all new collisions passing to children
-	for(int32 AssetIdx=0; AssetIdx < NumAssets; AssetIdx++)
+	for (int32 AssetIdx = 0; AssetIdx < NumAssets; AssetIdx++)
 	{
 		FClothingAssetData& Asset = SkeletalMesh->ClothingAssets[AssetIdx];
 		int32 NumCollisions = Asset.ClothCollisionVolumes.Num();
 
-		for(int32 ColIdx=0; ColIdx < NumCollisions; ColIdx++)
+		int32 ConvexCount = 0;
+		for (int32 ColIdx = 0; ColIdx < NumCollisions; ColIdx++)
 		{
 			FApexClothCollisionVolumeData& Collision = Asset.ClothCollisionVolumes[ColIdx];
 
-			if(Collision.BoneIndex < 0)
+			if (Collision.BoneIndex < 0)
 			{
 				continue;
 			}
@@ -2469,7 +2486,7 @@ void USkeletalMeshComponent::FindClothCollisions(TArray<FApexClothCollisionVolum
 
 			int32 BoneIndex = GetBoneIndex(BoneName);
 
-			if(BoneIndex < 0)
+			if (BoneIndex < 0)
 			{
 				continue;
 			}
@@ -2478,19 +2495,43 @@ void USkeletalMeshComponent::FindClothCollisions(TArray<FApexClothCollisionVolum
 
 			FMatrix LocalToWorld = Collision.LocalPose * BoneMat;
 
-			// support only capsule now 
-			if(Collision.IsCapsule())
+			if (Collision.IsCapsule())
 			{
 				FApexClothCollisionVolumeData NewCollision = Collision;
 				NewCollision.LocalPose = LocalToWorld;
 				OutCollisions.Add(NewCollision);
 			}
+			else
+			{
+				if (Asset.ClothCollisionConvexPlaneIndices.IsValidIndex(ConvexCount))
+				{
+					uint32 PlaneIndices = Asset.ClothCollisionConvexPlaneIndices[ConvexCount];
+
+					uint32 BitShiftIndex = 1;
+					for (uint32 PlaneIndex = 0; PlaneIndex < 32; PlaneIndex++)
+					{
+						if (PlaneIndices&BitShiftIndex)
+						{
+							FPlane UPlane = Asset.ClothCollisionVolumePlanes[PlaneIndex].PlaneData;
+
+							check(Collision.BoneIndex == Asset.ClothCollisionVolumePlanes[PlaneIndex].BoneIndex);
+							UPlane = UPlane.TransformBy(BoneMat);
+							Collision.BonePlanes.Add(UPlane);
+						}
+
+						// shift 1 bit because a bit index means a plane's index 
+						BitShiftIndex <<= 1;
+					}
+				}
+
+				ConvexCount++;
 			}
 		}
 	}
+}
 
 void USkeletalMeshComponent::CreateInternalClothCollisions(TArray<FApexClothCollisionVolumeData>& InCollisions, TArray<physx::apex::NxClothingCollision*>& OutCollisions)
-	{
+{
 	int32 NumCollisions = InCollisions.Num();
 
 	const int32 MaxNumCapsules = 16;
@@ -2500,7 +2541,7 @@ void USkeletalMeshComponent::CreateInternalClothCollisions(TArray<FApexClothColl
 
 		for(int32 ActorIdx=0; ActorIdx < NumActors; ActorIdx++)
 		{
-		if (!IsValidClothingActor(ActorIdx))
+			if (!IsValidClothingActor(ActorIdx))
 			{
 				continue;
 			}
@@ -2510,26 +2551,51 @@ void USkeletalMeshComponent::CreateInternalClothCollisions(TArray<FApexClothColl
 
 			for(int32 ColIdx=0; ColIdx < NumCollisions; ColIdx++)
 			{
-				// support only capsules now 
-			if (InCollisions[ColIdx].IsCapsule() && NumCurrentCapsules < MaxNumCapsules)
+				// capsules
+				if (InCollisions[ColIdx].IsCapsule())
 				{
-				FVector Origin = InCollisions[ColIdx].LocalPose.GetOrigin();
-					// apex uses y-axis as the up-axis of capsule
-				FVector UpAxis = InCollisions[ColIdx].LocalPose.GetScaledAxis(EAxis::Y);
-					
-				float Radius = InCollisions[ColIdx].CapsuleRadius*UpAxis.Size();
+					if(NumCurrentCapsules < MaxNumCapsules)
+					{
+				        FVector Origin = InCollisions[ColIdx].LocalPose.GetOrigin();
+					        // apex uses y-axis as the up-axis of capsule
+				        FVector UpAxis = InCollisions[ColIdx].LocalPose.GetScaledAxis(EAxis::Y);
+					        
+				        float Radius = InCollisions[ColIdx].CapsuleRadius*UpAxis.Size();
+        
+				        float HalfHeight = InCollisions[ColIdx].CapsuleHeight*0.5f;
+					        const FVector TopEnd = Origin + (HalfHeight * UpAxis);
+					        const FVector BottomEnd = Origin - (HalfHeight * UpAxis);
+        
+					        NxClothingSphere* Sphere1 = Actor->createCollisionSphere(U2PVector(TopEnd), Radius);
+					        NxClothingSphere* Sphere2 = Actor->createCollisionSphere(U2PVector(BottomEnd), Radius);
+        
+					        NxClothingCapsule* Capsule = Actor->createCollisionCapsule(*Sphere1, *Sphere2);
+        
+				        OutCollisions.Add(Capsule);
+				        NumCurrentCapsules++;
+					}
+				}
+				else // convexes
+				{
+					int32 NumPlanes = InCollisions[ColIdx].BonePlanes.Num();
 
-				float HalfHeight = InCollisions[ColIdx].CapsuleHeight*0.5f;
-					const FVector TopEnd = Origin + (HalfHeight * UpAxis);
-					const FVector BottomEnd = Origin - (HalfHeight * UpAxis);
+					TArray<NxClothingPlane*> ClothingPlanes;
 
-					NxClothingSphere* Sphere1 = Actor->createCollisionSphere(U2PVector(TopEnd), Radius);
-					NxClothingSphere* Sphere2 = Actor->createCollisionSphere(U2PVector(BottomEnd), Radius);
+					//can not exceed 32 planes
+					NumPlanes = FMath::Min(NumPlanes, 32);
 
-					NxClothingCapsule* Capsule = Actor->createCollisionCapsule(*Sphere1, *Sphere2);
+					ClothingPlanes.AddUninitialized(NumPlanes);
 
-				OutCollisions.Add(Capsule);
-				NumCurrentCapsules++;
+					for(int32 PlaneIdx=0; PlaneIdx < NumPlanes; PlaneIdx++)
+					{
+						PxPlane PPlane = U2PPlane(InCollisions[ColIdx].BonePlanes[PlaneIdx]);
+
+						ClothingPlanes[PlaneIdx] = Actor->createCollisionPlane(PPlane);
+					}
+
+					NxClothingConvex* Convex = Actor->createCollisionConvex(ClothingPlanes.GetTypedData(), ClothingPlanes.Num());
+
+					OutCollisions.Add(Convex);
 				}
 			}
 		}
@@ -3646,6 +3712,8 @@ void USkeletalMeshComponent::DrawClothingCollisionVolumes(FPrimitiveDrawInterfac
 		return;
 	}
 
+	static bool bDrawnPlanes = false;
+
 	FColor Colors[3] = { FColor::Red, FColor::Green, FColor::Blue };
 	FColor GrayColor(50, 50, 50); // dark gray to represent ignored collisions
 
@@ -3692,8 +3760,60 @@ void USkeletalMeshComponent::DrawClothingCollisionVolumes(FPrimitiveDrawInterfac
 				DrawWireCapsule(PDI, LocalToWorld.GetOrigin(), LocalToWorld.GetUnitAxis(EAxis::X), LocalToWorld.GetUnitAxis(EAxis::Z), LocalToWorld.GetUnitAxis(EAxis::Y), CapsuleColor, Collisions[i].CapsuleRadius, CapsuleHalfHeight, CapsuleSides, SDPG_World);
 				SphereCount += 2; // 1 capsule takes 2 spheres
 			}
+			else // convex
+			{
+				int32 NumVerts = Collisions[i].BoneVertices.Num();
+
+				TArray<FVector> TransformedVerts;
+				TransformedVerts.AddUninitialized(NumVerts);
+
+				for(int32 VertIdx=0; VertIdx < NumVerts; VertIdx++)
+				{
+					TransformedVerts[VertIdx] = BoneMat.TransformPosition(Collisions[i].BoneVertices[VertIdx]);
+				}
+
+				// just draw all connected wires to check convex shape
+				for(int32 i=0; i < NumVerts-2; i++)
+				{
+					for(int32 j=i+1; j < NumVerts; j++)
+					{
+						PDI->DrawLine(TransformedVerts[i], TransformedVerts[j], Colors[AssetIdx%3], SDPG_World);
+					}
+				}
+			}
+		}
+
+		TArray<FClothBonePlane>& BonePlanes = SkeletalMesh->ClothingAssets[AssetIdx].ClothCollisionVolumePlanes;
+		int32 NumPlanes = BonePlanes.Num();
+
+		FVector Origin = ComponentToWorld.GetLocation();
+
+
+		for(int32 PlaneIdx=0; PlaneIdx < NumPlanes; PlaneIdx++)
+		{
+			FName BoneName = SkeletalMesh->ClothingAssets[AssetIdx].ApexClothingAsset->GetConvertedBoneName(BonePlanes[PlaneIdx].BoneIndex);
+
+			int32 BoneIndex = GetBoneIndex(BoneName);
+
+			if(BoneIndex < 0)
+			{
+				continue;
+			}
+
+			FMatrix BoneMat = GetBoneMatrix(BoneIndex);
+			FPlane UPlane = BonePlanes[PlaneIdx].PlaneData.TransformBy(BoneMat);
+			FVector BoneMatOrigin = BoneMat.GetOrigin();
+
+			// draw once becasue of a bug in DrawDebugSolidPlane() method
+			if(!bDrawnPlanes)
+			{
+				DrawDebugSolidPlane(GetWorld(), UPlane, BoneMatOrigin, 2, Colors[AssetIdx%3]);
 			}
 
+		}
+
+		bDrawnPlanes = true;
+
 		// draw bone spheres
 		TArray<FApexClothBoneSphereData>& Spheres = SkeletalMesh->ClothingAssets[AssetIdx].ClothBoneSpheres;
 		TArray<FVector> SpherePositions;
diff --git a/Engine/Source/Runtime/Engine/Public/SkeletalMeshTypes.h b/Engine/Source/Runtime/Engine/Public/SkeletalMeshTypes.h
--- a/Engine/Source/Runtime/Engine/Public/SkeletalMeshTypes.h
+++ b/Engine/Source/Runtime/Engine/Public/SkeletalMeshTypes.h
@@ -632,8 +632,12 @@ struct FApexClothCollisionVolumeData
 	*/
 
 	int32	BoneIndex;
+	// for convexes
 	uint32	ConvexVerticesCount;
 	uint32	ConvexVerticesStart;
+	TArray<FVector> BoneVertices;
+	TArray<FPlane>  BonePlanes;
+	// for capsules
 	float	CapsuleRadius;
 	float	CapsuleHeight;
 	FMatrix LocalPose;
@@ -646,7 +650,7 @@ struct FApexClothCollisionVolumeData
 		CapsuleRadius = 0.0f;
 		CapsuleHeight = 0.0f;
 		LocalPose.SetIdentity();
-	};
+	}
 
 	bool IsCapsule()
 	{
