Fix GetActorsInSelectionRectangle template compile error

diff --git a/Engine/Source/Runtime/Engine/Classes/GameFramework/HUD.h b/Engine/Source/Runtime/Engine/Classes/GameFramework/HUD.h
--- a/Engine/Source/Runtime/Engine/Classes/GameFramework/HUD.h
+++ b/Engine/Source/Runtime/Engine/Classes/GameFramework/HUD.h
@@ -429,7 +429,7 @@ public:
 
 		//Run Inner Function, output to Base AActor Array
 		TArray<AActor*> OutActorsBaseArray;
-		GetActorsInSelectionRectangle(ClassFilter, FirstPoint, SecondPoint, OutActorsBaseArray, bActorMustBeFullyEnclosed);
+		GetActorsInSelectionRectangle(ClassFilter::StaticClass(), FirstPoint, SecondPoint, OutActorsBaseArray, bActorMustBeFullyEnclosed);
 
 		//Construct casted template type array
 		for (AActor* EachActor : OutActorsBaseArray)
