Add 'Random Point in Box' (C++ & BP)
https://github.com/EpicGames/UnrealEngine/pull/187 ([REDACTED_PROFILE_URL])

diff --git a/Engine/Source/Runtime/Core/Private/Math/UnrealMath.cpp b/Engine/Source/Runtime/Core/Private/Math/UnrealMath.cpp
--- a/Engine/Source/Runtime/Core/Private/Math/UnrealMath.cpp
+++ b/Engine/Source/Runtime/Core/Private/Math/UnrealMath.cpp
@@ -2207,6 +2207,13 @@ FVector FMath::VRandCone(FVector const& Dir, float HorizontalConeHalfAngleRad, f
 	}
 }
 
+FVector FMath::RandPointInBox(const FBox& Box)
+{
+	return FVector(	FRandRange(Box.Min.X, Box.Max.X),
+					FRandRange(Box.Min.Y, Box.Max.Y),
+					FRandRange(Box.Min.Z, Box.Max.Z) );
+}
+
 struct FClusterMovedHereToMakeCompile
 {
 	FVector ClusterPosAccum;
diff --git a/Engine/Source/Runtime/Core/Public/Math/UnrealMathUtility.h b/Engine/Source/Runtime/Core/Public/Math/UnrealMathUtility.h
--- a/Engine/Source/Runtime/Core/Public/Math/UnrealMathUtility.h
+++ b/Engine/Source/Runtime/Core/Public/Math/UnrealMathUtility.h
@@ -127,6 +127,8 @@ struct FMath : public FPlatformMath
 	 */
 	static CORE_API FVector VRandCone(FVector const& Dir, float HorizontalConeHalfAngleRad, float VerticalConeHalfAngleRad);
 
+	/** Returns a random point within the passed in bounding box */
+	static CORE_API FVector RandPointInBox(const FBox& Box);
 
 	// Predicates
 
diff --git a/Engine/Source/Runtime/Engine/Classes/Kismet/KismetMathLibrary.h b/Engine/Source/Runtime/Engine/Classes/Kismet/KismetMathLibrary.h
--- a/Engine/Source/Runtime/Engine/Classes/Kismet/KismetMathLibrary.h
+++ b/Engine/Source/Runtime/Engine/Classes/Kismet/KismetMathLibrary.h
@@ -470,6 +470,9 @@ class UKismetMathLibrary : public UBlueprintFunctionLibrary
 	UFUNCTION(BlueprintPure, Category="Math|Random")
 	static FVector RandomUnitVector();
 
+	/** Returns a random point within the specified bounding box */
+	UFUNCTION(BlueprintPure, Category = "Math|Random")
+	static FVector RandomPointInBoundingBox(FBox BoundingBox);
 	/* 
 	 * Returns a random vector with length of 1, within the specified cone, with uniform random distribution. 
 	 * @param ConeDir	The base "center" direction of the cone.
diff --git a/Engine/Source/Runtime/Engine/Private/KismetMathLibrary.cpp b/Engine/Source/Runtime/Engine/Private/KismetMathLibrary.cpp
--- a/Engine/Source/Runtime/Engine/Private/KismetMathLibrary.cpp
+++ b/Engine/Source/Runtime/Engine/Private/KismetMathLibrary.cpp
@@ -659,6 +659,11 @@ FVector UKismetMathLibrary::RandomUnitVector()
 	return FMath::VRand();
 }
 
+FVector UKismetMathLibrary::RandomPointInBoundingBox(FBox BoundingBox)
+{
+	return FMath::RandPointInBox(BoundingBox);
+}
+
 FVector UKismetMathLibrary::RandomUnitVectorInCone(FVector ConeDir, float ConeHalfAngle)
 {
 	return FMath::VRandCone(ConeDir, ConeHalfAngle);
