Pool the dynamic mesh builder's vertex and index buffers to avoid performance problems under OS X in very dense projects.
Uses the existing ResourcePool to trade memory for a reduction in the amount of buffer churn, which prevents validation of the GL buffer name or allocation of buffer backing store dominating execution time in large projects on OS X.

diff --git a/Engine/Source/Runtime/Engine/Private/DynamicMeshBuilder.cpp b/Engine/Source/Runtime/Engine/Private/DynamicMeshBuilder.cpp
--- a/Engine/Source/Runtime/Engine/Private/DynamicMeshBuilder.cpp
+++ b/Engine/Source/Runtime/Engine/Private/DynamicMeshBuilder.cpp
@@ -6,6 +6,174 @@
 
 #include "EnginePrivate.h"
 #include "DynamicMeshBuilder.h"
+#include "ResourcePool.h"
+
+class FGlobalDynamicMeshPoolPolicy
+{
+public:
+	/** Buffers are created with a simple byte size */
+	typedef uint32 CreationArguments;
+	enum
+	{
+		NumSafeFrames = 3, /** Number of frames to leaves buffers before reclaiming/reusing */
+		NumPoolBucketSizes = 16, /** Number of pool buckets */
+		NumToDrainPerFrame = 100, /** Max. number of resources to cull in a single frame */
+		CullAfterFramesNum = 10 /** Resources are culled if unused for more frames than this */
+	};
+	
+	/** Get the pool bucket index from the size
+	 * @param Size the number of bytes for the resource
+	 * @returns The bucket index.
+	 */
+	uint32 GetPoolBucketIndex(uint32 Size)
+	{
+		unsigned long Lower = 0;
+		unsigned long Upper = NumPoolBucketSizes;
+		unsigned long Middle;
+		
+		do
+		{
+			Middle = ( Upper + Lower ) >> 1;
+			if( Size <= BucketSizes[Middle-1] )
+			{
+				Upper = Middle;
+			}
+			else
+			{
+				Lower = Middle;
+			}
+		}
+		while( Upper - Lower > 1 );
+		
+		check( Size <= BucketSizes[Lower] );
+		check( (Lower == 0 ) || ( Size > BucketSizes[Lower-1] ) );
+		
+		return Lower;
+	}
+	
+	/** Get the pool bucket size from the index
+	 * @param Bucket the bucket index
+	 * @returns The bucket size.
+	 */
+	uint32 GetPoolBucketSize(uint32 Bucket)
+	{
+		check(Bucket < NumPoolBucketSizes);
+		return BucketSizes[Bucket];
+	}
+	
+private:
+	/** The bucket sizes */
+	static uint32 BucketSizes[NumPoolBucketSizes];
+};
+
+uint32 FGlobalDynamicMeshPoolPolicy::BucketSizes[NumPoolBucketSizes] = {
+	64, 128, 256, 512, 1024, 2048, 4096,
+	8*1024, 16*1024, 32*1024, 64*1024, 128*1024, 256*1024,
+	512*1024, 1*1024*1024, 2*1024*1024
+};
+
+class FGlobalDynamicMeshIndexPolicy : public FGlobalDynamicMeshPoolPolicy
+{
+public:
+	enum
+	{
+		NumSafeFrames = FGlobalDynamicMeshPoolPolicy::NumSafeFrames,
+		NumPoolBuckets = FGlobalDynamicMeshPoolPolicy::NumPoolBucketSizes,
+		NumToDrainPerFrame = FGlobalDynamicMeshPoolPolicy::NumToDrainPerFrame,
+		CullAfterFramesNum = FGlobalDynamicMeshPoolPolicy::CullAfterFramesNum
+	};
+	
+	/** Creates the resource
+	 * @param Args The buffer size in bytes.
+	 * @returns A suitably sized buffer or NULL on failure.
+	 */
+	FIndexBufferRHIRef CreateResource(FGlobalDynamicMeshPoolPolicy::CreationArguments Args)
+	{
+		FGlobalDynamicMeshPoolPolicy::CreationArguments BufferSize = GetPoolBucketSize(GetPoolBucketIndex(Args));
+		// The use of BUF_Static is deliberate - on OS X the buffer backing-store orphaning & reallocation will dominate execution time
+		// so to avoid this we don't reuse a buffer for several frames, thereby avoiding the pipeline stall and the reallocation cost.
+		FRHIResourceCreateInfo CreateInfo;
+		FIndexBufferRHIRef VertexBuffer = RHICreateIndexBuffer(sizeof(uint32), BufferSize, BUF_Static, CreateInfo);
+		return VertexBuffer;
+	}
+	
+	/** Gets the arguments used to create resource
+	 * @param Resource The buffer to get data for.
+	 * @returns The arguments used to create the buffer.
+	 */
+	FGlobalDynamicMeshPoolPolicy::CreationArguments GetCreationArguments(FIndexBufferRHIRef Resource)
+	{
+		return (Resource->GetSize());
+	}
+};
+
+class FGlobalDynamicMeshIndexPool : public TRenderResourcePool<FIndexBufferRHIRef, FGlobalDynamicMeshIndexPolicy, FGlobalDynamicMeshPoolPolicy::CreationArguments>
+{
+public:
+	/** Destructor */
+	virtual ~FGlobalDynamicMeshIndexPool()
+	{
+	}
+	
+public: // From FTickableObjectRenderThread
+	virtual TStatId GetStatId() const OVERRIDE
+	{
+		RETURN_QUICK_DECLARE_CYCLE_STAT(FGlobalDynamicMeshIndexPool, STATGROUP_Tickables);
+	}
+};
+TGlobalResource<FGlobalDynamicMeshIndexPool> GDynamicMeshIndexPool;
+
+class FGlobalDynamicMeshVertexPolicy : public FGlobalDynamicMeshPoolPolicy
+{
+public:
+	enum
+	{
+		NumSafeFrames = FGlobalDynamicMeshPoolPolicy::NumSafeFrames,
+		NumPoolBuckets = FGlobalDynamicMeshPoolPolicy::NumPoolBucketSizes,
+		NumToDrainPerFrame = FGlobalDynamicMeshPoolPolicy::NumToDrainPerFrame,
+		CullAfterFramesNum = FGlobalDynamicMeshPoolPolicy::CullAfterFramesNum
+	};
+	
+	/** Creates the resource
+	 * @param Args The buffer size in bytes.
+	 * @returns A suitably sized buffer or NULL on failure.
+	 */
+	FVertexBufferRHIRef CreateResource(FGlobalDynamicMeshPoolPolicy::CreationArguments Args)
+	{
+		FGlobalDynamicMeshPoolPolicy::CreationArguments BufferSize = GetPoolBucketSize(GetPoolBucketIndex(Args));
+		// The use of BUF_Static is deliberate - on OS X the buffer backing-store orphaning & reallocation will dominate execution time
+		// so to avoid this we don't reuse a buffer for several frames, thereby avoiding the pipeline stall and the reallocation cost.
+		FRHIResourceCreateInfo CreateInfo;
+		FVertexBufferRHIRef VertexBuffer = RHICreateVertexBuffer(BufferSize, BUF_Static, CreateInfo);
+		return VertexBuffer;
+	}
+	
+	/** Gets the arguments used to create resource
+	 * @param Resource The buffer to get data for.
+	 * @returns The arguments used to create the buffer.
+	 */
+	FGlobalDynamicMeshPoolPolicy::CreationArguments GetCreationArguments(FVertexBufferRHIRef Resource)
+	{
+		return (Resource->GetSize());
+	}
+};
+
+class FGlobalDynamicMeshVertexPool : public TRenderResourcePool<FVertexBufferRHIRef, FGlobalDynamicMeshVertexPolicy, FGlobalDynamicMeshPoolPolicy::CreationArguments>
+{
+public:
+	/** Destructor */
+	virtual ~FGlobalDynamicMeshVertexPool()
+	{
+	}
+	
+public: // From FTickableObjectRenderThread
+	virtual TStatId GetStatId() const OVERRIDE
+	{
+		RETURN_QUICK_DECLARE_CYCLE_STAT(FGlobalDynamicMeshVertexPool, STATGROUP_Tickables);
+	}
+};
+TGlobalResource<FGlobalDynamicMeshVertexPool> GDynamicMeshVertexPool;
+
 
 /** The index buffer type used for dynamic meshes. */
 class FDynamicMeshIndexBuffer : public FDynamicPrimitiveResource, public FIndexBuffer
@@ -17,14 +185,32 @@ public:
 	// FRenderResource interface.
 	virtual void InitRHI()
 	{
-		FRHIResourceCreateInfo CreateInfo;
-		IndexBufferRHI = RHICreateIndexBuffer(sizeof(int32),Indices.Num() * sizeof(int32),BUF_Static,CreateInfo);
-
+		uint32 SizeInBytes = Indices.Num() * sizeof(int32);
+		if(SizeInBytes <= FGlobalDynamicMeshIndexPolicy().GetPoolBucketSize(FGlobalDynamicMeshIndexPolicy::NumPoolBuckets - 1))
+		{
+			IndexBufferRHI = GDynamicMeshIndexPool.CreatePooledResource(SizeInBytes);
+		}
+		else
+		{
+			FRHIResourceCreateInfo CreateInfo;
+			IndexBufferRHI = RHICreateIndexBuffer(sizeof(uint32), SizeInBytes, BUF_Volatile, CreateInfo);
+		}
+		
 		// Write the indices to the index buffer.
 		void* Buffer = RHILockIndexBuffer(IndexBufferRHI,0,Indices.Num() * sizeof(int32),RLM_WriteOnly);
 		FMemory::Memcpy(Buffer,Indices.GetTypedData(),Indices.Num() * sizeof(int32));
 		RHIUnlockIndexBuffer(IndexBufferRHI);
 	}
+	
+	virtual void ReleaseRHI()
+	{
+		if(IndexBufferRHI->GetSize() <= FGlobalDynamicMeshIndexPolicy().GetPoolBucketSize(FGlobalDynamicMeshIndexPolicy::NumPoolBuckets - 1))
+		{
+			GDynamicMeshIndexPool.ReleasePooledResource(IndexBufferRHI);
+			IndexBufferRHI = NULL;
+		}
+		FIndexBuffer::ReleaseRHI();
+	}
 
 	// FDynamicPrimitiveResource interface.
 	virtual void InitPrimitiveResource()
@@ -45,17 +231,35 @@ public:
 
 	TArray<FDynamicMeshVertex> Vertices;
 
-	// FResourceResource interface.
+	// FRenderResource interface.
 	virtual void InitRHI()
 	{
-		FRHIResourceCreateInfo CreateInfo;
-		VertexBufferRHI = RHICreateVertexBuffer(Vertices.Num() * sizeof(FDynamicMeshVertex),BUF_Static, CreateInfo);
-
+		uint32 SizeInBytes = Vertices.Num() * sizeof(FDynamicMeshVertex);
+		if(SizeInBytes <= FGlobalDynamicMeshIndexPolicy().GetPoolBucketSize(FGlobalDynamicMeshIndexPolicy::NumPoolBuckets - 1))
+		{
+			VertexBufferRHI = GDynamicMeshVertexPool.CreatePooledResource(SizeInBytes);
+		}
+		else
+		{
+			FRHIResourceCreateInfo CreateInfo;
+			VertexBufferRHI = RHICreateVertexBuffer(SizeInBytes, BUF_Volatile, CreateInfo);
+		}
+		
 		// Copy the vertex data into the vertex buffer.
-		void* VertexBufferData = RHILockVertexBuffer(VertexBufferRHI,0,Vertices.Num() * sizeof(FDynamicMeshVertex),RLM_WriteOnly);
-		FMemory::Memcpy(VertexBufferData,Vertices.GetTypedData(),Vertices.Num() * sizeof(FDynamicMeshVertex));
+		void* VertexBufferData = RHILockVertexBuffer(VertexBufferRHI,0,SizeInBytes,RLM_WriteOnly);
+		FMemory::Memcpy(VertexBufferData,Vertices.GetTypedData(),SizeInBytes);
 		RHIUnlockVertexBuffer(VertexBufferRHI);
 	}
+	
+	virtual void ReleaseRHI()
+	{
+		if(VertexBufferRHI->GetSize() <= FGlobalDynamicMeshVertexPolicy().GetPoolBucketSize(FGlobalDynamicMeshVertexPolicy::NumPoolBuckets - 1))
+		{
+			GDynamicMeshVertexPool.ReleasePooledResource(VertexBufferRHI);
+			VertexBufferRHI = NULL;
+		}
+		FVertexBuffer::ReleaseRHI();
+	}
 
 	// FDynamicPrimitiveResource interface.
 	virtual void InitPrimitiveResource()
