Made AIModule not include Engine.h #UE4

diff --git a/Engine/Source/Runtime/AIModule/AIModule.Build.cs b/Engine/Source/Runtime/AIModule/AIModule.Build.cs
--- a/Engine/Source/Runtime/AIModule/AIModule.Build.cs
+++ b/Engine/Source/Runtime/AIModule/AIModule.Build.cs
@@ -11,7 +11,6 @@ namespace UnrealBuildTool.Rules
 			PublicIncludePaths.AddRange(
 				new string[] {
                     "Runtime/AIModule/Public",
-					// ... add public include paths required here ...
 				}
 				);
 
@@ -19,7 +18,6 @@ namespace UnrealBuildTool.Rules
 				new string[] {
 					"Runtime/AIModule/Private",
                     "Runtime/Engine/Private",
-					// ... add other private include paths required here ...
 				}
 				);
 
@@ -28,15 +26,14 @@ namespace UnrealBuildTool.Rules
 				{
 					"Core",
 					"CoreUObject",
-					"Engine"                    
-					// ... add other public dependencies that you statically link with here ...
+					"Engine",
 				}
 				);
 
 			PrivateDependencyModuleNames.AddRange(
 				new string[]
 				{
-					// ... add private dependencies that you statically link with here ...
+					"RHI",
 				}
 				);
 
diff --git a/Engine/Source/Runtime/AIModule/Classes/AISystem.h b/Engine/Source/Runtime/AIModule/Classes/AISystem.h
--- a/Engine/Source/Runtime/AIModule/Classes/AISystem.h
+++ b/Engine/Source/Runtime/AIModule/Classes/AISystem.h
@@ -2,10 +2,11 @@
 
 #pragma once
 #include "Misc/CoreMisc.h"
+#include "EngineDefines.h"
+#include "Engine/World.h"
 #include "AI/AISystemBase.h"
 #include "AISystem.generated.h"
 
-class UWorld;
 class UBehaviorTreeManager;
 class UEnvQueryManager;
 
diff --git a/Engine/Source/Runtime/AIModule/Classes/BehaviorTree/BlackboardData.h b/Engine/Source/Runtime/AIModule/Classes/BehaviorTree/BlackboardData.h
--- a/Engine/Source/Runtime/AIModule/Classes/BehaviorTree/BlackboardData.h
+++ b/Engine/Source/Runtime/AIModule/Classes/BehaviorTree/BlackboardData.h
@@ -1,5 +1,7 @@
 // Copyright 1998-2014 Epic Games, Inc. All Rights Reserved.
 #pragma once
+
+#include "Engine/DataAsset.h"
 #include "BlackboardData.generated.h"
 
 /** blackboard entry definition */
diff --git a/Engine/Source/Runtime/AIModule/Classes/BehaviorTree/Tasks/BTTask_PlaySound.h b/Engine/Source/Runtime/AIModule/Classes/BehaviorTree/Tasks/BTTask_PlaySound.h
--- a/Engine/Source/Runtime/AIModule/Classes/BehaviorTree/Tasks/BTTask_PlaySound.h
+++ b/Engine/Source/Runtime/AIModule/Classes/BehaviorTree/Tasks/BTTask_PlaySound.h
@@ -3,6 +3,9 @@
 #pragma once
 #include "BTTask_PlaySound.generated.h"
 
+class UBehaviorTreeComponent;
+class USoundCue;
+
 UCLASS()
 class UBTTask_PlaySound : public UBTTaskNode
 {
@@ -10,7 +13,7 @@ class UBTTask_PlaySound : public UBTTaskNode
 
 	/** CUE to play */
 	UPROPERTY(Category=Node, EditAnywhere)
-	class USoundCue* SoundToPlay;
+	USoundCue* SoundToPlay;
 
 	virtual EBTNodeResult::Type ExecuteTask(UBehaviorTreeComponent* OwnerComp, uint8* NodeMemory) OVERRIDE;
 	virtual FString GetStaticDescription() const OVERRIDE;
diff --git a/Engine/Source/Runtime/AIModule/Classes/EnvironmentQuery/EnvQueryTypes.h b/Engine/Source/Runtime/AIModule/Classes/EnvironmentQuery/EnvQueryTypes.h
--- a/Engine/Source/Runtime/AIModule/Classes/EnvironmentQuery/EnvQueryTypes.h
+++ b/Engine/Source/Runtime/AIModule/Classes/EnvironmentQuery/EnvQueryTypes.h
@@ -20,6 +20,8 @@ DECLARE_DWORD_COUNTER_STAT_EXTERN(TEXT("Num Instances"),STAT_AI_EQS_NumInstances
 DECLARE_DWORD_COUNTER_STAT_EXTERN(TEXT("Num Items"),STAT_AI_EQS_NumItems,STATGROUP_AI_EQS, );
 DECLARE_MEMORY_STAT_EXTERN(TEXT("Instance memory"),STAT_AI_EQS_InstanceMemory,STATGROUP_AI_EQS, AIMODULE_API);
 
+class ARecastNavMesh;
+
 UENUM()
 namespace EEnvTestPurpose
 {
@@ -913,7 +915,7 @@ public:
 namespace FEQSHelpers
 {
 #if WITH_RECAST
-	const class ARecastNavMesh* FindNavMeshForQuery(struct FEnvQueryInstance& QueryInstance);
+	const ARecastNavMesh* FindNavMeshForQuery(struct FEnvQueryInstance& QueryInstance);
 #endif // WITH_RECAST
 }
 
diff --git a/Engine/Source/Runtime/AIModule/Classes/EnvironmentQuery/Generators/EnvQueryGenerator_ProjectedPoints.h b/Engine/Source/Runtime/AIModule/Classes/EnvironmentQuery/Generators/EnvQueryGenerator_ProjectedPoints.h
--- a/Engine/Source/Runtime/AIModule/Classes/EnvironmentQuery/Generators/EnvQueryGenerator_ProjectedPoints.h
+++ b/Engine/Source/Runtime/AIModule/Classes/EnvironmentQuery/Generators/EnvQueryGenerator_ProjectedPoints.h
@@ -1,6 +1,7 @@
 // Copyright 1998-2013 Epic Games, Inc. All Rights Reserved.
 
 #pragma once
+#include "AI/Navigation/RecastNavMesh.h"
 #include "EnvironmentQuery/EnvQueryGenerator.h"
 #include "EnvironmentQuery/EnvQueryTypes.h"
 #include "EnvQueryGenerator_ProjectedPoints.generated.h"
diff --git a/Engine/Source/Runtime/AIModule/Classes/Navigation/CrowdManager.h b/Engine/Source/Runtime/AIModule/Classes/Navigation/CrowdManager.h
--- a/Engine/Source/Runtime/AIModule/Classes/Navigation/CrowdManager.h
+++ b/Engine/Source/Runtime/AIModule/Classes/Navigation/CrowdManager.h
@@ -4,6 +4,8 @@
 #include "Tickable.h"
 #include "CrowdManager.generated.h"
 
+struct FNavMeshPath;
+class UCrowdFollowingComponent;
 /**
  *  Crowd manager is responsible for handling crowds using Detour (Recast library)
  *
@@ -125,22 +127,22 @@ class AIMODULE_API UCrowdManager : public UObject
 	void UpdateAgentState(const class ICrowdAgentInterface* Agent) const;
 
 	/** sets move target for crowd agent (only for fully simulated) */
-	bool SetAgentMoveTarget(const class UCrowdFollowingComponent* AgentComponent, const FVector& MoveTarget, TSharedPtr<const FNavigationQueryFilter> Filter) const;
+	bool SetAgentMoveTarget(const UCrowdFollowingComponent* AgentComponent, const FVector& MoveTarget, TSharedPtr<const FNavigationQueryFilter> Filter) const;
 
 	/** sets move direction for crowd agent (only for fully simulated) */
-	bool SetAgentMoveDirection(const class UCrowdFollowingComponent* AgentComponent, const FVector& MoveDirection) const;
+	bool SetAgentMoveDirection(const UCrowdFollowingComponent* AgentComponent, const FVector& MoveDirection) const;
 
 	/** sets move target using path (only for fully simulated) */
-	bool SetAgentMovePath(const class UCrowdFollowingComponent* AgentComponent, const FNavMeshPath* Path, int32 PathSectionStart, int32 PathSectionEnd) const;
+	bool SetAgentMovePath(const UCrowdFollowingComponent* AgentComponent, const FNavMeshPath* Path, int32 PathSectionStart, int32 PathSectionEnd) const;
 
 	/** clears move target for crowd agent (only for fully simulated) */
-	void ClearAgentMoveTarget(const class UCrowdFollowingComponent* AgentComponent) const;
+	void ClearAgentMoveTarget(const UCrowdFollowingComponent* AgentComponent) const;
 
 	/** switch agent to waiting state */
-	void PauseAgent(const class UCrowdFollowingComponent* AgentComponent) const;
+	void PauseAgent(const UCrowdFollowingComponent* AgentComponent) const;
 
 	/** resumes agent movement */
-	void ResumeAgent(const class UCrowdFollowingComponent* AgentComponent, bool bForceReplanPath = true) const;
+	void ResumeAgent(const UCrowdFollowingComponent* AgentComponent, bool bForceReplanPath = true) const;
 
 	/** returns number of nearby agents */
 	int32 GetNumNearbyAgents(const class ICrowdAgentInterface* Agent) const;
@@ -158,7 +160,7 @@ class AIMODULE_API UCrowdManager : public UObject
 	void SetSingleAreaVisibilityOptimization(bool bEnable);
 
 	/** adjust current position in path's corridor, starting test from PathStartIdx */
-	void AdjustAgentPathStart(const class UCrowdFollowingComponent* AgentComponent, const FNavMeshPath* Path, int32& PathStartIdx) const;
+	void AdjustAgentPathStart(const UCrowdFollowingComponent* AgentComponent, const FNavMeshPath* Path, int32& PathStartIdx) const;
 
 #if WITH_EDITOR
 	virtual void PostEditChangeProperty(struct FPropertyChangedEvent& PropertyChangedEvent) OVERRIDE;
@@ -232,7 +234,7 @@ protected:
 	void UpdateSpatialData(const class ICrowdAgentInterface* Agent, const FCrowdAgentData& AgentData) const;
 
 	/** pass new velocity to movement components */
-	void ApplyVelocity(const class UCrowdFollowingComponent* AgentComponent, int32 AgentIndex) const;
+	void ApplyVelocity(const UCrowdFollowingComponent* AgentComponent, int32 AgentIndex) const;
 
 	/** check changes in crowd simulation and adjust UE specific properties (smart links, poly updates) */
 	void UpdateAgentPaths();
diff --git a/Engine/Source/Runtime/AIModule/Private/AIController.cpp b/Engine/Source/Runtime/AIModule/Private/AIController.cpp
--- a/Engine/Source/Runtime/AIModule/Private/AIController.cpp
+++ b/Engine/Source/Runtime/AIModule/Private/AIController.cpp
@@ -11,6 +11,8 @@
 #include "Navigation/PathFollowingComponent.h"
 #include "Navigation/NavigationComponent.h"
 #include "GameFramework/Pawn.h"
+#include "Engine/Canvas.h"
+#include "GameFramework/PhysicsVolume.h"
 #include "AIController.h"
 
 // mz@todo these need to be removed, legacy code
@@ -72,7 +74,7 @@ void AAIController::Reset()
 	}
 }
 
-void AAIController::DisplayDebug(class UCanvas* Canvas, const FDebugDisplayInfo& DebugDisplay, float& YL, float& YPos)
+void AAIController::DisplayDebug(UCanvas* Canvas, const FDebugDisplayInfo& DebugDisplay, float& YL, float& YPos)
 {
 	Super::DisplayDebug(Canvas, DebugDisplay, YL, YPos);
 
diff --git a/Engine/Source/Runtime/AIModule/Private/AIModulePrivate.h b/Engine/Source/Runtime/AIModule/Private/AIModulePrivate.h
--- a/Engine/Source/Runtime/AIModule/Private/AIModulePrivate.h
+++ b/Engine/Source/Runtime/AIModule/Private/AIModulePrivate.h
@@ -2,15 +2,36 @@
 
 #pragma once
 
-//#include "CoreUObject.h"
-//#include "Engine/EngineBaseTypes.h"
-//#include "Engine/EngineTypes.h"
-//// @todo consider making this non-shippable
-//#include "DisplayDebugHelpers.h"
-//#include "InputCoreTypes.h"
-//#include "GameFramework/Actor.h"
-#include "Engine.h"
+//----------------------------------------------------------------------//
+// Engine headers
+//----------------------------------------------------------------------//
 
+#include "Core.h"
+#include "CoreUObject.h"
+#include "InputCore.h"
+#include "EngineDefines.h"
+#include "EngineSettings.h"
+#include "EngineStats.h"
+#include "EngineLogs.h"
+#include "EngineGlobals.h"
+
+#include "Engine/EngineBaseTypes.h"
+#include "Engine/EngineTypes.h"
+#include "Components/ActorComponent.h"
+#include "GameFramework/Actor.h"
+#include "CollisionQueryParams.h"
+#include "WorldCollision.h"
+#include "Components/PrimitiveComponent.h"
+#include "GameFramework/PlayerController.h"
+#include "AI/Navigation/RecastNavMesh.h"
+#include "GameFramework/Pawn.h"
+#include "Components/CapsuleComponent.h"
+#include "Engine/Engine.h"
+#include "TimerManager.h"					// Game play timers
+
+//----------------------------------------------------------------------//
+// AIModule headers
+//----------------------------------------------------------------------//
 #include "AIModule.h"
 
 #include "AITypes.h"
diff --git a/Engine/Source/Runtime/AIModule/Private/AISystem.cpp b/Engine/Source/Runtime/AIModule/Private/AISystem.cpp
--- a/Engine/Source/Runtime/AIModule/Private/AISystem.cpp
+++ b/Engine/Source/Runtime/AIModule/Private/AISystem.cpp
@@ -3,6 +3,8 @@
 #include "AIModulePrivate.h"
 #include "BehaviorTree/BehaviorTreeManager.h"
 #include "EnvironmentQuery/EnvQueryManager.h"
+#include "GameFramework/PlayerController.h"
+#include "Debug/GameplayDebuggingControllerComponent.h"
 #include "AISystem.h"
 
 UAISystem::UAISystem(const FPostConstructInitializeProperties& PCIP) : Super(PCIP)
diff --git a/Engine/Source/Runtime/AIModule/Private/BehaviorTree/Decorators/BTDecorator_DoesPathExist.cpp b/Engine/Source/Runtime/AIModule/Private/BehaviorTree/Decorators/BTDecorator_DoesPathExist.cpp
--- a/Engine/Source/Runtime/AIModule/Private/BehaviorTree/Decorators/BTDecorator_DoesPathExist.cpp
+++ b/Engine/Source/Runtime/AIModule/Private/BehaviorTree/Decorators/BTDecorator_DoesPathExist.cpp
@@ -1,6 +1,9 @@
 // Copyright 1998-2014 Epic Games, Inc. All Rights Reserved.
 
 #include "AIModulePrivate.h"
+#if WITH_RECAST
+#include "AI/Navigation/RecastNavMesh.h"
+#endif // WITH_RECAST
 #include "Navigation/NavigationComponent.h"
 #include "BehaviorTree/BlackboardComponent.h"
 #include "BehaviorTree/Decorators/BTDecorator_DoesPathExist.h"
diff --git a/Engine/Source/Runtime/AIModule/Private/BehaviorTree/Tasks/BTTask_MakeNoise.cpp b/Engine/Source/Runtime/AIModule/Private/BehaviorTree/Tasks/BTTask_MakeNoise.cpp
--- a/Engine/Source/Runtime/AIModule/Private/BehaviorTree/Tasks/BTTask_MakeNoise.cpp
+++ b/Engine/Source/Runtime/AIModule/Private/BehaviorTree/Tasks/BTTask_MakeNoise.cpp
@@ -2,6 +2,7 @@
 
 #include "AIModulePrivate.h"
 #include "SoundDefinitions.h"
+#include "Sound/ReverbVolume.h"
 #include "BehaviorTree/Tasks/BTTask_MakeNoise.h"
 
 UBTTask_MakeNoise::UBTTask_MakeNoise(const class FPostConstructInitializeProperties& PCIP) 
diff --git a/Engine/Source/Runtime/AIModule/Private/BehaviorTree/Tasks/BTTask_PlaySound.cpp b/Engine/Source/Runtime/AIModule/Private/BehaviorTree/Tasks/BTTask_PlaySound.cpp
--- a/Engine/Source/Runtime/AIModule/Private/BehaviorTree/Tasks/BTTask_PlaySound.cpp
+++ b/Engine/Source/Runtime/AIModule/Private/BehaviorTree/Tasks/BTTask_PlaySound.cpp
@@ -2,6 +2,8 @@
 
 #include "AIModulePrivate.h"
 #include "SoundDefinitions.h"
+#include "Kismet/GameplayStatics.h"
+#include "Sound/SoundCue.h"
 #include "BehaviorTree/Tasks/BTTask_PlaySound.h"
 
 UBTTask_PlaySound::UBTTask_PlaySound(const class FPostConstructInitializeProperties& PCIP) : Super(PCIP)
diff --git a/Engine/Source/Runtime/AIModule/Private/Blueprint/AIBlueprintHelperLibrary.cpp b/Engine/Source/Runtime/AIModule/Private/Blueprint/AIBlueprintHelperLibrary.cpp
--- a/Engine/Source/Runtime/AIModule/Private/Blueprint/AIBlueprintHelperLibrary.cpp
+++ b/Engine/Source/Runtime/AIModule/Private/Blueprint/AIBlueprintHelperLibrary.cpp
@@ -1,10 +1,12 @@
 // Copyright 1998-2014 Epic Games, Inc. All Rights Reserved.
 
 #include "AIModulePrivate.h"
+#include "Engine/Blueprint.h"
 #include "Blueprint/AIAsyncTaskBlueprintProxy.h"
-#include "Blueprint/AIBlueprintHelperLibrary.h"
 #include "BehaviorTree/BlackboardComponent.h"
 #include "Animation/AnimInstance.h"
+#include "TimerManager.h"
+#include "Blueprint/AIBlueprintHelperLibrary.h"
 
 //----------------------------------------------------------------------//
 // UAIAsyncTaskBlueprintProxy
diff --git a/Engine/Source/Runtime/AIModule/Private/EnvironmentQuery/EQSTestingPawn.cpp b/Engine/Source/Runtime/AIModule/Private/EnvironmentQuery/EQSTestingPawn.cpp
--- a/Engine/Source/Runtime/AIModule/Private/EnvironmentQuery/EQSTestingPawn.cpp
+++ b/Engine/Source/Runtime/AIModule/Private/EnvironmentQuery/EQSTestingPawn.cpp
@@ -7,6 +7,9 @@
 #include "EnvironmentQuery/EnvQueryManager.h"
 #include "EnvironmentQuery/EQSRenderingComponent.h"
 
+#include "Engine/Selection.h"
+#include "Components/BillboardComponent.h"
+
 #if WITH_EDITOR
 #include "UnrealEd.h"
 #endif
diff --git a/Engine/Source/Runtime/AIModule/Private/EnvironmentQuery/Generators/EnvQueryGenerator_OnCircle.cpp b/Engine/Source/Runtime/AIModule/Private/EnvironmentQuery/Generators/EnvQueryGenerator_OnCircle.cpp
--- a/Engine/Source/Runtime/AIModule/Private/EnvironmentQuery/Generators/EnvQueryGenerator_OnCircle.cpp
+++ b/Engine/Source/Runtime/AIModule/Private/EnvironmentQuery/Generators/EnvQueryGenerator_OnCircle.cpp
@@ -4,6 +4,7 @@
 #include "EnvironmentQuery/Contexts/EnvQueryContext_Querier.h"
 #include "EnvironmentQuery/Items/EnvQueryItemType_Point.h"
 #include "EnvironmentQuery/Generators/EnvQueryGenerator_OnCircle.h"
+#include "AI/Navigation/RecastNavMesh.h"
 #include "VisualLog.h"
 
 #define LOCTEXT_NAMESPACE "EnvQueryGenerator"
diff --git a/Engine/Source/Runtime/AIModule/Private/EnvironmentQuery/Generators/EnvQueryGenerator_PathingGrid.cpp b/Engine/Source/Runtime/AIModule/Private/EnvironmentQuery/Generators/EnvQueryGenerator_PathingGrid.cpp
--- a/Engine/Source/Runtime/AIModule/Private/EnvironmentQuery/Generators/EnvQueryGenerator_PathingGrid.cpp
+++ b/Engine/Source/Runtime/AIModule/Private/EnvironmentQuery/Generators/EnvQueryGenerator_PathingGrid.cpp
@@ -5,7 +5,7 @@
 #include "EnvironmentQuery/Contexts/EnvQueryContext_Querier.h"
 #include "EnvironmentQuery/Generators/EnvQueryGenerator_PathingGrid.h"
 #include "AI/Navigation/NavFilters/RecastFilter_UseDefaultArea.h"
-
+#include "AI/Navigation/RecastNavMesh.h"
 
 #define LOCTEXT_NAMESPACE "EnvQueryGenerator"
 
diff --git a/Engine/Source/Runtime/AIModule/Private/EnvironmentQuery/Generators/EnvQueryGenerator_SimpleGrid.cpp b/Engine/Source/Runtime/AIModule/Private/EnvironmentQuery/Generators/EnvQueryGenerator_SimpleGrid.cpp
--- a/Engine/Source/Runtime/AIModule/Private/EnvironmentQuery/Generators/EnvQueryGenerator_SimpleGrid.cpp
+++ b/Engine/Source/Runtime/AIModule/Private/EnvironmentQuery/Generators/EnvQueryGenerator_SimpleGrid.cpp
@@ -4,6 +4,7 @@
 #include "EnvironmentQuery/Contexts/EnvQueryContext_Querier.h"
 #include "EnvironmentQuery/Items/EnvQueryItemType_Point.h"
 #include "EnvironmentQuery/Generators/EnvQueryGenerator_SimpleGrid.h"
+#include "AI/Navigation/RecastNavMesh.h"
 
 #define LOCTEXT_NAMESPACE "EnvQueryGenerator"
 
diff --git a/Engine/Source/Runtime/AIModule/Private/Navigation/CrowdFollowingComponent.cpp b/Engine/Source/Runtime/AIModule/Private/Navigation/CrowdFollowingComponent.cpp
--- a/Engine/Source/Runtime/AIModule/Private/Navigation/CrowdFollowingComponent.cpp
+++ b/Engine/Source/Runtime/AIModule/Private/Navigation/CrowdFollowingComponent.cpp
@@ -5,6 +5,7 @@
 #include "Navigation/CrowdManager.h"
 #include "Navigation/NavigationComponent.h"
 #include "AI/Navigation/SmartNavLinkComponent.h"
+#include "GameFramework/CharacterMovementComponent.h"
 #include "VisualLog.h"
 
 DEFINE_LOG_CATEGORY(LogCrowdFollowing);
diff --git a/Engine/Source/Runtime/AIModule/Private/Navigation/CrowdManager.cpp b/Engine/Source/Runtime/AIModule/Private/Navigation/CrowdManager.cpp
--- a/Engine/Source/Runtime/AIModule/Private/Navigation/CrowdManager.cpp
+++ b/Engine/Source/Runtime/AIModule/Private/Navigation/CrowdManager.cpp
@@ -13,6 +13,8 @@
 #include "Navigation/CrowdAgentInterface.h"
 #include "VisualLog.h"
 
+#include "DrawDebugHelpers.h"
+
 DECLARE_STATS_GROUP(TEXT("Crowd"), STATGROUP_AICrowd, STATCAT_Advanced);
 
 DECLARE_CYCLE_STAT(TEXT("Simulation Time"), STAT_AI_Crowd_Tick, STATGROUP_AICrowd);
diff --git a/Engine/Source/Runtime/AIModule/Private/Navigation/NavigationComponent.cpp b/Engine/Source/Runtime/AIModule/Private/Navigation/NavigationComponent.cpp
--- a/Engine/Source/Runtime/AIModule/Private/Navigation/NavigationComponent.cpp
+++ b/Engine/Source/Runtime/AIModule/Private/Navigation/NavigationComponent.cpp
@@ -3,6 +3,7 @@
 #include "AIModulePrivate.h"
 #include "VisualLog.h"
 #include "AI/Navigation/SmartNavLinkComponent.h"
+#include "DrawDebugHelpers.h"
 #include "Navigation/NavigationComponent.h"
 
 //----------------------------------------------------------------------//
diff --git a/Engine/Source/Runtime/AIModule/Private/Navigation/PathFollowingComponent.cpp b/Engine/Source/Runtime/AIModule/Private/Navigation/PathFollowingComponent.cpp
--- a/Engine/Source/Runtime/AIModule/Private/Navigation/PathFollowingComponent.cpp
+++ b/Engine/Source/Runtime/AIModule/Private/Navigation/PathFollowingComponent.cpp
@@ -4,8 +4,13 @@
 #if WITH_RECAST
 #	include "Detour/DetourNavMeshQuery.h"
 #endif
+#include "AI/Navigation/RecastNavMesh.h"
 #include "Navigation/NavigationComponent.h"
+#include "GameFramework/NavMovementComponent.h"
 #include "AI/Navigation/SmartNavLinkComponent.h"
+#include "GameFramework/Character.h"
+#include "Engine/Canvas.h"
+#include "TimerManager.h"
 #include "VisualLog.h"
 #include "Navigation/PathFollowingComponent.h"
 
diff --git a/Engine/Source/Runtime/AIModule/Private/Perception/PawnSensingComponent.cpp b/Engine/Source/Runtime/AIModule/Private/Perception/PawnSensingComponent.cpp
--- a/Engine/Source/Runtime/AIModule/Private/Perception/PawnSensingComponent.cpp
+++ b/Engine/Source/Runtime/AIModule/Private/Perception/PawnSensingComponent.cpp
@@ -1,6 +1,7 @@
 // Copyright 1998-2014 Epic Games, Inc. All Rights Reserved.
 
 #include "AIModulePrivate.h"
+#include "Components/PawnNoiseEmitterComponent.h"
 #include "Perception/PawnSensingComponent.h"
 
 DECLARE_CYCLE_STAT(TEXT("Sensing"),STAT_AI_Sensing,STATGROUP_AI);
diff --git a/Engine/Source/Runtime/Engine/Classes/AI/Navigation/NavigationData.h b/Engine/Source/Runtime/Engine/Classes/AI/Navigation/NavigationData.h
--- a/Engine/Source/Runtime/Engine/Classes/AI/Navigation/NavigationData.h
+++ b/Engine/Source/Runtime/Engine/Classes/AI/Navigation/NavigationData.h
@@ -3,6 +3,7 @@
 #pragma once
 #include "AI/Navigation/NavigationTypes.h"
 #include "AI/NavDataGenerator.h"
+#include "AI/Navigation/NavFilters/NavigationQueryFilter.h"
 #include "NavigationData.generated.h"
 
 class FNavDataGenerator; 
diff --git a/Engine/Source/Runtime/Engine/Classes/AI/Navigation/NavigationSystem.h b/Engine/Source/Runtime/Engine/Classes/AI/Navigation/NavigationSystem.h
--- a/Engine/Source/Runtime/Engine/Classes/AI/Navigation/NavigationSystem.h
+++ b/Engine/Source/Runtime/Engine/Classes/AI/Navigation/NavigationSystem.h
@@ -33,154 +33,6 @@ namespace NavigationDebugDrawing
 	extern const ENGINE_API FVector PathNodeBoxExtent;
 }
 
-namespace NavMeshMemory
-{
-#if STATS
-	// @todo could be made a more generic solution
-	class FNavigationMemoryStat : public FDefaultAllocator
-	{
-	public:
-		typedef FDefaultAllocator Super;
-
-		class ForAnyElementType : public FDefaultAllocator::ForAnyElementType
-		{
-		public:
-			typedef FDefaultAllocator::ForAnyElementType Super;
-		private:
-			int32 AllocatedSize;
-		public:
-
-			ForAnyElementType()
-			: AllocatedSize(0)
-			{
-
-			}
-
-			/** Destructor. */
-			~ForAnyElementType()
-			{
-				if (AllocatedSize)
-				{
-					DEC_DWORD_STAT_BY( STAT_NavigationMemory, AllocatedSize );
-				}
-			}
-
-			void ResizeAllocation(int32 PreviousNumElements,int32 NumElements,int32 NumBytesPerElement)
-			{
-				const int32 NewSize = NumElements * NumBytesPerElement;
-				INC_DWORD_STAT_BY( STAT_NavigationMemory, NewSize - AllocatedSize );
-				AllocatedSize = NewSize;
-
-				Super::ResizeAllocation(PreviousNumElements, NumElements, NumBytesPerElement);
-			}
-
-		private:
-			ForAnyElementType(const ForAnyElementType&);
-			ForAnyElementType& operator=(const ForAnyElementType&);
-		};
-	};
-
-	typedef FNavigationMemoryStat FNavAllocator;
-#else
-	typedef FDefaultAllocator FNavAllocator;
-#endif
-}
-
-#if STATS
-
-template <>
-struct TAllocatorTraits<NavMeshMemory::FNavigationMemoryStat> : TAllocatorTraits<NavMeshMemory::FNavigationMemoryStat::Super>
-{
-};
-
-#endif
-
-template<typename InElementType>
-class TNavStatArray : public TArray<InElementType, NavMeshMemory::FNavAllocator>
-{
-public:
-	typedef TArray<InElementType, NavMeshMemory::FNavAllocator> Super;
-
-#if PLATFORM_COMPILER_HAS_DEFAULTED_FUNCTIONS
-
-	TNavStatArray() = default;
-	TNavStatArray(const TNavStatArray&) = default;
-	TNavStatArray& operator=(const TNavStatArray&) = default;
-
-	#if PLATFORM_COMPILER_HAS_RVALUE_REFERENCES
-
-		TNavStatArray(TNavStatArray&&) = default;
-		TNavStatArray& operator=(TNavStatArray&&) = default;
-
-	#endif
-
-#else
-
-	FORCEINLINE TNavStatArray()
-	{
-	}
-
-	FORCEINLINE TNavStatArray(const TNavStatArray& Other)
-		: Super((const Super&)Other)
-	{
-	}
-
-	FORCEINLINE TNavStatArray& operator=(const TNavStatArray& Other)
-	{
-		(Super&)*this = (const Super&)Other;
-		return *this;
-	}
-
-	#if PLATFORM_COMPILER_HAS_RVALUE_REFERENCES
-
-		FORCEINLINE TNavStatArray(TNavStatArray&& Other)
-			: Super((Super&&)Other)
-		{
-		}
-
-		FORCEINLINE TNavStatArray& operator=(TNavStatArray&& Other)
-		{
-			(Super&)*this = (Super&&)Other;
-			return *this;
-		}
-
-	#endif
-
-#endif
-};
-
-template<typename InElementType>
-struct TContainerTraits<TNavStatArray<InElementType> > : public TContainerTraitsBase<TNavStatArray<InElementType> >
-{
-	enum { MoveWillEmptyContainer =
-		PLATFORM_COMPILER_HAS_RVALUE_REFERENCES &&
-		TContainerTraits<typename TNavStatArray<InElementType>::Super>::MoveWillEmptyContainer };
-};
-
-struct FNavigationPortalEdge
-{
-	FVector Left;
-	FVector Right;
-	NavNodeRef ToRef;
-	
-	FNavigationPortalEdge() : Left(0.f), Right(0.f)
-	{}
-
-	FNavigationPortalEdge(const FVector& InLeft, const FVector& InRight, NavNodeRef InToRef) 
-		: Left(InLeft), Right(InRight), ToRef(InToRef)
-	{}
-
-	FORCEINLINE FVector GetPoint(const int32 Index) const
-	{
-		check(Index >= 0 && Index < 2);
-		return ((FVector*)&Left)[Index];
-	}
-
-	FORCEINLINE float GetLength() const { return FVector::Dist(Left, Right); }
-
-	FORCEINLINE FVector GetMiddlePoint() const { return Left + (Right - Left) / 2; }
-};
-
 namespace FNavigationSystem
 {
 	enum EMode
diff --git a/Engine/Source/Runtime/Engine/Classes/AI/Navigation/NavigationTypes.h b/Engine/Source/Runtime/Engine/Classes/AI/Navigation/NavigationTypes.h
--- a/Engine/Source/Runtime/Engine/Classes/AI/Navigation/NavigationTypes.h
+++ b/Engine/Source/Runtime/Engine/Classes/AI/Navigation/NavigationTypes.h
@@ -124,6 +124,30 @@ struct FNavigationDirtyElement
 //////////////////////////////////////////////////////////////////////////
 // Path
 
+struct FNavigationPortalEdge
+{
+	FVector Left;
+	FVector Right;
+	NavNodeRef ToRef;
+
+	FNavigationPortalEdge() : Left(0.f), Right(0.f)
+	{}
+
+	FNavigationPortalEdge(const FVector& InLeft, const FVector& InRight, NavNodeRef InToRef)
+		: Left(InLeft), Right(InRight), ToRef(InToRef)
+	{}
+
+	FORCEINLINE FVector GetPoint(const int32 Index) const
+	{
+		check(Index >= 0 && Index < 2);
+		return ((FVector*)&Left)[Index];
+	}
+
+	FORCEINLINE float GetLength() const { return FVector::Dist(Left, Right); }
+
+	FORCEINLINE FVector GetMiddlePoint() const { return Left + (Right - Left) / 2; }
+};
+
 /** Describes a point in navigation data */
 struct FNavLocation
 {
@@ -508,3 +532,132 @@ class UNavigationTypes : public UObject
 
 	UNavigationTypes(const class FPostConstructInitializeProperties& PCIP) : Super(PCIP) { }
 };
+
+//////////////////////////////////////////////////////////////////////////
+// Memory stating
+
+namespace NavMeshMemory
+{
+#if STATS
+	// @todo could be made a more generic solution
+	class FNavigationMemoryStat : public FDefaultAllocator
+	{
+	public:
+		typedef FDefaultAllocator Super;
+
+		class ForAnyElementType : public FDefaultAllocator::ForAnyElementType
+		{
+		public:
+			typedef FDefaultAllocator::ForAnyElementType Super;
+		private:
+			int32 AllocatedSize;
+		public:
+
+			ForAnyElementType()
+				: AllocatedSize(0)
+			{
+
+			}
+
+			/** Destructor. */
+			~ForAnyElementType()
+			{
+				if (AllocatedSize)
+				{
+					DEC_DWORD_STAT_BY(STAT_NavigationMemory, AllocatedSize);
+				}
+			}
+
+			void ResizeAllocation(int32 PreviousNumElements, int32 NumElements, int32 NumBytesPerElement)
+			{
+				const int32 NewSize = NumElements * NumBytesPerElement;
+				INC_DWORD_STAT_BY(STAT_NavigationMemory, NewSize - AllocatedSize);
+				AllocatedSize = NewSize;
+
+				Super::ResizeAllocation(PreviousNumElements, NumElements, NumBytesPerElement);
+			}
+
+		private:
+			ForAnyElementType(const ForAnyElementType&);
+			ForAnyElementType& operator=(const ForAnyElementType&);
+		};
+	};
+
+	typedef FNavigationMemoryStat FNavAllocator;
+#else
+	typedef FDefaultAllocator FNavAllocator;
+#endif
+}
+
+#if STATS
+
+template <>
+struct TAllocatorTraits<NavMeshMemory::FNavigationMemoryStat> : TAllocatorTraits<NavMeshMemory::FNavigationMemoryStat::Super>
+{
+};
+
+#endif
+
+template<typename InElementType>
+class TNavStatArray : public TArray<InElementType, NavMeshMemory::FNavAllocator>
+{
+public:
+	typedef TArray<InElementType, NavMeshMemory::FNavAllocator> Super;
+
+#if PLATFORM_COMPILER_HAS_DEFAULTED_FUNCTIONS
+
+	TNavStatArray() = default;
+	TNavStatArray(const TNavStatArray&) = default;
+	TNavStatArray& operator=(const TNavStatArray&) = default;
+
+#if PLATFORM_COMPILER_HAS_RVALUE_REFERENCES
+
+	TNavStatArray(TNavStatArray&&) = default;
+	TNavStatArray& operator=(TNavStatArray&&) = default;
+
+#endif
+
+#else
+
+	FORCEINLINE TNavStatArray()
+	{
+	}
+
+	FORCEINLINE TNavStatArray(const TNavStatArray& Other)
+		: Super((const Super&)Other)
+	{
+	}
+
+	FORCEINLINE TNavStatArray& operator=(const TNavStatArray& Other)
+	{
+		(Super&)*this = (const Super&)Other;
+		return *this;
+	}
+
+#if PLATFORM_COMPILER_HAS_RVALUE_REFERENCES
+
+	FORCEINLINE TNavStatArray(TNavStatArray&& Other)
+		: Super((Super&&)Other)
+	{
+	}
+
+	FORCEINLINE TNavStatArray& operator=(TNavStatArray&& Other)
+	{
+		(Super&)*this = (Super&&)Other;
+		return *this;
+	}
+
+#endif
+
+#endif
+};
+
+template<typename InElementType>
+struct TContainerTraits<TNavStatArray<InElementType> > : public TContainerTraitsBase<TNavStatArray<InElementType> >
+{
+	enum {
+		MoveWillEmptyContainer =
+		PLATFORM_COMPILER_HAS_RVALUE_REFERENCES &&
+		TContainerTraits<typename TNavStatArray<InElementType>::Super>::MoveWillEmptyContainer
+	};
+};
\ No newline at end of file
diff --git a/Engine/Source/Runtime/Engine/Classes/AI/Navigation/RecastNavMesh.h b/Engine/Source/Runtime/Engine/Classes/AI/Navigation/RecastNavMesh.h
--- a/Engine/Source/Runtime/Engine/Classes/AI/Navigation/RecastNavMesh.h
+++ b/Engine/Source/Runtime/Engine/Classes/AI/Navigation/RecastNavMesh.h
@@ -2,6 +2,7 @@
 
 #pragma once
 
+#include "AI/Navigation/NavigationTypes.h"
 #include "NavAreas/NavArea.h"
 #include "Tickable.h"
 #include "AI/Navigation/NavigationData.h"
diff --git a/Engine/Source/Runtime/Engine/Classes/Animation/AnimInstance.h b/Engine/Source/Runtime/Engine/Classes/Animation/AnimInstance.h
--- a/Engine/Source/Runtime/Engine/Classes/Animation/AnimInstance.h
+++ b/Engine/Source/Runtime/Engine/Classes/Animation/AnimInstance.h
@@ -5,6 +5,7 @@
 #include "AnimMontage.h"
 #include "AnimSequence.h"
 #include "AnimationAsset.h"
+#include "Components/SkeletalMeshComponent.h"
 #include "AnimInstance.generated.h"
 
 struct FBoneContainer;
diff --git a/Engine/Source/Runtime/Engine/Classes/Animation/AnimSequenceBase.h b/Engine/Source/Runtime/Engine/Classes/Animation/AnimSequenceBase.h
--- a/Engine/Source/Runtime/Engine/Classes/Animation/AnimSequenceBase.h
+++ b/Engine/Source/Runtime/Engine/Classes/Animation/AnimSequenceBase.h
@@ -8,6 +8,7 @@
  */
 
 #include "AnimationAsset.h"
+#include "Curves/CurveBase.h"
 #include "AnimSequenceBase.generated.h"
 
 #define DEFAULT_SAMPLERATE			30.f
diff --git a/Engine/Source/Runtime/Engine/Classes/Animation/AnimationAsset.h b/Engine/Source/Runtime/Engine/Classes/Animation/AnimationAsset.h
--- a/Engine/Source/Runtime/Engine/Classes/Animation/AnimationAsset.h
+++ b/Engine/Source/Runtime/Engine/Classes/Animation/AnimationAsset.h
@@ -9,11 +9,10 @@
 
 #include "SkeletalMeshTypes.h"
 #include "AnimInterpFilter.h"
+#include "Animation/Skeleton.h"
+#include "Engine/SkeletalMesh.h"
 #include "AnimationAsset.generated.h"
 
-class USkeleton;
-class USkeletalMesh;
-
 /** Transform definition */
 USTRUCT(BlueprintType)
 struct FBlendSampleData
diff --git a/Engine/Source/Runtime/Engine/Classes/Animation/Skeleton.h b/Engine/Source/Runtime/Engine/Classes/Animation/Skeleton.h
--- a/Engine/Source/Runtime/Engine/Classes/Animation/Skeleton.h
+++ b/Engine/Source/Runtime/Engine/Classes/Animation/Skeleton.h
@@ -11,6 +11,8 @@
 #include "PreviewAssetAttachComponent.h"
 #include "Skeleton.generated.h"
 
+class UAnimSequence;
+class USkeletalMesh;
 
 /** This is a mapping table between bone in a particular skeletal mesh and bone of this skeleton set. */
 USTRUCT()
diff --git a/Engine/Source/Runtime/Engine/Classes/Components/ActorComponent.h b/Engine/Source/Runtime/Engine/Classes/Components/ActorComponent.h
--- a/Engine/Source/Runtime/Engine/Classes/Components/ActorComponent.h
+++ b/Engine/Source/Runtime/Engine/Classes/Components/ActorComponent.h
@@ -2,6 +2,7 @@
 
 #pragma once
 
+#include "Engine/EngineBaseTypes.h"
 #include "Engine/EngineTypes.h"
 #include "Interfaces/Interface_AssetUserData.h"
 
diff --git a/Engine/Source/Runtime/Engine/Classes/Components/CapsuleComponent.h b/Engine/Source/Runtime/Engine/Classes/Components/CapsuleComponent.h
--- a/Engine/Source/Runtime/Engine/Classes/Components/CapsuleComponent.h
+++ b/Engine/Source/Runtime/Engine/Classes/Components/CapsuleComponent.h
@@ -1,6 +1,7 @@
 // Copyright 1998-2014 Epic Games, Inc. All Rights Reserved.
 
 #pragma once
+#include "Components/ShapeComponent.h"
 #include "CapsuleComponent.generated.h"
 
 /** 
diff --git a/Engine/Source/Runtime/Engine/Classes/Components/ShapeComponent.h b/Engine/Source/Runtime/Engine/Classes/Components/ShapeComponent.h
--- a/Engine/Source/Runtime/Engine/Classes/Components/ShapeComponent.h
+++ b/Engine/Source/Runtime/Engine/Classes/Components/ShapeComponent.h
@@ -1,7 +1,7 @@
 // Copyright 1998-2014 Epic Games, Inc. All Rights Reserved.
 
-
 #pragma once
+#include "Components/PrimitiveComponent.h"
 #include "ShapeComponent.generated.h"
 
 UCLASS(abstract, hidecategories=(Object,LOD,Lighting,TextureStreaming,Activation,"Components|Activation"), editinlinenew, meta=(BlueprintSpawnableComponent), showcategories=(Mobility))
diff --git a/Engine/Source/Runtime/Engine/Classes/Components/SkeletalMeshComponent.h b/Engine/Source/Runtime/Engine/Classes/Components/SkeletalMeshComponent.h
--- a/Engine/Source/Runtime/Engine/Classes/Components/SkeletalMeshComponent.h
+++ b/Engine/Source/Runtime/Engine/Classes/Components/SkeletalMeshComponent.h
@@ -1,7 +1,7 @@
 // Copyright 1998-2014 Epic Games, Inc. All Rights Reserved.
 
 #pragma once
-
+#include "Components/SkinnedMeshComponent.h"
 #include "PhysicsEngine/PhysicsConstraintComponent.h"
 #include "SkeletalMeshComponent.generated.h"
 
diff --git a/Engine/Source/Runtime/Engine/Classes/Components/SkinnedMeshComponent.h b/Engine/Source/Runtime/Engine/Classes/Components/SkinnedMeshComponent.h
--- a/Engine/Source/Runtime/Engine/Classes/Components/SkinnedMeshComponent.h
+++ b/Engine/Source/Runtime/Engine/Classes/Components/SkinnedMeshComponent.h
@@ -9,6 +9,7 @@
 
 #pragma once
 #include "SkeletalMeshTypes.h"
+#include "Components/MeshComponent.h"
 #include "SkinnedMeshComponent.generated.h"
 
 //
diff --git a/Engine/Source/Runtime/Engine/Classes/EdGraph/EdGraphNode.h b/Engine/Source/Runtime/Engine/Classes/EdGraph/EdGraphNode.h
--- a/Engine/Source/Runtime/Engine/Classes/EdGraph/EdGraphNode.h
+++ b/Engine/Source/Runtime/Engine/Classes/EdGraph/EdGraphNode.h
@@ -10,6 +10,15 @@ class UEdGraphNode;
 class UEdGraphPin;
 class SGraphNode;
 
+/** Enum used to define which way data flows into or out of this pin */
+UENUM()
+enum EEdGraphPinDirection
+{
+	EGPD_Input,
+	EGPD_Output,
+	EGPD_MAX,
+};
+
 /** Enum to indicate what sort of title we want */
 UENUM()
 namespace ENodeTitleType
diff --git a/Engine/Source/Runtime/Engine/Classes/EdGraph/EdGraphPin.h b/Engine/Source/Runtime/Engine/Classes/EdGraph/EdGraphPin.h
--- a/Engine/Source/Runtime/Engine/Classes/EdGraph/EdGraphPin.h
+++ b/Engine/Source/Runtime/Engine/Classes/EdGraph/EdGraphPin.h
@@ -2,17 +2,9 @@
 
 #pragma once
 
+#include "EdGraph/EdGraphNode.h"
 #include "EdGraphPin.generated.h"
 
-/** Enum used to define which way data flows into or out of this pin */
-UENUM()
-enum EEdGraphPinDirection
-{
-	EGPD_Input,
-	EGPD_Output,
-	EGPD_MAX,
-};
-
 USTRUCT()
 struct FSimpleMemberReference
 {
diff --git a/Engine/Source/Runtime/Engine/Classes/Engine/Brush.h b/Engine/Source/Runtime/Engine/Classes/Engine/Brush.h
--- a/Engine/Source/Runtime/Engine/Classes/Engine/Brush.h
+++ b/Engine/Source/Runtime/Engine/Classes/Engine/Brush.h
@@ -6,6 +6,7 @@
 
 #pragma once
 #include "Model.h"
+#include "Engine/Level.h"
 #include "Brush.generated.h"
 
 //-----------------------------------------------------------------------------
diff --git a/Engine/Source/Runtime/Engine/Classes/Engine/Level.h b/Engine/Source/Runtime/Engine/Classes/Engine/Level.h
--- a/Engine/Source/Runtime/Engine/Classes/Engine/Level.h
+++ b/Engine/Source/Runtime/Engine/Classes/Engine/Level.h
@@ -5,10 +5,12 @@
 //
 
 #pragma once
-#include "Runtime/Engine/Classes/PhysicsEngine/BodyInstance.h"
-
+#include "PhysicsEngine/BodyInstance.h"
+#include "Engine/LevelBase.h"
 #include "Level.generated.h"
 
+class ALevelBounds;
+
 /**
  * Structure containing all information needed for determining the screen space
  * size of an object/ texture instance.
diff --git a/Engine/Source/Runtime/Engine/Classes/Engine/World.h b/Engine/Source/Runtime/Engine/Classes/Engine/World.h
--- a/Engine/Source/Runtime/Engine/Classes/Engine/World.h
+++ b/Engine/Source/Runtime/Engine/Classes/Engine/World.h
@@ -9,7 +9,7 @@
 #include "CollisionQueryParams.h"
 #include "WorldCollision.h"
 #include "PendingNetGame.h"
-#include "../GameFramework/MusicTrackDataStructures.h"
+#include "GameFramework/MusicTrackDataStructures.h"
 #include "World.generated.h"
 
 class FPhysScene;
@@ -19,6 +19,8 @@ class ULocalPlayer;
 class UGameViewportClient;
 class ABrush;
 class UModel;
+class APhysicsVolume;
+class UTexture2D;
 
 template<typename,typename> class TOctree;
 
diff --git a/Engine/Source/Runtime/Engine/Classes/GameFramework/Actor.h b/Engine/Source/Runtime/Engine/Classes/GameFramework/Actor.h
--- a/Engine/Source/Runtime/Engine/Classes/GameFramework/Actor.h
+++ b/Engine/Source/Runtime/Engine/Classes/GameFramework/Actor.h
@@ -4,6 +4,7 @@
 
 #include "Runtime/Engine/Classes/Engine/LatentActionManager.h"
 #include "Runtime/Engine/Classes/Camera/CameraTypes.h"
+#include "Components/SceneComponent.h"
 #include "RenderingThread.h"
 
 UENUM(BlueprintType)
diff --git a/Engine/Source/Runtime/Engine/Classes/GameFramework/CharacterMovementComponent.h b/Engine/Source/Runtime/Engine/Classes/GameFramework/CharacterMovementComponent.h
--- a/Engine/Source/Runtime/Engine/Classes/GameFramework/CharacterMovementComponent.h
+++ b/Engine/Source/Runtime/Engine/Classes/GameFramework/CharacterMovementComponent.h
@@ -7,6 +7,7 @@
 
 #pragma once
 
+#include "Interfaces/NetworkPredictionInterface.h"
 #include "CharacterMovementComponent.generated.h"
 
 /** Movement modes for Characters.  */
diff --git a/Engine/Source/Runtime/Engine/Classes/GameFramework/GameMode.h b/Engine/Source/Runtime/Engine/Classes/GameFramework/GameMode.h
--- a/Engine/Source/Runtime/Engine/Classes/GameFramework/GameMode.h
+++ b/Engine/Source/Runtime/Engine/Classes/GameFramework/GameMode.h
@@ -1,7 +1,7 @@
 // Copyright 1998-2014 Epic Games, Inc. All Rights Reserved.
 
 #pragma once
-
+#include "GameFramework/Info.h"
 #include "GameMode.generated.h"
 
 //=============================================================================
diff --git a/Engine/Source/Runtime/Engine/Classes/GameFramework/NavMovementComponent.h b/Engine/Source/Runtime/Engine/Classes/GameFramework/NavMovementComponent.h
--- a/Engine/Source/Runtime/Engine/Classes/GameFramework/NavMovementComponent.h
+++ b/Engine/Source/Runtime/Engine/Classes/GameFramework/NavMovementComponent.h
@@ -6,6 +6,7 @@
 
 #pragma once
 #include "AI/Navigation/NavAgentInterface.h"
+#include "AI/Navigation/NavigationTypes.h"
 #include "GameFramework/MovementComponent.h"
 #include "NavMovementComponent.generated.h"
 
diff --git a/Engine/Source/Runtime/Engine/Classes/GameFramework/Pawn.h b/Engine/Source/Runtime/Engine/Classes/GameFramework/Pawn.h
--- a/Engine/Source/Runtime/Engine/Classes/GameFramework/Pawn.h
+++ b/Engine/Source/Runtime/Engine/Classes/GameFramework/Pawn.h
@@ -178,7 +178,7 @@ public:
 	virtual void Reset() OVERRIDE;
 	virtual FString GetHumanReadableName() const OVERRIDE;
 	virtual float GetNetPriority(const FVector& ViewPos, const FVector& ViewDir, APlayerController* Viewer, UActorChannel* InChannel, float Time, bool bLowBandwidth) OVERRIDE;
-	virtual bool ShouldTickIfViewportsOnly() const OVERRIDE { return IsLocallyControlled() && Cast<APlayerController>(GetController()); }
+	virtual bool ShouldTickIfViewportsOnly() const OVERRIDE;
 	virtual bool IsNetRelevantFor(APlayerController* RealViewer, AActor* Viewer, const FVector& SrcLocation) OVERRIDE;
 	virtual void PostNetReceiveLocation() OVERRIDE;
 	virtual void PostNetReceiveVelocity(const FVector& NewVelocity) OVERRIDE;
diff --git a/Engine/Source/Runtime/Engine/Classes/GameFramework/PlayerController.h b/Engine/Source/Runtime/Engine/Classes/GameFramework/PlayerController.h
--- a/Engine/Source/Runtime/Engine/Classes/GameFramework/PlayerController.h
+++ b/Engine/Source/Runtime/Engine/Classes/GameFramework/PlayerController.h
@@ -10,6 +10,7 @@
 #include "SceneManagement.h"
 #include "ForceFeedbackEffect.h"
 #include "OnlineReplStructs.h"
+#include "GameFramework/Controller.h"
 #include "PlayerController.generated.h"
 
 /** delegate used to override default viewport audio listener position calculated from camera */
diff --git a/Engine/Source/Runtime/Engine/Classes/Kismet/GameplayStatics.h b/Engine/Source/Runtime/Engine/Classes/Kismet/GameplayStatics.h
--- a/Engine/Source/Runtime/Engine/Classes/Kismet/GameplayStatics.h
+++ b/Engine/Source/Runtime/Engine/Classes/Kismet/GameplayStatics.h
@@ -10,6 +10,8 @@
 //
 class USaveGame;
 struct FDialogueContext;
+class UParticleSystemComponent;
+class UParticleSystem;
 
 UENUM()
 namespace ESuggestProjVelocityTraceOption
diff --git a/Engine/Source/Runtime/Engine/Classes/PhysicsEngine/PhysicsConstraintComponent.h b/Engine/Source/Runtime/Engine/Classes/PhysicsEngine/PhysicsConstraintComponent.h
--- a/Engine/Source/Runtime/Engine/Classes/PhysicsEngine/PhysicsConstraintComponent.h
+++ b/Engine/Source/Runtime/Engine/Classes/PhysicsEngine/PhysicsConstraintComponent.h
@@ -1,6 +1,6 @@
 // Copyright 1998-2014 Epic Games, Inc. All Rights Reserved.
 #pragma once
-
+#include "PhysicsEngine/ConstraintInstance.h"
 #include "PhysicsConstraintComponent.generated.h"
 
 /** Dynamic delegate to use by components that want to route the broken-event into blueprints */
diff --git a/Engine/Source/Runtime/Engine/Classes/Sound/SoundCue.h b/Engine/Source/Runtime/Engine/Classes/Sound/SoundCue.h
--- a/Engine/Source/Runtime/Engine/Classes/Sound/SoundCue.h
+++ b/Engine/Source/Runtime/Engine/Classes/Sound/SoundCue.h
@@ -1,7 +1,7 @@
 // Copyright 1998-2014 Epic Games, Inc. All Rights Reserved.
 
 #pragma once
-
+#include "Sound/SoundBase.h"
 #include "SoundCue.generated.h"
 
 struct FActiveSound;
diff --git a/Engine/Source/Runtime/Engine/Private/AI/Navigation/PImplRecastNavMesh.h b/Engine/Source/Runtime/Engine/Private/AI/Navigation/PImplRecastNavMesh.h
--- a/Engine/Source/Runtime/Engine/Private/AI/Navigation/PImplRecastNavMesh.h
+++ b/Engine/Source/Runtime/Engine/Private/AI/Navigation/PImplRecastNavMesh.h
@@ -13,6 +13,7 @@
 #include "DetourCommon.h"
 #include "DetourNavMesh.h"
 #include "DetourNavMeshQuery.h"
+#include "AI/Navigation/RecastNavMesh.h"
 
 #define RECAST_VERY_SMALL_AGENT_RADIUS 0.0f
 
diff --git a/Engine/Source/Runtime/Engine/Private/DebugRenderSceneProxy.cpp b/Engine/Source/Runtime/Engine/Private/DebugRenderSceneProxy.cpp
--- a/Engine/Source/Runtime/Engine/Private/DebugRenderSceneProxy.cpp
+++ b/Engine/Source/Runtime/Engine/Private/DebugRenderSceneProxy.cpp
@@ -63,7 +63,7 @@ void FDebugRenderSceneProxy::DrawDebugLabels(UCanvas* Canvas, APlayerController*
 		if (PointInView(It->Location, View))
 		{
 			const FVector ScreenLoc = Canvas->Project(It->Location);
-			const FFontRenderInfo& FontInfo = TextWithoutShadowDistance >= 0 ? (PointWithinCorrectDistance(It->Location, View, TextWithoutShadowDistance) ? FontRenderInfoWithShadow : FontRenderInfo) : FontRenderInfo;
+			const FFontRenderInfo& FontInfo = TextWithoutShadowDistance >= 0 ? (PointInRange(It->Location, View, TextWithoutShadowDistance) ? FontRenderInfoWithShadow : FontRenderInfo) : FontRenderInfo;
 			Canvas->DrawText(RenderFont, It->Text, ScreenLoc.X, ScreenLoc.Y, 1, 1, FontInfo);
 		}
 	}
diff --git a/Engine/Source/Runtime/Engine/Private/Pawn.cpp b/Engine/Source/Runtime/Engine/Private/Pawn.cpp
--- a/Engine/Source/Runtime/Engine/Private/Pawn.cpp
+++ b/Engine/Source/Runtime/Engine/Private/Pawn.cpp
@@ -282,6 +282,10 @@ float APawn::GetNetPriority(const FVector& ViewPos, const FVector& ViewDir, APla
 	return NetPriority * Time;
 }
 
+bool APawn::ShouldTickIfViewportsOnly() const 
+{ 
+	return IsLocallyControlled() && Cast<APlayerController>(GetController()); 
+}
 
 FVector APawn::GetPawnViewLocation() const
 {
diff --git a/Engine/Source/Runtime/Engine/Private/World.cpp b/Engine/Source/Runtime/Engine/Private/World.cpp
--- a/Engine/Source/Runtime/Engine/Private/World.cpp
+++ b/Engine/Source/Runtime/Engine/Private/World.cpp
@@ -29,6 +29,7 @@
 #include "VisualLog.h"
 #include "LevelUtils.h"
 #include "PhysicsPublic.h"
+#include "GameFramework/CharacterMovementComponent.h"
 #include "AI/AISystemBase.h"
 #if WITH_EDITOR
 	#include "Editor/UnrealEd/Public/Kismet2/KismetEditorUtilities.h"
diff --git a/Engine/Source/Runtime/Engine/Public/ActiveSound.h b/Engine/Source/Runtime/Engine/Public/ActiveSound.h
--- a/Engine/Source/Runtime/Engine/Public/ActiveSound.h
+++ b/Engine/Source/Runtime/Engine/Public/ActiveSound.h
@@ -2,6 +2,8 @@
 
 #pragma once
 
+#include "Sound/ReverbVolume.h"
+
 struct ENGINE_API FActiveSound
 {
 public:
diff --git a/Engine/Source/Runtime/Engine/Public/DebugRenderSceneProxy.h b/Engine/Source/Runtime/Engine/Public/DebugRenderSceneProxy.h
--- a/Engine/Source/Runtime/Engine/Public/DebugRenderSceneProxy.h
+++ b/Engine/Source/Runtime/Engine/Public/DebugRenderSceneProxy.h
@@ -46,21 +46,14 @@ public:
 	/** called to clean up debug drawing delegate in UDebugDrawService */
 	ENGINE_API virtual void UnregisterDebugDrawDelgate();
 
-	FORCEINLINE bool PointInView(const FVector& Location, const FSceneView* View)
+	FORCEINLINE bool PointInView(const FVector& Location, const FSceneView* View) const
 	{
 		return View ? View->ViewFrustum.IntersectBox(Location, FVector::ZeroVector) : false;
 	}
 
-	FORCEINLINE bool PointWithinCorrectDistance(const FVector& Start, const FSceneView* View, float CorrectDistance = -1)
+	FORCEINLINE bool PointInRange(const FVector& Start, const FSceneView* View, float Range) const
 	{
-		const float MaxDistance = CorrectDistance > 0 ? (CorrectDistance*CorrectDistance) : ARecastNavMesh::GetDrawDistanceSq();
-
-		if (FVector::DistSquared(Start, View->ViewMatrices.ViewOrigin) > MaxDistance)
-		{
-			return false;
-		}
-
-		return true;
+		return FVector::DistSquared(Start, View->ViewMatrices.ViewOrigin) <= FMath::Square(Range);
 	}
 
 	/** Struct to hold info about lines to render. */
diff --git a/Engine/Source/Runtime/Engine/Public/EngineUtils.h b/Engine/Source/Runtime/Engine/Public/EngineUtils.h
--- a/Engine/Source/Runtime/Engine/Public/EngineUtils.h
+++ b/Engine/Source/Runtime/Engine/Public/EngineUtils.h
@@ -7,6 +7,9 @@
 #ifndef _INC_ENGINE_UTILS
 #define _INC_ENGINE_UTILS
 
+#include "Engine/Level.h"
+#include "GameFramework/WorldSettings.h"
+
 /*-----------------------------------------------------------------------------
 	Hit proxies.
 -----------------------------------------------------------------------------*/
