Nodes which are deemed to have previews with no particular value are now auto collapsed by default.
#TTP 335636 - UE4: Material Editor: Certain nodes should be collapsed by default to save space

#branch UE4

#proj Runtime.Engine

#change Set all the requested material nodes to be defaulted to collapsed.

diff --git a/Engine/Source/Runtime/Engine/Private/Materials/MaterialExpressions.cpp b/Engine/Source/Runtime/Engine/Private/Materials/MaterialExpressions.cpp
--- a/Engine/Source/Runtime/Engine/Private/Materials/MaterialExpressions.cpp
+++ b/Engine/Source/Runtime/Engine/Private/Materials/MaterialExpressions.cpp
@@ -2127,7 +2127,7 @@ UMaterialExpressionConstant::UMaterialExpressionConstant(const class FPostConstr
 
 	MenuCategories.Add(ConstructorStatics.NAME_Constants);
 
-	bCollapsed = false;
+	bCollapsed = true;
 }
 
 int32 UMaterialExpressionConstant::Compile(class FMaterialCompiler* Compiler, int32 OutputIndex, int32 MultiplexIndex)
@@ -2170,7 +2170,7 @@ UMaterialExpressionConstant2Vector::UMaterialExpressionConstant2Vector(const cla
 
 	MenuCategories.Add(ConstructorStatics.NAME_Constants);
 	MenuCategories.Add(ConstructorStatics.NAME_Vectors);
-	bCollapsed = false;
+	bCollapsed = true;
 }
 
 int32 UMaterialExpressionConstant2Vector::Compile(class FMaterialCompiler* Compiler, int32 OutputIndex, int32 MultiplexIndex)
@@ -2506,7 +2506,7 @@ UMaterialExpressionTextureCoordinate::UMaterialExpressionTextureCoordinate(const
 
 	MenuCategories.Add(ConstructorStatics.NAME_Coordinates);
 
-	bCollapsed = false;
+	bCollapsed = true;
 }
 
 
@@ -2872,7 +2872,7 @@ UMaterialExpressionPanner::UMaterialExpressionPanner(const class FPostConstructI
 	static FConstructorStatics ConstructorStatics;
 
 	MenuCategories.Add(ConstructorStatics.NAME_Coordinates);
-	bCollapsed = false;
+	bCollapsed = true;
 	ConstCoordinate = 0;
 }
 
@@ -2933,7 +2933,7 @@ UMaterialExpressionRotator::UMaterialExpressionRotator(const class FPostConstruc
 
 	MenuCategories.Add(ConstructorStatics.NAME_Coordinates);
 
-	bCollapsed = false;
+	bCollapsed = true;
 }
 
 int32 UMaterialExpressionRotator::Compile(class FMaterialCompiler* Compiler, int32 OutputIndex, int32 MultiplexIndex)
