New destruction API from NVIDIA:

Chunk visibility event callback
Using "finalStep" flag for chunk callbacks
Adding getChunkPhysXShapes API

diff --git a/Engine/Source/Runtime/Engine/Private/PhysicsEngine/PhysXSupport.cpp b/Engine/Source/Runtime/Engine/Private/PhysicsEngine/PhysXSupport.cpp
--- a/Engine/Source/Runtime/Engine/Private/PhysicsEngine/PhysXSupport.cpp
+++ b/Engine/Source/Runtime/Engine/Private/PhysicsEngine/PhysXSupport.cpp
@@ -708,8 +708,9 @@ void FApexChunkReport::onDamageNotify(const NxApexDamageEventReportData& damageE
 		return;
 	}
 
-#if WITH_SUBSTEPPING
+	#if WITH_SUBSTEPPING
 	if (UPhysicsSettings::Get()->bSubstepping)
+
 	{
 		DestructibleComponent->GetWorld()->GetPhysicsScene()->DeferredDestructibleDamageNotify(damageEvent);
 	}
@@ -720,6 +721,11 @@ void FApexChunkReport::onDamageNotify(const NxApexDamageEventReportData& damageE
 	}
 }
 
+void FApexChunkReport::onStateChangeNotify(const NxApexChunkStateEventData& visibilityEvent)
+{
+
+}
+
 ///////// FApexPhysX3Interface //////////////////////////////////
 void FApexPhysX3Interface::setContactReportFlags(physx::PxShape* PShape, physx::PxPairFlags PFlags, NxDestructibleActor* actor, PxU16 actorChunkIndex)
 {
diff --git a/Engine/Source/Runtime/Engine/Private/PhysicsEngine/PhysXSupport.h b/Engine/Source/Runtime/Engine/Private/PhysicsEngine/PhysXSupport.h
--- a/Engine/Source/Runtime/Engine/Private/PhysicsEngine/PhysXSupport.h
+++ b/Engine/Source/Runtime/Engine/Private/PhysicsEngine/PhysXSupport.h
@@ -602,6 +602,7 @@ public:
 	// NxUserChunkReport interface.
 
 	virtual void	onDamageNotify(const NxApexDamageEventReportData& damageEvent) OVERRIDE;
+	virtual void	onStateChangeNotify(const NxApexChunkStateEventData& visibilityEvent) OVERRIDE;
 };
 extern FApexChunkReport GApexChunkReport;
 #endif // #if WITH_APEX
