Destruction uses the new visibility event API
- Remove PostTick

Destruction uses the deferred damage API
- fixes particle effects not spawning durning sub-stepping

diff --git a/Engine/Source/Runtime/Engine/Classes/Components/DestructibleComponent.h b/Engine/Source/Runtime/Engine/Classes/Components/DestructibleComponent.h
--- a/Engine/Source/Runtime/Engine/Classes/Components/DestructibleComponent.h
+++ b/Engine/Source/Runtime/Engine/Classes/Components/DestructibleComponent.h
@@ -12,6 +12,7 @@ namespace physx
 	{
 		class  NxDestructibleActor;
 		struct NxApexDamageEventReportData;
+		struct NxApexChunkStateEventData;
 	}
 #endif
 	class PxRigidDynamic;
@@ -123,7 +124,6 @@ public:
 	virtual void AddRadialImpulse(FVector Origin, float Radius, float Strength, ERadialImpulseFalloff Falloff, bool bVelChange=false) OVERRIDE;
 	virtual void AddRadialForce(FVector Origin, float Radius, float Strength, ERadialImpulseFalloff Falloff) OVERRIDE;
 	virtual void ReceiveComponentDamage(float DamageAmount, FDamageEvent const& DamageEvent, AController* EventInstigator, AActor* DamageCauser) OVERRIDE;
-	virtual void PostPhysicsTick(FPrimitiveComponentPostPhysicsTickFunction &ThisTickFunction) OVERRIDE;
 
 	virtual bool LineTraceComponent( FHitResult& OutHit, const FVector Start, const FVector End, const FCollisionQueryParams& Params ) OVERRIDE;
 	virtual bool SweepComponent( FHitResult& OutHit, const FVector Start, const FVector End, const FCollisionShape& CollisionShape, bool bTraceComplex=false) OVERRIDE;
@@ -174,6 +174,9 @@ public:
 
 	/** Callback from physics system to notify the actor that it has been damaged */
 	void OnDamageEvent(const physx::apex::NxApexDamageEventReportData& InDamageEvent);
+
+	/** Callback from physics system to notify the actor that a chunk's visibilty has changed */
+	void OnVisibilityEvent(const physx::apex::NxApexChunkStateEventData & InDamageEvent);
 #endif // WITH_APEX
 
 	// End DestructibleComponent interface.
diff --git a/Engine/Source/Runtime/Engine/Private/DestructibleComponent.cpp b/Engine/Source/Runtime/Engine/Private/DestructibleComponent.cpp
--- a/Engine/Source/Runtime/Engine/Private/DestructibleComponent.cpp
+++ b/Engine/Source/Runtime/Engine/Private/DestructibleComponent.cpp
@@ -319,6 +319,8 @@ void UDestructibleComponent::CreatePhysicsState()
 
 	// Start asleep if requested
 	PxRigidDynamic* PRootActor = ApexDestructibleActor->getChunkPhysXActor(0);
+
+
 	//  Put to sleep or wake up only if the component is physics-simulated
 	if (PRootActor != NULL && BodyInstance.bSimulatePhysics)
 	{
@@ -633,54 +635,22 @@ void UDestructibleComponent::OnDamageEvent(const NxApexDamageEventReportData& In
 		MarkRenderStateDirty();
 	}
 }
-#endif // WITH_APEX
 
-void UDestructibleComponent::RefreshBoneTransforms(FActorComponentTickFunction* TickFunction)
+void UDestructibleComponent::OnVisibilityEvent(const NxApexChunkStateEventData & InVisibilityEvent)
 {
-#if WITH_APEX
-	if(ApexDestructibleActor != NULL && SkeletalMesh)
+	for (uint32 EventIndex = 0; EventIndex < InVisibilityEvent.stateEventListSize; ++EventIndex)
 	{
-		UDestructibleMesh* TheDestructibleMesh = GetDestructibleMesh();
-
-		// Save a pointer to the APEX NxDestructibleAsset
-		physx::NxDestructibleAsset* ApexDestructibleAsset = TheDestructibleMesh->ApexDestructibleAsset;
-		check(ApexDestructibleAsset);
-
-		{
-			// Lock here so we don't encounter race conditions with the destruction processing
-			FPhysScene* PhysScene = World->GetPhysicsScene();
-			check(PhysScene);
-			const uint32 SceneType = (BodyInstance.bUseAsyncScene && PhysScene->HasAsyncScene()) ? PST_Async : PST_Sync;
-			PxScene* PScene = PhysScene->GetPhysXScene(SceneType);
-			check(PScene);
-			SCOPED_SCENE_WRITE_LOCK(PScene);
-			SCOPED_SCENE_READ_LOCK(PScene);
-
-			// Try to acquire event buffer
-			const physx::NxDestructibleChunkEvent* EventBuffer;
-			physx::PxU32 EventBufferSize;
-			if (ApexDestructibleActor->acquireChunkEventBuffer(EventBuffer, EventBufferSize))
-			{
-				// Buffer acquired
-				while (EventBufferSize--)
-				{
-					const physx::NxDestructibleChunkEvent& Event = *EventBuffer++;
-					// Right now the only events are visibility changes.  So as an optimization we won't check for the event type.
-	//				if (Event.event & physx::NxDestructibleChunkEvent::VisibilityChanged)
-					const bool bVisible = (Event.event & physx::NxDestructibleChunkEvent::ChunkVisible) != 0;
-					SetChunkVisible(Event.chunkIndex, bVisible);
-				}
-				// Release buffer (will be cleared)
-				ApexDestructibleActor->releaseChunkEventBuffer();
-			}
-		}
-
-
-
-		// Send bones to render thread at end of frame
-		MarkRenderDynamicDataDirty();
+		const NxDestructibleChunkEvent &  Event = InVisibilityEvent.stateEventList[EventIndex];
+		// Right now the only events are visibility changes.  So as an optimization we won't check for the event type.
+		//				if (Event.event & physx::NxDestructibleChunkEvent::VisibilityChanged)
+		const bool bVisible = (Event.event & physx::NxDestructibleChunkEvent::ChunkVisible) != 0;
+		SetChunkVisible(Event.chunkIndex, bVisible);
+	}
 }
-#endif	// #if WITH_APEX
+#endif // WITH_APEX
+
+void UDestructibleComponent::RefreshBoneTransforms(FActorComponentTickFunction* TickFunction)
+{
 }
 
 void UDestructibleComponent::SetDestructibleMesh(class UDestructibleMesh* NewMesh)
@@ -811,6 +781,7 @@ void UDestructibleComponent::SetChunkVisible( int32 ChunkIndex, bool bVisible )
 
 				PActor->userData = UserData;
 
+
 				// Set collision response to non-root chunks
 				if (GetDestructibleMesh()->ApexDestructibleAsset->getChunkParentIndex(ChunkIndex) >= 0)
 				{
@@ -861,6 +832,9 @@ void UDestructibleComponent::SetChunkWorldRT( int32 ChunkIndex, const FQuat& Wor
 	// Mark the transform as dirty, so the bounds are updated and sent to the render thread
 	MarkRenderTransformDirty();
 
+	// New bone positions need to be sent to render thread
+	MarkRenderDynamicDataDirty();
+
 #if 0
 	// Scale is already applied to the ComponentToWorld transform, and is carried into the bones _locally_.
 	// So there is no need to set scale in the bone local transforms
@@ -1010,20 +984,6 @@ void UDestructibleComponent::Deactivate()
 		bIsActive = false;
 	}
 }
-void UDestructibleComponent::PostPhysicsTick( FPrimitiveComponentPostPhysicsTickFunction &ThisTickFunction )
-{
-	Super::PostPhysicsTick(ThisTickFunction);
-
-	if (IsRegistered())
-	{
-		if (BodyInstance.bSimulatePhysics)
-		{
-			SyncComponentToRBPhysics();
-		}
-
-		RefreshBoneTransforms();
-	}
-}
 
 bool UDestructibleComponent::ShouldUpdateTransform(bool bLODHasChanged) const
 {
diff --git a/Engine/Source/Runtime/Engine/Private/PhysicsEngine/PhysLevel.cpp b/Engine/Source/Runtime/Engine/Private/PhysicsEngine/PhysLevel.cpp
--- a/Engine/Source/Runtime/Engine/Private/PhysicsEngine/PhysLevel.cpp
+++ b/Engine/Source/Runtime/Engine/Private/PhysicsEngine/PhysLevel.cpp
@@ -345,6 +345,8 @@ void InitGamePhys()
 	static const auto CVarSortDynamicChunksByBenefit = IConsoleManager::Get().FindTConsoleVariableDataInt(TEXT("p.bAPEXSortDynamicChunksByBenefit"));
 	GApexModuleDestructible->setSortByBenefit(CVarSortDynamicChunksByBenefit->GetValueOnGameThread()!=0);
 
+	GApexModuleDestructible->setChunkReportSendChunkStateEvents(true);
+
 	// APEX 1.3 to preserve 1.2 behavior
 	GApexModuleDestructible->setUseLegacyDamageRadiusSpread(true); 
 	GApexModuleDestructible->setUseLegacyChunkBoundsTesting(true);
diff --git a/Engine/Source/Runtime/Engine/Private/PhysicsEngine/PhysXSupport.cpp b/Engine/Source/Runtime/Engine/Private/PhysicsEngine/PhysXSupport.cpp
--- a/Engine/Source/Runtime/Engine/Private/PhysicsEngine/PhysXSupport.cpp
+++ b/Engine/Source/Runtime/Engine/Private/PhysicsEngine/PhysXSupport.cpp
@@ -708,22 +708,20 @@ void FApexChunkReport::onDamageNotify(const NxApexDamageEventReportData& damageE
 		return;
 	}
 
-	#if WITH_SUBSTEPPING
-	if (UPhysicsSettings::Get()->bSubstepping)
-
-	{
-		DestructibleComponent->GetWorld()->GetPhysicsScene()->DeferredDestructibleDamageNotify(damageEvent);
-	}
-	else
-#endif
-	{
-		DestructibleComponent->OnDamageEvent(damageEvent);
-	}
+	DestructibleComponent->OnDamageEvent(damageEvent);
 }
 
 void FApexChunkReport::onStateChangeNotify(const NxApexChunkStateEventData& visibilityEvent)
 {
+	UDestructibleComponent* DestructibleComponent = Cast<UDestructibleComponent>(FPhysxUserData::Get<UPrimitiveComponent>(visibilityEvent.destructible->userData));
+	check(DestructibleComponent);
+
+	if (DestructibleComponent->IsPendingKill())	//don't notify if object is being destroyed
+	{
+		return;
+	}
 
+	DestructibleComponent->OnVisibilityEvent(visibilityEvent);
 }
 
 ///////// FApexPhysX3Interface //////////////////////////////////
