Fix vehicle tires not automatically fixing up when simulation type is set to default

diff --git a/Engine/Source/Runtime/Engine/Private/PhysicsEngine/BodyInstance.cpp b/Engine/Source/Runtime/Engine/Private/PhysicsEngine/BodyInstance.cpp
--- a/Engine/Source/Runtime/Engine/Private/PhysicsEngine/BodyInstance.cpp
+++ b/Engine/Source/Runtime/Engine/Private/PhysicsEngine/BodyInstance.cpp
@@ -1693,7 +1693,7 @@ void FBodyInstance::UpdateInstanceSimulatePhysics(bool bIgnoreParent)
 {
 	// In skeletal case, we need both our bone and skelcomponent flag to be true.
 	// This might be 'and'ing us with ourself, but thats fine.
-	const bool bUseSimulate = IsInstanceSimulatingPhysics();
+	const bool bUseSimulate = IsInstanceSimulatingPhysics(bIgnoreParent);
 
 #if WITH_PHYSX
 	PxRigidDynamic* PRigidDynamic = GetPxRigidDynamic();
