#ue4 - Call ClearJumpInput() in ACharacter::Restart().

diff --git a/Engine/Source/Runtime/Engine/Private/Character.cpp b/Engine/Source/Runtime/Engine/Private/Character.cpp
--- a/Engine/Source/Runtime/Engine/Private/Character.cpp
+++ b/Engine/Source/Runtime/Engine/Private/Character.cpp
@@ -560,6 +560,8 @@ void ACharacter::Restart()
 {
 	Super::Restart();
 	bPressedJump = false;
+	JumpKeyHoldTime = 0.0f;
+	ClearJumpInput();
 }
 
 void ACharacter::PawnClientRestart()
