#ue4 - Fix for replicated 'MovementMode' not being set on SimulatedProxies in some cases (dependent on framerate).

- Also fixes replicated movement mode being incorrect when playing root motion on proxies.

(integrate change by [REDACTED_PERSON])

diff --git a/Engine/Source/Runtime/Engine/Private/Character.cpp b/Engine/Source/Runtime/Engine/Private/Character.cpp
--- a/Engine/Source/Runtime/Engine/Private/Character.cpp
+++ b/Engine/Source/Runtime/Engine/Private/Character.cpp
@@ -831,7 +831,7 @@ void ACharacter::PostNetReceive()
 	if (Role == ROLE_SimulatedProxy)
 	{
 		CharacterMovement->bNetworkUpdateReceived = true;
-		CharacterMovement->bNetworkMovementModeChanged = (SavedMovementMode != ReplicatedMovementMode);
+		CharacterMovement->bNetworkMovementModeChanged = (CharacterMovement->bNetworkMovementModeChanged || (SavedMovementMode != ReplicatedMovementMode));
 	}
 
 	Super::PostNetReceive();
diff --git a/Engine/Source/Runtime/Engine/Private/CharacterMovementComponent.cpp b/Engine/Source/Runtime/Engine/Private/CharacterMovementComponent.cpp
--- a/Engine/Source/Runtime/Engine/Private/CharacterMovementComponent.cpp
+++ b/Engine/Source/Runtime/Engine/Private/CharacterMovementComponent.cpp
@@ -927,6 +927,13 @@ void UCharacterMovementComponent::SimulateRootMotion(float DeltaSeconds, const F
 		// Do not override Velocity.Z if falling.
 		Velocity = FVector(RootMotionVelocity.X, RootMotionVelocity.Y, (MovementMode == MOVE_Falling ? Velocity.Z : RootMotionVelocity.Z));
 
+		// Update replicated movement mode.
+		if (bNetworkMovementModeChanged)
+		{
+			bNetworkMovementModeChanged = false;
+			ApplyNetworkMovementMode(CharacterOwner->GetReplicatedMovementMode());
+		}
+
 		StartNewPhysics(DeltaSeconds, 0);
 		// fixme laurent - simulate movement seems to have step up issues? investigate as that would be cheaper to use.
 		// 		SimulateMovement(DeltaSeconds);
