Physics: Added limited support for collision response filtering in 2D
- Split out sim/query filter packing code to be somewhat PhysX/Box2D agnostic
- Adjust Box2D b2Filter structure to accomodate UE4 filtering requirements (incl. updated libraries for Win32/Win64)

diff --git a/Engine/Source/Runtime/Engine/Private/Collision/PhysXCollision.cpp b/Engine/Source/Runtime/Engine/Private/Collision/PhysXCollision.cpp
--- a/Engine/Source/Runtime/Engine/Private/Collision/PhysXCollision.cpp
+++ b/Engine/Source/Runtime/Engine/Private/Collision/PhysXCollision.cpp
@@ -1141,81 +1141,6 @@ PxTransform ConvertToPhysXCapsulePose(const FTransform& GeomPose)
 
 //////////////////////////////////////////////////////////////////////////
 
-
-void CreateShapeFilterData(const uint8 MyChannel, const int32 ActorID, const FCollisionResponseContainer& ResponseToChannels, uint32 SkelMeshCompID, uint16 BodyIndex, PxFilterData& OutQueryData, PxFilterData& OutSimData, bool bEnableCCD, bool bEnableContactNotify, bool bStaticShape)
-{
-	/**
-	 * Format for SimData : 
-	 * 		word0 (body index)
-	 *		word1 (blocking channels)
-	 *		word2 (skeletal mesh component ID)
-	 * 		word3 (MyChannel (top 8) + Flags (lower 24))
-	 * 
-	 * Format for QueryData : 
-	 *		word0 (actor ID)
-	 *		word1 (blocking channels)
-	 *		word2 (touching channels)
-	 *		word3 (MyChannel (top 8) as ECollisionChannel + Flags (lower 24))
-	 */
-	 
-	OutQueryData.setToDefault();
-	OutSimData.setToDefault();
-
-	// query word 0 has ActorID
-	OutQueryData.word0 = ActorID;
-
-	// query word 0 bodyindex 
-	OutSimData.word0 = BodyIndex;
-	// query word 2 stores asset instance ID
-	OutSimData.word2 = SkelMeshCompID;
-
-	// word1 encodes 'what i block', word2 encodes 'what i touch'
-	for(int32 i=0; i<ARRAY_COUNT(ResponseToChannels.EnumArray); i++)
-	{
-		// if i block, set that in word1
-		if(ResponseToChannels.EnumArray[i] == ECR_Block)
-		{
-			PxU32 const ChannelBit = CRC_TO_BITFIELD(i);
-			OutQueryData.word1 |= ChannelBit;
-			OutSimData.word1 |= ChannelBit;
-		}
-		else if(ResponseToChannels.EnumArray[i] == ECR_Overlap)
-		{
-			// if i touch, set that in word2
-			OutQueryData.word2 |= CRC_TO_BITFIELD(i);
-		}
-	}
-
-	// enable CCD at word 3
-	if(bEnableCCD)
-	{
-		OutSimData.word3 |= EPDF_CCD;
-	}
-
-	// Contact notify
-	if(bEnableContactNotify)
-	{
-		OutSimData.word3 |= EPDF_ContactNotify;
-	}
-
-	// Whether shape is static
-	if(bStaticShape)
-	{
-		OutSimData.word3 |= EPDF_StaticShape;
-	}
-	
-	// make sure it doesn't have flag beyond usage
-	// we only allow lower 24 bits for flag
-	// warn if that's used before setting MyChannel
-	check((OutSimData.word3 & 0xFF000000) == 0);
-
-	// pack MyChannel
-	OutSimData.word3 |=  (MyChannel << 24);
-
-	// QueryData word3 is same as SimData
-	OutQueryData.word3 = OutSimData.word3;
-}
-
 PxFilterData CreateObjectQueryFilterData(const bool bTraceComplex, const int32 MultiTrace/*=1 if multi. 0 otherwise*/, const struct FCollisionObjectQueryParams & ObjectParam)
 {
 	/**
diff --git a/Engine/Source/Runtime/Engine/Private/Collision/PhysXCollision.h b/Engine/Source/Runtime/Engine/Private/Collision/PhysXCollision.h
--- a/Engine/Source/Runtime/Engine/Private/Collision/PhysXCollision.h
+++ b/Engine/Source/Runtime/Engine/Private/Collision/PhysXCollision.h
@@ -183,11 +183,7 @@ PxTransform ConvertToPhysXCapsulePose(const FTransform& GeomPose);
 
 // FILTER DATA
 
-/** Utility for creating a PhysX PxFilterData for filtering query (trace) and sim (physics) from the Unreal filtering info. */
-void CreateShapeFilterData(const uint8 MyChannel, const int32 ActorID, const FCollisionResponseContainer& ResponseToChannels, uint32 SkelMeshCompID, uint16 BodyIndex, PxFilterData& OutQueryData, PxFilterData& OutSimData, bool bEnableCCD, bool bEnableContactNotify, bool bStaticShape);
-
 /** Utility for creating a PhysX PxFilterData for performing a query (trace) against the scene */
 PxFilterData CreateQueryFilterData(const uint8 MyChannel, const bool bTraceComplex, const FCollisionResponseContainer & InCollisionResponseContainer, const struct FCollisionObjectQueryParams & ObjectParam, const bool bMultitrace);
 
 #endif // WITH_PHYX
-
diff --git a/Engine/Source/Runtime/Engine/Private/Collision/PhysicsFiltering.cpp b/Engine/Source/Runtime/Engine/Private/Collision/PhysicsFiltering.cpp
new file mode 100644
--- /dev/null
+++ b/Engine/Source/Runtime/Engine/Private/Collision/PhysicsFiltering.cpp
@@ -0,0 +1,31 @@
+// Copyright 1998-2014 Epic Games, Inc. All Rights Reserved.
+
+#include "EnginePrivate.h"
+
+#include "PhysicsFiltering.h"
+
+//////////////////////////////////////////////////////////////////////////
+// FPhysicsFilterBuilder
+
+FPhysicsFilterBuilder::FPhysicsFilterBuilder(TEnumAsByte<enum ECollisionChannel> InObjectType, const struct FCollisionResponseContainer& ResponseToChannels)
+	: BlockingBits(0)
+	, TouchingBits(0)
+	, Word3(0)
+{
+	for (int32 i = 0; i < ARRAY_COUNT(ResponseToChannels.EnumArray); i++)
+	{
+		if (ResponseToChannels.EnumArray[i] == ECR_Block)
+		{
+			const uint32 ChannelBit = CRC_TO_BITFIELD(i);
+			BlockingBits |= ChannelBit;
+		}
+		else if (ResponseToChannels.EnumArray[i] == ECR_Overlap)
+		{
+			const uint32 ChannelBit = CRC_TO_BITFIELD(i);
+			TouchingBits |= ChannelBit;
+		}
+	}
+
+	uint8 MyChannel = InObjectType;
+	Word3 = (MyChannel << 24);
+}
diff --git a/Engine/Source/Runtime/Engine/Private/Collision/PhysicsFiltering.h b/Engine/Source/Runtime/Engine/Private/Collision/PhysicsFiltering.h
new file mode 100644
--- /dev/null
+++ b/Engine/Source/Runtime/Engine/Private/Collision/PhysicsFiltering.h
@@ -0,0 +1,106 @@
+// Copyright 1998-2014 Epic Games, Inc. All Rights Reserved.
+
+#pragma once
+
+/** 
+ * Set of flags stored in the PhysX FilterData
+ *
+ * When this flag is saved in CreateShapeFilterData or CreateQueryFilterData, we only use 24 bits
+ * If you plan to use more than 24 bits, you'll also need to change the format of ShapeFilterData,QueryFilterData
+ * Make sure you also change preFilter/SimFilterShader where it's used
+ */
+enum EPhysXFilterDataFlags
+{
+	EPDF_SimpleCollision	=	0x0001,
+	EPDF_ComplexCollision	=	0x0002,
+	EPDF_CCD				=	0x0004,
+	EPDF_ContactNotify		=	0x0008,
+	EPDF_StaticShape		=	0x0010
+};
+
+struct FPhysicsFilterBuilder
+{
+	FPhysicsFilterBuilder(TEnumAsByte<enum ECollisionChannel> InObjectType, const struct FCollisionResponseContainer& ResponseToChannels);
+
+	inline void ConditionalSetFlags(EPhysXFilterDataFlags Flag, bool bEnabled)
+	{
+		if (bEnabled)
+		{
+			Word3 |= Flag;
+		}
+	}
+
+	inline void GetQueryData(uint32 ActorID, uint32& OutWord0, uint32& OutWord1, uint32& OutWord2, uint32& OutWord3) const
+	{
+		/**
+		 * Format for QueryData : 
+		 *		word0 (actor ID)
+		 *		word1 (blocking channels)
+		 *		word2 (touching channels)
+		 *		word3 (MyChannel (top 8) as ECollisionChannel + Flags (lower 24))
+		 */
+		OutWord0 = ActorID;
+		OutWord1 = BlockingBits;
+		OutWord2 = TouchingBits;
+		OutWord3 = Word3;
+	}
+
+	inline void GetSimData(uint32 BodyIndex, uint32 SkelMeshCompID, uint32& OutWord0, uint32& OutWord1, uint32& OutWord2, uint32& OutWord3) const
+	{
+		/**
+		 * Format for SimData : 
+		 * 		word0 (body index)
+		 *		word1 (blocking channels)
+		 *		word2 (skeletal mesh component ID)
+		 * 		word3 (MyChannel (top 8) + Flags (lower 24))
+		 */
+		OutWord0 = BodyIndex;
+		OutWord1 = BlockingBits;
+		OutWord2 = SkelMeshCompID;
+		OutWord3 = Word3;
+	}
+
+	inline void GetCombinedData(uint32& OutBlockingBits, uint32& OutTouchingBits, uint32& OutObjectTypeAndFlags) const
+	{
+		OutBlockingBits = BlockingBits;
+		OutTouchingBits = TouchingBits;
+		OutObjectTypeAndFlags = Word3;
+	}
+
+private:
+	uint32 BlockingBits;
+	uint32 TouchingBits;
+	uint32 Word3;
+};
+
+/** Utility for creating a PhysX PxFilterData for filtering query (trace) and sim (physics) from the Unreal filtering info. */
+#if WITH_PHYSX
+inline void CreateShapeFilterData(
+	const uint8 MyChannel,
+	const int32 ActorID,
+	const FCollisionResponseContainer& ResponseToChannels,
+	uint32 SkelMeshCompID,
+	uint16 BodyIndex,
+	PxFilterData& OutQueryData,
+	PxFilterData& OutSimData,
+	bool bEnableCCD,
+	bool bEnableContactNotify,
+	bool bStaticShape)
+{
+	FPhysicsFilterBuilder Builder((ECollisionChannel)MyChannel, ResponseToChannels);
+	Builder.ConditionalSetFlags(EPDF_CCD, bEnableCCD);
+	Builder.ConditionalSetFlags(EPDF_ContactNotify, bEnableContactNotify);
+	Builder.ConditionalSetFlags(EPDF_StaticShape, bStaticShape);
+
+	OutQueryData.setToDefault();
+	OutSimData.setToDefault();
+	Builder.GetQueryData(ActorID, OutQueryData.word0, OutQueryData.word1, OutQueryData.word2, OutQueryData.word3);
+	Builder.GetSimData(BodyIndex, SkelMeshCompID, OutSimData.word0, OutSimData.word1, OutSimData.word2, OutSimData.word3);
+}
+#endif //WITH_PHYSX
+
+/** Utility for creating a PhysX PxFilterData for performing a query (trace) against the scene */
+// PxFilterData ZZ_CreateQueryFilterData(const uint8 MyChannel, const bool bTraceComplex,
+// 	const FCollisionResponseContainer & InCollisionResponseContainer,
+// 	const struct FCollisionObjectQueryParams & ObjectParam,
+// 	const bool bMultitrace);
diff --git a/Engine/Source/Runtime/Engine/Private/PhysicsEngine/BodyInstance.cpp b/Engine/Source/Runtime/Engine/Private/PhysicsEngine/BodyInstance.cpp
--- a/Engine/Source/Runtime/Engine/Private/PhysicsEngine/BodyInstance.cpp
+++ b/Engine/Source/Runtime/Engine/Private/PhysicsEngine/BodyInstance.cpp
@@ -565,13 +565,13 @@ void FBodyInstance::UpdatePhysicsFilterData()
 	if (PxRigidActor* PActor = GetPxRigidActor())
 	{
 		// Create the filterdata structs
-		PxFilterData PQueryFilterData, PSimFilterData;
-		PxFilterData PSimpleQueryData, PComplexQueryData;
+		PxFilterData PSimFilterData;
+		PxFilterData PSimpleQueryData;
+		PxFilterData PComplexQueryData;
 		if (UseCollisionEnabled != ECollisionEnabled::NoCollision)
 		{
-			CreateShapeFilterData(ObjectType, OwnerID, UseResponse, SkelMeshCompID, InstanceBodyIndex, PQueryFilterData, PSimFilterData, bUseCCD && !bPhysicsStatic, bUseNotifyRBCollision, bPhysicsStatic);
-			PSimpleQueryData = PQueryFilterData;
-			PComplexQueryData = PQueryFilterData;
+			CreateShapeFilterData(ObjectType, OwnerID, UseResponse, SkelMeshCompID, InstanceBodyIndex, PSimpleQueryData, PSimFilterData, bUseCCD && !bPhysicsStatic, bUseNotifyRBCollision, bPhysicsStatic);
+			PComplexQueryData = PSimpleQueryData;
 
 			// Build filterdata variations for complex and simple
 			PSimpleQueryData.word3 |= EPDF_SimpleCollision;
@@ -607,7 +607,7 @@ void FBodyInstance::UpdatePhysicsFilterData()
 			PShape->setSimulationFilterData(PSimFilterData);
 
 			// If query collision is enabled..
-			if (UseCollisionEnabled == ECollisionEnabled::QueryAndPhysics || UseCollisionEnabled == ECollisionEnabled::QueryOnly)
+			if ((UseCollisionEnabled == ECollisionEnabled::QueryAndPhysics) || (UseCollisionEnabled == ECollisionEnabled::QueryOnly))
 			{
 				// Only perform scene queries in the synchronous scene for static shapes
 				if (bPhysicsStatic)
@@ -695,16 +695,37 @@ void FBodyInstance::UpdatePhysicsFilterData()
 #endif
 
 #if WITH_BOX2D
-
 	if (BodyInstancePtr != nullptr)
 	{
-		//@TODO: BOX2D: Implement UpdatePhysicsFilterData
+		if (UseCollisionEnabled != ECollisionEnabled::NoCollision)
+		{
+			// Create the simulation/query filter data
+			FPhysicsFilterBuilder FilterBuilder(ObjectType, UseResponse);
+ 			FilterBuilder.ConditionalSetFlags(EPDF_CCD, bUseCCD && !bPhysicsStatic);
+			FilterBuilder.ConditionalSetFlags(EPDF_ContactNotify, bUseNotifyRBCollision);
+ 			FilterBuilder.ConditionalSetFlags(EPDF_StaticShape, bPhysicsStatic);
 
-		bool bUpdateMassProperties = false;
+			b2Filter BoxSimFilterData;
+			FilterBuilder.GetCombinedData(/*out*/ BoxSimFilterData.BlockingChannels, /*out*/ BoxSimFilterData.TouchingChannels, /*out*/ BoxSimFilterData.ObjectTypeAndFlags);
+			BoxSimFilterData.UniqueComponentID = SkelMeshCompID;
+			BoxSimFilterData.BodyIndex = InstanceBodyIndex;
 
-		if (bUpdateMassProperties)
+			// Update the body data
+			const bool bSimCollision = (UseCollisionEnabled == ECollisionEnabled::QueryAndPhysics);
+			BodyInstancePtr->SetBullet(bSimCollision && !bPhysicsStatic && bUseCCD);
+			BodyInstancePtr->SetActive(true);
+
+			// Copy the filter data to each fixture in the body
+			for (b2Fixture* Fixture = BodyInstancePtr->GetFixtureList(); Fixture; Fixture = Fixture->GetNext())
+			{
+				Fixture->SetFilterData(BoxSimFilterData);
+				Fixture->SetSensor(UseCollisionEnabled == ECollisionEnabled::QueryOnly);
+			}
+		}
+		else
 		{
-			UpdateMassProperties();
+			// No collision
+			BodyInstancePtr->SetActive(false);
 		}
 	}
 #endif
@@ -2551,7 +2572,7 @@ void FBodyInstance::AddForce(const FVector& Force, bool bAllowSubstepping)
 {
 #if WITH_PHYSX
 	PxRigidDynamic* PRigidDynamic = GetPxRigidDynamic();
-	if(IsRigidDynamicNonKinematic(PRigidDynamic))
+	if (IsRigidDynamicNonKinematic(PRigidDynamic))
 	{
 		const PxScene* PScene = PRigidDynamic->getScene();
 		FPhysScene* PhysScene = FPhysxUserData::Get<FPhysScene>(PScene->userData);
@@ -2567,7 +2588,7 @@ void FBodyInstance::AddForceAtPosition(const FVector& Force, const FVector& Posi
 {
 #if WITH_PHYSX
 	PxRigidDynamic* PRigidDynamic = GetPxRigidDynamic();
-	if(IsRigidDynamicNonKinematic(PRigidDynamic))
+	if (IsRigidDynamicNonKinematic(PRigidDynamic))
 	{
 		const PxScene* PScene = PRigidDynamic->getScene();
 		FPhysScene* PhysScene = FPhysxUserData::Get<FPhysScene>(PScene->userData);
@@ -2582,7 +2603,7 @@ void FBodyInstance::AddTorque(const FVector& Torque, bool bAllowSubstepping)
 {
 #if WITH_PHYSX
 	PxRigidDynamic* PRigidDynamic = GetPxRigidDynamic();
-	if(IsRigidDynamicNonKinematic(PRigidDynamic))
+	if (IsRigidDynamicNonKinematic(PRigidDynamic))
 	{
 		const PxScene* PScene = PRigidDynamic->getScene();
 		FPhysScene* PhysScene = FPhysxUserData::Get<FPhysScene>(PScene->userData);
@@ -2597,7 +2618,7 @@ void FBodyInstance::AddImpulse(const FVector& Impulse, bool bVelChange)
 {
 #if WITH_PHYSX
 	PxRigidDynamic* PRigidDynamic = GetPxRigidDynamic();
-	if(IsRigidDynamicNonKinematic(PRigidDynamic))
+	if (IsRigidDynamicNonKinematic(PRigidDynamic))
 	{
 		PxForceMode::Enum Mode = bVelChange ? PxForceMode::eVELOCITY_CHANGE : PxForceMode::eIMPULSE;
 		PRigidDynamic->addForce(U2PVector(Impulse), Mode, true);
@@ -2611,7 +2632,7 @@ void FBodyInstance::AddImpulseAtPosition(const FVector& Impulse, const FVector&
 {
 #if WITH_PHYSX
 	PxRigidDynamic* PRigidDynamic = GetPxRigidDynamic();
-	if(IsRigidDynamicNonKinematic(PRigidDynamic))
+	if (IsRigidDynamicNonKinematic(PRigidDynamic))
 	{
 		PxForceMode::Enum Mode = PxForceMode::eIMPULSE; // does not support eVELOCITY_CHANGE
 		PxRigidBodyExt::addForceAtPos(*PRigidDynamic, U2PVector(Impulse), U2PVector(Position), Mode, true);
diff --git a/Engine/Source/Runtime/Engine/Private/PhysicsEngine/PhysXSupport.h b/Engine/Source/Runtime/Engine/Private/PhysicsEngine/PhysXSupport.h
--- a/Engine/Source/Runtime/Engine/Private/PhysicsEngine/PhysXSupport.h
+++ b/Engine/Source/Runtime/Engine/Private/PhysicsEngine/PhysXSupport.h
@@ -147,22 +147,6 @@ const bool bGlobalCCD = true;
 
 /////// UTILS
 
-/** 
- * Set of flags stored in the PhysX FilterData
- *
- * When this flag is saved in CreateShapeFilterData or CreateQueryFilterData, we only use 24 bits
- * If you plan to use more than 24 bits, you'll also need to change the format of ShapeFilterData,QueryFilterData
- * Make sure you also change preFilter/SimFilterShader where it's used
- */
-enum EPhysXFilterDataFlags
-{
-	EPDF_SimpleCollision	=	0x0001,
-	EPDF_ComplexCollision	=	0x0002,
-	EPDF_CCD				=	0x0004,
-	EPDF_ContactNotify		=	0x0008,
-	EPDF_StaticShape		=	0x0010
-};
-
 /** Get a pointer to the PxScene from an SceneIndex (will be NULL if scene already shut down) */
 PxScene* GetPhysXSceneFromIndex(int32 InSceneIndex);
 
@@ -637,3 +621,6 @@ public:
  **/
 ENGINE_API SIZE_T GetPhysxObjectSize(PxBase* Obj, const PxCollection* SharedCollection);
 #endif // WITH_PHYSX
+
+
+#include "../Collision/PhysicsFiltering.h"
diff --git a/Engine/Source/Runtime/Engine/Private/PhysicsEngine2D/Box2DIntegration.cpp b/Engine/Source/Runtime/Engine/Private/PhysicsEngine2D/Box2DIntegration.cpp
--- a/Engine/Source/Runtime/Engine/Private/PhysicsEngine2D/Box2DIntegration.cpp
+++ b/Engine/Source/Runtime/Engine/Private/PhysicsEngine2D/Box2DIntegration.cpp
@@ -125,7 +125,84 @@ FPhysicsScene2D::~FPhysicsScene2D()
 
 bool FPhysicsScene2D::ShouldCollide(b2Fixture* FixtureA, b2Fixture* FixtureB)
 {
-	return b2ContactFilter::ShouldCollide(FixtureA, FixtureB);
+	//@TODO: BOX2D: May need to extend Box2D to be able to indicate that we don't want them to collide now but keep tracking the interaction, not sure!
+	const bool KillInteraction = false;
+	const bool SupressInteraction = false;
+	const bool NormalInteraction = true;
+
+	const b2Filter& FilterDataA = FixtureA->GetFilterData();
+	const b2Filter& FilterDataB = FixtureB->GetFilterData();
+
+	const ECollisionChannel ChannelA = static_cast<ECollisionChannel>(FilterDataA.ObjectTypeAndFlags >> 24);
+	const ECollisionChannel ChannelB = static_cast<ECollisionChannel>(FilterDataB.ObjectTypeAndFlags >> 24);
+	const b2BodyType BodyTypeA = FixtureA->GetBody()->GetType();
+	const b2BodyType BodyTypeB = FixtureB->GetBody()->GetType();
+
+	// ignore kinematic-kinematic interactions which don't involve a destructible
+	const bool bIsKinematicA = BodyTypeA == b2_kinematicBody;
+	const bool bIsKinematicB = BodyTypeB == b2_kinematicBody;
+	if (bIsKinematicA && bIsKinematicB && (ChannelA != ECC_Destructible) && (ChannelB != ECC_Destructible))
+	{
+		return KillInteraction;
+	}
+
+	// ignore static-kinematic interactions
+	const bool bIsStaticA = BodyTypeA == b2_staticBody;
+	const bool bIsStaticB = BodyTypeB == b2_staticBody;
+	if ((bIsKinematicA || bIsKinematicB) && (bIsStaticA || bIsStaticB))
+	{
+		// should return eSUPPRESS here instead eKILL so that kinematics vs statics will still be considered once kinematics become dynamic (dying ragdoll case)
+		return SupressInteraction;
+	}
+
+	//@TODO: BOX2D: Skeletal collision filtering
+#if 0
+	// if these bodies are from the same skeletal mesh component, use the disable table to see if we should disable collision
+	if ((FilterDataA.UniqueComponentID != 0) && (FilterDataA.UniqueComponentID == FilterDataB.UniqueComponentID))
+	{
+		FPhysScene* PhysScene = *magic*;
+		check(PhysScene);
+
+		const TMap<uint32, TMap<FRigidBodyIndexPair, bool> *> & CollisionDisableTableLookup = PhysScene->GetCollisionDisableTableLookup();
+		TMap<FRigidBodyIndexPair, bool>* const * DisableTablePtrPtr = CollisionDisableTableLookup.FindChecked(FilterDataA.UniqueComponentID);
+		TMap<FRigidBodyIndexPair, bool>* DisableTablePtr = *DisableTablePtrPtr;
+		FRigidBodyIndexPair BodyPair(FilterDataA.BodyIndex, FilterDataB.BodyIndex);
+		if (DisableTablePtr->Find(BodyPair))
+		{
+			return KillInteraction;
+		}
+	}
+#endif
+
+	// See if either one would like to block the other
+	const uint32 BlockFlagToB = ECC_TO_BITFIELD(ChannelB) & FilterDataA.BlockingChannels;
+	const uint32 BlockFlagToA = ECC_TO_BITFIELD(ChannelA) & FilterDataB.BlockingChannels;
+	const bool bDoesWantToBlock = BlockFlagToB && BlockFlagToA;
+	if (!bDoesWantToBlock)
+	{
+		return SupressInteraction;
+	}
+
+	//@TODO: BOX2D: CCD / contact notification flags
+#if 0
+	const uint32 FilterFlags0 = FilterDataA.ObjectTypeAndFlags & 0xFFFFFF;
+	const uint32 FilterFlags1 = FilterDataB.ObjectTypeAndFlags & 0xFFFFFF;
+
+	bool bResult = NormalInteraction;
+
+	//todo enabling CCD objects against everything else for now
+	if (!(bIsKinematicA && bIsKinematicB) && ((FilterFlagsA & EPDF_CCD) || (FilterFlagsB & EPDF_CCD)))
+	{
+		pairFlags |= PxPairFlag::eCCD_LINEAR;
+	}
+
+	if ((FilterFlagsA & EPDF_ContactNotify) || (FilterFlagsB & EPDF_ContactNotify))
+	{
+		pairFlags |= (PxPairFlag::eNOTIFY_TOUCH_FOUND | PxPairFlag::eNOTIFY_TOUCH_PERSISTS | PxPairFlag::eNOTIFY_CONTACT_POINTS);
+	}
+#endif
+
+	return NormalInteraction;
 }
 
 //////////////////////////////////////////////////////////////////////////
