Integrate Xbox Live voice chat, from Lionhead
Works in ShooterGame

diff --git a/Engine/Source/Runtime/Engine/Private/VoiceChannel.cpp b/Engine/Source/Runtime/Engine/Private/VoiceChannel.cpp
--- a/Engine/Source/Runtime/Engine/Private/VoiceChannel.cpp
+++ b/Engine/Source/Runtime/Engine/Private/VoiceChannel.cpp
@@ -59,19 +59,28 @@ void UVoiceChannel::ReceivedBunch(FInBunch& Bunch)
  */
 void UVoiceChannel::Tick()
 {
-	// If the handshaking hasn't completed throw away all voice data
+	// If the handshaking hasn't completed throw away all unreliable voice data
 	if (Connection->PlayerController &&
 		Connection->PlayerController->MuteList.bHasVoiceHandshakeCompleted)
 	{
 		// Try to append each packet in turn
-		for (int32 Index = 0; Index < VoicePackets.Num(); Index++)
+		int32 Index;
+		for (Index = 0; Index < VoicePackets.Num(); Index++)
 		{
+			if (Connection->IsNetReady(0) == false)
+			{
+				// If the network is saturated bail early
+				UE_LOG(LogNet, Warning, TEXT("Network saturated"));
+				break;
+			}
+
 			FOutBunch Bunch(this, 0);
-			// Don't want reliable delivery. The bunch will be lost if the connection is saturated
-			// First send needs to be reliable
-			Bunch.bReliable = OpenAcked == false;
 
 			TSharedPtr<FVoicePacket> Packet = VoicePackets[Index];
+
+			// First send must be reliable as must any packet marked reliable
+			Bunch.bReliable = OpenAcked == false || Packet->IsReliable();
+
 			// Append the packet data (copies into the bunch)
 			Packet->Serialize(Bunch);
 
@@ -86,17 +95,41 @@ void UVoiceChannel::Tick()
 				// Submit the bunching with merging on
 				SendBunch(&Bunch, 1);
 			}
-			// If the network is saturated, throw away any remaining packets
-			if (Connection->IsNetReady(0) == false)
+			else
 			{
-				// Empty early
-				VoicePackets.Empty();
+				//bail and try again next frame
+				UE_LOG(LogNet, Warning, TEXT("Bunch error"));
+				break;
 			}
 		}
+
+		if(Index >= VoicePackets.Num())
+		{
+			// all sent, can throw everything away
+			VoicePackets.Empty();
+		}
+		else if(Index > 0)
+		{
+			// didn't send everything, just remove all packets actually sent
+			VoicePackets.RemoveAt(0,Index);
+		}
 	}
 
-	// Let the packets free itself if no longer in use
-	VoicePackets.Empty();
+	// make sure we keep any reliable messages to try again next time
+	// but ditch any unreliable messages we've not managed to send
+	int PacketLoss = 0;
+	for(int i=VoicePackets.Num() - 1; i >= 0; i--)
+	{
+		if(!VoicePackets[i]->IsReliable())
+		{
+			VoicePackets.RemoveAt(i,1,false);
+			PacketLoss++;
+		}
+	}
+	if(PacketLoss > 0)
+	{
+		UE_LOG(LogNet, Warning, TEXT("Dropped %d packets due to congestion in the voicechannel"), PacketLoss);
+	}
 }
 
 /**
diff --git a/Engine/Source/Runtime/Engine/Public/Net/VoiceDataCommon.h b/Engine/Source/Runtime/Engine/Public/Net/VoiceDataCommon.h
--- a/Engine/Source/Runtime/Engine/Public/Net/VoiceDataCommon.h
+++ b/Engine/Source/Runtime/Engine/Public/Net/VoiceDataCommon.h
@@ -52,6 +52,9 @@ public:
 	/** @return the sender of this voice packet */
 	virtual TSharedPtr<class FUniqueNetId> GetSender() = 0;
 
+	/** @return true if this packet should be sent reliably */
+	virtual bool IsReliable() OVERRIDE = 0;
+
 	/** 
 	 * Serialize the voice packet data to a buffer 
 	 *
diff --git a/Engine/Source/Runtime/Online/OnlineSubsystemSteam/Private/VoicePacketSteam.h b/Engine/Source/Runtime/Online/OnlineSubsystemSteam/Private/VoicePacketSteam.h
--- a/Engine/Source/Runtime/Online/OnlineSubsystemSteam/Private/VoicePacketSteam.h
+++ b/Engine/Source/Runtime/Online/OnlineSubsystemSteam/Private/VoicePacketSteam.h
@@ -49,6 +49,8 @@ public:
 	/** @return the sender of this voice packet */
 	virtual TSharedPtr<class FUniqueNetId> GetSender() OVERRIDE;
 
+	virtual bool IsReliable() OVERRIDE { return false; }
+
 	/** 
 	 * Serialize the voice packet data to a buffer 
 	 *
diff --git a/Engine/Source/Runtime/Online/OnlineSubsystemUtils/Public/VoicePacketImpl.h b/Engine/Source/Runtime/Online/OnlineSubsystemUtils/Public/VoicePacketImpl.h
--- a/Engine/Source/Runtime/Online/OnlineSubsystemUtils/Public/VoicePacketImpl.h
+++ b/Engine/Source/Runtime/Online/OnlineSubsystemUtils/Public/VoicePacketImpl.h
@@ -52,6 +52,8 @@ public:
 	/** @return the sender of this voice packet */
 	virtual TSharedPtr<class FUniqueNetId> GetSender() OVERRIDE;
 
+	virtual bool IsReliable() OVERRIDE { return false; }
+
 	/** 
 	 * Serialize the voice packet data to a buffer 
 	 *
