Fixing compiler error is VoiceDataCommon.h

diff --git a/Engine/Source/Runtime/Engine/Public/Net/VoiceDataCommon.h b/Engine/Source/Runtime/Engine/Public/Net/VoiceDataCommon.h
--- a/Engine/Source/Runtime/Engine/Public/Net/VoiceDataCommon.h
+++ b/Engine/Source/Runtime/Engine/Public/Net/VoiceDataCommon.h
@@ -39,7 +39,7 @@ public:
 	 * @param Other packet to copy
 	 * @param InRefCount the starting ref count to use
 	 */
- 	FVoicePacket(const FVoicePacket& Other)
+	FVoicePacket(const FVoicePacket& Other)
 	{
 	}
 
@@ -53,7 +53,7 @@ public:
 	virtual TSharedPtr<class FUniqueNetId> GetSender() = 0;
 
 	/** @return true if this packet should be sent reliably */
-	virtual bool IsReliable() OVERRIDE = 0;
+	virtual bool IsReliable() = 0;
 
 	/** 
 	 * Serialize the voice packet data to a buffer 
