Fix numerous issues with undo/redo and attachment when the root component of the Actor is a construction script created component.
- Construction script created components are now not transactional (i.e. not in Undo buffer)
- Use component instance data cache to reconstruct construction script created components during undo/redo
- ChildActorComponent spawned ChildActors are not transactional

diff --git a/Engine/Source/Editor/UnrealEd/Private/EditorObject.cpp b/Engine/Source/Editor/UnrealEd/Private/EditorObject.cpp
--- a/Engine/Source/Editor/UnrealEd/Private/EditorObject.cpp
+++ b/Engine/Source/Editor/UnrealEd/Private/EditorObject.cpp
@@ -437,7 +437,7 @@ static const TCHAR* ImportProperties(
 				}
 
 				// Propagate object flags to the sub object.
-				const EObjectFlags NewFlags = SubobjectOuter->GetMaskedFlags( RF_PropagateToSubObjects );
+				EObjectFlags NewFlags = SubobjectOuter->GetMaskedFlags( RF_PropagateToSubObjects );
 
 				if (!Archetype) // no override and we didn't find one from the class table, so go with the base
 				{
@@ -476,6 +476,16 @@ static const TCHAR* ImportProperties(
 				}
 				else
 				{
+					// We do not want to set RF_Transactional for construction script created components, so we have to monkey with things here
+					if (NewFlags & RF_Transactional)
+					{
+						UActorComponent* Component = Cast<UActorComponent>(ComponentTemplate);
+						if (Component && Component->bCreatedByConstructionScript)
+						{
+							NewFlags &= ~RF_Transactional;
+						}
+					}
+
 					// Make sure desired flags are set - existing object could be pending kill
 					ComponentTemplate->ClearFlags(RF_AllFlags);
 					ComponentTemplate->SetFlags(NewFlags);
diff --git a/Engine/Source/Editor/UnrealEd/Private/LevelEditorViewport.cpp b/Engine/Source/Editor/UnrealEd/Private/LevelEditorViewport.cpp
--- a/Engine/Source/Editor/UnrealEd/Private/LevelEditorViewport.cpp
+++ b/Engine/Source/Editor/UnrealEd/Private/LevelEditorViewport.cpp
@@ -1684,52 +1684,30 @@ void FTrackingTransaction::Begin(const FText& Description)
 
 	TrackingTransactionState = ETransactionState::Active;
 
-	TSet<AActor*> RelatedActors;
+	TSet<AGroupActor*> GroupActors;
 
 	// Modify selected actors to record their state at the start of the transaction
-	for ( FSelectionIterator It( GEditor->GetSelectedActorIterator() ) ; It ; ++It )
+	for (FSelectionIterator It(GEditor->GetSelectedActorIterator()); It; ++It)
 	{
-		AActor* Actor = CastChecked<AActor>( *It );
+		AActor* Actor = CastChecked<AActor>(*It);
 
 		Actor->Modify();
 
-		if (Actor->GetRootComponent())
-		{
-			TArray<USceneComponent*> RootComponents;
-			RootComponents.Add(Actor->GetRootComponent());
-			while (RootComponents.Num() > 0)
-			{
-				USceneComponent* RootComponent = RootComponents.Pop();
-				for (USceneComponent* AttachedComponent : RootComponent->AttachChildren)
-				{
-					if (AttachedComponent)
-					{
-						AActor* AttachedActor = AttachedComponent->GetOwner();
-						if (AttachedComponent == AttachedActor->GetRootComponent())
-						{
-							RelatedActors.Add(AttachedActor);
-							RootComponents.Add(AttachedComponent);
-						}
-					}
-				}
-			}
-		}
-
 		if (GEditor->bGroupingActive)
 		{
 			// if this actor is in a group, add the GroupActor into a list to be modified shortly
 			AGroupActor* ActorLockedRootGroup = AGroupActor::GetRootForActor(Actor, true);
 			if (ActorLockedRootGroup != nullptr)
 			{
-				RelatedActors.Add(ActorLockedRootGroup);
+				GroupActors.Add(ActorLockedRootGroup);
 			}
 		}
 	}
 
 	// Modify unique group actors
-	for (AActor* RelatedActor : RelatedActors)
+	for (auto* GroupActor : GroupActors)
 	{
-		RelatedActor->Modify();
+		GroupActor->Modify();
 	}
 }
 
diff --git a/Engine/Source/Runtime/Engine/Classes/GameFramework/Actor.h b/Engine/Source/Runtime/Engine/Classes/GameFramework/Actor.h
--- a/Engine/Source/Runtime/Engine/Classes/GameFramework/Actor.h
+++ b/Engine/Source/Runtime/Engine/Classes/GameFramework/Actor.h
@@ -6,6 +6,7 @@
 #include "Runtime/Engine/Classes/Camera/CameraTypes.h"
 #include "Components/SceneComponent.h"
 #include "RenderingThread.h"
+#include "ComponentInstanceDataCache.h"
 
 UENUM(BlueprintType)
 namespace EEndPlayReason
@@ -1043,7 +1044,48 @@ public:
 #if WITH_EDITOR
 	virtual void PreEditChange(UProperty* PropertyThatWillChange) OVERRIDE;
 	virtual void PostEditChangeProperty(FPropertyChangedEvent& PropertyChangedEvent) OVERRIDE;
+	virtual void PreEditUndo() OVERRIDE;
 	virtual void PostEditUndo() OVERRIDE;
+
+	struct FActorRootComponentReconstructionData
+	{
+		// Struct to store info about attached actors
+		struct FAttachedActorInfo
+		{
+			TWeakObjectPtr<AActor> Actor;
+			FName SocketName;
+			FTransform RelativeTransform;
+		};
+
+		// The RootComponent's transform
+		FTransform Transform;
+
+		// The Actor the RootComponent is attached to
+		FAttachedActorInfo AttachedParentInfo;
+
+		// Actors that are attached to this RootComponent
+		TArray<FAttachedActorInfo> AttachedToInfo;
+	};
+
+	class FActorTransactionAnnotation : public ITransactionObjectAnnotation
+	{
+	public:
+		FActorTransactionAnnotation(const AActor* Actor);
+
+		bool HasInstanceData() const;
+
+		FComponentInstanceDataCache ComponentInstanceData;
+
+		// Root component reconstruction data
+		bool bRootComponentDataCached;
+		FActorRootComponentReconstructionData RootComponentData;
+	};
+
+	/* Cached pointer to the transaction annotation data from PostEditUndo to be used in the next RerunConstructionScript */
+	TSharedPtr<FActorTransactionAnnotation> CurrentTransactionAnnotation;
+
+	virtual TSharedPtr<ITransactionObjectAnnotation> GetTransactionAnnotation() const OVERRIDE;
+	virtual void PostEditUndo(TSharedPtr<ITransactionObjectAnnotation> TransactionAnnotation) OVERRIDE;
 #endif // WITH_EDITOR
 	// End UObject Interface
 
diff --git a/Engine/Source/Runtime/Engine/Private/ActorConstruction.cpp b/Engine/Source/Runtime/Engine/Private/ActorConstruction.cpp
--- a/Engine/Source/Runtime/Engine/Private/ActorConstruction.cpp
+++ b/Engine/Source/Runtime/Engine/Private/ActorConstruction.cpp
@@ -139,61 +139,102 @@ void AActor::RerunConstructionScripts()
 		GUndo = NULL;
 		
 		// Create cache to store component data across rerunning construction scripts
-		FComponentInstanceDataCache InstanceDataCache(this);
-
-		// If there are attached objects detach them and store the socket names
-		TArray<AActor*> AttachedActors;
-		GetAttachedActors(AttachedActors);
+		FActorTransactionAnnotation* ActorTransactionAnnotation = CurrentTransactionAnnotation.Get();
+		FComponentInstanceDataCache* InstanceDataCache;
+		
+		FTransform OldTransform = FTransform::Identity;
+		FName  SocketName;
+		AActor* Parent = NULL;
 
 		// Struct to store info about attached actors
 		struct FAttachedActorInfo
 		{
 			AActor* AttachedActor;
 			FName AttachedToSocket;
+			bool bSetRelativeTransform;
+			FTransform RelativeTransform;
 		};
 
 		// Save info about attached actors
 		TArray<FAttachedActorInfo> AttachedActorInfos;
-		for( AActor* AttachedActor : AttachedActors)
+
+		if (ActorTransactionAnnotation)
 		{
-			USceneComponent* EachRoot = AttachedActor->GetRootComponent();
-			// If the component we are attached to is about to go away...
-			if( EachRoot && EachRoot->AttachParent && EachRoot->AttachParent->bCreatedByConstructionScript )
+			InstanceDataCache = &ActorTransactionAnnotation->ComponentInstanceData;
+
+			if (ActorTransactionAnnotation->bRootComponentDataCached)
 			{
-				// Save info about actor to reattach
-				FAttachedActorInfo Info;
-				Info.AttachedActor = AttachedActor;
-				Info.AttachedToSocket = EachRoot->AttachSocketName;
-				AttachedActorInfos.Add(Info);
-
-				// Now detach it
-				AttachedActor->Modify();
-				EachRoot->DetachFromParent(true);					
+				OldTransform = ActorTransactionAnnotation->RootComponentData.Transform;
+				Parent = ActorTransactionAnnotation->RootComponentData.AttachedParentInfo.Actor.Get();
+				if (Parent)
+				{
+					DetachRootComponentFromParent();
+					SocketName = ActorTransactionAnnotation->RootComponentData.AttachedParentInfo.SocketName;
+				}
+
+				for (const auto& CachedAttachInfo : ActorTransactionAnnotation->RootComponentData.AttachedToInfo)
+				{
+					AActor* AttachedActor = CachedAttachInfo.Actor.Get();
+					if (AttachedActor)
+					{
+						FAttachedActorInfo Info;
+						Info.AttachedActor = AttachedActor;
+						Info.AttachedToSocket = CachedAttachInfo.SocketName;
+						Info.bSetRelativeTransform = true;
+						Info.RelativeTransform = CachedAttachInfo.RelativeTransform;
+						AttachedActorInfos.Add(Info);
+
+						AttachedActor->DetachRootComponentFromParent();
+					}
+				}
 			}
 		}
-
-		// Save off original pose of the actor
-		FTransform OldTransform = FTransform::Identity;
-		FName  SocketName;
-		AActor* Parent = NULL;
-		if (RootComponent != NULL)
+		else
 		{
-			// Do not need to detach if root component is not going away
-			if(RootComponent->AttachParent != NULL && RootComponent->bCreatedByConstructionScript)
+			InstanceDataCache = new FComponentInstanceDataCache(this);
+
+			// If there are attached objects detach them and store the socket names
+			TArray<AActor*> AttachedActors;
+			GetAttachedActors(AttachedActors);
+
+			for (AActor* AttachedActor : AttachedActors)
 			{
-				Parent = RootComponent->AttachParent->GetOwner();
-				// Root component should never be attached to another component in the same actor!
-				if(Parent == this)
+				USceneComponent* EachRoot = AttachedActor->GetRootComponent();
+				// If the component we are attached to is about to go away...
+				if (EachRoot && EachRoot->AttachParent && EachRoot->AttachParent->bCreatedByConstructionScript)
 				{
-					UE_LOG(LogActor, Warning, TEXT("RerunConstructionScripts: RootComponent (%s) attached to another component in this Actor (%s)."), *RootComponent->GetPathName(), *Parent->GetPathName());
-					Parent = NULL;
+					// Save info about actor to reattach
+					FAttachedActorInfo Info;
+					Info.AttachedActor = AttachedActor;
+					Info.AttachedToSocket = EachRoot->AttachSocketName;
+					Info.bSetRelativeTransform = false;
+					AttachedActorInfos.Add(Info);
+
+					// Now detach it
+					AttachedActor->Modify();
+					EachRoot->DetachFromParent(true);
 				}
+			}
 
-				SocketName = RootComponent->AttachSocketName;
-				//detach it to remove any scaling 
-				RootComponent->DetachFromParent(true);
+			if (RootComponent != NULL)
+			{
+				// Do not need to detach if root component is not going away
+				if (RootComponent->AttachParent != NULL && RootComponent->bCreatedByConstructionScript)
+				{
+					Parent = RootComponent->AttachParent->GetOwner();
+					// Root component should never be attached to another component in the same actor!
+					if (Parent == this)
+					{
+						UE_LOG(LogActor, Warning, TEXT("RerunConstructionScripts: RootComponent (%s) attached to another component in this Actor (%s)."), *RootComponent->GetPathName(), *Parent->GetPathName());
+						Parent = NULL;
+					}
+
+					SocketName = RootComponent->AttachSocketName;
+					//detach it to remove any scaling 
+					RootComponent->DetachFromParent(true);
+				}
+				OldTransform = RootComponent->ComponentToWorld;
 			}
-			OldTransform = RootComponent->ComponentToWorld;
 		}
 
 		// Destroy existing components
@@ -210,7 +251,7 @@ void AActor::RerunConstructionScripts()
 		}
 
 		// Run the construction scripts
-		ExecuteConstruction(OldTransform, &InstanceDataCache);
+		ExecuteConstruction(OldTransform, InstanceDataCache);
 
 		if(Parent)
 		{
@@ -232,6 +273,10 @@ void AActor::RerunConstructionScripts()
 				if (ChildRoot && ChildRoot->AttachParent != RootComponent)
 				{
 					ChildRoot->AttachTo(RootComponent, Info.AttachedToSocket, EAttachLocation::KeepWorldPosition);
+					if (Info.bSetRelativeTransform)
+					{
+						ChildRoot->SetRelativeTransform(Info.RelativeTransform);
+					}
 					ChildRoot->UpdateComponentToWorld();
 				}
 			}
@@ -239,6 +284,16 @@ void AActor::RerunConstructionScripts()
 
 		// Restore the undo buffer
 		GUndo = CurrentTransaction;
+
+		if (ActorTransactionAnnotation)
+		{
+			CurrentTransactionAnnotation = NULL;
+		}
+		else
+		{
+			delete InstanceDataCache;
+		}
+
 	}
 }
 
diff --git a/Engine/Source/Runtime/Engine/Private/ActorEditor.cpp b/Engine/Source/Runtime/Engine/Private/ActorEditor.cpp
--- a/Engine/Source/Runtime/Engine/Private/ActorEditor.cpp
+++ b/Engine/Source/Runtime/Engine/Private/ActorEditor.cpp
@@ -217,6 +217,73 @@ void AActor::DebugShowOneComponentHierarchy( USceneComponent* SceneComp, int32&
 	}
 }
 
+AActor::FActorTransactionAnnotation::FActorTransactionAnnotation(const AActor* Actor)
+	: ComponentInstanceData(Actor)
+{
+	USceneComponent* RootComponent = Actor->GetRootComponent();
+	if (RootComponent && RootComponent->bCreatedByConstructionScript)
+	{
+		bRootComponentDataCached = true;
+		RootComponentData.Transform = RootComponent->ComponentToWorld;
+
+		if (RootComponent->AttachParent)
+		{
+			RootComponentData.AttachedParentInfo.Actor = RootComponent->AttachParent->GetOwner();
+			RootComponentData.AttachedParentInfo.SocketName = RootComponent->AttachSocketName;
+			RootComponentData.AttachedParentInfo.RelativeTransform = RootComponent->GetRelativeTransform();
+		}
+
+		for (USceneComponent* AttachChild : RootComponent->AttachChildren)
+		{
+			if (AttachChild)
+			{
+				// Save info about actor to reattach
+				FActorRootComponentReconstructionData::FAttachedActorInfo Info;
+				Info.Actor = AttachChild->GetOwner();
+				Info.SocketName = AttachChild->AttachSocketName;
+				Info.RelativeTransform = AttachChild->GetRelativeTransform();
+				RootComponentData.AttachedToInfo.Add(Info);
+			}
+		}
+	}
+	else
+	{
+		bRootComponentDataCached = false;
+	}
+}
+
+bool AActor::FActorTransactionAnnotation::HasInstanceData() const
+{
+	return (bRootComponentDataCached || ComponentInstanceData.HasInstanceData());
+}
+
+TSharedPtr<ITransactionObjectAnnotation> AActor::GetTransactionAnnotation() const
+{
+	TSharedPtr<FActorTransactionAnnotation> TransactionAnnotation = MakeShareable(new FActorTransactionAnnotation(this));
+
+	if (!TransactionAnnotation->HasInstanceData())
+	{
+		// If there is nothing in the annotation don't bother storing it.
+		TransactionAnnotation = NULL;
+	}
+
+	return TransactionAnnotation;
+}
+
+void AActor::PreEditUndo()
+{
+	// Since child actor components will rebuild themselves get rid of the Actor before we make changes
+	TArray<UChildActorComponent*> ChildActorComponents;
+	GetComponents(ChildActorComponents);
+
+	for (UChildActorComponent* ChildActorComponent : ChildActorComponents)
+	{
+		ChildActorComponent->DestroyChildActor();
+	}
+
+	Super::PreEditUndo();
+}
+
 void AActor::PostEditUndo()
 {
 	// Notify LevelBounds actor that level bounding box might be changed
@@ -228,6 +295,19 @@ void AActor::PostEditUndo()
 	Super::PostEditUndo();
 }
 
+void AActor::PostEditUndo(TSharedPtr<ITransactionObjectAnnotation> TransactionAnnotation)
+{
+	CurrentTransactionAnnotation = StaticCastSharedPtr<FActorTransactionAnnotation>(TransactionAnnotation);
+
+	// Notify LevelBounds actor that level bounding box might be changed
+	if (!IsTemplate() && GetLevel()->LevelBoundsActor.IsValid())
+	{
+		GetLevel()->LevelBoundsActor.Get()->OnLevelBoundsDirtied();
+	}
+
+	Super::PostEditUndo(TransactionAnnotation);
+}
+
 // @todo: Remove this hack once we have decided on the scaling method to use.
 bool AActor::bUsePercentageBasedScaling = false;
 
diff --git a/Engine/Source/Runtime/Engine/Private/ChildActorComponent.cpp b/Engine/Source/Runtime/Engine/Private/ChildActorComponent.cpp
--- a/Engine/Source/Runtime/Engine/Private/ChildActorComponent.cpp
+++ b/Engine/Source/Runtime/Engine/Private/ChildActorComponent.cpp
@@ -97,6 +97,7 @@ void UChildActorComponent::CreateChildActor()
 			Params.bAllowDuringConstructionScript = true;
 			Params.OverrideLevel = GetOwner()->GetLevel();
 			Params.Name = ChildActorName;
+			Params.ObjectFlags &= ~RF_Transactional;
 
 			// Spawn actor of desired class
 			FVector Location = GetComponentLocation();
@@ -108,10 +109,8 @@ void UChildActorComponent::CreateChildActor()
 			{
 				ChildActorName = ChildActor->GetFName();
 
-#if WITH_EDITOR
 				// Remember which actor spawned it (for selection in editor etc)
 				ChildActor->ParentComponentActor = GetOwner();
-#endif
 
 				// Parts that we deferred from SpawnActor
 				ChildActor->ExecuteConstruction(ComponentToWorld, NULL);
diff --git a/Engine/Source/Runtime/Engine/Private/ComponentInstanceDataCache.cpp b/Engine/Source/Runtime/Engine/Private/ComponentInstanceDataCache.cpp
--- a/Engine/Source/Runtime/Engine/Private/ComponentInstanceDataCache.cpp
+++ b/Engine/Source/Runtime/Engine/Private/ComponentInstanceDataCache.cpp
@@ -67,7 +67,7 @@ bool FComponentInstanceDataBase::MatchesComponent(const UActorComponent* Compone
 	return bMatches;
 }
 
-FComponentInstanceDataCache::FComponentInstanceDataCache(AActor* Actor)
+FComponentInstanceDataCache::FComponentInstanceDataCache(const AActor* Actor)
 {
 	if(Actor != NULL)
 	{
diff --git a/Engine/Source/Runtime/Engine/Private/SCS_Node.cpp b/Engine/Source/Runtime/Engine/Private/SCS_Node.cpp
--- a/Engine/Source/Runtime/Engine/Private/SCS_Node.cpp
+++ b/Engine/Source/Runtime/Engine/Private/SCS_Node.cpp
@@ -55,7 +55,6 @@ void USCS_Node::ExecuteNodeOnActor(AActor* Actor, USceneComponent* ParentCompone
 			// If NULL is passed in, we are the root, so set transform and assign as RootComponent on Actor
 			if (ParentComponent == NULL || (ParentComponent && ParentComponent->IsPendingKill()))
 			{
-				NewSceneComp->SetFlags(RF_Transactional);
 				NewSceneComp->SetWorldTransform(*RootTransform);
 				Actor->SetRootComponent(NewSceneComp);
 			}
diff --git a/Engine/Source/Runtime/Engine/Public/ComponentInstanceDataCache.h b/Engine/Source/Runtime/Engine/Public/ComponentInstanceDataCache.h
--- a/Engine/Source/Runtime/Engine/Public/ComponentInstanceDataCache.h
+++ b/Engine/Source/Runtime/Engine/Public/ComponentInstanceDataCache.h
@@ -2,6 +2,9 @@
 
 #pragma once
 
+class UActorComponent;
+class AActor;
+
 /** Base class for component instance cached data of a particular type. */
 class FComponentInstanceDataBase
 {
@@ -36,7 +39,7 @@ public:
 	FComponentInstanceDataCache() {}
 
 	/** Constructor that also populatees cache from Actor */
-	FComponentInstanceDataCache(AActor* InActor);
+	FComponentInstanceDataCache(const AActor* InActor);
 
 	/** Add some new instance data to the cache */
 	void AddInstanceData(TSharedPtr<FComponentInstanceDataBase> NewData);
@@ -44,6 +47,8 @@ public:
 	/** Util to iterate over components and apply data to each */
 	void ApplyToActor(AActor* Actor) const;
 
+	bool HasInstanceData() const { return TypeToDataMap.Num() > 0; }
+
 private:
 	/** Map of data type name to data of that type */
 	TMultiMap< FName, TSharedPtr<FComponentInstanceDataBase> >	TypeToDataMap;
