Check for pending kill

diff --git a/Engine/Source/Runtime/Engine/Private/ActorReplication.cpp b/Engine/Source/Runtime/Engine/Private/ActorReplication.cpp
--- a/Engine/Source/Runtime/Engine/Private/ActorReplication.cpp
+++ b/Engine/Source/Runtime/Engine/Private/ActorReplication.cpp
@@ -262,7 +262,7 @@ void AActor::GetSubobjectsWithStableNamesForNetworking(TArray<UObject*> &ObjList
 	// For experimenting with replicating ALL stably named components initially
 	for (UActorComponent* Component : OwnedComponents)
 	{
-		if (Component && Component->IsNameStableForNetworking())
+		if (Component && !Component->IsPendingKill() && Component->IsNameStableForNetworking())
 		{
 			ObjList.Add(Component);
 			Component->GetSubobjectsWithStableNamesForNetworking(ObjList);
