Fix compile issue in non unity

diff --git a/Engine/Source/Runtime/Engine/Private/Collision/PhysicsFiltering.h b/Engine/Source/Runtime/Engine/Private/Collision/PhysicsFiltering.h
--- a/Engine/Source/Runtime/Engine/Private/Collision/PhysicsFiltering.h
+++ b/Engine/Source/Runtime/Engine/Private/Collision/PhysicsFiltering.h
@@ -2,6 +2,8 @@
 
 #pragma once
 
+#include "../PhysicsEngine/PhysXSupport.h"
+
 /** 
  * Set of flags stored in the PhysX FilterData
  *
