Distance field AO using a surface cache

diff --git a/Engine/Source/Runtime/Core/Public/Math/UnrealMathSSE.h b/Engine/Source/Runtime/Core/Public/Math/UnrealMathSSE.h
--- a/Engine/Source/Runtime/Core/Public/Math/UnrealMathSSE.h
+++ b/Engine/Source/Runtime/Core/Public/Math/UnrealMathSSE.h
@@ -103,6 +103,15 @@ static const VectorRegister SSE_INFINITY = MakeVectorRegister( (uint32)0x7F80000
  */
 #define VectorOne()						SSE_ONE
 
+/**
+ * Returns an component from a vector.
+ *
+ * @param Vec				Vector register
+ * @param ComponentIndex	Which component to get, X=0, Y=1, Z=2, W=3
+ * @return					The component as a float
+ */
+#define VectorGetComponent( Vec, ComponentIndex )	(((float*) &(Vec))[ComponentIndex])
+
 /**
  * Loads 4 FLOATs from unaligned memory.
  *
@@ -171,6 +180,15 @@ FORCEINLINE VectorRegister VectorLoadFloat3_W1( const void* Ptr )
  */
 #define VectorSetFloat3( X, Y, Z )		MakeVectorRegister( X, Y, Z, 0.0f )
 
+/**
+ * Propagates passed in float to all registers
+ *
+ * @param F		Float to set
+ * @return		VectorRegister(F,F,F,F)
+ */
+#define VectorSetFloat1( F )	_mm_set1_ps( F )
+
+
 /**
  * Creates a vector out of four FLOATs.
  *
@@ -780,6 +798,79 @@ FORCEINLINE VectorRegister VectorTransformVector(const VectorRegister&  VecP,  c
 #define VectorShuffle( Vec1, Vec2, X, Y, Z, W )	_mm_shuffle_ps( Vec1, Vec2, SHUFFLEMASK(X,Y,Z,W) )
 
 
+
+/**
+ * These functions return a vector mask to indicate which components pass the comparison.
+ * Each component is 0xffffffff if it passes, 0x00000000 if it fails.
+ *
+ * @param Vec1			1st source vector
+ * @param Vec2			2nd source vector
+ * @return				Vector with a mask for each component.
+ */
+#define VectorMask_LT( Vec1, Vec2 )			_mm_cmplt_ps(Vec1, Vec2)
+#define VectorMask_LE( Vec1, Vec2 )			_mm_cmple_ps(Vec1, Vec2)
+#define VectorMask_GT( Vec1, Vec2 )			_mm_cmpgt_ps(Vec1, Vec2)
+#define VectorMask_GE( Vec1, Vec2 )			_mm_cmpge_ps(Vec1, Vec2)
+#define VectorMask_EQ( Vec1, Vec2 )			_mm_cmpeq_ps(Vec1, Vec2)
+#define VectorMask_NE( Vec1, Vec2 )			_mm_cmpneq_ps(Vec1, Vec2)
+
+/**
+ * Returns an integer bit-mask (0x00 - 0x0f) based on the sign-bit for each component in a vector.
+ *
+ * @param VecMask		Vector
+ * @return				Bit 0 = sign(VecMask.x), Bit 1 = sign(VecMask.y), Bit 2 = sign(VecMask.z), Bit 3 = sign(VecMask.w)
+ */
+#define VectorMaskBits( VecMask )			_mm_movemask_ps( VecMask )
+
+/**
+ * Returns the bitwise AND.
+ *
+ * @param	Vec1	Vector to AND
+ * @param	Vec2	Vector to AND
+ * @return	bitwise per component AND operation.
+ */
+#define VectorBitwiseAND( Vec1, Vec2 )	_mm_and_ps( (Vec1), (Vec2) )
+
+/**
+ * Divides two vectors (component-wise) and returns the result.
+ *
+ * @param Vec1	1st vector
+ * @param Vec2	2nd vector
+ * @return		VectorRegister( Vec1.x/Vec2.x, Vec1.y/Vec2.y, Vec1.z/Vec2.z, Vec1.w/Vec2.w )
+ */
+#define VectorDivide( Vec1, Vec2 )		_mm_div_ps( Vec1, Vec2 )
+
+/**
+ * Counts the number of trailing zeros in the bit representation of the value,
+ * counting from least-significant bit to most.
+ *
+ * @param Value the value to determine the number of leading zeros for
+ * @return the number of zeros before the first "on" bit
+ */
+#if PLATFORM_WINDOWS
+#pragma intrinsic( _BitScanForward )
+FORCEINLINE uint32 appCountTrailingZeros(uint32 Value)
+{
+	if (Value == 0)
+	{
+		return 32;
+	}
+	uint32 BitIndex;	// 0-based, where the LSB is 0 and MSB is 31
+	_BitScanForward( (::DWORD *)&BitIndex, Value );	// Scans from LSB to MSB
+	return BitIndex;
+}
+#else // PLATFORM_WINDOWS
+FORCEINLINE uint32 appCountTrailingZeros(uint32 Value)
+{
+	if (Value == 0)
+	{
+		return 32;
+	}
+	return __builtin_ffs(Value) - 1;
+}
+#endif // PLATFORM_WINDOWS
+
+
 /**
  * Merges the XYZ components of one vector with the W component of another vector and returns the result.
  *
diff --git a/Engine/Source/Runtime/Core/Public/Math/UnrealMathUtility.h b/Engine/Source/Runtime/Core/Public/Math/UnrealMathUtility.h
--- a/Engine/Source/Runtime/Core/Public/Math/UnrealMathUtility.h
+++ b/Engine/Source/Runtime/Core/Public/Math/UnrealMathUtility.h
@@ -260,6 +260,12 @@ struct FMath : public FPlatformMath
 		return (Dividend + Divisor - 1) / Divisor;
 	}
 
+	template <class T>
+	static FORCEINLINE T DivideAndRoundDown(T Dividend,T Divisor)
+	{
+		return Dividend / Divisor;
+	}
+
 	/**
 	 * Computes the base 2 logarithm of the specified value
 	 *
diff --git a/Engine/Source/Runtime/CoreUObject/Classes/Object.h b/Engine/Source/Runtime/CoreUObject/Classes/Object.h
--- a/Engine/Source/Runtime/CoreUObject/Classes/Object.h
+++ b/Engine/Source/Runtime/CoreUObject/Classes/Object.h
@@ -102,6 +102,8 @@ enum EPixelFormat
 	PF_ETC1,
 	PF_ETC2_RGB,
 	PF_ETC2_RGBA,
+	PF_R32G32B32A32_UINT,
+	PF_R16G16_UINT,
 	PF_MAX,
 };
 
diff --git a/Engine/Source/Runtime/Engine/Classes/Engine/EngineTypes.h b/Engine/Source/Runtime/Engine/Classes/Engine/EngineTypes.h
--- a/Engine/Source/Runtime/Engine/Classes/Engine/EngineTypes.h
+++ b/Engine/Source/Runtime/Engine/Classes/Engine/EngineTypes.h
@@ -1672,6 +1672,8 @@ struct FMeshBuildSettings
 	/** The local scale applied when building the mesh */
 	UPROPERTY(EditAnywhere, Category=BuildSettings, meta=(DisplayName="Build Scale"))
 	FVector BuildScale3D;
+	UPROPERTY(EditAnywhere, Category=BuildSettings)
+	float DistanceFieldResolutionScale;
 
 	/** Default settings. */
 	FMeshBuildSettings()
@@ -1682,6 +1684,7 @@ struct FMeshBuildSettings
 		, BuildScale_DEPRECATED(1.0f)
 		, BuildScale3D(1.0f, 1.0f, 1.0f)
 	{
+		DistanceFieldResolutionScale = 1;
 	}
 
 	/** Equality operator. */
@@ -1691,7 +1694,8 @@ struct FMeshBuildSettings
 			&& bRecomputeTangents == Other.bRecomputeTangents
 			&& bRemoveDegenerates == Other.bRemoveDegenerates
 			&& bUseFullPrecisionUVs == Other.bUseFullPrecisionUVs
-			&& BuildScale3D == Other.BuildScale3D;
+			&& BuildScale3D == Other.BuildScale3D
+			&& DistanceFieldResolutionScale == Other.DistanceFieldResolutionScale;
 	}
 
 	/** Inequality. */
diff --git a/Engine/Source/Runtime/Engine/Classes/Engine/RendererSettings.h b/Engine/Source/Runtime/Engine/Classes/Engine/RendererSettings.h
--- a/Engine/Source/Runtime/Engine/Classes/Engine/RendererSettings.h
+++ b/Engine/Source/Runtime/Engine/Classes/Engine/RendererSettings.h
@@ -121,6 +121,11 @@ class ENGINE_API URendererSettings
 		ToolTip="Whether to allow any static lighting to be generated and used, like lightmaps and shadowmaps. Games that only use dynamic lighting should set this to 0 to save some static lighting overhead. Changing this setting requires restarting the editor."))
 	uint32 bAllowStaticLighting:1;
 
+	UPROPERTY(config, EditAnywhere, Category=Lighting, meta=(
+		ConsoleVariable="r.AllowMeshDistanceFieldRepresentations",
+		ToolTip="Whether to build distance fields of static meshes, needed for distance field AO, which is used to implement Movable SkyLight shadows.  Enabling will increase mesh build times and memory usage.  Changing this setting requires restarting the editor."))
+	uint32 bAllowMeshDistanceFieldRepresentations:1;
+
 	UPROPERTY(config, EditAnywhere, Category=Lighting, meta=(
 		ConsoleVariable="r.Shadow.DistanceFieldPenumbraSize",
 		ToolTip="Controls the size of the uniform penumbra produced by static shadowing."))
diff --git a/Engine/Source/Runtime/Engine/Private/DynamicMeshBuilder.cpp b/Engine/Source/Runtime/Engine/Private/DynamicMeshBuilder.cpp
--- a/Engine/Source/Runtime/Engine/Private/DynamicMeshBuilder.cpp
+++ b/Engine/Source/Runtime/Engine/Private/DynamicMeshBuilder.cpp
@@ -432,6 +432,7 @@ void FDynamicMeshBuilder::Draw(FPrimitiveDrawInterface* PDI,const FMatrix& Local
 			FBoxSphereBounds(EForceInit::ForceInit),
 			FBoxSphereBounds(EForceInit::ForceInit),
 			bReceivesDecals,
+			false,
 			1.0f		// LPV bias
 			);
 
diff --git a/Engine/Source/Runtime/Engine/Private/NiagaraComponent.cpp b/Engine/Source/Runtime/Engine/Private/NiagaraComponent.cpp
--- a/Engine/Source/Runtime/Engine/Private/NiagaraComponent.cpp
+++ b/Engine/Source/Runtime/Engine/Private/NiagaraComponent.cpp
@@ -138,7 +138,8 @@ private:
 					GetActorPosition(),
 					GetBounds(),
 					GetLocalBounds(),
-					ReceivesDecals()
+					ReceivesDecals(),
+					false
 					);
 				WorldSpacePrimitiveUniformBuffer.SetContents(PrimitiveUniformShaderParameters);
 				WorldSpacePrimitiveUniformBuffer.InitResource();
diff --git a/Engine/Source/Runtime/Engine/Private/Particles/ParticleSystemRender.cpp b/Engine/Source/Runtime/Engine/Private/Particles/ParticleSystemRender.cpp
--- a/Engine/Source/Runtime/Engine/Private/Particles/ParticleSystemRender.cpp
+++ b/Engine/Source/Runtime/Engine/Private/Particles/ParticleSystemRender.cpp
@@ -7219,6 +7219,7 @@ void FParticleSystemSceneProxy::UpdateWorldSpacePrimitiveUniformBuffer()
 			GetBounds(),
 			GetLocalBounds(),
 			ReceivesDecals(),
+			false,
 			1.0f			// LPV bias
 			);
 		WorldSpacePrimitiveUniformBuffer.SetContents(PrimitiveUniformShaderParameters);
diff --git a/Engine/Source/Runtime/Engine/Private/PrimitiveSceneProxy.cpp b/Engine/Source/Runtime/Engine/Private/PrimitiveSceneProxy.cpp
--- a/Engine/Source/Runtime/Engine/Private/PrimitiveSceneProxy.cpp
+++ b/Engine/Source/Runtime/Engine/Private/PrimitiveSceneProxy.cpp
@@ -185,7 +185,7 @@ void FPrimitiveSceneProxy::UpdateActorPosition(FVector InActorPosition)
 			PrimitiveSceneProxy->ActorPosition = InActorPosition;
 			// Update the uniform shader parameters.
 			const FPrimitiveUniformShaderParameters PrimitiveUniformShaderParameters =
-			GetPrimitiveUniformShaderParameters(PrimitiveSceneProxy->LocalToWorld, InActorPosition, PrimitiveSceneProxy->Bounds,PrimitiveSceneProxy->LocalBounds, PrimitiveSceneProxy->bReceivesDecals, PrimitiveSceneProxy->GetLpvBiasMultiplier() );
+			GetPrimitiveUniformShaderParameters(PrimitiveSceneProxy->LocalToWorld, InActorPosition, PrimitiveSceneProxy->Bounds,PrimitiveSceneProxy->LocalBounds, PrimitiveSceneProxy->bReceivesDecals, PrimitiveSceneProxy->HasDistanceFieldRepresentation(), PrimitiveSceneProxy->GetLpvBiasMultiplier() );
 
 			PrimitiveSceneProxy->UniformBuffer.SetContents(PrimitiveUniformShaderParameters);
 			PrimitiveSceneProxy->OnActorPositionChanged();
@@ -207,7 +207,7 @@ void FPrimitiveSceneProxy::SetTransform(const FMatrix& InLocalToWorld, const FBo
 	ActorPosition = InActorPosition;
 	
 	// Update the uniform shader parameters.
-	const FPrimitiveUniformShaderParameters PrimitiveUniformShaderParameters = GetPrimitiveUniformShaderParameters(LocalToWorld, ActorPosition, Bounds, LocalBounds, bReceivesDecals, LpvBiasMultiplier );
+	const FPrimitiveUniformShaderParameters PrimitiveUniformShaderParameters = GetPrimitiveUniformShaderParameters(LocalToWorld, ActorPosition, Bounds, LocalBounds, bReceivesDecals, HasDistanceFieldRepresentation(), LpvBiasMultiplier );
 	UniformBuffer.SetContents(PrimitiveUniformShaderParameters);
 	
 	// Notify the proxy's implementation of the change.
diff --git a/Engine/Source/Runtime/Engine/Private/SceneView.cpp b/Engine/Source/Runtime/Engine/Private/SceneView.cpp
--- a/Engine/Source/Runtime/Engine/Private/SceneView.cpp
+++ b/Engine/Source/Runtime/Engine/Private/SceneView.cpp
@@ -170,6 +170,7 @@ FSceneView::FSceneView(const FSceneViewInitOptions& InitOptions)
 	, bIsViewInfo(false)
 	, bIsSceneCapture(false)
 	, bIsReflectionCapture(false)
+	, bStaticSceneOnly(false)
 	, RenderingCompositePassContext(0)
 #if WITH_EDITOR
 	, OverrideLODViewOrigin(InitOptions.OverrideLODViewOrigin)
diff --git a/Engine/Source/Runtime/Engine/Private/SkyLightComponent.cpp b/Engine/Source/Runtime/Engine/Private/SkyLightComponent.cpp
--- a/Engine/Source/Runtime/Engine/Private/SkyLightComponent.cpp
+++ b/Engine/Source/Runtime/Engine/Private/SkyLightComponent.cpp
@@ -72,6 +72,7 @@ FSkyLightSceneProxy::FSkyLightSceneProxy(const USkyLightComponent* InLightCompon
 	, ProcessedTexture(InLightComponent->ProcessedSkyTexture)
 	, SkyDistanceThreshold(InLightComponent->SkyDistanceThreshold)
 	, bCastShadows(InLightComponent->CastShadows)
+	, bWantsStaticShadowing(InLightComponent->Mobility == EComponentMobility::Stationary)
 	, bPrecomputedLightingIsValid(InLightComponent->bPrecomputedLightingIsValid)
 	, bHasStaticLighting(InLightComponent->HasStaticLighting())
 	, LightColor(FLinearColor(InLightComponent->LightColor) * InLightComponent->Intensity)
@@ -199,12 +200,6 @@ void USkyLightComponent::PostEditChangeProperty(FPropertyChangedEvent& PropertyC
 {
 	const FName CategoryName = FObjectEditorUtils::GetCategoryFName(PropertyChangedEvent.Property);
 
-	// Movable not supported yet
-	if (Mobility == EComponentMobility::Movable)
-	{
-		Mobility = EComponentMobility::Stationary;
-	}
-
 	Super::PostEditChangeProperty(PropertyChangedEvent);
 
 	SetCaptureIsDirty();
@@ -365,6 +360,7 @@ void USkyLightComponent::UpdateSkyCaptureContents(UWorld* WorldToUpdate)
 						CaptureComponent->MarkRenderStateDirty();
 					}
 
+					//@todo - defer this until shader compilation is finished, just like reflection captures, otherwise we will capture an incomplete scene.
 					WorldToUpdate->Scene->UpdateSkyCaptureContents(CaptureComponent, false, CaptureComponent->ProcessedSkyTexture, CaptureComponent->IrradianceEnvironmentMap);
 
 					CaptureComponent->MarkRenderStateDirty();
@@ -468,4 +464,4 @@ void ASkyLight::OnRep_bEnabled()
 	LightComponent->SetVisibility(bEnabled);
 }
 
-#undef LOCTEXT_NAMESPACE
\ No newline at end of file
+#undef LOCTEXT_NAMESPACE
diff --git a/Engine/Source/Runtime/Engine/Private/StaticMesh.cpp b/Engine/Source/Runtime/Engine/Private/StaticMesh.cpp
--- a/Engine/Source/Runtime/Engine/Private/StaticMesh.cpp
+++ b/Engine/Source/Runtime/Engine/Private/StaticMesh.cpp
@@ -250,6 +250,146 @@ void FStaticMeshVertexBuffer::AllocateData( bool bNeedsCPUAccess /*= true*/ )
 	Stride = VertexData->GetStride();
 }
 
+TGlobalResource<FDistanceFieldVolumeTextureAtlas> GDistanceFieldVolumeTextureAtlas = TGlobalResource<FDistanceFieldVolumeTextureAtlas>(PF_R16F);
+
+FDistanceFieldVolumeTextureAtlas::FDistanceFieldVolumeTextureAtlas(EPixelFormat InFormat) :
+	BlockAllocator(0, 0, 0, 512, 512, 512, false, false)
+{
+	Format = InFormat;
+}
+
+void FDistanceFieldVolumeTextureAtlas::AddAllocation(FDistanceFieldVolumeTexture* Texture)
+{
+	PendingAllocations.AddUnique(Texture);
+}
+
+void FDistanceFieldVolumeTextureAtlas::RemoveAllocation(FDistanceFieldVolumeTexture* Texture)
+{
+	PendingAllocations.Remove(Texture);
+
+	if (CurrentAllocations.Contains(Texture))
+	{
+		const FIntVector Min = Texture->GetAllocationMin();
+		const FIntVector Size = Texture->VolumeData.Size;
+		verify(BlockAllocator.RemoveElement(Min.X, Min.Y, Min.Z, Size.X, Size.Y, Size.Z));
+		CurrentAllocations.Remove(Texture);
+	}
+}
+
+void FDistanceFieldVolumeTextureAtlas::UpdateAllocations()
+{
+	if (PendingAllocations.Num() > 0)
+	{
+		//@todo - sort largest to smallest for best packing
+		for (int32 AllocationIndex = 0; AllocationIndex < PendingAllocations.Num(); AllocationIndex++)
+		{
+			FDistanceFieldVolumeTexture* Texture = PendingAllocations[AllocationIndex];
+			const FIntVector Size = Texture->VolumeData.Size;
+
+			if (!BlockAllocator.AddElement((uint32&)Texture->AtlasAllocationMin.X, (uint32&)Texture->AtlasAllocationMin.Y, (uint32&)Texture->AtlasAllocationMin.Z, Size.X, Size.Y, Size.Z))
+			{
+				UE_LOG(LogStaticMesh,Error,TEXT("Failed to allocate %ux%ux%u in distance field atlas"), Size.X, Size.Y, Size.Z);
+				PendingAllocations.RemoveAt(AllocationIndex);
+				AllocationIndex--;
+			}
+		}
+
+		if (!VolumeTextureRHI
+			|| BlockAllocator.GetSizeX() > VolumeTextureRHI->GetSizeX()
+			|| BlockAllocator.GetSizeY() > VolumeTextureRHI->GetSizeY()
+			|| BlockAllocator.GetSizeZ() > VolumeTextureRHI->GetSizeZ())
+		{
+			FRHIResourceCreateInfo CreateInfo;
+
+			VolumeTextureRHI = RHICreateTexture3D(
+				BlockAllocator.GetSizeX(), 
+				BlockAllocator.GetSizeY(), 
+				BlockAllocator.GetSizeZ(), 
+				Format,
+				1,
+				TexCreate_ShaderResource,
+				CreateInfo);
+
+			const int32 FormatSize = GPixelFormats[Format].BlockBytes;
+			float MemorySize = VolumeTextureRHI->GetSizeX() * VolumeTextureRHI->GetSizeY() * VolumeTextureRHI->GetSizeZ() * FormatSize / 1024.0f / 1024.0f;
+			UE_LOG(LogStaticMesh,Log,TEXT("Allocated %ux%ux%u distance field atlas = %.1fMb"), VolumeTextureRHI->GetSizeX(), VolumeTextureRHI->GetSizeY(), VolumeTextureRHI->GetSizeZ(), MemorySize);
+
+			// Re-upload all textures since we had to reallocate
+			PendingAllocations.Append(CurrentAllocations);
+			CurrentAllocations.Empty();
+		}
+
+		for (int32 AllocationIndex = 0; AllocationIndex < PendingAllocations.Num(); AllocationIndex++)
+		{
+			FDistanceFieldVolumeTexture* Texture = PendingAllocations[AllocationIndex];
+			const FIntVector Size = Texture->VolumeData.Size;
+
+			const FUpdateTextureRegion3D UpdateRegion(
+				Texture->AtlasAllocationMin.X,
+				Texture->AtlasAllocationMin.Y,
+				Texture->AtlasAllocationMin.Z,
+				0,
+				0,
+				0,
+				Size.X,
+				Size.Y,
+				Size.Z);
+
+			const int32 FormatSize = GPixelFormats[Format].BlockBytes;
+
+			// Update the volume texture atlas
+			RHIUpdateTexture3D(VolumeTextureRHI, 0, UpdateRegion, Size.X * FormatSize, Size.X * Size.Y * FormatSize, (const uint8*)Texture->VolumeData.DistanceFieldVolume.GetData());
+		}
+
+		CurrentAllocations.Append(PendingAllocations);
+		PendingAllocations.Empty();
+	}
+}
+
+void FDistanceFieldVolumeTexture::Initialize()
+{
+	if (IsValidDistanceFieldVolume())
+	{
+		bReferencedByAtlas = true;
+
+		ENQUEUE_UNIQUE_RENDER_COMMAND_TWOPARAMETER(
+			AddAllocation,
+			FDistanceFieldVolumeTextureAtlas*, Atlas, Atlas,
+			FDistanceFieldVolumeTexture*, DistanceFieldVolumeTexture, this,
+			{
+				Atlas->AddAllocation(DistanceFieldVolumeTexture);
+			}
+		);
+	}
+}
+
+void FDistanceFieldVolumeTexture::Release()
+{
+	if (bReferencedByAtlas)
+	{
+		bReferencedByAtlas = false;
+
+		ENQUEUE_UNIQUE_RENDER_COMMAND_TWOPARAMETER(
+			ReleaseAllocation,
+			FDistanceFieldVolumeTextureAtlas*, Atlas, Atlas,
+			FDistanceFieldVolumeTexture*, DistanceFieldVolumeTexture, this,
+			{
+				Atlas->RemoveAllocation(DistanceFieldVolumeTexture);
+			}
+		);
+	}
+}
+
+FIntVector FDistanceFieldVolumeTexture::GetAllocationSize() const
+{
+	return VolumeData.Size;
+}
+
+bool FDistanceFieldVolumeTexture::IsValidDistanceFieldVolume() const
+{
+	return VolumeData.Size.GetMax() > 0;
+}
+
 /*-----------------------------------------------------------------------------
 	FStaticMeshLODResources
 -----------------------------------------------------------------------------*/
@@ -289,14 +429,18 @@ void FStaticMeshLODResources::Serialize(FArchive& Ar, UObject* Owner, int32 Inde
 	if( !StripFlags.IsDataStrippedForServer() )
 	{
 		PositionVertexBuffer.Serialize( Ar, bNeedsCPUAccess );
-			VertexBuffer.Serialize( Ar, bNeedsCPUAccess );
-				ColorVertexBuffer.Serialize( Ar, bNeedsCPUAccess );
+		VertexBuffer.Serialize( Ar, bNeedsCPUAccess );
+		ColorVertexBuffer.Serialize( Ar, bNeedsCPUAccess );
 		IndexBuffer.Serialize( Ar, bNeedsCPUAccess );
 		DepthOnlyIndexBuffer.Serialize(Ar, bNeedsCPUAccess);
+		
+		Ar << DistanceFieldData;
+
 		if( !StripFlags.IsEditorDataStripped() )
 		{
 			WireframeIndexBuffer.Serialize(Ar, bNeedsCPUAccess);
 		}
+
 		if ( !StripFlags.IsClassDataStripped( AdjacencyDataStripFlag ) )
 		{
 			AdjacencyIndexBuffer.Serialize( Ar, bNeedsCPUAccess );
@@ -499,6 +643,8 @@ void FStaticMeshLODResources::InitResources(UStaticMesh* Parent)
 
 	InitVertexFactory(VertexFactory, Parent, NULL);
 	BeginInitResource(&VertexFactory);
+	
+	DistanceFieldData.VolumeTexture.Initialize();
 
 	const uint32 StaticMeshVertexMemory = 
 		VertexBuffer.GetStride() * VertexBuffer.GetNumVertices() + 
@@ -541,6 +687,8 @@ void FStaticMeshLODResources::ReleaseResources()
 	BeginReleaseResource(&ColorVertexBuffer);
 	BeginReleaseResource(&DepthOnlyIndexBuffer);
 
+	DistanceFieldData.VolumeTexture.Release();
+
 	// Release the vertex factories.
 	BeginReleaseResource(&VertexFactory);
 }
@@ -944,6 +1092,7 @@ FArchive& operator<<(FArchive& Ar, FMeshBuildSettings& BuildSettings)
 		Ar << BuildSettings.BuildScale3D;
 	}
 	
+	Ar << BuildSettings.DistanceFieldResolutionScale;
 	return Ar;
 }
 
@@ -951,18 +1100,22 @@ FArchive& operator<<(FArchive& Ar, FMeshBuildSettings& BuildSettings)
 // differences, etc.) replace the version GUID below with a new one.
 // In case of merge conflicts with DDC versions, you MUST generate a new GUID
 // and set this new GUID as the version.
-#define STATICMESH_DERIVEDDATA_VER TEXT("9E3F518AAD424921BA8F9A1C5966F0B9")
+#define STATICMESH_DERIVEDDATA_VER TEXT("4668778561B445A9523C94440EA899D")
 
 static const FString& GetStaticMeshDerivedDataVersion()
 {
 	static FString CachedVersionString;
 	if (CachedVersionString.IsEmpty())
 	{
+		static const auto CVar = IConsoleManager::Get().FindTConsoleVariableDataInt(TEXT("r.AllowMeshDistanceFieldRepresentations"));
+		FString DistanceFieldAllowed = (CVar && CVar->GetValueOnGameThread() != 0) ? TEXT("Dist") : TEXT("");
+
 		// Static mesh versioning is controlled by the version reported by the mesh utilities module.
 		IMeshUtilities& MeshUtilities = FModuleManager::Get().LoadModuleChecked<IMeshUtilities>(TEXT("MeshUtilities"));
-		CachedVersionString = FString::Printf(TEXT("%s_%s"),
+		CachedVersionString = FString::Printf(TEXT("%s_%s_%s"),
 			STATICMESH_DERIVEDDATA_VER,
-			*MeshUtilities.GetVersionString()
+			*MeshUtilities.GetVersionString(),
+			*DistanceFieldAllowed
 			);
 	}
 	return CachedVersionString;
@@ -1065,7 +1218,7 @@ void FStaticMeshRenderData::Cache(UStaticMesh* Owner, const FStaticMeshLODSettin
 		FStaticMeshStatusMessageContext StatusContext( FText::Format( NSLOCTEXT("Engine", "BuildingStaticMeshStatus", "Building static mesh {StaticMeshName}..."), Args ) );
 
 		IMeshUtilities& MeshUtilities = FModuleManager::Get().LoadModuleChecked<IMeshUtilities>(TEXT("MeshUtilities"));
-		MeshUtilities.BuildStaticMesh(*this, Owner->SourceModels, LODGroup);
+		MeshUtilities.BuildStaticMesh(*this, Owner->SourceModels, Owner->Materials, LODGroup);
 		bLODsShareStaticLighting = Owner->CanLODsShareStaticLighting();
 		FMemoryWriter Ar(DerivedData, /*bIsPersistent=*/ true);
 		Serialize(Ar, Owner, /*bCooked=*/ false);
diff --git a/Engine/Source/Runtime/Engine/Private/StaticMeshRender.cpp b/Engine/Source/Runtime/Engine/Private/StaticMeshRender.cpp
--- a/Engine/Source/Runtime/Engine/Private/StaticMeshRender.cpp
+++ b/Engine/Source/Runtime/Engine/Private/StaticMeshRender.cpp
@@ -880,6 +880,20 @@ void FStaticMeshSceneProxy::GetLightRelevance(const FLightSceneProxy* LightScene
 	}
 }
 
+void FStaticMeshSceneProxy::GetDistancefieldAtlasData(FBox& LocalVolumeBounds, FIntVector& OutBlockMin, FIntVector& OutBlockSize) const
+{
+	const FDistanceFieldVolumeData& DistanceField = RenderData->LODResources[0].DistanceFieldData;
+	LocalVolumeBounds = DistanceField.LocalBoundingBox;
+	OutBlockMin = DistanceField.VolumeTexture.GetAllocationMin();
+	OutBlockSize = DistanceField.VolumeTexture.GetAllocationSize();
+}
+
+bool FStaticMeshSceneProxy::HasDistanceFieldRepresentation() const
+{
+	const FDistanceFieldVolumeData& DistanceField = RenderData->LODResources[0].DistanceFieldData;
+	return (CastsDynamicShadow() || CastsStaticShadow()) && DistanceField.VolumeTexture.IsValidDistanceFieldVolume();
+}
+
 /** Initialization constructor. */
 FStaticMeshSceneProxy::FLODInfo::FLODInfo(const UStaticMeshComponent* InComponent,int32 LODIndex):
 	OverrideColorVertexBuffer(0),
diff --git a/Engine/Source/Runtime/Engine/Public/PixelFormat.h b/Engine/Source/Runtime/Engine/Public/PixelFormat.h
--- a/Engine/Source/Runtime/Engine/Public/PixelFormat.h
+++ b/Engine/Source/Runtime/Engine/Public/PixelFormat.h
@@ -54,7 +54,9 @@ enum EPixelFormat
 	PF_ETC1					=45,
 	PF_ETC2_RGB				=46,
 	PF_ETC2_RGBA			=47,
-	PF_MAX                  =48,
+	PF_R32G32B32A32_UINT	=48,
+	PF_R16G16_UINT			=49,
+	PF_MAX                  =50,
 };
 #define FOREACH_ENUM_EPIXELFORMAT(op) \
 	op(PF_Unknown) \
@@ -100,11 +102,11 @@ enum EPixelFormat
 	op(PF_R8G8) \
 	op(PF_ATC_RGB) \
 	op(PF_ATC_RGBA_E) \
-	op(PF_ATC_RGBA_I) \
 	op(PF_X24_G8) \
 	op(PF_ETC1) \
 	op(PF_ETC2_RGB) \
-	op(PF_ETC2_RGBA)
+	op(PF_ETC2_RGBA) \
+	op(PF_R32G32B32A32_UINT) \
+	op(PF_R16G16_UINT)
 
-
-#endif		// __PixelFormat_h__
\ No newline at end of file
+#endif		// __PixelFormat_h__
diff --git a/Engine/Source/Runtime/Engine/Public/PrimitiveSceneProxy.h b/Engine/Source/Runtime/Engine/Public/PrimitiveSceneProxy.h
--- a/Engine/Source/Runtime/Engine/Public/PrimitiveSceneProxy.h
+++ b/Engine/Source/Runtime/Engine/Public/PrimitiveSceneProxy.h
@@ -183,6 +183,13 @@ public:
 		bShadowMapped = false;
 	}
 
+	virtual void GetDistancefieldAtlasData(FBox& LocalVolumeBounds, FIntVector& OutBlockMin, FIntVector& OutBlockSize) const 
+	{
+		LocalVolumeBounds = FBox(0);
+		OutBlockMin = FIntVector(-1, -1, -1);
+		OutBlockSize = FIntVector(0, 0, 0);
+	}
+
 	/**
 	 *	Called when the rendering thread adds the proxy to the scene.
 	 *	This function allows for generating renderer-side resources.
@@ -361,6 +368,11 @@ public:
 		return GetBounds();
 	}
 
+	virtual bool HasDistanceFieldRepresentation() const
+	{
+		return false;
+	}
+
 	/** 
 	 * Drawing helper. Draws nice bouncy line.
 	 */
diff --git a/Engine/Source/Runtime/Engine/Public/SceneManagement.h b/Engine/Source/Runtime/Engine/Public/SceneManagement.h
--- a/Engine/Source/Runtime/Engine/Public/SceneManagement.h
+++ b/Engine/Source/Runtime/Engine/Public/SceneManagement.h
@@ -705,6 +705,7 @@ public:
 	FTexture* ProcessedTexture;
 	float SkyDistanceThreshold;
 	bool bCastShadows;
+	bool bWantsStaticShadowing;
 	bool bPrecomputedLightingIsValid;
 	bool bHasStaticLighting;
 	FLinearColor LightColor;
@@ -1134,6 +1135,7 @@ BEGIN_UNIFORM_BUFFER_STRUCT(FPrimitiveUniformShaderParameters,ENGINE_API)
 	DECLARE_UNIFORM_BUFFER_STRUCT_MEMBER_EX(float,LocalToWorldDeterminantSign,EShaderPrecisionModifier::Half)
 	DECLARE_UNIFORM_BUFFER_STRUCT_MEMBER(FVector,ActorWorldPosition)
 	DECLARE_UNIFORM_BUFFER_STRUCT_MEMBER_EX(float,DecalReceiverMask,EShaderPrecisionModifier::Half)
+	DECLARE_UNIFORM_BUFFER_STRUCT_MEMBER_EX(float,HasDistanceFieldRepresentation,EShaderPrecisionModifier::Half)
 	DECLARE_UNIFORM_BUFFER_STRUCT_MEMBER_EX(FVector4,ObjectOrientation,EShaderPrecisionModifier::Half)
 	DECLARE_UNIFORM_BUFFER_STRUCT_MEMBER_EX(FVector4,NonUniformScale,EShaderPrecisionModifier::Half)
 	DECLARE_UNIFORM_BUFFER_STRUCT_MEMBER_EX(FVector4,InvNonUniformScale,EShaderPrecisionModifier::Half)
@@ -1149,6 +1151,7 @@ inline FPrimitiveUniformShaderParameters GetPrimitiveUniformShaderParameters(
 	const FBoxSphereBounds& WorldBounds,
 	const FBoxSphereBounds& LocalBounds,
 	bool bReceivesDecals,
+	bool bHasDistanceFieldRepresentation,
 	float LpvBiasMultiplier = 1.0f
 )
 {
@@ -1182,6 +1185,7 @@ inline FPrimitiveUniformShaderParameters GetPrimitiveUniformShaderParameters(
 
 	Result.LocalToWorldDeterminantSign = FMath::FloatSelect(LocalToWorld.RotDeterminant(),1,-1);
 	Result.DecalReceiverMask = bReceivesDecals ? 1 : 0;
+	Result.HasDistanceFieldRepresentation = bHasDistanceFieldRepresentation ? 1 : 0;
 	return Result;
 }
 
@@ -1195,7 +1199,7 @@ inline TUniformBufferRef<FPrimitiveUniformShaderParameters> CreatePrimitiveUnifo
 {
 	check(IsInRenderingThread());
 	return TUniformBufferRef<FPrimitiveUniformShaderParameters>::CreateUniformBufferImmediate(
-		GetPrimitiveUniformShaderParameters(LocalToWorld, WorldBounds.Origin, WorldBounds, LocalBounds, bReceivesDecals, LpvBiasMultiplier ),
+		GetPrimitiveUniformShaderParameters(LocalToWorld, WorldBounds.Origin, WorldBounds, LocalBounds, bReceivesDecals, false, LpvBiasMultiplier ),
 		UniformBuffer_MultiUse
 		);
 }
@@ -1221,6 +1225,7 @@ public:
 			FBoxSphereBounds(EForceInit::ForceInit),
 			FBoxSphereBounds(EForceInit::ForceInit),
 			true,
+			false,
 			1.0f		// LPV bias
 			));
 	}
@@ -2635,6 +2640,12 @@ public:
 	/** Whether this view is being used to render a reflection capture. */
 	bool bIsReflectionCapture;
 
+	/** 
+	 * Whether to only render static lights and objects.  
+	 * This is used when capturing the scene for reflection captures, which aren't updated at runtime. 
+	 */
+	bool bStaticSceneOnly;
+
 	/** 0 if valid (we are rendering a screen postprocess pass )*/
 	struct FRenderingCompositePassContext* RenderingCompositePassContext; 
 
diff --git a/Engine/Source/Runtime/Engine/Public/ShowFlags.h b/Engine/Source/Runtime/Engine/Public/ShowFlags.h
--- a/Engine/Source/Runtime/Engine/Public/ShowFlags.h
+++ b/Engine/Source/Runtime/Engine/Public/ShowFlags.h
@@ -208,6 +208,7 @@ struct FEngineShowFlags
 		SetVolumeLightingSamples(false);
 		// we enable it manually on the editor view ports
 		SetSnap(false);
+		SetVisualizeDistanceFieldAO(false);
 	}
 
 	static bool CanBeToggledInEditor(const TCHAR* Name)
@@ -247,6 +248,7 @@ struct FEngineShowFlags
 		SetHighResScreenshotMask(false);
 		SetHMDDistortion(false);
 		SetStereoRendering(false);
+		SetDistanceFieldAO(false);
 	}
 
 	// ---------------------------------------------------------
diff --git a/Engine/Source/Runtime/Engine/Public/ShowFlagsValues.inl b/Engine/Source/Runtime/Engine/Public/ShowFlagsValues.inl
--- a/Engine/Source/Runtime/Engine/Public/ShowFlagsValues.inl
+++ b/Engine/Source/Runtime/Engine/Public/ShowFlagsValues.inl
@@ -262,6 +262,10 @@ SHOWFLAG_ALWAYS_ACCESSIBLE(VolumeLightingSamples, SFG_Visualize, LOCTEXT("Volume
 SHOWFLAG_ALWAYS_ACCESSIBLE(LpvLightingOnly, SFG_Hidden, LOCTEXT("VisualizeLPVSF_ViewMode", "Visualize LPV"))
 /** Render Paper2D sprites */
 SHOWFLAG_ALWAYS_ACCESSIBLE(Paper2DSprites, SFG_Advanced, LOCTEXT("Paper2DSpritesSF", "Paper 2D Sprites"))
+/** Visualization of distance field AO */
+SHOWFLAG_ALWAYS_ACCESSIBLE(VisualizeDistanceFieldAO, SFG_Visualize, LOCTEXT("VisualizeDistanceFieldAOSF", "Visualize Distance Field AO"))
+/** Distance field AO */
+SHOWFLAG_ALWAYS_ACCESSIBLE(DistanceFieldAO, SFG_LightingFeatures, LOCTEXT("DistanceFieldAOSF", "Distance Field AO"))
 
 #undef SHOWFLAG_ALWAYS_ACCESSIBLE
 #undef SHOWFLAG_FIXED_IN_SHIPPING
diff --git a/Engine/Source/Runtime/Engine/Public/StaticMeshResources.h b/Engine/Source/Runtime/Engine/Public/StaticMeshResources.h
--- a/Engine/Source/Runtime/Engine/Public/StaticMeshResources.h
+++ b/Engine/Source/Runtime/Engine/Public/StaticMeshResources.h
@@ -7,6 +7,7 @@
 #pragma once
 
 #include "RawIndexBuffer.h"
+#include "TextureLayout3d.h"
 
 /**
  * The LOD settings to use for a group of static meshes.
@@ -521,6 +522,125 @@ private:
 	void AllocateData( bool bNeedsCPUAccess = true );
 };
 
+class FDistanceFieldVolumeTexture;
+
+/** Global volume texture atlas that collects all static mesh resource distance fields. */
+class FDistanceFieldVolumeTextureAtlas : public FRenderResource
+{
+public:
+	FDistanceFieldVolumeTextureAtlas(EPixelFormat InFormat);
+
+	virtual void ReleaseRHI() OVERRIDE
+	{
+		VolumeTextureRHI.SafeRelease();
+	}
+
+	int32 GetSizeX() const { return VolumeTextureRHI->GetSizeX(); }
+	int32 GetSizeY() const { return VolumeTextureRHI->GetSizeY(); }
+	int32 GetSizeZ() const { return VolumeTextureRHI->GetSizeZ(); }
+
+	/** Add an allocation to the atlas. */
+	void AddAllocation(FDistanceFieldVolumeTexture* Texture);
+
+	/** Remove an allocation from the atlas. This must be done prior to deleting the FDistanceFieldVolumeTexture object. */
+	void RemoveAllocation(FDistanceFieldVolumeTexture* Texture);
+
+	/** Reallocates the volume texture if necessary and uploads new allocations. */
+	ENGINE_API void UpdateAllocations();
+
+	EPixelFormat Format;
+	FTexture3DRHIRef VolumeTextureRHI;
+
+private:
+	/** Manages the atlas layout. */
+	FTextureLayout3d BlockAllocator;
+
+	/** Allocations that are waiting to be added until the next update. */
+	TArray<FDistanceFieldVolumeTexture*> PendingAllocations;
+
+	/** Allocations that have already been added, stored in case we need to realloc. */
+	TArray<FDistanceFieldVolumeTexture*> CurrentAllocations;
+};
+
+extern ENGINE_API TGlobalResource<FDistanceFieldVolumeTextureAtlas> GDistanceFieldVolumeTextureAtlas;
+
+/** Represents a distance field volume texture for a single UStaticMesh. */
+class FDistanceFieldVolumeTexture
+{
+public:
+	FDistanceFieldVolumeTexture(const class FDistanceFieldVolumeData& InVolumeData) :
+		VolumeData(InVolumeData),
+		Atlas(NULL),
+		AtlasAllocationMin(FIntVector(-1, -1, -1)),
+		bReferencedByAtlas(false)
+	{}
+
+	~FDistanceFieldVolumeTexture()
+	{
+		// Make sure we have been properly removed from the atlas before deleting
+		check(!bReferencedByAtlas);
+	}
+
+	void Initialize();
+	void Release();
+
+	FIntVector GetAllocationMin() const
+	{
+		return AtlasAllocationMin;
+	}
+
+	FIntVector GetAllocationSize() const;
+
+	void SetAtlas(FDistanceFieldVolumeTextureAtlas* InAtlas)
+	{
+		Atlas = InAtlas;
+	}
+
+	bool IsValidDistanceFieldVolume() const;
+
+private:
+	const FDistanceFieldVolumeData& VolumeData;
+	FDistanceFieldVolumeTextureAtlas* Atlas;
+	FIntVector AtlasAllocationMin;
+	bool bReferencedByAtlas;
+
+	friend class FDistanceFieldVolumeTextureAtlas;
+};
+
+/** Distance field data payload and output of the mesh build process. */
+class FDistanceFieldVolumeData
+{
+public:
+
+	/** Signed distance field volume stored in local space. */
+	TArray<FFloat16> DistanceFieldVolume;
+
+	/** Dimensions of DistanceFieldVolume. */
+	FIntVector Size;
+
+	/** Local space bounding box of the distance field volume. */
+	FBox LocalBoundingBox;
+
+	/** Whether the mesh was closed and therefore a valid distance field was supported. */
+	bool bMeshWasClosed;
+
+	FDistanceFieldVolumeTexture VolumeTexture;
+
+	FDistanceFieldVolumeData() :
+		Size(FIntVector(0, 0, 0)),
+		LocalBoundingBox(0),
+		bMeshWasClosed(true),
+		VolumeTexture(*this)
+	{}
+
+	friend FArchive& operator<<(FArchive& Ar,FDistanceFieldVolumeData& Data)
+	{
+		Ar << Data.DistanceFieldVolume << Data.Size << Data.LocalBoundingBox << Data.bMeshWasClosed;
+		return Ar;
+	}
+};
+
+
 /** Rendering resources needed to render an individual static mesh LOD. */
 struct FStaticMeshLODResources
 {
@@ -546,6 +666,8 @@ struct FStaticMeshLODResources
 	/** Sections for this LOD. */
 	TArray<FStaticMeshSection> Sections;
 
+	FDistanceFieldVolumeData DistanceFieldData;
+
 	/** The maximum distance by which this LOD deviates from the base from which it was generated. */
 	float MaxDeviation;
 
@@ -557,6 +679,7 @@ struct FStaticMeshLODResources
 		: MaxDeviation(0.0f)
 		, bHasAdjacencyInfo(false)
 	{
+		DistanceFieldData.VolumeTexture.SetAtlas(&GDistanceFieldVolumeTextureAtlas);
 	}
 
 	/** Initializes all rendering resources. */
@@ -762,6 +885,8 @@ public:
 	virtual FPrimitiveViewRelevance GetViewRelevance(const FSceneView* View) OVERRIDE;
 	virtual bool CanBeOccluded() const OVERRIDE;
 	virtual void GetLightRelevance(const FLightSceneProxy* LightSceneProxy, bool& bDynamic, bool& bRelevant, bool& bLightMapped, bool& bShadowMapped) const OVERRIDE;
+	virtual void GetDistancefieldAtlasData(FBox& LocalVolumeBounds, FIntVector& OutBlockMin, FIntVector& OutBlockSize) const OVERRIDE;
+	virtual bool HasDistanceFieldRepresentation() const OVERRIDE;
 	virtual uint32 GetMemoryFootprint( void ) const OVERRIDE { return( sizeof( *this ) + GetAllocatedSize() ); }
 	uint32 GetAllocatedSize( void ) const { return( FPrimitiveSceneProxy::GetAllocatedSize() + LODs.GetAllocatedSize() ); }
 
diff --git a/Engine/Source/Runtime/Engine/Public/TextureLayout3d.h b/Engine/Source/Runtime/Engine/Public/TextureLayout3d.h
--- a/Engine/Source/Runtime/Engine/Public/TextureLayout3d.h
+++ b/Engine/Source/Runtime/Engine/Public/TextureLayout3d.h
@@ -166,6 +166,21 @@ public:
 
 	uint32 GetSizeZ() const { return SizeZ; }
 
+	uint32 GetMaxSizeX() const 
+	{
+		return Nodes[0].SizeX;
+	}
+
+	uint32 GetMaxSizeY() const 
+	{
+		return Nodes[0].SizeY;
+	}
+
+	uint32 GetMaxSizeZ() const 
+	{
+		return Nodes[0].SizeZ;
+	}
+
 private:
 
 	struct FTextureLayoutNode3d
diff --git a/Engine/Source/Runtime/Engine/Public/kDOP.h b/Engine/Source/Runtime/Engine/Public/kDOP.h
--- a/Engine/Source/Runtime/Engine/Public/kDOP.h
+++ b/Engine/Source/Runtime/Engine/Public/kDOP.h
@@ -1,98 +1,406 @@
 // Copyright 1998-2014 Epic Games, Inc. All Rights Reserved.
 
+#pragma once
 
-// ONLY HERE FOR BACKWARDS COMPAT SERIALIZATION
-// Can remove this file once VER_UE4_REMOVE_KDOP_PROPERTIES is no longer supported
+// Indicates how many "k / 2" there are in the k-DOP. 3 == AABB == 6 DOP. The code relies on this being 3.
+#define NUM_PLANES	3
 
+struct FkHitResult
+{
+	/** Normal vector in coordinate system of the returner. Zero==none.	*/
+	FVector4	Normal;
+	/** Time until hit.													*/
+	float		Time;
+	/** Primitive data item which was hit, INDEX_NONE=none				*/
+	int32			Item;		
 
-#ifndef _KDOP_H
-#define _KDOP_H
+	FkHitResult()
+	: Normal	(0,0,0,0)
+	, Time		(1.0f)
+	, Item		(INDEX_NONE)
+	{
+	}
+};
 
+/**
+ * Stores XYZ from 4 vectors in one Struct Of Arrays.
+ */
+struct FVector3SOA
+{
+	/** X = (v0.x, v1.x, v2.x, v3.x) */
+	VectorRegister	X;
+	/** Y = (v0.y, v1.y, v2.y, v3.y) */
+	VectorRegister	Y;
+	/** Z = (v0.z, v1.z, v2.z, v3.z) */
+	VectorRegister	Z;
+};
 
-// Indicates how many "k / 2" there are in the k-DOP. 3 == AABB == 6 DOP
-#define NUM_PLANES	3
+/**
+ * Stores XYZW from 4 vectors in one Struct Of Arrays.
+ */
+struct FVector4SOA
+{
+	/** X = (v0.x, v1.x, v2.x, v3.x) */
+	VectorRegister	X;
+	/** Y = (v0.y, v1.y, v2.y, v3.y) */
+	VectorRegister	Y;
+	/** Z = (v0.z, v1.z, v2.z, v3.z) */
+	VectorRegister	Z;
+	/** W = (v0.w, v1.w, v2.w, v3.w) */
+	VectorRegister	W;
+};
 
-template<typename BOUND_TYPE>
-struct MS_ALIGN(16) FEnsurePadding : public BOUND_TYPE
+/**
+ * Stores 4 triangles in one struct (Struct Of Arrays).
+ */
+struct FTriangleSOA
 {
-	float Pad[2];
-	FEnsurePadding()
+	/** 3 vertex positions for each of the 4 triangles. */
+	FVector3SOA	Positions[3];
+	/** Triangle normal (including -W for a plane) for each of the 4 triangles. */
+	FVector4SOA	Normals;
+	/** A 32-bit payload value for each of the 4 triangles. */
+	uint32		Payload[4];
+};
+
+/** Tracks information about a traversal through the kDOP tree. */
+template<typename KDOP_IDX_TYPE>
+struct TTraversalHistory
+{
+	/** 
+	 * Number of parent nodes to track information for. 
+	 * This has been tweaked for optimal boolean shadow ray performance.
+	 */
+	enum { NodeHistoryLength = 1 };
+	KDOP_IDX_TYPE Nodes[NodeHistoryLength];
+
+	TTraversalHistory()
 	{
+		for (int32 NodeIndex = 0; NodeIndex < NodeHistoryLength; NodeIndex++)
+		{
+			// Initialize the history to invalid index
+			Nodes[NodeIndex] = 0xFFFFFFFF;
+		}
 	}
-	FEnsurePadding(const BOUND_TYPE& InBound)
-		: BOUND_TYPE(InBound)
+
+	TTraversalHistory AddNode(KDOP_IDX_TYPE NewNodeIndex)
 	{
+		TTraversalHistory NewHistory;
+		for (int32 NodeIndex = NodeHistoryLength - 1; NodeIndex > 0; NodeIndex--)
+		{
+			// Move all the indices toward the end of the array one element
+			NewHistory.Nodes[NodeIndex] = Nodes[NodeIndex - 1];
+		}
+		// Insert the new node at the beginning of the array
+		NewHistory.Nodes[0] = NewNodeIndex;
+		return NewHistory;
 	}
-} GCC_ALIGN(16);
 
+	KDOP_IDX_TYPE GetOldestNode() const
+	{
+		// Accesses the oldest node that this traversal history is tracking
+		return Nodes[NodeHistoryLength - 1];
+	}
+};
 
-template<typename KDOP_IDX_TYPE> struct FkDOPCollisionTriangle
+/**
+ * Line vs triangle intersection test.
+ *
+ * @param Start		Start of the line segment
+ * @param End		End of the line segment
+ * @param Dir		Direction of the line segment (not normalized, just End-Start)
+ * @param V0		First vertex of the triangle
+ * @param V1		Second vertex of the triangle
+ * @param V2		Third vertex of the triangle
+ * @param Normal	Triangle normal (including +W for the plane equation)
+ * @param IntersectionTime	[in/out] Best intersection time so far (0..1), as in: IntersectionPoint = Start + IntersectionTime * Dir.
+ * @return			true if the line intersected the triangle
+ */
+FORCEINLINE bool appLineCheckTriangle(const FVector4& Start, const FVector4& End, const FVector4& Dir, const FVector4& V0, const FVector4& V1, const FVector4& V2, const FVector4& Normal, float& IntersectionTime)
 {
-	KDOP_IDX_TYPE v1, v2, v3;
-	KDOP_IDX_TYPE MaterialIndex;
+	const float StartDist = FPlane(Normal).PlaneDot(Start);
+	const float EndDist = FPlane(Normal).PlaneDot(End);
 
-	friend FArchive& operator<<(FArchive& Ar, FkDOPCollisionTriangle<KDOP_IDX_TYPE>& Tri)
+	// Check if the line is completely on one side of the triangle, or if it's co-planar.
+#if 1
+	if ((StartDist == EndDist) || (StartDist < -0.001f && EndDist < -0.001f) || (StartDist > 0.001f && EndDist > 0.001f))
+#else
+	if ( (StartDist * EndDist) > -0.0001f )
+#endif
 	{
-		Ar << Tri.v1 << Tri.v2 << Tri.v3 << Tri.MaterialIndex;
-		return Ar;
+		return false;
 	}
-};
 
+	// Figure out when it will hit the triangle
+	float Time = -StartDist / (EndDist - StartDist);
 
-// Forward declarations
-template <typename KDOP_IDX_TYPE> struct TkDOP;
-template <typename KDOP_IDX_TYPE> struct TkDOPNode;
-template <typename KDOP_IDX_TYPE> struct TkDOPTree;
-template <typename KDOP_IDX_TYPE> struct TkDOPCompact;
-template <typename KDOP_IDX_TYPE> struct TkDOPNodeCompact;
-template <typename KDOP_IDX_TYPE> struct TkDOPTreeCompact;
+	// If this triangle is not closer than the previous hit, reject it
+	if (Time < 0.f || Time >= IntersectionTime)
+	{
+		return false;
+	}
 
+	// Calculate the line's point of intersection with the node's plane
+	const FVector4& Intersection = Start + Dir * Time;
+	const FVector4* Verts[3] = 
+	{ 
+		&V0, &V1, &V2
+	};
 
-struct FkDOPPlanes
-{
-	static FVector PlaneNormals[NUM_PLANES];
-};
+	// Check if the point of intersection is inside the triangle's edges.
+	for( int32 SideIndex = 0; SideIndex < 3; SideIndex++ )
+	{
+		const FVector4& SideDirection = Normal ^ (*Verts[(SideIndex + 1) % 3] - *Verts[SideIndex]);
+		const float SideW = Dot3(SideDirection, *Verts[SideIndex]);
+		const float DotW = Dot3(SideDirection, Intersection);
+		if ((DotW - SideW) >= 0.001f)
+		{
+			return false;
+		}
+	}
+	IntersectionTime = Time;
+	return true;
+}
+
+/** ( -0.0001f, -0.0001f, -0.0001f, -0.0001f ) */
+static const VectorRegister GSmallNegativeNumber = { -0.0001f, -0.0001f, -0.0001f, -0.0001f };
+//extern const VectorRegister GSmallNegativeNumber;
+
+/** ( 0.0001f, 0.0001f, 0.0001f, 0.0001f ) */
+static const VectorRegister GSmallNumber = { 0.0001f, 0.0001f, 0.0001f, 0.0001f };
+//extern const VectorRegister GSmallNumber;
+
+static const VectorRegister GZeroVectorRegister = { 0, 0, 0, 0 };
 
+static const VectorRegister VectorNegativeOne = MakeVectorRegister( -1.0f, -1.0f, -1.0f, -1.0f );
 
-template <typename KDOP_IDX_TYPE> struct TkDOP :
-	public FkDOPPlanes
+/**
+ * Line vs triangle intersection test. Tests 1 line against 4 triangles at once.
+ *
+ * @param Start		Start of the line segment
+ * @param End		End of the line segment
+ * @param Dir		Direction of the line segment (not normalized, just End-Start)
+ * @param Triangle4	Four triangles
+ * @param IntersectionTime	[in/out] Best intersection time so far (0..1), as in: IntersectionPoint = Start + IntersectionTime * Dir.
+ * @return			Index (0-3) to specify which of the 4 triangles the line intersected, or -1 if none was found.
+ */
+int32 appLineCheckTriangleSOA(const FVector3SOA& Start, const FVector3SOA& End, const FVector3SOA& Dir, const FTriangleSOA& Triangle4, float& InOutIntersectionTime)
 {
-	float Min[NUM_PLANES];
-	float Max[NUM_PLANES];
+	VectorRegister TriangleMask;
 
-	// Serialization
-	friend FArchive& operator<<(FArchive& Ar,TkDOP<KDOP_IDX_TYPE>& kDOP)
+ 	VectorRegister StartDist;
+	StartDist = VectorMultiplyAdd( Triangle4.Normals.X, Start.X, Triangle4.Normals.W );
+	StartDist = VectorMultiplyAdd( Triangle4.Normals.Y, Start.Y, StartDist );
+	StartDist = VectorMultiplyAdd( Triangle4.Normals.Z, Start.Z, StartDist );
+
+	VectorRegister EndDist;
+	EndDist = VectorMultiplyAdd( Triangle4.Normals.X, End.X, Triangle4.Normals.W );
+	EndDist = VectorMultiplyAdd( Triangle4.Normals.Y, End.Y, EndDist );
+	EndDist = VectorMultiplyAdd( Triangle4.Normals.Z, End.Z, EndDist );
+
+	// Are both end-points of the line on the same side of the triangle (or parallel to the triangle plane)?
+	TriangleMask = VectorMask_LE( VectorMultiply( StartDist, EndDist ), GSmallNegativeNumber );
+	if ( VectorMaskBits(TriangleMask) == 0 )
 	{
-		for (int32 nIndex = 0; nIndex < NUM_PLANES; nIndex++)
-		{
-			Ar << kDOP.Min[nIndex];
-		}
+		return -1;
+	}
+
+	// Figure out when it will hit the triangle
+	VectorRegister Time = VectorDivide( StartDist, VectorSubtract(StartDist, EndDist) );
+
+	// If this triangle is not closer than the previous hit, reject it
+	VectorRegister IntersectionTime = VectorLoadFloat1( &InOutIntersectionTime );
+	TriangleMask = VectorBitwiseAND( TriangleMask, VectorMask_GE( Time, VectorZero() ) );
+	TriangleMask = VectorBitwiseAND( TriangleMask, VectorMask_LT( Time, IntersectionTime ) );
+	if ( VectorMaskBits(TriangleMask) == 0 )
+	{
+		return -1;
+	}
+
+	// Calculate the line's point of intersection with the node's plane
+	const VectorRegister IntersectionX = VectorMultiplyAdd( Dir.X, Time, Start.X );
+	const VectorRegister IntersectionY = VectorMultiplyAdd( Dir.Y, Time, Start.Y );
+	const VectorRegister IntersectionZ = VectorMultiplyAdd( Dir.Z, Time, Start.Z );
 
-		for (int32 nIndex = 0; nIndex < NUM_PLANES; nIndex++)
+	// Check if the point of intersection is inside the triangle's edges.
+	for( int32 SideIndex = 0; SideIndex < 3; SideIndex++ )
+	{
+		const VectorRegister EdgeX = VectorSubtract( Triangle4.Positions[(SideIndex + 1) % 3].X, Triangle4.Positions[SideIndex].X );
+		const VectorRegister EdgeY = VectorSubtract( Triangle4.Positions[(SideIndex + 1) % 3].Y, Triangle4.Positions[SideIndex].Y );
+		const VectorRegister EdgeZ = VectorSubtract( Triangle4.Positions[(SideIndex + 1) % 3].Z, Triangle4.Positions[SideIndex].Z );
+		const VectorRegister SideDirectionX = VectorSubtract( VectorMultiply( Triangle4.Normals.Y, EdgeZ ), VectorMultiply(Triangle4.Normals.Z, EdgeY) );
+		const VectorRegister SideDirectionY = VectorSubtract( VectorMultiply( Triangle4.Normals.Z, EdgeX ), VectorMultiply(Triangle4.Normals.X, EdgeZ) );
+		const VectorRegister SideDirectionZ = VectorSubtract( VectorMultiply( Triangle4.Normals.X, EdgeY ), VectorMultiply(Triangle4.Normals.Y, EdgeX) );
+		VectorRegister SideW;
+		SideW = VectorMultiply( SideDirectionX, Triangle4.Positions[SideIndex].X );
+		SideW = VectorMultiplyAdd( SideDirectionY, Triangle4.Positions[SideIndex].Y, SideW );
+		SideW = VectorMultiplyAdd( SideDirectionZ, Triangle4.Positions[SideIndex].Z, SideW );
+		VectorRegister DotW;
+		DotW = VectorMultiply( SideDirectionX, IntersectionX );
+		DotW = VectorMultiplyAdd( SideDirectionY, IntersectionY, DotW );
+		DotW = VectorMultiplyAdd( SideDirectionZ, IntersectionZ, DotW );
+		TriangleMask = VectorBitwiseAND( TriangleMask, VectorMask_LT( VectorSubtract(DotW, SideW), GSmallNumber ) );
+		if ( VectorMaskBits(TriangleMask) == 0 )
 		{
-			Ar << kDOP.Max[nIndex];
+			return -1;
 		}
-		return Ar;
+	}
+
+	// Set all non-hitting times to 1.0
+	Time = VectorSelect( TriangleMask, Time, VectorOne() );
+
+	// Get the best intersection time out of the 4 possibilities.
+	VectorRegister BestTimes = VectorMin( Time, VectorSwizzle(Time, 2, 3, 0, 0) );
+	BestTimes = VectorMin( BestTimes, VectorSwizzle(BestTimes, 1, 0, 0, 0) );
+	IntersectionTime = VectorReplicate( BestTimes, 0 );
+
+	// Get the triangle index that corresponds to the best time.
+	// NOTE: This will pick the first triangle, in case there are multiple hits at the same spot.
+	int32 SubIndex = VectorMaskBits( VectorMask_EQ( Time, IntersectionTime ) );
+	SubIndex = appCountTrailingZeros( SubIndex );
+
+	// Return results.
+	VectorStoreFloat1( IntersectionTime, &InOutIntersectionTime );
+	return SubIndex;
+}
+
+// This structure is used during the build process. It contains the triangle's
+// centroid for calculating which plane it should be split or not with
+template<typename KDOP_IDX_TYPE>
+struct FkDOPBuildCollisionTriangle
+{
+	/**
+	 * First vertex in the triangle
+	 */
+	FVector4 V0;
+	/**
+	 * Second vertex in the triangle
+	 */
+	FVector4 V1;
+	/**
+	 * Third vertex in the triangle
+	 */
+	FVector4 V2;
+
+	/** The material of this triangle */
+	KDOP_IDX_TYPE MaterialIndex;
+
+	inline FVector4 GetCentroid() const
+	{
+		return (V0 + V1 + V2) / 3.f;
+	}
+	FVector4 GetLocalNormal() const
+	{
+		FVector4 LocalNormal = ((V1 - V2) ^ (V0 - V2)).SafeNormal();
+		LocalNormal.W = Dot3(V0, LocalNormal);
+		return LocalNormal;
+	}
+
+	/**
+	 * Sets the indices, material index, calculates the centroid using the
+	 * specified triangle vertex positions
+	 */
+	FkDOPBuildCollisionTriangle(
+		KDOP_IDX_TYPE InMaterialIndex,
+		const FVector4& vert0,const FVector4& vert1,const FVector4& vert2) :
+		V0(vert0), V1(vert1), V2(vert2),
+		MaterialIndex(InMaterialIndex)
+	{
 	}
 };
 
+// Forward declarations
+template <typename COLL_DATA_PROVIDER,typename KDOP_IDX_TYPE> struct TkDOPNode;
+template <typename COLL_DATA_PROVIDER,typename KDOP_IDX_TYPE> struct TkDOPTree;
+template <typename COLL_DATA_PROVIDER,typename KDOP_IDX_TYPE> struct TkDOPLineCollisionCheck;
 
-template <typename KDOP_IDX_TYPE> struct TkDOPNode
+/**
+ * Holds the min/max planes that make up a set of 4 bounding volumes.
+ */
+struct FFourBox
 {
-	typedef TkDOP<KDOP_IDX_TYPE>		kDOPType;
-	typedef TkDOPNode<KDOP_IDX_TYPE>	NodeType;
+	/**
+	 * Min planes for this set of bounding volumes. Array index is X/Y/Z.
+	 */
+	MS_ALIGN(16) FVector4 Min[3]  GCC_ALIGN(16);
 
-	kDOPType BoundingVolume;
+	/** 
+	 * Max planes for this set of bounding volumes. Array index is X/Y/Z.
+	 */
+	FVector4 Max[3];
+
+	/**
+	 * Sets the box at the passed in index to the passed in box.
+	 *
+	 * @param	BoundingVolumeIndex		Index of box to set
+	 * @param	Box						Box to set
+	 */
+	void SetBox( int32 BoundingVolumeIndex, const FBox& Box )
+	{
+		Min[0].Component(BoundingVolumeIndex) = Box.Min.X;
+		Min[1].Component(BoundingVolumeIndex) = Box.Min.Y;
+		Min[2].Component(BoundingVolumeIndex) = Box.Min.Z;
+		Max[0].Component(BoundingVolumeIndex) = Box.Max.X;
+		Max[1].Component(BoundingVolumeIndex) = Box.Max.Y;
+		Max[2].Component(BoundingVolumeIndex) = Box.Max.Z;
+	}
+
+	/**
+	 * Returns the bounding volume at the passed in index.
+	 *
+	 * @param	BoundingVolumeIndex		Index of bounding volume to return
+	 *
+	 * @return Bounding volume at the passed in index
+	 */
+	FBox GetBox( int32 BoundingVolumeIndex )
+	{
+		FBox Box;
+		Box.Min = FVector4(Min[0][BoundingVolumeIndex],Min[1][BoundingVolumeIndex],Min[2][BoundingVolumeIndex],1);
+		Box.Max = FVector4(Max[0][BoundingVolumeIndex],Max[1][BoundingVolumeIndex],Max[2][BoundingVolumeIndex],1);
+		return Box;
+	}
+};
+
+#if PLATFORM_64BITS
+	#define kDOPArray TArray
+#else
+	#define kDOPArray TChunkedArray
+#endif
+
+/**
+ * A node in the kDOP tree. The node contains the kDOP volume that encompasses
+ * it's children and/or triangles
+ */
+template <typename COLL_DATA_PROVIDER, typename KDOP_IDX_TYPE>
+struct TkDOPNode
+{
+	/** Exposes data provider type to clients. */
+	typedef COLL_DATA_PROVIDER							DataProviderType;
+
+	/** Exposes node type to clients. */
+	typedef TkDOPNode<DataProviderType,KDOP_IDX_TYPE>	NodeType;
+
+	/** Set of bounding volumes for child nodes. */
+	FFourBox BoundingVolumes;
+
+	// Note this isn't smaller since 4 byte alignment will take over anyway
 	bool bIsLeaf;
+	uint8 Occupancy;
 
+	// Union of either child kDOP nodes or a list of enclosed triangles
 	union
 	{
+		// This structure contains the left and right child kDOP indices
+		// These index values correspond to the array in the FkDOPTree
 		struct
 		{
 			KDOP_IDX_TYPE LeftNode;
 			KDOP_IDX_TYPE RightNode;
 		} n;
-
+		// This structure contains the list of enclosed triangles
+		// These index values correspond to the triangle information in the
+		// FkDOPTree using the start and count as the means of delineating
+		// which triangles are involved
 		struct
 		{
 			KDOP_IDX_TYPE NumTriangles;
@@ -100,89 +408,684 @@ template <typename KDOP_IDX_TYPE> struct TkDOPNode
 		} t;
 	};
 
-	friend FArchive& operator<<(FArchive& Ar,NodeType& Node)
+	/**
+	 * Inits the data to no child nodes and an inverted volume
+	 */
+	FORCEINLINE TkDOPNode()
 	{
-		Ar << Node.BoundingVolume << Node.bIsLeaf;
-		Ar << Node.n.LeftNode << Node.n.RightNode;
-		return Ar;
+		n.LeftNode = ((KDOP_IDX_TYPE) -1);
+        n.RightNode = ((KDOP_IDX_TYPE) -1);
 	}
-};
 
-template<typename KDOP_IDX_TYPE> struct TkDOPTree
-{
-	typedef TkDOPNode<KDOP_IDX_TYPE>	NodeType;
+	/**
+	 * Determines if the node is a leaf or not. If it is not a leaf, it subdivides
+	 * the list of triangles again adding two child nodes and splitting them on
+	 * the mean (splatter method). Otherwise it sets up the triangle information.
+	 *
+	 * @param Start -- The triangle index to start processing with
+	 * @param NumTris -- The number of triangles to process
+	 * @param BuildTriangles -- The list of triangles to use for the build process
+	 * @param Nodes -- The list of nodes in this tree
+	 * @return bounding box for this node
+	 */
+	FBox SplitTriangleList(int32 Start,int32 NumTris,
+		TArray<FkDOPBuildCollisionTriangle<KDOP_IDX_TYPE> >& BuildTriangles,
+		kDOPArray<FTriangleSOA>& SOATriangles,
+		kDOPArray<NodeType>& Nodes)
+	{
+		// Figure out if we are a leaf node or not
+		if (NumTris > 4)
+		{
+			// Still too many triangles, so continue subdividing the triangle list
+			bIsLeaf = 0;
+			Occupancy = 0;
+			int32 BestPlane = -1;
+			float BestMean = 0.f;
+			float BestVariance = 0.f;
+			// Determine how to split using the splatter algorithm
+			for (int32 nPlane = 0; nPlane < NUM_PLANES; nPlane++)
+			{
+				float Mean = 0.f;
+				float Variance = 0.f;
+				// Compute the mean for the triangle list
+				for (int32 nTriangle = Start; nTriangle < Start + NumTris; nTriangle++)
+				{
+					// Project the centroid of the triangle against the plane
+					// normals and accumulate to find the total projected
+					// weighting
+					Mean += BuildTriangles[nTriangle].GetCentroid()[nPlane];
+				}
+				// Divide by the number of triangles to get the average
+				Mean /= float(NumTris);
+				// Compute variance of the triangle list
+				for (int32 nTriangle = Start; nTriangle < Start + NumTris;nTriangle++)
+				{
+					// Project the centroid again
+					float Dot = BuildTriangles[nTriangle].GetCentroid()[nPlane];
+					// Now calculate the variance and accumulate it
+					Variance += (Dot - Mean) * (Dot - Mean);
+				}
+				// Get the average variance
+				Variance /= float(NumTris);
+				// Determine if this plane is the best to split on or not
+				if (Variance >= BestVariance)
+				{
+					BestPlane = nPlane;
+					BestVariance = Variance;
+					BestMean = Mean;
+				}
+			}
+			// Now that we have the plane to split on, work through the triangle
+			// list placing them on the left or right of the splitting plane
+			int32 Left = Start - 1;
+			int32 Right = Start + NumTris;
+			// Keep working through until the left index passes the right
+			while (Left < Right)
+			{
+				float Dot;
+				// Find all the triangles to the "left" of the splitting plane
+				do
+				{
+					Dot = BuildTriangles[++Left].GetCentroid()[BestPlane];
+				}
+				while (Dot < BestMean && Left < Right);
+				// Find all the triangles to the "right" of the splitting plane
+				do
+				{
+					Dot = BuildTriangles[--Right].GetCentroid()[BestPlane];
+				}
+				while (Dot >= BestMean && Right > 0 && Left < Right);
+				// Don't swap the triangle data if we just hit the end
+				if (Left < Right)
+				{
+					// Swap the triangles since they are on the wrong sides of the
+					// splitting plane
+					FkDOPBuildCollisionTriangle<KDOP_IDX_TYPE> Temp = BuildTriangles[Left];
+					BuildTriangles[Left] = BuildTriangles[Right];
+					BuildTriangles[Right] = Temp;
+				}
+			}
+			// Check for wacky degenerate case where more than GKDOPMaxTrisPerLeaf
+			// fall all in the same kDOP
+			if (Left == Start + NumTris || Right == Start)
+			{
+				Left = Start + (NumTris / 2);
+			}
+			// Add the two child nodes
+			n.LeftNode = Nodes.AddZeroed(2);
+			n.RightNode = n.LeftNode + 1;
+			// Have the left node recursively subdivide it's list and set bounding volume.
+			FBox LeftBoundingVolume = Nodes[n.LeftNode].SplitTriangleList(Start,Left - Start,BuildTriangles,SOATriangles,Nodes);
+			BoundingVolumes.SetBox(0,LeftBoundingVolume);
+			// Set unused index 2,3 to child nodes of left node.
+			BoundingVolumes.SetBox(2,Nodes[n.LeftNode].BoundingVolumes.GetBox(0));
+			BoundingVolumes.SetBox(3,Nodes[n.LeftNode].BoundingVolumes.GetBox(1));
 
-	TArray<NodeType> Nodes;
-	TArray<FkDOPCollisionTriangle<KDOP_IDX_TYPE> > Triangles;
+			// And now have the right node recursively subdivide it's list and set bounding volume.			
+			FBox RightBoundingVolume = Nodes[n.RightNode].SplitTriangleList(Left,Start + NumTris - Left,BuildTriangles,SOATriangles,Nodes);
+			BoundingVolumes.SetBox(1,RightBoundingVolume);
 
-	friend FArchive& operator<<(FArchive& Ar,TkDOPTree<KDOP_IDX_TYPE>& Tree)
-	{
-		Tree.Nodes.BulkSerialize( Ar );
-		Tree.Triangles.BulkSerialize( Ar );
-		return Ar;
+			// Non-leaf node bounds are the "sum" of the left and right nodes' volumes.
+			return LeftBoundingVolume + RightBoundingVolume;
+		}
+		else
+		{
+			// Build SOA triangles
+
+			// "NULL triangle", used when a leaf can't fill all 4 triangles in a FTriangleSOA.
+			// No line should ever hit these triangles, set the values so that it can never happen.
+			FkDOPBuildCollisionTriangle<KDOP_IDX_TYPE> EmptyTriangle(0,FVector4(0,0,0,0),FVector4(0,0,0,0),FVector4(0,0,0,0));
+			
+			t.StartIndex = SOATriangles.Num();
+			t.NumTriangles = Align<int32>(NumTris, 4) / 4;
+			SOATriangles.AddZeroed( t.NumTriangles );
+
+			int32 BuildTriIndex = Start;
+			for ( uint32 SOAIndex=0; SOAIndex < t.NumTriangles; ++SOAIndex )
+			{
+				FkDOPBuildCollisionTriangle<KDOP_IDX_TYPE>* Tris[4] = { &EmptyTriangle, &EmptyTriangle, &EmptyTriangle, &EmptyTriangle };
+				FTriangleSOA& SOA = SOATriangles[t.StartIndex + SOAIndex];
+				int32 SubIndex = 0;
+				for ( ; SubIndex < 4 && BuildTriIndex < (Start+NumTris); ++SubIndex, ++BuildTriIndex )
+				{
+					Tris[SubIndex] = &BuildTriangles[BuildTriIndex];
+					SOA.Payload[SubIndex] = Tris[SubIndex]->MaterialIndex;
+				}
+				for ( ; SubIndex < 4; ++SubIndex )
+				{
+					SOA.Payload[SubIndex] = 0xffffffff;
+				}
+
+				SOA.Positions[0].X = VectorSet( Tris[0]->V0.X, Tris[1]->V0.X, Tris[2]->V0.X, Tris[3]->V0.X );
+				SOA.Positions[0].Y = VectorSet( Tris[0]->V0.Y, Tris[1]->V0.Y, Tris[2]->V0.Y, Tris[3]->V0.Y );
+				SOA.Positions[0].Z = VectorSet( Tris[0]->V0.Z, Tris[1]->V0.Z, Tris[2]->V0.Z, Tris[3]->V0.Z );
+				SOA.Positions[1].X = VectorSet( Tris[0]->V1.X, Tris[1]->V1.X, Tris[2]->V1.X, Tris[3]->V1.X );
+				SOA.Positions[1].Y = VectorSet( Tris[0]->V1.Y, Tris[1]->V1.Y, Tris[2]->V1.Y, Tris[3]->V1.Y );
+				SOA.Positions[1].Z = VectorSet( Tris[0]->V1.Z, Tris[1]->V1.Z, Tris[2]->V1.Z, Tris[3]->V1.Z );
+				SOA.Positions[2].X = VectorSet( Tris[0]->V2.X, Tris[1]->V2.X, Tris[2]->V2.X, Tris[3]->V2.X );
+				SOA.Positions[2].Y = VectorSet( Tris[0]->V2.Y, Tris[1]->V2.Y, Tris[2]->V2.Y, Tris[3]->V2.Y );
+				SOA.Positions[2].Z = VectorSet( Tris[0]->V2.Z, Tris[1]->V2.Z, Tris[2]->V2.Z, Tris[3]->V2.Z );
+
+				const FVector4& Tris0LocalNormal = Tris[0]->GetLocalNormal();
+				const FVector4& Tris1LocalNormal = Tris[1]->GetLocalNormal();
+				const FVector4& Tris2LocalNormal = Tris[2]->GetLocalNormal();
+				const FVector4& Tris3LocalNormal = Tris[3]->GetLocalNormal();
+
+				SOA.Normals.X = VectorSet( Tris0LocalNormal.X, Tris1LocalNormal.X, Tris2LocalNormal.X, Tris3LocalNormal.X );
+				SOA.Normals.Y = VectorSet( Tris0LocalNormal.Y, Tris1LocalNormal.Y, Tris2LocalNormal.Y, Tris3LocalNormal.Y );
+				SOA.Normals.Z = VectorSet( Tris0LocalNormal.Z, Tris1LocalNormal.Z, Tris2LocalNormal.Z, Tris3LocalNormal.Z );
+				SOA.Normals.W = VectorSet( -Tris0LocalNormal.W, -Tris1LocalNormal.W, -Tris2LocalNormal.W, -Tris3LocalNormal.W );
+			}
+
+			// No need to subdivide further so make this a leaf node
+			bIsLeaf = 1;
+			Occupancy = NumTris;
+			
+			// Generate bounding volume for leaf which is passed up the call chain.
+			FBox BoundingVolume(0);
+			for (int32 TriangleIndex=Start; TriangleIndex<Start + NumTris; TriangleIndex++)
+			{
+				BoundingVolume += BuildTriangles[TriangleIndex].V0;
+				BoundingVolume += BuildTriangles[TriangleIndex].V1;
+				BoundingVolume += BuildTriangles[TriangleIndex].V2;			
+			}
+			BoundingVolumes.SetBox(0,BoundingVolume);
+			BoundingVolumes.SetBox(1,BoundingVolume);
+			BoundingVolumes.SetBox(2,BoundingVolume);
+			BoundingVolumes.SetBox(3,BoundingVolume);
+
+			return BoundingVolume;
+		}
 	}
-};
 
+	/**
+	 * The slab testing algorithm is based on the following papers. We chose to use the
+	 * faster final hit determination, which means we'll get some false positives.
+     *
+	 * "A Cross-Platform Framework for Interactive Ray Tracing"
+	 * Markus Geimer, Stefan Mueller
+	 * http://www.uni-koblenz.de/~cg/publikationen/cp_raytrace.pdf
+	 *
+  	 * "Efficiency Issues for Ray Tracing"
+	 * Brian Smits
+	 * http://www.cs.utah.edu/~bes/papers/fastRT/paper-node10.html
+	 *
+	 * We don't need to handle +/- INF, but we have permission to use Thierry 
+	 * Berger-Perrin trick available below.
+	 *
+     * http://www.flipcode.com/archives/SSE_RayBox_Intersection_Test.shtml
+	 *
+	 * @param	Check				Information about the ray to trace
+	 * @param	HitTime	[out]	Time of hit
+	 */
+	FORCEINLINE void LineCheckBounds(TkDOPLineCollisionCheck<COLL_DATA_PROVIDER,KDOP_IDX_TYPE>& Check, FVector4& HitTime, int32 NodeHit[4] ) const
+	{
+#define PLAIN_C 0
+#if PLAIN_C
+		for( int32 BoxIndex=0; BoxIndex<4; BoxIndex++ )
+		{
+			// 0: Create constants
+			FVector4 BoxMin( BoundingVolumes.Min[0][BoxIndex], BoundingVolumes.Min[1][BoxIndex], BoundingVolumes.Min[2][BoxIndex], 0 );
+			FVector4 BoxMax( BoundingVolumes.Max[0][BoxIndex], BoundingVolumes.Max[1][BoxIndex], BoundingVolumes.Max[2][BoxIndex], 0 );
 
-template <typename KDOP_IDX_TYPE> struct TkDOPCompact :
-	public FkDOPPlanes
-{
-	typedef TkDOP<KDOP_IDX_TYPE> BoundType;
+			// 1: Calculate slabs.
+			FVector4 Slab1 = (BoxMin - Check.LocalStart) * Check.LocalOneOverDir;
+			FVector4 Slab2 = (BoxMax - Check.LocalStart) * Check.LocalOneOverDir;
+
+			// 2: Figure out per component min/ max
+			FVector4 SlabMin = FVector4( FMath::Min(Slab1.X, Slab2.X), FMath::Min(Slab1.Y, Slab2.Y), FMath::Min(Slab1.Z, Slab2.Z), FMath::Min(Slab1.W, Slab2.W) );
+			FVector4 SlabMax = FVector4( FMath::Max(Slab1.X, Slab2.X), FMath::Max(Slab1.Y, Slab2.Y), FMath::Max(Slab1.Z, Slab2.Z), FMath::Max(Slab1.W, Slab2.W) );
+
+			// 3: Figure out global min/ max
+			float MinTime = Max3( SlabMin.X, SlabMin.Y, SlabMin.Z );
+			float MaxTime = Min3( SlabMax.X, SlabMax.Y, SlabMax.Z );
+
+			// 4: Calculate hit time and determine whether there was a hit.
+			HitTime[BoxIndex] = MinTime;			
+			NodeHit[BoxIndex] = (MaxTime >= 0 && MaxTime >= MinTime && MinTime < Check.Result->Time) ? 0xFFFFFFFF : 0;
+		}
+#else
+		// 0: load everything into registers
+		const VectorRegister OriginX		= VectorSetFloat1( Check.LocalStart.X );
+		const VectorRegister OriginY		= VectorSetFloat1( Check.LocalStart.Y );
+		const VectorRegister OriginZ		= VectorSetFloat1( Check.LocalStart.Z );
+		const VectorRegister InvDirX		= VectorSetFloat1( Check.LocalOneOverDir.X );
+		const VectorRegister InvDirY		= VectorSetFloat1( Check.LocalOneOverDir.Y );
+		const VectorRegister InvDirZ		= VectorSetFloat1( Check.LocalOneOverDir.Z );
+		const VectorRegister CurrentHitTime	= VectorSetFloat1( Check.Result->Time );
+		// Boxes are FVector2D so we need to unshuffle the data.
+		const VectorRegister BoxMinX		= VectorLoadAligned( &BoundingVolumes.Min[0] );
+		const VectorRegister BoxMinY		= VectorLoadAligned( &BoundingVolumes.Min[1] );
+		const VectorRegister BoxMinZ		= VectorLoadAligned( &BoundingVolumes.Min[2] );
+		const VectorRegister BoxMaxX		= VectorLoadAligned( &BoundingVolumes.Max[0] );
+		const VectorRegister BoxMaxY		= VectorLoadAligned( &BoundingVolumes.Max[1] );
+		const VectorRegister BoxMaxZ		= VectorLoadAligned( &BoundingVolumes.Max[2] );
+
+		// 1: Calculate slabs.
+		const VectorRegister BoxMinSlabX	= VectorMultiply( VectorSubtract( BoxMinX, OriginX ), InvDirX );
+		const VectorRegister BoxMinSlabY	= VectorMultiply( VectorSubtract( BoxMinY, OriginY ), InvDirY );
+		const VectorRegister BoxMinSlabZ	= VectorMultiply( VectorSubtract( BoxMinZ, OriginZ ), InvDirZ );		
+		const VectorRegister BoxMaxSlabX	= VectorMultiply( VectorSubtract( BoxMaxX, OriginX ), InvDirX );
+		const VectorRegister BoxMaxSlabY	= VectorMultiply( VectorSubtract( BoxMaxY, OriginY ), InvDirY );
+		const VectorRegister BoxMaxSlabZ	= VectorMultiply( VectorSubtract( BoxMaxZ, OriginZ ), InvDirZ );
+
+		// 2: Figure out per component min/ max
+		const VectorRegister SlabMinX		= VectorMin( BoxMinSlabX, BoxMaxSlabX );
+		const VectorRegister SlabMinY		= VectorMin( BoxMinSlabY, BoxMaxSlabY );
+		const VectorRegister SlabMinZ		= VectorMin( BoxMinSlabZ, BoxMaxSlabZ );
+		const VectorRegister SlabMaxX		= VectorMax( BoxMinSlabX, BoxMaxSlabX );
+		const VectorRegister SlabMaxY		= VectorMax( BoxMinSlabY, BoxMaxSlabY );
+		const VectorRegister SlabMaxZ		= VectorMax( BoxMinSlabZ, BoxMaxSlabZ );
+		
+		// 3: Figure out global min/ max
+		const VectorRegister SlabMinXY		= VectorMax( SlabMinX , SlabMinY );
+		const VectorRegister MinTime		= VectorMax( SlabMinXY, SlabMinZ );
+		const VectorRegister SlabMaxXY		= VectorMin( SlabMaxX , SlabMaxY );
+		const VectorRegister MaxTime		= VectorMin( SlabMaxXY, SlabMaxZ );
+
+		// 4: Calculate hit time and determine whether there was a hit.		
+		VectorStoreAligned( MinTime, &HitTime );
+		const VectorRegister OutNodeHit		= VectorBitwiseAND( VectorCompareGE( MaxTime, VectorZero() ), VectorCompareGE( MaxTime, MinTime ) );
+		const VectorRegister CloserNodeHit	= VectorBitwiseAND( OutNodeHit, VectorCompareGT( CurrentHitTime, MinTime ) );
+		VectorStoreAligned( CloserNodeHit, (float*) NodeHit );
+#endif
+	}
+
+	/** 
+	 * Determines the line in the FkDOPLineCollisionCheck intersects this node. It
+	 * also will check the child nodes if it is not a leaf, otherwise it will check
+	 * against the triangle data.
+	 *
+	 * @param Check -- The aggregated line check data
+	 */
+	bool LineCheck(TkDOPLineCollisionCheck<COLL_DATA_PROVIDER,KDOP_IDX_TYPE>& Check, TTraversalHistory<KDOP_IDX_TYPE> History) const
+	{
+		bool bHit = 0;
+		// If this is a node, check the two child nodes and pick the closest one
+		// to recursively check against and only check the second one if there is
+		// not a hit or the hit returned is further out than the second node
+		if (bIsLeaf == 0)
+		{
+			// Check both left and right node at the same time.
+			FVector4 NodeHitTime;
+			MS_ALIGN(16) int32	NodeHit[4] GCC_ALIGN(16);
+			LineCheckBounds( Check, NodeHitTime, NodeHit );
 
-	uint8 Min[NUM_PLANES];
-	uint8 Max[NUM_PLANES];
+			// Left node was hit.
+			if( NodeHit[0] )
+			{
+				// Left and Right node are hit.
+				if( NodeHit[1] )
+				{
+					// Left node is closer than right node.
+					if( NodeHitTime.X < NodeHitTime.Y )
+					{
+						bHit = Check.Nodes[n.LeftNode].LineCheckPreCalculated(Check,NodeHitTime,History.AddNode(n.LeftNode),NodeHit);
+						// Only check right node if it could possibly be a closer hit
+						if(	NodeHitTime.Y < Check.Result->Time
+						// No need to check if we have a hit and don't care about closest
+						&&	(!bHit || Check.bFindClosestIntersection))
+						{
+							bHit |= Check.Nodes[n.RightNode].LineCheck(Check,History.AddNode(n.RightNode));
+						}
+					}
+					// Right node is closer than left node.
+					else
+					{
+						bHit = Check.Nodes[n.RightNode].LineCheck(Check,History.AddNode(n.RightNode));
+						// Only check left node if it could possibly be a closer hit
+						if(	NodeHitTime.X < Check.Result->Time
+						// No need to check if we have a hit and don't care about closest
+						&&	(!bHit || Check.bFindClosestIntersection))
+						{
+							bHit |= Check.Nodes[n.LeftNode].LineCheckPreCalculated(Check,NodeHitTime,History.AddNode(n.LeftNode),NodeHit);
+						}
+					}
+				}
+				// Only left node was hit.
+				else
+				{
+					bHit = Check.Nodes[n.LeftNode].LineCheckPreCalculated(Check,NodeHitTime,History.AddNode(n.LeftNode),NodeHit);
+				}
+			}
+			// Left node was not hit.
+			else
+			{
+				// Only right node was hit.
+				if( NodeHit[1] )
+				{			
+					bHit = Check.Nodes[n.RightNode].LineCheck(Check,History.AddNode(n.RightNode));					
+				}
+				// No node was hit.
+				else
+				{
+					return false;
+				}
+			}
+		}
+		else
+		{
+			// This is a leaf, check the triangles for a hit
+			bHit = LineCheckTriangles(Check,History);
+		}
+		return bHit;
+	}
 
-	friend FArchive& operator<<(FArchive& Ar,TkDOPCompact<KDOP_IDX_TYPE>& kDOP)
+	/** 
+	 * Determines the line in the FkDOPLineCollisionCheck intersects this node. It
+	 * also will check the child nodes if it is not a leaf, otherwise it will check
+	 * against the triangle data.
+	 *
+	 * @param Check -- The aggregated line check data
+	 */
+	bool LineCheckPreCalculated(TkDOPLineCollisionCheck<COLL_DATA_PROVIDER,KDOP_IDX_TYPE>& Check, const FVector4& NodeHitTime, TTraversalHistory<KDOP_IDX_TYPE> History, int32* NodeHit) const
 	{
-		for (int32 nIndex = 0; nIndex < NUM_PLANES; nIndex++)
+		bool bHit = 0;
+		// If this is a node, check the two child nodes and pick the closest one
+		// to recursively check against and only check the second one if there is
+		// not a hit or the hit returned is further out than the second node
+		if (bIsLeaf == 0)
+		{
+			// Left node was hit.
+			if( NodeHit[2] )
+			{
+				// Left and Right node are hit.
+				if( NodeHit[3] )
+				{
+					// Left node is closer than right node.
+					if( NodeHitTime.Z < NodeHitTime.W )
+					{
+						bHit = Check.Nodes[n.LeftNode].LineCheck(Check,History.AddNode(n.LeftNode));
+						// Only check right node if it could possibly be a closer hit
+						if(	NodeHitTime.W < Check.Result->Time
+						// No need to check if we have a hit and don't care about closest
+						&&	(!bHit || Check.bFindClosestIntersection))
+						{
+							bHit |= Check.Nodes[n.RightNode].LineCheck(Check,History.AddNode(n.RightNode));
+						}
+					}
+					// Right node is closer than left node.
+					else
+					{
+						bHit = Check.Nodes[n.RightNode].LineCheck(Check,History.AddNode(n.RightNode));
+						// Only check left node if it could possibly be a closer hit
+						if(	NodeHitTime.Z < Check.Result->Time
+						// No need to check if we have a hit and don't care about closest
+						&&	(!bHit || Check.bFindClosestIntersection))
+						{
+							bHit |= Check.Nodes[n.LeftNode].LineCheck(Check,History.AddNode(n.LeftNode));
+						}
+					}
+				}
+				// Only left node was hit.
+				else
+				{
+					bHit = Check.Nodes[n.LeftNode].LineCheck(Check,History.AddNode(n.LeftNode));
+				}
+			}
+			// Left node was not hit.
+			else
+			{
+				// Only right node was hit.
+				if( NodeHit[3] )
+				{			
+					bHit = Check.Nodes[n.RightNode].LineCheck(Check,History.AddNode(n.RightNode));					
+				}
+				// No node was hit.
+				else
+				{
+					return false;
+				}
+			}
+		}
+		else
 		{
-			Ar << kDOP.Min[nIndex];
+			// This is a leaf, check the triangles for a hit
+			bHit = LineCheckTriangles(Check,History);
 		}
+		return bHit;
+	}
 
-		for (int32 nIndex = 0; nIndex < NUM_PLANES; nIndex++)
+	/**
+	 * Works through the list of triangles in this node checking each one for a
+	 * collision.
+	 *
+	 * @param Check -- The aggregated line check data
+	 */
+	bool LineCheckTriangles(TkDOPLineCollisionCheck<COLL_DATA_PROVIDER,KDOP_IDX_TYPE>& Check, const TTraversalHistory<KDOP_IDX_TYPE>& History) const
+	{
+		// Assume a miss
+		bool bHit = false;
+		for ( KDOP_IDX_TYPE SOAIndex = t.StartIndex; SOAIndex < (t.StartIndex + t.NumTriangles); SOAIndex++ )
 		{
-			Ar << kDOP.Max[nIndex];
+			const FTriangleSOA& TriangleSOA = Check.SOATriangles[SOAIndex];
+			int32 SubIndex = appLineCheckTriangleSOA( Check.StartSOA, Check.EndSOA, Check.DirSOA, TriangleSOA, Check.Result->Time );
+			if ( SubIndex >= 0 )
+			{
+				bHit = true;
+				Check.LocalHitNormal.X = VectorGetComponent(TriangleSOA.Normals.X, SubIndex);
+				Check.LocalHitNormal.Y = VectorGetComponent(TriangleSOA.Normals.Y, SubIndex);
+				Check.LocalHitNormal.Z = VectorGetComponent(TriangleSOA.Normals.Z, SubIndex);
+				Check.Result->Item = TriangleSOA.Payload[SubIndex];
+				Check.HitNodeIndex = History.GetOldestNode();
+
+				// Early out if we don't care about the closest intersection.
+				if( !Check.bFindClosestIntersection )
+				{
+					break;
+				}
+			}
 		}
-		return Ar;
+		return bHit;
 	}
 };
 
-
-template <typename KDOP_IDX_TYPE> struct TkDOPNodeCompact
+/**
+ * This is the tree of kDOPs that spatially divides the static mesh. It is
+ * a binary tree of kDOP nodes.
+ */
+template<typename COLL_DATA_PROVIDER, typename KDOP_IDX_TYPE>
+struct TkDOPTree
 {
-	typedef TkDOPNodeCompact<KDOP_IDX_TYPE>	NodeCompactType;
-	typedef TkDOPCompact<KDOP_IDX_TYPE>		kDOPCompactType;
-	typedef TkDOP<KDOP_IDX_TYPE>				kDOPUncompressedType;
+	/** Exposes data provider type to clients. */
+	typedef COLL_DATA_PROVIDER							DataProviderType;
 
-	kDOPCompactType BoundingVolume;
+	/** Exposes node type to clients. */
+	typedef TkDOPNode<DataProviderType,KDOP_IDX_TYPE>	NodeType;
 
-	friend FArchive& operator<<(FArchive& Ar,NodeCompactType& Node)
+	/** The list of nodes contained within this tree. Node 0 is always the root node. */
+	kDOPArray<NodeType> Nodes;
+
+	/** The list of collision triangles in this tree. */
+	kDOPArray<FTriangleSOA> SOATriangles;
+
+	/**
+	 * Creates the root node and recursively splits the triangles into smaller
+	 * volumes
+	 *
+	 * @param BuildTriangles -- The list of triangles to use for the build process
+	 */
+	void Build(TArray<FkDOPBuildCollisionTriangle<KDOP_IDX_TYPE> >& BuildTriangles)
 	{
-		Ar << Node.BoundingVolume;
-		return Ar;
+		float kDOPBuildTime = 0;
+		{
+			// Empty the current set of nodes and preallocate the memory so it doesn't
+			// reallocate memory while we are recursively walking the tree
+			// With near-perfect packing, we could easily size these to be
+			// Nodes = (n / 2) + 1 and SOATriangles = (n / 4) + 1
+			Nodes.Empty(BuildTriangles.Num());
+			SOATriangles.Empty(BuildTriangles.Num());
+
+			// Add the root node
+			Nodes.AddZeroed();
+
+			// Now tell that node to recursively subdivide the entire set of triangles
+			Nodes[0].SplitTriangleList(0,BuildTriangles.Num(),BuildTriangles,SOATriangles,Nodes);
+
+			// Don't waste memory.
+			Nodes.Shrink();
+			SOATriangles.Shrink();
+		}
 	}
-};
 
+	/**
+	 * Figures out whether the check even hits the root node's bounding volume. If
+	 * it does, it recursively searches for a triangle to hit.
+	 *
+	 * @param Check -- The aggregated line check data
+	 */
+	bool LineCheck(TkDOPLineCollisionCheck<COLL_DATA_PROVIDER,KDOP_IDX_TYPE>& Check) const
+	{
+		// Recursively check for a hit
+		TTraversalHistory<KDOP_IDX_TYPE> History;
+		bool bHit = Nodes[0].LineCheck(Check, History.AddNode(0));
+		return bHit;
+	}
+};
 
-template<typename KDOP_IDX_TYPE> struct TkDOPTreeCompact
+/**
+ * Base struct for all collision checks. Holds a reference to the collision
+ * data provider, which is a struct that abstracts out the access to a
+ * particular mesh/primitives data
+ */
+template <typename COLL_DATA_PROVIDER, typename KDOP_IDX_TYPE> struct TkDOPCollisionCheck
 {
-	typedef TkDOPNodeCompact<KDOP_IDX_TYPE>	NodeType;
-	typedef typename NodeType::kDOPUncompressedType				kDOPUncompressedType;
+	/** Exposes data provider type to clients. */
+	typedef COLL_DATA_PROVIDER DataProviderType;
+
+	/** Exposes node type to clients. */
+	typedef TkDOPNode<DataProviderType,KDOP_IDX_TYPE> NodeType;
 
-	TArray<NodeType>											Nodes;
-	TArray<FkDOPCollisionTriangle<KDOP_IDX_TYPE> >				Triangles;
-	FEnsurePadding<kDOPUncompressedType>						RootBound;
+	/** Exposes tree type to clients. */
+	typedef TkDOPTree<DataProviderType,KDOP_IDX_TYPE> TreeType;
 
-	// Serialization
-	friend FArchive& operator<<(FArchive& Ar,TkDOPTreeCompact<KDOP_IDX_TYPE>& Tree)
+	/**
+	 * Used to get access to local->world, vertices, etc. without using virtuals
+	 */
+	const DataProviderType& CollDataProvider;
+	/**
+	 * The kDOP tree
+	 */
+	const TreeType& kDOPTree;
+	/**
+	 * The array of the nodes for the kDOP tree
+	 */
+	const kDOPArray<NodeType>& Nodes;
+	/**
+	 * The collision triangle data for the kDOP tree
+	 */
+	const kDOPArray<FTriangleSOA>& SOATriangles;
+
+	/**
+	 * Hide the default ctor
+	 */
+	TkDOPCollisionCheck(const DataProviderType& InCollDataProvider) :
+		CollDataProvider(InCollDataProvider),
+		kDOPTree(CollDataProvider.GetkDOPTree()),
+		Nodes(kDOPTree.Nodes),
+		SOATriangles(kDOPTree.SOATriangles)
 	{
-		Ar << Tree.RootBound;
-		Tree.Nodes.BulkSerialize( Ar );
-		Tree.Triangles.BulkSerialize( Ar );
-		check(!Tree.Nodes.Num() || (PTRINT(&Tree.Nodes[0]) & 3) == 0); // SIMD will need this avoid the potential for page faults
-		return Ar;
 	}
-
 };
 
-#endif
+/**
+ * This struct holds the information used to do a line check against the kDOP
+ * tree. The collision provider gives access to various matrices, vertex data
+ * etc. without having to use virtual functions.
+ */
+template <typename COLL_DATA_PROVIDER, typename KDOP_IDX_TYPE> struct TkDOPLineCollisionCheck :
+	public TkDOPCollisionCheck<COLL_DATA_PROVIDER,KDOP_IDX_TYPE>
+{
+	/**
+	 * Where the collision results get stored
+	 */
+	FkHitResult* Result;
+	// Constant input vars
+	const FVector4& Start;
+	const FVector4& End;
+
+	/**
+	 * Flags for optimizing a trace
+	 */
+	const bool bFindClosestIntersection;
+	// Locally calculated vectors
+	FVector4 LocalStart;
+	FVector4 LocalEnd;
+	FVector4 LocalDir;
+	FVector4 LocalOneOverDir;
+	// Normal in local space which gets transformed to world at the very end
+	FVector4 LocalHitNormal;
+
+	/** Index into the kDOP tree's nodes of the node that was hit. */
+	KDOP_IDX_TYPE HitNodeIndex;
+	
+	/** Start of the line, where each component is replicated into their own vector registers. */
+	FVector3SOA	StartSOA;
+	/** End of the line, where each component is replicated into their own vector registers. */
+	FVector3SOA	EndSOA;
+	/** Direction of the line (not normalized, just EndSOA-StartSOA), where each component is replicated into their own vector registers. */
+	FVector3SOA	DirSOA;
+
+	/**
+	 * Sets up the FkDOPLineCollisionCheck structure for performing line checks
+	 * against a kDOPTree. Initializes all of the variables that are used
+	 * throughout the line check.
+	 *
+	 * @param InStart -- The starting point of the trace
+	 * @param InEnd -- The ending point of the trace
+	 * @param InbFindClosestIntersection -- Whether to stop at the first hit or not
+	 * @param InCollDataProvider -- The struct that provides access to mesh/primitive
+	 *		specific data, such as L2W, W2L, Vertices, and so on
+	 * @param InResult -- The out param for hit result information
+	 */
+	TkDOPLineCollisionCheck(const FVector4& InStart,const FVector4& InEnd,
+		bool bInbFindClosestIntersection,
+		const COLL_DATA_PROVIDER& InCollDataProvider,
+		FkHitResult* InResult) 
+		:
+		TkDOPCollisionCheck<COLL_DATA_PROVIDER,KDOP_IDX_TYPE>(InCollDataProvider),
+		Result(InResult), 
+		Start(InStart), 
+		End(InEnd), 
+		bFindClosestIntersection(bInbFindClosestIntersection), 
+		HitNodeIndex(0xFFFFFFFF)
+	{
+		const FMatrix& WorldToLocal = TkDOPCollisionCheck<COLL_DATA_PROVIDER,KDOP_IDX_TYPE>::CollDataProvider.GetWorldToLocal();
+		// Move start and end to local space
+		LocalStart = WorldToLocal.TransformPosition(Start);
+		LocalEnd = WorldToLocal.TransformPosition(End);
+		// Calculate the vector's direction in local space
+		LocalDir = LocalEnd - LocalStart;
+		// Build the one over dir
+		LocalOneOverDir.X = LocalDir.X ? 1.f / LocalDir.X : MAX_FLT;
+		LocalOneOverDir.Y = LocalDir.Y ? 1.f / LocalDir.Y : MAX_FLT;
+		LocalOneOverDir.Z = LocalDir.Z ? 1.f / LocalDir.Z : MAX_FLT;
+
+		// Construct the SOA data
+		StartSOA.X = VectorLoadFloat1( &LocalStart.X );
+		StartSOA.Y = VectorLoadFloat1( &LocalStart.Y );
+		StartSOA.Z = VectorLoadFloat1( &LocalStart.Z );
+		EndSOA.X = VectorLoadFloat1( &LocalEnd.X );
+		EndSOA.Y = VectorLoadFloat1( &LocalEnd.Y );
+		EndSOA.Z = VectorLoadFloat1( &LocalEnd.Z );
+		DirSOA.X = VectorLoadFloat1( &LocalDir.X );
+		DirSOA.Y = VectorLoadFloat1( &LocalDir.Y );
+		DirSOA.Z = VectorLoadFloat1( &LocalDir.Z );
+	}
+
+	/**
+	 * Transforms the local hit normal into a world space normal using the transpose
+	 * adjoint and flips the normal if need be
+	 */
+	FORCEINLINE FVector4 GetHitNormal(void)
+	{
+		// Transform the hit back into world space using the transpose adjoint
+		FVector4 Normal = TkDOPCollisionCheck<COLL_DATA_PROVIDER,KDOP_IDX_TYPE>::CollDataProvider.GetLocalToWorldTransposeAdjoint().TransformVector(LocalHitNormal).SafeNormal();
+		// Flip the normal if the triangle is inverted
+		if (TkDOPCollisionCheck<COLL_DATA_PROVIDER,KDOP_IDX_TYPE>::CollDataProvider.GetDeterminant() < 0.f)
+		{
+			Normal = -Normal;
+		}
+		return Normal;
+	}
+};
\ No newline at end of file
diff --git a/Engine/Source/Runtime/RenderCore/Private/RenderResource.cpp b/Engine/Source/Runtime/RenderCore/Private/RenderResource.cpp
--- a/Engine/Source/Runtime/RenderCore/Private/RenderResource.cpp
+++ b/Engine/Source/Runtime/RenderCore/Private/RenderResource.cpp
@@ -67,7 +67,7 @@ FRenderResource::~FRenderResource()
 	if (bInitialized && !GIsCriticalError)
 	{
 		// Deleting an initialized FRenderResource will result in a crash later since it is still linked
-		UE_LOG(LogRendererCore, Fatal,TEXT("An FRenderResource was deleted without being released first!"));
+		UE_LOG(LogRendererCore, Fatal,TEXT("A FRenderResource was deleted without being released first!"));
 	}
 }
 
diff --git a/Engine/Source/Runtime/RenderCore/Private/RenderUtils.cpp b/Engine/Source/Runtime/RenderCore/Private/RenderUtils.cpp
--- a/Engine/Source/Runtime/RenderCore/Private/RenderUtils.cpp
+++ b/Engine/Source/Runtime/RenderCore/Private/RenderUtils.cpp
@@ -41,8 +41,6 @@ FArchive& operator<<(FArchive& Ar,FPackedNormal& N)
 //	Pixel format information.
 //
 
-// NOTE: If you add a new basic texture format (ie a format that could be cooked - currently PF_A32B32G32R32F through
-// PF_UYVY) you MUST also update XeTools.cpp and PS3Tools.cpp to match up!
 FPixelFormatInfo	GPixelFormats[PF_MAX] =
 {
 	// Name						BlockSizeX	BlockSizeY	BlockSizeZ	BlockBytes	NumComponents	PlatformFormat	Supported		UnrealFormat
@@ -99,6 +97,8 @@ FPixelFormatInfo	GPixelFormats[PF_MAX] =
 	{ TEXT("ETC1"),				4,			4,			1,			8,			3,				0,				0,				PF_ETC1				},
 	{ TEXT("ETC2_RGB"),			4,			4,			1,			8,			3,				0,				0,				PF_ETC2_RGB			},
 	{ TEXT("ETC2_RGBA"),		4,			4,			1,			16,			4,				0,				0,				PF_ETC2_RGBA		},
+	{ TEXT("PF_R32G32B32A32_UINT"),1,		1,			1,			16,			4,				0,				1,				PF_R32G32B32A32_UINT},
+	{ TEXT("PF_R16G16_UINT"),	1,			1,			1,			4,			4,				0,				1,				PF_R16G16_UINT},
 };
 
 static struct FValidatePixelFormats
@@ -770,4 +770,4 @@ RENDERCORE_API bool IsSimpleDynamicLightingEnabled()
 {
 	static const auto CVar = IConsoleManager::Get().FindTConsoleVariableDataInt(TEXT("r.SimpleDynamicLighting"));
 	return (CVar->GetValueOnAnyThread() != 0);
-}
\ No newline at end of file
+}
diff --git a/Engine/Source/Runtime/Renderer/Private/BasePassRendering.cpp b/Engine/Source/Runtime/Renderer/Private/BasePassRendering.cpp
--- a/Engine/Source/Runtime/Renderer/Private/BasePassRendering.cpp
+++ b/Engine/Source/Runtime/Renderer/Private/BasePassRendering.cpp
@@ -182,7 +182,7 @@ public:
 				LightMapPolicy,
 				Parameters.BlendMode,
 				Parameters.TextureMode,
-				Parameters.ShadingModel != MSM_Unlit && Scene && Scene->SkyLight && !Scene->SkyLight->bHasStaticLighting,
+				Parameters.ShadingModel != MSM_Unlit && Scene && Scene->SkyLight && Scene->SkyLight->bWantsStaticShadowing && !Scene->SkyLight->bHasStaticLighting,
 				IsTranslucentBlendMode(Parameters.BlendMode) && Scene->HasAtmosphericFog()
 				),
 			Scene->GetFeatureLevel()
@@ -281,7 +281,7 @@ public:
 			LightMapPolicy,
 			Parameters.BlendMode,
 			Parameters.TextureMode,
-			Scene && Scene->SkyLight && !Scene->SkyLight->bHasStaticLighting && bIsLitMaterial,
+			Scene && Scene->SkyLight && !Scene->SkyLight->bHasStaticLighting && Scene->SkyLight->bWantsStaticShadowing && bIsLitMaterial,
 			IsTranslucentBlendMode(Parameters.BlendMode) && (Scene && Scene->HasAtmosphericFog()) && View.Family->EngineShowFlags.Atmosphere,
 			View.Family->EngineShowFlags.ShaderComplexity,
 			false,
@@ -589,4 +589,4 @@ void FSelfShadowedCachedPointIndirectLightingPolicy::SetMesh(
 		VertexFactory,
 		MaterialRenderProxy,
 		ElementData);
-}
\ No newline at end of file
+}
diff --git a/Engine/Source/Runtime/Renderer/Private/DeferredShadingRenderer.cpp b/Engine/Source/Runtime/Renderer/Private/DeferredShadingRenderer.cpp
--- a/Engine/Source/Runtime/Renderer/Private/DeferredShadingRenderer.cpp
+++ b/Engine/Source/Runtime/Renderer/Private/DeferredShadingRenderer.cpp
@@ -638,7 +638,9 @@ void FDeferredShadingSceneRenderer::Render()
 			SCOPED_CONDITIONAL_DRAW_EVENTF(EventView,Views.Num() > 1, DEC_SCENE_ITEMS, TEXT("View%d"), ViewIndex);
 			GCompositionLighting.ProcessAfterBasePass(Views[ViewIndex]);
 		}
-		
+
+		RenderDynamicSkyLighting();
+
 		// Clear the translucent lighting volumes before we accumulate
 		ClearTranslucentVolumeLighting();
 
@@ -739,6 +741,12 @@ void FDeferredShadingSceneRenderer::Render()
 		RenderLightShaftBloom();
 	}
 
+	if (ViewFamily.EngineShowFlags.VisualizeDistanceFieldAO)
+	{
+		FSceneRenderTargetItem DummyOutput(NULL, NULL, NULL);
+		RenderDistanceFieldAOSurfaceCache(DummyOutput, true);
+	}
+
 	// Resolve the scene color for post processing.
 	GSceneRenderTargets.ResolveSceneColor(FResolveRect(0, 0, ViewFamily.FamilySizeX, ViewFamily.FamilySizeY));
 
diff --git a/Engine/Source/Runtime/Renderer/Private/DeferredShadingRenderer.h b/Engine/Source/Runtime/Renderer/Private/DeferredShadingRenderer.h
--- a/Engine/Source/Runtime/Renderer/Private/DeferredShadingRenderer.h
+++ b/Engine/Source/Runtime/Renderer/Private/DeferredShadingRenderer.h
@@ -148,7 +148,7 @@ private:
 	void SetupInteractionShadows(
 		FLightPrimitiveInteraction* Interaction, 
 		FVisibleLightInfo& VisibleLightInfo, 
-		bool bReflectionCaptureScene,
+		bool bStaticSceneOnly,
 		const TArray<FProjectedShadowInfo*,SceneRenderingAllocator>& ViewDependentWholeSceneShadows,
 		TArray<FProjectedShadowInfo*,SceneRenderingAllocator>& PreShadows);
 
@@ -176,6 +176,12 @@ private:
 	/** Renders reflections that can be done in a deferred pass. */
 	void RenderDeferredReflections();
 
+	/** Render dynamic sky lighting from Movable sky lights. */
+	void RenderDynamicSkyLighting();
+
+	/** Render Ambient Occlusion using mesh distance fields and the surface cache, which supports dynamic rigid meshes. */
+	bool RenderDistanceFieldAOSurfaceCache(FSceneRenderTargetItem& OutBentNormalAO, bool bApplyToSceneColor);
+
 	/** Whether tiled deferred is supported and can be used at all. */
 	bool CanUseTiledDeferred() const;
 
diff --git a/Engine/Source/Runtime/Renderer/Private/DistanceFieldLightingShared.h b/Engine/Source/Runtime/Renderer/Private/DistanceFieldLightingShared.h
new file mode 100644
--- /dev/null
+++ b/Engine/Source/Runtime/Renderer/Private/DistanceFieldLightingShared.h
@@ -0,0 +1,201 @@
+// Copyright 1998-2013 Epic Games, Inc. All Rights Reserved.
+
+/*=============================================================================
+	
+=============================================================================*/
+
+#pragma once
+
+#include "StaticMeshResources.h"
+
+/** Maximum object / tile interactions to allocate memory for. */
+static const int32 GMaxTileShadingJobs = 500000;
+
+/** Tile sized used for most AO compute shaders. */
+const int32 GDistanceFieldAOTileSizeX = 8;
+const int32 GDistanceFieldAOTileSizeY = 8;
+/** Base downsample factor that all distance field AO operations are done at. */
+const int32 GAODownsampleFactor = 2;
+/** Number of cone traced directions. */
+const int32 NumConeSampleDirections = 9;
+
+/** Generates unit length, stratified and uniformly distributed direction samples in a hemisphere. */
+extern void GenerateStratifiedUniformHemisphereSamples2(int32 NumThetaSteps, int32 NumPhiSteps, FRandomStream& RandomStream, TArray<FVector4>& Samples);
+
+// Optimized with GA, Best =  -2.764536 after 117591 Steps
+const FVector RandomVectors9[] = 
+{
+	FVector(-0.134453, 0.570145, 0.810467),
+	FVector(0.600695, 0.732059, 0.321333),
+	FVector(0.591856, -0.692181, 0.413028),
+	FVector(0.065080, -0.293309, 0.953800),
+	FVector(-0.952426, -0.251775, 0.171737),
+	FVector(0.047865, -0.998847, 0.003783),
+	FVector(-0.476063, -0.743138, 0.470224),
+	FVector(-0.837584, 0.480720, 0.259541),
+	FVector(0.780739, 0.035320, 0.623858)
+};
+
+// Sample set restricted to not self-intersect a surface based on cone angle .475882232
+// Coverage of hemisphere = 0.755312979
+const FVector SpacedVectors9[] = 
+{
+	FVector(-0.573257625, 0.625250816, 0.529563010),
+	FVector(0.253354192, -0.840093017, 0.479640961),
+	FVector(-0.421664953, -0.718063235, 0.553700149),
+	FVector(0.249163717, 0.796005428, 0.551627457),
+	FVector(0.375082791, 0.295851320, 0.878512800),
+	FVector(-0.217619032, 0.00193520682, 0.976031899),
+	FVector(-0.852834642, 0.0111727007, 0.522061586),
+	FVector(0.745701790, 0.239393353, 0.621787369),
+	FVector(-0.151036426, -0.465937436, 0.871831656)
+};
+
+BEGIN_UNIFORM_BUFFER_STRUCT(FAOSampleData2,)
+	DECLARE_UNIFORM_BUFFER_STRUCT_MEMBER_ARRAY(FVector4,SampleDirections,[NumConeSampleDirections])
+END_UNIFORM_BUFFER_STRUCT(FAOSampleData2)
+
+inline float GetAOMaxDistance()
+{
+	extern float GAOLargestSampleOffset;
+	extern float GAOConeHalfAngle;
+	float AOMaxDistanceValue = GAOLargestSampleOffset * (1 + FMath::Tan(GAOConeHalfAngle));
+	return AOMaxDistanceValue;
+}
+
+class FAOParameters
+{
+public:
+	void Bind(const FShaderParameterMap& ParameterMap)
+	{
+		AOMaxDistance.Bind(ParameterMap,TEXT("AOMaxDistance"));
+		AOStepScale.Bind(ParameterMap,TEXT("AOStepScale"));
+		AOStepExponentScale.Bind(ParameterMap,TEXT("AOStepExponentScale"));
+		AOMaxViewDistance.Bind(ParameterMap,TEXT("AOMaxViewDistance"));
+		DistanceFieldTexture.Bind(ParameterMap, TEXT("DistanceFieldTexture"));
+		DistanceFieldSampler.Bind(ParameterMap, TEXT("DistanceFieldSampler"));
+		DistanceFieldAtlasTexelSize.Bind(ParameterMap, TEXT("DistanceFieldAtlasTexelSize"));
+	}
+
+	friend FArchive& operator<<(FArchive& Ar,FAOParameters& Parameters)
+	{
+		Ar << Parameters.AOMaxDistance;
+		Ar << Parameters.AOStepScale;
+		Ar << Parameters.AOStepExponentScale;
+		Ar << Parameters.AOMaxViewDistance;
+		Ar << Parameters.DistanceFieldTexture;
+		Ar << Parameters.DistanceFieldSampler;
+		Ar << Parameters.DistanceFieldAtlasTexelSize;
+		return Ar;
+	}
+
+	template<typename ShaderRHIParamRef>
+	void Set(const ShaderRHIParamRef ShaderRHI)
+	{
+		float AOMaxDistanceValue = GetAOMaxDistance();
+		SetShaderValue(ShaderRHI, AOMaxDistance, AOMaxDistanceValue);
+
+		extern float GAOLargestSampleOffset;
+		extern float GAOStepExponentScale;
+		extern uint32 GAONumConeSteps;
+		float AOStepScaleValue = GAOLargestSampleOffset / FMath::Pow(2.0f, GAOStepExponentScale * (GAONumConeSteps - 1));
+		SetShaderValue(ShaderRHI, AOStepScale, AOStepScaleValue);
+
+		SetShaderValue(ShaderRHI, AOStepExponentScale, GAOStepExponentScale);
+
+		extern float GAOMaxViewDistance;
+		SetShaderValue(ShaderRHI, AOMaxViewDistance, GAOMaxViewDistance);
+
+		SetTextureParameter(
+			ShaderRHI,
+			DistanceFieldTexture,
+			DistanceFieldSampler,
+			TStaticSamplerState<SF_Bilinear,AM_Clamp,AM_Clamp,AM_Clamp>::GetRHI(),
+			GDistanceFieldVolumeTextureAtlas.VolumeTextureRHI
+			);
+
+		const int32 NumTexelsOneDimX = GDistanceFieldVolumeTextureAtlas.GetSizeX();
+		const int32 NumTexelsOneDimY = GDistanceFieldVolumeTextureAtlas.GetSizeY();
+		const int32 NumTexelsOneDimZ = GDistanceFieldVolumeTextureAtlas.GetSizeZ();
+		const FVector InvTextureDim(1.0f / NumTexelsOneDimX, 1.0f / NumTexelsOneDimY, 1.0f / NumTexelsOneDimZ);
+		SetShaderValue(ShaderRHI, DistanceFieldAtlasTexelSize, InvTextureDim);
+	}
+
+private:
+	FShaderParameter AOMaxDistance;
+	FShaderParameter AOStepScale;
+	FShaderParameter AOStepExponentScale;
+	FShaderParameter AOMaxViewDistance;
+	FShaderResourceParameter DistanceFieldTexture;
+	FShaderResourceParameter DistanceFieldSampler;
+	FShaderParameter DistanceFieldAtlasTexelSize;
+};
+
+class FQuadVertexBuffer : public FVertexBuffer
+{
+public:
+
+	virtual void InitRHI()
+	{
+		// Used as a non-indexed triangle list, so 6 vertices per quad
+		const uint32 Size = 6 * sizeof(FScreenVertex);
+		FRHIResourceCreateInfo CreateInfo;
+		VertexBufferRHI = RHICreateVertexBuffer(Size, BUF_Static, CreateInfo);
+		void* Buffer = RHILockVertexBuffer(VertexBufferRHI, 0, Size, RLM_WriteOnly);
+		FScreenVertex* DestVertex = (FScreenVertex*)Buffer;
+
+		DestVertex[0].Position = FVector2D(0, 0);
+		DestVertex[0].UV = FVector2D(1, 1);
+		DestVertex[1].Position = FVector2D(0, 0);
+		DestVertex[1].UV = FVector2D(1, 0);
+		DestVertex[2].Position = FVector2D(0, 0);
+		DestVertex[2].UV = FVector2D(0, 1);
+		DestVertex[3].Position = FVector2D(0, 0);
+		DestVertex[3].UV = FVector2D(1, 0);
+		DestVertex[4].Position = FVector2D(0, 0);
+		DestVertex[4].UV = FVector2D(0, 0);
+		DestVertex[5].Position = FVector2D(0, 0);
+		DestVertex[5].UV = FVector2D(0, 1);
+
+		RHIUnlockVertexBuffer(VertexBufferRHI);      
+	}
+};
+
+extern TGlobalResource<FQuadVertexBuffer> GQuadVertexBuffer;
+
+class FCircleVertexBuffer : public FVertexBuffer
+{
+public:
+
+	virtual void InitRHI()
+	{
+		int32 NumSections = 8;
+
+		// Used as a non-indexed triangle list, so 3 vertices per triangle
+		const uint32 Size = 3 * NumSections * sizeof(FScreenVertex);
+		FRHIResourceCreateInfo CreateInfo;
+		VertexBufferRHI = RHICreateVertexBuffer(Size, BUF_Static, CreateInfo);
+		void* Buffer = RHILockVertexBuffer(VertexBufferRHI, 0, Size, RLM_WriteOnly);
+		FScreenVertex* DestVertex = (FScreenVertex*)Buffer;
+
+		for (int32 SectionIndex = 0; SectionIndex < NumSections; SectionIndex++)
+		{
+			float Fraction = SectionIndex / (float)NumSections;
+			float CurrentAngle = Fraction * 2 * PI;
+			float NextAngle = ((SectionIndex + 1) / (float)NumSections) * 2 * PI;
+			FVector2D CurrentPosition(FMath::Cos(CurrentAngle), FMath::Sin(CurrentAngle));
+			FVector2D NextPosition(FMath::Cos(NextAngle), FMath::Sin(NextAngle));
+
+			DestVertex[SectionIndex * 3 + 0].Position = FVector2D(0, 0);
+			DestVertex[SectionIndex * 3 + 0].UV = CurrentPosition;
+			DestVertex[SectionIndex * 3 + 1].Position = FVector2D(0, 0);
+			DestVertex[SectionIndex * 3 + 1].UV = NextPosition;
+			DestVertex[SectionIndex * 3 + 2].Position = FVector2D(0, 0);
+			DestVertex[SectionIndex * 3 + 2].UV = FVector2D(.5f, .5f);
+		}
+
+		RHIUnlockVertexBuffer(VertexBufferRHI);      
+	}
+};
+
+extern TGlobalResource<FCircleVertexBuffer> GCircleVertexBuffer;
