Fixed AISystem not being created for no-mcp servers #UE4

TTP#337533

diff --git a/Engine/Source/Runtime/Engine/Private/UnrealEngine.cpp b/Engine/Source/Runtime/Engine/Private/UnrealEngine.cpp
--- a/Engine/Source/Runtime/Engine/Private/UnrealEngine.cpp
+++ b/Engine/Source/Runtime/Engine/Private/UnrealEngine.cpp
@@ -8435,6 +8435,9 @@ bool UEngine::LoadMap( FWorldContext& WorldContext, FURL URL, class UPendingNetG
 	}
 	
 	UNavigationSystem::InitializeForWorld(WorldContext.World(), FNavigationSystem::GameMode);
+	
+	// Note that AI system will be created only if ai-system-creation conditions are met
+	WorldContext.World()->CreateAISystem();
 
 	// Initialize gameplay for the level.
 	WorldContext.World()->InitializeActorsForPlay(URL);
