Fix (CL# 2093455) when WITH_EDITOR is 0

diff --git a/Engine/Source/Runtime/Engine/Private/ActorConstruction.cpp b/Engine/Source/Runtime/Engine/Private/ActorConstruction.cpp
--- a/Engine/Source/Runtime/Engine/Private/ActorConstruction.cpp
+++ b/Engine/Source/Runtime/Engine/Private/ActorConstruction.cpp
@@ -139,7 +139,9 @@ void AActor::RerunConstructionScripts()
 		GUndo = NULL;
 		
 		// Create cache to store component data across rerunning construction scripts
+#if WITH_EDITOR
 		FActorTransactionAnnotation* ActorTransactionAnnotation = CurrentTransactionAnnotation.Get();
+#endif
 		FComponentInstanceDataCache* InstanceDataCache;
 		
 		FTransform OldTransform = FTransform::Identity;
@@ -158,6 +160,7 @@ void AActor::RerunConstructionScripts()
 		// Save info about attached actors
 		TArray<FAttachedActorInfo> AttachedActorInfos;
 
+#if WITH_EDITOR
 		if (ActorTransactionAnnotation)
 		{
 			InstanceDataCache = &ActorTransactionAnnotation->ComponentInstanceData;
@@ -190,6 +193,7 @@ void AActor::RerunConstructionScripts()
 			}
 		}
 		else
+#endif
 		{
 			InstanceDataCache = new FComponentInstanceDataCache(this);
 
@@ -285,11 +289,13 @@ void AActor::RerunConstructionScripts()
 		// Restore the undo buffer
 		GUndo = CurrentTransaction;
 
+#if WITH_EDITOR
 		if (ActorTransactionAnnotation)
 		{
 			CurrentTransactionAnnotation = NULL;
 		}
 		else
+#endif
 		{
 			delete InstanceDataCache;
 		}
