Packagemap refactor

* We no longer use a package map for server side network guid management
* MasterMap is now just direct to a stored GuidCache, which is then referenced by all PackageMapClients

diff --git a/Engine/Source/Runtime/CoreUObject/Private/UObject/CoreNet.cpp b/Engine/Source/Runtime/CoreUObject/Private/UObject/CoreNet.cpp
--- a/Engine/Source/Runtime/CoreUObject/Private/UObject/CoreNet.cpp
+++ b/Engine/Source/Runtime/CoreUObject/Private/UObject/CoreNet.cpp
@@ -106,10 +106,7 @@ bool UPackageMap::SerializeName(FArchive& Ar, FName& Name)
 	}
 	return true;
 }
-bool UPackageMap::CanSerializeObject(const UObject* Obj)
-{
-	return SupportsObject(Obj);
-}
+
 bool UPackageMap::SerializeObject( FArchive& Ar, UClass* Class, UObject*& Object, FNetworkGUID *OutNetGUID )
 {
 	// For now this is unsupported, but we could implement this serialization in a way that is compatible with UPackageMapClient,
@@ -124,10 +121,21 @@ bool UPackageMap::WriteObject( FArchive& Ar, UObject* Outer, FNetworkGUID NetGUI
 	return true;
 }
 
+void UPackageMap::LogDebugInfo( FOutputDevice & Ar)
+{
+	UE_LOG(LogCoreNet, Fatal,TEXT("Unexpected UPackageMap::LogDebugInfo"));
+}
+
+UObject * UPackageMap::GetObjectFromNetGUID( const FNetworkGUID & NetGUID )
+{
+	UE_LOG(LogCoreNet, Fatal,TEXT("Unexpected UPackageMap::GetObjectFromNetGUID"));
+	return NULL;
+}
+
 FClassNetCache* UPackageMap::GetClassNetCache( UClass* Class )
 {
 	FClassNetCache* Result = ClassFieldIndices.FindRef(Class);
-	if( !Result )// && SupportsObject(Class) )
+	if( !Result )
 	{
 		Result                       = ClassFieldIndices.Add( Class, new FClassNetCache(Class) );
 		Result->Super                = NULL;
@@ -143,11 +151,8 @@ FClassNetCache* UPackageMap::GetClassNetCache( UClass* Class )
 		{
 			// Add sandboxed items to net cache.  
 			UField* Field = Class->NetFields[i];
-			//if( SupportsObject(Field ) )
-			{
-				int32 ThisIndex      = Result->GetMaxIndex();
-				new(Result->Fields)FFieldNetCache( Field, ThisIndex );
-			}
+			int32 ThisIndex	= Result->GetMaxIndex();
+			new(Result->Fields)FFieldNetCache( Field, ThisIndex );
 		}
 
 		Result->Fields.Shrink();
@@ -168,293 +173,17 @@ void UPackageMap::ClearClassNetCache()
 	ClassFieldIndices.Empty();
 }
 
-void UPackageMap::LogDebugInfo(FOutputDevice& Ar)
-{
-	for (auto It = Cache->ObjectLookup.CreateIterator(); It; ++It)
-	{
-		FNetworkGUID NetGUID = It.Key();
-		UObject *Obj = It.Value().Object.Get();
-		UE_LOG(LogCoreNet, Log, TEXT("%s - %s"), *NetGUID.ToString(), Obj ? *Obj->GetPathName() : TEXT("NULL"));
-	}
-}
-
-bool UPackageMap::SupportsObject( const UObject* Object )
-{
-	/** In this PackageMap implementation, all Objects are supported.*/
-	return true;
-}
-bool UPackageMap::SupportsPackage( UObject* InOuter )
-{
-	/** In this PackageMap implementation, all Packages are supported.*/
-	return true;
-}
-
 void UPackageMap::Serialize( FArchive& Ar )
 {
 	Super::Serialize( Ar );
-	Ar << Cache;
-
-	if ( Ar.IsCountingMemory() && Cache )
-	{
-		Cache->ObjectLookup.CountBytes( Ar );
-		Cache->NetGUIDLookup.CountBytes( Ar );
-	}
 }
 
 void UPackageMap::AddReferencedObjects(UObject* InThis, FReferenceCollector& Collector)
 {
 	UPackageMap * This = CastChecked<UPackageMap>( InThis );
-
-	if ( This->Cache != NULL )
-	{
-		Collector.AddReferencedObject( This->Cache, This );		
-	}
-
 	return Super::AddReferencedObjects( InThis, Collector );
 }
 
-FNetworkGUID UPackageMap::GetObjectNetGUID(const UObject *Object)
-{
-	FNetworkGUID NetGUID;
-	if (!Object || !SupportsObject(Object))
-	{
-		return NetGUID;
-	}
-	
-	NetGUID = Cache->NetGUIDLookup.FindRef(Object);
-	if (!NetGUID.IsValid())
-	{
-		NetGUID = AssignNewNetGUID(Object);
-	}
-	return NetGUID;
-}
-
-UObject * UPackageMap::GetObjectFromNetGUID( FNetworkGUID NetGUID )
-{
-	if ( !ensure( NetGUID.IsValid() ) )
-	{
-		return NULL;
-	}
-
-	if ( !ensure( !NetGUID.IsDefault() ) )
-	{
-		return NULL;
-	}
-
-	FNetGuidCacheObject * CacheObjectPtr = Cache->ObjectLookup.Find( NetGUID );
-
-	if ( CacheObjectPtr == NULL )
-	{
-		// We don't have the object mapped yet
-		return NULL;
-	}
-
-	UObject * Object = CacheObjectPtr->Object.Get();
-
-	if ( Object != NULL )
-	{
-		check( Object->GetPathName() == CacheObjectPtr->FullPath );
-		return Object;
-	}
-
-	if ( CacheObjectPtr->FullPath.IsEmpty() )
-	{
-		// This probably isn't possible, but making it a warning
-		UE_LOG( LogNetPackageMap, Warning, TEXT( "GetObjectFromNetGUID: No full path for %s" ), *NetGUID.ToString() );
-		return NULL;
-	}
-
-	if ( NetGUID.IsDynamic() )
-	{
-		// Dynamic objects don't have stable names, so we can't possibly reload the same object
-		UE_LOG( LogNetPackageMap, VeryVerbose, TEXT( "GetObjectFromNetGUID: Cannot re-create dynamic object after GC <%s, %s>" ), *NetGUID.ToString(), *CacheObjectPtr->FullPath );
-		return NULL;
-	}
-
-	//
-	// The object was previously mapped, but we GC'd it
-	// We need to reload it now
-	//
-
-	Object = StaticFindObject( UObject::StaticClass(), NULL, *CacheObjectPtr->FullPath, false );
-
-	if ( Object == NULL )
-	{
-		if ( IsNetGUIDAuthority() )
-		{
-			// The authority shouldn't be loading resources on demand, at least for now.
-			// This could be garbage or a security risk
-			// Another possibility is in dynamic property replication if the server reads the previously serialized state
-			// that has a now destroyed actor in it.
-			UE_LOG( LogNetPackageMap, Warning, TEXT( "GetObjectFromNetGUID: Could not find Object for: NetGUID <%s, %s> (and IsNetGUIDAuthority())" ), *NetGUID.ToString(), *CacheObjectPtr->FullPath );
-			return NULL;
-		}
-		else
-		{
-			UE_LOG( LogNetPackageMap, Log, TEXT( "GetObjectFromNetGUID: Could not find Object for: NetGUID <%s, %s>"), *NetGUID.ToString(), *CacheObjectPtr->FullPath );
-		}
-
-		Object = StaticLoadObject( UObject::StaticClass(), NULL, *CacheObjectPtr->FullPath, NULL, LOAD_NoWarn );
-
-		if ( Object )
-		{
-			UE_LOG( LogNetPackageMap, Log, TEXT( "GetObjectFromNetGUID: StaticLoadObject. Found: %s" ), Object ? *Object->GetName() : TEXT( "NULL" ) );
-		}
-		else
-		{
-			Object = LoadPackage( NULL, *CacheObjectPtr->FullPath, LOAD_None );
-			UE_LOG( LogNetPackageMap, Log, TEXT( "GetObjectFromNetGUID: LoadPackage. Found: %s" ), Object ? *Object->GetName() : TEXT( "NULL" ) );
-		}
-	}
-	else
-	{
-		// If we actually found the object, we probably shouldn't have GC'd it, so that's odd
-		//UE_LOG( LogNetPackageMap, Warning, TEXT( "GetObjectFromNetGUID: Object should not be found after GC: <%s, %s>" ), *NetGUID.ToString(), *CacheObjectPtr->FullPath );
-	}
-
-	if ( Object == NULL )
-	{
-		// If we get here, that means we have an invalid path, which shouldn't be possible, but making it a warning
-		UE_LOG( LogNetPackageMap, Warning, TEXT( "GetObjectFromNetGUID: FAILED to re-create object after GC: <%s, %s>" ), *NetGUID.ToString(), *CacheObjectPtr->FullPath );
-	}
-	else
-	{
-		// Reassign the object pointer for quick access next time
-		CacheObjectPtr->Object = Object;		
-		// Make sure the object is in the GUIDToObjectLookup.
-		Cache->NetGUIDLookup.Add( Object, NetGUID );
-	}
-
-	return Object;
-}
-
-/**
- *	Generate a new NetGUID for this object and assign it.
- */
-FNetworkGUID UPackageMap::AssignNewNetGUID(const UObject* Object)
-{
-	if (!IsNetGUIDAuthority())
-	{
-		// We cannot make new NetGUIDs
-		return FNetworkGUID::GetDefault();
-	}
-
-	// Generate new NetGUID and assign it
-	int32 idx = 0;
-	if (!IsDynamicObject(Object))
-	{
-		idx = 1;
-	}
-
-	return AssignNetGUID(Object, (++Cache->UniqueNetIDs[idx] << 1) | idx);
-}
-
-/**
- *	Assigns the given UObject the give NetNetworkGUID.
- */
-FNetworkGUID UPackageMap::AssignNetGUID(const UObject* Object, FNetworkGUID NewNetworkGUID)
-{
-	if (NewNetworkGUID.IsDefault())
-	{
-		// This is the default NetGUID, look up or generate a real one
-		NewNetworkGUID = Cache->NetGUIDLookup.FindRef(Object);
-		if (!NewNetworkGUID.IsValid())
-		{
-			NewNetworkGUID = AssignNewNetGUID(Object);
-		}
-		HandleUnAssignedObject(Object);
-		return NewNetworkGUID;
-	}
-
-	UE_LOG(LogNetPackageMap, Log, TEXT("Assigning %s NetGUID <%s>"), Object ? *Object->GetName() : TEXT("NULL"), *NewNetworkGUID.ToString() );
-	
-	// Check for reassignment: this can happen during seamless travel, so leaving this off for now, but is useful for debugging purposes, so leaving it in.
-#if 0
-	{
-		FNetworkGUID OldNetGUID = Cache->NetGUIDLookup.FindRef(Object);
-		if (OldNetGUID.IsValid() && !OldNetGUID.IsDefault())
-		{
-			UE_LOG(LogNetPackageMap, Log, TEXT("Attempting to reassign NetGUID to %s (was previously <%s>. Now is <%s>"), Object ? *Object->GetName() : TEXT("NULL"), *OldNetGUID.ToString(),  *NewNetworkGUID.ToString() );
-		}
-
-		UObject *OldObject = Cache->ObjectLookup.FindRef(NewNetworkGUID).Get();
-		if (OldObject)
-		{
-			UE_LOG(LogNetPackageMap, Warning, TEXT("Attempting to reassign NetGUID <%s> to %s (It is already assigned to object %s"), *NewNetworkGUID.ToString(), Object ? *Object->GetName() : TEXT("NULL"), *OldNetGUID.ToString(), *OldObject->GetName() );
-		}
-	}
-#endif
-
-	// Actually assign
-	FNetGuidCacheObject CacheObject;
-
-	CacheObject.Object = Object;
-	CacheObject.FullPath = Object->GetPathName();
-
-	// Remove any existing entries
-	{
-		const FNetGuidCacheObject * ExistingCacheObjectPtr = Cache->ObjectLookup.Find( NewNetworkGUID );
-
-		if ( ExistingCacheObjectPtr )
-		{
-			const UObject * OldObject = ExistingCacheObjectPtr->Object.Get();
-
-			// If we find static guid in the cache, this means we didn't find it in GetObjectFromNetGUID
-			// This is a problem, since we should never get this far if we do find it there
-			if ( !NewNetworkGUID.IsDynamic() )
-			{
-				UE_LOG( LogNetPackageMap, Warning, TEXT( "AssignNetGUID: NewNetworkGUID was in ObjectLookup cache but is static: <%s> %s" ), *NewNetworkGUID.ToString(), Object ? *Object->GetPathName() : TEXT( "NULL" ) );
-			}
-
-			// If this net guid was found but the old object is NULL, this can happen due to:
-			//	1. Actor channel was closed locally (but we don't remove the net guid entry, since we can't know for sure if it will be referenced again)
-			//		a. Then when we re-create a channel, and assign this actor, we will find the old guid entry here
-			//	2. Dynamic object was locally GC'd, but then exported again from the server
-			//
-			// If this net guid was found and the objects match, we don't care. This can happen due to:
-			//	1. Same thing above can happen, but if we for some reason didn't destroy the actor/object we will see this case
-			//
-			// If the object pointers are different, this can be a problem, 
-			//	since this should only be possible if something gets out of sync during the net guid exchange code
-			if ( OldObject != NULL && OldObject != Object )
-			{
-				UE_LOG( LogNetPackageMap, Warning, TEXT( "Reassigning NetGUID <%s> to %s (was assigned to object %s)" ), *NewNetworkGUID.ToString(), Object ? *Object->GetPathName() : TEXT( "NULL" ), OldObject ? *OldObject->GetPathName() : TEXT( "NULL" ) );
-			}
-			else
-			{
-				UE_LOG( LogNetPackageMap, Verbose, TEXT( "Reassigning NetGUID <%s> to %s (was assigned to object %s)" ), *NewNetworkGUID.ToString(), Object ? *Object->GetPathName() : TEXT( "NULL" ), OldObject ? *OldObject->GetPathName() : TEXT( "NULL" ) );
-			}
-
-			Cache->NetGUIDLookup.Remove( ExistingCacheObjectPtr->Object );
-		}
-
-		const FNetworkGUID * ExistingNetworkGUIDPtr = Cache->NetGUIDLookup.Find( Object );
-
-		if ( ExistingNetworkGUIDPtr )
-		{
-			// This can happen when the server destroys a package/object and the client doesn't.
-			// The server will simply assign a new guid, and the client will assign to the old object that never got deleted locally
-			UE_LOG( LogNetPackageMap, Log, TEXT( "Changing NetGUID on object %s from <%s:%s> to <%s:%s>" ), Object ? *Object->GetPathName() : TEXT( "NULL" ), *ExistingNetworkGUIDPtr->ToString(), ExistingNetworkGUIDPtr->IsDynamic() ? TEXT("TRUE") : TEXT("FALSE"), *NewNetworkGUID.ToString(), NewNetworkGUID.IsDynamic() ? TEXT("TRUE") : TEXT("FALSE") );
-			Cache->ObjectLookup.Remove( *ExistingNetworkGUIDPtr );
-		}
-	}
-
-	Cache->ObjectLookup.Add( NewNetworkGUID, CacheObject );
-	Cache->NetGUIDLookup.Add( Object, NewNetworkGUID );
-
-#if !(UE_BUILD_SHIPPING || UE_BUILD_TEST)
-	Cache->History.Add(NewNetworkGUID, Object ? Object->GetPathName() : TEXT("NULL"));
-#endif
-
-	return NewNetworkGUID;
-}
-
-bool UPackageMap::IsDynamicObject(const UObject* Object)
-{
-	check(ObjectIsDynamicDelegate.IsBound());
-	return ObjectIsDynamicDelegate.Execute(Object);
-}
-
 void UPackageMap::FinishDestroy()
 {
 	ClearClassNetCache();
@@ -463,83 +192,7 @@ void UPackageMap::FinishDestroy()
 
 void UPackageMap::ResetPackageMap()
 {
-	Cache->ObjectLookup.Empty();
-	Cache->NetGUIDLookup.Empty();
-	Cache->UniqueNetIDs[0] = Cache->UniqueNetIDs[1] = 0;
-}
-
-class FArchiveCountMemGUID : public FArchive
-{
-public:
-	FArchiveCountMemGUID() : Mem( 0 ) { ArIsCountingMemory = true; }
-	void CountBytes( SIZE_T InNum, SIZE_T InMax ) { Mem += InMax; }
-	SIZE_T Mem;
-};
-
-void UPackageMap::CleanPackageMap()
-{
-	// NOTE - Quick fix for some package map issues where client was cleaning up guids, which was causing issues.
-	// For now, we no longer clean these up.  The client can sometimes cleanup guids that the server doesn't which causes issues.
-	// The guid memory will leak for now until we find a better way.
-
-	// Free any dynamic guids that refer to objects that are dead
-	for ( auto It = Cache->ObjectLookup.CreateIterator(); It; ++It )
-	{
-		if ( It.Value().Object.IsValid() )
-		{
-			// Make sure the path matches (which can change during seamless travel)
-			It.Value().FullPath = It.Value().Object->GetPathName();
-			continue;
-		}
-
-		if ( It.Key().IsDynamic() )
-		{
-			UE_LOG( LogNetPackageMap, Log, TEXT( "Cleaning reference to NetGUID: <%s> (%s) (ObjectLookup)" ), *It.Key().ToString(), *It.Value().FullPath );
-			It.RemoveCurrent();
-		}
-	}
-	
-
-	for ( auto It = Cache->NetGUIDLookup.CreateIterator(); It; ++It )
-	{
-		if ( !It.Key().IsValid() )
-		{
-			UE_LOG( LogNetPackageMap, Log, TEXT( "Cleaning reference to NetGUID: <%s> (NetGUIDLookup)" ), *It.Value().ToString() );
-			It.RemoveCurrent();
-		}
-	}
-
-	// Sanity check
-	for ( auto It = Cache->ObjectLookup.CreateIterator(); It; ++It )
-	{
-		check( It.Value().Object.IsValid() || It.Key().IsStatic() );
-
-		if ( It.Value().Object.IsValid() )
-		{
-			checkf( !It.Value().Object.IsValid() || Cache->NetGUIDLookup.FindRef( It.Value().Object ) == It.Key(), TEXT("Failed to validate ObjectLookup map in UPackageMap. Object '%s' was not in the NetGUIDLookup map with with value '%s'."), *It.Value().Object.Get()->GetPathName(), *It.Key().ToString());
-		}
-		else
-		{
-			for ( auto It2 = Cache->NetGUIDLookup.CreateIterator(); It2; ++It2 )
-			{
-				checkf( It2.Value() != It.Key(), TEXT("Failed to validate ObjectLookup map in UPackageMap. Object '%s' was in the NetGUIDLookup map with with value '%s'."), *It.Value().Object.Get()->GetPathName(), *It.Key().ToString());
-			}
-		}
-	}
-
-	for ( auto It = Cache->NetGUIDLookup.CreateIterator(); It; ++It )
-	{
-		check( It.Key().IsValid() );
-		checkf( Cache->ObjectLookup.FindRef( It.Value() ).Object == It.Key(), TEXT("Failed to validate NetGUIDLookup map in UPackageMap. GUID '%s' was not in the ObjectLookup map with with object '%s'."), *It.Value().ToString(), *It.Key().Get()->GetPathName());
-	}
-
-	FArchiveCountMemGUID CountBytesAr;
-
-	Cache->ObjectLookup.CountBytes( CountBytesAr );
-	Cache->NetGUIDLookup.CountBytes( CountBytesAr );
-
-	// Make this a warning to be a constant reminder that we need to ultimately free this memory when we find a solution
-	UE_LOG( LogNetPackageMap, Warning, TEXT( "CleanPackageMap ObjectLookup: %i, NetGUIDLookup: %i, Mem: %i kB" ), Cache->ObjectLookup.Num(), Cache->NetGUIDLookup.Num(), ( CountBytesAr.Mem / 1024 ) );
+	UE_LOG(LogCoreNet, Fatal,TEXT("Unexpected UPackageMap::ResetPackageMap"));
 }
 
 void UPackageMap::PostInitProperties()
@@ -547,17 +200,6 @@ void UPackageMap::PostInitProperties()
 	Super::PostInitProperties();
 	bShouldSerializeUnAckedObjects = true;
 	bSerializedUnAckedObject = false;
-
-	// Create the NetGUIDCache now if we don't have one. The pattern should be that the MasterMap (a plain UPackageMap) will always
-	// create one here, but UPackageMapClient and other subclasses will take thier Cache in as a parameter so that is it shared 
-	// (e.g., one 1 UNetGUIDCache exists on the server and is referenced be each client connection's package map).
-	// 
-	// This also keeps subclasses from <i>requiring</i> a master map to pass in their UNetGUIDCache; if we don't pass one in on construction
-	// they can use their own and still work. This may be useful for 'sideways' connections or other backend communication
-	if ( !Cache && !HasAnyFlags(RF_ClassDefaultObject) )
-	{
-		Cache = new UNetGUIDCache (FPostConstructInitializeProperties());
-	}
 }
 
 IMPLEMENT_CORE_INTRINSIC_CLASS(UPackageMap, UObject,
@@ -567,19 +209,6 @@ IMPLEMENT_CORE_INTRINSIC_CLASS(UPackageMap, UObject,
 
 // ----------------------------------------------------------------
 
-void UNetGUIDCache::PostInitProperties()
-{
-	Super::PostInitProperties();
-	UniqueNetIDs[0] = UniqueNetIDs[1] = 0;
-}
-
-IMPLEMENT_CORE_INTRINSIC_CLASS(UNetGUIDCache, UObject,
-	{
-	}
-);
-
-// ----------------------------------------------------------------
-
 void SerializeChecksum(FArchive &Ar, uint32 x, bool ErrorOK)
 {
 	if (Ar.IsLoading() )
diff --git a/Engine/Source/Runtime/CoreUObject/Private/UObject/Linker.cpp b/Engine/Source/Runtime/CoreUObject/Private/UObject/Linker.cpp
--- a/Engine/Source/Runtime/CoreUObject/Private/UObject/Linker.cpp
+++ b/Engine/Source/Runtime/CoreUObject/Private/UObject/Linker.cpp
@@ -680,6 +680,7 @@ ULinkerLoad* GetPackageLinker
 			}
 		}
 
+#if 0
 		// Make sure the package is accessible in the sandbox.
 		if( Sandbox && !Sandbox->SupportsPackage(InOuter) )
 		{
@@ -689,6 +690,7 @@ ULinkerLoad* GetPackageLinker
 			LogGetPackageLinkerError(InLongPackageName, FText::Format(LOCTEXT("Sandbox", "Asset '{AssetName}' is not accessible in this sandbox"), Arguments), LOCTEXT("SandboxShort", "Asset is not accessible in this sandbox"), InOuter, LoadFlags);
 			return NULL;
 		}
+#endif
 
 		// Create new linker.
 		if( !Result )
diff --git a/Engine/Source/Runtime/CoreUObject/Public/UObject/CoreNet.h b/Engine/Source/Runtime/CoreUObject/Public/UObject/CoreNet.h
--- a/Engine/Source/Runtime/CoreUObject/Public/UObject/CoreNet.h
+++ b/Engine/Source/Runtime/CoreUObject/Public/UObject/CoreNet.h
@@ -99,33 +99,6 @@ public:
  */
 #define MAX_OBJECT_INDEX (uint32(1) << 31)
 
-/** Stores an object with path associated with FNetworkGUID */
-class FNetGuidCacheObject
-{
-public:
-	TWeakObjectPtr< UObject >		Object;
-	FString							FullPath;
-};
-
-class COREUOBJECT_API UNetGUIDCache : public UObject
-{
-	DECLARE_CLASS_INTRINSIC(UNetGUIDCache,UObject,CLASS_Transient|0,CoreUObject);
-
-	virtual void PostInitProperties();
-
-	// Fixme: use a single data structure for this
-	TMap<FNetworkGUID, FNetGuidCacheObject >		ObjectLookup;
-	TMap<TWeakObjectPtr<UObject>, FNetworkGUID>		NetGUIDLookup;
-
-	int32	UniqueNetIDs[2];
-
-#if !(UE_BUILD_SHIPPING || UE_BUILD_TEST)
-	// History for debugging entries in the packagemap
-	TMap<FNetworkGUID, FString>	History;
-#endif
-};
-
-
 //
 // Maps objects and names to and from indices for network communication.
 //
@@ -133,12 +106,6 @@ class COREUOBJECT_API UPackageMap : public UObject
 {
 	DECLARE_CLASS_INTRINSIC(UPackageMap,UObject,CLASS_Transient|0,CoreUObject);
 
-	
-	UPackageMap( const class FPostConstructInitializeProperties& PCIP, FNetObjectIsDynamic Delegate ) :
-		UObject(PCIP), ObjectIsDynamicDelegate(Delegate)
-	{
-	}
-
 	// Begin UObject interface.
 	virtual void	PostInitProperties();
 	virtual void	Serialize( FArchive& Ar ) OVERRIDE;
@@ -146,15 +113,6 @@ class COREUOBJECT_API UPackageMap : public UObject
 	static void		AddReferencedObjects(UObject* InThis, FReferenceCollector& Collector);
 	// End UObject interface.
 	
-	// @todo document
-	virtual bool CanSerializeObject(const UObject* Obj);
-	
-	// @todo document
-	virtual bool SupportsPackage( UObject* InOuter );
-
-	// @todo document
-	virtual bool SupportsObject( const UObject* Obj );
-
 	virtual bool WriteObject( FArchive& Ar, UObject* Outer, FNetworkGUID NetGUID, FString ObjName );
 
 	// @todo document
@@ -164,34 +122,16 @@ class COREUOBJECT_API UPackageMap : public UObject
 	virtual bool SerializeName( FArchive& Ar, FName& Name );
 
 	/** get the cached field to index mappings for the given class */
-	FClassNetCache* GetClassNetCache(UClass* Class);
-	void	ClearClassNetCache();
-
-	virtual UObject * ResolvePathAndAssignNetGUID( FNetworkGUID & InOutNetGUID, const FString & PathName, UObject * ObjOuter, const bool bNoLoad = false ) { return NULL; }
-	
-	virtual void HandleUnAssignedObject(const UObject* Obj) { }
-
-	virtual FNetworkGUID GetObjectNetGUID(const UObject* Obj);
-
-	virtual UObject * GetObjectFromNetGUID(FNetworkGUID NetGUID);
+	FClassNetCache *	GetClassNetCache(UClass* Class);
+	void				ClearClassNetCache();
 
 	virtual void ResetPackageMap();
 
-	/** Removes deleted objects from the package map lookup tables */
-	virtual void CleanPackageMap();
-
-	/**
-	 * logs debug info (package list, etc) to the specified output device
-	 * @param Ar - the device to log to
-	 */
-	virtual void LogDebugInfo(FOutputDevice& Ar);
-
 	virtual void SetLocked(bool L) { }
 
 	// Variables.
 
-	virtual FNetworkGUID AssignNewNetGUID(const UObject* Object);
-	virtual FNetworkGUID AssignNetGUID(const UObject* Object, FNetworkGUID NewNetworkGUID);
+	virtual UObject * ResolvePathAndAssignNetGUID( FNetworkGUID & InOutNetGUID, const FString & PathName, UObject * ObjOuter, const bool bNoLoad = false ) { return NULL; }
 
 	void ResetUnAckedObject() { bSerializedUnAckedObject = false; }
 	bool SerializedUnAckedObject() { return bSerializedUnAckedObject; }
@@ -208,10 +148,6 @@ class COREUOBJECT_API UPackageMap : public UObject
 
 	virtual void GetNetGUIDStats(int32 &AckCount, int32 &UnAckCount, int32 &PendingCount) { }
 
-	UNetGUIDCache * GetNetGUIDCache() { return Cache; }
-
-	FNetObjectIsDynamic ObjectIsDynamicDelegate;
-
 	virtual void NotifyStreamingLevelUnload(UObject* UnloadedLevel) { }
 
 	virtual bool PrintExportBatch() { return false; }
@@ -223,28 +159,25 @@ class COREUOBJECT_API UPackageMap : public UObject
 	bool			GetLoadedUnmappedObject() const { return bLoadedUnmappedObject; }
 	FNetworkGUID	GetLastUnmappedNetGUID() const { return LastUnmappedNetGUID; }
 
+	virtual void		LogDebugInfo( FOutputDevice & Ar);
+	virtual UObject *	GetObjectFromNetGUID( const FNetworkGUID & NetGUID );
+
 protected:
 
-	virtual bool IsDynamicObject(const UObject* Object);
 	virtual bool IsNetGUIDAuthority() { return true; }
 
-	UNetGUIDCache * Cache;
-
-	/** map from package names to their index in List */
-	TMap<FName, int32> PackageListMap;
-
 	// @todo document
 	TMap<TWeakObjectPtr<UClass>, FClassNetCache*> ClassFieldIndices;
 
-	bool	bSuppressLogs;
-	bool	bShouldSerializeUnAckedObjects;
-	bool	bSerializedUnAckedObject;
-	bool	bSerializedCDO;
+	bool			bSuppressLogs;
+	bool			bShouldSerializeUnAckedObjects;
+	bool			bSerializedUnAckedObject;
+	bool			bSerializedCDO;
 
 	bool			bLoadedUnmappedObject;
 	FNetworkGUID	LastUnmappedNetGUID;
 
-	FString	DebugContextString;
+	FString			DebugContextString;
 };
 
 /** Represents a range of PacketIDs, inclusive */
@@ -292,13 +225,13 @@ enum ELifetimeCondition
 	COND_None				= 0,		// This property has no condition, and will send anytime it changes
 	COND_InitialOnly		= 1,		// This property will only attempt to send on the initial bunch
 	COND_OwnerOnly			= 2,		// This property will only send to the actor's owner
-	COND_SkipOwner			= 4,		// This property send to every connection EXCEPT the owner
-	COND_SimulatedOnly		= 5,		// This property will only send to simulated actors
-	COND_AutonomousOnly		= 6,		// This property will only send to autonomous actors
-	COND_SimulatedOrPhysics	= 7,		// This property will send to simulated OR bRepPhysics actors
-	COND_InitialOrOwner		= 8,		// This property will send on the initial packet, or to the actors owner
-	COND_Custom				= 9,		// This property has no particular condition, but wants the ability to toggle on/off via SetCustomIsActiveOverride
-	COND_Max				= 10,
+	COND_SkipOwner			= 3,		// This property send to every connection EXCEPT the owner
+	COND_SimulatedOnly		= 4,		// This property will only send to simulated actors
+	COND_AutonomousOnly		= 5,		// This property will only send to autonomous actors
+	COND_SimulatedOrPhysics	= 6,		// This property will send to simulated OR bRepPhysics actors
+	COND_InitialOrOwner		= 7,		// This property will send on the initial packet, or to the actors owner
+	COND_Custom				= 8,		// This property has no particular condition, but wants the ability to toggle on/off via SetCustomIsActiveOverride
+	COND_Max				= 9,
 };
 
 /** FLifetimeProperty
diff --git a/Engine/Source/Runtime/Engine/Classes/Engine/NetDriver.h b/Engine/Source/Runtime/Engine/Classes/Engine/NetDriver.h
--- a/Engine/Source/Runtime/Engine/Classes/Engine/NetDriver.h
+++ b/Engine/Source/Runtime/Engine/Classes/Engine/NetDriver.h
@@ -166,8 +166,7 @@ public:
 	class UWorld* World;
 
 	/** @todo document */
-	UPROPERTY()
-	UPackageMap* MasterMap;
+	TSharedPtr< class FNetGUIDCache > GuidCache;
 
 	/** The loaded UClass of the net connection type to use */
 	UPROPERTY()
diff --git a/Engine/Source/Runtime/Engine/Classes/Engine/PackageMapClient.h b/Engine/Source/Runtime/Engine/Classes/Engine/PackageMapClient.h
--- a/Engine/Source/Runtime/Engine/Classes/Engine/PackageMapClient.h
+++ b/Engine/Source/Runtime/Engine/Classes/Engine/PackageMapClient.h
@@ -31,16 +31,51 @@ public:
 	FNetworkGUID	OuterGUID;
 };
 
+/** Stores an object with path associated with FNetworkGUID */
+class FNetGuidCacheObject
+{
+public:
+	TWeakObjectPtr< UObject >		Object;
+	FString							FullPath;
+};
+
+class ENGINE_API FNetGUIDCache
+{
+public:
+	FNetGUIDCache( UNetDriver * InDriver );
+
+	void			Reset();
+	void			CleanReferences();
+	bool			SupportsObject( const UObject * Object );
+	bool			IsDynamicObject( const UObject * Object );
+	bool			IsNetGUIDAuthority();
+	FNetworkGUID	GetOrAssignNetGUID( const UObject * Object );
+	FNetworkGUID	AssignNewNetGUID( const UObject * Object );
+	void			AssignNetGUID( const UObject * Object, const FNetworkGUID & NewNetworkGUID );
+	UObject *		GetObjectFromNetGUID( const FNetworkGUID & NetGUID );
+
+	TMap< FNetworkGUID, FNetGuidCacheObject >		ObjectLookup;
+	TMap< TWeakObjectPtr< UObject >, FNetworkGUID >	NetGUIDLookup;
+	int32											UniqueNetIDs[2];
+
+	UNetDriver *									Driver;
+
+#if !(UE_BUILD_SHIPPING || UE_BUILD_TEST)
+	// History for debugging entries in the guid cache
+	TMap<FNetworkGUID, FString>						History;
+#endif
+};
+
 UCLASS(transient)
 class UPackageMapClient : public UPackageMap
 {
     GENERATED_UCLASS_BODY()
 
-	UPackageMapClient( const class FPostConstructInitializeProperties& PCIP, UNetConnection *InConnection, FNetObjectIsDynamic Delegate, UNetGUIDCache *InNetGUIDCache )
-	: UPackageMap(PCIP, Delegate)
-	, Connection(InConnection)
+	UPackageMapClient( const class FPostConstructInitializeProperties & PCIP, UNetConnection * InConnection, TSharedPtr< FNetGUIDCache > InNetGUIDCache ) : 
+		UPackageMap( PCIP )
+	,	Connection(InConnection)
 	{
-		Cache = InNetGUIDCache;
+		GuidCache = InNetGUIDCache;
 		SerializeOuter = true;
 		Locked = false;
 		ExportNetGUIDCount = 0;
@@ -60,8 +95,6 @@ class UPackageMapClient : public UPackageMap
 	// UPackageMap Interface
 	virtual bool SerializeObject( FArchive& Ar, UClass* Class, UObject*& Obj, FNetworkGUID *OutNetGUID = NULL ) OVERRIDE;
 	virtual bool SerializeNewActor( FArchive& Ar, class UActorChannel *Channel, class AActor*& Actor) OVERRIDE;
-	virtual void LogDebugInfo(FOutputDevice& Ar) OVERRIDE;
-	virtual bool SupportsObject( const UObject* Object ) OVERRIDE;
 	
 	virtual bool WriteObject( FArchive& Ar, UObject* Outer, FNetworkGUID NetGUID, FString ObjName ) OVERRIDE;
 
@@ -83,7 +116,7 @@ class UPackageMapClient : public UPackageMap
 
 	TMap<FNetworkGUID, int32>	NetGUIDExportCountMap;	// How many times we've exported each NetGUID on this connection. Public for ListNetGUIDExports 
 
-	virtual void HandleUnAssignedObject(const UObject* Obj) OVERRIDE;
+	void HandleUnAssignedObject( const UObject * Obj );
 
 	static void	AddReferencedObjects(UObject* InThis, FReferenceCollector& Collector);
 
@@ -91,26 +124,28 @@ class UPackageMapClient : public UPackageMap
 
 	virtual bool PrintExportBatch() OVERRIDE;
 
+	virtual void		LogDebugInfo( FOutputDevice & Ar) OVERRIDE;
+	virtual UObject *	GetObjectFromNetGUID( const FNetworkGUID & NetGUID ) OVERRIDE;
+
 protected:
 
 	bool	ExportNetGUID( FNetworkGUID NetGUID, const UObject * Object, FString PathName, UObject * ObjOuter );
 	void	ExportNetGUIDHeader();
 
-	FNetworkGUID	InternalGetOrAssignNetGUID( const UObject * Obj );
 	bool	InternalIsMapped(UObject* Obj, FNetworkGUID &NetGUID);
 
 	void	InternalWriteObject( FArchive& Ar, FNetworkGUID NetGUID, const UObject * Object, FString ObjectPathName, UObject* ObjectOuter );	
 	void	InternalWriteObjectPath( FArchive& Ar, FNetworkGUID NetGUID, const UObject * Object, FString ObjectPathName, UObject* ObjectOuter );
 
-	FNetworkGUID	InternalLoadObject( FArchive & Ar, UObject *& Object, const bool bIsOuterLevel, const int InternalLoadObjectRecursionCount );
+	FNetworkGUID	InternalLoadObject( FArchive & Ar, UObject *& Object, const bool bIsOuterLevel, int InternalLoadObjectRecursionCount );
 
 	void		ResolvePathAndAssignNetGUID_Deferred( const FNetworkGUID & NetGUID, const FString & PathName, const FNetworkGUID & OuterGUID );
 	UObject *	ResolvePathAndAssignNetGUID( FNetworkGUID & InOutNetGUID, const FString & PathName, const FNetworkGUID & OuterGUID, const bool bNoLoad = false );
-	
+
 	virtual UObject * ResolvePathAndAssignNetGUID( FNetworkGUID & InOutNetGUID, const FString & PathName, UObject * ObjOuter, const bool bNoLoad = false ) OVERRIDE;
 
 	bool	ShouldSendFullPath(const UObject* Object, const FNetworkGUID &NetGUID);
-
+	
 	virtual bool IsNetGUIDAuthority() OVERRIDE;
 
 	bool Locked;
@@ -120,8 +155,6 @@ protected:
 
 	bool ObjectLevelHasFinishedLoading(UObject* Obj);
 
-	virtual bool IsDynamicObject(const UObject* Object) OVERRIDE;
-
 	TSet< FNetworkGUID >			CurrentExportNetGUIDs;				// Current list of NetGUIDs being written to the Export Bunch.
 
 	TMap< FNetworkGUID, int32 >		NetGUIDAckStatus;
@@ -140,4 +173,6 @@ protected:
 	int32							ExportNetGUIDCount;
 
 	TMap< FNetworkGUID, FDeferredResolvePath >	DeferredResolvePaths;
+
+	TSharedPtr< FNetGUIDCache >		GuidCache;
 };
diff --git a/Engine/Source/Runtime/Engine/Private/DataChannel.cpp b/Engine/Source/Runtime/Engine/Private/DataChannel.cpp
--- a/Engine/Source/Runtime/Engine/Private/DataChannel.cpp
+++ b/Engine/Source/Runtime/Engine/Private/DataChannel.cpp
@@ -1997,7 +1997,7 @@ UObject * UActorChannel::ReadContentBlockHeader( FInBunch & Bunch )
 		SubObj = ConstructObject< UObject >( SubObjClass, Actor );
 		check( SubObj != NULL );
 		Actor->OnSubobjectCreatedFromReplication( SubObj );
-		Connection->PackageMap->AssignNetGUID( SubObj, NetGUID );
+		Connection->Driver->GuidCache->AssignNetGUID( SubObj, NetGUID );
 	}
 
 	check( Cast< AActor >( SubObj ) == NULL );
@@ -2072,9 +2072,9 @@ bool UActorChannel::ReplicateSubobject(UObject *Obj, FOutBunch &Bunch, const FRe
 	// here, by forcing the packagemap to give them a NetGUID.
 	//
 	// Once we can lazily handle unmapped references on the client side, this can be simplified.
-	if (!Connection->PackageMap->SupportsObject(Obj))
+	if ( !Connection->Driver->GuidCache->SupportsObject( Obj ) )
 	{
-		FNetworkGUID NetGUID = Connection->PackageMap->AssignNewNetGUID(Obj);	//Make sure he gets a NetGUID so that he is now 'supported'
+		FNetworkGUID NetGUID = Connection->Driver->GuidCache->AssignNewNetGUID( Obj );	//Make sure he gets a NetGUID so that he is now 'supported'
 	}
 
 	bool NewSubobject = false;
@@ -2223,16 +2223,19 @@ FAutoConsoleCommandWithWorld DeleteDormantActorCommand(
 #if !(UE_BUILD_SHIPPING || UE_BUILD_TEST)
 static void	FindNetGUID( const TArray<FString>& Args, UWorld* InWorld )
 {
-	for (FObjectIterator ObjIt(UNetGUIDCache::StaticClass()); ObjIt; ++ObjIt)
+	for (FObjectIterator ObjIt(UNetDriver::StaticClass()); ObjIt; ++ObjIt)
 	{
-		UNetGUIDCache *Cache = Cast<UNetGUIDCache>(*ObjIt);
-		if (Cache->HasAnyFlags(RF_ClassDefaultObject | RF_ArchetypeObject))
+		UNetDriver * Driver = Cast< UNetDriver >( *ObjIt );
+
+		if ( Driver->HasAnyFlags( RF_ClassDefaultObject | RF_ArchetypeObject ) )
+		{
 			continue;
+		}
 
 		if (Args.Num() <= 0)
 		{
 			// Display all
-			for (auto It = Cache->History.CreateIterator(); It; ++It)
+			for (auto It = Driver->GuidCache->History.CreateIterator(); It; ++It)
 			{
 				FString Str = It.Value();
 				FNetworkGUID NetGUID = It.Key();
@@ -2246,7 +2249,7 @@ static void	FindNetGUID( const TArray<FString>& Args, UWorld* InWorld )
 			FNetworkGUID NetGUID(GUIDValue);
 
 			// Search
-			FString Str = Cache->History.FindRef(NetGUID);
+			FString Str = Driver->GuidCache->History.FindRef(NetGUID);
 
 			if (!Str.IsEmpty())
 			{
diff --git a/Engine/Source/Runtime/Engine/Private/DataReplication.cpp b/Engine/Source/Runtime/Engine/Private/DataReplication.cpp
--- a/Engine/Source/Runtime/Engine/Private/DataReplication.cpp
+++ b/Engine/Source/Runtime/Engine/Private/DataReplication.cpp
@@ -208,13 +208,8 @@ void FObjectReplicator::StartReplicating( class UActorChannel * InActorChannel )
 	OwningChannel = InActorChannel;
 
 	// Cache off netGUID so if this object gets deleted we can close it
-	ObjectNetGUID = OwningChannel->Connection->PackageMap->GetObjectNetGUID( GetObject() );
-
-	if ( !ObjectNetGUID.IsValid() )
-	{
-		ObjectNetGUID = OwningChannel->Connection->PackageMap->AssignNewNetGUID( GetObject() );
-		check( !ObjectNetGUID.IsDefault() && ObjectNetGUID.IsValid() );
-	}
+	ObjectNetGUID = OwningChannel->Connection->Driver->GuidCache->GetOrAssignNetGUID( GetObject() );
+	check( !ObjectNetGUID.IsDefault() && ObjectNetGUID.IsValid() );
 
 	// Allocate retirement list.
 	// SetNum now constructs, so this is safe
diff --git a/Engine/Source/Runtime/Engine/Private/NetConnection.cpp b/Engine/Source/Runtime/Engine/Private/NetConnection.cpp
--- a/Engine/Source/Runtime/Engine/Private/NetConnection.cpp
+++ b/Engine/Source/Runtime/Engine/Private/NetConnection.cpp
@@ -114,8 +114,7 @@ void UNetConnection::InitBase(UNetDriver* InDriver,class FSocket* InSocket, cons
 	}
 
 	// Create package map.
-	//PackageMap = new(this)UPackageMapClient(FPostConstructInitializeProperties(), this, Driver->MasterMap->ObjectIsDynamicDelegate, NULL /*Driver->MasterMap->GetNetGUIDCache() */);
-	PackageMap = new(this)UPackageMapClient(FPostConstructInitializeProperties(), this, Driver->MasterMap->ObjectIsDynamicDelegate, Driver->MasterMap->GetNetGUIDCache() );
+	PackageMap = new( this )UPackageMapClient( FPostConstructInitializeProperties(), this, Driver->GuidCache );
 
 	// Create the voice channel
 	CreateChannel(CHTYPE_Voice, true, VOICE_CHANNEL_INDEX);
@@ -162,7 +161,7 @@ void UNetConnection::InitConnection(UNetDriver* InDriver, EConnectionState InSta
 	}
 
 	// Create package map.
-	PackageMap = new(this)UPackageMapClient(FPostConstructInitializeProperties(), this, Driver->MasterMap->ObjectIsDynamicDelegate, Driver->MasterMap->GetNetGUIDCache() );
+	PackageMap = new( this )UPackageMapClient( FPostConstructInitializeProperties(), this, Driver->GuidCache );
 }
 
 void UNetConnection::Serialize( FArchive& Ar )
diff --git a/Engine/Source/Runtime/Engine/Private/NetworkDriver.cpp b/Engine/Source/Runtime/Engine/Private/NetworkDriver.cpp
--- a/Engine/Source/Runtime/Engine/Private/NetworkDriver.cpp
+++ b/Engine/Source/Runtime/Engine/Private/NetworkDriver.cpp
@@ -99,11 +99,12 @@ void UNetDriver::PostInitProperties()
 		PacketSimulationSettings.RegisterCommands();
 		PacketSimulationSettings.ParseSettings(FCommandLine::Get());
 #endif
-		RoleProperty       = FindObjectChecked<UProperty>( AActor::StaticClass(), TEXT("Role"      ) );
-		RemoteRoleProperty = FindObjectChecked<UProperty>( AActor::StaticClass(), TEXT("RemoteRole") );
+		RoleProperty		= FindObjectChecked<UProperty>( AActor::StaticClass(), TEXT("Role"      ) );
+		RemoteRoleProperty	= FindObjectChecked<UProperty>( AActor::StaticClass(), TEXT("RemoteRole") );
 
-		MasterMap          = new UPackageMap(FPostConstructInitializeProperties(), FNetObjectIsDynamic::CreateUObject(this, &UNetDriver::NetObjectIsDynamic));
-		ProfileStats	   = FParse::Param(FCommandLine::Get(),TEXT("profilestats"));
+		GuidCache			= TSharedPtr< FNetGUIDCache >( new FNetGUIDCache( this ) );
+
+		ProfileStats		= FParse::Param(FCommandLine::Get(),TEXT("profilestats"));
 	}
 	// By default we're the game net driver and any child ones must override this
 	NetDriverName = NAME_GameNetDriver;
@@ -961,7 +962,7 @@ void UNetDriver::Serialize( FArchive& Ar )
 	Super::Serialize( Ar );
 
 	// Prevent referenced objects from being garbage collected.
-	Ar << ClientConnections << ServerConnection << MasterMap << RoleProperty << RemoteRoleProperty;
+	Ar << ClientConnections << ServerConnection << RoleProperty << RemoteRoleProperty;
 
 	if (Ar.IsCountingMemory())
 	{
@@ -986,14 +987,14 @@ void UNetDriver::FinishDestroy()
 		// Low level destroy.
 		LowLevelDestroy();
 
-		// Delete the master package map.
-		MasterMap = NULL;
+		// Delete the guid cache
+		GuidCache.Reset();
 	}
 	else
 	{
 		check(ServerConnection==NULL);
 		check(ClientConnections.Num()==0);
-		check(MasterMap==NULL);
+		check(!GuidCache.IsValid());
 	}
 	Super::FinishDestroy();
 }
@@ -1144,7 +1145,7 @@ bool UNetDriver::HandleNetDisconnectCommand( const TCHAR* Cmd, FOutputDevice& Ar
 bool UNetDriver::HandleNetDumpServerRPCCommand( const TCHAR* Cmd, FOutputDevice& Ar )
 {
 #if WITH_SERVER_CODE
-	UPackageMap * PackageMap = new( this )UPackageMap( FPostConstructInitializeProperties(), MasterMap->ObjectIsDynamicDelegate );
+	UPackageMap * PackageMap = new( this )UPackageMap( FPostConstructInitializeProperties() );
 
 	for ( TObjectIterator<UClass> ClassIt; ClassIt; ++ClassIt )
 	{
@@ -1285,7 +1286,8 @@ FActorDestructionInfo *	CreateDestructionInfo( UNetDriver * NetDriver, AActor* T
 	if (DestructionInfo)
 		return DestructionInfo;
 
-	FNetworkGUID NetGUID = NetDriver->MasterMap->GetObjectNetGUID(ThisActor);
+	FNetworkGUID NetGUID = NetDriver->GuidCache->GetOrAssignNetGUID( ThisActor );
+
 	FActorDestructionInfo &NewInfo = NetDriver->DestroyedStartupOrDormantActors.FindOrAdd( NetGUID );
 	NewInfo.DestroyedPosition = ThisActor->GetActorLocation();
 	NewInfo.NetGUID = NetGUID;
@@ -1325,7 +1327,7 @@ void UNetDriver::NotifyActorDestroyed( AActor* ThisActor, bool IsSeamlessTravel
 	const bool bIsServer = ServerConnection == NULL;
 	if (bIsServer)
 	{
-		const bool bIsActorStatic = MasterMap->GetObjectNetGUID(ThisActor).IsStatic();
+		const bool bIsActorStatic = !GuidCache->IsDynamicObject( ThisActor );
 		const bool bActorHasRole = ThisActor->GetRemoteRole() != ROLE_None;
 		const bool ShouldCreateDestructionInfo = bIsServer && bIsActorStatic && bActorHasRole && !IsSeamlessTravel;
 
@@ -2358,8 +2360,8 @@ int32 UNetDriver::ServerReplicateActors(float DeltaSeconds)
 							// we can't create the channel if the client is in a different world than we are
 							// or the package map doesn't support the actor's class/archetype (or the actor itself in the case of serializable actors)
 							// or it's an editor placed actor and the client hasn't initialized the level it's in
-							if ( Channel == NULL && Connection->PackageMap->SupportsObject(Actor->GetClass()) &&
-									Connection->PackageMap->SupportsObject(Actor->IsNetStartupActor() ? Actor : Actor->GetArchetype()) )
+							if ( Channel == NULL && GuidCache->SupportsObject(Actor->GetClass()) &&
+									GuidCache->SupportsObject(Actor->IsNetStartupActor() ? Actor : Actor->GetArchetype()) )
 							{
 								if (bLevelInitializedForActor)
 								{
@@ -2833,9 +2835,9 @@ bool UNetDriver::VerifyPackageInfo(FPackageInfo& Info)
 
 void UNetDriver::ResetPackageMaps()
 {
-	if (MasterMap)
+	if ( GuidCache.IsValid() )
 	{
-		MasterMap->ResetPackageMap();
+		GuidCache->Reset();
 	}
 
 	if (ServerConnection && ServerConnection->PackageMap)
@@ -2899,13 +2901,9 @@ void UNetDriver::ResetGameWorldState()
 
 void UNetDriver::CleanPackageMaps()
 {
-	if (MasterMap)
-	{
-		MasterMap->CleanPackageMap();
-	}
-	if (ServerConnection && ServerConnection->PackageMap)
+	if ( GuidCache.IsValid() )
 	{
-		ServerConnection->PackageMap->CleanPackageMap();
+		GuidCache->CleanReferences();
 	}
 }
 
diff --git a/Engine/Source/Runtime/Engine/Private/PackageMapClient.cpp b/Engine/Source/Runtime/Engine/Private/PackageMapClient.cpp
--- a/Engine/Source/Runtime/Engine/Private/PackageMapClient.cpp
+++ b/Engine/Source/Runtime/Engine/Private/PackageMapClient.cpp
@@ -14,8 +14,6 @@
 static const int GUID_PACKET_NOT_ACKED	= -2;		
 static const int GUID_PACKET_ACKED		= -1;		
 
-#define USE_NEW_DYNAMIC_AND_SUPPORT_CHECKS 1
-
 /**
  * Don't allow infinite recursion of InternalLoadObject - an attacker could
  * send malicious packets that cause a stack overflow on the server.
@@ -57,7 +55,7 @@ bool UPackageMapClient::SerializeObject( FArchive& Ar, UClass* Class, UObject*&
 		}
 
 		bSerializedUnAckedObject = false;
-		FNetworkGUID NetGUID = InternalGetOrAssignNetGUID( Object );
+		FNetworkGUID NetGUID = GuidCache->GetOrAssignNetGUID( Object );
 		bool Mapped = InternalIsMapped(Object, NetGUID);
 
 		// Write out NetGUID to caller if necessary
@@ -97,7 +95,7 @@ bool UPackageMapClient::SerializeObject( FArchive& Ar, UClass* Class, UObject*&
 		FNetworkGUID NetGUID = InternalLoadObject( Ar, Object, false, 0 );
 
 		// Write out NetGUID to caller if necessary
-		if (OutNetGUID)
+		if ( OutNetGUID )
 		{
 			*OutNetGUID = NetGUID;
 		}	
@@ -107,16 +105,16 @@ bool UPackageMapClient::SerializeObject( FArchive& Ar, UClass* Class, UObject*&
 		// ----------------	
 
 		// NULL if we haven't finished loading the objects level yet
-		if (!ObjectLevelHasFinishedLoading(Object))
+		if ( !ObjectLevelHasFinishedLoading( Object ) )
 		{
 			UE_LOG(LogNetPackageMap, Log, TEXT("Using None instead of replicated reference to %s because the level it's in has not been made visible"), *Object->GetFullName());
 			Object = NULL;
 		}
 
 		// Check that we got the right class
-		if( Object && !Object->IsA(Class) )
+		if ( Object && !Object->IsA( Class ) )
 		{
-			UE_LOG(LogNetPackageMap, Log, TEXT("Forged object: got %s, expecting %s"),*Object->GetFullName(),*Class->GetFullName());
+			UE_LOG(LogNetPackageMap, Warning, TEXT("Forged object: got %s, expecting %s"),*Object->GetFullName(),*Class->GetFullName());
 			Object = NULL;
 		}
 
@@ -177,7 +175,7 @@ bool UPackageMapClient::WriteObject( FArchive& Ar, UObject* ObjOuter, FNetworkGU
  */
 bool UPackageMapClient::SerializeNewActor(FArchive& Ar, class UActorChannel *Channel, class AActor*& Actor)
 {
-	UE_LOG(LogNetPackageMap, VeryVerbose, TEXT("SerializeNewActor START"));
+	UE_LOG( LogNetPackageMap, VeryVerbose, TEXT( "SerializeNewActor START" ) );
 
 	uint8 bIsClosingChannel = 0;
 
@@ -208,6 +206,7 @@ bool UPackageMapClient::SerializeNewActor(FArchive& Ar, class UActorChannel *Cha
 	{
 		// This must be a destruction info coming through or something is wrong
 		// If so, we should be both closing the channel and we should find the actor
+		// This can happen when dormant actors that don't have channels get destroyed
 		if ( bIsClosingChannel == 0 || Actor == NULL )
 		{
 			UE_LOG( LogNetPackageMap, Error, TEXT( "UPackageMapClient::SerializeNewActor: bIsClosingChannel == 0 || Actor == NULL : %s" ), Actor ? *Actor->GetName() : TEXT( "NULL" ) );
@@ -277,7 +276,7 @@ bool UPackageMapClient::SerializeNewActor(FArchive& Ar, class UActorChannel *Cha
 					Actor = Connection->Driver->GetWorld()->SpawnActor(Archetype->GetClass(), &Location, &Rotation, SpawnInfo );
 					Actor->PostNetReceiveVelocity(Velocity);
 
-					AssignNetGUID(Actor, NetGUID);
+					GuidCache->AssignNetGUID(Actor, NetGUID);
 					SpawnedActor = true;
 				}
 				else
@@ -312,13 +311,13 @@ bool UPackageMapClient::SerializeNewActor(FArchive& Ar, class UActorChannel *Cha
 
 			for ( UObject * Obj : Subobjs )
 			{
-				FNetworkGUID SubobjNetGUID = Cache->NetGUIDLookup.FindRef( Obj );
+				FNetworkGUID SubobjNetGUID = GuidCache->NetGUIDLookup.FindRef( Obj );
 
 				uint8 bForceSendGuid = SubobjNetGUID.IsValid() ? 1 : 0;
 
 				if ( !bForceSendGuid )
 				{
-					SubobjNetGUID = InternalGetOrAssignNetGUID( Obj );
+					SubobjNetGUID = GuidCache->GetOrAssignNetGUID( Obj );
 				}
 
 				check( SubobjNetGUID.IsDynamic() );		// Make sure some assumptions are correct
@@ -417,7 +416,7 @@ bool UPackageMapClient::SerializeNewActor(FArchive& Ar, class UActorChannel *Cha
 
 					for ( int32 i = 0; i < Subobjs.Num(); i++ )
 					{
-						AssignNetGUID( Subobjs[i], SubObjGuids[i] );
+						GuidCache->AssignNetGUID( Subobjs[i], SubObjGuids[i] );
 						UE_LOG( LogNetPackageMap, Log, TEXT( "SerializeNewActor: [Loading] Assigned NetGUID <%s> to subobject %s" ), *SubObjGuids[i].ToString(), *Subobjs[i]->GetPathName() );
 					}
 				}
@@ -436,27 +435,6 @@ bool UPackageMapClient::SerializeNewActor(FArchive& Ar, class UActorChannel *Cha
 //
 //--------------------------------------------------------------------
 
-/** Gets or assigns a new NetGUID to this object. Returns whether the object is fully mapped or not */
-FNetworkGUID UPackageMapClient::InternalGetOrAssignNetGUID( const UObject * Object )
-{
-	if (!Object || !SupportsObject(Object))
-	{
-		// Null of unsupported object, leave as default NetGUID and just return mapped=true
-		return FNetworkGUID();
-	}
-
-	// ----------------
-	// Assign NetGUID if necessary
-	// ----------------
-	FNetworkGUID NetGUID = Cache->NetGUIDLookup.FindRef(Object);
-	if (!NetGUID.IsValid())
-	{
-		NetGUID = AssignNewNetGUID(Object);
-	}
-
-	return NetGUID;
-}
-
 bool UPackageMapClient::InternalIsMapped(UObject* Object, FNetworkGUID &NetGUID)
 {
 	if (!Object)
@@ -496,7 +474,7 @@ bool UPackageMapClient::InternalIsMapped(UObject* Object, FNetworkGUID &NetGUID)
 		TestObj = Actor;
 	}
 
-	bool IsDynamic = IsDynamicObject(TestObj);
+	bool IsDynamic = GuidCache->IsDynamicObject(TestObj);
 	bool Mapped = true;
 	if (IsDynamic)
 	{
@@ -656,7 +634,7 @@ void UPackageMapClient::InternalWriteObjectPath( FArchive & Ar, FNetworkGUID Net
 	}
 
 	// Serialize reference to outer. This is basically a form of compression.
-	FNetworkGUID OuterNetGUID = InternalGetOrAssignNetGUID( ObjectOuter );
+	FNetworkGUID OuterNetGUID = GuidCache->GetOrAssignNetGUID( ObjectOuter );
 
 	InternalWriteObject( Ar, OuterNetGUID, ObjectOuter, TEXT( "" ), NULL );
 
@@ -890,7 +868,16 @@ UObject * UPackageMapClient::ResolvePathAndAssignNetGUID( FNetworkGUID & InOutNe
 
 	if ( Object )
 	{
-		InOutNetGUID = AssignNetGUID( Object, InOutNetGUID );
+		if ( InOutNetGUID.IsDefault() )
+		{
+			check( IsNetGUIDAuthority() );
+			InOutNetGUID = GuidCache->GetOrAssignNetGUID( Object );
+			HandleUnAssignedObject( Object );
+		}
+		else
+		{
+			GuidCache->AssignNetGUID( Object, InOutNetGUID );
+		}
 	}
 
 	return Object;
@@ -1088,8 +1075,9 @@ void UPackageMapClient::ReceiveNetGUIDBunch( FInBunch &InBunch )
 	int32 NumGUIDsRead = 0;
 	while( NumGUIDsRead < NumGUIDsInBunch )
 	{
-		UObject *Obj;
+		UObject * Obj = NULL;
 		InternalLoadObject( InBunch, Obj, false, 0 );
+
 		if ( InBunch.IsError() )
 		{
 			UE_LOG( LogNetPackageMap, Error, TEXT( "UPackageMapClient::ReceiveNetGUIDBunch: InBunch.IsError() after InternalLoadObject" ) );
@@ -1276,7 +1264,7 @@ void UPackageMapClient::HandleUnAssignedObject(const UObject* Obj)
 {
 	check( Obj != NULL );
 
-	FNetworkGUID NetGUID = InternalGetOrAssignNetGUID( Obj );
+	FNetworkGUID NetGUID = GuidCache->GetOrAssignNetGUID( Obj );
 
 	if ( !NetGUID.IsDefault() && ShouldSendFullPath( Obj, NetGUID ) )
 	{
@@ -1325,58 +1313,12 @@ bool UPackageMapClient::ShouldSendFullPath(const UObject* Object, const FNetwork
 	return Locked || !NetGUIDHasBeenAckd( NetGUID );
 }
 
-/**
- *	Dynamic objects are actors that were spawned in the world at run time, and therefor cannot be
- *	referenced with a path name to the client.
- */
-bool UPackageMapClient::IsDynamicObject(const UObject* Object)
-{
-#if USE_NEW_DYNAMIC_AND_SUPPORT_CHECKS == 1
-	check( Object != NULL );
-	check( Object->IsSupportedForNetworking() );
-
-	// Any non net addressable object is dynamic
-	return !Object->IsFullNameStableForNetworking();
-#else
-	const AActor * Actor = Cast<const AActor>(Object);
-	bool NonAddressableSubobject = false;
-	if (!Actor)
-	{
-		Actor = Cast<AActor>(Object->GetOuter());
-		if (Actor)
-		{
-			const UActorComponent* ActorComponent = Cast<const UActorComponent>(Object);
-			if (ActorComponent)
-			{
-				// Actor compnents define if they are NetAddressable via UActorComponent::IsNameStableForNetworking()
-				if (!ActorComponent->IsNameStableForNetworking())
-				{
-					NonAddressableSubobject  = true;
-				}
-			}
-			// Non actor component subobjects are net addressable if they are default subobjects or RF_WasLoaded
-			else if (!Object->IsDefaultSubobject() && !Object->HasAnyFlags(RF_WasLoaded))
-			{
-				NonAddressableSubobject = true;
-			}
-		}
-	}
-	
-	if (!Actor || Actor->HasAnyFlags(RF_ClassDefaultObject | RF_ArchetypeObject) || (Actor->IsNetStartupActor() && !NonAddressableSubobject))
-	{
-		return false;
-	}
-
-	return true;
-#endif
-}
-
 /**
  *	Prints debug info about this package map's state
  */
-void UPackageMapClient::LogDebugInfo(FOutputDevice& Ar)
+void UPackageMapClient::LogDebugInfo( FOutputDevice & Ar )
 {
-	for (auto It = Cache->NetGUIDLookup.CreateIterator(); It; ++It)
+	for ( auto It = GuidCache->NetGUIDLookup.CreateIterator(); It; ++It )
 	{
 		FNetworkGUID NetGUID = It.Value();
 
@@ -1399,7 +1341,7 @@ void UPackageMapClient::LogDebugInfo(FOutputDevice& Ar)
 		}
 
 		UObject *Obj = It.Key().Get();
-		FString Str = FString::Printf(TEXT("%s [%s] [%s] - %s"), *NetGUID.ToString(), *Status, IsDynamicObject(Obj) ? TEXT("Dynamic") : TEXT("Static") , Obj ? *Obj->GetPathName() : TEXT("NULL"));
+		FString Str = FString::Printf(TEXT("%s [%s] [%s] - %s"), *NetGUID.ToString(), *Status, NetGUID.IsDynamic() ? TEXT("Dynamic") : TEXT("Static") , Obj ? *Obj->GetPathName() : TEXT("NULL"));
 		Ar.Logf(*Str);
 		UE_LOG(LogNetPackageMap, Log, TEXT("%s"), *Str);
 	}
@@ -1439,111 +1381,12 @@ bool UPackageMapClient::ObjectLevelHasFinishedLoading(UObject* Object)
  */
 bool UPackageMapClient::IsNetGUIDAuthority()
 {
-	return !(Connection && Connection->Driver && Connection->Driver->ServerConnection == Connection);
-}
-
-bool UPackageMapClient::SupportsObject( const UObject* Object )
-{
-#if USE_NEW_DYNAMIC_AND_SUPPORT_CHECKS == 1
-	// NULL is always supported
-	if ( !Object )
-	{
-		return true;
-	}
-
-	// If we already gave it a NetGUID, its supported.
-	// This should happen for dynamic subobjects.
-	FNetworkGUID NetGUID = Cache->NetGUIDLookup.FindRef(Object);
-	if (NetGUID.IsValid())
-	{
-		return true;
-	}
-
-	if ( Object->IsFullNameStableForNetworking() )
-	{
-		// If object is fully net addressable, it's definitely supported
-		return true;
-	}
-
-	if ( Object->IsSupportedForNetworking() )
-	{
-		// This means the server will explicitly tell the client to spawn and assign the id for this object
-		return true;
-	}
-
-	UE_LOG( LogNetPackageMap, Warning, TEXT( "UPackageMapClient::SupportsObject: %s NOT Supported." ), *Object->GetFullName() );
-	UE_LOG( LogNetPackageMap, Warning, TEXT( "   %s"), *DebugContextString );
-
-	return false;
-#else
-	// NULL is always supported
-	if (!Object)
-	{
-		return true;
-	}
-
-	// Non actor UObjects
-	if (Cast<const AActor>(Object) == NULL && Cast<const ULevel>(Object) == NULL )
-	{
-		// If we already gave it a NetGUID, its supported.
-		// This should happen for dynamic subobjects.
-		FNetworkGUID NetGUID = Cache->NetGUIDLookup.FindRef(Object);
-		if (NetGUID.IsValid())
-		{
-			return true;
-		}
-		
-		// Some special rules for actor components
-		const UActorComponent *ActorComponent = Cast<const UActorComponent>(Object);
-		if (ActorComponent)
-		{
-			// NetAddressable actor components are supported even before they have a NetGUID
-			// So things like an actor's static mesh component can be referenced even if the SMC doesn't have replicated properties itself
-			if (ActorComponent->IsNameStableForNetworking())
-			{
-				return true;
-			}
-
-			// This case handles a dynamically spawned actor component who doesn't have a NetGUID yet, but probably will have one soon.
-			// For now, it is unsupported. Once it is assigned a netguid by the owning actor channel, it will be supported.
-			if (ActorComponent->GetIsReplicated())
-			{
-				// This is ok. Don't warn.
-				return false;
-			}
-		}
-
-		// Everything else
-		UObject * ThisOuter = Object->GetOuter();
-		while(ThisOuter)
-		{
-			// If the UObject lives in a world, we can only serialize it if it was loaded (not if it was dynamically spawned)
-			if (Cast<UWorld>(ThisOuter))
-			{
-				// This object lives in a world, it can only be supported if RF_WasLoaded is true
-				if (Object->HasAnyFlags(RF_WasLoaded) == false)
-				{
-					UE_LOG(LogNetPackageMap, Warning, TEXT("UPackageMapClient::SupportsObject %s NOT Supported.Not RF_WasLoaded. NetGUID: <%s>"), *Object->GetPathName(), *NetGUID.ToString() );
-					UE_LOG(LogNetPackageMap, Warning, TEXT("   %s"), *DebugContextString);
-					return false;
-				}
-
-				break;
-			}
-			ThisOuter = ThisOuter->GetOuter();
-		}
-	}
-
-	// Everything else is now supported
-	return true;
-#endif
+	return GuidCache->IsNetGUIDAuthority();
 }
 
 /** Resets packagemap state for server travel */
 void UPackageMapClient::ResetPackageMap()
 {
-	UPackageMap::ResetPackageMap();
-	
 	NetGUIDExportCountMap.Empty();
 	CurrentExportNetGUIDs.Empty();
 	NetGUIDAckStatus.Empty();
@@ -1597,7 +1440,7 @@ bool UPackageMapClient::PrintExportBatch()
 
 	// Print the whole thing for reference
 #if !(UE_BUILD_SHIPPING || UE_BUILD_TEST)
-	for (auto It = Cache->History.CreateIterator(); It; ++It)
+	for (auto It = GuidCache->History.CreateIterator(); It; ++It)
 	{
 		FString Str = It.Value();
 		FNetworkGUID NetGUID = It.Key();
@@ -1623,6 +1466,370 @@ bool UPackageMapClient::PrintExportBatch()
 	return true;
 }
 
+UObject * UPackageMapClient::GetObjectFromNetGUID( const FNetworkGUID & NetGUID )
+{
+	return GuidCache->GetObjectFromNetGUID( NetGUID );
+}
+
+//----------------------------------------------------------------------------------------
+//	FNetGUIDCache
+//----------------------------------------------------------------------------------------
+
+FNetGUIDCache::FNetGUIDCache( UNetDriver * InDriver ) : Driver( InDriver )
+{
+	UniqueNetIDs[0] = UniqueNetIDs[1] = 0;
+}
+
+void FNetGUIDCache::Reset()
+{
+	ObjectLookup.Empty();
+	NetGUIDLookup.Empty();
+	UniqueNetIDs[0] = UniqueNetIDs[1] = 0;
+}
+
+class FArchiveCountMemGUID : public FArchive
+{
+public:
+	FArchiveCountMemGUID() : Mem( 0 ) { ArIsCountingMemory = true; }
+	void CountBytes( SIZE_T InNum, SIZE_T InMax ) { Mem += InMax; }
+	SIZE_T Mem;
+};
+
+void FNetGUIDCache::CleanReferences()
+{
+	// NOTE - Quick fix for some package map issues where client was cleaning up guids, which was causing issues.
+	// For now, we no longer clean these up.  The client can sometimes cleanup guids that the server doesn't which causes issues.
+	// The guid memory will leak for now until we find a better way.
+
+	// Free any dynamic guids that refer to objects that are dead
+	for ( auto It = ObjectLookup.CreateIterator(); It; ++It )
+	{
+		if ( It.Value().Object.IsValid() )
+		{
+			// Make sure the path matches (which can change during seamless travel)
+			It.Value().FullPath = It.Value().Object->GetPathName();
+			continue;
+		}
+
+		if ( It.Key().IsDynamic() )
+		{
+			UE_LOG( LogNetPackageMap, Log, TEXT( "Cleaning reference to NetGUID: <%s> (%s) (ObjectLookup)" ), *It.Key().ToString(), *It.Value().FullPath );
+			It.RemoveCurrent();
+		}
+	}
+
+	for ( auto It = NetGUIDLookup.CreateIterator(); It; ++It )
+	{
+		if ( !It.Key().IsValid() )
+		{
+			UE_LOG( LogNetPackageMap, Log, TEXT( "Cleaning reference to NetGUID: <%s> (NetGUIDLookup)" ), *It.Value().ToString() );
+			It.RemoveCurrent();
+		}
+	}
+
+	// Sanity check
+	for ( auto It = ObjectLookup.CreateIterator(); It; ++It )
+	{
+		check( It.Value().Object.IsValid() || It.Key().IsStatic() );
+
+		if ( It.Value().Object.IsValid() )
+		{
+			checkf( !It.Value().Object.IsValid() || NetGUIDLookup.FindRef( It.Value().Object ) == It.Key(), TEXT("Failed to validate ObjectLookup map in UPackageMap. Object '%s' was not in the NetGUIDLookup map with with value '%s'."), *It.Value().Object.Get()->GetPathName(), *It.Key().ToString());
+		}
+		else
+		{
+			for ( auto It2 = NetGUIDLookup.CreateIterator(); It2; ++It2 )
+			{
+				checkf( It2.Value() != It.Key(), TEXT("Failed to validate ObjectLookup map in UPackageMap. Object '%s' was in the NetGUIDLookup map with with value '%s'."), *It.Value().Object.Get()->GetPathName(), *It.Key().ToString());
+			}
+		}
+	}
+
+	for ( auto It = NetGUIDLookup.CreateIterator(); It; ++It )
+	{
+		check( It.Key().IsValid() );
+		checkf( ObjectLookup.FindRef( It.Value() ).Object == It.Key(), TEXT("Failed to validate NetGUIDLookup map in UPackageMap. GUID '%s' was not in the ObjectLookup map with with object '%s'."), *It.Value().ToString(), *It.Key().Get()->GetPathName());
+	}
+
+	FArchiveCountMemGUID CountBytesAr;
+
+	ObjectLookup.CountBytes( CountBytesAr );
+	NetGUIDLookup.CountBytes( CountBytesAr );
+
+	// Make this a warning to be a constant reminder that we need to ultimately free this memory when we find a solution
+	UE_LOG( LogNetPackageMap, Warning, TEXT( "FNetGUIDCache::CleanReferences: ObjectLookup: %i, NetGUIDLookup: %i, Mem: %i kB" ), ObjectLookup.Num(), NetGUIDLookup.Num(), ( CountBytesAr.Mem / 1024 ) );
+}
+
+bool FNetGUIDCache::SupportsObject( const UObject * Object )
+{
+	// NULL is always supported
+	if ( !Object )
+	{
+		return true;
+	}
+
+	// If we already gave it a NetGUID, its supported.
+	// This should happen for dynamic subobjects.
+	FNetworkGUID NetGUID = NetGUIDLookup.FindRef( Object );
+
+	if ( NetGUID.IsValid() )
+	{
+		return true;
+	}
+
+	if ( Object->IsFullNameStableForNetworking() )
+	{
+		// If object is fully net addressable, it's definitely supported
+		return true;
+	}
+
+	if ( Object->IsSupportedForNetworking() )
+	{
+		// This means the server will explicitly tell the client to spawn and assign the id for this object
+		return true;
+	}
+
+	UE_LOG( LogNetPackageMap, Warning, TEXT( "FNetGUIDCache::SupportsObject: %s NOT Supported." ), *Object->GetFullName() );
+	//UE_LOG( LogNetPackageMap, Warning, TEXT( "   %s"), *DebugContextString );
+
+	return false;
+}
+
+/**
+ *	Dynamic objects are actors that were spawned in the world at run time, and therefor cannot be
+ *	referenced with a path name to the client.
+ */
+bool FNetGUIDCache::IsDynamicObject( const UObject * Object )
+{
+	check( Object != NULL );
+	check( Object->IsSupportedForNetworking() );
+
+	// Any non net addressable object is dynamic
+	return !Object->IsFullNameStableForNetworking();
+}
+
+bool FNetGUIDCache::IsNetGUIDAuthority()
+{
+	return Driver == NULL || Driver->IsServer();
+}
+
+/** Gets or assigns a new NetGUID to this object. Returns whether the object is fully mapped or not */
+FNetworkGUID FNetGUIDCache::GetOrAssignNetGUID( const UObject * Object )
+{
+	if ( !Object || !SupportsObject( Object ) )
+	{
+		// Null of unsupported object, leave as default NetGUID and just return mapped=true
+		return FNetworkGUID();
+	}
+
+	// ----------------
+	// Assign NetGUID if necessary
+	// ----------------
+	FNetworkGUID NetGUID = NetGUIDLookup.FindRef( Object );
+
+	if ( NetGUID.IsValid() )
+	{
+		return NetGUID;
+	}
+
+	return AssignNewNetGUID( Object );
+}
+
+/**
+ *	Generate a new NetGUID for this object and assign it.
+ */
+FNetworkGUID FNetGUIDCache::AssignNewNetGUID( const UObject * Object )
+{
+	if ( !IsNetGUIDAuthority() )
+	{
+		// We cannot make new NetGUIDs
+		return FNetworkGUID::GetDefault();
+	}
+
+#define COMPOSE_NET_GUID( Index, IsStatic )		( ( ( Index ) << 1 ) | ( IsStatic ) )
+#define ALLOC_NEW_NET_GUID( IsStatic )			( COMPOSE_NET_GUID( ++UniqueNetIDs[ IsStatic ], IsStatic ) )
+
+	// Generate new NetGUID and assign it
+	const int32 IsStatic = IsDynamicObject( Object ) ? 0 : 1;
+
+	const FNetworkGUID NewNetGuid( ALLOC_NEW_NET_GUID( IsStatic ) );
+
+	AssignNetGUID( Object, NewNetGuid );
+
+	return NewNetGuid;
+}
+
+/**
+ *	Assigns the given UObject the give NetNetworkGUID.
+ */
+void FNetGUIDCache::AssignNetGUID( const UObject * Object, const FNetworkGUID & NewNetworkGUID )
+{
+	check( !Object->IsPendingKill() );
+	check( !NewNetworkGUID.IsDefault() );
+
+	UE_LOG( LogNetPackageMap, Log, TEXT("Assigning %s NetGUID <%s>"), Object ? *Object->GetName() : TEXT("NULL"), *NewNetworkGUID.ToString() );
+	
+	// Actually assign
+	FNetGuidCacheObject CacheObject;
+
+	CacheObject.Object = Object;
+	CacheObject.FullPath = Object->GetPathName();
+
+	// Remove any existing entries
+	{
+		const FNetGuidCacheObject * ExistingCacheObjectPtr = ObjectLookup.Find( NewNetworkGUID );
+
+		if ( ExistingCacheObjectPtr )
+		{
+			const UObject * OldObject = ExistingCacheObjectPtr->Object.Get();
+
+			// If we find static guid in the cache, this means we didn't find it in GetObjectFromNetGUID
+			// This is a problem, since we should never get this far if we do find it there
+			if ( !NewNetworkGUID.IsDynamic() )
+			{
+				UE_LOG( LogNetPackageMap, Warning, TEXT( "AssignNetGUID: NewNetworkGUID was in ObjectLookup cache but is static: <%s> %s" ), *NewNetworkGUID.ToString(), Object ? *Object->GetPathName() : TEXT( "NULL" ) );
+			}
+
+			// If this net guid was found but the old object is NULL, this can happen due to:
+			//	1. Actor channel was closed locally (but we don't remove the net guid entry, since we can't know for sure if it will be referenced again)
+			//		a. Then when we re-create a channel, and assign this actor, we will find the old guid entry here
+			//	2. Dynamic object was locally GC'd, but then exported again from the server
+			//
+			// If this net guid was found and the objects match, we don't care. This can happen due to:
+			//	1. Same thing above can happen, but if we for some reason didn't destroy the actor/object we will see this case
+			//
+			// If the object pointers are different, this can be a problem, 
+			//	since this should only be possible if something gets out of sync during the net guid exchange code
+			if ( OldObject != NULL && OldObject != Object )
+			{
+				UE_LOG( LogNetPackageMap, Warning, TEXT( "Reassigning NetGUID <%s> to %s (was assigned to object %s)" ), *NewNetworkGUID.ToString(), Object ? *Object->GetPathName() : TEXT( "NULL" ), OldObject ? *OldObject->GetPathName() : TEXT( "NULL" ) );
+			}
+			else
+			{
+				UE_LOG( LogNetPackageMap, Verbose, TEXT( "Reassigning NetGUID <%s> to %s (was assigned to object %s)" ), *NewNetworkGUID.ToString(), Object ? *Object->GetPathName() : TEXT( "NULL" ), OldObject ? *OldObject->GetPathName() : TEXT( "NULL" ) );
+			}
+
+			NetGUIDLookup.Remove( ExistingCacheObjectPtr->Object );
+		}
+
+		const FNetworkGUID * ExistingNetworkGUIDPtr = NetGUIDLookup.Find( Object );
+
+		if ( ExistingNetworkGUIDPtr )
+		{
+			// This can happen when the server destroys a package/object and the client doesn't.
+			// The server will simply assign a new guid, and the client will assign to the old object that never got deleted locally
+			UE_LOG( LogNetPackageMap, Log, TEXT( "Changing NetGUID on object %s from <%s:%s> to <%s:%s>" ), Object ? *Object->GetPathName() : TEXT( "NULL" ), *ExistingNetworkGUIDPtr->ToString(), ExistingNetworkGUIDPtr->IsDynamic() ? TEXT("TRUE") : TEXT("FALSE"), *NewNetworkGUID.ToString(), NewNetworkGUID.IsDynamic() ? TEXT("TRUE") : TEXT("FALSE") );
+			ObjectLookup.Remove( *ExistingNetworkGUIDPtr );
+		}
+	}
+
+	ObjectLookup.Add( NewNetworkGUID, CacheObject );
+	NetGUIDLookup.Add( Object, NewNetworkGUID );
+
+#if !(UE_BUILD_SHIPPING || UE_BUILD_TEST)
+	History.Add(NewNetworkGUID, Object ? Object->GetPathName() : TEXT("NULL"));
+#endif
+}
+
+UObject * FNetGUIDCache::GetObjectFromNetGUID( const FNetworkGUID & NetGUID )
+{
+	if ( !ensure( NetGUID.IsValid() ) )
+	{
+		return NULL;
+	}
+
+	if ( !ensure( !NetGUID.IsDefault() ) )
+	{
+		return NULL;
+	}
+
+	FNetGuidCacheObject * CacheObjectPtr = ObjectLookup.Find( NetGUID );
+
+	if ( CacheObjectPtr == NULL )
+	{
+		// We don't have the object mapped yet
+		return NULL;
+	}
+
+	UObject * Object = CacheObjectPtr->Object.Get();
+
+	if ( Object != NULL )
+	{
+		check( Object->GetPathName() == CacheObjectPtr->FullPath );
+		return Object;
+	}
+
+	if ( CacheObjectPtr->FullPath.IsEmpty() )
+	{
+		// This probably isn't possible, but making it a warning
+		UE_LOG( LogNetPackageMap, Warning, TEXT( "GetObjectFromNetGUID: No full path for %s" ), *NetGUID.ToString() );
+		return NULL;
+	}
+
+	if ( NetGUID.IsDynamic() )
+	{
+		// Dynamic objects don't have stable names, so we can't possibly reload the same object
+		UE_LOG( LogNetPackageMap, VeryVerbose, TEXT( "GetObjectFromNetGUID: Cannot re-create dynamic object after GC <%s, %s>" ), *NetGUID.ToString(), *CacheObjectPtr->FullPath );
+		return NULL;
+	}
+
+	//
+	// The object was previously mapped, but we GC'd it
+	// We need to reload it now
+	//
+
+	Object = StaticFindObject( UObject::StaticClass(), NULL, *CacheObjectPtr->FullPath, false );
+
+	if ( Object == NULL )
+	{
+		if ( IsNetGUIDAuthority() )
+		{
+			// The authority shouldn't be loading resources on demand, at least for now.
+			// This could be garbage or a security risk
+			// Another possibility is in dynamic property replication if the server reads the previously serialized state
+			// that has a now destroyed actor in it.
+			UE_LOG( LogNetPackageMap, Warning, TEXT( "GetObjectFromNetGUID: Could not find Object for: NetGUID <%s, %s> (and IsNetGUIDAuthority())" ), *NetGUID.ToString(), *CacheObjectPtr->FullPath );
+			return NULL;
+		}
+		else
+		{
+			UE_LOG( LogNetPackageMap, Log, TEXT( "GetObjectFromNetGUID: Could not find Object for: NetGUID <%s, %s>"), *NetGUID.ToString(), *CacheObjectPtr->FullPath );
+		}
+
+		Object = StaticLoadObject( UObject::StaticClass(), NULL, *CacheObjectPtr->FullPath, NULL, LOAD_NoWarn );
+
+		if ( Object )
+		{
+			UE_LOG( LogNetPackageMap, Log, TEXT( "GetObjectFromNetGUID: StaticLoadObject. Found: %s" ), Object ? *Object->GetName() : TEXT( "NULL" ) );
+		}
+		else
+		{
+			Object = LoadPackage( NULL, *CacheObjectPtr->FullPath, LOAD_None );
+			UE_LOG( LogNetPackageMap, Log, TEXT( "GetObjectFromNetGUID: LoadPackage. Found: %s" ), Object ? *Object->GetName() : TEXT( "NULL" ) );
+		}
+	}
+	else
+	{
+		// If we actually found the object, we probably shouldn't have GC'd it, so that's odd
+		//UE_LOG( LogNetPackageMap, Warning, TEXT( "GetObjectFromNetGUID: Object should not be found after GC: <%s, %s>" ), *NetGUID.ToString(), *CacheObjectPtr->FullPath );
+	}
+
+	if ( Object == NULL )
+	{
+		// If we get here, that means we have an invalid path, which shouldn't be possible, but making it a warning
+		UE_LOG( LogNetPackageMap, Warning, TEXT( "GetObjectFromNetGUID: FAILED to re-create object after GC: <%s, %s>" ), *NetGUID.ToString(), *CacheObjectPtr->FullPath );
+	}
+	else
+	{
+		// Reassign the object pointer for quick access next time
+		CacheObjectPtr->Object = Object;		
+		// Make sure the object is in the GUIDToObjectLookup.
+		NetGUIDLookup.Add( Object, NetGUID );
+	}
+
+	return Object;
+}
+
 //------------------------------------------------------
 // Debug command to see how many times we've exported each NetGUID
 // Used for measuring inefficiencies. Some duplication is unavoidable since we cannot garuntee atomicicity across multiple channels.
diff --git a/Engine/Source/Runtime/Online/OnlineSubsystemUtils/Private/OnlineBeaconClient.cpp b/Engine/Source/Runtime/Online/OnlineSubsystemUtils/Private/OnlineBeaconClient.cpp
--- a/Engine/Source/Runtime/Online/OnlineSubsystemUtils/Private/OnlineBeaconClient.cpp
+++ b/Engine/Source/Runtime/Online/OnlineSubsystemUtils/Private/OnlineBeaconClient.cpp
@@ -133,7 +133,7 @@ void AOnlineBeaconClient::NotifyControlMessage(UNetConnection* Connection, uint8
 				FNetControlMessage<NMT_BeaconAssignGUID>::Receive(Bunch, NetGUID);
 				if (NetGUID.IsValid())
 				{
-					Connection->PackageMap->AssignNetGUID(this, NetGUID);
+					Connection->Driver->GuidCache->AssignNetGUID( this, NetGUID );
 
 					FString BeaconType = GetBeaconType();
 					FNetControlMessage<NMT_BeaconNetGUIDAck>::Send(Connection, BeaconType);
diff --git a/Engine/Source/Runtime/Online/OnlineSubsystemUtils/Private/OnlineBeaconHost.cpp b/Engine/Source/Runtime/Online/OnlineSubsystemUtils/Private/OnlineBeaconHost.cpp
--- a/Engine/Source/Runtime/Online/OnlineSubsystemUtils/Private/OnlineBeaconHost.cpp
+++ b/Engine/Source/Runtime/Online/OnlineSubsystemUtils/Private/OnlineBeaconHost.cpp
@@ -134,7 +134,7 @@ void AOnlineBeaconHost::NotifyControlMessage(UNetConnection* Connection, uint8 M
 						AOnlineBeaconClient* NewClientActor = SpawnBeaconActor();
 						if (NewClientActor && BeaconType == NewClientActor->GetBeaconType())
 						{
-							FNetworkGUID NetGUID = Connection->PackageMap->AssignNewNetGUID(NewClientActor);
+							FNetworkGUID NetGUID = Connection->Driver->GuidCache->AssignNewNetGUID( NewClientActor );
 							NewClientActor->SetNetConnection(Connection);
 							Connection->OwningActor = NewClientActor;
 							NewClientActor->Role = ROLE_None;
