Compute sub-object guids when owning actor guids are generated

diff --git a/Engine/Source/Runtime/Engine/Private/PackageMapClient.cpp b/Engine/Source/Runtime/Engine/Private/PackageMapClient.cpp
--- a/Engine/Source/Runtime/Engine/Private/PackageMapClient.cpp
+++ b/Engine/Source/Runtime/Engine/Private/PackageMapClient.cpp
@@ -288,13 +288,13 @@ bool UPackageMapClient::SerializeNewActor(FArchive& Ar, class UActorChannel *Cha
 		}
 
 		//
-		// Generate net guids for every subobject that has a stable name, this way they can safely be referenced over the network
+		// Serialize Initially replicating sub-objects here
 		//
 
-		// The idea here is to assume both the server and client will generate the same list when calling GetSubobjectsWithStableNamesForNetworking
-		// We can take advantage of this by only sending the first GUID, and then having the client generate the same guids deterministically from there
-		// The only caveat is that some guids have already been assigned before getting here (they might be referenced, and were pre-assigned)
-		// To get around this, we cherry pick these guids, and send them anyhow, skipping over them when doing the range assign
+		//
+		// We expect that all sub-objects that have stable network names have already been assigned guids, which are relative to their actor owner
+		// Because of this, we can expect the client to deterministically generate the same guids, and we just need to verify with checksum
+		//
 
 		TArray< UObject * > Subobjs;
 
@@ -304,121 +304,55 @@ bool UPackageMapClient::SerializeNewActor(FArchive& Ar, class UActorChannel *Cha
 		{
 			uint32 Checksum = 0;
 
-			uint16 NumSubobjects = Subobjs.Num();
-			Ar << NumSubobjects;
-
-			bool bRangeStartSend = false;
-
-			for ( UObject * Obj : Subobjs )
+			for ( int32 i = 0; i < Subobjs.Num(); i++ )
 			{
-				FNetworkGUID SubobjNetGUID = GuidCache->NetGUIDLookup.FindRef( Obj );
+				UObject * SubObj = Subobjs[i];
 
-				uint8 bForceSendGuid = SubobjNetGUID.IsValid() ? 1 : 0;
+				FNetworkGUID SubobjNetGUID = GuidCache->NetGUIDLookup.FindRef( SubObj );
 
-				if ( !bForceSendGuid )
-				{
-					SubobjNetGUID = GuidCache->GetOrAssignNetGUID( Obj );
-				}
-
-				check( SubobjNetGUID.IsDynamic() );		// Make sure some assumptions are correct
-
-				Ar.SerializeBits( &bForceSendGuid, 1 );
-
-				if ( bForceSendGuid || !bRangeStartSend )
-				{
-					Ar << SubobjNetGUID;				// Save guid if forced, or if we haven't sent the start of the range
-
-					if ( !bForceSendGuid )
-					{
-						// If this is the start of the range, we don't need to send it again
-						bRangeStartSend = true;
-					}
-				}
+				check( SubobjNetGUID.Value == ( ( ( NetGUID.Value >> 1 ) + i + 1 ) << 1 ) );
 
 				// Evolve checksum so we can sanity check
 				Checksum = FCrc::MemCrc32( &SubobjNetGUID.Value, sizeof( SubobjNetGUID.Value ), Checksum );
-				Checksum = FCrc::StrCrc32( *Obj->GetName(), Checksum );
+				Checksum = FCrc::StrCrc32( *SubObj->GetName(), Checksum );
 
 				check( !PendingAckGUIDs.Contains( SubobjNetGUID ) );
-
-				UE_LOG( LogNetPackageMap, Log, TEXT( "SerializeNewActor: [Saving] Assigned NetGUID <%s> to subobject %s" ), *SubobjNetGUID.ToString(), *Obj->GetPathName() );
 			}
 
-			if ( NumSubobjects > 0 )
-			{
-				Ar << Checksum;
-			}
+			Ar << Checksum;
 		}
 		else
 		{
-			uint16 NumSubobjects = 0;
-			Ar << NumSubobjects;
+			uint32 ServerChecksum = 0;
+			uint32 LocalChecksum = 0;
+			Ar << ServerChecksum;
 
-			if ( NumSubobjects > 0 )
+			TArray< FNetworkGUID > SubObjGuids;
+
+			for ( int32 i = 0; i < Subobjs.Num(); i++ )
 			{
-				FNetworkGUID LastRangedNetGUID;
+				SubObjGuids.Add( ( ( ( NetGUID.Value >> 1 ) + i + 1 ) << 1 ) );
 
-				uint32 Checksum = 0;
+				// Evolve checksum so we can sanity check
+				LocalChecksum = FCrc::MemCrc32( &SubObjGuids[i].Value, sizeof( SubObjGuids[i].Value ), LocalChecksum );
+				LocalChecksum = FCrc::StrCrc32( *Subobjs[i]->GetName(), LocalChecksum );
 
-				TArray< FNetworkGUID > SubObjGuids;
+				UE_LOG( LogNetPackageMap, Log, TEXT( "SerializeNewActor: [Loading] Assigned NetGUID <%s> to subobject %s" ), *SubObjGuids[i].ToString(), *Subobjs[i]->GetPathName() );
+			}
 
-				for ( int32 i = 0; i < NumSubobjects; i++ )
+			if ( LocalChecksum != ServerChecksum )
+			{
+				UE_LOG( LogNetPackageMap, Error, TEXT( "SerializeNewActor: Subobject checksum FAILED: %s" ), *Actor->GetFullName() );
+			}
+			else
+			{
+				// Everything worked, assign the guids
+				for ( int32 i = 0; i < Subobjs.Num(); i++ )
 				{
-					FNetworkGUID LocalNetGUID;
-
-					uint8 bForceSendGuid = false;
-					Ar.SerializeBits( &bForceSendGuid, 1 );
-
-					if ( bForceSendGuid || !LastRangedNetGUID.IsValid() )
-					{
-						// Server only saves out the very first guid (or guid that aren't part of the range), we derive the rest
-						Ar << LocalNetGUID;
-					}
-					else
-					{
-						// Derive the next guid from the last one
-						LocalNetGUID.Value = ( ( LastRangedNetGUID.Value >> 1 ) + 1 ) << 1;
-					}
-
-					if ( !bForceSendGuid )
-					{
-						LastRangedNetGUID = LocalNetGUID;
-					}
-
-					SubObjGuids.Add( LocalNetGUID );
-
-					if ( i < Subobjs.Num() )
-					{
-						UObject * Obj = Subobjs[i];
-
-						// Evolve checksum so we can sanity check
-						Checksum = FCrc::MemCrc32( &LocalNetGUID.Value, sizeof( LocalNetGUID.Value ), Checksum );
-						Checksum = FCrc::StrCrc32( *Obj->GetName(), Checksum );
-					}
-				}
+					check( !GuidCache->NetGUIDLookup.Contains( Subobjs[i] ) );		// We should NOT have a guid yet
+					GuidCache->AssignNetGUID( Subobjs[i], SubObjGuids[i] );
 
-				uint32 ServerChecksum = 0;
-				Ar << ServerChecksum;
-
-				// Sanity check the results
-				if ( NumSubobjects != Subobjs.Num() )
-				{
-					UE_LOG( LogNetPackageMap, Error, TEXT( "SerializeNewActor: Invalid number of sub-objects: %s (%i != %i)" ), *Actor->GetFullName(), NumSubobjects, Subobjs.Num() );
-				}
-				else if ( Checksum != ServerChecksum )
-				{
-					UE_LOG( LogNetPackageMap, Error, TEXT( "SerializeNewActor: Subobject checksum FAILED: %s" ), *Actor->GetFullName() );
-				}
-				else
-				{
-					// Everything worked, assign the guids
-					check( SubObjGuids.Num() == Subobjs.Num() );
-
-					for ( int32 i = 0; i < Subobjs.Num(); i++ )
-					{
-						GuidCache->AssignNetGUID( Subobjs[i], SubObjGuids[i] );
-						UE_LOG( LogNetPackageMap, Log, TEXT( "SerializeNewActor: [Loading] Assigned NetGUID <%s> to subobject %s" ), *SubObjGuids[i].ToString(), *Subobjs[i]->GetPathName() );
-					}
+					UE_LOG( LogNetPackageMap, Log, TEXT( "SerializeNewActor: [Loading] Assigned NetGUID <%s> to subobject %s" ), *SubObjGuids[i].ToString(), *Subobjs[i]->GetPathName() );
 				}
 			}
 		}
@@ -1635,6 +1569,25 @@ FNetworkGUID FNetGUIDCache::GetOrAssignNetGUID( const UObject * Object )
 	return AssignNewNetGUID( Object );
 }
 
+static const AActor * GetOuterActor( const UObject * Object )
+{
+	Object = Object->GetOuter();
+
+	while ( Object != NULL )
+	{
+		const AActor * Actor = Cast< const AActor >( Object );
+
+		if ( Actor != NULL )
+		{
+			return Actor;
+		}
+
+		Object = Object->GetOuter();
+	}
+
+	return NULL;
+}
+
 /**
  *	Generate a new NetGUID for this object and assign it.
  */
@@ -1646,6 +1599,31 @@ FNetworkGUID FNetGUIDCache::AssignNewNetGUID( const UObject * Object )
 		return FNetworkGUID::GetDefault();
 	}
 
+	const AActor * Actor = Cast< const AActor >( Object );
+
+	if ( Actor == NULL )
+	{
+		const AActor * OuterActor = GetOuterActor( Object );
+
+		if ( OuterActor != NULL && IsDynamicObject( OuterActor ) )
+		{
+			check( !NetGUIDLookup.Contains( OuterActor ) );
+
+			TArray< UObject * > Subobjs;
+			const_cast< AActor * >( OuterActor )->GetSubobjectsWithStableNamesForNetworking( Subobjs );
+
+			if ( Subobjs.Contains( const_cast< UObject * >( Object ) ) )
+			{
+				// Assign our owning actor a guid first (which should also assign our net guid if things are working correctly)
+				AssignNewNetGUID( OuterActor );
+
+				return NetGUIDLookup.FindChecked( Object );
+			}
+		}
+	}
+
+#define GET_NET_GUID_INDEX( GUID )				( GUID.Value >> 1 )
+#define GET_NET_GUID_STATIC_MOD( GUID )			( GUID.Value & 1 )
 #define COMPOSE_NET_GUID( Index, IsStatic )		( ( ( Index ) << 1 ) | ( IsStatic ) )
 #define ALLOC_NEW_NET_GUID( IsStatic )			( COMPOSE_NET_GUID( ++UniqueNetIDs[ IsStatic ], IsStatic ) )
 
@@ -1656,6 +1634,31 @@ FNetworkGUID FNetGUIDCache::AssignNewNetGUID( const UObject * Object )
 
 	AssignNetGUID( Object, NewNetGuid );
 
+	// If this is a dynamic actor, assign all of our stably named sub-objects their net guids now as well (they come as a group)
+	if ( Actor != NULL && !IsStatic )
+	{
+		TArray< UObject * > Subobjs;
+
+		const_cast< AActor * >( Actor )->GetSubobjectsWithStableNamesForNetworking( Subobjs );
+
+		for ( int32 i = 0; i < Subobjs.Num(); i++ )
+		{
+			UObject * SubObj = Subobjs[i];
+
+			check( GetOuterActor( SubObj ) == Actor );
+
+			if ( NetGUIDLookup.Contains( SubObj ) )		// We should NOT have a guid yet
+			{
+				UE_LOG( LogNetPackageMap, Error, TEXT( "AssignNewNetGUID: Sub object already registered: Actor: %s, SubObj: %s, NetGuid: %s" ), *Actor->GetPathName(), *SubObj->GetPathName(), *NetGUIDLookup.FindChecked( SubObj ).ToString() );
+			}
+			const int32 SubIsStatic = IsDynamicObject( SubObj ) ? 0 : 1;
+			check( SubIsStatic == IsStatic );
+			const FNetworkGUID SubobjNetGUID( ALLOC_NEW_NET_GUID( SubIsStatic ) );
+			check( SubobjNetGUID.Value == COMPOSE_NET_GUID( GET_NET_GUID_INDEX( NewNetGuid ) + 1 + i, GET_NET_GUID_STATIC_MOD( NewNetGuid ) ) );
+			AssignNetGUID( SubObj, SubobjNetGUID );
+		}
+	}
+
 	return NewNetGuid;
 }
 
