Maintain attachments to child actor component through reconstruction

diff --git a/Engine/Source/Runtime/Engine/Classes/Components/ChildActorComponent.h b/Engine/Source/Runtime/Engine/Classes/Components/ChildActorComponent.h
--- a/Engine/Source/Runtime/Engine/Classes/Components/ChildActorComponent.h
+++ b/Engine/Source/Runtime/Engine/Classes/Components/ChildActorComponent.h
@@ -4,6 +4,8 @@
 #pragma once
 #include "ChildActorComponent.generated.h"
 
+class FChildActorComponentInstanceData;
+
 /** A component that spawns an Actor when registered, and destroys it when unregistered.*/
 UCLASS(ClassGroup=Utility, hidecategories=(Object,LOD,Physics,Lighting,TextureStreaming,Activation,"Components|Activation",Collision), meta=(BlueprintSpawnableComponent), MinimalAPI)
 class UChildActorComponent : public USceneComponent
@@ -21,6 +23,9 @@ class UChildActorComponent : public USceneComponent
 	/** We try to keep the child actor's name as best we can, so we store it off here when destroying */
 	FName ChildActorName;
 
+	/** Cached copy of the instance data when the ChildActor is destroyed to be available when needed */
+	TSharedPtr<FChildActorComponentInstanceData> CachedInstanceData;
+
 	// Begin ActorComponent interface.
 	virtual void OnComponentCreated() OVERRIDE;
 	virtual void OnComponentDestroyed() OVERRIDE;
