Missed file for CL#2093698

diff --git a/Engine/Source/Runtime/Engine/Private/ChildActorComponent.cpp b/Engine/Source/Runtime/Engine/Private/ChildActorComponent.cpp
--- a/Engine/Source/Runtime/Engine/Private/ChildActorComponent.cpp
+++ b/Engine/Source/Runtime/Engine/Private/ChildActorComponent.cpp
@@ -31,9 +31,37 @@ public:
 		: FComponentInstanceDataBase(Component)
 		, ChildActorName(Component->ChildActorName)
 	{
+		if (Component->ChildActor)
+		{
+			USceneComponent* ChildRootComponent = Component->ChildActor->GetRootComponent();
+			for (USceneComponent* AttachedComponent : ChildRootComponent->AttachChildren)
+			{
+				if (AttachedComponent)
+				{
+					AActor* AttachedActor = AttachedComponent->GetOwner();
+					if (AttachedActor != Component->ChildActor)
+					{
+						FAttachedActorInfo Info;
+						Info.Actor = AttachedActor;
+						Info.SocketName = AttachedComponent->AttachSocketName;
+						Info.RelativeTransform = AttachedComponent->GetRelativeTransform();
+						AttachedActors.Add(Info);
+					}
+				}
+			}
+		}
 	}
 
 	FName ChildActorName;
+
+	struct FAttachedActorInfo
+	{
+		TWeakObjectPtr<AActor> Actor;
+		FName SocketName;
+		FTransform RelativeTransform;
+	};
+
+	TArray<FAttachedActorInfo> AttachedActors;
 };
 
 FName UChildActorComponent::GetComponentInstanceDataType() const
@@ -44,32 +72,48 @@ FName UChildActorComponent::GetComponentInstanceDataType() const
 
 TSharedPtr<FComponentInstanceDataBase> UChildActorComponent::GetComponentInstanceData() const
 {
-	TSharedPtr<FChildActorComponentInstanceData> InstanceData;
-	if (!ChildActorName.IsNone())
-	{
-		// Allocate new struct for holding light map data
-		 InstanceData = MakeShareable(new FChildActorComponentInstanceData(this));
-	}
-
-	return InstanceData;
+	return (CachedInstanceData.IsValid() ? CachedInstanceData : MakeShareable(new FChildActorComponentInstanceData(this)));
 }
 
 void UChildActorComponent::ApplyComponentInstanceData(TSharedPtr<FComponentInstanceDataBase> ComponentInstanceData)
 {
 	check(ComponentInstanceData.IsValid());
-	ChildActorName = StaticCastSharedPtr<FChildActorComponentInstanceData>(ComponentInstanceData)->ChildActorName;
+	FChildActorComponentInstanceData* ChildActorInstanceData = StaticCastSharedPtr<FChildActorComponentInstanceData>(ComponentInstanceData).Get();
+	ChildActorName = ChildActorInstanceData->ChildActorName;
 	if (ChildActor)
 	{
 		ChildActor->Rename(*ChildActorName.ToString(), NULL, REN_DoNotDirty);
+
+		USceneComponent* ChildActorRoot = ChildActor->GetRootComponent();
+		if (ChildActorRoot)
+		{
+			for (const auto& AttachInfo : ChildActorInstanceData->AttachedActors)
+			{
+				AActor* AttachedActor = AttachInfo.Actor.Get();
+				if (AttachedActor)
+				{
+					USceneComponent* AttachedRootComponent = AttachedActor->GetRootComponent();
+					if (AttachedRootComponent)
+					{
+						AttachedActor->DetachRootComponentFromParent();
+						AttachedRootComponent->AttachTo(ChildActorRoot, AttachInfo.SocketName, EAttachLocation::KeepWorldPosition);
+						AttachedRootComponent->SetRelativeTransform(AttachInfo.RelativeTransform);
+						AttachedRootComponent->UpdateComponentToWorld();
+					}
+				}
+			}
+		}
 	}
 }
 
-
 void UChildActorComponent::CreateChildActor()
 {
 	// Kill spawned actor if we have one
 	DestroyChildActor();
 
+	// This is no longer needed
+	CachedInstanceData = NULL;
+
 	// If we have a class to spawn.
 	if(ChildActorClass != NULL)
 	{
@@ -91,38 +135,38 @@ void UChildActorComponent::CreateChildActor()
 
 			if (bSpawn)
 			{
-			FActorSpawnParameters Params;
-			Params.bNoCollisionFail = true;
-			Params.bDeferConstruction = true; // We defer construction so that we set ParentComponentActor prior to component registration so they appear selected
-			Params.bAllowDuringConstructionScript = true;
-			Params.OverrideLevel = GetOwner()->GetLevel();
-			Params.Name = ChildActorName;
-			Params.ObjectFlags &= ~RF_Transactional;
-
-			// Spawn actor of desired class
-			FVector Location = GetComponentLocation();
-			FRotator Rotation = GetComponentRotation();
-			ChildActor = World->SpawnActor(ChildActorClass, &Location, &Rotation, Params);
-
-			// If spawn was successful, 
-			if(ChildActor != NULL)
-			{
-				ChildActorName = ChildActor->GetFName();
+				FActorSpawnParameters Params;
+				Params.bNoCollisionFail = true;
+				Params.bDeferConstruction = true; // We defer construction so that we set ParentComponentActor prior to component registration so they appear selected
+				Params.bAllowDuringConstructionScript = true;
+				Params.OverrideLevel = GetOwner()->GetLevel();
+				Params.Name = ChildActorName;
+				Params.ObjectFlags &= ~RF_Transactional;
+
+				// Spawn actor of desired class
+				FVector Location = GetComponentLocation();
+				FRotator Rotation = GetComponentRotation();
+				ChildActor = World->SpawnActor(ChildActorClass, &Location, &Rotation, Params);
+
+				// If spawn was successful, 
+				if(ChildActor != NULL)
+				{
+					ChildActorName = ChildActor->GetFName();
 
-				// Remember which actor spawned it (for selection in editor etc)
-				ChildActor->ParentComponentActor = GetOwner();
+					// Remember which actor spawned it (for selection in editor etc)
+					ChildActor->ParentComponentActor = GetOwner();
 
-				// Parts that we deferred from SpawnActor
-				ChildActor->ExecuteConstruction(ComponentToWorld, NULL);
-				ChildActor->PostActorConstruction();
+					// Parts that we deferred from SpawnActor
+					ChildActor->ExecuteConstruction(ComponentToWorld, NULL);
+					ChildActor->PostActorConstruction();
 
-				// attach new actor to this component
-				ChildActor->AttachRootComponentTo(this, NAME_None, EAttachLocation::SnapToTarget);
+					// attach new actor to this component
+					ChildActor->AttachRootComponentTo(this, NAME_None, EAttachLocation::SnapToTarget);
+				}
 			}
 		}
 	}
 }
-}
 
 void UChildActorComponent::DestroyChildActor()
 {
@@ -132,6 +176,8 @@ void UChildActorComponent::DestroyChildActor()
 		// if still alive, destroy, otherwise just clear the pointer
 		if(!ChildActor->IsPendingKill())
 		{
+			CachedInstanceData = MakeShareable(new FChildActorComponentInstanceData(this));
+
 			UWorld* World = ChildActor->GetWorld();
 			// World may be NULL during shutdown
 			if(World != NULL)
