#ue4 - Various console commands to control framerate now work on consoles.

- Now supported: t.MaxFPS, t.UnsteadyFPS, CauseHitches

TTP 331313

diff --git a/Engine/Source/Runtime/Engine/Private/GameEngine.cpp b/Engine/Source/Runtime/Engine/Private/GameEngine.cpp
--- a/Engine/Source/Runtime/Engine/Private/GameEngine.cpp
+++ b/Engine/Source/Runtime/Engine/Private/GameEngine.cpp
@@ -786,13 +786,13 @@ float UGameEngine::GetMaxTickRate( float DeltaTime, bool bAllowFrameRateSmoothin
 				}
 			}*/
 		}
+	}
 
-		// See if the code in the base class wants to replace this
-		float SuperTickRate = Super::GetMaxTickRate(DeltaTime, bAllowFrameRateSmoothing);
-		if(SuperTickRate != 0.0)
-		{
-			MaxTickRate = SuperTickRate;
-		}
+	// See if the code in the base class wants to replace this
+	float SuperTickRate = Super::GetMaxTickRate(DeltaTime, bAllowFrameRateSmoothing);
+	if(SuperTickRate != 0.0)
+	{
+		MaxTickRate = SuperTickRate;
 	}
 
 	return MaxTickRate;
diff --git a/Engine/Source/Runtime/Engine/Private/UnrealEngine.cpp b/Engine/Source/Runtime/Engine/Private/UnrealEngine.cpp
--- a/Engine/Source/Runtime/Engine/Private/UnrealEngine.cpp
+++ b/Engine/Source/Runtime/Engine/Private/UnrealEngine.cpp
@@ -5459,11 +5459,11 @@ static TAutoConsoleVariable<float> CVarMaxFPS(
 // CauseHitches cvar
 static TAutoConsoleVariable<int32> CVarCauseHitches(
 	TEXT("CauseHitches"),0,
-	TEXT("Causes a 100ms hitch every second."));
+	TEXT("Causes a 200ms hitch every second."));
 
 static TAutoConsoleVariable<int32> CVarUnsteadyFPS(
 	TEXT("t.UnsteadyFPS"),0,
-	TEXT("Causes FPS to bounce around randomly in the 10-30 range."));
+	TEXT("Causes FPS to bounce around randomly in the 8-32 range."));
 
 /** Get tick rate limitor. */
 float UEngine::GetMaxTickRate( float DeltaTime, bool bAllowFrameRateSmoothing )
