Timer manager updates.

Added support for identifying timers by handle. Timers with a handle do not require a delegate but may still have one.
Added support for timers that execute on the next tick.
Added automation tests.

diff --git a/Engine/Source/Runtime/Engine/Private/TimerManager.cpp b/Engine/Source/Runtime/Engine/Private/TimerManager.cpp
--- a/Engine/Source/Runtime/Engine/Private/TimerManager.cpp
+++ b/Engine/Source/Runtime/Engine/Private/TimerManager.cpp
@@ -6,6 +6,19 @@
 
 #include "EnginePrivate.h"
 
+void FTimerHandle::MakeValid()
+{
+	static int LastAssignedHandle = -1;
+
+	if (!IsValid())
+	{
+		++LastAssignedHandle;
+		Handle = LastAssignedHandle;
+	}
+
+	check(IsValid());
+}
+
 
 // ---------------------------------
 // Private members
@@ -38,9 +51,47 @@ FTimerData const* FTimerManager::FindTimer(FTimerUnifiedDelegate const& InDelega
 
 	int32 PendingTimerIdx = FindTimerInList(PendingTimerList, InDelegate);
 	if (PendingTimerIdx != INDEX_NONE)
+	{
+		// found it in the pending list
+		if (OutTimerIndex)
+		{
+			*OutTimerIndex = PendingTimerIdx;
+		}
+		return &PendingTimerList[PendingTimerIdx];
+	}
+
+	return NULL;
+}
+
+FTimerData const* FTimerManager::FindTimer(FTimerHandle const& InHandle, int32* OutTimerIndex) const
+{
+	int32 ActiveTimerIdx = FindTimerInList(ActiveTimerHeap, InHandle);
+	if (ActiveTimerIdx != INDEX_NONE)
+	{
+		// found it in the active heap
+		if (OutTimerIndex)
+		{
+			*OutTimerIndex = ActiveTimerIdx;
+		}
+		return &ActiveTimerHeap[ActiveTimerIdx];
+	}
+
+	int32 PausedTimerIdx = FindTimerInList(PausedTimerList, InHandle);
+	if (PausedTimerIdx != INDEX_NONE)
 	{
 		// found it in the paused list
 		if (OutTimerIndex)
+		{
+			*OutTimerIndex = PausedTimerIdx;
+		}
+		return &PausedTimerList[PausedTimerIdx];
+	}
+
+	int32 PendingTimerIdx = FindTimerInList(PendingTimerList, InHandle);
+	if (PendingTimerIdx != INDEX_NONE)
+	{
+		// found it in the pending list
+		if (OutTimerIndex)
 		{
 			*OutTimerIndex = PendingTimerIdx;
 		}
@@ -56,7 +107,7 @@ int32 FTimerManager::FindTimerInList(const TArray<FTimerData> &SearchArray, FTim
 {
 	for (int32 Idx=0; Idx<SearchArray.Num(); ++Idx)
 	{
-		if (SearchArray[Idx].TimerDelegate == InDelegate)
+		if (SearchArray[Idx].TimerDelegate == InDelegate && !SearchArray[Idx].TimerHandle.IsValid())
 		{
 			return Idx;
 		}
@@ -65,7 +116,26 @@ int32 FTimerManager::FindTimerInList(const TArray<FTimerData> &SearchArray, FTim
 	return INDEX_NONE;
 }
 
-void FTimerManager::InternalSetTimer(FTimerUnifiedDelegate const& InDelegate, float InRate, bool InbLoop)
+/** Will find the given timer in the given TArray and return its index. */
+int32 FTimerManager::FindTimerInList(const TArray<FTimerData> &SearchArray, FTimerHandle const& InHandle) const
+{
+	if (!(InHandle.IsValid()))
+	{
+		return INDEX_NONE;
+	}
+
+	for (int32 Idx = 0; Idx < SearchArray.Num(); ++Idx)
+	{
+		if (SearchArray[Idx].TimerHandle == InHandle)
+		{
+			return Idx;
+		}
+	}
+
+	return INDEX_NONE;
+}
+
+void FTimerManager::InternalSetTimer(FTimerUnifiedDelegate const& InDelegate, float InRate, bool InbLoop, float InFirstDelay)
 {
 	// not currently threadsafe
 	check(IsInGameThread());
@@ -78,48 +148,125 @@ void FTimerManager::InternalSetTimer(FTimerUnifiedDelegate const& InDelegate, fl
 	{
 		// set up the new timer
 		FTimerData NewTimerData;
-		NewTimerData.Rate = InRate;
-		NewTimerData.bLoop = InbLoop;
 		NewTimerData.TimerDelegate = InDelegate;
-		
-		if( HasBeenTickedThisFrame() )
+
+		InternalSetTimer(NewTimerData, InRate, InbLoop, InFirstDelay);
+	}
+}
+
+void FTimerManager::InternalSetTimer(FTimerHandle& InOutHandle, FTimerUnifiedDelegate const& InDelegate, float InRate, bool InbLoop, float InFirstDelay)
+{
+	// not currently threadsafe
+	check(IsInGameThread());
+
+	if (InOutHandle.IsValid())
+	{
+		// if the timer is already set, just clear it and we'll re-add it, since 
+		// there's no data to maintain.
+		InternalClearTimer(InOutHandle);
+	}
+
+	if (InRate > 0.f)
+	{
+		InOutHandle.MakeValid();
+
+		// set up the new timer
+		FTimerData NewTimerData;
+		NewTimerData.TimerHandle = InOutHandle;
+		NewTimerData.TimerDelegate = InDelegate;
+
+		InternalSetTimer(NewTimerData, InRate, InbLoop, InFirstDelay);
+	}
+}
+
+void FTimerManager::InternalSetTimer(FTimerData& NewTimerData, float InRate, bool InbLoop, float InFirstDelay)
+{
+	check(NewTimerData.TimerDelegate.IsBound() || NewTimerData.TimerHandle.IsValid());
+
+	NewTimerData.Rate = InRate;
+	NewTimerData.bLoop = InbLoop;
+
+	float FirstDelay = (InFirstDelay >= 0.f) ? InFirstDelay : InRate;
+
+	if (HasBeenTickedThisFrame())
+	{
+		NewTimerData.ExpireTime = InternalTime + FirstDelay;
+		NewTimerData.Status = ETimerStatus::Active;
+		ActiveTimerHeap.HeapPush(NewTimerData);
+	}
+	else
+	{
+		// Store time remaining in ExpireTime while pending
+		NewTimerData.ExpireTime = FirstDelay;
+		NewTimerData.Status = ETimerStatus::Pending;
+		PendingTimerList.Add(NewTimerData);
+	}
+}
+
+void FTimerManager::InternalSetTimerForNextTick(FTimerUnifiedDelegate const& InDelegate)
+{
+	// not currently threadsafe
+	check(IsInGameThread());
+
+	FTimerData NewTimerData;
+	NewTimerData.Rate = 0.f;
+	NewTimerData.bLoop = false;
+	NewTimerData.TimerDelegate = InDelegate;
+	NewTimerData.ExpireTime = InternalTime;
+	NewTimerData.Status = ETimerStatus::Active;
+	ActiveTimerHeap.HeapPush(NewTimerData);
+}
+
+void FTimerManager::InternalClearTimer(FTimerUnifiedDelegate const& InDelegate)
+{
+	// not currently threadsafe
+	check(IsInGameThread());
+
+	int32 TimerIdx;
+	FTimerData const* const TimerData = FindTimer(InDelegate, &TimerIdx);
+	if (TimerData)
+	{
+		switch (TimerData->Status)
 		{
-			NewTimerData.ExpireTime = InternalTime + InRate;
-			NewTimerData.Status = ETimerStatus::Active;
-			ActiveTimerHeap.HeapPush(NewTimerData);
+		case ETimerStatus::Pending: PendingTimerList.RemoveAtSwap(TimerIdx); break;
+		case ETimerStatus::Active: ActiveTimerHeap.HeapRemoveAt(TimerIdx); break;
+		case ETimerStatus::Paused: PausedTimerList.RemoveAtSwap(TimerIdx); break;
+		default: check(false);
 		}
-		else
+	}
+	else
+	{
+		// Edge case. We're currently handling this timer when it got cleared.  Unbind it to prevent it firing again
+		// in case it was scheduled to fire multiple times.
+		if (CurrentlyExecutingTimer.TimerDelegate == InDelegate)
 		{
-			// Store time remaining in ExpireTime while pending
-			NewTimerData.ExpireTime = InRate;
-			NewTimerData.Status = ETimerStatus::Pending;
-			PendingTimerList.Add(NewTimerData);
+			CurrentlyExecutingTimer.TimerDelegate.Unbind();
 		}
 	}
 }
 
-void FTimerManager::InternalClearTimer(FTimerUnifiedDelegate const& InDelegate)
+void FTimerManager::InternalClearTimer(FTimerHandle const& InHandle)
 {
 	// not currently threadsafe
 	check(IsInGameThread());
 
 	int32 TimerIdx;
-	FTimerData const* const TimerData = FindTimer(InDelegate, &TimerIdx);
-	if( TimerData )
+	FTimerData const* const TimerData = FindTimer(InHandle, &TimerIdx);
+	if (TimerData)
 	{
-		switch( TimerData->Status )
+		switch (TimerData->Status)
 		{
-			case ETimerStatus::Pending : PendingTimerList.RemoveAtSwap(TimerIdx); break;
-			case ETimerStatus::Active : ActiveTimerHeap.HeapRemoveAt(TimerIdx); break;
-			case ETimerStatus::Paused : PausedTimerList.RemoveAtSwap(TimerIdx); break;
-			default : check(false);
+		case ETimerStatus::Pending: PendingTimerList.RemoveAtSwap(TimerIdx); break;
+		case ETimerStatus::Active: ActiveTimerHeap.HeapRemoveAt(TimerIdx); break;
+		case ETimerStatus::Paused: PausedTimerList.RemoveAtSwap(TimerIdx); break;
+		default: check(false);
 		}
 	}
 	else
 	{
 		// Edge case. We're currently handling this timer when it got cleared.  Unbind it to prevent it firing again
 		// in case it was scheduled to fire multiple times.
-		if (CurrentlyExecutingTimer.TimerDelegate == InDelegate)
+		if (CurrentlyExecutingTimer.TimerHandle == InHandle)
 		{
 			CurrentlyExecutingTimer.TimerDelegate.Unbind();
 		}
@@ -171,9 +318,8 @@ void FTimerManager::InternalClearAllTimers(void const* Object)
 
 
 
-float FTimerManager::InternalGetTimerRemaining(FTimerUnifiedDelegate const& InDelegate) const
+float FTimerManager::InternalGetTimerRemaining(FTimerData const* const TimerData) const
 {
-	FTimerData const* const TimerData = FindTimer(InDelegate);
 	if (TimerData)
 	{
 		if( TimerData->Status != ETimerStatus::Active )
@@ -190,9 +336,8 @@ float FTimerManager::InternalGetTimerRemaining(FTimerUnifiedDelegate const& InDe
 	return -1.f;
 }
 
-float FTimerManager::InternalGetTimerElapsed(FTimerUnifiedDelegate const& InDelegate) const
+float FTimerManager::InternalGetTimerElapsed(FTimerData const* const TimerData) const
 {
-	FTimerData const* const TimerData = FindTimer(InDelegate);
 	if (TimerData)
 	{
 		if( TimerData->Status != ETimerStatus::Active)
@@ -209,9 +354,8 @@ float FTimerManager::InternalGetTimerElapsed(FTimerUnifiedDelegate const& InDele
 	return -1.f;
 }
 
-float FTimerManager::InternalGetTimerRate(FTimerUnifiedDelegate const& InDelegate) const
+float FTimerManager::InternalGetTimerRate(FTimerData const* const TimerData) const
 {
-	FTimerData const* const TimerData = FindTimer(InDelegate);
 	if (TimerData)
 	{
 		return TimerData->Rate;
@@ -219,13 +363,11 @@ float FTimerManager::InternalGetTimerRate(FTimerUnifiedDelegate const& InDelegat
 	return -1.f;
 }
 
-void FTimerManager::InternalPauseTimer(FTimerUnifiedDelegate const& InDelegate)
+void FTimerManager::InternalPauseTimer( FTimerData const* TimerToPause, int32 TimerIdx )
 {
 	// not currently threadsafe
 	check(IsInGameThread());
 
-	int32 TimerIdx;
-	FTimerData const* TimerToPause = FindTimer(InDelegate, &TimerIdx);
 	if( TimerToPause && (TimerToPause->Status != ETimerStatus::Paused) )
 	{
 		ETimerStatus::Type PreviousStatus = TimerToPause->Status;
@@ -255,12 +397,11 @@ void FTimerManager::InternalPauseTimer(FTimerUnifiedDelegate const& InDelegate)
 	}
 }
 
-void FTimerManager::InternalUnPauseTimer(FTimerUnifiedDelegate const& InDelegate)
+void FTimerManager::InternalUnPauseTimer(int32 PausedTimerIdx)
 {
 	// not currently threadsafe
 	check(IsInGameThread());
 
-	int32 PausedTimerIdx = FindTimerInList(PausedTimerList, InDelegate);
 	if (PausedTimerIdx != INDEX_NONE)
 	{
 		FTimerData& TimerToUnPause = PausedTimerList[PausedTimerIdx];
@@ -313,19 +454,21 @@ void FTimerManager::Tick(float DeltaTime)
 
 			// Now call the function
 			for (int32 CallIdx=0; CallIdx<CallCount; ++CallIdx)
-			{
+			{ 
 				CurrentlyExecutingTimer.TimerDelegate.Execute();
 
 				// If timer was cleared in the delegate execution, don't execute further 
-				if( !CurrentlyExecutingTimer.TimerDelegate.IsBound() )
+				if( !CurrentlyExecutingTimer.TimerDelegate.IsBound() && !CurrentlyExecutingTimer.TimerHandle.IsValid() )
 				{
 					break;
 				}
 			}
 
 			if( CurrentlyExecutingTimer.bLoop && 
-				CurrentlyExecutingTimer.TimerDelegate.IsBound() && 							// did not get cleared during execution
-				(FindTimer(CurrentlyExecutingTimer.TimerDelegate) == NULL)		// did not get manually re-added during execution
+				(CurrentlyExecutingTimer.TimerDelegate.IsBound() ||
+				 CurrentlyExecutingTimer.TimerHandle.IsValid()) && 							// did not get cleared during execution
+				((FindTimer(CurrentlyExecutingTimer.TimerDelegate) == NULL) || 
+				 (CurrentlyExecutingTimer.TimerHandle.IsValid() && FindTimer(CurrentlyExecutingTimer.TimerHandle) == NULL))		// did not get manually re-added during execution
 			  )
 			{
 				// Put this timer back on the heap
diff --git a/Engine/Source/Runtime/Engine/Private/TimerManagerTests.cpp b/Engine/Source/Runtime/Engine/Private/TimerManagerTests.cpp
new file mode 100644
--- /dev/null
+++ b/Engine/Source/Runtime/Engine/Private/TimerManagerTests.cpp
@@ -0,0 +1,415 @@
+// Copyright 1998-2014 Epic Games, Inc. All Rights Reserved.
+
+#include "EnginePrivate.h"
+#include "TimerManager.h"
+
+IMPLEMENT_SIMPLE_AUTOMATION_TEST(FTimerManagerTest, "Engine.TimerManager", EAutomationTestFlags::ATF_Editor)
+
+#define TIMER_TEST_TEXT( Format, ... ) FString::Printf(TEXT("%s - %d: %s"), TEXT(__FILE__) , __LINE__ , *FString::Printf(TEXT(Format), ##__VA_ARGS__) )
+
+void TimerTest_TickWorld(UWorld* World, float Time = 1.f)
+{
+	const float Step = 0.1f;
+	while(Time > 0.f)
+	{
+		World->Tick(ELevelTick::LEVELTICK_All, FMath::Min(Time, Step) );
+		Time-=Step;
+
+		// This is terrible but required for subticking like this.
+		// we could always cache the real GFrameCounter at the start of our tests and restore it when finished.
+		GFrameCounter++;
+	}
+}
+
+class FDummy
+{
+public:
+	FDummy() { Count = 0; }
+
+	void Callback() { ++Count; }
+	void Reset() { Count = 0; }
+
+	uint8 Count;
+};
+
+// Make sure that the timer manager works as expected when given invalid delegates and handles
+bool TimerManagerTest_InvalidTimers(UWorld* World, FAutomationTestBase* Test)
+{
+	FTimerManager& TimerManager = World->GetTimerManager();
+	FTimerDelegate Delegate;
+	FTimerHandle Handle;
+
+	Test->TestFalse(TIMER_TEST_TEXT("TimerExists called with an invalid delegate"), TimerManager.TimerExists(Delegate));
+	Test->TestFalse(TIMER_TEST_TEXT("TimerExists called with an invalid handle"), TimerManager.TimerExists(Handle));
+
+	Test->TestFalse(TIMER_TEST_TEXT("IsTimerActive called with an invalid delegate"), TimerManager.IsTimerActive(Delegate));
+	Test->TestFalse(TIMER_TEST_TEXT("IsTimerActive called with an invalid handle"), TimerManager.IsTimerActive(Handle));
+
+	Test->TestFalse(TIMER_TEST_TEXT("IsTimerPaused called with an invalid delegate"), TimerManager.IsTimerPaused(Delegate));
+	Test->TestFalse(TIMER_TEST_TEXT("IsTimerPaused called with an invalid handle"), TimerManager.IsTimerPaused(Handle));
+
+	Test->TestTrue(TIMER_TEST_TEXT("GetTimerRate called with an invalid delegate"), (TimerManager.GetTimerRate(Delegate) == -1.f));
+	Test->TestTrue(TIMER_TEST_TEXT("GetTimerRate called with an invalid handle"), (TimerManager.GetTimerRate(Handle) == -1.f));
+
+	Test->TestTrue(TIMER_TEST_TEXT("GetTimerElapsed called with an invalid delegate"), (TimerManager.GetTimerElapsed(Delegate) == -1.f));
+	Test->TestTrue(TIMER_TEST_TEXT("GetTimerElapsed called with an invalid handle"), (TimerManager.GetTimerElapsed(Handle) == -1.f));
+
+	Test->TestTrue(TIMER_TEST_TEXT("GetTimerRemaining called with an invalid delegate"), (TimerManager.GetTimerRemaining(Delegate) == -1.f));
+	Test->TestTrue(TIMER_TEST_TEXT("GetTimerRemaining called with an invalid handle"), (TimerManager.GetTimerRemaining(Handle) == -1.f));
+
+	// these don't return values but we should run them to make sure they don't do something horrible like crash
+	TimerManager.PauseTimer(Delegate);
+	TimerManager.PauseTimer(Handle);
+
+	TimerManager.UnPauseTimer(Delegate);
+	TimerManager.UnPauseTimer(Handle);
+
+	TimerManager.ClearTimer(Delegate);
+	TimerManager.ClearTimer(Handle);
+	
+	return true;
+}
+
+// Make sure that the timer manager works as expected when given delegates and handles that aren't in the timer manager
+bool TimerManagerTest_MissingTimers(UWorld* World, FAutomationTestBase* Test)
+{
+	FTimerManager& TimerManager = World->GetTimerManager();
+	FTimerDelegate Delegate;
+	FTimerHandle Handle;
+
+	FDummy Dummy;
+	Delegate.BindRaw(&Dummy, &FDummy::Callback);
+	Handle.MakeValid();
+
+	Test->TestFalse(TIMER_TEST_TEXT("TimerExists called with an invalid delegate"), TimerManager.TimerExists(Delegate));
+	Test->TestFalse(TIMER_TEST_TEXT("TimerExists called with an invalid handle"), TimerManager.TimerExists(Handle));
+
+	Test->TestFalse(TIMER_TEST_TEXT("IsTimerActive called with an invalid delegate"), TimerManager.IsTimerActive(Delegate));
+	Test->TestFalse(TIMER_TEST_TEXT("IsTimerActive called with an invalid handle"), TimerManager.IsTimerActive(Handle));
+
+	Test->TestFalse(TIMER_TEST_TEXT("IsTimerPaused called with an invalid delegate"), TimerManager.IsTimerPaused(Delegate));
+	Test->TestFalse(TIMER_TEST_TEXT("IsTimerPaused called with an invalid handle"), TimerManager.IsTimerPaused(Handle));
+
+	Test->TestTrue(TIMER_TEST_TEXT("GetTimerRate called with an invalid delegate"), (TimerManager.GetTimerRate(Delegate) == -1.f));
+	Test->TestTrue(TIMER_TEST_TEXT("GetTimerRate called with an invalid handle"), (TimerManager.GetTimerRate(Handle) == -1.f));
+
+	Test->TestTrue(TIMER_TEST_TEXT("GetTimerElapsed called with an invalid delegate"), (TimerManager.GetTimerElapsed(Delegate) == -1.f));
+	Test->TestTrue(TIMER_TEST_TEXT("GetTimerElapsed called with an invalid handle"), (TimerManager.GetTimerElapsed(Handle) == -1.f));
+
+	Test->TestTrue(TIMER_TEST_TEXT("GetTimerRemaining called with an invalid delegate"), (TimerManager.GetTimerRemaining(Delegate) == -1.f));
+	Test->TestTrue(TIMER_TEST_TEXT("GetTimerRemaining called with an invalid handle"), (TimerManager.GetTimerRemaining(Handle) == -1.f));
+
+	// these don't return values but we should run them to make sure they don't do something horrible like crash
+	TimerManager.PauseTimer(Delegate);
+	TimerManager.PauseTimer(Handle);
+
+	TimerManager.UnPauseTimer(Delegate);
+	TimerManager.UnPauseTimer(Handle);
+
+	TimerManager.ClearTimer(Delegate);
+	TimerManager.ClearTimer(Handle);
+
+	return true;
+}
+
+bool TimerManagerTest_ValidTimer_Delegate(UWorld* World, FAutomationTestBase* Test)
+{
+	FTimerManager& TimerManager = World->GetTimerManager();
+	FTimerDelegate Delegate;
+	const float Rate = 1.5f;
+
+	FDummy Dummy;
+
+	Delegate.BindRaw(&Dummy, &FDummy::Callback);
+
+	TimerManager.SetTimer(Delegate, Rate, false);
+
+	Test->TestTrue(TIMER_TEST_TEXT("TimerExists called with a pending timer"), TimerManager.TimerExists(Delegate));
+	Test->TestTrue(TIMER_TEST_TEXT("IsTimerActive called with a pending timer"), TimerManager.IsTimerActive(Delegate));
+	Test->TestFalse(TIMER_TEST_TEXT("IsTimerPaused called with a pending timer"), TimerManager.IsTimerPaused(Delegate));
+	Test->TestTrue(TIMER_TEST_TEXT("GetTimerRate called with a pending timer"), (TimerManager.GetTimerRate(Delegate) == Rate));
+	Test->TestTrue(TIMER_TEST_TEXT("GetTimerElapsed called with a pending timer"), (TimerManager.GetTimerElapsed(Delegate) == 0.f));
+	Test->TestTrue(TIMER_TEST_TEXT("GetTimerRemaining called with a pending timer"), (TimerManager.GetTimerRemaining(Delegate) == Rate));
+
+	// small tick to move the timer from the pending list to the active list, the timer will start counting time after this tick
+	TimerTest_TickWorld(World, KINDA_SMALL_NUMBER);
+
+	Test->TestTrue(TIMER_TEST_TEXT("TimerExists called with a pending timer"), TimerManager.TimerExists(Delegate));
+	Test->TestTrue(TIMER_TEST_TEXT("IsTimerActive called with an active timer"), TimerManager.IsTimerActive(Delegate));
+	Test->TestFalse(TIMER_TEST_TEXT("IsTimerPaused called with an active timer"), TimerManager.IsTimerPaused(Delegate));
+
+	TimerTest_TickWorld(World);
+
+	Test->TestTrue(TIMER_TEST_TEXT("GetTimerElapsed called with an active timer after one step"), (FMath::IsNearlyEqual(TimerManager.GetTimerElapsed(Delegate), 1.f, KINDA_SMALL_NUMBER)));
+	Test->TestTrue(TIMER_TEST_TEXT("GetTimerRemaining called with an active timer after one step"), (FMath::IsNearlyEqual(TimerManager.GetTimerRemaining(Delegate), Rate - 1.f, KINDA_SMALL_NUMBER)));
+
+	TimerManager.PauseTimer(Delegate);
+
+	Test->TestTrue(TIMER_TEST_TEXT("TimerExists called with a paused timer"), TimerManager.TimerExists(Delegate));
+	Test->TestFalse(TIMER_TEST_TEXT("IsTimerActive called with a paused timer"), TimerManager.IsTimerActive(Delegate));
+	Test->TestTrue(TIMER_TEST_TEXT("IsTimerPaused called with a paused timer"), TimerManager.IsTimerPaused(Delegate));
+
+	TimerTest_TickWorld(World);
+
+	Test->TestTrue(TIMER_TEST_TEXT("TimerExists called with a paused timer"), TimerManager.TimerExists(Delegate));
+	Test->TestFalse(TIMER_TEST_TEXT("IsTimerActive called with a paused timer"), TimerManager.IsTimerActive(Delegate));
+	Test->TestTrue(TIMER_TEST_TEXT("IsTimerPaused called with a paused timer"), TimerManager.IsTimerPaused(Delegate));
+
+	Test->TestTrue(TIMER_TEST_TEXT("GetTimerElapsed called with a paused timer after one step"), (FMath::IsNearlyEqual(TimerManager.GetTimerElapsed(Delegate), 1.f, KINDA_SMALL_NUMBER)));
+	Test->TestTrue(TIMER_TEST_TEXT("GetTimerRemaining called with a paused timer after one step"), (FMath::IsNearlyEqual(TimerManager.GetTimerRemaining(Delegate), Rate - 1.f, KINDA_SMALL_NUMBER)));
+
+	TimerManager.UnPauseTimer(Delegate);
+
+	Test->TestTrue(TIMER_TEST_TEXT("TimerExists called with a pending timer"), TimerManager.TimerExists(Delegate));
+	Test->TestTrue(TIMER_TEST_TEXT("IsTimerActive called with a pending timer"), TimerManager.IsTimerActive(Delegate));
+	Test->TestFalse(TIMER_TEST_TEXT("IsTimerPaused called with a pending timer"), TimerManager.IsTimerPaused(Delegate));
+
+	TimerTest_TickWorld(World);
+
+	Test->TestFalse(TIMER_TEST_TEXT("TimerExists called with a completed timer"), TimerManager.TimerExists(Delegate));
+	Test->TestTrue(TIMER_TEST_TEXT("Count of callback executions"), Dummy.Count == 1);
+
+
+	// Test resetting the timer
+	TimerManager.SetTimer(Delegate, Rate, false);
+	TimerManager.SetTimer(Delegate, 0.f, false);
+
+	Test->TestFalse(TIMER_TEST_TEXT("TimerExists called with a reset timer"), TimerManager.TimerExists(Delegate));
+
+	// Test looping timers
+	Dummy.Reset();
+	TimerManager.SetTimer(Delegate, Rate, true);
+	TimerTest_TickWorld(World, KINDA_SMALL_NUMBER);
+
+	TimerTest_TickWorld(World, 2.f);
+
+	Test->TestTrue(TIMER_TEST_TEXT("TimerExists called with a looping timer"), TimerManager.TimerExists(Delegate));
+	Test->TestTrue(TIMER_TEST_TEXT("IsTimerActive called with a looping timer"), TimerManager.IsTimerActive(Delegate));
+
+	Test->TestTrue(TIMER_TEST_TEXT("Count of callback executions"), Dummy.Count == 1);
+	Test->TestTrue(TIMER_TEST_TEXT("GetTimerElapsed called with a looping timer"), (FMath::IsNearlyEqual(TimerManager.GetTimerElapsed(Delegate), 2.f - (Rate * Dummy.Count), KINDA_SMALL_NUMBER)));
+	Test->TestTrue(TIMER_TEST_TEXT("GetTimerRemaining called with a looping timer"), (FMath::IsNearlyEqual(TimerManager.GetTimerRemaining(Delegate), Rate * (Dummy.Count + 1) - 2.f, KINDA_SMALL_NUMBER)));
+
+	TimerTest_TickWorld(World, 2.f);
+
+	Test->TestTrue(TIMER_TEST_TEXT("Count of callback executions"), Dummy.Count == 2);
+	Test->TestTrue(TIMER_TEST_TEXT("GetTimerElapsed called with a looping timer"), (FMath::IsNearlyEqual(TimerManager.GetTimerElapsed(Delegate), 4.f - (Rate * Dummy.Count), KINDA_SMALL_NUMBER)));
+	Test->TestTrue(TIMER_TEST_TEXT("GetTimerRemaining called with a looping timer"), (FMath::IsNearlyEqual(TimerManager.GetTimerRemaining(Delegate), Rate * (Dummy.Count + 1) - 4.f, KINDA_SMALL_NUMBER)));
+
+	TimerManager.SetTimer(Delegate, 0.f, false);
+
+	Test->TestFalse(TIMER_TEST_TEXT("TimerExists called with a reset looping timer"), TimerManager.TimerExists(Delegate));
+
+	Dummy.Reset();
+	TimerManager.SetTimerForNextTick(Delegate);
+
+	Test->TestTrue(TIMER_TEST_TEXT("TimerExists called with a timer executing on the next tick"), TimerManager.TimerExists(Delegate));
+	Test->TestTrue(TIMER_TEST_TEXT("IsTimerActive called with a timer executing on the next tick"), TimerManager.IsTimerActive(Delegate));
+
+	TimerTest_TickWorld(World, SMALL_NUMBER);
+
+	Test->TestFalse(TIMER_TEST_TEXT("TimerExists called with a timer executing on the previous tick"), TimerManager.TimerExists(Delegate));
+	Test->TestTrue(TIMER_TEST_TEXT("Count of callback executions"), Dummy.Count == 1);
+
+	return true;
+}
+
+bool TimerManagerTest_ValidTimer_HandleWithDelegate(UWorld* World, FAutomationTestBase* Test)
+{
+	FTimerManager& TimerManager = World->GetTimerManager();
+	FTimerDelegate Delegate;
+	FTimerHandle Handle;
+	const float Rate = 1.5f;
+
+	FDummy Dummy;
+	Delegate.BindRaw(&Dummy, &FDummy::Callback);
+
+	TimerManager.SetTimer(Handle, Delegate, Rate, false);
+
+	Test->TestTrue(TIMER_TEST_TEXT("Handle should be valid after calling SetTimer"), Handle.IsValid());
+	Test->TestTrue(TIMER_TEST_TEXT("TimerExists called with a pending timer"), TimerManager.TimerExists(Handle));
+	Test->TestFalse(TIMER_TEST_TEXT("TimerExists called by delegate with a pending timer"), TimerManager.TimerExists(Delegate));
+	Test->TestTrue(TIMER_TEST_TEXT("IsTimerActive called with a pending timer"), TimerManager.IsTimerActive(Handle));
+	Test->TestFalse(TIMER_TEST_TEXT("IsTimerPaused called with a pending timer"), TimerManager.IsTimerPaused(Handle));
+	Test->TestTrue(TIMER_TEST_TEXT("GetTimerRate called with a pending timer"), (TimerManager.GetTimerRate(Handle) == Rate));
+	Test->TestTrue(TIMER_TEST_TEXT("GetTimerElapsed called with a pending timer"), (TimerManager.GetTimerElapsed(Handle) == 0.f));
+	Test->TestTrue(TIMER_TEST_TEXT("GetTimerRemaining called with a pending timer"), (TimerManager.GetTimerRemaining(Handle) == Rate));
+
+	// small tick to move the timer from the pending list to the active list, the timer will start counting time after this tick
+	TimerTest_TickWorld(World, KINDA_SMALL_NUMBER);
+
+	Test->TestTrue(TIMER_TEST_TEXT("TimerExists called with a pending timer"), TimerManager.TimerExists(Handle));
+	Test->TestTrue(TIMER_TEST_TEXT("IsTimerActive called with an active timer"), TimerManager.IsTimerActive(Handle));
+	Test->TestFalse(TIMER_TEST_TEXT("IsTimerPaused called with an active timer"), TimerManager.IsTimerPaused(Handle));
+
+	TimerTest_TickWorld(World);
+
+	Test->TestTrue(TIMER_TEST_TEXT("GetTimerElapsed called with an active timer after one step"), (FMath::IsNearlyEqual(TimerManager.GetTimerElapsed(Handle), 1.f, KINDA_SMALL_NUMBER)));
+	Test->TestTrue(TIMER_TEST_TEXT("GetTimerRemaining called with an active timer after one step"), (FMath::IsNearlyEqual(TimerManager.GetTimerRemaining(Handle), Rate - 1.f, KINDA_SMALL_NUMBER)));
+
+	TimerManager.PauseTimer(Handle);
+
+	Test->TestTrue(TIMER_TEST_TEXT("TimerExists called with a paused timer"), TimerManager.TimerExists(Handle));
+	Test->TestFalse(TIMER_TEST_TEXT("IsTimerActive called with a paused timer"), TimerManager.IsTimerActive(Handle));
+	Test->TestTrue(TIMER_TEST_TEXT("IsTimerPaused called with a paused timer"), TimerManager.IsTimerPaused(Handle));
+
+	TimerTest_TickWorld(World);
+
+	Test->TestTrue(TIMER_TEST_TEXT("TimerExists called with a paused timer"), TimerManager.TimerExists(Handle));
+	Test->TestFalse(TIMER_TEST_TEXT("IsTimerActive called with a paused timer"), TimerManager.IsTimerActive(Handle));
+	Test->TestTrue(TIMER_TEST_TEXT("IsTimerPaused called with a paused timer"), TimerManager.IsTimerPaused(Handle));
+
+	Test->TestTrue(TIMER_TEST_TEXT("GetTimerElapsed called with a paused timer after one step"), (FMath::IsNearlyEqual(TimerManager.GetTimerElapsed(Handle), 1.f, KINDA_SMALL_NUMBER)));
+	Test->TestTrue(TIMER_TEST_TEXT("GetTimerRemaining called with a paused timer after one step"), (FMath::IsNearlyEqual(TimerManager.GetTimerRemaining(Handle), Rate - 1.f, KINDA_SMALL_NUMBER)));
+
+	TimerManager.UnPauseTimer(Handle);
+
+	Test->TestTrue(TIMER_TEST_TEXT("TimerExists called with a pending timer"), TimerManager.TimerExists(Handle));
+	Test->TestTrue(TIMER_TEST_TEXT("IsTimerActive called with a pending timer"), TimerManager.IsTimerActive(Handle));
+	Test->TestFalse(TIMER_TEST_TEXT("IsTimerPaused called with a pending timer"), TimerManager.IsTimerPaused(Handle));
+
+	TimerTest_TickWorld(World);
+
+	Test->TestFalse(TIMER_TEST_TEXT("TimerExists called with a completed timer"), TimerManager.TimerExists(Handle));
+	Test->TestTrue(TIMER_TEST_TEXT("Count of callback executions"), Dummy.Count == 1);
+
+
+	// Test resetting the timer
+	TimerManager.SetTimer(Handle, Delegate, Rate, false);
+	TimerManager.SetTimer(Handle, 0.f, false);
+
+	Test->TestFalse(TIMER_TEST_TEXT("TimerExists called with a reset timer"), TimerManager.TimerExists(Handle));
+
+	// Test looping timers
+	Dummy.Reset();
+	TimerManager.SetTimer(Handle, Delegate, Rate, true);
+	TimerTest_TickWorld(World, KINDA_SMALL_NUMBER);
+
+	TimerTest_TickWorld(World, 2.f);
+
+	Test->TestTrue(TIMER_TEST_TEXT("TimerExists called with a looping timer"), TimerManager.TimerExists(Handle));
+	Test->TestTrue(TIMER_TEST_TEXT("IsTimerActive called with a looping timer"), TimerManager.IsTimerActive(Handle));
+
+	Test->TestTrue(TIMER_TEST_TEXT("Count of callback executions"), Dummy.Count == 1);
+	Test->TestTrue(TIMER_TEST_TEXT("GetTimerElapsed called with a looping timer"), (FMath::IsNearlyEqual(TimerManager.GetTimerElapsed(Handle), 2.f - (Rate * Dummy.Count), KINDA_SMALL_NUMBER)));
+	Test->TestTrue(TIMER_TEST_TEXT("GetTimerRemaining called with a looping timer"), (FMath::IsNearlyEqual(TimerManager.GetTimerRemaining(Handle), Rate * (Dummy.Count + 1) - 2.f, KINDA_SMALL_NUMBER)));
+
+	TimerTest_TickWorld(World, 2.f);
+
+	Test->TestTrue(TIMER_TEST_TEXT("Count of callback executions"), Dummy.Count == 2);
+	Test->TestTrue(TIMER_TEST_TEXT("GetTimerElapsed called with a looping timer"), (FMath::IsNearlyEqual(TimerManager.GetTimerElapsed(Handle), 4.f - (Rate * Dummy.Count), KINDA_SMALL_NUMBER)));
+	Test->TestTrue(TIMER_TEST_TEXT("GetTimerRemaining called with a looping timer"), (FMath::IsNearlyEqual(TimerManager.GetTimerRemaining(Handle), Rate * (Dummy.Count + 1) - 4.f, KINDA_SMALL_NUMBER)));
+
+	TimerManager.SetTimer(Handle, 0.f, false);
+
+	Test->TestFalse(TIMER_TEST_TEXT("TimerExists called with a reset looping timer"), TimerManager.TimerExists(Handle));
+
+	return true;
+}
+
+bool TimerManagerTest_ValidTimer_HandleNoDelegate(UWorld* World, FAutomationTestBase* Test)
+{
+	FTimerManager& TimerManager = World->GetTimerManager();
+	FTimerHandle Handle;
+	const float Rate = 1.5f;
+
+	TimerManager.SetTimer(Handle, Rate, false);
+
+	Test->TestTrue(TIMER_TEST_TEXT("Handle should be valid after calling SetTimer"), Handle.IsValid());
+	Test->TestTrue(TIMER_TEST_TEXT("TimerExists called with a pending timer"), TimerManager.TimerExists(Handle));
+	Test->TestTrue(TIMER_TEST_TEXT("IsTimerActive called with a pending timer"), TimerManager.IsTimerActive(Handle));
+	Test->TestFalse(TIMER_TEST_TEXT("IsTimerPaused called with a pending timer"), TimerManager.IsTimerPaused(Handle));
+	Test->TestTrue(TIMER_TEST_TEXT("GetTimerRate called with a pending timer"), (TimerManager.GetTimerRate(Handle) == Rate));
+	Test->TestTrue(TIMER_TEST_TEXT("GetTimerElapsed called with a pending timer"), (TimerManager.GetTimerElapsed(Handle) == 0.f));
+	Test->TestTrue(TIMER_TEST_TEXT("GetTimerRemaining called with a pending timer"), (TimerManager.GetTimerRemaining(Handle) == Rate));
+
+	// small tick to move the timer from the pending list to the active list, the timer will start counting time after this tick
+	TimerTest_TickWorld(World, KINDA_SMALL_NUMBER);
+
+	Test->TestTrue(TIMER_TEST_TEXT("TimerExists called with a pending timer"), TimerManager.TimerExists(Handle));
+	Test->TestTrue(TIMER_TEST_TEXT("IsTimerActive called with an active timer"), TimerManager.IsTimerActive(Handle));
+	Test->TestFalse(TIMER_TEST_TEXT("IsTimerPaused called with an active timer"), TimerManager.IsTimerPaused(Handle));
+
+	TimerTest_TickWorld(World);
+
+	Test->TestTrue(TIMER_TEST_TEXT("GetTimerElapsed called with an active timer after one step"), (FMath::IsNearlyEqual(TimerManager.GetTimerElapsed(Handle), 1.f, KINDA_SMALL_NUMBER)));
+	Test->TestTrue(TIMER_TEST_TEXT("GetTimerRemaining called with an active timer after one step"), (FMath::IsNearlyEqual(TimerManager.GetTimerRemaining(Handle), Rate - 1.f, KINDA_SMALL_NUMBER)));
+
+	TimerManager.PauseTimer(Handle);
+
+	Test->TestTrue(TIMER_TEST_TEXT("TimerExists called with a paused timer"), TimerManager.TimerExists(Handle));
+	Test->TestFalse(TIMER_TEST_TEXT("IsTimerActive called with a paused timer"), TimerManager.IsTimerActive(Handle));
+	Test->TestTrue(TIMER_TEST_TEXT("IsTimerPaused called with a paused timer"), TimerManager.IsTimerPaused(Handle));
+
+	TimerTest_TickWorld(World);
+
+	Test->TestTrue(TIMER_TEST_TEXT("TimerExists called with a paused timer"), TimerManager.TimerExists(Handle));
+	Test->TestFalse(TIMER_TEST_TEXT("IsTimerActive called with a paused timer"), TimerManager.IsTimerActive(Handle));
+	Test->TestTrue(TIMER_TEST_TEXT("IsTimerPaused called with a paused timer"), TimerManager.IsTimerPaused(Handle));
+
+	Test->TestTrue(TIMER_TEST_TEXT("GetTimerElapsed called with a paused timer after one step"), (FMath::IsNearlyEqual(TimerManager.GetTimerElapsed(Handle), 1.f, KINDA_SMALL_NUMBER)));
+	Test->TestTrue(TIMER_TEST_TEXT("GetTimerRemaining called with a paused timer after one step"), (FMath::IsNearlyEqual(TimerManager.GetTimerRemaining(Handle), Rate - 1.f, KINDA_SMALL_NUMBER)));
+
+	TimerManager.UnPauseTimer(Handle);
+
+	Test->TestTrue(TIMER_TEST_TEXT("TimerExists called with a pending timer"), TimerManager.TimerExists(Handle));
+	Test->TestTrue(TIMER_TEST_TEXT("IsTimerActive called with a pending timer"), TimerManager.IsTimerActive(Handle));
+	Test->TestFalse(TIMER_TEST_TEXT("IsTimerPaused called with a pending timer"), TimerManager.IsTimerPaused(Handle));
+
+	TimerTest_TickWorld(World);
+
+	Test->TestFalse(TIMER_TEST_TEXT("TimerExists called with a completed timer"), TimerManager.TimerExists(Handle));
+
+	// Test resetting the timer
+	TimerManager.SetTimer(Handle, Rate, false);
+	TimerManager.SetTimer(Handle, 0.f, false);
+
+	Test->TestFalse(TIMER_TEST_TEXT("TimerExists called with a reset timer"), TimerManager.TimerExists(Handle));
+
+	// Test looping timers
+	TimerManager.SetTimer(Handle, Rate, true);
+	TimerTest_TickWorld(World, KINDA_SMALL_NUMBER);
+
+	TimerTest_TickWorld(World, 2.f);
+
+	Test->TestTrue(TIMER_TEST_TEXT("TimerExists called with a looping timer"), TimerManager.TimerExists(Handle));
+	Test->TestTrue(TIMER_TEST_TEXT("IsTimerActive called with a looping timer"), TimerManager.IsTimerActive(Handle));
+
+	Test->TestTrue(TIMER_TEST_TEXT("GetTimerElapsed called with a looping timer"), (FMath::IsNearlyEqual(TimerManager.GetTimerElapsed(Handle), 2.f - Rate, KINDA_SMALL_NUMBER)));
+	Test->TestTrue(TIMER_TEST_TEXT("GetTimerRemaining called with a looping timer"), (FMath::IsNearlyEqual(TimerManager.GetTimerRemaining(Handle), Rate * 2 - 2.f, KINDA_SMALL_NUMBER)));
+
+	TimerTest_TickWorld(World, 2.f);
+
+	Test->TestTrue(TIMER_TEST_TEXT("GetTimerElapsed called with a looping timer"), (FMath::IsNearlyEqual(TimerManager.GetTimerElapsed(Handle), 4.f - (Rate * 2), KINDA_SMALL_NUMBER)));
+	Test->TestTrue(TIMER_TEST_TEXT("GetTimerRemaining called with a looping timer"), (FMath::IsNearlyEqual(TimerManager.GetTimerRemaining(Handle), Rate * 3 - 4.f, KINDA_SMALL_NUMBER)));
+
+	TimerManager.SetTimer(Handle, 0.f, false);
+
+	Test->TestFalse(TIMER_TEST_TEXT("TimerExists called with a reset looping timer"), TimerManager.TimerExists(Handle));
+
+	return true;
+}
+
+bool FTimerManagerTest::RunTest(const FString& Parameters)
+{
+	UWorld *World = UWorld::CreateWorld(EWorldType::Game, false);
+	FWorldContext &WorldContext = GEngine->CreateNewWorldContext(EWorldType::Game);
+	WorldContext.SetCurrentWorld(World);
+	
+	FURL URL;
+	World->InitializeActorsForPlay(URL);
+	World->BeginPlay();
+
+	TimerManagerTest_InvalidTimers(World, this);
+	TimerManagerTest_MissingTimers(World, this);
+	TimerManagerTest_ValidTimer_Delegate(World, this);
+	TimerManagerTest_ValidTimer_HandleWithDelegate(World, this);
+	TimerManagerTest_ValidTimer_HandleNoDelegate(World, this);
+
+	GEngine->DestroyWorldContext(World);
+	World->DestroyWorld(false);
+
+	return true;
+}
+
+
diff --git a/Engine/Source/Runtime/Engine/Public/TimerManager.h b/Engine/Source/Runtime/Engine/Public/TimerManager.h
--- a/Engine/Source/Runtime/Engine/Public/TimerManager.h
+++ b/Engine/Source/Runtime/Engine/Public/TimerManager.h
@@ -81,6 +81,52 @@ struct FTimerUnifiedDelegate
 	}
 };
 
+// Unique handle that can be used to distinguish timers that have identical delegates.
+struct FTimerHandle
+{
+	FTimerHandle()
+	: Handle(INDEX_NONE)
+	{
+
+	}
+
+	FTimerHandle(int32 InHandle)
+		: Handle(InHandle)
+	{
+
+	}
+
+	bool IsValid() const
+	{
+		return Handle != INDEX_NONE;
+	}
+
+	void Invalidate()
+	{
+		Handle = INDEX_NONE;
+	}
+
+	void MakeValid();
+
+	bool operator==(const FTimerHandle& Other) const
+	{
+		return Handle == Other.Handle;
+	}
+
+	bool operator!=(const FTimerHandle& Other) const
+	{
+		return Handle != Other.Handle;
+	}
+
+	FString ToString() const
+	{
+		return FString::Printf(TEXT("%d"), Handle);
+	}
+
+private:
+	int32 Handle;
+};
+
 namespace ETimerStatus
 {
 	enum Type
@@ -112,6 +158,8 @@ struct FTimerData
 	/** Holds the delegate to call. */
 	FTimerUnifiedDelegate TimerDelegate;
 
+	FTimerHandle TimerHandle;
+
 	FTimerData()
 		: bLoop(false), Status(ETimerStatus::Active)
 		, Rate(0), ExpireTime(0)
@@ -148,36 +196,104 @@ public:
 	{}
 
 
-/**
+	/**
 	 * Sets a timer to call the given native function at a set interval.  If a timer is already set
 	 * for this delegate, it will update the current timer to the new parameters and reset its 
 	 * elapsed time to 0.
 	 *
-	 * @param inObj			Object to call the timer function on.
-	 * @param inTimerMethod Method to call when timer fires.
-	 * @param inRate		The amount of time between set and firing.  If <= 0.f, clears existing timers.
-	 * @param inbLoop		true to keep firing at Rate intervals, false to fire only once.
+	 * @param InObj					Object to call the timer function on.
+	 * @param InTimerMethod			Method to call when timer fires.
+	 * @param InRate				The amount of time between set and firing.  If <= 0.f, clears existing timers.
+	 * @param InbLoop				true to keep firing at Rate intervals, false to fire only once.
+	 * @param InFirstDelay			The time for the first iteration of a looping timer. If < 0.f inRate will be used.
 	 */
 	template< class UserClass >	
-	FORCEINLINE void SetTimer( UserClass* inObj, typename FTimerDelegate::TUObjectMethodDelegate< UserClass >::FMethodPtr inTimerMethod, float inRate, bool inbLoop = false)
+	FORCEINLINE void SetTimer(UserClass* InObj, typename FTimerDelegate::TUObjectMethodDelegate< UserClass >::FMethodPtr InTimerMethod, float InRate, bool InbLoop = false, float InFirstDelay = -1.f)
 	{
-		SetTimer( FTimerDelegate::CreateUObject(inObj, inTimerMethod), inRate, inbLoop );
+		SetTimer( FTimerDelegate::CreateUObject(InObj, InTimerMethod), InRate, InbLoop, InFirstDelay );
 	}
 	template< class UserClass >	
-	FORCEINLINE void SetTimer( UserClass* inObj, typename FTimerDelegate::TUObjectMethodDelegate_Const< UserClass >::FMethodPtr inTimerMethod, float inRate, bool inbLoop = false)
+	FORCEINLINE void SetTimer(UserClass* InObj, typename FTimerDelegate::TUObjectMethodDelegate_Const< UserClass >::FMethodPtr InTimerMethod, float InRate, bool InbLoop = false, float InFirstDelay = -1.f)
+	{
+		SetTimer( FTimerDelegate::CreateUObject(InObj, InTimerMethod), InRate, InbLoop, InFirstDelay );
+	}
+
+	/** Version that takes any generic delegate. */
+	FORCEINLINE void SetTimer(FTimerDelegate const& InDelegate, float InRate, bool InbLoop, float InFirstDelay = -1.f)
+	{
+		InternalSetTimer( FTimerUnifiedDelegate(InDelegate), InRate, InbLoop, InFirstDelay );
+	}
+	/** Version that takes a dynamic delegate (e.g. for UFunctions). */
+	FORCEINLINE void SetTimer(FTimerDynamicDelegate const& InDynDelegate, float InRate, bool InbLoop, float InFirstDelay = -1.f)
+	{
+		InternalSetTimer( FTimerUnifiedDelegate(InDynDelegate), InRate, InbLoop, InFirstDelay );
+	}
+
+	/**
+	* Sets a timer to call the given native function at a set interval.  If a timer is already set
+	* for this delegate, it will update the current timer to the new parameters and reset its
+	* elapsed time to 0.
+	*
+	* @param InOutHandle			Handle to identify this timer. If it is invalid when passed in it will be made into a valid handle.
+	* @param InObj					Object to call the timer function on.
+	* @param InTimerMethod			Method to call when timer fires.
+	* @param InRate					The amount of time between set and firing.  If <= 0.f, clears existing timers.
+	* @param InbLoop				true to keep firing at Rate intervals, false to fire only once.
+	* @param InFirstDelay			The time for the first iteration of a looping timer. If < 0.f inRate will be used.
+	*/
+	template< class UserClass >
+	FORCEINLINE void SetTimer(FTimerHandle& InOutHandle, UserClass* InObj, typename FTimerDelegate::TUObjectMethodDelegate< UserClass >::FMethodPtr InTimerMethod, float InRate, bool InbLoop = false, float InFirstDelay = -1.f)
+	{
+		SetTimer(InOutHandle, FTimerDelegate::CreateUObject(InObj, InTimerMethod), InRate, InbLoop, InFirstDelay);
+	}
+	template< class UserClass >
+	FORCEINLINE void SetTimer(FTimerHandle& InOutHandle, UserClass* InObj, typename FTimerDelegate::TUObjectMethodDelegate_Const< UserClass >::FMethodPtr InTimerMethod, float InRate, bool InbLoop = false, float InFirstDelay = -1.f)
+	{
+		SetTimer(InOutHandle, FTimerDelegate::CreateUObject(InObj, InTimerMethod), InRate, InbLoop, InFirstDelay);
+	}
+
+	/** Version that takes any generic delegate. */
+	FORCEINLINE void SetTimer(FTimerHandle& InOutHandle, FTimerDelegate const& InDelegate, float InRate, bool InbLoop, float InFirstDelay = -1.f)
+	{
+		InternalSetTimer(InOutHandle, FTimerUnifiedDelegate(InDelegate), InRate, InbLoop, InFirstDelay);
+	}
+	/** Version that takes a dynamic delegate (e.g. for UFunctions). */
+	FORCEINLINE void SetTimer(FTimerHandle& InOutHandle, FTimerDynamicDelegate const& InDynDelegate, float InRate, bool InbLoop, float InFirstDelay = -1.f)
+	{
+		InternalSetTimer(InOutHandle, FTimerUnifiedDelegate(InDynDelegate), InRate, InbLoop, InFirstDelay);
+	}
+	/*** Version that doesn't take a delegate */
+	FORCEINLINE void SetTimer(FTimerHandle& InOutHandle, float InRate, bool InbLoop, float InFirstDelay = -1.f)
+	{
+		InternalSetTimer(InOutHandle, FTimerUnifiedDelegate(), InRate, InbLoop, InFirstDelay);
+	}
+
+	/**
+	* Sets a timer to call the given native function on the next tick.
+	*
+	* @param inObj					Object to call the timer function on.
+	* @param inTimerMethod			Method to call when timer fires.
+	*/
+	template< class UserClass >
+	FORCEINLINE void SetTimerForNextTick(UserClass* inObj, typename FTimerDelegate::TUObjectMethodDelegate< UserClass >::FMethodPtr inTimerMethod)
+	{
+		SetTimerForNextTick(FTimerDelegate::CreateUObject(inObj, inTimerMethod));
+	}
+	template< class UserClass >
+	FORCEINLINE void SetTimerForNextTick(UserClass* inObj, typename FTimerDelegate::TUObjectMethodDelegate_Const< UserClass >::FMethodPtr inTimerMethod)
 	{
-		SetTimer( FTimerDelegate::CreateUObject(inObj, inTimerMethod), inRate, inbLoop );
+		SetTimerForNextTick(FTimerDelegate::CreateUObject(inObj, inTimerMethod));
 	}
 
 	/** Version that takes any generic delegate. */
-	FORCEINLINE void SetTimer(FTimerDelegate const& InDelegate, float InRate, bool InbLoop)
+	FORCEINLINE void SetTimerForNextTick(FTimerDelegate const& InDelegate)
 	{
-		InternalSetTimer( FTimerUnifiedDelegate(InDelegate), InRate, InbLoop );
+		InternalSetTimerForNextTick(FTimerUnifiedDelegate(InDelegate));
 	}
 	/** Version that takes a dynamic delegate (e.g. for UFunctions). */
-	FORCEINLINE void SetTimer(FTimerDynamicDelegate const& InDynDelegate, float InRate, bool InbLoop)
+	FORCEINLINE void SetTimerForNextTick(FTimerDynamicDelegate const& InDynDelegate)
 	{
-		InternalSetTimer( FTimerUnifiedDelegate(InDynDelegate), InRate, InbLoop );
+		InternalSetTimerForNextTick(FTimerUnifiedDelegate(InDynDelegate));
 	}
 
 
@@ -208,6 +324,11 @@ public:
 	{
 		InternalClearTimer( FTimerUnifiedDelegate(InDynDelegate) );
 	}
+	/** Version that takes a handle */
+	FORCEINLINE void ClearTimer(FTimerHandle const& InHandle)
+	{
+		InternalClearTimer(InHandle);
+	}
 
 	/** Clears all timers that are bound to functions on the given object. */
 	FORCEINLINE void ClearAllTimersForObject(void const* Object)
@@ -239,12 +360,23 @@ public:
 	/** Version that takes any generic delegate. */
 	FORCEINLINE void PauseTimer(FTimerDelegate const& InDelegate)
 	{
-		InternalPauseTimer( FTimerUnifiedDelegate(InDelegate) );
+		int32 TimerIdx;
+		FTimerData const* TimerToPause = FindTimer(FTimerUnifiedDelegate(InDelegate), &TimerIdx);
+		InternalPauseTimer(TimerToPause, TimerIdx);
 	}
 	/** Version that takes a dynamic delegate (e.g. for UFunctions). */
 	FORCEINLINE void PauseTimer(FTimerDynamicDelegate const& InDynDelegate)
 	{
-		InternalPauseTimer( FTimerUnifiedDelegate(InDynDelegate) );
+		int32 TimerIdx;
+		FTimerData const* TimerToPause = FindTimer(FTimerUnifiedDelegate(InDynDelegate), &TimerIdx);
+		InternalPauseTimer(TimerToPause, TimerIdx);
+	}
+	/** Version that takes a handle */
+	FORCEINLINE void PauseTimer(FTimerHandle const& InHandle)
+	{
+		int32 TimerIdx;
+		FTimerData const* TimerToPause = FindTimer(InHandle, &TimerIdx);
+		InternalPauseTimer(TimerToPause, TimerIdx);
 	}
 
 	/**
@@ -267,12 +399,20 @@ public:
 	/** Version that takes any generic delegate. */
 	FORCEINLINE void UnPauseTimer(FTimerDelegate const& InDelegate)
 	{
-		InternalUnPauseTimer( FTimerUnifiedDelegate(InDelegate) );
+		int32 TimerIdx = FindTimerInList( PausedTimerList, FTimerUnifiedDelegate(InDelegate) );
+		InternalUnPauseTimer(TimerIdx);
 	}
 	/** Version that takes a dynamic delegate (e.g. for UFunctions). */
 	FORCEINLINE void UnPauseTimer(FTimerDynamicDelegate const& InDynDelegate)
 	{
-		InternalUnPauseTimer( FTimerUnifiedDelegate(InDynDelegate) );
+		int32 TimerIdx = FindTimerInList( PausedTimerList, FTimerUnifiedDelegate(InDynDelegate) );
+		InternalUnPauseTimer(TimerIdx);
+	}
+	/** Version that takes a handle */
+	FORCEINLINE void UnPauseTimer(FTimerHandle const& InHandle)
+	{
+		int32 TimerIdx = FindTimerInList(PausedTimerList, InHandle);
+		InternalUnPauseTimer(TimerIdx);
 	}
 
 	/**
@@ -297,13 +437,22 @@ public:
 	/** Version that takes any generic delegate. */
 	FORCEINLINE float GetTimerRate(FTimerDelegate const& InDelegate)
 	{
-		return InternalGetTimerRate( FTimerUnifiedDelegate(InDelegate) );
+		FTimerData const* const TimerData = FindTimer( FTimerUnifiedDelegate(InDelegate) );
+		return InternalGetTimerRate(TimerData);
 	}
 
 	/** Version that takes a dynamic delegate (e.g. for UFunctions). */
 	FORCEINLINE float GetTimerRate(FTimerDynamicDelegate const& InDynDelegate)
 	{
-		return InternalGetTimerRate( FTimerUnifiedDelegate(InDynDelegate) );
+		FTimerData const* const TimerData = FindTimer( FTimerUnifiedDelegate(InDynDelegate) );
+		return InternalGetTimerRate(TimerData);
+	}
+
+	/** Version that takes a handle */
+	FORCEINLINE float GetTimerRate(FTimerHandle const& InHandle)
+	{
+		FTimerData const* const TimerData = FindTimer(InHandle);
+		return InternalGetTimerRate(TimerData);
 	}
 
 	/**
@@ -339,6 +488,13 @@ public:
 		return ( (TimerData != NULL) && (TimerData->Status != ETimerStatus::Paused) );
 	}
 
+	/** Version that takes a handle */
+	FORCEINLINE bool IsTimerActive(FTimerHandle const& InHandle)
+	{
+		FTimerData const* const TimerData = FindTimer( InHandle );
+		return ( (TimerData != NULL) && (TimerData->Status != ETimerStatus::Paused) );
+	}
+
 	/**
 	* Returns true if the specified timer exists and is paused
 	*
@@ -372,6 +528,53 @@ public:
 		return ((TimerData != NULL) && (TimerData->Status == ETimerStatus::Paused));
 	}
 
+	/** Version that takes a handle */
+	FORCEINLINE bool IsTimerPaused(FTimerHandle const& InHandle)
+	{
+		FTimerData const* const TimerData = FindTimer(InHandle);
+		return ((TimerData != NULL) && (TimerData->Status == ETimerStatus::Paused));
+	}
+
+	/**
+	* Returns true if the specified timer exists and is pending
+	*
+	* @param inObj			Object binding for query.
+	* @param inTimerMethod	Method binding for query.
+	* @return				true if the timer matching the given criteria exists and is paused, false otherwise.
+	*/
+	template< class UserClass >
+	FORCEINLINE bool IsTimerPending(UserClass* inObj, typename FTimerDelegate::TUObjectMethodDelegate< UserClass >::FMethodPtr inTimerMethod) const
+	{
+		return IsTimerPending(FTimerDelegate::CreateUObject(inObj, inTimerMethod));
+	}
+
+	template< class UserClass >
+	FORCEINLINE bool IsTimerPending(UserClass* inObj, typename FTimerDelegate::TUObjectMethodDelegate_Const< UserClass >::FMethodPtr inTimerMethod)  const
+	{
+		return IsTimerPending(FTimerDelegate::CreateUObject(inObj, inTimerMethod));
+	}
+
+	/** Version that takes any generic delegate. */
+	FORCEINLINE bool IsTimerPending(FTimerDelegate const& InDelegate) const
+	{
+		FTimerData const* const TimerData = FindTimer(FTimerUnifiedDelegate(InDelegate));
+		return ((TimerData != NULL) && (TimerData->Status == ETimerStatus::Pending));
+	}
+
+	/** Version that takes a dynamic delegate (e.g. for UFunctions). */
+	FORCEINLINE bool IsTimerPending(FTimerDynamicDelegate const& InDynDelegate) const
+	{
+		FTimerData const* const TimerData = FindTimer(FTimerUnifiedDelegate(InDynDelegate));
+		return ((TimerData != NULL) && (TimerData->Status == ETimerStatus::Pending));
+	}
+
+	/** Version that takes a handle */
+	FORCEINLINE bool IsTimerPending(FTimerHandle const& InHandle)
+	{
+		FTimerData const* const TimerData = FindTimer(InHandle);
+		return ((TimerData != NULL) && (TimerData->Status == ETimerStatus::Pending));
+	}
+
 	/**
 	* Returns true if the specified timer exists
 	*
@@ -404,6 +607,12 @@ public:
 		return FindTimer(FTimerUnifiedDelegate(InDynDelegate)) != NULL;
 	}
 
+	/** Version that takes a handle */
+	FORCEINLINE bool TimerExists(FTimerHandle const& InHandle)
+	{
+		return FindTimer(InHandle) != NULL;
+	}
+
 	/**
 	 * Gets the current elapsed time for the specified timer.
 	 *
@@ -424,12 +633,21 @@ public:
 	/** Version that takes any generic delegate. */
 	FORCEINLINE float GetTimerElapsed(FTimerDelegate const& InDelegate)
 	{
-		return InternalGetTimerElapsed( FTimerUnifiedDelegate(InDelegate) );
+		FTimerData const* const TimerData = FindTimer( FTimerUnifiedDelegate(InDelegate) );
+		return InternalGetTimerElapsed(TimerData);
 	}
 	/** Version that takes a dynamic delegate (e.g. for UFunctions). */
 	FORCEINLINE float GetTimerElapsed(FTimerDynamicDelegate const& InDynDelegate)
 	{
-		return InternalGetTimerElapsed( FTimerUnifiedDelegate(InDynDelegate) );
+		FTimerData const* const TimerData = FindTimer( FTimerUnifiedDelegate(InDynDelegate) );
+		return InternalGetTimerElapsed(TimerData);
+	}
+
+	/** Version that takes a handle */
+	FORCEINLINE float GetTimerElapsed(FTimerHandle const& InHandle)
+	{
+		FTimerData const* const TimerData = FindTimer(InHandle);
+		return InternalGetTimerElapsed(TimerData);
 	}
 
 	/**
@@ -452,12 +670,20 @@ public:
 	/** Version that takes any generic delegate. */
 	FORCEINLINE float GetTimerRemaining(FTimerDelegate const& InDelegate)
 	{
-		return InternalGetTimerRemaining( FTimerUnifiedDelegate(InDelegate) );
+		FTimerData const* const TimerData = FindTimer( FTimerUnifiedDelegate(InDelegate) );
+		return InternalGetTimerRemaining(TimerData);
 	}
 	/** Version that takes a dynamic delegate (e.g. for UFunctions). */
 	FORCEINLINE float GetTimerRemaining(FTimerDynamicDelegate const& InDynDelegate)
 	{
-		return InternalGetTimerRemaining( FTimerUnifiedDelegate(InDynDelegate) );
+		FTimerData const* const TimerData = FindTimer( FTimerUnifiedDelegate(InDynDelegate) );
+		return InternalGetTimerRemaining(TimerData);
+	}
+	/** Version that takes a handle */
+	FORCEINLINE float GetTimerRemaining(FTimerHandle const& InHandle)
+	{
+		FTimerData const* const TimerData = FindTimer(InHandle);
+		return InternalGetTimerRemaining(TimerData);
 	}
 
 	/** Timer manager has been ticked this frame? */
@@ -468,22 +694,28 @@ public:
 
 private:
 
-	void InternalSetTimer( FTimerUnifiedDelegate const& InDelegate, float inRate, bool inbLoop );
+	void InternalSetTimer( FTimerUnifiedDelegate const& InDelegate, float InRate, bool InbLoop, float InFirstDelay );
+	void InternalSetTimer( FTimerHandle& InOutHandle, FTimerUnifiedDelegate const& InDelegate, float InRate, bool InbLoop, float InFirstDelay );
+	void InternalSetTimer( FTimerData& NewTimerData, float InRate, bool InbLoop, float InFirstDelay );
+	void InternalSetTimerForNextTick( FTimerUnifiedDelegate const& InDelegate );
 	void InternalClearTimer( FTimerUnifiedDelegate const& InDelegate );
-	void InternalClearAllTimers(void const* Object);
+	void InternalClearTimer( FTimerHandle const& InDelegate );
+	void InternalClearAllTimers( void const* Object );
 
-	/** Will find a timer either in the active or paused list. */
+	/** Will find a timer in the active, paused, or pending list. */
 	FTimerData const* FindTimer( FTimerUnifiedDelegate const& InDelegate, int32* OutTimerIndex=NULL ) const;
+	FTimerData const* FindTimer( FTimerHandle const& InHandle, int32* OutTimerIndex = NULL ) const;
 
-	void InternalPauseTimer( FTimerUnifiedDelegate const& InDelegate );
-	void InternalUnPauseTimer( FTimerUnifiedDelegate const& InDelegate );
-
-	float InternalGetTimerRate( FTimerUnifiedDelegate const& InDelegate ) const;
-	float InternalGetTimerElapsed( FTimerUnifiedDelegate const& InDelegate ) const;
-	float InternalGetTimerRemaining( FTimerUnifiedDelegate const& InDelegate ) const;
+	void InternalPauseTimer( FTimerData const* TimerToPause, int32 TimerIdx );
+	void InternalUnPauseTimer( int32 PausedTimerIdx );
+	
+	float InternalGetTimerRate( FTimerData const* const TimerData ) const;
+	float InternalGetTimerElapsed( FTimerData const* const TimerData ) const;
+	float InternalGetTimerRemaining( FTimerData const* const TimerData ) const;
 
 	/** Will find the given timer in the given TArray and return its index. */ 
-	int32 FindTimerInList(const TArray<FTimerData> &SearchArray, FTimerUnifiedDelegate const& InDelegate ) const;
+	int32 FindTimerInList( const TArray<FTimerData> &SearchArray, FTimerUnifiedDelegate const& InDelegate ) const;
+	int32 FindTimerInList( const TArray<FTimerData> &SearchArray, FTimerHandle const& InHandle ) const;
 
 	/** Heap of actively running timers. */
 	TArray<FTimerData> ActiveTimerHeap;
