Fix CIS error with EdGraphSchema undefined

diff --git a/Engine/Source/Runtime/Engine/Classes/EdGraph/EdGraphNode.h b/Engine/Source/Runtime/Engine/Classes/EdGraph/EdGraphNode.h
--- a/Engine/Source/Runtime/Engine/Classes/EdGraph/EdGraphNode.h
+++ b/Engine/Source/Runtime/Engine/Classes/EdGraph/EdGraphNode.h
@@ -7,6 +7,7 @@
 class UBlueprint;
 class UEdGraph;
 class UEdGraphNode;
+class UEdGraphSchema;
 class UEdGraphPin;
 class SGraphNode;
 
@@ -322,7 +323,7 @@ public:
 	virtual void NodeConnectionListChanged() {}
 
 	/** Shorthand way to access the schema of the graph that owns this node */
-	const class UEdGraphSchema* GetSchema() const;
+	const UEdGraphSchema* GetSchema() const;
 
 	/** Whether or not this node can be safely duplicated (via copy/paste, etc...) in the graph */
 	virtual bool CanDuplicateNode() const;
