removed NavigationSystem::bWholeWorldNavigable from editable properties - we might want to remove the feature #UE4

diff --git a/Engine/Source/Runtime/Engine/Classes/AI/Navigation/NavigationSystem.h b/Engine/Source/Runtime/Engine/Classes/AI/Navigation/NavigationSystem.h
--- a/Engine/Source/Runtime/Engine/Classes/AI/Navigation/NavigationSystem.h
+++ b/Engine/Source/Runtime/Engine/Classes/AI/Navigation/NavigationSystem.h
@@ -94,7 +94,9 @@ class ENGINE_API UNavigationSystem : public UBlueprintFunctionLibrary
 	 *	bounds volumes (like ANavMeshBoundsVolume). Set to false means navigation should be generated
 	 *	everywhere.
 	 */
-	UPROPERTY(config, EditAnywhere, Category=NavigationSystem)
+	// @todo removing it from edition since it's currently broken and I'm not sure we want that at all
+	// since I'm not sure we can make it efficient in a generic case
+	//UPROPERTY(config, EditAnywhere, Category=NavigationSystem)
 	uint32 bWholeWorldNavigable:1;
 
 	/** If set to true (default) generation seeds will include locations of all player controlled pawns */
