Changed SoundWave Compression Quality so that it doesn't re-compress during interactive property changes and the new compression settings can be heard next time the sound is played.

TTP#256433 - EDITOR: SOUND: Unable to change the compression quality for sound wave files.

Made sure that Sound resources are freed after compressed data is invalidated by a quality change. Also had to make sure that all playing sounds using the SoundWave were stopped so that realtime sounds don't attempt to keep reading from freed memory. Renamed StopSoundsForReimport to StopSoundsUsingResource to make it clear that it can be used for other purposes and added a warning in case this is used in game code and sounds stop unexpectedly.

diff --git a/Engine/Source/Editor/UnrealEd/Private/Factories/EditorFactories.cpp b/Engine/Source/Editor/UnrealEd/Private/Factories/EditorFactories.cpp
--- a/Engine/Source/Editor/UnrealEd/Private/Factories/EditorFactories.cpp
+++ b/Engine/Source/Editor/UnrealEd/Private/Factories/EditorFactories.cpp
@@ -1633,7 +1633,7 @@ UObject* USoundFactory::FactoryCreateBinary
 		FAudioDevice* AudioDevice = GEngine->GetAudioDevice();
 		if (AudioDevice && ExistingSound)
 		{
-			AudioDevice->StopSoundsForReimport(ExistingSound, ComponentsToRestart);
+			AudioDevice->StopSoundsUsingResource(ExistingSound, ComponentsToRestart);
 		}
 
 		bool bUseExistingSettings = bSoundFactorySuppressImportOverwriteDialog;
@@ -1717,9 +1717,6 @@ UObject* USoundFactory::FactoryCreateBinary
 		{
 			// Clear resources so that if it's already been played, it will reload the wave data
 			Sound->FreeResources();
-			Sound->DecompressionType = DTYPE_Setup;
-			Sound->bDecompressedFromOgg = 0;
-			Sound->RawPCMDataSize = 0;
 		}
 		
 		// Store the current file path and timestamp for re-import purposes
diff --git a/Engine/Source/Runtime/Engine/Private/AudioDevice.cpp b/Engine/Source/Runtime/Engine/Private/AudioDevice.cpp
--- a/Engine/Source/Runtime/Engine/Private/AudioDevice.cpp
+++ b/Engine/Source/Runtime/Engine/Private/AudioDevice.cpp
@@ -2498,21 +2498,23 @@ bool FAudioDevice::HandleListSoundsCommand(const TCHAR* Cmd, FOutputDevice& Ar)
 }
 #endif // !UE_BUILD_SHIPPING
 
-#if WITH_EDITOR
-void FAudioDevice::StopSoundsForReimport(USoundWave* ReimportedSoundWave, TArray<UAudioComponent*>& ComponentsToRestart)
+void FAudioDevice::StopSoundsUsingResource(USoundWave* SoundWave, TArray<UAudioComponent*>& StoppedComponents)
 {
+	bool bStoppedSounds = false;
+
 	for (int32 ActiveSoundIndex = ActiveSounds.Num() - 1; ActiveSoundIndex >= 0; --ActiveSoundIndex)
 	{
 		FActiveSound* ActiveSound = ActiveSounds[ActiveSoundIndex];
 		if (ActiveSound->Sound->IsA(USoundWave::StaticClass()))
 		{
-			if (ActiveSound->Sound == ReimportedSoundWave)
+			if (ActiveSound->Sound == SoundWave)
 			{
 				if (ActiveSound->AudioComponent.IsValid())
 				{
-					ComponentsToRestart.Add(ActiveSound->AudioComponent.Get());
+					StoppedComponents.Add(ActiveSound->AudioComponent.Get());
 				}
 				ActiveSound->Stop(this);
+				bStoppedSounds = true;
 			}
 		}
 		else
@@ -2521,22 +2523,29 @@ void FAudioDevice::StopSoundsForReimport(USoundWave* ReimportedSoundWave, TArray
 
 			for (auto WaveInstanceIt(ActiveSound->WaveInstances.CreateConstIterator()); WaveInstanceIt; ++WaveInstanceIt)
 			{
-				// If anything in the SoundCue uses the wave we're going to restart the whole thing for simplicity
+				// If anything in the SoundCue uses the wave we're going to stop the whole thing for simplicity
 				FWaveInstance* WaveInstance = WaveInstanceIt.Value();
-				if (WaveInstance->WaveData == ReimportedSoundWave)
+				if (WaveInstance->WaveData == SoundWave)
 				{
 					if (ActiveSound->AudioComponent.IsValid())
 					{
-						ComponentsToRestart.Add(ActiveSound->AudioComponent.Get());
+						StoppedComponents.Add(ActiveSound->AudioComponent.Get());
 					}
 					ActiveSound->Stop(this);
+					bStoppedSounds = true;
 					break;
 				}
 			}
 		}
 	}
+
+	if (!GIsEditor && bStoppedSounds)
+	{
+		UE_LOG(LogAudio, Warning, TEXT( "All Sounds using SoundWave '%s' have been stopped" ), *SoundWave->GetName() );
+	}
 }
 
+#if WITH_EDITOR
 void FAudioDevice::OnBeginPIE(const bool bIsSimulating)
 {
 	for (TObjectIterator<USoundNode> It; It; ++It)
diff --git a/Engine/Source/Runtime/Engine/Private/SoundWave.cpp b/Engine/Source/Runtime/Engine/Private/SoundWave.cpp
--- a/Engine/Source/Runtime/Engine/Private/SoundWave.cpp
+++ b/Engine/Source/Runtime/Engine/Private/SoundWave.cpp
@@ -314,12 +314,17 @@ void USoundWave::PostEditChangeProperty(FPropertyChangedEvent& PropertyChangedEv
 
 	static FName CompressionQualityFName = FName( TEXT( "CompressionQuality" ) );
 
-	UProperty* PropertyThatChanged = PropertyChangedEvent.Property;
-	// Regenerate on save any compressed sound formats
-	if ( PropertyThatChanged && PropertyThatChanged->GetFName() == CompressionQualityFName )
+	// Prevent constant re-compression of SoundWave while properties are being changed interactively
+	if (PropertyChangedEvent.ChangeType != EPropertyChangeType::Interactive)
 	{
-		InvalidateCompressedData();
-		MarkPackageDirty();
+		UProperty* PropertyThatChanged = PropertyChangedEvent.Property;
+		// Regenerate on save any compressed sound formats
+		if ( PropertyThatChanged && PropertyThatChanged->GetFName() == CompressionQualityFName )
+		{
+			InvalidateCompressedData();
+			FreeResources();
+			MarkPackageDirty();
+		}
 	}
 }
 #endif // WITH_EDITOR
@@ -338,6 +343,8 @@ void USoundWave::FreeResources()
 		FAudioDevice* AudioDevice = GEngine->GetAudioDevice();
 		if (AudioDevice != NULL)
 		{
+			TArray<UAudioComponent*> StoppedComponents;
+			AudioDevice->StopSoundsUsingResource(this, StoppedComponents);
 			AudioDevice->FreeResource( this );
 		}
 	}
@@ -347,6 +354,9 @@ void USoundWave::FreeResources()
 	Duration = 0.0f;
 	ResourceID = 0;
 	bDynamicResource = false;
+	DecompressionType = DTYPE_Setup;
+	bDecompressedFromOgg = 0;
+	RawPCMDataSize = 0;
 }
 
 FWaveInstance* USoundWave::HandleStart( FActiveSound& ActiveSound, const UPTRINT WaveInstanceHash )
diff --git a/Engine/Source/Runtime/Engine/Public/AudioDevice.h b/Engine/Source/Runtime/Engine/Public/AudioDevice.h
--- a/Engine/Source/Runtime/Engine/Public/AudioDevice.h
+++ b/Engine/Source/Runtime/Engine/Public/AudioDevice.h
@@ -261,10 +261,15 @@ public:
 	 */
 	void Flush( class UWorld* WorldToFlush, bool bClearActivatedReverb = true );
 
-#if WITH_EDITOR
-	/* Stop any playing sounds so that we can reimport a specific sound wave */
-	void StopSoundsForReimport(USoundWave* ReimportedSoundWave, TArray<UAudioComponent*>& ComponentsToRestart);
+	/**
+	 * Stop any playing sounds that are using a particular SoundWave
+	 *
+	 * @param SoundWave					Resource to stop any sounds that are using it
+	 * @param[out] StoppedComponents	List of Audio Components that were stopped
+	 */
+	void StopSoundsUsingResource(USoundWave* SoundWave, TArray<UAudioComponent*>& StoppedComponents);
 
+#if WITH_EDITOR
 	/** Deals with anything audio related that should happen when PIE starts */
 	void OnBeginPIE(const bool bIsSimulating);
 
