Optimize ArrowComponent by each owning its own vertex/index buffer rather than creating it from scratch each frame

diff --git a/Engine/Source/Runtime/Engine/Private/ArrowComponent.cpp b/Engine/Source/Runtime/Engine/Private/ArrowComponent.cpp
--- a/Engine/Source/Runtime/Engine/Private/ArrowComponent.cpp
+++ b/Engine/Source/Runtime/Engine/Private/ArrowComponent.cpp
@@ -1,6 +1,7 @@
 // Copyright 1998-2014 Epic Games, Inc. All Rights Reserved.
 
 #include "EnginePrivate.h"
+#include "DynamicMeshBuilder.h"
 
 #define DEFAULT_SCREEN_SIZE	(0.0025f)
 #define ARROW_SCALE			(80.0f)
@@ -12,6 +13,75 @@
 float UArrowComponent::EditorScale = 1.0f;
 #endif
 
+/** Vertex Buffer */
+class FArrowVertexBuffer : public FVertexBuffer
+{
+public:
+	TArray<FDynamicMeshVertex> Vertices;
+
+	virtual void InitRHI()
+	{
+		FRHIResourceCreateInfo CreateInfo;
+		VertexBufferRHI = RHICreateVertexBuffer(Vertices.Num() * sizeof(FDynamicMeshVertex), BUF_Static, CreateInfo);
+
+		// Copy the vertex data into the vertex buffer.
+		void* VertexBufferData = RHILockVertexBuffer(VertexBufferRHI, 0, Vertices.Num() * sizeof(FDynamicMeshVertex), RLM_WriteOnly);
+		FMemory::Memcpy(VertexBufferData, Vertices.GetTypedData(), Vertices.Num() * sizeof(FDynamicMeshVertex));
+		RHIUnlockVertexBuffer(VertexBufferRHI);
+	}
+
+};
+
+/** Index Buffer */
+class FArrowIndexBuffer : public FIndexBuffer
+{
+public:
+	TArray<int32> Indices;
+
+	virtual void InitRHI()
+	{
+		FRHIResourceCreateInfo CreateInfo;
+		IndexBufferRHI = RHICreateIndexBuffer(sizeof(int32), Indices.Num() * sizeof(int32), BUF_Static, CreateInfo);
+
+		// Write the indices to the index buffer.
+		void* Buffer = RHILockIndexBuffer(IndexBufferRHI, 0, Indices.Num() * sizeof(int32), RLM_WriteOnly);
+		FMemory::Memcpy(Buffer, Indices.GetTypedData(), Indices.Num() * sizeof(int32));
+		RHIUnlockIndexBuffer(IndexBufferRHI);
+	}
+};
+
+/** Vertex Factory */
+class FArrowVertexFactory : public FLocalVertexFactory
+{
+public:
+
+	FArrowVertexFactory()
+	{}
+
+
+	/** Initialization */
+	void Init(const FArrowVertexBuffer* VertexBuffer)
+	{
+		check(!IsInRenderingThread());
+
+		ENQUEUE_UNIQUE_RENDER_COMMAND_TWOPARAMETER(
+			InitArrowVertexFactory,
+			FArrowVertexFactory*, VertexFactory, this,
+			const FArrowVertexBuffer*, VertexBuffer, VertexBuffer,
+			{
+			// Initialize the vertex factory's stream components.
+			DataType NewData;
+			NewData.PositionComponent = STRUCTMEMBER_VERTEXSTREAMCOMPONENT(VertexBuffer, FDynamicMeshVertex, Position, VET_Float3);
+			NewData.TextureCoordinates.Add(
+				FVertexStreamComponent(VertexBuffer, STRUCT_OFFSET(FDynamicMeshVertex, TextureCoordinate), sizeof(FDynamicMeshVertex), VET_Float2)
+				);
+			NewData.TangentBasisComponents[0] = STRUCTMEMBER_VERTEXSTREAMCOMPONENT(VertexBuffer, FDynamicMeshVertex, TangentX, VET_PackedNormal);
+			NewData.TangentBasisComponents[1] = STRUCTMEMBER_VERTEXSTREAMCOMPONENT(VertexBuffer, FDynamicMeshVertex, TangentZ, VET_PackedNormal);
+			VertexFactory->SetData(NewData);
+		});
+	}
+};
+
 /** Represents a UArrowComponent to the scene manager. */
 class FArrowSceneProxy : public FPrimitiveSceneProxy
 {
@@ -38,6 +108,32 @@ public:
 			SpriteCategoryIndex = GEngine->GetSpriteCategoryIndex( Component->SpriteInfo.Category );
 		}
 #endif	//WITH_EDITOR
+
+		const float HeadAngle = FMath::DegreesToRadians(ARROW_HEAD_ANGLE);
+		const float TotalLength = ArrowSize * ARROW_SCALE;
+		const float HeadLength = TotalLength * ARROW_HEAD_FACTOR;
+		const float ShaftRadius = TotalLength * ARROW_RADIUS_FACTOR;
+		const float ShaftLength = (TotalLength - HeadLength) * 1.1f; // 10% overlap between shaft and head
+		const FVector ShaftCenter = FVector(0.5f * ShaftLength, 0, 0);
+
+		BuildConeVerts(HeadAngle, HeadAngle, -HeadLength, TotalLength, 32, VertexBuffer.Vertices, IndexBuffer.Indices);
+		BuildCylinderVerts(ShaftCenter, FVector(0,0,1), FVector(0,1,0), FVector(1,0,0), ShaftRadius, 0.5f * ShaftLength, 16, VertexBuffer.Vertices, IndexBuffer.Indices);
+
+
+		// Init vertex factory
+		VertexFactory.Init(&VertexBuffer);
+
+		// Enqueue initialization of render resource
+		BeginInitResource(&VertexBuffer);
+		BeginInitResource(&IndexBuffer);
+		BeginInitResource(&VertexFactory);
+	}
+
+	virtual ~FArrowSceneProxy()
+	{
+		VertexBuffer.ReleaseResource();
+		IndexBuffer.ReleaseResource();
+		VertexFactory.ReleaseResource();
 	}
 
 	// FPrimitiveSceneProxy interface.
@@ -64,27 +160,21 @@ public:
 			EffectiveLocalToWorld = GetLocalToWorld();
 		}
 
-		const FColoredMaterialRenderProxy CollisionMaterialInstance(
+		const FColoredMaterialRenderProxy ArrowMaterialRenderProxy(
 			GEngine->ArrowMaterial->GetRenderProxy(IsSelected(), IsHovered()),
 			ArrowColor,
 			"GizmoColor"
 			);
 
-		const float TotalLength = ArrowSize * ARROW_SCALE;
-		const float HeadLength = TotalLength * ARROW_HEAD_FACTOR;
-		const float ShaftRadius = TotalLength * ARROW_RADIUS_FACTOR;
-		const float ShaftLength = (TotalLength - HeadLength) * 1.1f; // 10% overlap between shaft and head
-		const FVector Offset(0, 0, 0.5f * ShaftLength);
-		const float HeadAngle = FMath::DegreesToRadians(ARROW_HEAD_ANGLE);
 
 		// Calculate the view-dependent scaling factor.
 		float ViewScale = 1.0f;
 		if (bIsScreenSizeScaled && (View->ViewMatrices.ProjMatrix.M[3][3] != 1.0f))
 		{
-			const float ZoomFactor	= FMath::Min<float>(View->ViewMatrices.ProjMatrix.M[0][0], View->ViewMatrices.ProjMatrix.M[1][1]);
-			if(ZoomFactor != 0.0f)
+			const float ZoomFactor = FMath::Min<float>(View->ViewMatrices.ProjMatrix.M[0][0], View->ViewMatrices.ProjMatrix.M[1][1]);
+			if (ZoomFactor != 0.0f)
 			{
-				const float Radius		= View->WorldToScreen(Origin).W * (ScreenSize / ZoomFactor);
+				const float Radius = View->WorldToScreen(Origin).W * (ScreenSize / ZoomFactor);
 				if (Radius < 1.0f)
 				{
 					ViewScale *= Radius;
@@ -96,11 +186,33 @@ public:
 		ViewScale *= EditorScale;
 #endif
 
-		const FMatrix CylinderTransform = FScaleMatrix(ViewScale) * FRotationMatrix( FRotator(-90,0.f,0) ) * FTranslationMatrix(FVector(ViewScale, 0.0f, 0.0f)) * EffectiveLocalToWorld;
-		DrawCylinder(PDI, CylinderTransform, Offset, FVector(1,0,0), FVector(0,1,0), FVector(0,0,1), ShaftRadius, 0.5f * ShaftLength, 16, &CollisionMaterialInstance, SDPG_World);
+
+		// Draw the mesh.
+		FMeshBatch Mesh;
+		FMeshBatchElement& BatchElement = Mesh.Elements[0];
+		BatchElement.IndexBuffer = &IndexBuffer;
+		Mesh.bWireframe = false;
+		Mesh.VertexFactory = &VertexFactory;
+		Mesh.MaterialRenderProxy = &ArrowMaterialRenderProxy;
+		BatchElement.PrimitiveUniformBuffer = CreatePrimitiveUniformBufferImmediate(FScaleMatrix(ViewScale) * EffectiveLocalToWorld, GetBounds(), GetLocalBounds(), true);
+		BatchElement.FirstIndex = 0;
+		BatchElement.NumPrimitives = IndexBuffer.Indices.Num() / 3;
+		BatchElement.MinVertexIndex = 0;
+		BatchElement.MaxVertexIndex = VertexBuffer.Vertices.Num() - 1;
+		Mesh.ReverseCulling = IsLocalToWorldDeterminantNegative();
+		Mesh.Type = PT_TriangleList;
+		Mesh.DepthPriorityGroup = SDPG_World;
+		PDI->DrawMesh(Mesh);
+
+
+
+
+
+		//const FMatrix CylinderTransform = FScaleMatrix(ViewScale) * FRotationMatrix( FRotator(-90,0.f,0) ) * FTranslationMatrix(FVector(ViewScale, 0.0f, 0.0f)) * EffectiveLocalToWorld;
+		//DrawCylinder(PDI, CylinderTransform, Offset, FVector(1,0,0), FVector(0,1,0), FVector(0,0,1), ShaftRadius, 0.5f * ShaftLength, 16, &ArrowMaterialRenderProxy, SDPG_World);
 		
-		const FMatrix ConeTransform = FScaleMatrix(ViewScale * -HeadLength) * FTranslationMatrix(FVector(ViewScale * TotalLength,0,0)) * EffectiveLocalToWorld;
-		DrawCone(PDI, ConeTransform, HeadAngle, HeadAngle, 32, false, FColor::White, &CollisionMaterialInstance, SDPG_World);
+		//const FMatrix ConeTransform = FScaleMatrix(ViewScale * -HeadLength) * FTranslationMatrix(FVector(ViewScale * TotalLength,0,0)) * EffectiveLocalToWorld;
+		//DrawCone(PDI, ConeTransform, HeadAngle, HeadAngle, 32, false, FColor::White, &CollisionMaterialInstance, SDPG_World);
 	}
 
 	virtual FPrimitiveViewRelevance GetViewRelevance(const FSceneView* View)
@@ -121,14 +233,20 @@ public:
 		Result.bEditorPrimitiveRelevance = UseEditorCompositing(View);
 		return Result;
 	}
+
 	virtual void OnTransformChanged() OVERRIDE
 	{
 		Origin = GetLocalToWorld().GetOrigin();
 	}
+
 	virtual uint32 GetMemoryFootprint( void ) const { return( sizeof( *this ) + GetAllocatedSize() ); }
 	uint32 GetAllocatedSize( void ) const { return( FPrimitiveSceneProxy::GetAllocatedSize() ); }
 
 private:
+	FArrowVertexBuffer VertexBuffer;
+	FArrowIndexBuffer IndexBuffer;
+	FArrowVertexFactory VertexFactory;
+
 	FVector Origin;
 	FColor ArrowColor;
 	float ArrowSize;
diff --git a/Engine/Source/Runtime/Engine/Private/PrimitiveDrawingUtils.cpp b/Engine/Source/Runtime/Engine/Private/PrimitiveDrawingUtils.cpp
--- a/Engine/Source/Runtime/Engine/Private/PrimitiveDrawingUtils.cpp
+++ b/Engine/Source/Runtime/Engine/Private/PrimitiveDrawingUtils.cpp
@@ -214,7 +214,7 @@ void DrawSphere(FPrimitiveDrawInterface* PDI,const FVector& Center,const FVector
 	MeshBuilder.Draw(PDI, FScaleMatrix( Radii ) * FTranslationMatrix( Center ), MaterialRenderProxy, DepthPriority,bDisableBackfaceCulling);
 }
 
-void DrawCone(FPrimitiveDrawInterface* PDI,const FMatrix& ConeToWorld, float Angle1, float Angle2, int32 NumSides, bool bDrawSideLines, const FLinearColor& SideLineColor, const FMaterialRenderProxy* MaterialRenderProxy, uint8 DepthPriority)
+FVector CalcConeVert(float Angle1, float Angle2, float AzimuthAngle)
 {
 	float ang1 = FMath::Clamp<float>(Angle1, 0.01f, (float)PI - 0.01f);
 	float ang2 = FMath::Clamp<float>(Angle2, 0.01f, (float)PI - 0.01f);
@@ -228,107 +228,117 @@ void DrawCone(FPrimitiveDrawInterface* PDI,const FMatrix& ConeToWorld, float Ang
 	float tanX_2 = FMath::Tan(0.5f * ang1);
 	float tanY_2 = FMath::Tan(0.5f * ang2);
 
-	TArray<FVector> ConeVerts;
-	ConeVerts.AddUninitialized(NumSides);
 
-	for(int32 i = 0; i < NumSides; i++)
-	{
-		float Fraction = (float)i/(float)(NumSides);
-		float thi = 2.f*PI*Fraction;
-		float phi = FMath::Atan2(FMath::Sin(thi)*sinY_2, FMath::Cos(thi)*sinX_2);
-		float sinPhi = FMath::Sin(phi);
-		float cosPhi = FMath::Cos(phi);
-		float sinSqPhi = sinPhi*sinPhi;
-		float cosSqPhi = cosPhi*cosPhi;
-
-		float rSq, r, Sqr, alpha, beta;
-
-		rSq = sinSqX_2*sinSqY_2/(sinSqX_2*sinSqPhi + sinSqY_2*cosSqPhi);
-		r = FMath::Sqrt(rSq);
-		Sqr = FMath::Sqrt(1-rSq);
-		alpha = r*cosPhi;
-		beta  = r*sinPhi;
-
-		ConeVerts[i].X = (1-2*rSq);
-		ConeVerts[i].Y = 2*Sqr*alpha;
-		ConeVerts[i].Z = 2*Sqr*beta;
-	}
-
-	FDynamicMeshBuilder MeshBuilder;
-	{
-		for(int32 i=0; i < NumSides; i++)
-		{
-			FDynamicMeshVertex V0, V1, V2;
+	float phi = FMath::Atan2(FMath::Sin(AzimuthAngle)*sinY_2, FMath::Cos(AzimuthAngle)*sinX_2);
+	float sinPhi = FMath::Sin(phi);
+	float cosPhi = FMath::Cos(phi);
+	float sinSqPhi = sinPhi*sinPhi;
+	float cosSqPhi = cosPhi*cosPhi;
 
-			// Normal of the previous face connected to this face
-			FVector TriTangentZPrev = ConeVerts[i] ^ ConeVerts [ i == 0 ? NumSides - 1 : i-1 ];
+	float rSq, r, Sqr, alpha, beta;
 
-			// Normal of the current face 
-			FVector TriTangentZ = ConeVerts[(i+1)%NumSides] ^ ConeVerts[ i ]; // aka triangle normal
+	rSq = sinSqX_2*sinSqY_2 / (sinSqX_2*sinSqPhi + sinSqY_2*cosSqPhi);
+	r = FMath::Sqrt(rSq);
+	Sqr = FMath::Sqrt(1 - rSq);
+	alpha = r*cosPhi;
+	beta = r*sinPhi;
 
-			// Normal of the next face connected to this face
-			FVector TriTangentZNext = ConeVerts[(i+2)%NumSides] ^ ConeVerts[(i+1)%NumSides];
+	FVector ConeVert;
 
-			FVector TriTangentY = ConeVerts[i];
-			FVector TriTangentX = TriTangentZ ^ TriTangentY;
+	ConeVert.X = (1 - 2 * rSq);
+	ConeVert.Y = 2 * Sqr*alpha;
+	ConeVert.Z = 2 * Sqr*beta;
 
-			{
-				V0.Position = FVector(0);
-				V0.TextureCoordinate.X = 0.0f;
-				V0.TextureCoordinate.Y = (float)i/NumSides;
-
-				V0.SetTangents(TriTangentX,TriTangentY,FVector(-1,0,0));
-			}
-		
-			{
-				V1.Position = ConeVerts[i];
-				V1.TextureCoordinate.X = 1.0f;
-				V1.TextureCoordinate.Y = (float)i/NumSides;
-
-				FVector VertexNormal = (TriTangentZPrev + TriTangentZ).SafeNormal();
-				V1.SetTangents(TriTangentX,TriTangentY,VertexNormal);
-			}
+	return ConeVert;
+}
 
-	
-			{
-				V2.Position = ConeVerts[(i+1)%NumSides];
-				V2.TextureCoordinate.X = 1.0f;
-				V2.TextureCoordinate.Y = (float)((i+1)%NumSides)/NumSides;
-				V2.SetTangents(TriTangentX,TriTangentY,TriTangentZ);
-			
-				FVector VertexNormal = (TriTangentZNext + TriTangentZ).SafeNormal();
-				V2.SetTangents(TriTangentX,TriTangentY,VertexNormal);
-			}
+void BuildConeVerts(float Angle1, float Angle2, float Scale, float XOffset, int32 NumSides, TArray<FDynamicMeshVertex>& OutVerts, TArray<int32>& OutIndices)
+{
+	TArray<FVector> ConeVerts;
+	ConeVerts.AddUninitialized(NumSides);
 
-			const int32 VertexStart = MeshBuilder.AddVertex(V0);
-			MeshBuilder.AddVertex(V1);
-			MeshBuilder.AddVertex(V2);
+	for (int32 i = 0; i < NumSides; i++)
+	{
+		float Fraction = (float)i / (float)(NumSides);
+		float Azi = 2.f*PI*Fraction;
+		ConeVerts[i] = (CalcConeVert(Angle1, Angle2, Azi) * Scale) + FVector(XOffset,0,0);
+	}
 
-			MeshBuilder.AddTriangle(VertexStart,VertexStart+1,VertexStart+2);
+	for (int32 i = 0; i < NumSides; i++)
+	{
+		// Normal of the current face 
+		FVector TriTangentZ = ConeVerts[(i + 1) % NumSides] ^ ConeVerts[i]; // aka triangle normal
+		FVector TriTangentY = ConeVerts[i];
+		FVector TriTangentX = TriTangentZ ^ TriTangentY;
+
+
+		FDynamicMeshVertex V0, V1, V2;
+
+		V0.Position = FVector(0) + FVector(XOffset,0,0);
+		V0.TextureCoordinate.X = 0.0f;
+		V0.TextureCoordinate.Y = (float)i / NumSides;
+		V0.SetTangents(TriTangentX, TriTangentY, FVector(-1, 0, 0));
+		int32 I0 = OutVerts.Add(V0);
+
+		V1.Position = ConeVerts[i];
+		V1.TextureCoordinate.X = 1.0f;
+		V1.TextureCoordinate.Y = (float)i / NumSides;
+		FVector TriTangentZPrev = ConeVerts[i] ^ ConeVerts[i == 0 ? NumSides - 1 : i - 1]; // Normal of the previous face connected to this face
+		V1.SetTangents(TriTangentX, TriTangentY, (TriTangentZPrev + TriTangentZ).SafeNormal());
+		int32 I1 = OutVerts.Add(V1);
+
+		V2.Position = ConeVerts[(i + 1) % NumSides];
+		V2.TextureCoordinate.X = 1.0f;
+		V2.TextureCoordinate.Y = (float)((i + 1) % NumSides) / NumSides;
+		FVector TriTangentZNext = ConeVerts[(i + 2) % NumSides] ^ ConeVerts[(i + 1) % NumSides]; // Normal of the next face connected to this face
+		V2.SetTangents(TriTangentX, TriTangentY, (TriTangentZNext + TriTangentZ).SafeNormal());
+		int32 I2 = OutVerts.Add(V2);
+
+		// Flip winding for negative scale
+		if(Scale >= 0.f)
+		{
+			OutIndices.Add(I0);
+			OutIndices.Add(I1);
+			OutIndices.Add(I2);
+		}
+		else
+		{
+			OutIndices.Add(I0);
+			OutIndices.Add(I2);
+			OutIndices.Add(I1);
 		}
 	}
-	MeshBuilder.Draw(PDI, ConeToWorld, MaterialRenderProxy, DepthPriority, 0.f);
+}
 
+void DrawCone(FPrimitiveDrawInterface* PDI,const FMatrix& ConeToWorld, float Angle1, float Angle2, int32 NumSides, bool bDrawSideLines, const FLinearColor& SideLineColor, const FMaterialRenderProxy* MaterialRenderProxy, uint8 DepthPriority)
+{
+	TArray<FDynamicMeshVertex> MeshVerts;
+	TArray<int32> MeshIndices;
+	BuildConeVerts(Angle1, Angle2, 1.f, 0.f, NumSides, MeshVerts, MeshIndices);
+
+	FDynamicMeshBuilder MeshBuilder;
+	MeshBuilder.AddVertices(MeshVerts);
+	MeshBuilder.AddTriangles(MeshIndices);
+	MeshBuilder.Draw(PDI, ConeToWorld, MaterialRenderProxy, DepthPriority, 0.f);
 
 	if(bDrawSideLines)
 	{
+		TArray<FVector> ConeVerts;
+		ConeVerts.AddUninitialized(NumSides);
+
 		// Draw lines down major directions
 		for(int32 i=0; i<4; i++)
 		{
-			PDI->DrawLine( ConeToWorld.GetOrigin(), ConeToWorld.TransformPosition( ConeVerts[ (i*NumSides/4)%NumSides ] ), SideLineColor, DepthPriority );
+			float Fraction = (float)i / (float)(4);
+			float Azi = 2.f*PI*Fraction;
+			FVector ConeVert = CalcConeVert(Angle1, Angle2, Azi);
+			PDI->DrawLine( ConeToWorld.GetOrigin(), ConeToWorld.TransformPosition(ConeVert), SideLineColor, DepthPriority );
 		}
 	}
 }
 
 
-void DrawCylinder(FPrimitiveDrawInterface* PDI,const FVector& Base, const FVector& XAxis, const FVector& YAxis, const FVector& ZAxis,
-	float Radius, float HalfHeight, int32 Sides, const FMaterialRenderProxy* MaterialRenderProxy, uint8 DepthPriority)
-{
-	DrawCylinder( PDI, FMatrix::Identity, Base, XAxis, YAxis, ZAxis, Radius, HalfHeight, Sides, MaterialRenderProxy, DepthPriority );
-}
-
-void DrawCylinder(FPrimitiveDrawInterface* PDI, const FMatrix& CylToWorld, const FVector& Base, const FVector& XAxis, const FVector& YAxis, const FVector& ZAxis,
-	float Radius, float HalfHeight, int32 Sides, const FMaterialRenderProxy* MaterialRenderProxy, uint8 DepthPriority)
+void BuildCylinderVerts(const FVector& Base, const FVector& XAxis, const FVector& YAxis, const FVector& ZAxis, float Radius, float HalfHeight, int32 Sides, TArray<FDynamicMeshVertex>& OutVerts, TArray<int32>& OutIndices)
 {
 	const float	AngleDelta = 2.0f * PI / Sides;
 	FVector	LastVertex = Base + XAxis * Radius;
@@ -338,11 +348,10 @@ void DrawCylinder(FPrimitiveDrawInterface* PDI, const FMatrix& CylToWorld, const
 
 	FVector TopOffset = HalfHeight * ZAxis;
 
-	FDynamicMeshBuilder MeshBuilder;
-
+	int32 BaseVertIndex = OutVerts.Num();
 
 	//Compute vertices for base circle.
-	for(int32 SideIndex = 0;SideIndex < Sides;SideIndex++)
+	for (int32 SideIndex = 0; SideIndex < Sides; SideIndex++)
 	{
 		const FVector Vertex = Base + (XAxis * FMath::Cos(AngleDelta * (SideIndex + 1)) + YAxis * FMath::Sin(AngleDelta * (SideIndex + 1))) * Radius;
 		FVector Normal = Vertex - Base;
@@ -359,7 +368,7 @@ void DrawCylinder(FPrimitiveDrawInterface* PDI, const FMatrix& CylToWorld, const
 			Normal
 			);
 
-		MeshBuilder.AddVertex(MeshVertex); //Add bottom vertex
+		OutVerts.Add(MeshVertex); //Add bottom vertex
 
 		LastVertex = Vertex;
 		TC.X += TCStep;
@@ -369,7 +378,7 @@ void DrawCylinder(FPrimitiveDrawInterface* PDI, const FMatrix& CylToWorld, const
 	TC = FVector2D(0.0f, 1.0f);
 
 	//Compute vertices for the top circle
-	for(int32 SideIndex = 0;SideIndex < Sides;SideIndex++)
+	for (int32 SideIndex = 0; SideIndex < Sides; SideIndex++)
 	{
 		const FVector Vertex = Base + (XAxis * FMath::Cos(AngleDelta * (SideIndex + 1)) + YAxis * FMath::Sin(AngleDelta * (SideIndex + 1))) * Radius;
 		FVector Normal = Vertex - Base;
@@ -386,7 +395,7 @@ void DrawCylinder(FPrimitiveDrawInterface* PDI, const FMatrix& CylToWorld, const
 			Normal
 			);
 
-		MeshBuilder.AddVertex(MeshVertex); //Add top vertex
+		OutVerts.Add(MeshVertex); //Add top vertex
 
 		LastVertex = Vertex;
 		TC.X += TCStep;
@@ -395,29 +404,60 @@ void DrawCylinder(FPrimitiveDrawInterface* PDI, const FMatrix& CylToWorld, const
 	//Add top/bottom triangles, in the style of a fan.
 	//Note if we wanted nice rendering of the caps then we need to duplicate the vertices and modify
 	//texture/tangent coordinates.
-	for(int32 SideIndex = 1; SideIndex < Sides; SideIndex++)
+	for (int32 SideIndex = 1; SideIndex < Sides; SideIndex++)
 	{
-		int32 V0 = 0;
-		int32 V1 = SideIndex;
-		int32 V2 = (SideIndex + 1) % Sides;
-
-		MeshBuilder.AddTriangle(V0, V1, V2); //bottom
-		MeshBuilder.AddTriangle(Sides + V2, Sides + V1 , Sides + V0); //top
+		int32 V0 = BaseVertIndex;
+		int32 V1 = BaseVertIndex + SideIndex;
+		int32 V2 = BaseVertIndex + ((SideIndex + 1) % Sides);
+
+		//bottom
+		OutIndices.Add(V0);
+		OutIndices.Add(V1);
+		OutIndices.Add(V2);
+
+		// top
+		OutIndices.Add(Sides + V2);
+		OutIndices.Add(Sides + V1);
+		OutIndices.Add(Sides + V0);
 	}
 
 	//Add sides.
 
-	for(int32 SideIndex = 0; SideIndex < Sides; SideIndex++)
+	for (int32 SideIndex = 0; SideIndex < Sides; SideIndex++)
 	{
-		int32 V0 = SideIndex;
-		int32 V1 = (SideIndex + 1) % Sides;
+		int32 V0 = BaseVertIndex + SideIndex;
+		int32 V1 = BaseVertIndex + ((SideIndex + 1) % Sides);
 		int32 V2 = V0 + Sides;
 		int32 V3 = V1 + Sides;
 
-		MeshBuilder.AddTriangle(V0, V2, V1);
-		MeshBuilder.AddTriangle(V2, V3, V1);
+		OutIndices.Add(V0);
+		OutIndices.Add(V2);
+		OutIndices.Add(V1);
+
+		OutIndices.Add(V2);
+		OutIndices.Add(V3);
+		OutIndices.Add(V1);
 	}
 
+}
+
+void DrawCylinder(FPrimitiveDrawInterface* PDI,const FVector& Base, const FVector& XAxis, const FVector& YAxis, const FVector& ZAxis,
+	float Radius, float HalfHeight, int32 Sides, const FMaterialRenderProxy* MaterialRenderProxy, uint8 DepthPriority)
+{
+	DrawCylinder( PDI, FMatrix::Identity, Base, XAxis, YAxis, ZAxis, Radius, HalfHeight, Sides, MaterialRenderProxy, DepthPriority );
+}
+
+void DrawCylinder(FPrimitiveDrawInterface* PDI, const FMatrix& CylToWorld, const FVector& Base, const FVector& XAxis, const FVector& YAxis, const FVector& ZAxis, float Radius, float HalfHeight, int32 Sides, const FMaterialRenderProxy* MaterialRenderProxy, uint8 DepthPriority)
+{
+	TArray<FDynamicMeshVertex> MeshVerts;
+	TArray<int32> MeshIndices;
+	BuildCylinderVerts(Base, XAxis, YAxis, ZAxis, Radius, HalfHeight, Sides, MeshVerts, MeshIndices);
+
+
+	FDynamicMeshBuilder MeshBuilder;
+	MeshBuilder.AddVertices(MeshVerts);
+	MeshBuilder.AddTriangles(MeshIndices);
+
 	MeshBuilder.Draw(PDI, CylToWorld, MaterialRenderProxy, DepthPriority,0.f);
 }
 
diff --git a/Engine/Source/Runtime/Engine/Public/SceneManagement.h b/Engine/Source/Runtime/Engine/Public/SceneManagement.h
--- a/Engine/Source/Runtime/Engine/Public/SceneManagement.h
+++ b/Engine/Source/Runtime/Engine/Public/SceneManagement.h
@@ -29,6 +29,7 @@ class FVertexFactory;
 class FAtmosphericFogSceneInfo;
 class ISceneViewExtension;
 class HHitProxy;
+struct FDynamicMeshVertex;
 
 
 DECLARE_LOG_CATEGORY_EXTERN(LogBufferVisualization, Log, All);
@@ -3001,6 +3002,11 @@ extern ENGINE_API void DrawFrustumWireframe(
 	uint8 DepthPriority
 	);
 
+void BuildConeVerts(float Angle1, float Angle2, float Scale, float XOffset, int32 NumSides, TArray<FDynamicMeshVertex>& OutVerts, TArray<int32>& OutIndices);
+
+void BuildCylinderVerts(const FVector& Base, const FVector& XAxis, const FVector& YAxis, const FVector& ZAxis, float Radius, float HalfHeight, int32 Sides, TArray<FDynamicMeshVertex>& OutVerts, TArray<int32>& OutIndices);
+
+
 /**
  * Given a base color and a selection state, returns a color which accounts for the selection state.
  * @param BaseColor - The base color of the object.
