Added blueprint events for montage blending out/ending

#TTP 335600 - UE4: Animation: Add support for BP Delegate to AnimMontage events

#proj UE4

diff --git a/Engine/Source/Runtime/Engine/Classes/Animation/AnimInstance.h b/Engine/Source/Runtime/Engine/Classes/Animation/AnimInstance.h
--- a/Engine/Source/Runtime/Engine/Classes/Animation/AnimInstance.h
+++ b/Engine/Source/Runtime/Engine/Classes/Animation/AnimInstance.h
@@ -11,6 +11,22 @@
 struct FBoneContainer;
 class  USkeletalMesh;
 
+/**
+* Delegate for when Montage is completed, whether interrupted or finished
+* Weight of this montage is 0.f, so it stops contributing to output pose
+*
+* bInterrupted = true if it was not property finished
+*/
+DECLARE_DYNAMIC_MULTICAST_DELEGATE_TwoParams(FOnMontageEndedMCDelegate, class UAnimMontage*, Montage, bool, bInterrupted);
+
+/**
+* Delegate for when Montage started to blend out, whether interrupted or finished
+* DesiredWeight of this montage becomes 0.f, but this still contributes to the output pose
+*
+* bInterrupted = true if it was not property finished
+*/
+DECLARE_DYNAMIC_MULTICAST_DELEGATE_TwoParams(FOnMontageBlendingOutStartedMCDelegate, class UAnimMontage*, Montage, bool, bInterrupted);
+
 /** Enum for controlling which reference frame a controller is applied in. */
 UENUM()
 enum EBoneControlSpace
@@ -522,6 +538,15 @@ public:
 	void Montage_SetBlendingOutDelegate(FOnMontageBlendingOutStarted & OnMontageBlendingOut);
 	FOnMontageBlendingOutStarted* Montage_GetBlendingOutDelegate();
 
+
+	/** Called when a montage starts blending out, whether interrupted or finished */
+	UPROPERTY(BlueprintAssignable)
+	FOnMontageBlendingOutStartedMCDelegate OnMontageBlendingOut;
+	
+	/** Called when a montage has ended, whether interrupted or finished*/
+	UPROPERTY(BlueprintAssignable)
+	FOnMontageEndedMCDelegate OnMontageEnded;
+
 public:
 	/** Get Current Montage Position */
 	float Montage_GetPosition(UAnimMontage* Montage);
diff --git a/Engine/Source/Runtime/Engine/Private/Animation/AnimMontage.cpp b/Engine/Source/Runtime/Engine/Private/Animation/AnimMontage.cpp
--- a/Engine/Source/Runtime/Engine/Private/Animation/AnimMontage.cpp
+++ b/Engine/Source/Runtime/Engine/Private/Animation/AnimMontage.cpp
@@ -705,10 +705,11 @@ void FAnimMontageInstance::Stop(float BlendOut, bool bInterrupt)
 		// check where this function gets called and see how we calculate BlendTime
 		BlendTime = BlendOut;
 
-		if (OnMontageBlendingOutStarted.IsBound())
+		if (UAnimInstance* Inst = AnimInstance.Get())
 		{
-			OnMontageBlendingOutStarted.Execute(Montage, bInterrupted);
+			Inst->OnMontageBlendingOut.Broadcast(Montage, bInterrupt);
 		}
+		OnMontageBlendingOutStarted.ExecuteIfBound(Montage, bInterrupted);
 	}
 
 	if (BlendTime == 0.0f)
@@ -780,10 +781,11 @@ void FAnimMontageInstance::Terminate()
 	Montage = NULL;
 
 	// terminating, trigger end
-	if (OnMontageEnded.IsBound())
+	if (UAnimInstance* Inst = AnimInstance.Get())
 	{
-		OnMontageEnded.Execute(OldMontage, bInterrupted);
+		Inst->OnMontageEnded.Broadcast(Montage, bInterrupted);
 	}
+	OnMontageEnded.ExecuteIfBound(OldMontage, bInterrupted);
 }
 
 bool FAnimMontageInstance::ChangeNextSection(FName SectionName, FName NextSectionName)
