Expose SplineMeshComponent to BPs
Move SplineMeshComponent code from LandscapeSplines.cpp to SplineMeshComponent.cpp
Add a new header for SplineMeshSceneProxy
Add BP exposed funcs for setting start/end pos/tangent

diff --git a/Engine/Source/Runtime/Engine/Classes/Components/SplineMeshComponent.h b/Engine/Source/Runtime/Engine/Classes/Components/SplineMeshComponent.h
--- a/Engine/Source/Runtime/Engine/Classes/Components/SplineMeshComponent.h
+++ b/Engine/Source/Runtime/Engine/Classes/Components/SplineMeshComponent.h
@@ -25,43 +25,43 @@ struct FSplineMeshParams
 	GENERATED_USTRUCT_BODY()
 
 	/** Start location of spline, in component space */
-	UPROPERTY()
+	UPROPERTY(EditAnywhere, Category=SplineMesh)
 	FVector StartPos;
 
 	/** Start tangent of spline, in component space */
-	UPROPERTY()
+	UPROPERTY(EditAnywhere, Category=SplineMesh)
 	FVector StartTangent;
 
 	/** X and Y scale applied to mesh at start of spline */
-	UPROPERTY()
+	UPROPERTY(EditAnywhere, Category=SplineMesh, AdvancedDisplay)
 	FVector2D StartScale;
 
 	/** Roll around spline applied at start */
-	UPROPERTY()
+	UPROPERTY(EditAnywhere, Category=SplineMesh, AdvancedDisplay)
 	float StartRoll;
 
 	/** Starting offset of the mesh from the spline, in component space */
-	UPROPERTY()
+	UPROPERTY(EditAnywhere, Category=SplineMesh)
 	FVector2D StartOffset;
 
 	/** End location of spline, in component space */
-	UPROPERTY()
+	UPROPERTY(EditAnywhere, Category=SplineMesh)
 	FVector EndPos;
 
 	/** End tangent of spline, in component space */
-	UPROPERTY()
+	UPROPERTY(EditAnywhere, Category=SplineMesh)
 	FVector EndTangent;
 
 	/** X and Y scale applied to mesh at end of spline */
-	UPROPERTY()
+	UPROPERTY(EditAnywhere, Category=SplineMesh, AdvancedDisplay)
 	FVector2D EndScale;
 
 	/** Roll around spline applied at end */
-	UPROPERTY()
+	UPROPERTY(EditAnywhere, Category=SplineMesh, AdvancedDisplay)
 	float EndRoll;
 
 	/** Ending offset of the mesh from the spline, in component space */
-	UPROPERTY()
+	UPROPERTY(EditAnywhere, Category=SplineMesh, AdvancedDisplay)
 	FVector2D EndOffset;
 
 
@@ -81,29 +81,29 @@ struct FSplineMeshParams
 
 };
 
-UCLASS(MinimalAPI)
-class USplineMeshComponent : public UStaticMeshComponent, public IInterface_CollisionDataProvider
+UCLASS(ClassGroup=Rendering, hidecategories=(Physics), meta=(BlueprintSpawnableComponent))
+class ENGINE_API USplineMeshComponent : public UStaticMeshComponent, public IInterface_CollisionDataProvider
 {
 	GENERATED_UCLASS_BODY()
 
 	/** Spline that is used to deform mesh */
-	UPROPERTY()
+	UPROPERTY(EditAnywhere, Category=SplineMesh, meta=(ShowOnlyInnerProperties))
 	struct FSplineMeshParams SplineParams;
 
 	/** Axis (in component space) that is used to determine X axis for co-ordinates along spline */
-	UPROPERTY()
+	UPROPERTY(EditAnywhere, Category=SplineMesh)
 	FVector SplineXDir;
 
 	/** If true, will use smooth interpolation (ease in/out) for Scale, Roll, and Offset along this section of spline. If false, uses linear */
-	UPROPERTY()
+	UPROPERTY(EditAnywhere, Category=SplineMesh, AdvancedDisplay)
 	uint32 bSmoothInterpRollScale:1;
 
 	/** Chooses the forward axis for the spline mesh orientation */
-	UPROPERTY()
+	UPROPERTY(EditAnywhere, Category=SplineMesh)
 	TEnumAsByte<ESplineMeshAxis::Type> ForwardAxis;
 
 	// Physics data.
-	UPROPERTY(EditAnywhere, BlueprintReadWrite, editinline, Category=Physics)
+	UPROPERTY()
 	class UBodySetup* BodySetup;
 
 #if WITH_EDITORONLY_DATA
@@ -147,6 +147,21 @@ class USplineMeshComponent : public UStaticMeshComponent, public IInterface_Coll
 	virtual bool WantsNegXTriMesh() OVERRIDE { return false; }
 	// End Interface_CollisionDataProvider Interface
 
+	/** Called when spline params are changed, to notify render thread and possibly collision */
+	void MarkSplineParamsDirty();
+
+	UFUNCTION(BlueprintCallable, Category=SplineMesh)
+	void SetStartPosition(FVector StartPos);
+
+	UFUNCTION(BlueprintCallable, Category = SplineMesh)
+	void SetStartTangent(FVector StartTangent);
+
+	UFUNCTION(BlueprintCallable, Category = SplineMesh)
+	void SetEndPosition(FVector EndPos);
+
+	UFUNCTION(BlueprintCallable, Category = SplineMesh)
+	void SetEndTangent(FVector EndTangent);
+
 	// Destroys the body setup, used to clear collision if the mesh goes missing
 	void DestroyBodySetup();
 #if WITH_EDITOR
@@ -158,7 +173,11 @@ class USplineMeshComponent : public UStaticMeshComponent, public IInterface_Coll
 	 * Calculates the spline transform, including roll, scale, and offset along the spline at a specified distance
 	 * @Note:  This is mirrored to Lightmass::CalcSliceTransform() and LocalVertexShader.usf.  If you update one of these, please update them all! 
 	 */
-	ENGINE_API FTransform CalcSliceTransform(const float DistanceAlong) const;
+	FTransform CalcSliceTransform(const float DistanceAlong) const;
+
+	static const float& GetAxisValue(const FVector& InVector, ESplineMeshAxis::Type InAxis);
+	static float& GetAxisValue(FVector& InVector, ESplineMeshAxis::Type InAxis);
+
 };
 
 
diff --git a/Engine/Source/Runtime/Engine/Private/Landscape/LandscapeSplines.cpp b/Engine/Source/Runtime/Engine/Private/Landscape/LandscapeSplines.cpp
--- a/Engine/Source/Runtime/Engine/Private/Landscape/LandscapeSplines.cpp
+++ b/Engine/Source/Runtime/Engine/Private/Landscape/LandscapeSplines.cpp
@@ -1141,37 +1141,7 @@ static FVector2D RotateVec2D(const FVector2D InVec, float RotAngle)
 	return OutVec;
 }
 
-static const float& GetAxisValue(const FVector& InVector, ESplineMeshAxis::Type InAxis)
-{
-	switch (InAxis)
-	{
-	case ESplineMeshAxis::X:
-		return InVector.X;
-	case ESplineMeshAxis::Y:
-		return InVector.Y;
-	case ESplineMeshAxis::Z:
-		return InVector.Z;
-	default:
-		check(0);
-		return InVector.Z;
-	}
-}
 
-static float& GetAxisValue(FVector& InVector, ESplineMeshAxis::Type InAxis)
-{
-	switch (InAxis)
-	{
-	case ESplineMeshAxis::X:
-		return InVector.X;
-	case ESplineMeshAxis::Y:
-		return InVector.Y;
-	case ESplineMeshAxis::Z:
-		return InVector.Z;
-	default:
-		check(0);
-		return InVector.Z;
-	}
-}
 
 static ESplineMeshAxis::Type CrossAxis(ESplineMeshAxis::Type InForwardAxis, ESplineMeshAxis::Type InUpAxis)
 {
@@ -1363,10 +1333,10 @@ void ULandscapeSplineSegment::UpdateSplinePoints(bool bUpdateCollision)
 			FVector Scale = MeshEntry->Scale;
 			if (MeshEntry->bScaleToWidth)
 			{
-				Scale *= Width / GetAxisValue(MeshBounds.BoxExtent, CrossAxis(MeshEntry->ForwardAxis, MeshEntry->UpAxis));
+				Scale *= Width / USplineMeshComponent::GetAxisValue(MeshBounds.BoxExtent, CrossAxis(MeshEntry->ForwardAxis, MeshEntry->UpAxis));
 			}
 
-			const float MeshLength = FMath::Abs(GetAxisValue(MeshBounds.BoxExtent, MeshEntry->ForwardAxis) * 2 * GetAxisValue(Scale, MeshEntry->ForwardAxis));
+			const float MeshLength = FMath::Abs(USplineMeshComponent::GetAxisValue(MeshBounds.BoxExtent, MeshEntry->ForwardAxis) * 2 * USplineMeshComponent::GetAxisValue(Scale, MeshEntry->ForwardAxis));
 			float MeshT = (MeshLength / SplineLength);
 
 			// Improve our approximation if we're not going off the end of the spline
@@ -1468,7 +1438,7 @@ void ULandscapeSplineSegment::UpdateSplinePoints(bool bUpdateCollision)
 			FVector Scale = MeshEntry->Scale;
 			if (MeshEntry->bScaleToWidth)
 			{
-				Scale *= Width / GetAxisValue(MeshBounds.BoxExtent, SideAxis);
+				Scale *= Width / USplineMeshComponent::GetAxisValue(MeshBounds.BoxExtent, SideAxis);
 			}
 
 			FVector2D Offset = MeshEntry->Offset;
@@ -1476,11 +1446,11 @@ void ULandscapeSplineSegment::UpdateSplinePoints(bool bUpdateCollision)
 			{
 				if (bDoOrientationRoll)
 				{
-					Offset.Y -= GetAxisValue(MeshBounds.Origin, SideAxis);
+					Offset.Y -= USplineMeshComponent::GetAxisValue(MeshBounds.Origin, SideAxis);
 				}
 				else
 				{
-					Offset.X -= GetAxisValue(MeshBounds.Origin, SideAxis);
+					Offset.X -= USplineMeshComponent::GetAxisValue(MeshBounds.Origin, SideAxis);
 				}
 			}
 
@@ -1516,7 +1486,7 @@ void ULandscapeSplineSegment::UpdateSplinePoints(bool bUpdateCollision)
 			FVector ScaleEnd = MeshEntry->Scale;
 			if (MeshEntry->bScaleToWidth)
 			{
-				ScaleEnd *= WidthEnd / GetAxisValue(MeshBounds.BoxExtent, SideAxis);
+				ScaleEnd *= WidthEnd / USplineMeshComponent::GetAxisValue(MeshBounds.BoxExtent, SideAxis);
 			}
 
 			FVector2D OffsetEnd = MeshEntry->Offset;
@@ -1524,11 +1494,11 @@ void ULandscapeSplineSegment::UpdateSplinePoints(bool bUpdateCollision)
 			{
 				if (bDoOrientationRoll)
 				{
-					OffsetEnd.Y -= GetAxisValue(MeshBounds.Origin, SideAxis);
+					OffsetEnd.Y -= USplineMeshComponent::GetAxisValue(MeshBounds.Origin, SideAxis);
 				}
 				else
 				{
-					OffsetEnd.X -= GetAxisValue(MeshBounds.Origin, SideAxis);
+					OffsetEnd.X -= USplineMeshComponent::GetAxisValue(MeshBounds.Origin, SideAxis);
 				}
 			}
 
@@ -1560,7 +1530,7 @@ void ULandscapeSplineSegment::UpdateSplinePoints(bool bUpdateCollision)
 			MeshComponent->SplineXDir = FVector(0,0,1); // Up, to be consistent between joined meshes. We rotate it to horizontal using roll if using Z Forward X Up or X Forward Y Up
 			MeshComponent->ForwardAxis = MeshEntry->ForwardAxis;
 
-			if (GetAxisValue(MeshEntry->Scale, MeshEntry->ForwardAxis) < 0)
+			if (USplineMeshComponent::GetAxisValue(MeshEntry->Scale, MeshEntry->ForwardAxis) < 0)
 			{
 				Swap(MeshComponent->SplineParams.StartPos, MeshComponent->SplineParams.EndPos);
 				Swap(MeshComponent->SplineParams.StartTangent, MeshComponent->SplineParams.EndTangent);
@@ -1892,862 +1862,4 @@ void ULandscapeSplineSegment::PostEditChangeProperty(FPropertyChangedEvent& Prop
 #endif // WITH_EDITOR
 
 
-/** A vertex factory for spline-deformed static meshes */
-struct FSplineMeshVertexFactory : public FLocalVertexFactory
-{
-	DECLARE_VERTEX_FACTORY_TYPE(FSplineMeshVertexFactory);
-public:
-
-	FSplineMeshVertexFactory(class FSplineMeshSceneProxy* InSplineProxy) :
-		SplineSceneProxy(InSplineProxy)
-	{
-	}
-
-	/** Should we cache the material's shadertype on this platform with this vertex factory? */
-	static bool ShouldCache(EShaderPlatform Platform, const class FMaterial* Material, const class FShaderType* ShaderType)
-	{
-		return (Material->IsUsedWithSplineMeshes() || Material->IsSpecialEngineMaterial())
-			&& FLocalVertexFactory::ShouldCache(Platform, Material, ShaderType)
-			&& IsFeatureLevelSupported(Platform, ERHIFeatureLevel::SM3);
-	}
-
-	/** Modify compile environment to enable spline deformation */
-	static void ModifyCompilationEnvironment(EShaderPlatform Platform, const FMaterial* Material, FShaderCompilerEnvironment& OutEnvironment)
-	{
-		FLocalVertexFactory::ModifyCompilationEnvironment(Platform, Material, OutEnvironment);
-		OutEnvironment.SetDefine(TEXT("USE_SPLINEDEFORM"),TEXT("1"));
-	}
-
-	/** Copy the data from another vertex factory */
-	void Copy(const FSplineMeshVertexFactory& Other)
-	{
-		ENQUEUE_UNIQUE_RENDER_COMMAND_TWOPARAMETER(
-			FSplineMeshVertexFactoryCopyData,
-			FSplineMeshVertexFactory*,VertexFactory,this,
-			const DataType*,DataCopy,&Other.Data,
-		{
-			VertexFactory->Data = *DataCopy;
-		});
-		BeginUpdateResourceRHI(this);
-	}
-
-	// FRenderResource interface.
-	virtual void InitRHI()
-	{
-		FLocalVertexFactory::InitRHI();
-	}
-
-	static FVertexFactoryShaderParameters* ConstructShaderParameters(EShaderFrequency ShaderFrequency);
-
-	class FSplineMeshSceneProxy* SplineSceneProxy;
-
-private:
-};
-
-
-
-/** Factory specific params */
-class FSplineMeshVertexFactoryShaderParameters : public FVertexFactoryShaderParameters
-{
-	void Bind(const FShaderParameterMap& ParameterMap) OVERRIDE;
-
-	void SetMesh(FShader* Shader,const FVertexFactory* VertexFactory,const FSceneView& View,const FMeshBatchElement& BatchElement,uint32 DataFlags) const OVERRIDE;
-
-	void Serialize(FArchive& Ar) OVERRIDE
-	{
-		Ar << SplineStartPosParam;
-		Ar << SplineStartTangentParam;
-		Ar << SplineStartRollParam;
-		Ar << SplineStartScaleParam;
-		Ar << SplineStartOffsetParam;
-		
-		Ar << SplineEndPosParam;
-		Ar << SplineEndTangentParam;
-		Ar << SplineEndRollParam;
-		Ar << SplineEndScaleParam;
-		Ar << SplineEndOffsetParam;
-
-		Ar << SplineUpDirParam;
-		Ar << SmoothInterpRollScaleParam;
-
-		Ar << SplineMeshMinZParam;
-		Ar << SplineMeshScaleZParam;
-
-		Ar << SplineMeshDirParam;
-		Ar << SplineMeshXParam;
-		Ar << SplineMeshYParam;
-	}
-
-	virtual uint32 GetSize() const OVERRIDE
-	{
-		return sizeof(*this);
-	}
-
-private:
-	FShaderParameter SplineStartPosParam;
-	FShaderParameter SplineStartTangentParam;
-	FShaderParameter SplineStartRollParam;
-	FShaderParameter SplineStartScaleParam;
-	FShaderParameter SplineStartOffsetParam;
-	
-	FShaderParameter SplineEndPosParam;
-	FShaderParameter SplineEndTangentParam;
-	FShaderParameter SplineEndRollParam;
-	FShaderParameter SplineEndScaleParam;
-	FShaderParameter SplineEndOffsetParam;
-	
-	FShaderParameter SplineUpDirParam;
-	FShaderParameter SmoothInterpRollScaleParam;
-	
-	FShaderParameter SplineMeshMinZParam;
-	FShaderParameter SplineMeshScaleZParam;
-
-	FShaderParameter SplineMeshDirParam;
-	FShaderParameter SplineMeshXParam;
-	FShaderParameter SplineMeshYParam;
-};
-
-/** Scene proxy for SplineMesh instance */
-class FSplineMeshSceneProxy : public FStaticMeshSceneProxy
-{
-protected:
-	struct FLODResources
-	{
-		/** Pointer to vertex factory object */
-		FSplineMeshVertexFactory* VertexFactory;
-
-		FLODResources(FSplineMeshVertexFactory* InVertexFactory) :
-			VertexFactory(InVertexFactory)
-		{
-		}
-	};
-public:
-
-	FSplineMeshSceneProxy(USplineMeshComponent* InComponent):
-		FStaticMeshSceneProxy(InComponent)
-	{
-		// make sure all the materials are okay to be rendered as a spline mesh
-		for (FStaticMeshSceneProxy::FLODInfo& LODInfo : LODs)
-		{
-			for (FStaticMeshSceneProxy::FLODInfo::FSectionInfo& Section : LODInfo.Sections)
-			{
-				if (!Section.Material->CheckMaterialUsage(MATUSAGE_SplineMesh))
-				{
-					Section.Material = UMaterial::GetDefaultMaterial(MD_Surface);
-				}
-			}
-		}
-
-		// Copy spline params from component
-		SplineParams = InComponent->SplineParams;
-		SplineXDir = InComponent->SplineXDir;
-		bSmoothInterpRollScale = InComponent->bSmoothInterpRollScale;
-		ForwardAxis = InComponent->ForwardAxis;
-
-		// Fill in info about the mesh
-		FBoxSphereBounds StaticMeshBounds = StaticMesh->GetBounds();
-		SplineMeshScaleZ = 0.5f / GetAxisValue(StaticMeshBounds.BoxExtent, ForwardAxis); // 1/(2 * Extent)
-		SplineMeshMinZ = GetAxisValue(StaticMeshBounds.Origin, ForwardAxis) * SplineMeshScaleZ - 0.5f;
-
-		LODResources.Reset(InComponent->StaticMesh->RenderData->LODResources.Num());
-
-		for (int32 LODIndex = 0; LODIndex < LODs.Num(); LODIndex++)
-		{
-			FSplineMeshVertexFactory* VertexFactory = new FSplineMeshVertexFactory(this);
-
-			LODResources.Add(VertexFactory);
-
-			InitResources(InComponent, LODIndex);
-		}
-	}
-
-	virtual ~FSplineMeshSceneProxy() OVERRIDE
-	{
-		ReleaseResources();
-
-		for (FSplineMeshSceneProxy::FLODResources& LODResource : LODResources)
-		{
-			delete LODResource.VertexFactory;
-		}
-		LODResources.Empty();
-	}
-
-	void InitResources(USplineMeshComponent* InComponent, int32 InLODIndex);
-
-	void ReleaseResources();
-
-	/** Sets up a shadow FMeshBatch for a specific LOD. */
-	virtual bool GetShadowMeshElement(int32 LODIndex, uint8 InDepthPriorityGroup, FMeshBatch& OutMeshElement) const OVERRIDE
-	{
-		//checkf(LODIndex == 0, TEXT("Getting spline static mesh element with invalid LOD [%d]"), LODIndex);
-
-		if (FStaticMeshSceneProxy::GetShadowMeshElement(LODIndex, InDepthPriorityGroup, OutMeshElement))
-		{
-			OutMeshElement.VertexFactory = LODResources[LODIndex].VertexFactory;
-			OutMeshElement.ReverseCulling ^= (SplineParams.StartScale.X < 0) ^ (SplineParams.StartScale.Y < 0);
-			return true;
-		}
-		return false;
-	}
-
-	/** Sets up a FMeshBatch for a specific LOD and element. */
-	virtual bool GetMeshElement(int32 LODIndex,int32 SectionIndex,uint8 InDepthPriorityGroup,FMeshBatch& OutMeshElement, const bool bUseSelectedMaterial, const bool bUseHoveredMaterial) const OVERRIDE
-	{
-		//checkf(LODIndex == 0 /*&& SectionIndex == 0*/, TEXT("Getting spline static mesh element with invalid params [%d, %d]"), LODIndex, SectionIndex);
-
-		if (FStaticMeshSceneProxy::GetMeshElement(LODIndex, SectionIndex, InDepthPriorityGroup, OutMeshElement, bUseSelectedMaterial, bUseHoveredMaterial))
-		{
-			OutMeshElement.VertexFactory = LODResources[LODIndex].VertexFactory;
-			OutMeshElement.ReverseCulling ^= (SplineParams.StartScale.X < 0) ^ (SplineParams.StartScale.Y < 0);
-			return true;
-		}
-		return false;
-	}
-
-	/** Sets up a wireframe FMeshBatch for a specific LOD. */
-	virtual bool GetWireframeMeshElement(int32 LODIndex, const FMaterialRenderProxy* WireframeRenderProxy, uint8 InDepthPriorityGroup, FMeshBatch& OutMeshElement) const OVERRIDE
-	{
-		//checkf(LODIndex == 0, TEXT("Getting spline static mesh element with invalid LOD [%d]"), LODIndex);
-
-		if (FStaticMeshSceneProxy::GetWireframeMeshElement(LODIndex, WireframeRenderProxy, InDepthPriorityGroup, OutMeshElement))
-		{
-			OutMeshElement.VertexFactory = LODResources[LODIndex].VertexFactory;
-			OutMeshElement.ReverseCulling ^= (SplineParams.StartScale.X < 0) ^ (SplineParams.StartScale.Y < 0);
-			return true;
-		}
-		return false;
-	}
-
-	virtual FPrimitiveViewRelevance GetViewRelevance(const FSceneView* View) OVERRIDE
-	{
-		return FStaticMeshSceneProxy::GetViewRelevance(View);
-	}
-
-	// 	  virtual uint32 GetMemoryFootprint( void ) const { return 0; }
-
-	/** Parameters that define the spline, used to deform mesh */
-	FSplineMeshParams SplineParams;
-	/** Axis (in component space) that is used to determine X axis for co-ordinates along spline */
-	FVector SplineXDir;
-	/** Smoothly (cubic) interpolate the Roll and Scale params over spline. */
-	bool bSmoothInterpRollScale;
-	/** Chooses the forward axis for the spline mesh orientation */
-	ESplineMeshAxis::Type ForwardAxis;
-	
-	/** Minimum Z value of the entire mesh */
-	float SplineMeshMinZ;
-	/** Range of Z values over entire mesh */
-	float SplineMeshScaleZ;
-
-protected:
-	TArray<FLODResources> LODResources;
-};
-
-FVertexFactoryShaderParameters* FSplineMeshVertexFactory::ConstructShaderParameters(EShaderFrequency ShaderFrequency)
-{
-	return ShaderFrequency == SF_Vertex ? new FSplineMeshVertexFactoryShaderParameters() : NULL;
-}
-
-IMPLEMENT_VERTEX_FACTORY_TYPE(FSplineMeshVertexFactory, "LocalVertexFactory", true, true, true, true, true);
-
-
-void FSplineMeshVertexFactoryShaderParameters::Bind(const FShaderParameterMap& ParameterMap)
-{
-	SplineStartPosParam.Bind(ParameterMap,TEXT("SplineStartPos"), SPF_Mandatory);
-	SplineStartTangentParam.Bind(ParameterMap,TEXT("SplineStartTangent"), SPF_Mandatory);
-	SplineStartRollParam.Bind(ParameterMap,TEXT("SplineStartRoll"), SPF_Mandatory);
-	SplineStartScaleParam.Bind(ParameterMap,TEXT("SplineStartScale"), SPF_Mandatory);
-	SplineStartOffsetParam.Bind(ParameterMap,TEXT("SplineStartOffset"), SPF_Mandatory);
-
-	SplineEndPosParam.Bind(ParameterMap,TEXT("SplineEndPos"), SPF_Mandatory);
-	SplineEndTangentParam.Bind(ParameterMap,TEXT("SplineEndTangent"), SPF_Mandatory);
-	SplineEndRollParam.Bind(ParameterMap,TEXT("SplineEndRoll"), SPF_Mandatory);
-	SplineEndScaleParam.Bind(ParameterMap,TEXT("SplineEndScale"), SPF_Mandatory);
-	SplineEndOffsetParam.Bind(ParameterMap,TEXT("SplineEndOffset"), SPF_Mandatory);
-	
-	SplineUpDirParam.Bind(ParameterMap,TEXT("SplineUpDir"), SPF_Mandatory);
-	SmoothInterpRollScaleParam.Bind(ParameterMap,TEXT("SmoothInterpRollScale"), SPF_Mandatory);
-
-	SplineMeshMinZParam.Bind(ParameterMap,TEXT("SplineMeshMinZ"), SPF_Mandatory);
-	SplineMeshScaleZParam.Bind(ParameterMap,TEXT("SplineMeshScaleZ"), SPF_Mandatory);
-
-	SplineMeshDirParam.Bind(ParameterMap,TEXT("SplineMeshDir"), SPF_Mandatory);
-	SplineMeshXParam.Bind(ParameterMap,TEXT("SplineMeshX"), SPF_Mandatory);
-	SplineMeshYParam.Bind(ParameterMap,TEXT("SplineMeshY"), SPF_Mandatory);
-}
-
-void FSplineMeshVertexFactoryShaderParameters::SetMesh(FShader* Shader,const FVertexFactory* VertexFactory,const FSceneView& View,const FMeshBatchElement& BatchElement,uint32 DataFlags) const
-{	
-	if(Shader->GetVertexShader())
-	{
-		FSplineMeshVertexFactory* SplineVertexFactory = (FSplineMeshVertexFactory*)VertexFactory;
-		FSplineMeshSceneProxy* SplineProxy = SplineVertexFactory->SplineSceneProxy;
-		FSplineMeshParams& SplineParams = SplineProxy->SplineParams;
-
-		SetShaderValue(Shader->GetVertexShader(), SplineStartPosParam, SplineParams.StartPos);
-		SetShaderValue(Shader->GetVertexShader(), SplineStartTangentParam, SplineParams.StartTangent);
-		SetShaderValue(Shader->GetVertexShader(), SplineStartRollParam, SplineParams.StartRoll);
-		SetShaderValue(Shader->GetVertexShader(), SplineStartScaleParam, SplineParams.StartScale);
-		SetShaderValue(Shader->GetVertexShader(), SplineStartOffsetParam, SplineParams.StartOffset);
-
-		SetShaderValue(Shader->GetVertexShader(), SplineEndPosParam, SplineParams.EndPos);
-		SetShaderValue(Shader->GetVertexShader(), SplineEndTangentParam, SplineParams.EndTangent);
-		SetShaderValue(Shader->GetVertexShader(), SplineEndRollParam, SplineParams.EndRoll);
-		SetShaderValue(Shader->GetVertexShader(), SplineEndScaleParam, SplineParams.EndScale);
-		SetShaderValue(Shader->GetVertexShader(), SplineEndOffsetParam, SplineParams.EndOffset);
-	
-		SetShaderValue(Shader->GetVertexShader(), SplineUpDirParam, SplineProxy->SplineXDir);
-		SetShaderValue(Shader->GetVertexShader(), SmoothInterpRollScaleParam, SplineProxy->bSmoothInterpRollScale);
-
-		SetShaderValue(Shader->GetVertexShader(), SplineMeshMinZParam, SplineProxy->SplineMeshMinZ);
-		SetShaderValue(Shader->GetVertexShader(), SplineMeshScaleZParam, SplineProxy->SplineMeshScaleZ);
-
-		FVector DirMask(0,0,0);
-		DirMask[SplineProxy->ForwardAxis] = 1;
-		SetShaderValue(Shader->GetVertexShader(), SplineMeshDirParam, DirMask);
-		DirMask = FVector::ZeroVector;
-		DirMask[(SplineProxy->ForwardAxis+1)%3] = 1;
-		SetShaderValue(Shader->GetVertexShader(), SplineMeshXParam, DirMask);
-		DirMask = FVector::ZeroVector;
-		DirMask[(SplineProxy->ForwardAxis+2)%3] = 1;
-		SetShaderValue(Shader->GetVertexShader(), SplineMeshYParam, DirMask);
-	}
-}
-
-void FSplineMeshSceneProxy::InitResources(USplineMeshComponent* InComponent, int32 InLODIndex)
-{
-	// Initialize the static mesh's vertex factory.
-	ENQUEUE_UNIQUE_RENDER_COMMAND_THREEPARAMETER(
-		InitSplineMeshVertexFactory,
-		FSplineMeshVertexFactory*,VertexFactory,LODResources[InLODIndex].VertexFactory,
-		const FStaticMeshLODResources*,RenderData,&InComponent->StaticMesh->RenderData->LODResources[InLODIndex],
-		UStaticMesh*,Parent,InComponent->StaticMesh,
-	{
-		FLocalVertexFactory::DataType Data;
-
-		Data.PositionComponent = FVertexStreamComponent(
-			&RenderData->PositionVertexBuffer,
-			STRUCT_OFFSET(FPositionVertex,Position),
-			RenderData->PositionVertexBuffer.GetStride(),
-			VET_Float3
-			);
-		Data.TangentBasisComponents[0] = FVertexStreamComponent(
-			&RenderData->VertexBuffer,
-			STRUCT_OFFSET(FStaticMeshFullVertex,TangentX),
-			RenderData->VertexBuffer.GetStride(),
-			VET_PackedNormal
-			);
-		Data.TangentBasisComponents[1] = FVertexStreamComponent(
-			&RenderData->VertexBuffer,
-			STRUCT_OFFSET(FStaticMeshFullVertex,TangentZ),
-			RenderData->VertexBuffer.GetStride(),
-			VET_PackedNormal
-			);
-
-		if( RenderData->ColorVertexBuffer.GetNumVertices() > 0 )
-		{
-			Data.ColorComponent = FVertexStreamComponent(
-				&RenderData->ColorVertexBuffer,
-				0,	// Struct offset to color
-				RenderData->ColorVertexBuffer.GetStride(),
-				VET_Color
-				);
-		}
-
-		Data.TextureCoordinates.Empty();
-
-		if( !RenderData->VertexBuffer.GetUseFullPrecisionUVs() )
-		{
-			for(uint32 UVIndex = 0;UVIndex < RenderData->VertexBuffer.GetNumTexCoords();UVIndex++)
-			{
-				Data.TextureCoordinates.Add(FVertexStreamComponent(
-					&RenderData->VertexBuffer,
-					STRUCT_OFFSET(TStaticMeshFullVertexFloat16UVs<MAX_TEXCOORDS>,UVs) + sizeof(FVector2DHalf) * UVIndex,
-					RenderData->VertexBuffer.GetStride(),
-					VET_Half2
-					));
-			}
-			if(	Parent->LightMapCoordinateIndex >= 0 && (uint32)Parent->LightMapCoordinateIndex < RenderData->VertexBuffer.GetNumTexCoords())
-			{
-				Data.LightMapCoordinateComponent = FVertexStreamComponent(
-					&RenderData->VertexBuffer,
-					STRUCT_OFFSET(TStaticMeshFullVertexFloat16UVs<MAX_TEXCOORDS>,UVs) + sizeof(FVector2DHalf) * Parent->LightMapCoordinateIndex,
-					RenderData->VertexBuffer.GetStride(),
-					VET_Half2
-					);
-			}
-		}
-		else
-		{
-			for(uint32 UVIndex = 0;UVIndex < RenderData->VertexBuffer.GetNumTexCoords();UVIndex++)
-			{
-				Data.TextureCoordinates.Add(FVertexStreamComponent(
-					&RenderData->VertexBuffer,
-					STRUCT_OFFSET(TStaticMeshFullVertexFloat32UVs<MAX_TEXCOORDS>,UVs) + sizeof(FVector2D) * UVIndex,
-					RenderData->VertexBuffer.GetStride(),
-					VET_Float2
-					));
-			}
-
-			if (Parent->LightMapCoordinateIndex >= 0 && (uint32)Parent->LightMapCoordinateIndex < RenderData->VertexBuffer.GetNumTexCoords())
-			{
-				Data.LightMapCoordinateComponent = FVertexStreamComponent(
-					&RenderData->VertexBuffer,
-					STRUCT_OFFSET(TStaticMeshFullVertexFloat32UVs<MAX_TEXCOORDS>,UVs) + sizeof(FVector2D) * Parent->LightMapCoordinateIndex,
-					RenderData->VertexBuffer.GetStride(),
-					VET_Float2
-					);
-			}
-		}	
-
-		VertexFactory->SetData(Data);
-
-		VertexFactory->InitResource();
-	});
-}
-
-
-void FSplineMeshSceneProxy::ReleaseResources()
-{
-	for (FSplineMeshSceneProxy::FLODResources& LODResource : LODResources)
-	{
-		LODResource.VertexFactory->ReleaseResource();
-	}
-}
-
-USplineMeshComponent::USplineMeshComponent(const class FPostConstructInitializeProperties& PCIP)
-	: Super(PCIP)
-{
-	Mobility = EComponentMobility::Static;
-
-	BodyInstance.bEnableCollision_DEPRECATED = false;
-	SetCollisionProfileName(UCollisionProfile::NoCollision_ProfileName);
-	bHasCustomNavigableGeometry = EHasCustomNavigableGeometry::Yes;
-
-	SplineXDir.X = 1.0f;
-}
 
-void USplineMeshComponent::Serialize(FArchive& Ar)
-{
-	Super::Serialize(Ar);
-
-	if (Ar.UE4Ver() < VER_UE4_SPLINE_MESH_ORIENTATION)
-	{
-		ForwardAxis = ESplineMeshAxis::Z;
-		SplineParams.StartRoll -= HALF_PI;
-		SplineParams.EndRoll -= HALF_PI;
-
-		float Temp = SplineParams.StartOffset.X;
-		SplineParams.StartOffset.X = -SplineParams.StartOffset.Y;
-		SplineParams.StartOffset.Y = Temp;
-		Temp = SplineParams.EndOffset.X;
-		SplineParams.EndOffset.X = -SplineParams.EndOffset.Y;
-		SplineParams.EndOffset.Y = Temp;
-	}
-
-#if WITH_EDITOR
-	if (BodySetup != NULL)
-	{
-		BodySetup->SetFlags(RF_Transactional);
-	}
-#endif
-}
-
-bool USplineMeshComponent::Modify(bool bAlwaysMarkDirty)
-{
-	bool bSavedToTransactionBuffer = Super::Modify(bAlwaysMarkDirty);
-
-	if (BodySetup != NULL)
-	{
-		BodySetup->Modify(bAlwaysMarkDirty);
-	}
-
-	return bSavedToTransactionBuffer;
-}
-
-FPrimitiveSceneProxy* USplineMeshComponent::CreateSceneProxy()
-{
-	// Verify that the mesh is valid before using it.
-	const bool bMeshIsValid = 
-		// make sure we have an actual staticmesh
-		StaticMesh &&
-		StaticMesh->HasValidRenderData();
-
-	if (bMeshIsValid && (GetScene()->GetFeatureLevel() >= ERHIFeatureLevel::SM3))
-	{
-		return ::new FSplineMeshSceneProxy(this);
-	}
-	else
-	{
-		return NULL;
-	}
-}
-
-FBoxSphereBounds USplineMeshComponent::CalcBounds(const FTransform & LocalToWorld) const
-{
-	// Use util to generate bounds of spline
-	FInterpCurvePoint<FVector> Start(0.f, SplineParams.StartPos, SplineParams.StartTangent, SplineParams.StartTangent, CIM_CurveUser);
-	FInterpCurvePoint<FVector> End(1.f, SplineParams.EndPos, SplineParams.EndTangent, SplineParams.EndTangent, CIM_CurveUser);
-
-	FVector CurveMax(-BIG_NUMBER, -BIG_NUMBER, -BIG_NUMBER);
-	FVector CurveMin(BIG_NUMBER, BIG_NUMBER, BIG_NUMBER);
-	CurveVectorFindIntervalBounds(Start, End, CurveMin, CurveMax);
-
-	FBox LocalBox(CurveMin, CurveMax);
-
-	// Find largest extent of mesh in XY, and add on all around
-	if(StaticMesh)
-	{
-		FVector MinMeshExtent = StaticMesh->GetBounds().Origin - StaticMesh->GetBounds().BoxExtent;
-		FVector MaxMeshExtent = StaticMesh->GetBounds().Origin + StaticMesh->GetBounds().BoxExtent;
-		GetAxisValue(MinMeshExtent, ForwardAxis) = 0;
-		GetAxisValue(MaxMeshExtent, ForwardAxis) = 0;
-		float MaxDim = FMath::Max<float>(MinMeshExtent.GetAbsMax(), MaxMeshExtent.GetAbsMax());
-
-		float MaxScale = FMath::Max(SplineParams.StartScale.GetAbsMax(), SplineParams.EndScale.GetAbsMax());
-
-		LocalBox = LocalBox.ExpandBy(MaxScale * MaxDim);
-	}
-
-	return FBoxSphereBounds( LocalBox.TransformBy(LocalToWorld) );
-}
-
-
-/** 
- * Functions used for transforming a static mesh component based on a spline.  
- * This needs to be updated if the spline functionality changes!
- */
-static float SmoothStep(float A, float B, float X)
-{
-	if (X < A)
-	{
-		return 0.0f;
-	}
-	else if (X >= B)
-	{
-		return 1.0f;
-	}
-	const float InterpFraction = (X - A) / (B - A);
-	return InterpFraction * InterpFraction * (3.0f - 2.0f * InterpFraction);
-}
-
-static FVector SplineEvalPos(const FVector& StartPos, const FVector& StartTangent, const FVector& EndPos, const FVector& EndTangent, float A)
-{
-	const float A2 = A  * A;
-	const float A3 = A2 * A;
-
-	return (((2*A3)-(3*A2)+1) * StartPos) + ((A3-(2*A2)+A) * StartTangent) + ((A3-A2) * EndTangent) + (((-2*A3)+(3*A2)) * EndPos);
-}
-
-static FVector SplineEvalDir(const FVector& StartPos, const FVector& StartTangent, const FVector& EndPos, const FVector& EndTangent, float A)
-{
-	const FVector C = (6*StartPos) + (3*StartTangent) + (3*EndTangent) - (6*EndPos);
-	const FVector D = (-6*StartPos) - (4*StartTangent) - (2*EndTangent) + (6*EndPos);
-	const FVector E = StartTangent;
-
-	const float A2 = A  * A;
-
-	return ((C * A2) + (D * A) + E).SafeNormal();
-}
-
-
-FTransform USplineMeshComponent::CalcSliceTransform(const float DistanceAlong) const
-{
-	// Find how far 'along' mesh we are
-	FBoxSphereBounds StaticMeshBounds = StaticMesh->GetBounds();
-	const float MeshMinZ = GetAxisValue(StaticMeshBounds.Origin, ForwardAxis) - GetAxisValue(StaticMeshBounds.BoxExtent, ForwardAxis);
-	const float MeshRangeZ = 2.0f * GetAxisValue(StaticMeshBounds.BoxExtent, ForwardAxis);
-	const float Alpha = (DistanceAlong - MeshMinZ) / MeshRangeZ;
-
-	// Apply hermite interp to Alpha if desired
-	const float HermiteAlpha = bSmoothInterpRollScale ? SmoothStep(0.0, 1.0, Alpha) : Alpha;
-
-	// Then find the point and direction of the spline at this point along
-	FVector SplinePos = SplineEvalPos( SplineParams.StartPos, SplineParams.StartTangent, SplineParams.EndPos, SplineParams.EndTangent, Alpha );
-	const FVector SplineDir = SplineEvalDir( SplineParams.StartPos, SplineParams.StartTangent, SplineParams.EndPos, SplineParams.EndTangent, Alpha );
-
-	// Find base frenet frame
-	const FVector BaseXVec = (SplineXDir ^ SplineDir).SafeNormal();
-	const FVector BaseYVec = (SplineDir ^ BaseXVec).SafeNormal();
-
-	// Offset the spline by the desired amount
-	const FVector2D SliceOffset = FMath::Lerp<FVector2D>(SplineParams.StartOffset, SplineParams.EndOffset, HermiteAlpha);
-	SplinePos += SliceOffset.X * BaseXVec;
-	SplinePos += SliceOffset.Y * BaseYVec;
-
-	// Apply roll to frame around spline
-	const float UseRoll = FMath::Lerp(SplineParams.StartRoll, SplineParams.EndRoll, HermiteAlpha);
-	const float CosAng = FMath::Cos(UseRoll);
-	const float SinAng = FMath::Sin(UseRoll);
-	const FVector XVec = (CosAng * BaseXVec) - (SinAng * BaseYVec);
-	const FVector YVec = (CosAng * BaseYVec) + (SinAng * BaseXVec);
-
-	// Find scale at this point along spline
-	const FVector2D UseScale = FMath::Lerp(SplineParams.StartScale, SplineParams.EndScale, HermiteAlpha);
-
-	// Build overall transform
-	FTransform SliceTransform;
-	switch (ForwardAxis)
-	{
-	case ESplineMeshAxis::X:
-		SliceTransform = FTransform(SplineDir, XVec, YVec, SplinePos);
-		SliceTransform.SetScale3D(FVector(1, UseScale.X, UseScale.Y));
-		break;
-	case ESplineMeshAxis::Y:
-		SliceTransform = FTransform(YVec, SplineDir, XVec, SplinePos);
-		SliceTransform.SetScale3D(FVector(UseScale.Y, 1, UseScale.X));
-		break;
-	case ESplineMeshAxis::Z:
-		SliceTransform = FTransform(XVec, YVec, SplineDir, SplinePos);
-		SliceTransform.SetScale3D(FVector(UseScale.X, UseScale.Y, 1));
-		break;
-	default:
-		check(0);
-		break;
-	}
-
-	return SliceTransform;
-}
-
-
-bool USplineMeshComponent::GetPhysicsTriMeshData(struct FTriMeshCollisionData* CollisionData, bool InUseAllTriData)
-{
-	if (StaticMesh)
-	{
-		StaticMesh->GetPhysicsTriMeshData(CollisionData, InUseAllTriData);
-
-		FVector Mask = FVector(1,1,1);
-		GetAxisValue(Mask, ForwardAxis) = 0;
-
-		for (FVector& CollisionVert : CollisionData->Vertices)
-		{
-			CollisionVert = CalcSliceTransform(GetAxisValue(CollisionVert, ForwardAxis)).TransformPosition(CollisionVert * Mask);
-		}
-		return true;
-	}
-
-	return false;
-}
-
-bool USplineMeshComponent::ContainsPhysicsTriMeshData(bool InUseAllTriData) const
-{
-	if (StaticMesh)
-	{
-		return StaticMesh->ContainsPhysicsTriMeshData(InUseAllTriData);
-	}
-
-	return false;
-}
-
-void USplineMeshComponent::CreatePhysicsState()
-{
-#if WITH_EDITOR
-	// With editor code we can recreate the collision if the mesh changes
-	const FGuid MeshBodySetupGuid = (StaticMesh != NULL ? StaticMesh->BodySetup->BodySetupGuid : FGuid());
-	if (CachedMeshBodySetupGuid != MeshBodySetupGuid)
-	{
-		RecreateCollision();
-	}
-#else
-	// Without editor code we can only destroy the collision if the mesh is missing
-	if (StaticMesh == NULL && BodySetup != NULL)
-	{
-		DestroyBodySetup();
-	}
-#endif
-
-	return Super::CreatePhysicsState();
-}
-
-UBodySetup* USplineMeshComponent::GetBodySetup()
-{
-	// Don't return a body setup that has no collision, it means we are interactively moving the spline and don't want to build collision.
-	// Instead we explicitly build collision with USplineMeshComponent::RecreateCollision()
-	if (BodySetup != NULL && (BodySetup->TriMesh != NULL || BodySetup->AggGeom.GetElementCount() > 0))
-	{
-		return BodySetup;
-	}
-	return NULL;
-}
-
-bool USplineMeshComponent::DoCustomNavigableGeometryExport(struct FNavigableGeometryExport* GeomExport) const
-{
-	// the NavCollision is supposed to be faster than exporting the regular collision,
-	// but I'm not sure that's true here, as the regular collision is pre-distorted to the spline
-
-	if (StaticMesh != NULL && StaticMesh->NavCollision != NULL)
-	{
-		UNavCollision* NavCollision = StaticMesh->NavCollision;
-		
-		if (ensure(!NavCollision->bIsDynamicObstacle))
-		{
-			if (NavCollision->bHasConvexGeometry)
-			{
-				TArray<FVector> VertexBuffer;
-				VertexBuffer.Reserve(FMath::Max(NavCollision->ConvexCollision.VertexBuffer.Num(), NavCollision->TriMeshCollision.VertexBuffer.Num()));
-
-				for (int32 i = 0; i < NavCollision->ConvexCollision.VertexBuffer.Num(); ++i)
-				{
-					FVector Vertex = NavCollision->ConvexCollision.VertexBuffer[i];
-					Vertex = CalcSliceTransform(GetAxisValue(Vertex, ForwardAxis)).TransformPosition(Vertex);
-					VertexBuffer.Add(Vertex);
-				}
-				GeomExport->ExportCustomMesh(VertexBuffer.GetData(), VertexBuffer.Num(),
-					NavCollision->ConvexCollision.IndexBuffer.GetData(), NavCollision->ConvexCollision.IndexBuffer.Num(),
-					ComponentToWorld);
-
-				VertexBuffer.Reset();
-				for (int32 i = 0; i < NavCollision->TriMeshCollision.VertexBuffer.Num(); ++i)
-				{
-					FVector Vertex = NavCollision->TriMeshCollision.VertexBuffer[i];
-					Vertex = CalcSliceTransform(GetAxisValue(Vertex, ForwardAxis)).TransformPosition(Vertex);
-					VertexBuffer.Add(Vertex);
-				}
-				GeomExport->ExportCustomMesh(VertexBuffer.GetData(), VertexBuffer.Num(),
-					NavCollision->TriMeshCollision.IndexBuffer.GetData(), NavCollision->TriMeshCollision.IndexBuffer.Num(),
-					ComponentToWorld);
-
-				return false;
-			}
-		}
-	}
-
-	return true;
-}
-
-void USplineMeshComponent::DestroyBodySetup()
-{
-	if (BodySetup != NULL)
-	{
-		BodySetup->MarkPendingKill();
-		BodySetup = NULL;
-#if WITH_EDITORONLY_DATA
-		CachedMeshBodySetupGuid.Invalidate();
-#endif
-	}
-}
-
-#if WITH_EDITOR
-void USplineMeshComponent::RecreateCollision()
-{
-	if (StaticMesh && IsCollisionEnabled())
-	{
-		if (BodySetup == NULL)
-		{
-			BodySetup = DuplicateObject<UBodySetup>(StaticMesh->BodySetup, this);
-			BodySetup->SetFlags(RF_Transactional);
-			BodySetup->InvalidatePhysicsData();
-		}
-		else
-		{
-			BodySetup->Modify();
-			BodySetup->InvalidatePhysicsData();
-			BodySetup->CopyBodyPropertiesFrom(StaticMesh->BodySetup);
-			BodySetup->CollisionTraceFlag = StaticMesh->BodySetup->CollisionTraceFlag;
-		}
-
-		if (BodySetup->CollisionTraceFlag == CTF_UseComplexAsSimple)
-		{
-			BodySetup->AggGeom.EmptyElements();
-		}
-		else
-		{
-			FVector Mask = FVector(1,1,1);
-			GetAxisValue(Mask, ForwardAxis) = 0;
-
-			// distortion of a sphere can't be done nicely, so we just transform the origin and size
-			for (FKSphereElem& SphereElem : BodySetup->AggGeom.SphereElems)
-			{
-				const float Z = GetAxisValue(SphereElem.Center, ForwardAxis);
-				FTransform SliceTransform = CalcSliceTransform(Z);
-				SphereElem.Center *= Mask;
-
-				SphereElem.Radius *= SliceTransform.GetMaximumAxisScale();
-
-				SliceTransform.RemoveScaling();
-				SphereElem.Center = SliceTransform.TransformPosition(SphereElem.Center);
-			}
-
-			// distortion of a sphyl can't be done nicely, so we just transform the origin and size
-			for (FKSphylElem& SphylElem : BodySetup->AggGeom.SphylElems)
-			{
-				const float Z = GetAxisValue(SphylElem.Center, ForwardAxis);
-				FTransform SliceTransform = CalcSliceTransform(Z);
-				SphylElem.Center *= Mask;
-
-				FTransform TM = SphylElem.GetTransform();
-				SphylElem.Length = (TM * SliceTransform).TransformVector(FVector(0, 0, SphylElem.Length)).Size();
-				SphylElem.Radius *= SliceTransform.GetMaximumAxisScale();
-
-				SliceTransform.RemoveScaling();
-				SphylElem.SetTransform(TM * SliceTransform);
-			}
-
-			// Convert boxes to convex hulls to better respect distortion
-			for (FKBoxElem& BoxElem : BodySetup->AggGeom.BoxElems)
-			{
-				FKConvexElem& ConvexElem = *new(BodySetup->AggGeom.ConvexElems) FKConvexElem();
-
-				const FVector Radii = FVector(BoxElem.X / 2, BoxElem.Y / 2, BoxElem.Z / 2);
-				const FTransform ElementTM = BoxElem.GetTransform();
-				ConvexElem.VertexData.Empty(8);
-				ConvexElem.VertexData.Add(ElementTM.TransformPosition(Radii * FVector(-1,-1,-1)));
-				ConvexElem.VertexData.Add(ElementTM.TransformPosition(Radii * FVector(-1,-1, 1)));
-				ConvexElem.VertexData.Add(ElementTM.TransformPosition(Radii * FVector(-1, 1,-1)));
-				ConvexElem.VertexData.Add(ElementTM.TransformPosition(Radii * FVector(-1, 1, 1)));
-				ConvexElem.VertexData.Add(ElementTM.TransformPosition(Radii * FVector( 1,-1,-1)));
-				ConvexElem.VertexData.Add(ElementTM.TransformPosition(Radii * FVector( 1,-1, 1)));
-				ConvexElem.VertexData.Add(ElementTM.TransformPosition(Radii * FVector( 1, 1,-1)));
-				ConvexElem.VertexData.Add(ElementTM.TransformPosition(Radii * FVector( 1, 1, 1)));
-
-				ConvexElem.UpdateElemBox();
-			}
-			BodySetup->AggGeom.BoxElems.Empty();
-
-			// transform the points of the convex hulls into spline space
-			for (FKConvexElem& ConvexElem : BodySetup->AggGeom.ConvexElems)
-			{
-				for (FVector& Point : ConvexElem.VertexData)
-				{
-					Point = CalcSliceTransform(GetAxisValue(Point, ForwardAxis)).TransformPosition(Point * Mask);
-				}
-			}
-
-			BodySetup->CreatePhysicsMeshes();
-			CachedMeshBodySetupGuid = StaticMesh->BodySetup->BodySetupGuid;
-		}
-	}
-	else
-	{
-		DestroyBodySetup();
-	}
-}
-#endif
-
-#include "StaticMeshLight.h"
-/** */
-class FSplineStaticLightingMesh : public FStaticMeshStaticLightingMesh
-{
-public:
-
-	FSplineStaticLightingMesh(const USplineMeshComponent* InPrimitive,int32 InLODIndex,const TArray<ULightComponent*>& InRelevantLights):
-		FStaticMeshStaticLightingMesh(InPrimitive, InLODIndex, InRelevantLights),
-		SplineComponent(InPrimitive)
-	{
-	}
-
-#if WITH_EDITOR
-	virtual const struct FSplineMeshParams* GetSplineParameters() const
-	{
-		return &SplineComponent->SplineParams;
-	}
-#endif	//WITH_EDITOR
-
-private:
-	const USplineMeshComponent* SplineComponent;
-};
-
-FStaticMeshStaticLightingMesh* USplineMeshComponent::AllocateStaticLightingMesh(int32 LODIndex, const TArray<ULightComponent*>& InRelevantLights)
-{
-	return new FSplineStaticLightingMesh(this, LODIndex, InRelevantLights);
-}
diff --git a/Engine/Source/Runtime/Engine/Private/SplineMeshComponent.cpp b/Engine/Source/Runtime/Engine/Private/SplineMeshComponent.cpp
new file mode 100644
--- /dev/null
+++ b/Engine/Source/Runtime/Engine/Private/SplineMeshComponent.cpp
@@ -0,0 +1,706 @@
+// Copyright 1998-2014 Epic Games, Inc. All Rights Reserved.
+
+#include "EnginePrivate.h"
+#include "Components/SplineMeshComponent.h"
+#include "SplineMeshSceneProxy.h"
+
+
+//////////////////////////////////////////////////////////////////////////
+// FSplineMeshVertexFactoryShaderParameters
+
+void FSplineMeshVertexFactoryShaderParameters::Bind(const FShaderParameterMap& ParameterMap)
+{
+	SplineStartPosParam.Bind(ParameterMap, TEXT("SplineStartPos"), SPF_Mandatory);
+	SplineStartTangentParam.Bind(ParameterMap, TEXT("SplineStartTangent"), SPF_Mandatory);
+	SplineStartRollParam.Bind(ParameterMap, TEXT("SplineStartRoll"), SPF_Mandatory);
+	SplineStartScaleParam.Bind(ParameterMap, TEXT("SplineStartScale"), SPF_Mandatory);
+	SplineStartOffsetParam.Bind(ParameterMap, TEXT("SplineStartOffset"), SPF_Mandatory);
+
+	SplineEndPosParam.Bind(ParameterMap, TEXT("SplineEndPos"), SPF_Mandatory);
+	SplineEndTangentParam.Bind(ParameterMap, TEXT("SplineEndTangent"), SPF_Mandatory);
+	SplineEndRollParam.Bind(ParameterMap, TEXT("SplineEndRoll"), SPF_Mandatory);
+	SplineEndScaleParam.Bind(ParameterMap, TEXT("SplineEndScale"), SPF_Mandatory);
+	SplineEndOffsetParam.Bind(ParameterMap, TEXT("SplineEndOffset"), SPF_Mandatory);
+
+	SplineUpDirParam.Bind(ParameterMap, TEXT("SplineUpDir"), SPF_Mandatory);
+	SmoothInterpRollScaleParam.Bind(ParameterMap, TEXT("SmoothInterpRollScale"), SPF_Mandatory);
+
+	SplineMeshMinZParam.Bind(ParameterMap, TEXT("SplineMeshMinZ"), SPF_Mandatory);
+	SplineMeshScaleZParam.Bind(ParameterMap, TEXT("SplineMeshScaleZ"), SPF_Mandatory);
+
+	SplineMeshDirParam.Bind(ParameterMap, TEXT("SplineMeshDir"), SPF_Mandatory);
+	SplineMeshXParam.Bind(ParameterMap, TEXT("SplineMeshX"), SPF_Mandatory);
+	SplineMeshYParam.Bind(ParameterMap, TEXT("SplineMeshY"), SPF_Mandatory);
+}
+
+void FSplineMeshVertexFactoryShaderParameters::SetMesh(FShader* Shader, const FVertexFactory* VertexFactory, const FSceneView& View, const FMeshBatchElement& BatchElement, uint32 DataFlags) const
+{
+	if (Shader->GetVertexShader())
+	{
+		FSplineMeshVertexFactory* SplineVertexFactory = (FSplineMeshVertexFactory*)VertexFactory;
+		FSplineMeshSceneProxy* SplineProxy = SplineVertexFactory->SplineSceneProxy;
+		FSplineMeshParams& SplineParams = SplineProxy->SplineParams;
+
+		SetShaderValue(Shader->GetVertexShader(), SplineStartPosParam, SplineParams.StartPos);
+		SetShaderValue(Shader->GetVertexShader(), SplineStartTangentParam, SplineParams.StartTangent);
+		SetShaderValue(Shader->GetVertexShader(), SplineStartRollParam, SplineParams.StartRoll);
+		SetShaderValue(Shader->GetVertexShader(), SplineStartScaleParam, SplineParams.StartScale);
+		SetShaderValue(Shader->GetVertexShader(), SplineStartOffsetParam, SplineParams.StartOffset);
+
+		SetShaderValue(Shader->GetVertexShader(), SplineEndPosParam, SplineParams.EndPos);
+		SetShaderValue(Shader->GetVertexShader(), SplineEndTangentParam, SplineParams.EndTangent);
+		SetShaderValue(Shader->GetVertexShader(), SplineEndRollParam, SplineParams.EndRoll);
+		SetShaderValue(Shader->GetVertexShader(), SplineEndScaleParam, SplineParams.EndScale);
+		SetShaderValue(Shader->GetVertexShader(), SplineEndOffsetParam, SplineParams.EndOffset);
+
+		SetShaderValue(Shader->GetVertexShader(), SplineUpDirParam, SplineProxy->SplineXDir);
+		SetShaderValue(Shader->GetVertexShader(), SmoothInterpRollScaleParam, SplineProxy->bSmoothInterpRollScale);
+
+		SetShaderValue(Shader->GetVertexShader(), SplineMeshMinZParam, SplineProxy->SplineMeshMinZ);
+		SetShaderValue(Shader->GetVertexShader(), SplineMeshScaleZParam, SplineProxy->SplineMeshScaleZ);
+
+		FVector DirMask(0, 0, 0);
+		DirMask[SplineProxy->ForwardAxis] = 1;
+		SetShaderValue(Shader->GetVertexShader(), SplineMeshDirParam, DirMask);
+		DirMask = FVector::ZeroVector;
+		DirMask[(SplineProxy->ForwardAxis + 1) % 3] = 1;
+		SetShaderValue(Shader->GetVertexShader(), SplineMeshXParam, DirMask);
+		DirMask = FVector::ZeroVector;
+		DirMask[(SplineProxy->ForwardAxis + 2) % 3] = 1;
+		SetShaderValue(Shader->GetVertexShader(), SplineMeshYParam, DirMask);
+	}
+}
+
+//////////////////////////////////////////////////////////////////////////
+// SplineMeshVertexFactory
+
+IMPLEMENT_VERTEX_FACTORY_TYPE(FSplineMeshVertexFactory, "LocalVertexFactory", true, true, true, true, true);
+
+
+FVertexFactoryShaderParameters* FSplineMeshVertexFactory::ConstructShaderParameters(EShaderFrequency ShaderFrequency)
+{
+	return ShaderFrequency == SF_Vertex ? new FSplineMeshVertexFactoryShaderParameters() : NULL;
+}
+
+//////////////////////////////////////////////////////////////////////////
+// SplineMeshSceneProxy
+
+
+
+void FSplineMeshSceneProxy::InitResources(USplineMeshComponent* InComponent, int32 InLODIndex)
+{
+	// Initialize the static mesh's vertex factory.
+	ENQUEUE_UNIQUE_RENDER_COMMAND_THREEPARAMETER(
+		InitSplineMeshVertexFactory,
+		FSplineMeshVertexFactory*, VertexFactory, LODResources[InLODIndex].VertexFactory,
+		const FStaticMeshLODResources*, RenderData, &InComponent->StaticMesh->RenderData->LODResources[InLODIndex],
+		UStaticMesh*, Parent, InComponent->StaticMesh,
+		{
+		FLocalVertexFactory::DataType Data;
+
+		Data.PositionComponent = FVertexStreamComponent(
+			&RenderData->PositionVertexBuffer,
+			STRUCT_OFFSET(FPositionVertex, Position),
+			RenderData->PositionVertexBuffer.GetStride(),
+			VET_Float3
+			);
+		Data.TangentBasisComponents[0] = FVertexStreamComponent(
+			&RenderData->VertexBuffer,
+			STRUCT_OFFSET(FStaticMeshFullVertex, TangentX),
+			RenderData->VertexBuffer.GetStride(),
+			VET_PackedNormal
+			);
+		Data.TangentBasisComponents[1] = FVertexStreamComponent(
+			&RenderData->VertexBuffer,
+			STRUCT_OFFSET(FStaticMeshFullVertex, TangentZ),
+			RenderData->VertexBuffer.GetStride(),
+			VET_PackedNormal
+			);
+
+		if (RenderData->ColorVertexBuffer.GetNumVertices() > 0)
+		{
+			Data.ColorComponent = FVertexStreamComponent(
+				&RenderData->ColorVertexBuffer,
+				0,	// Struct offset to color
+				RenderData->ColorVertexBuffer.GetStride(),
+				VET_Color
+				);
+		}
+
+		Data.TextureCoordinates.Empty();
+
+		if (!RenderData->VertexBuffer.GetUseFullPrecisionUVs())
+		{
+			for (uint32 UVIndex = 0; UVIndex < RenderData->VertexBuffer.GetNumTexCoords(); UVIndex++)
+			{
+				Data.TextureCoordinates.Add(FVertexStreamComponent(
+					&RenderData->VertexBuffer,
+					STRUCT_OFFSET(TStaticMeshFullVertexFloat16UVs<MAX_TEXCOORDS>, UVs) + sizeof(FVector2DHalf)* UVIndex,
+					RenderData->VertexBuffer.GetStride(),
+					VET_Half2
+					));
+			}
+			if (Parent->LightMapCoordinateIndex >= 0 && (uint32)Parent->LightMapCoordinateIndex < RenderData->VertexBuffer.GetNumTexCoords())
+			{
+				Data.LightMapCoordinateComponent = FVertexStreamComponent(
+					&RenderData->VertexBuffer,
+					STRUCT_OFFSET(TStaticMeshFullVertexFloat16UVs<MAX_TEXCOORDS>, UVs) + sizeof(FVector2DHalf)* Parent->LightMapCoordinateIndex,
+					RenderData->VertexBuffer.GetStride(),
+					VET_Half2
+					);
+			}
+		}
+		else
+		{
+			for (uint32 UVIndex = 0; UVIndex < RenderData->VertexBuffer.GetNumTexCoords(); UVIndex++)
+			{
+				Data.TextureCoordinates.Add(FVertexStreamComponent(
+					&RenderData->VertexBuffer,
+					STRUCT_OFFSET(TStaticMeshFullVertexFloat32UVs<MAX_TEXCOORDS>, UVs) + sizeof(FVector2D)* UVIndex,
+					RenderData->VertexBuffer.GetStride(),
+					VET_Float2
+					));
+			}
+
+			if (Parent->LightMapCoordinateIndex >= 0 && (uint32)Parent->LightMapCoordinateIndex < RenderData->VertexBuffer.GetNumTexCoords())
+			{
+				Data.LightMapCoordinateComponent = FVertexStreamComponent(
+					&RenderData->VertexBuffer,
+					STRUCT_OFFSET(TStaticMeshFullVertexFloat32UVs<MAX_TEXCOORDS>, UVs) + sizeof(FVector2D)* Parent->LightMapCoordinateIndex,
+					RenderData->VertexBuffer.GetStride(),
+					VET_Float2
+					);
+			}
+		}
+
+		VertexFactory->SetData(Data);
+
+		VertexFactory->InitResource();
+	});
+}
+
+
+void FSplineMeshSceneProxy::ReleaseResources()
+{
+	for (FSplineMeshSceneProxy::FLODResources& LODResource : LODResources)
+	{
+		LODResource.VertexFactory->ReleaseResource();
+	}
+}
+
+//////////////////////////////////////////////////////////////////////////
+// SplineMeshComponent
+
+USplineMeshComponent::USplineMeshComponent(const class FPostConstructInitializeProperties& PCIP)
+	: Super(PCIP)
+{
+	Mobility = EComponentMobility::Static;
+
+	BodyInstance.bEnableCollision_DEPRECATED = false;
+	SetCollisionProfileName(UCollisionProfile::NoCollision_ProfileName);
+	bHasCustomNavigableGeometry = EHasCustomNavigableGeometry::Yes;
+
+	SplineXDir.Z = 1.0f;
+
+	// Default to useful length and scale
+	SplineParams.StartTangent = FVector(100.f, 0.f, 0.f);
+	SplineParams.StartScale = FVector2D(1.f, 1.f);
+
+	SplineParams.EndPos = FVector(100.f, 0.f, 0.f);
+	SplineParams.EndTangent = FVector(100.f, 0.f, 0.f);
+	SplineParams.EndScale = FVector2D(1.f, 1.f);
+}
+
+void USplineMeshComponent::SetStartPosition(FVector StartPos)
+{
+	SplineParams.StartPos = StartPos;
+	MarkSplineParamsDirty();
+}
+
+void USplineMeshComponent::SetStartTangent(FVector StartTangent)
+{
+	SplineParams.StartTangent = StartTangent;
+	MarkSplineParamsDirty();
+}
+
+void USplineMeshComponent::SetEndPosition(FVector EndPos)
+{
+	SplineParams.EndPos = EndPos;
+	MarkSplineParamsDirty();
+}
+
+void USplineMeshComponent::SetEndTangent(FVector EndTangent)
+{
+	SplineParams.EndTangent = EndTangent;
+	MarkSplineParamsDirty();
+}
+
+void USplineMeshComponent::MarkSplineParamsDirty()
+{
+	MarkRenderStateDirty();
+
+#if WITH_EDITOR
+	if (!GetWorld()->AreActorsInitialized())
+	{
+		RecreateCollision();
+	}
+#endif // WITH_EDITOR
+}
+
+void USplineMeshComponent::Serialize(FArchive& Ar)
+{
+	Super::Serialize(Ar);
+
+	if (Ar.UE4Ver() < VER_UE4_SPLINE_MESH_ORIENTATION)
+	{
+		ForwardAxis = ESplineMeshAxis::Z;
+		SplineParams.StartRoll -= HALF_PI;
+		SplineParams.EndRoll -= HALF_PI;
+
+		float Temp = SplineParams.StartOffset.X;
+		SplineParams.StartOffset.X = -SplineParams.StartOffset.Y;
+		SplineParams.StartOffset.Y = Temp;
+		Temp = SplineParams.EndOffset.X;
+		SplineParams.EndOffset.X = -SplineParams.EndOffset.Y;
+		SplineParams.EndOffset.Y = Temp;
+	}
+
+#if WITH_EDITOR
+	if (BodySetup != NULL)
+	{
+		BodySetup->SetFlags(RF_Transactional);
+	}
+#endif
+}
+
+bool USplineMeshComponent::Modify(bool bAlwaysMarkDirty)
+{
+	bool bSavedToTransactionBuffer = Super::Modify(bAlwaysMarkDirty);
+
+	if (BodySetup != NULL)
+	{
+		BodySetup->Modify(bAlwaysMarkDirty);
+	}
+
+	return bSavedToTransactionBuffer;
+}
+
+FPrimitiveSceneProxy* USplineMeshComponent::CreateSceneProxy()
+{
+	// Verify that the mesh is valid before using it.
+	const bool bMeshIsValid = 
+		// make sure we have an actual staticmesh
+		StaticMesh &&
+		StaticMesh->HasValidRenderData();
+
+	if (bMeshIsValid && (GetScene()->GetFeatureLevel() >= ERHIFeatureLevel::SM3))
+	{
+		return ::new FSplineMeshSceneProxy(this);
+	}
+	else
+	{
+		return NULL;
+	}
+}
+
+FBoxSphereBounds USplineMeshComponent::CalcBounds(const FTransform & LocalToWorld) const
+{
+	// Use util to generate bounds of spline
+	FInterpCurvePoint<FVector> Start(0.f, SplineParams.StartPos, SplineParams.StartTangent, SplineParams.StartTangent, CIM_CurveUser);
+	FInterpCurvePoint<FVector> End(1.f, SplineParams.EndPos, SplineParams.EndTangent, SplineParams.EndTangent, CIM_CurveUser);
+
+	FVector CurveMax(-BIG_NUMBER, -BIG_NUMBER, -BIG_NUMBER);
+	FVector CurveMin(BIG_NUMBER, BIG_NUMBER, BIG_NUMBER);
+	CurveVectorFindIntervalBounds(Start, End, CurveMin, CurveMax);
+
+	FBox LocalBox(CurveMin, CurveMax);
+
+	// Find largest extent of mesh in XY, and add on all around
+	if(StaticMesh)
+	{
+		FVector MinMeshExtent = StaticMesh->GetBounds().Origin - StaticMesh->GetBounds().BoxExtent;
+		FVector MaxMeshExtent = StaticMesh->GetBounds().Origin + StaticMesh->GetBounds().BoxExtent;
+		GetAxisValue(MinMeshExtent, ForwardAxis) = 0;
+		GetAxisValue(MaxMeshExtent, ForwardAxis) = 0;
+		float MaxDim = FMath::Max<float>(MinMeshExtent.GetAbsMax(), MaxMeshExtent.GetAbsMax());
+
+		float MaxScale = FMath::Max(SplineParams.StartScale.GetAbsMax(), SplineParams.EndScale.GetAbsMax());
+
+		LocalBox = LocalBox.ExpandBy(MaxScale * MaxDim);
+	}
+
+	return FBoxSphereBounds( LocalBox.TransformBy(LocalToWorld) );
+}
+
+
+/** 
+ * Functions used for transforming a static mesh component based on a spline.  
+ * This needs to be updated if the spline functionality changes!
+ */
+static float SmoothStep(float A, float B, float X)
+{
+	if (X < A)
+	{
+		return 0.0f;
+	}
+	else if (X >= B)
+	{
+		return 1.0f;
+	}
+	const float InterpFraction = (X - A) / (B - A);
+	return InterpFraction * InterpFraction * (3.0f - 2.0f * InterpFraction);
+}
+
+static FVector SplineEvalPos(const FVector& StartPos, const FVector& StartTangent, const FVector& EndPos, const FVector& EndTangent, float A)
+{
+	const float A2 = A  * A;
+	const float A3 = A2 * A;
+
+	return (((2*A3)-(3*A2)+1) * StartPos) + ((A3-(2*A2)+A) * StartTangent) + ((A3-A2) * EndTangent) + (((-2*A3)+(3*A2)) * EndPos);
+}
+
+static FVector SplineEvalDir(const FVector& StartPos, const FVector& StartTangent, const FVector& EndPos, const FVector& EndTangent, float A)
+{
+	const FVector C = (6*StartPos) + (3*StartTangent) + (3*EndTangent) - (6*EndPos);
+	const FVector D = (-6*StartPos) - (4*StartTangent) - (2*EndTangent) + (6*EndPos);
+	const FVector E = StartTangent;
+
+	const float A2 = A  * A;
+
+	return ((C * A2) + (D * A) + E).SafeNormal();
+}
+
+
+FTransform USplineMeshComponent::CalcSliceTransform(const float DistanceAlong) const
+{
+	// Find how far 'along' mesh we are
+	FBoxSphereBounds StaticMeshBounds = StaticMesh->GetBounds();
+	const float MeshMinZ = GetAxisValue(StaticMeshBounds.Origin, ForwardAxis) - GetAxisValue(StaticMeshBounds.BoxExtent, ForwardAxis);
+	const float MeshRangeZ = 2.0f * GetAxisValue(StaticMeshBounds.BoxExtent, ForwardAxis);
+	const float Alpha = (DistanceAlong - MeshMinZ) / MeshRangeZ;
+
+	// Apply hermite interp to Alpha if desired
+	const float HermiteAlpha = bSmoothInterpRollScale ? SmoothStep(0.0, 1.0, Alpha) : Alpha;
+
+	// Then find the point and direction of the spline at this point along
+	FVector SplinePos = SplineEvalPos( SplineParams.StartPos, SplineParams.StartTangent, SplineParams.EndPos, SplineParams.EndTangent, Alpha );
+	const FVector SplineDir = SplineEvalDir( SplineParams.StartPos, SplineParams.StartTangent, SplineParams.EndPos, SplineParams.EndTangent, Alpha );
+
+	// Find base frenet frame
+	const FVector BaseXVec = (SplineXDir ^ SplineDir).SafeNormal();
+	const FVector BaseYVec = (SplineDir ^ BaseXVec).SafeNormal();
+
+	// Offset the spline by the desired amount
+	const FVector2D SliceOffset = FMath::Lerp<FVector2D>(SplineParams.StartOffset, SplineParams.EndOffset, HermiteAlpha);
+	SplinePos += SliceOffset.X * BaseXVec;
+	SplinePos += SliceOffset.Y * BaseYVec;
+
+	// Apply roll to frame around spline
+	const float UseRoll = FMath::Lerp(SplineParams.StartRoll, SplineParams.EndRoll, HermiteAlpha);
+	const float CosAng = FMath::Cos(UseRoll);
+	const float SinAng = FMath::Sin(UseRoll);
+	const FVector XVec = (CosAng * BaseXVec) - (SinAng * BaseYVec);
+	const FVector YVec = (CosAng * BaseYVec) + (SinAng * BaseXVec);
+
+	// Find scale at this point along spline
+	const FVector2D UseScale = FMath::Lerp(SplineParams.StartScale, SplineParams.EndScale, HermiteAlpha);
+
+	// Build overall transform
+	FTransform SliceTransform;
+	switch (ForwardAxis)
+	{
+	case ESplineMeshAxis::X:
+		SliceTransform = FTransform(SplineDir, XVec, YVec, SplinePos);
+		SliceTransform.SetScale3D(FVector(1, UseScale.X, UseScale.Y));
+		break;
+	case ESplineMeshAxis::Y:
+		SliceTransform = FTransform(YVec, SplineDir, XVec, SplinePos);
+		SliceTransform.SetScale3D(FVector(UseScale.Y, 1, UseScale.X));
+		break;
+	case ESplineMeshAxis::Z:
+		SliceTransform = FTransform(XVec, YVec, SplineDir, SplinePos);
+		SliceTransform.SetScale3D(FVector(UseScale.X, UseScale.Y, 1));
+		break;
+	default:
+		check(0);
+		break;
+	}
+
+	return SliceTransform;
+}
+
+
+bool USplineMeshComponent::GetPhysicsTriMeshData(struct FTriMeshCollisionData* CollisionData, bool InUseAllTriData)
+{
+	if (StaticMesh)
+	{
+		StaticMesh->GetPhysicsTriMeshData(CollisionData, InUseAllTriData);
+
+		FVector Mask = FVector(1,1,1);
+		GetAxisValue(Mask, ForwardAxis) = 0;
+
+		for (FVector& CollisionVert : CollisionData->Vertices)
+		{
+			CollisionVert = CalcSliceTransform(GetAxisValue(CollisionVert, ForwardAxis)).TransformPosition(CollisionVert * Mask);
+		}
+		return true;
+	}
+
+	return false;
+}
+
+bool USplineMeshComponent::ContainsPhysicsTriMeshData(bool InUseAllTriData) const
+{
+	if (StaticMesh)
+	{
+		return StaticMesh->ContainsPhysicsTriMeshData(InUseAllTriData);
+	}
+
+	return false;
+}
+
+void USplineMeshComponent::CreatePhysicsState()
+{
+#if WITH_EDITOR
+	// With editor code we can recreate the collision if the mesh changes
+	const FGuid MeshBodySetupGuid = (StaticMesh != NULL ? StaticMesh->BodySetup->BodySetupGuid : FGuid());
+	if (CachedMeshBodySetupGuid != MeshBodySetupGuid)
+	{
+		RecreateCollision();
+	}
+#else
+	// Without editor code we can only destroy the collision if the mesh is missing
+	if (StaticMesh == NULL && BodySetup != NULL)
+	{
+		DestroyBodySetup();
+	}
+#endif
+
+	return Super::CreatePhysicsState();
+}
+
+UBodySetup* USplineMeshComponent::GetBodySetup()
+{
+	// Don't return a body setup that has no collision, it means we are interactively moving the spline and don't want to build collision.
+	// Instead we explicitly build collision with USplineMeshComponent::RecreateCollision()
+	if (BodySetup != NULL && (BodySetup->TriMesh != NULL || BodySetup->AggGeom.GetElementCount() > 0))
+	{
+		return BodySetup;
+	}
+	return NULL;
+}
+
+bool USplineMeshComponent::DoCustomNavigableGeometryExport(struct FNavigableGeometryExport* GeomExport) const
+{
+	// the NavCollision is supposed to be faster than exporting the regular collision,
+	// but I'm not sure that's true here, as the regular collision is pre-distorted to the spline
+
+	if (StaticMesh != NULL && StaticMesh->NavCollision != NULL)
+	{
+		UNavCollision* NavCollision = StaticMesh->NavCollision;
+		
+		if (ensure(!NavCollision->bIsDynamicObstacle))
+		{
+			if (NavCollision->bHasConvexGeometry)
+			{
+				TArray<FVector> VertexBuffer;
+				VertexBuffer.Reserve(FMath::Max(NavCollision->ConvexCollision.VertexBuffer.Num(), NavCollision->TriMeshCollision.VertexBuffer.Num()));
+
+				for (int32 i = 0; i < NavCollision->ConvexCollision.VertexBuffer.Num(); ++i)
+				{
+					FVector Vertex = NavCollision->ConvexCollision.VertexBuffer[i];
+					Vertex = CalcSliceTransform(GetAxisValue(Vertex, ForwardAxis)).TransformPosition(Vertex);
+					VertexBuffer.Add(Vertex);
+				}
+				GeomExport->ExportCustomMesh(VertexBuffer.GetData(), VertexBuffer.Num(),
+					NavCollision->ConvexCollision.IndexBuffer.GetData(), NavCollision->ConvexCollision.IndexBuffer.Num(),
+					ComponentToWorld);
+
+				VertexBuffer.Reset();
+				for (int32 i = 0; i < NavCollision->TriMeshCollision.VertexBuffer.Num(); ++i)
+				{
+					FVector Vertex = NavCollision->TriMeshCollision.VertexBuffer[i];
+					Vertex = CalcSliceTransform(GetAxisValue(Vertex, ForwardAxis)).TransformPosition(Vertex);
+					VertexBuffer.Add(Vertex);
+				}
+				GeomExport->ExportCustomMesh(VertexBuffer.GetData(), VertexBuffer.Num(),
+					NavCollision->TriMeshCollision.IndexBuffer.GetData(), NavCollision->TriMeshCollision.IndexBuffer.Num(),
+					ComponentToWorld);
+
+				return false;
+			}
+		}
+	}
+
+	return true;
+}
+
+void USplineMeshComponent::DestroyBodySetup()
+{
+	if (BodySetup != NULL)
+	{
+		BodySetup->MarkPendingKill();
+		BodySetup = NULL;
+#if WITH_EDITORONLY_DATA
+		CachedMeshBodySetupGuid.Invalidate();
+#endif
+	}
+}
+
+#if WITH_EDITOR
+void USplineMeshComponent::RecreateCollision()
+{
+	if (StaticMesh && IsCollisionEnabled())
+	{
+		if (BodySetup == NULL)
+		{
+			BodySetup = DuplicateObject<UBodySetup>(StaticMesh->BodySetup, this);
+			BodySetup->SetFlags(RF_Transactional);
+			BodySetup->InvalidatePhysicsData();
+		}
+		else
+		{
+			BodySetup->Modify();
+			BodySetup->InvalidatePhysicsData();
+			BodySetup->CopyBodyPropertiesFrom(StaticMesh->BodySetup);
+			BodySetup->CollisionTraceFlag = StaticMesh->BodySetup->CollisionTraceFlag;
+		}
+
+		if (BodySetup->CollisionTraceFlag == CTF_UseComplexAsSimple)
+		{
+			BodySetup->AggGeom.EmptyElements();
+		}
+		else
+		{
+			FVector Mask = FVector(1,1,1);
+			GetAxisValue(Mask, ForwardAxis) = 0;
+
+			// distortion of a sphere can't be done nicely, so we just transform the origin and size
+			for (FKSphereElem& SphereElem : BodySetup->AggGeom.SphereElems)
+			{
+				const float Z = GetAxisValue(SphereElem.Center, ForwardAxis);
+				FTransform SliceTransform = CalcSliceTransform(Z);
+				SphereElem.Center *= Mask;
+
+				SphereElem.Radius *= SliceTransform.GetMaximumAxisScale();
+
+				SliceTransform.RemoveScaling();
+				SphereElem.Center = SliceTransform.TransformPosition(SphereElem.Center);
+			}
+
+			// distortion of a sphyl can't be done nicely, so we just transform the origin and size
+			for (FKSphylElem& SphylElem : BodySetup->AggGeom.SphylElems)
+			{
+				const float Z = GetAxisValue(SphylElem.Center, ForwardAxis);
+				FTransform SliceTransform = CalcSliceTransform(Z);
+				SphylElem.Center *= Mask;
+
+				FTransform TM = SphylElem.GetTransform();
+				SphylElem.Length = (TM * SliceTransform).TransformVector(FVector(0, 0, SphylElem.Length)).Size();
+				SphylElem.Radius *= SliceTransform.GetMaximumAxisScale();
+
+				SliceTransform.RemoveScaling();
+				SphylElem.SetTransform(TM * SliceTransform);
+			}
+
+			// Convert boxes to convex hulls to better respect distortion
+			for (FKBoxElem& BoxElem : BodySetup->AggGeom.BoxElems)
+			{
+				FKConvexElem& ConvexElem = *new(BodySetup->AggGeom.ConvexElems) FKConvexElem();
+
+				const FVector Radii = FVector(BoxElem.X / 2, BoxElem.Y / 2, BoxElem.Z / 2);
+				const FTransform ElementTM = BoxElem.GetTransform();
+				ConvexElem.VertexData.Empty(8);
+				ConvexElem.VertexData.Add(ElementTM.TransformPosition(Radii * FVector(-1,-1,-1)));
+				ConvexElem.VertexData.Add(ElementTM.TransformPosition(Radii * FVector(-1,-1, 1)));
+				ConvexElem.VertexData.Add(ElementTM.TransformPosition(Radii * FVector(-1, 1,-1)));
+				ConvexElem.VertexData.Add(ElementTM.TransformPosition(Radii * FVector(-1, 1, 1)));
+				ConvexElem.VertexData.Add(ElementTM.TransformPosition(Radii * FVector( 1,-1,-1)));
+				ConvexElem.VertexData.Add(ElementTM.TransformPosition(Radii * FVector( 1,-1, 1)));
+				ConvexElem.VertexData.Add(ElementTM.TransformPosition(Radii * FVector( 1, 1,-1)));
+				ConvexElem.VertexData.Add(ElementTM.TransformPosition(Radii * FVector( 1, 1, 1)));
+
+				ConvexElem.UpdateElemBox();
+			}
+			BodySetup->AggGeom.BoxElems.Empty();
+
+			// transform the points of the convex hulls into spline space
+			for (FKConvexElem& ConvexElem : BodySetup->AggGeom.ConvexElems)
+			{
+				for (FVector& Point : ConvexElem.VertexData)
+				{
+					Point = CalcSliceTransform(GetAxisValue(Point, ForwardAxis)).TransformPosition(Point * Mask);
+				}
+			}
+
+			BodySetup->CreatePhysicsMeshes();
+			CachedMeshBodySetupGuid = StaticMesh->BodySetup->BodySetupGuid;
+		}
+	}
+	else
+	{
+		DestroyBodySetup();
+	}
+}
+#endif
+
+#include "StaticMeshLight.h"
+/** */
+class FSplineStaticLightingMesh : public FStaticMeshStaticLightingMesh
+{
+public:
+
+	FSplineStaticLightingMesh(const USplineMeshComponent* InPrimitive,int32 InLODIndex,const TArray<ULightComponent*>& InRelevantLights):
+		FStaticMeshStaticLightingMesh(InPrimitive, InLODIndex, InRelevantLights),
+		SplineComponent(InPrimitive)
+	{
+	}
+
+#if WITH_EDITOR
+	virtual const struct FSplineMeshParams* GetSplineParameters() const
+	{
+		return &SplineComponent->SplineParams;
+	}
+#endif	//WITH_EDITOR
+
+private:
+	const USplineMeshComponent* SplineComponent;
+};
+
+FStaticMeshStaticLightingMesh* USplineMeshComponent::AllocateStaticLightingMesh(int32 LODIndex, const TArray<ULightComponent*>& InRelevantLights)
+{
+	return new FSplineStaticLightingMesh(this, LODIndex, InRelevantLights);
+}
+
+
+const float& USplineMeshComponent::GetAxisValue(const FVector& InVector, ESplineMeshAxis::Type InAxis)
+{
+	switch (InAxis)
+	{
+	case ESplineMeshAxis::X:
+		return InVector.X;
+	case ESplineMeshAxis::Y:
+		return InVector.Y;
+	case ESplineMeshAxis::Z:
+		return InVector.Z;
+	default:
+		check(0);
+		return InVector.Z;
+	}
+}
+
+float& USplineMeshComponent::GetAxisValue(FVector& InVector, ESplineMeshAxis::Type InAxis)
+{
+	switch (InAxis)
+	{
+	case ESplineMeshAxis::X:
+		return InVector.X;
+	case ESplineMeshAxis::Y:
+		return InVector.Y;
+	case ESplineMeshAxis::Z:
+		return InVector.Z;
+	default:
+		check(0);
+		return InVector.Z;
+	}
+}
\ No newline at end of file
diff --git a/Engine/Source/Runtime/Engine/Private/SplineMeshSceneProxy.h b/Engine/Source/Runtime/Engine/Private/SplineMeshSceneProxy.h
new file mode 100644
--- /dev/null
+++ b/Engine/Source/Runtime/Engine/Private/SplineMeshSceneProxy.h
@@ -0,0 +1,265 @@
+// Copyright 1998-2014 Epic Games, Inc. All Rights Reserved.
+
+#pragma once
+
+#include "StaticMeshResources.h"
+
+class FVertexFactoryShaderParameters;
+
+//////////////////////////////////////////////////////////////////////////
+// SplineMeshVertexFactory
+
+/** A vertex factory for spline-deformed static meshes */
+struct FSplineMeshVertexFactory : public FLocalVertexFactory
+{
+	DECLARE_VERTEX_FACTORY_TYPE(FSplineMeshVertexFactory);
+public:
+
+	FSplineMeshVertexFactory(class FSplineMeshSceneProxy* InSplineProxy) :
+		SplineSceneProxy(InSplineProxy)
+	{
+	}
+
+	/** Should we cache the material's shadertype on this platform with this vertex factory? */
+	static bool ShouldCache(EShaderPlatform Platform, const class FMaterial* Material, const class FShaderType* ShaderType)
+	{
+		return (Material->IsUsedWithSplineMeshes() || Material->IsSpecialEngineMaterial())
+			&& FLocalVertexFactory::ShouldCache(Platform, Material, ShaderType)
+			&& IsFeatureLevelSupported(Platform, ERHIFeatureLevel::SM3);
+	}
+
+	/** Modify compile environment to enable spline deformation */
+	static void ModifyCompilationEnvironment(EShaderPlatform Platform, const FMaterial* Material, FShaderCompilerEnvironment& OutEnvironment)
+	{
+		FLocalVertexFactory::ModifyCompilationEnvironment(Platform, Material, OutEnvironment);
+		OutEnvironment.SetDefine(TEXT("USE_SPLINEDEFORM"), TEXT("1"));
+	}
+
+	/** Copy the data from another vertex factory */
+	void Copy(const FSplineMeshVertexFactory& Other)
+	{
+		ENQUEUE_UNIQUE_RENDER_COMMAND_TWOPARAMETER(
+			FSplineMeshVertexFactoryCopyData,
+			FSplineMeshVertexFactory*, VertexFactory, this,
+			const DataType*, DataCopy, &Other.Data,
+			{
+			VertexFactory->Data = *DataCopy;
+		});
+		BeginUpdateResourceRHI(this);
+	}
+
+	// FRenderResource interface.
+	virtual void InitRHI()
+	{
+		FLocalVertexFactory::InitRHI();
+	}
+
+	static FVertexFactoryShaderParameters* ConstructShaderParameters(EShaderFrequency ShaderFrequency);
+
+	class FSplineMeshSceneProxy* SplineSceneProxy;
+
+private:
+};
+
+//////////////////////////////////////////////////////////////////////////
+// FSplineMeshVertexFactoryShaderParameters
+
+/** Factory specific params */
+class FSplineMeshVertexFactoryShaderParameters : public FVertexFactoryShaderParameters
+{
+	void Bind(const FShaderParameterMap& ParameterMap) OVERRIDE;
+
+	void SetMesh(FShader* Shader, const FVertexFactory* VertexFactory, const FSceneView& View, const FMeshBatchElement& BatchElement, uint32 DataFlags) const OVERRIDE;
+
+	void Serialize(FArchive& Ar) OVERRIDE
+	{
+		Ar << SplineStartPosParam;
+		Ar << SplineStartTangentParam;
+		Ar << SplineStartRollParam;
+		Ar << SplineStartScaleParam;
+		Ar << SplineStartOffsetParam;
+
+		Ar << SplineEndPosParam;
+		Ar << SplineEndTangentParam;
+		Ar << SplineEndRollParam;
+		Ar << SplineEndScaleParam;
+		Ar << SplineEndOffsetParam;
+
+		Ar << SplineUpDirParam;
+		Ar << SmoothInterpRollScaleParam;
+
+		Ar << SplineMeshMinZParam;
+		Ar << SplineMeshScaleZParam;
+
+		Ar << SplineMeshDirParam;
+		Ar << SplineMeshXParam;
+		Ar << SplineMeshYParam;
+	}
+
+	virtual uint32 GetSize() const OVERRIDE
+	{
+		return sizeof(*this);
+	}
+
+private:
+	FShaderParameter SplineStartPosParam;
+	FShaderParameter SplineStartTangentParam;
+	FShaderParameter SplineStartRollParam;
+	FShaderParameter SplineStartScaleParam;
+	FShaderParameter SplineStartOffsetParam;
+
+	FShaderParameter SplineEndPosParam;
+	FShaderParameter SplineEndTangentParam;
+	FShaderParameter SplineEndRollParam;
+	FShaderParameter SplineEndScaleParam;
+	FShaderParameter SplineEndOffsetParam;
+
+	FShaderParameter SplineUpDirParam;
+	FShaderParameter SmoothInterpRollScaleParam;
+
+	FShaderParameter SplineMeshMinZParam;
+	FShaderParameter SplineMeshScaleZParam;
+
+	FShaderParameter SplineMeshDirParam;
+	FShaderParameter SplineMeshXParam;
+	FShaderParameter SplineMeshYParam;
+};
+
+//////////////////////////////////////////////////////////////////////////
+// SplineMeshSceneProxy
+
+/** Scene proxy for SplineMesh instance */
+class FSplineMeshSceneProxy : public FStaticMeshSceneProxy
+{
+protected:
+	struct FLODResources
+	{
+		/** Pointer to vertex factory object */
+		FSplineMeshVertexFactory* VertexFactory;
+
+		FLODResources(FSplineMeshVertexFactory* InVertexFactory) :
+			VertexFactory(InVertexFactory)
+		{
+		}
+	};
+public:
+
+	FSplineMeshSceneProxy(USplineMeshComponent* InComponent) :
+		FStaticMeshSceneProxy(InComponent)
+	{
+		// make sure all the materials are okay to be rendered as a spline mesh
+		for (FStaticMeshSceneProxy::FLODInfo& LODInfo : LODs)
+		{
+			for (FStaticMeshSceneProxy::FLODInfo::FSectionInfo& Section : LODInfo.Sections)
+			{
+				if (!Section.Material->CheckMaterialUsage(MATUSAGE_SplineMesh))
+				{
+					Section.Material = UMaterial::GetDefaultMaterial(MD_Surface);
+				}
+			}
+		}
+
+		// Copy spline params from component
+		SplineParams = InComponent->SplineParams;
+		SplineXDir = InComponent->SplineXDir;
+		bSmoothInterpRollScale = InComponent->bSmoothInterpRollScale;
+		ForwardAxis = InComponent->ForwardAxis;
+
+		// Fill in info about the mesh
+		FBoxSphereBounds StaticMeshBounds = StaticMesh->GetBounds();
+		SplineMeshScaleZ = 0.5f / USplineMeshComponent::GetAxisValue(StaticMeshBounds.BoxExtent, ForwardAxis); // 1/(2 * Extent)
+		SplineMeshMinZ = USplineMeshComponent::GetAxisValue(StaticMeshBounds.Origin, ForwardAxis) * SplineMeshScaleZ - 0.5f;
+
+		LODResources.Reset(InComponent->StaticMesh->RenderData->LODResources.Num());
+
+		for (int32 LODIndex = 0; LODIndex < LODs.Num(); LODIndex++)
+		{
+			FSplineMeshVertexFactory* VertexFactory = new FSplineMeshVertexFactory(this);
+
+			LODResources.Add(VertexFactory);
+
+			InitResources(InComponent, LODIndex);
+		}
+	}
+
+	virtual ~FSplineMeshSceneProxy() OVERRIDE
+	{
+		ReleaseResources();
+
+		for (FSplineMeshSceneProxy::FLODResources& LODResource : LODResources)
+		{
+			delete LODResource.VertexFactory;
+		}
+		LODResources.Empty();
+	}
+
+	void InitResources(USplineMeshComponent* InComponent, int32 InLODIndex);
+
+	void ReleaseResources();
+
+	/** Sets up a shadow FMeshBatch for a specific LOD. */
+	virtual bool GetShadowMeshElement(int32 LODIndex, uint8 InDepthPriorityGroup, FMeshBatch& OutMeshElement) const OVERRIDE
+	{
+		//checkf(LODIndex == 0, TEXT("Getting spline static mesh element with invalid LOD [%d]"), LODIndex);
+
+		if (FStaticMeshSceneProxy::GetShadowMeshElement(LODIndex, InDepthPriorityGroup, OutMeshElement))
+		{
+			OutMeshElement.VertexFactory = LODResources[LODIndex].VertexFactory;
+			OutMeshElement.ReverseCulling ^= (SplineParams.StartScale.X < 0) ^ (SplineParams.StartScale.Y < 0);
+			return true;
+		}
+		return false;
+	}
+
+	/** Sets up a FMeshBatch for a specific LOD and element. */
+	virtual bool GetMeshElement(int32 LODIndex, int32 SectionIndex, uint8 InDepthPriorityGroup, FMeshBatch& OutMeshElement, const bool bUseSelectedMaterial, const bool bUseHoveredMaterial) const OVERRIDE
+	{
+		//checkf(LODIndex == 0 /*&& SectionIndex == 0*/, TEXT("Getting spline static mesh element with invalid params [%d, %d]"), LODIndex, SectionIndex);
+
+		if (FStaticMeshSceneProxy::GetMeshElement(LODIndex, SectionIndex, InDepthPriorityGroup, OutMeshElement, bUseSelectedMaterial, bUseHoveredMaterial))
+		{
+			OutMeshElement.VertexFactory = LODResources[LODIndex].VertexFactory;
+			OutMeshElement.ReverseCulling ^= (SplineParams.StartScale.X < 0) ^ (SplineParams.StartScale.Y < 0);
+			return true;
+		}
+		return false;
+	}
+
+	/** Sets up a wireframe FMeshBatch for a specific LOD. */
+	virtual bool GetWireframeMeshElement(int32 LODIndex, const FMaterialRenderProxy* WireframeRenderProxy, uint8 InDepthPriorityGroup, FMeshBatch& OutMeshElement) const OVERRIDE
+	{
+		//checkf(LODIndex == 0, TEXT("Getting spline static mesh element with invalid LOD [%d]"), LODIndex);
+
+		if (FStaticMeshSceneProxy::GetWireframeMeshElement(LODIndex, WireframeRenderProxy, InDepthPriorityGroup, OutMeshElement))
+		{
+			OutMeshElement.VertexFactory = LODResources[LODIndex].VertexFactory;
+			OutMeshElement.ReverseCulling ^= (SplineParams.StartScale.X < 0) ^ (SplineParams.StartScale.Y < 0);
+			return true;
+		}
+		return false;
+	}
+
+	virtual FPrimitiveViewRelevance GetViewRelevance(const FSceneView* View) OVERRIDE
+	{
+		return FStaticMeshSceneProxy::GetViewRelevance(View);
+	}
+
+	// 	  virtual uint32 GetMemoryFootprint( void ) const { return 0; }
+
+	/** Parameters that define the spline, used to deform mesh */
+	FSplineMeshParams SplineParams;
+	/** Axis (in component space) that is used to determine X axis for co-ordinates along spline */
+	FVector SplineXDir;
+	/** Smoothly (cubic) interpolate the Roll and Scale params over spline. */
+	bool bSmoothInterpRollScale;
+	/** Chooses the forward axis for the spline mesh orientation */
+	ESplineMeshAxis::Type ForwardAxis;
+
+	/** Minimum Z value of the entire mesh */
+	float SplineMeshMinZ;
+	/** Range of Z values over entire mesh */
+	float SplineMeshScaleZ;
+
+protected:
+	TArray<FLODResources> LODResources;
+};
\ No newline at end of file
