Expose default terminal velocity in physics settings

diff --git a/Engine/Source/Runtime/Engine/Classes/PhysicsEngine/PhysicsSettings.h b/Engine/Source/Runtime/Engine/Classes/PhysicsEngine/PhysicsSettings.h
--- a/Engine/Source/Runtime/Engine/Classes/PhysicsEngine/PhysicsSettings.h
+++ b/Engine/Source/Runtime/Engine/Classes/PhysicsEngine/PhysicsSettings.h
@@ -61,6 +61,10 @@ class ENGINE_API UPhysicsSettings
 	UPROPERTY(config, EditAnywhere, Category = Constants)
 	float DefaultGravityZ;
 
+	/**Default Terminal Velocity. */
+	UPROPERTY(config, EditAnywhere, Category = Constants)
+	float DefaultTerminalVelocity;
+
 	/**Enables the use of an async scene */
 	UPROPERTY(config, EditAnywhere, AdvancedDisplay, Category=Simulation)
 	bool bEnableAsyncScene;
diff --git a/Engine/Source/Runtime/Engine/Private/PhysicsEngine/PhysicsSettings.cpp b/Engine/Source/Runtime/Engine/Private/PhysicsEngine/PhysicsSettings.cpp
--- a/Engine/Source/Runtime/Engine/Private/PhysicsEngine/PhysicsSettings.cpp
+++ b/Engine/Source/Runtime/Engine/Private/PhysicsEngine/PhysicsSettings.cpp
@@ -5,6 +5,7 @@
 UPhysicsSettings::UPhysicsSettings(const class FPostConstructInitializeProperties& PCIP)
 : Super(PCIP),
 DefaultGravityZ(-980.f),
+DefaultTerminalVelocity(4000.f),
 bEnableAsyncScene(false),
 MaxPhysicsDeltaTime(1.f / 30.f),
 bSubstepping(false),
diff --git a/Engine/Source/Runtime/Engine/Private/PhysicsVolume.cpp b/Engine/Source/Runtime/Engine/Private/PhysicsVolume.cpp
--- a/Engine/Source/Runtime/Engine/Private/PhysicsVolume.cpp
+++ b/Engine/Source/Runtime/Engine/Private/PhysicsVolume.cpp
@@ -11,7 +11,7 @@ APhysicsVolume::APhysicsVolume(const class FPostConstructInitializeProperties& P
 	BrushComponent->SetCollisionProfileName(CollisionProfileName);
 
 	FluidFriction = 0.3f;
-	TerminalVelocity = 4000.0f;
+	TerminalVelocity = UPhysicsSettings::Get()->DefaultTerminalVelocity;
 	bAlwaysRelevant = true;
 	NetUpdateFrequency = 0.1f;
 	bReplicateMovement = false;
diff --git a/Engine/Source/Runtime/Engine/Private/World.cpp b/Engine/Source/Runtime/Engine/Private/World.cpp
--- a/Engine/Source/Runtime/Engine/Private/World.cpp
+++ b/Engine/Source/Runtime/Engine/Private/World.cpp
@@ -770,6 +770,9 @@ void UWorld::InitWorld(const InitializationValues IVS)
 		DefaultPhysicsVolume->Priority = -1000000;
 	}
 
+	// update terminal velocity when world is restarted
+	DefaultPhysicsVolume->TerminalVelocity = UPhysicsSettings::Get()->DefaultTerminalVelocity;
+
 	// Find gravity
 	if (GetPhysicsScene())
 	{
