Fix crash when undo/redo/undo adding a mesh to the level.

diff --git a/Engine/Source/Runtime/Engine/Private/Level.cpp b/Engine/Source/Runtime/Engine/Private/Level.cpp
--- a/Engine/Source/Runtime/Engine/Private/Level.cpp
+++ b/Engine/Source/Runtime/Engine/Private/Level.cpp
@@ -598,6 +598,10 @@ void ULevel::PostLoad()
 	}
 #endif //WITH_EDITOR
 
+#if WITH_EDITOR
+	Actors.Remove(nullptr);
+#endif
+
 	// in the Editor, sort Actor list immediately (at runtime we wait for the level to be added to the world so that it can be delayed in the level streaming case)
 	if (GIsEditor)
 	{
@@ -1131,7 +1135,7 @@ void ULevel::PostEditUndo()
 
 	// Non-transactional actors may disappear from the actors list but still exist, so we need to re-add them
 	// Likewise they won't get recreated if we undo to before they were deleted, so we'll have nulls in the actors list to remove
-	Actors.Remove(nullptr);
+	//Actors.Remove(nullptr); // removed because TTransArray exploded (undo followed by redo ends up with a different ArrayNum to originally)
 	TSet<AActor*> ActorsSet(Actors);
 	TArray<UObject *> InnerObjects;
 	GetObjectsWithOuter(this, InnerObjects, /*bIncludeNestedObjects*/ false, /*ExclusionFlags*/ RF_PendingKill);
