Fix NULL entries ending up in the AttachChildren when working with ChildActorComponent
TTP #337464

diff --git a/Engine/Source/Runtime/Engine/Private/ActorConstruction.cpp b/Engine/Source/Runtime/Engine/Private/ActorConstruction.cpp
--- a/Engine/Source/Runtime/Engine/Private/ActorConstruction.cpp
+++ b/Engine/Source/Runtime/Engine/Private/ActorConstruction.cpp
@@ -271,7 +271,7 @@ void AActor::RerunConstructionScripts()
 		for(FAttachedActorInfo& Info : AttachedActorInfos)
 		{
 			// If this actor is no longer attached to anything, reattach
-			if (Info.AttachedActor->GetAttachParentActor() == NULL)
+			if (!Info.AttachedActor->IsPendingKill() && Info.AttachedActor->GetAttachParentActor() == NULL)
 			{
 				USceneComponent* ChildRoot = Info.AttachedActor->GetRootComponent();
 				if (ChildRoot && ChildRoot->AttachParent != RootComponent)
