Ensure that each destroyed ChildActor from a ChildActorComponent has a unique name even when class is changing
TTP #337298

diff --git a/Engine/Source/Runtime/Engine/Private/ChildActorComponent.cpp b/Engine/Source/Runtime/Engine/Private/ChildActorComponent.cpp
--- a/Engine/Source/Runtime/Engine/Private/ChildActorComponent.cpp
+++ b/Engine/Source/Runtime/Engine/Private/ChildActorComponent.cpp
@@ -182,7 +182,9 @@ void UChildActorComponent::DestroyChildActor()
 			// World may be NULL during shutdown
 			if(World != NULL)
 			{
-				ChildActor->Rename(*MakeUniqueObjectName(GetOuter(), ChildActor->GetClass(), TEXT("DESTROYED_CHILDACTOR")).ToString(), NULL, REN_DoNotDirty);
+				UClass* ChildActorClass = ChildActor->GetClass();
+				const FString ObjectBaseName = FString::Printf(TEXT("DESTROYED_%s_CHILDACTOR"), *ChildActorClass->GetName());
+				ChildActor->Rename(*MakeUniqueObjectName(GetOuter(), ChildActorClass, *ObjectBaseName).ToString(), NULL, REN_DoNotDirty);
 				World->DestroyActor(ChildActor);
 			}
 		}
