#ue4 - Better comments for ComputeGroundMovementDelta() and MoveAlongFloor(), and better variable naming within MoveAlongFloor().

diff --git a/Engine/Source/Runtime/Engine/Classes/GameFramework/CharacterMovementComponent.h b/Engine/Source/Runtime/Engine/Classes/GameFramework/CharacterMovementComponent.h
--- a/Engine/Source/Runtime/Engine/Classes/GameFramework/CharacterMovementComponent.h
+++ b/Engine/Source/Runtime/Engine/Classes/GameFramework/CharacterMovementComponent.h
@@ -1151,7 +1151,7 @@ protected:
 	 * @param RampHit:				Hit result of sweep that found the ramp below the capsule
 	 * @param bHitFromLineTrace:	Whether the floor trace came from a line trace
 	 *
-	 * @return If on a walkable surface, this returns a vector that moves parallel to the surface while maintaining horizontal velocity.
+	 * @return If on a walkable surface, this returns a vector that moves parallel to the surface. The magnitude may be scaled if bMaintainHorizontalGroundVelocity is true.
 	 * If a ramp vector can't be computed, this will just return Delta.
 	 */
 	virtual FVector ComputeGroundMovementDelta(const FVector& Delta, const FHitResult& RampHit, const bool bHitFromLineTrace) const;
diff --git a/Engine/Source/Runtime/Engine/Private/CharacterMovementComponent.cpp b/Engine/Source/Runtime/Engine/Private/CharacterMovementComponent.cpp
--- a/Engine/Source/Runtime/Engine/Private/CharacterMovementComponent.cpp
+++ b/Engine/Source/Runtime/Engine/Private/CharacterMovementComponent.cpp
@@ -2870,6 +2870,7 @@ FVector UCharacterMovementComponent::ComputeGroundMovementDelta(const FVector& D
 
 	if (FloorNormal.Z < (1.f - KINDA_SMALL_NUMBER) && FloorNormal.Z > KINDA_SMALL_NUMBER && ContactNormal.Z > KINDA_SMALL_NUMBER && !bHitFromLineTrace && IsWalkable(RampHit))
 	{
+		// Compute a vector that moves parallel to the surface, by projecting the horizontal movement direction onto the ramp.
 		const float FloorDotDelta = (FloorNormal | Delta);
 		FVector RampMovement(Delta.X, Delta.Y, -FloorDotDelta / FloorNormal.Z);
 		
@@ -2896,6 +2897,7 @@ void UCharacterMovementComponent::MoveAlongFloor(const FVector& InVelocity, cons
 		return;
 	}
 
+	// Move along the current floor
 	FHitResult Hit(1.f);
 	FVector RampVector = ComputeGroundMovementDelta(Delta, CurrentFloor.HitResult, CurrentFloor.bLineTrace);
 	SafeMoveUpdatedComponent(RampVector, CharacterOwner->GetActorRotation(), true, Hit);
@@ -2909,15 +2911,16 @@ void UCharacterMovementComponent::MoveAlongFloor(const FVector& InVelocity, cons
 
 	if (Hit.IsValidBlockingHit())
 	{
-		// See if we impacted something (most likely another ramp, but possibly a barrier). Try to slide along it as well.
+		// We impacted something (most likely another ramp, but possibly a barrier).
 		float TimeApplied = Hit.Time;
 		if ((Hit.Time > 0.f) && (Hit.Normal.Z > KINDA_SMALL_NUMBER) && IsWalkable(Hit))
 		{
-			const float PreSlideTimeRemaining = 1.f - Hit.Time;
-			RampVector = ComputeGroundMovementDelta(Delta * PreSlideTimeRemaining, Hit, false);
+			// Another walkable ramp.
+			const float InitialTimeRemaining = 1.f - TimeApplied;
+			RampVector = ComputeGroundMovementDelta(Delta * InitialTimeRemaining, Hit, false);
 			SafeMoveUpdatedComponent(RampVector, CharacterOwner->GetActorRotation(), true, Hit);
 
-			const float SecondHitPercent = Hit.Time * (1.f - TimeApplied);
+			const float SecondHitPercent = Hit.Time * InitialTimeRemaining;
 			TimeApplied = FMath::Clamp(TimeApplied + SecondHitPercent, 0.f, 1.f);
 		}
 
