Improving const correctness in the timer manager.

#github

This was a pull request submitted by Zoubi (https://github.com/EpicGames/UnrealEngine/pull/159/commits).

diff --git a/Engine/Source/Runtime/Engine/Public/TimerManager.h b/Engine/Source/Runtime/Engine/Public/TimerManager.h
--- a/Engine/Source/Runtime/Engine/Public/TimerManager.h
+++ b/Engine/Source/Runtime/Engine/Public/TimerManager.h
@@ -423,33 +423,33 @@ public:
 	 * @return				The current rate or -1.f if timer does not exist.
 	 */
 	template< class UserClass >	
-	FORCEINLINE float GetTimerRate(UserClass* inObj, typename FTimerDelegate::TUObjectMethodDelegate< UserClass >::FMethodPtr inTimerMethod)
+	FORCEINLINE float GetTimerRate(UserClass* inObj, typename FTimerDelegate::TUObjectMethodDelegate< UserClass >::FMethodPtr inTimerMethod) const
 	{
 		return GetTimerRate( FTimerDelegate::CreateUObject(inObj, inTimerMethod) );
 	}
 
 	template< class UserClass >	
-	FORCEINLINE float GetTimerRate(UserClass* inObj, typename FTimerDelegate::TUObjectMethodDelegate_Const< UserClass >::FMethodPtr inTimerMethod)
+	FORCEINLINE float GetTimerRate(UserClass* inObj, typename FTimerDelegate::TUObjectMethodDelegate_Const< UserClass >::FMethodPtr inTimerMethod) const
 	{
 		return GetTimerRate( FTimerDelegate::CreateUObject(inObj, inTimerMethod) );
 	}
 
 	/** Version that takes any generic delegate. */
-	FORCEINLINE float GetTimerRate(FTimerDelegate const& InDelegate)
+	FORCEINLINE float GetTimerRate(FTimerDelegate const& InDelegate) const
 	{
 		FTimerData const* const TimerData = FindTimer( FTimerUnifiedDelegate(InDelegate) );
 		return InternalGetTimerRate(TimerData);
 	}
 
 	/** Version that takes a dynamic delegate (e.g. for UFunctions). */
-	FORCEINLINE float GetTimerRate(FTimerDynamicDelegate const& InDynDelegate)
+	FORCEINLINE float GetTimerRate(FTimerDynamicDelegate const& InDynDelegate) const
 	{
 		FTimerData const* const TimerData = FindTimer( FTimerUnifiedDelegate(InDynDelegate) );
 		return InternalGetTimerRate(TimerData);
 	}
 
 	/** Version that takes a handle */
-	FORCEINLINE float GetTimerRate(FTimerHandle const& InHandle)
+	FORCEINLINE float GetTimerRate(FTimerHandle const& InHandle) const
 	{
 		FTimerData const* const TimerData = FindTimer(InHandle);
 		return InternalGetTimerRate(TimerData);
@@ -463,33 +463,33 @@ public:
 	 * @return				true if the timer matching the given criteria exists and is active, false otherwise.
 	 */
 	template< class UserClass >	
-	FORCEINLINE bool IsTimerActive(UserClass* inObj, typename FTimerDelegate::TUObjectMethodDelegate< UserClass >::FMethodPtr inTimerMethod)
+	FORCEINLINE bool IsTimerActive(UserClass* inObj, typename FTimerDelegate::TUObjectMethodDelegate< UserClass >::FMethodPtr inTimerMethod) const
 	{
 		return IsTimerActive( FTimerDelegate::CreateUObject(inObj, inTimerMethod) );
 	}
 		
 	template< class UserClass >	
-	FORCEINLINE bool IsTimerActive(UserClass* inObj, typename FTimerDelegate::TUObjectMethodDelegate_Const< UserClass >::FMethodPtr inTimerMethod)
+	FORCEINLINE bool IsTimerActive(UserClass* inObj, typename FTimerDelegate::TUObjectMethodDelegate_Const< UserClass >::FMethodPtr inTimerMethod) const
 	{
 		return IsTimerActive( FTimerDelegate::CreateUObject(inObj, inTimerMethod) );
 	}
 
 	/** Version that takes any generic delegate. */
-	FORCEINLINE bool IsTimerActive(FTimerDelegate const& InDelegate)
+	FORCEINLINE bool IsTimerActive(FTimerDelegate const& InDelegate) const
 	{
 		FTimerData const* const TimerData = FindTimer( FTimerUnifiedDelegate(InDelegate) );
 		return ( (TimerData != NULL) && (TimerData->Status != ETimerStatus::Paused) );
 	}
 
 	/** Version that takes a dynamic delegate (e.g. for UFunctions). */
-	FORCEINLINE bool IsTimerActive(FTimerDynamicDelegate const& InDynDelegate)
+	FORCEINLINE bool IsTimerActive(FTimerDynamicDelegate const& InDynDelegate) const
 	{
 		FTimerData const* const TimerData = FindTimer( FTimerUnifiedDelegate(InDynDelegate) );
 		return ( (TimerData != NULL) && (TimerData->Status != ETimerStatus::Paused) );
 	}
 
 	/** Version that takes a handle */
-	FORCEINLINE bool IsTimerActive(FTimerHandle const& InHandle)
+	FORCEINLINE bool IsTimerActive(FTimerHandle const& InHandle) const
 	{
 		FTimerData const* const TimerData = FindTimer( InHandle );
 		return ( (TimerData != NULL) && (TimerData->Status != ETimerStatus::Paused) );
@@ -621,30 +621,30 @@ public:
 	 * @return				The current time elapsed or -1.f if timer does not exist.
 	 */
 	template< class UserClass >	
-	FORCEINLINE float GetTimerElapsed(UserClass* inObj, typename FTimerDelegate::TUObjectMethodDelegate< UserClass >::FMethodPtr inTimerMethod)
+	FORCEINLINE float GetTimerElapsed(UserClass* inObj, typename FTimerDelegate::TUObjectMethodDelegate< UserClass >::FMethodPtr inTimerMethod) const
 	{
 		return GetTimerElapsed( FTimerDelegate::CreateUObject(inObj, inTimerMethod) );
 	}
 	template< class UserClass >	
-	FORCEINLINE float GetTimerElapsed(UserClass* inObj, typename FTimerDelegate::TUObjectMethodDelegate_Const< UserClass >::FMethodPtr inTimerMethod)
+	FORCEINLINE float GetTimerElapsed(UserClass* inObj, typename FTimerDelegate::TUObjectMethodDelegate_Const< UserClass >::FMethodPtr inTimerMethod) const
 	{
 		return GetTimerElapsed( FTimerDelegate::CreateUObject(inObj, inTimerMethod) );
 	}
 	/** Version that takes any generic delegate. */
-	FORCEINLINE float GetTimerElapsed(FTimerDelegate const& InDelegate)
+	FORCEINLINE float GetTimerElapsed(FTimerDelegate const& InDelegate) const
 	{
 		FTimerData const* const TimerData = FindTimer( FTimerUnifiedDelegate(InDelegate) );
 		return InternalGetTimerElapsed(TimerData);
 	}
 	/** Version that takes a dynamic delegate (e.g. for UFunctions). */
-	FORCEINLINE float GetTimerElapsed(FTimerDynamicDelegate const& InDynDelegate)
+	FORCEINLINE float GetTimerElapsed(FTimerDynamicDelegate const& InDynDelegate) const
 	{
 		FTimerData const* const TimerData = FindTimer( FTimerUnifiedDelegate(InDynDelegate) );
 		return InternalGetTimerElapsed(TimerData);
 	}
 
 	/** Version that takes a handle */
-	FORCEINLINE float GetTimerElapsed(FTimerHandle const& InHandle)
+	FORCEINLINE float GetTimerElapsed(FTimerHandle const& InHandle) const
 	{
 		FTimerData const* const TimerData = FindTimer(InHandle);
 		return InternalGetTimerElapsed(TimerData);
@@ -658,29 +658,29 @@ public:
 	 * @return				The current time remaining, or -1.f if timer does not exist
 	 */
 	template< class UserClass >	
-	FORCEINLINE float GetTimerRemaining(UserClass* inObj, typename FTimerDelegate::TUObjectMethodDelegate< UserClass >::FMethodPtr inTimerMethod)
+	FORCEINLINE float GetTimerRemaining(UserClass* inObj, typename FTimerDelegate::TUObjectMethodDelegate< UserClass >::FMethodPtr inTimerMethod) const
 	{
 		return GetTimerRemaining( FTimerDelegate::CreateUObject(inObj, inTimerMethod) );
 	}
 	template< class UserClass >	
-	FORCEINLINE float GetTimerRemaining(UserClass* inObj, typename FTimerDelegate::TUObjectMethodDelegate_Const< UserClass >::FMethodPtr inTimerMethod)
+	FORCEINLINE float GetTimerRemaining(UserClass* inObj, typename FTimerDelegate::TUObjectMethodDelegate_Const< UserClass >::FMethodPtr inTimerMethod) const
 	{
 		return GetTimerRemaining( FTimerDelegate::CreateUObject(inObj, inTimerMethod) );
 	}
 	/** Version that takes any generic delegate. */
-	FORCEINLINE float GetTimerRemaining(FTimerDelegate const& InDelegate)
+	FORCEINLINE float GetTimerRemaining(FTimerDelegate const& InDelegate) const
 	{
 		FTimerData const* const TimerData = FindTimer( FTimerUnifiedDelegate(InDelegate) );
 		return InternalGetTimerRemaining(TimerData);
 	}
 	/** Version that takes a dynamic delegate (e.g. for UFunctions). */
-	FORCEINLINE float GetTimerRemaining(FTimerDynamicDelegate const& InDynDelegate)
+	FORCEINLINE float GetTimerRemaining(FTimerDynamicDelegate const& InDynDelegate) const
 	{
 		FTimerData const* const TimerData = FindTimer( FTimerUnifiedDelegate(InDynDelegate) );
 		return InternalGetTimerRemaining(TimerData);
 	}
 	/** Version that takes a handle */
-	FORCEINLINE float GetTimerRemaining(FTimerHandle const& InHandle)
+	FORCEINLINE float GetTimerRemaining(FTimerHandle const& InHandle) const
 	{
 		FTimerData const* const TimerData = FindTimer(InHandle);
 		return InternalGetTimerRemaining(TimerData);
