Sky light captures are deferred until after shader compiling finishes

diff --git a/Engine/Source/Runtime/Engine/Private/SkyLightComponent.cpp b/Engine/Source/Runtime/Engine/Private/SkyLightComponent.cpp
--- a/Engine/Source/Runtime/Engine/Private/SkyLightComponent.cpp
+++ b/Engine/Source/Runtime/Engine/Private/SkyLightComponent.cpp
@@ -339,7 +339,9 @@ void USkyLightComponent::ApplyComponentInstanceData(TSharedPtr<FComponentInstanc
 
 void USkyLightComponent::UpdateSkyCaptureContents(UWorld* WorldToUpdate)
 {
-	if (WorldToUpdate->Scene)
+	if (WorldToUpdate->Scene
+		// Don't process any sky capture requests until async shader compiling completes
+		&& (GShaderCompilingManager == NULL || !GShaderCompilingManager->IsCompiling()))
 	{
 		// Iterate backwards so we can remove elements without changing the index
 		for (int32 CaptureIndex = SkyCapturesToUpdate.Num() - 1; CaptureIndex >= 0; CaptureIndex--)
