#ue4 - Perf: avoid retrieving HitResult.GetActor() when we are going to ignore the overlap in UpdateOverlaps() anyway.

diff --git a/Engine/Source/Runtime/Engine/Private/PrimitiveComponent.cpp b/Engine/Source/Runtime/Engine/Private/PrimitiveComponent.cpp
--- a/Engine/Source/Runtime/Engine/Private/PrimitiveComponent.cpp
+++ b/Engine/Source/Runtime/Engine/Private/PrimitiveComponent.cpp
@@ -2104,10 +2104,10 @@ void UPrimitiveComponent::UpdateOverlaps(TArray<FOverlapInfo> const* PendingOver
 						const FOverlapResult& Result = Overlaps[ResultIdx];
 
 						UPrimitiveComponent const* const HitComp = Result.Component.Get();
-						if (!Result.bBlockingHit && HitComp  && HitComp->bGenerateOverlapEvents)
+						if (!Result.bBlockingHit && HitComp && (HitComp != this) && HitComp->bGenerateOverlapEvents)
 						{
 							AActor* const ResultActor = Result.GetActor();
-							if ( (ResultActor != NULL) && (HitComp != this) && !ResultActor->IsAttachedTo(MyActor) && !MyActor->IsAttachedTo(ResultActor) && ResultActor != MyWorld->GetWorldSettings() && ResultActor->bActorInitialized)
+							if ( (ResultActor != NULL) && !ResultActor->IsAttachedTo(MyActor) && !MyActor->IsAttachedTo(ResultActor) && ResultActor != MyWorld->GetWorldSettings() && ResultActor->bActorInitialized)
 							{
 								NewOverlappingComponents.Add(FOverlapInfo(Result.Component.Get(), Result.ItemIndex));		// don't need to add unique unless the overlap check can return dupes
 							}						
