Add support for renamed UPackage's, cleanup

diff --git a/Engine/Source/Runtime/CoreUObject/Public/UObject/CoreNet.h b/Engine/Source/Runtime/CoreUObject/Public/UObject/CoreNet.h
--- a/Engine/Source/Runtime/CoreUObject/Public/UObject/CoreNet.h
+++ b/Engine/Source/Runtime/CoreUObject/Public/UObject/CoreNet.h
@@ -131,7 +131,7 @@ class COREUOBJECT_API UPackageMap : public UObject
 
 	// Variables.
 
-	virtual UObject * ResolvePathAndAssignNetGUID( FNetworkGUID & InOutNetGUID, const FString & PathName, UObject * ObjOuter, const bool bNoLoad = false ) { return NULL; }
+	virtual UObject * ResolvePathAndAssignNetGUID( FNetworkGUID & InOutNetGUID, const FString & PathName, const FString & FilenameOverride, UObject * ObjOuter, const bool bNoLoad = false ) { return NULL; }
 
 	void ResetUnAckedObject() { bSerializedUnAckedObject = false; }
 	bool SerializedUnAckedObject() { return bSerializedUnAckedObject; }
diff --git a/Engine/Source/Runtime/Engine/Classes/Engine/PackageMapClient.h b/Engine/Source/Runtime/Engine/Classes/Engine/PackageMapClient.h
--- a/Engine/Source/Runtime/Engine/Classes/Engine/PackageMapClient.h
+++ b/Engine/Source/Runtime/Engine/Classes/Engine/PackageMapClient.h
@@ -28,6 +28,7 @@ public:
 	FDeferredResolvePath() {}
 
 	FName			PathName;
+	FString			FilenameOverride;
 	FNetworkGUID	OuterGUID;
 };
 
@@ -132,17 +133,13 @@ protected:
 	bool	ExportNetGUID( FNetworkGUID NetGUID, const UObject * Object, FString PathName, UObject * ObjOuter );
 	void	ExportNetGUIDHeader();
 
-	bool	InternalIsMapped(UObject* Obj, FNetworkGUID &NetGUID);
-
-	void	InternalWriteObject( FArchive& Ar, FNetworkGUID NetGUID, const UObject * Object, FString ObjectPathName, UObject* ObjectOuter );	
-	void	InternalWriteObjectPath( FArchive& Ar, FNetworkGUID NetGUID, const UObject * Object, FString ObjectPathName, UObject* ObjectOuter );
-
+	void			InternalWriteObject( FArchive& Ar, FNetworkGUID NetGUID, const UObject * Object, FString ObjectPathName, UObject * ObjectOuter );	
 	FNetworkGUID	InternalLoadObject( FArchive & Ar, UObject *& Object, const bool bIsOuterLevel, int InternalLoadObjectRecursionCount );
 
-	void		ResolvePathAndAssignNetGUID_Deferred( const FNetworkGUID & NetGUID, const FString & PathName, const FNetworkGUID & OuterGUID );
-	UObject *	ResolvePathAndAssignNetGUID( FNetworkGUID & InOutNetGUID, const FString & PathName, const FNetworkGUID & OuterGUID, const bool bNoLoad = false );
+	void		ResolvePathAndAssignNetGUID_Deferred( const FNetworkGUID & NetGUID, const FString & PathName, const FString & FilenameOverride, const FNetworkGUID & OuterGUID );
+	UObject *	ResolvePathAndAssignNetGUID( FNetworkGUID & InOutNetGUID, const FString & PathName, const FString & FilenameOverride, const FNetworkGUID & OuterGUID, const bool bNoLoad = false );
 
-	virtual UObject * ResolvePathAndAssignNetGUID( FNetworkGUID & InOutNetGUID, const FString & PathName, UObject * ObjOuter, const bool bNoLoad = false ) OVERRIDE;
+	virtual UObject * ResolvePathAndAssignNetGUID( FNetworkGUID & InOutNetGUID, const FString & PathName, const FString & FilenameOverride, UObject * ObjOuter, const bool bNoLoad = false ) OVERRIDE;
 
 	bool	ShouldSendFullPath(const UObject* Object, const FNetworkGUID &NetGUID);
 	
diff --git a/Engine/Source/Runtime/Engine/Private/PackageMapClient.cpp b/Engine/Source/Runtime/Engine/Private/PackageMapClient.cpp
--- a/Engine/Source/Runtime/Engine/Private/PackageMapClient.cpp
+++ b/Engine/Source/Runtime/Engine/Private/PackageMapClient.cpp
@@ -56,7 +56,6 @@ bool UPackageMapClient::SerializeObject( FArchive& Ar, UClass* Class, UObject*&
 
 		bSerializedUnAckedObject = false;
 		FNetworkGUID NetGUID = GuidCache->GetOrAssignNetGUID( Object );
-		bool Mapped = InternalIsMapped(Object, NetGUID);
 
 		// Write out NetGUID to caller if necessary
 		if (OutNetGUID)
@@ -85,7 +84,7 @@ bool UPackageMapClient::SerializeObject( FArchive& Ar, UClass* Class, UObject*&
 		}
 #endif
 
-		return Mapped;
+		return true;
 	}
 	else if (Ar.IsLoading())
 	{
@@ -363,119 +362,6 @@ bool UPackageMapClient::SerializeNewActor(FArchive& Ar, class UActorChannel *Cha
 	return SpawnedActor;
 }
 
-//--------------------------------------------------------------------
-//
-//	Assignment of NetGUIDs
-//
-//--------------------------------------------------------------------
-
-bool UPackageMapClient::InternalIsMapped(UObject* Object, FNetworkGUID &NetGUID)
-{
-	if (!Object)
-	{
-		// Null of unsupported object, leave as default NetGUID and just return mapped=true
-		return true;
-	}
-
-	if (!NetGUID.IsValid())
-	{
-		// This will happen with subobjects that will replicate but who have not been assigned NetGUIDs yet.
-		if (bShouldSerializeUnAckedObjects)
-		{
-			// We should serialize unacked objects, so just note that we did it by setting this flag to true.
-			UE_LOG(LogNetPackageMap, Log, TEXT("Object %s is not supported, so unmapped"), *Object->GetName() );
-			bSerializedUnAckedObject = true;
-		}
-		else
-		{
-			// We arent suppose to serialize unacked objects now, reset the serialized NetGUID
-			UE_LOG(LogNetPackageMap, Log, TEXT("Object %s is not supported, so unmapped. Reseting"), *Object->GetName() );
-			NetGUID.Reset();
-		}
-		return false;
-	}
-
-	// ----------------
-	// Make sure we can serialize reference, determine Mapped
-	//	-Actor is not mapped if actor channel is not ackd
-	//	-Static objects are not mapped if streaming level not loaded
-	// ----------------
-	AActor *Actor = Cast<AActor>(Object);
-	UObject *TestObj = Object;
-	if (Actor == NULL && Cast<AActor>(Object->GetOuter()))
-	{
-		Actor = Cast<AActor>(Object->GetOuter());
-		TestObj = Actor;
-	}
-
-	bool IsDynamic = GuidCache->IsDynamicObject(TestObj);
-	bool Mapped = true;
-	if (IsDynamic)
-	{
-		check(Actor);
-		UActorChannel* Ch = Connection->ActorChannels.FindRef(Actor);
-		if (Actor->bTearOff)
-		{
-			// Assume tear off actors are mapped, but this may not be 100% accurate.
-			// We may need to track 'actors that were succesfully torn off per client' on the server.
-			Mapped = true;
-		}
-		else if ( Ch )
-		{
-			// Actor has open channel - mapped if Channel is OpenAcked
-			Mapped = Ch->OpenAcked;
-		}
-		else if (Connection->DormantActors.Contains(Actor) || Connection->RecentlyDormantActors.Contains(Actor))
-		{
-			// Actor is dormant - must be mapped
-			Mapped = true;
-		}
-		else
-		{
-			// Actor has no channel and not dormant - its not mapped
-			Mapped = false;
-		}
-			
-		// If this is an unmapped object 
-		if (!Mapped)
-		{
-			if (bShouldSerializeUnAckedObjects)
-			{
-				// We should serialize unacked objects, so just note that we did it by setting this flag to true.
-				bSerializedUnAckedObject = true;
-			}
-			else
-			{
-				// We arent suppose to serialize unacked objects now, reset the serialized NetGUID
-				NetGUID.Reset();
-			}
-		}
-
-		// If this is an actor component whose channel has not been fully ack'd
-		if (TestObj != Object && !Mapped)
-		{
-			if (NetGUID.IsStatic() && !NetGUIDHasBeenAckd(NetGUID))
-			{
-				// We cannot serialize a reference to this until the channel has been fully ack'd.
-				//	We can't allow a reference like this to go out since we can't garuntee the exported
-				//	NetGUID bunch containing it will be handled after the Actor Bunch that contains the parent actor.
-				//	Everything in the NetGUID bunch must, 100%, must be resolved on client side, otherwise the NetGUID'
-				//	ACKing system would not work. PackageMapClient assumes that if the NetGUID bunch packet ID is ACK'd, all
-				//	NetGUIDs inside of it are ACK'd.
-				NetGUID.Reset();
-			}
-		}
-	}
-	else if(Connection->Driver->IsServer() && !Connection->ClientHasInitializedLevelFor(Object))
-	{
-		// Static actor the client has not initialized its level yet, always serialize these as NULL
-		NetGUID.Reset();
-		Mapped = false;
-	}
-
-	return Mapped;
-}
-
 //--------------------------------------------------------------------
 //
 //	Writing
@@ -488,9 +374,10 @@ struct FExportFlags
 	{
 		struct
 		{
-			uint8 bHasPath		: 1;
-			uint8 bIsLevel		: 1;
-			uint8 bIsPackage	: 1;
+			uint8 bHasPath			: 1;
+			uint8 bIsLevel			: 1;
+			uint8 bIsPackage		: 1;
+			uint8 bIsRenamedPackage	: 1;
 		};
 
 		uint8	Value;
@@ -533,59 +420,59 @@ void UPackageMapClient::InternalWriteObject( FArchive & Ar, FNetworkGUID NetGUID
 			ExportFlags.bHasPath = ObjectPathName.IsEmpty() ? 0 : 1;
 		}
 
-		ExportFlags.bIsLevel	= Cast< const ULevel >( Object ) != NULL ? 1 : 0;
-		ExportFlags.bIsPackage	= Cast< const UPackage >( Object ) != NULL ? 1 : 0;
+		const UPackage * Package = Cast< const UPackage >( Object );
+
+		ExportFlags.bIsLevel	= ( Cast< const ULevel >( Object ) != NULL ) ? 1 : 0;
+		ExportFlags.bIsPackage	= ( Package != NULL ) ? 1 : 0;
+
+		ExportFlags.bIsRenamedPackage = ( Package != NULL && Package->FileName != NAME_None && Package->FileName != Package->GetFName() ) ? 1 : 0;
 
 		Ar << ExportFlags.Value;
 	}
 
 	if ( ExportFlags.bHasPath )
 	{
-		InternalWriteObjectPath( Ar, NetGUID, Object, ObjectPathName, ObjectOuter );
-	}
-	else
-	{
-		UE_CLOG(!bSuppressLogs, LogNetPackageMap, Log, TEXT("UPackageMapClient::InternalWriteObjectPath Serialized ObjectPathName %s as <%s>)"), *ObjectPathName, *NetGUID.ToString() );
-	}
-}
+		if ( Object != NULL )
+		{
+			// If the object isn't NULL, expect an empty path name, then fill it out with the actual info
+			check( ObjectOuter == NULL );
+			check( ObjectPathName.IsEmpty() );
+			ObjectPathName = Object->GetName();
+			ObjectOuter = Object->GetOuter();
+		}
+		else
+		{
+			// If we don't have an object, expect an already filled out path name
+			check( ObjectOuter != NULL );
+			check( !ObjectPathName.IsEmpty() );
+		}
 
-/** Actually writes the ObjectPath, given a FString ObjectPathName and Outer* */
-void UPackageMapClient::InternalWriteObjectPath( FArchive & Ar, FNetworkGUID NetGUID, const UObject * Object, FString ObjectPathName, UObject * ObjectOuter )
-{
-	if ( Object != NULL )
-	{
-		// If the object isn't NULL, expect an empty path name, then fill it out with the actual info
-		check( ObjectOuter == NULL );
-		check( ObjectPathName.IsEmpty() );
-		ObjectPathName = Object->GetName();
-		ObjectOuter = Object->GetOuter();
-	}
-	else
-	{
-		// If we don't have an object, expect an already filled out path name
-		check( ObjectOuter != NULL );
-		check( !ObjectPathName.IsEmpty() );
-	}
+		// Serialize reference to outer. This is basically a form of compression.
+		FNetworkGUID OuterNetGUID = GuidCache->GetOrAssignNetGUID( ObjectOuter );
 
-	// Serialize reference to outer. This is basically a form of compression.
-	FNetworkGUID OuterNetGUID = GuidCache->GetOrAssignNetGUID( ObjectOuter );
+		InternalWriteObject( Ar, OuterNetGUID, ObjectOuter, TEXT( "" ), NULL );
 
-	InternalWriteObject( Ar, OuterNetGUID, ObjectOuter, TEXT( "" ), NULL );
+		GEngine->NetworkRemapPath(Connection->Driver->GetWorld(), ObjectPathName, false);
 
-	GEngine->NetworkRemapPath(Connection->Driver->GetWorld(), ObjectPathName, false);
+		// Serialize Name of object
+		Ar << ObjectPathName;
 
-	// Serialize Name of object
-	Ar << ObjectPathName;
+		if ( ExportFlags.bIsRenamedPackage )
+		{
+			const UPackage * Package = CastChecked< const UPackage >( Object );
+			FString FilenameOverride = Package->FileName.ToString();
+			Ar << FilenameOverride;
+			//UE_LOG( LogNetPackageMap, Warning, TEXT("UPackageMapClient::InternalWriteObject: Renamed package. Name: %s, Filename: %s"), *Package->GetName(), *FilenameOverride );
+		}
 
-	if ( IsExportingNetGUIDBunch )
-	{
-		CurrentExportNetGUIDs.Add( NetGUID );
+		if ( IsExportingNetGUIDBunch )
+		{
+			CurrentExportNetGUIDs.Add( NetGUID );
 
-		int32& Count = NetGUIDExportCountMap.FindOrAdd(NetGUID);
-		Count++;
+			int32 & Count = NetGUIDExportCountMap.FindOrAdd( NetGUID );
+			Count++;
+		}
 	}
-
-	UE_CLOG(!bSuppressLogs, LogNetPackageMap, Log, TEXT("UPackageMapClient::InternalWriteObjectPath Serialized Object %s as <%s,%s>)"), *ObjectPathName, *NetGUID.ToString(), *ObjectPathName );
 }
 
 //--------------------------------------------------------------------
@@ -594,7 +481,7 @@ void UPackageMapClient::InternalWriteObjectPath( FArchive & Ar, FNetworkGUID Net
 //
 //--------------------------------------------------------------------
 
-static void SanityCheckExport( const UObject * Object, const FNetworkGUID & NetGUID, const FString & ExpectedPathName, const UObject * ExpectedOuter )
+static void SanityCheckExport( const UObject * Object, const FNetworkGUID & NetGUID, const FString & ExpectedPathName, const FString & ExpectedFilename, const UObject * ExpectedOuter, const FExportFlags & ExportFlags )
 {
 	if ( Object->GetName() != ExpectedPathName )
 	{
@@ -607,6 +494,21 @@ static void SanityCheckExport( const UObject * Object, const FNetworkGUID & NetG
 		const FString ExpectedOuterName = ExpectedOuter != NULL ? *ExpectedOuter->GetName() : TEXT( "NULL" );
 		UE_LOG( LogNetPackageMap, Warning, TEXT( "SanityCheckExport: Outer mismatch. Object: %s, NetGUID: %s, Current: %s, Expected: %s" ), *Object->GetPathName(), *NetGUID.ToString(), *CurrentOuterName, *ExpectedOuterName );
 	}
+
+	if ( ExportFlags.bIsPackage || ExportFlags.bIsRenamedPackage )
+	{
+		const UPackage * Package = Cast< const UPackage >( Object );
+
+		if ( Package == NULL )
+		{
+			UE_LOG( LogNetPackageMap, Warning, TEXT( "SanityCheckExport: Expected package. Object:%s, NetGUID: %s" ), *Object->GetPathName(), *NetGUID.ToString() );
+		}
+
+		if ( ExportFlags.bIsRenamedPackage && Package->FileName.ToString() != ExpectedFilename )
+		{
+			UE_LOG( LogNetPackageMap, Warning, TEXT( "SanityCheckExport: Renamed package filename mismatch. Filename: %s, Expected: %s" ), *Package->FileName.ToString(), *ExpectedFilename );
+		}
+	}
 }
 
 /** Loads a UObject from an FArchive stream. Reads object path if there, and tries to load object if its not already loaded */
@@ -678,6 +580,13 @@ FNetworkGUID UPackageMapClient::InternalLoadObject( FArchive & Ar, UObject *& Ob
 		FString PathName;
 		Ar << PathName;
 
+		FString FilenameOverride;
+
+		if ( ExportFlags.bIsRenamedPackage )
+		{
+			Ar << FilenameOverride;
+		}
+
 		if ( Ar.IsError() )
 		{
 			UE_LOG( LogNetPackageMap, Error, TEXT( "InternalLoadObject: Failed to load path name" ) );
@@ -690,13 +599,13 @@ FNetworkGUID UPackageMapClient::InternalLoadObject( FArchive & Ar, UObject *& Ob
 		if ( Object == NULL )
 		{
 			// Try to resolve the name (we may have already loaded it), but DON'T load yet if we haven't already
-			Object = ResolvePathAndAssignNetGUID( NetGUID, PathName, OuterGUID, true );
+			Object = ResolvePathAndAssignNetGUID( NetGUID, PathName, FilenameOverride, OuterGUID, true );
 		}
 
 		if ( Object != NULL )
 		{
 			// If we already have the object, just do some sanity checking and return
-			SanityCheckExport( Object, NetGUID, PathName, ObjOuter );
+			SanityCheckExport( Object, NetGUID, PathName, FilenameOverride, ObjOuter, ExportFlags );
 			return NetGUID;
 		}
 
@@ -705,12 +614,12 @@ FNetworkGUID UPackageMapClient::InternalLoadObject( FArchive & Ar, UObject *& Ob
 
 		if ( bCanDefer )
 		{
-			UE_LOG( LogNetPackageMap, Log, TEXT( "InternalLoadObject: Deferring resolve. Path: %s, NetGUID: %s, Outer: %s" ), *PathName, *NetGUID.ToString(), ObjOuter != NULL ? *ObjOuter->GetFullName() : TEXT( "NULL" ) );
-			ResolvePathAndAssignNetGUID_Deferred( NetGUID, PathName, OuterGUID );
+			UE_LOG( LogNetPackageMap, Warning, TEXT( "InternalLoadObject: Deferring resolve. Path: %s, NetGUID: %s, Outer: %s" ), *PathName, *NetGUID.ToString(), ObjOuter != NULL ? *ObjOuter->GetFullName() : TEXT( "NULL" ) );
+			ResolvePathAndAssignNetGUID_Deferred( NetGUID, PathName, FilenameOverride, OuterGUID );
 		}
 		else
 		{
-			Object = ResolvePathAndAssignNetGUID( NetGUID, PathName, OuterGUID );
+			Object = ResolvePathAndAssignNetGUID( NetGUID, PathName, FilenameOverride, OuterGUID );
 		}
 	}
 
@@ -722,7 +631,7 @@ FNetworkGUID UPackageMapClient::InternalLoadObject( FArchive & Ar, UObject *& Ob
 	return NetGUID;
 }
 
-void UPackageMapClient::ResolvePathAndAssignNetGUID_Deferred( const FNetworkGUID & NetGUID, const FString & PathName, const FNetworkGUID & OuterGUID )
+void UPackageMapClient::ResolvePathAndAssignNetGUID_Deferred( const FNetworkGUID & NetGUID, const FString & PathName, const FString & FilenameOverride, const FNetworkGUID & OuterGUID )
 {
 	check( !NetGUID.IsDefault() );		// We can't defer default guids, doesn't make sense
 
@@ -730,8 +639,9 @@ void UPackageMapClient::ResolvePathAndAssignNetGUID_Deferred( const FNetworkGUID
 	{
 		FDeferredResolvePath DeferredResolvePath;
 
-		DeferredResolvePath.OuterGUID	= OuterGUID;
-		DeferredResolvePath.PathName	= FName( *PathName );
+		DeferredResolvePath.OuterGUID			= OuterGUID;
+		DeferredResolvePath.PathName			= FName( *PathName );
+		DeferredResolvePath.FilenameOverride	= FilenameOverride;
 
 		DeferredResolvePaths.Add( NetGUID, DeferredResolvePath );
 	}
@@ -747,7 +657,7 @@ void UPackageMapClient::ResolvePathAndAssignNetGUID_Deferred( const FNetworkGUID
 	}
 }
 
-UObject * UPackageMapClient::ResolvePathAndAssignNetGUID( FNetworkGUID & InOutNetGUID, const FString & PathName, UObject * ObjOuter, const bool bNoLoad )
+UObject * UPackageMapClient::ResolvePathAndAssignNetGUID( FNetworkGUID & InOutNetGUID, const FString & PathName, const FString & FilenameOverride, UObject * ObjOuter, const bool bNoLoad )
 {
 	if ( InOutNetGUID.IsDynamic() )
 	{
@@ -785,17 +695,47 @@ UObject * UPackageMapClient::ResolvePathAndAssignNetGUID( FNetworkGUID & InOutNe
 			// If we failed to load it as an object, try to load it as a package
 			//
 
-			UPackage * PackageOuter = Cast< UPackage >( ObjOuter );
-
-			// If we have an outer at this point, it only makes sense for it to be a package (or NULL)
-			if ( PackageOuter != NULL || ObjOuter == NULL )
+			if ( !FilenameOverride.IsEmpty() )
 			{
-				Object = LoadPackage( PackageOuter, *PathName, LOAD_None );
-				UE_LOG( LogNetPackageMap, Log, TEXT( "UPackageMapClient::ResolvePathAndAssignNetGUID: LoadPackage: %s. Found: %s" ), *PathName, Object ? *Object->GetName() : TEXT( "NULL" ) );
-			} 
+				// Package was renamed, we need to load a package on disk into a renamed package in memory
+				if ( ObjOuter != NULL )
+				{
+					UE_LOG( LogNetPackageMap, Warning, TEXT( "UPackageMapClient::ResolvePathAndAssignNetGUID: Expected outer to be NULL. Name: %s, Override: %s, Outer: %s" ), *PathName, *FilenameOverride, *ObjOuter->GetName() );
+					return NULL;
+				}
+
+				// Create the package
+				UPackage * PackageOuter = CreatePackage( NULL, *PathName );
+
+				// Load the package off disk into this renamed package
+				Object = LoadPackage( PackageOuter, *FilenameOverride, LOAD_None );
+
+				UPackage * Package = Cast< UPackage >( Object );
+
+				if ( Package != NULL )
+				{
+					if ( Package->GetName() != PathName || Package->FileName.ToString() != FilenameOverride )
+					{
+						UE_LOG( LogNetPackageMap, Error, TEXT( "UPackageMapClient::ResolvePathAndAssignNetGUID: Renamed package mismatch. Name: %s, Filename: %s, ActualName: %s, ActualFile: %s" ), *PathName, *FilenameOverride, *Package->GetName(), *Package->FileName.ToString() );
+					}
+				}
+
+				UE_LOG( LogNetPackageMap, Log, TEXT( "UPackageMapClient::ResolvePathAndAssignNetGUID: Loaded renamed package. Name: %s. Found: %s, FileName: %s" ), *PathName, Object ? *Object->GetName() : TEXT( "NULL" ), *FilenameOverride );
+			}
 			else
 			{
-				UE_LOG( LogNetPackageMap, Warning, TEXT( "UPackageMapClient::ResolvePathAndAssignNetGUID: Outer %s is not a package for object %s" ), ObjOuter ? *ObjOuter->GetName() : TEXT( "NULL" ), *PathName );
+				UPackage * PackageOuter = Cast< UPackage >( ObjOuter );
+
+				// If we have an outer at this point, it only makes sense for it to be a package (or NULL)
+				if ( PackageOuter != NULL || ObjOuter == NULL )
+				{
+					Object = LoadPackage( PackageOuter, *PathName, LOAD_None );
+					UE_LOG( LogNetPackageMap, Log, TEXT( "UPackageMapClient::ResolvePathAndAssignNetGUID: LoadPackage: %s. Found: %s" ), *PathName, Object ? *Object->GetName() : TEXT( "NULL" ) );
+				} 
+				else
+				{
+					UE_LOG( LogNetPackageMap, Warning, TEXT( "UPackageMapClient::ResolvePathAndAssignNetGUID: Outer %s is not a package for object %s" ), ObjOuter ? *ObjOuter->GetName() : TEXT( "NULL" ), *PathName );
+				}
 			}
 		}
 	}
@@ -804,6 +744,9 @@ UObject * UPackageMapClient::ResolvePathAndAssignNetGUID( FNetworkGUID & InOutNe
 	{
 		if ( InOutNetGUID.IsDefault() )
 		{
+			// This should be from the client
+			// If we get here, we want to go ahead and assign a network guid, then export that to the client 
+			// on the next available opportunity
 			check( IsNetGUIDAuthority() );
 			InOutNetGUID = GuidCache->GetOrAssignNetGUID( Object );
 			HandleUnAssignedObject( Object );
@@ -817,17 +760,24 @@ UObject * UPackageMapClient::ResolvePathAndAssignNetGUID( FNetworkGUID & InOutNe
 	return Object;
 }
 
-UObject * UPackageMapClient::ResolvePathAndAssignNetGUID( FNetworkGUID & InOutNetGUID, const FString & PathName, const FNetworkGUID & OuterGUID, const bool bNoLoad )
+UObject * UPackageMapClient::ResolvePathAndAssignNetGUID( FNetworkGUID & InOutNetGUID, const FString & PathName, const FString & FilenameOverride, const FNetworkGUID & OuterGUID, const bool bNoLoad )
 {
-	UObject * ObjOuter = OuterGUID.IsValid() ? GetObjectFromNetGUID( OuterGUID ) : NULL;
+	UObject * ObjOuter = NULL;
 
-	if ( OuterGUID.IsValid() && ObjOuter == NULL )
+	// Resolve the outer if the guid is valid
+	if ( OuterGUID.IsValid() )
 	{
-		UE_LOG( LogNetPackageMap, Warning, TEXT( "ResolvePathAndAssignNetGUID: Unable resolve outer. Path: %s, NetGUID: %s, OuterGUID: %s " ), *PathName, *InOutNetGUID.ToString(), *OuterGUID.ToString() );
-		return NULL;
+		ObjOuter = GetObjectFromNetGUID( OuterGUID );
+
+		// If we have a valid outer guid, we expect it to resolve
+		if ( ObjOuter == NULL )
+		{
+			UE_LOG( LogNetPackageMap, Warning, TEXT( "ResolvePathAndAssignNetGUID: Unable resolve outer. Path: %s, NetGUID: %s, OuterGUID: %s " ), *PathName, *InOutNetGUID.ToString(), *OuterGUID.ToString() );
+			return NULL;
+		}
 	}
 
-	UObject * Object = ResolvePathAndAssignNetGUID( InOutNetGUID, PathName, ObjOuter, bNoLoad );
+	UObject * Object = ResolvePathAndAssignNetGUID( InOutNetGUID, PathName, FilenameOverride, ObjOuter, bNoLoad );
 
 	if ( Object == NULL && !bNoLoad )
 	{
diff --git a/Engine/Source/Runtime/Engine/Private/PendingNetGame.cpp b/Engine/Source/Runtime/Engine/Private/PendingNetGame.cpp
--- a/Engine/Source/Runtime/Engine/Private/PendingNetGame.cpp
+++ b/Engine/Source/Runtime/Engine/Private/PendingNetGame.cpp
@@ -325,7 +325,7 @@ void UPendingNetGame::NotifyControlMessage(UNetConnection* Connection, uint8 Mes
 			FNetworkGUID NetGUID;
 			FString Path;
 			FNetControlMessage<NMT_NetGUIDAssign>::Receive(Bunch, NetGUID, Path);
-			NetDriver->ServerConnection->PackageMap->ResolvePathAndAssignNetGUID(NetGUID, Path, NULL);
+			NetDriver->ServerConnection->PackageMap->ResolvePathAndAssignNetGUID(NetGUID, Path, TEXT(""), NULL);
 			break;
 		}
 		default:
diff --git a/Engine/Source/Runtime/Engine/Private/World.cpp b/Engine/Source/Runtime/Engine/Private/World.cpp
--- a/Engine/Source/Runtime/Engine/Private/World.cpp
+++ b/Engine/Source/Runtime/Engine/Private/World.cpp
@@ -3388,7 +3388,7 @@ void UWorld::NotifyControlMessage(UNetConnection* Connection, uint8 MessageType,
 				FNetControlMessage<NMT_NetGUIDAssign>::Receive(Bunch, NetGUID, Path);
 
 				UE_LOG(LogNet, Verbose, TEXT("NMT_NetGUIDAssign  NetGUID %s. Path: %s. "), *NetGUID.ToString(), *Path );
-				Connection->PackageMap->ResolvePathAndAssignNetGUID(NetGUID, Path, NULL);
+				Connection->PackageMap->ResolvePathAndAssignNetGUID(NetGUID, Path, TEXT(""), NULL);
 				break;
 			}
 		}
