#ue4 - Ensure UMovementComponent::ResolvePenetration() maintains the plane constraint in the case of multiple overlaps.

diff --git a/Engine/Source/Runtime/Engine/Private/MovementComponent.cpp b/Engine/Source/Runtime/Engine/Private/MovementComponent.cpp
--- a/Engine/Source/Runtime/Engine/Private/MovementComponent.cpp
+++ b/Engine/Source/Runtime/Engine/Private/MovementComponent.cpp
@@ -344,7 +344,7 @@ bool UMovementComponent::ResolvePenetration(const FVector& ProposedAdjustment, c
 			if (!bMoved)
 			{
 				// Try moving the proposed adjustment plus the attempted move direction. This can sometimes get out of penetrations with multiple objects
-				const FVector MoveDelta = (Hit.TraceEnd - Hit.TraceStart);
+				const FVector MoveDelta = ConstrainDirectionToPlane(Hit.TraceEnd - Hit.TraceStart);
 				if (!MoveDelta.IsZero())
 				{
 					bMoved = MoveUpdatedComponent(Adjustment + MoveDelta, NewRotation, true);
