#ue4 Adding function to parse an FString into a character array

diff --git a/Engine/Source/Runtime/Engine/Classes/Kismet/KismetStringLibrary.h b/Engine/Source/Runtime/Engine/Classes/Kismet/KismetStringLibrary.h
--- a/Engine/Source/Runtime/Engine/Classes/Kismet/KismetStringLibrary.h
+++ b/Engine/Source/Runtime/Engine/Classes/Kismet/KismetStringLibrary.h
@@ -262,4 +262,12 @@ class UKismetStringLibrary : public UBlueprintFunctionLibrary
 	 */
 	UFUNCTION(BlueprintPure, meta=(Delimiter=" "), Category="Utilities|String")
 	static TArray<FString> ParseIntoArray(const FString& SourceString, const FString& Delimiter, const bool CullEmptyStrings = true);
+
+	/**
+	 * Returns an array that contains one entry for each character in SourceString
+	 * @param	SourceString	The string to break apart into characters
+	 * @return	An array containing one entry for each character in SourceString
+	 */
+	UFUNCTION(BlueprintPure, Category="Utilities|String")
+	static TArray<FString> GetCharacterArrayFromString(const FString& SourceString);
 };
diff --git a/Engine/Source/Runtime/Engine/Private/KismetStringLibrary.cpp b/Engine/Source/Runtime/Engine/Private/KismetStringLibrary.cpp
--- a/Engine/Source/Runtime/Engine/Private/KismetStringLibrary.cpp
+++ b/Engine/Source/Runtime/Engine/Private/KismetStringLibrary.cpp
@@ -283,3 +283,19 @@ TArray<FString> UKismetStringLibrary::ParseIntoArray(const FString& SourceString
 	const int32 nArraySize = SourceString.ParseIntoArray(&SeparatedStrings, *Delimiter, CullEmptyStrings);
 	return SeparatedStrings;
 }
+
+TArray<FString> UKismetStringLibrary::GetCharacterArrayFromString(const FString& SourceString)
+{
+	TArray<FString> SeparatedChars;
+
+	for (auto CharIt(SourceString.CreateConstIterator()); CharIt; ++CharIt)
+	{
+		TCHAR Char = *CharIt;
+		SeparatedChars.Add(FString(1, &Char));
+	}
+
+	// Remove the null terminator on the end
+	SeparatedChars.RemoveAt(SeparatedChars.Num() - 1, 1);
+
+	return SeparatedChars;
+}
\ No newline at end of file
