Fixing Misaligned Material Comments

TTP#337463 - UE4: Samples: Misaligned material nodes need to be tidied up

Noticed that the position of comments was not updated by the scaled size of the comment so it was never going to be moved far enough to the left. Have fixed this for any older materials that still need to be flipped/scaled and added a version number for those inbetween, which will shift the comments left by the desired amount. If the comments have already been lined up manually, they might be forced further out of alignment but there should be less cases of this than those that are actually broken.

diff --git a/Engine/Source/Runtime/Core/Public/UObject/ObjectVersion.h b/Engine/Source/Runtime/Core/Public/UObject/ObjectVersion.h
--- a/Engine/Source/Runtime/Core/Public/UObject/ObjectVersion.h
+++ b/Engine/Source/Runtime/Core/Public/UObject/ObjectVersion.h
@@ -550,6 +550,8 @@ enum EUnrealEngineObjectUE4Version
 	VER_UE4_MOVE_LANDSCAPE_MICS_AND_TEXTURES_WITHIN_LEVEL,
 	// FTexts have creation history data, removed Key, Namespaces, and SourceString
 	VER_UE4_FTEXT_HISTORY,
+	// Shift comments to the left to contain expressions properly
+	VER_UE4_FIX_MATERIAL_COMMENTS,
 	// -----<new versions can be added before this line>-------------------------------------------------
 	// - this needs to be the last line (see note below)
 	VER_UE4_AUTOMATIC_VERSION_PLUS_ONE,
diff --git a/Engine/Source/Runtime/Engine/Classes/Materials/Material.h b/Engine/Source/Runtime/Engine/Classes/Materials/Material.h
--- a/Engine/Source/Runtime/Engine/Classes/Materials/Material.h
+++ b/Engine/Source/Runtime/Engine/Classes/Materials/Material.h
@@ -1158,6 +1158,13 @@ public:
 	 */
 	static void FlipExpressionPositions(const TArray<UMaterialExpression*>& Expressions, const TArray<UMaterialExpressionComment*>& Comments, bool bScaleCoords, UMaterial* Material = NULL);
 
+	/**
+	 * Shifts the positions of comments so that they are aligned correctly with other expressions
+	 *
+	 * @param	Comments	Array of comments to fix
+	 */
+	static void FixCommentPositions(const TArray<UMaterialExpressionComment*>& Comments);
+
 	/**
 	 * Checks whether a Material is arranged in the old style, with inputs flowing from right to left
 	 */
diff --git a/Engine/Source/Runtime/Engine/Private/Materials/Material.cpp b/Engine/Source/Runtime/Engine/Private/Materials/Material.cpp
--- a/Engine/Source/Runtime/Engine/Private/Materials/Material.cpp
+++ b/Engine/Source/Runtime/Engine/Private/Materials/Material.cpp
@@ -52,6 +52,7 @@ FUObjectAnnotationSparseBool GMaterialsThatNeedPhysicalConversion;
 FUObjectAnnotationSparse<FMaterialsWithDirtyUsageFlags,true> GMaterialsWithDirtyUsageFlags;
 FUObjectAnnotationSparseBool GMaterialsThatNeedExpressionsFlipped;
 FUObjectAnnotationSparseBool GMaterialsThatNeedCoordinateCheck;
+FUObjectAnnotationSparseBool GMaterialsThatNeedCommentFix;
 
 #endif // #if WITH_EDITOR
 
@@ -1636,6 +1637,10 @@ void UMaterial::Serialize(FArchive& Ar)
 	{
 		GMaterialsThatNeedCoordinateCheck.Set(this);
 	}
+	else if (Ar.UE4Ver() < VER_UE4_FIX_MATERIAL_COMMENTS)
+	{
+		GMaterialsThatNeedCommentFix.Set(this);
+	}
 #endif // #if WITH_EDITOR
 
 	if( Ar.UE4Ver() < VER_UE4_MATERIAL_ATTRIBUTES_REORDERING )
@@ -2022,6 +2027,12 @@ void UMaterial::PostLoad()
 		{
 			FlipExpressionPositions(Expressions, EditorComments, false, this);
 		}
+		FixCommentPositions(EditorComments);
+	}
+	else if (GMaterialsThatNeedCommentFix.Get(this))
+	{
+		GMaterialsThatNeedCommentFix.Clear(this);
+		FixCommentPositions(EditorComments);
 	}
 #endif // #if WITH_EDITOR
 }
@@ -3548,13 +3559,25 @@ void UMaterial::FlipExpressionPositions(const TArray<UMaterialExpression*>& Expr
 	for (int32 ExpressionIndex = 0; ExpressionIndex < Comments.Num(); ExpressionIndex++)
 	{
 		UMaterialExpressionComment* Comment = Comments[ExpressionIndex];
-		Comment->MaterialExpressionEditorX = -Comment->MaterialExpressionEditorX * PosScaling - Comment->SizeX;
+		Comment->MaterialExpressionEditorX = (-Comment->MaterialExpressionEditorX - Comment->SizeX) * PosScaling;
 		Comment->MaterialExpressionEditorY *= PosScaling;
 		Comment->SizeX *= PosScaling;
 		Comment->SizeY *= PosScaling;
 	}
 }
 
+void UMaterial::FixCommentPositions(const TArray<UMaterialExpressionComment*>& Comments)
+{
+	// equivalent to 1/1.25 * 0.25 to get the amount that should have been used when first flipping
+	const float SizeScaling = 0.2f;
+
+	for (int32 Index = 0; Index < Comments.Num(); Index++)
+	{
+		UMaterialExpressionComment* Comment = Comments[Index];
+		Comment->MaterialExpressionEditorX -= Comment->SizeX * SizeScaling;
+	}
+}
+
 bool UMaterial::HasFlippedCoordinates()
 {
 	uint32 ReversedInputCount = 0;
diff --git a/Engine/Source/Runtime/Engine/Private/Materials/MaterialExpressions.cpp b/Engine/Source/Runtime/Engine/Private/Materials/MaterialExpressions.cpp
--- a/Engine/Source/Runtime/Engine/Private/Materials/MaterialExpressions.cpp
+++ b/Engine/Source/Runtime/Engine/Private/Materials/MaterialExpressions.cpp
@@ -154,6 +154,7 @@ if( ExpressionInput.Expression == ToBeRemovedExpression )										\
 #if WITH_EDITOR
 FUObjectAnnotationSparseBool GMaterialFunctionsThatNeedExpressionsFlipped;
 FUObjectAnnotationSparseBool GMaterialFunctionsThatNeedCoordinateCheck;
+FUObjectAnnotationSparseBool GMaterialFunctionsThatNeedCommentFix;
 #endif // #if WITH_EDITOR
 
 /** Returns whether the given expression class is allowed. */
@@ -6151,6 +6152,10 @@ void UMaterialFunction::Serialize(FArchive& Ar)
 	{
 		GMaterialFunctionsThatNeedCoordinateCheck.Set(this);
 	}
+	else if (Ar.UE4Ver() < VER_UE4_FIX_MATERIAL_COMMENTS)
+	{
+		GMaterialFunctionsThatNeedCommentFix.Set(this);
+	}
 #endif // #if WITH_EDITOR
 }
 
@@ -6203,6 +6208,12 @@ void UMaterialFunction::PostLoad()
 		{
 			UMaterial::FlipExpressionPositions(FunctionExpressions, FunctionEditorComments, false);
 		}
+		UMaterial::FixCommentPositions(FunctionEditorComments);
+	}
+	else if (GMaterialFunctionsThatNeedCommentFix.Get(this))
+	{
+		GMaterialFunctionsThatNeedCommentFix.Clear(this);
+		UMaterial::FixCommentPositions(FunctionEditorComments);
 	}
 #endif // #if WITH_EDITOR
 }
