Fixed issue causing timed particle effect notify to repeatedly disable the same particle component if the animation looped around before the component had been destroyed
#ttp 337718

diff --git a/Engine/Source/Runtime/Engine/Private/Animation/AnimNotifyState_TimedParticleEffect.cpp b/Engine/Source/Runtime/Engine/Private/Animation/AnimNotifyState_TimedParticleEffect.cpp
--- a/Engine/Source/Runtime/Engine/Private/Animation/AnimNotifyState_TimedParticleEffect.cpp
+++ b/Engine/Source/Runtime/Engine/Private/Animation/AnimNotifyState_TimedParticleEffect.cpp
@@ -46,7 +46,7 @@ void UAnimNotifyState_TimedParticleEffect::NotifyEnd(USkeletalMeshComponent * Me
 			bTemplateMatch |= PreviousPSTemplates.Contains(ParticleComponent->Template);
 #endif
 
-			if(bSocketMatch && bTemplateMatch)
+			if(bSocketMatch && bTemplateMatch && !ParticleComponent->bWasDeactivated)
 			{
 				// Either destroy the component or deactivate it to have it's active particles finish.
 				// The component will auto destroy once all particle are gone.
