Fix ArrowComponent crash when building lighting

diff --git a/Engine/Source/Runtime/Engine/Private/ArrowComponent.cpp b/Engine/Source/Runtime/Engine/Private/ArrowComponent.cpp
--- a/Engine/Source/Runtime/Engine/Private/ArrowComponent.cpp
+++ b/Engine/Source/Runtime/Engine/Private/ArrowComponent.cpp
@@ -62,13 +62,8 @@ public:
 	/** Initialization */
 	void Init(const FArrowVertexBuffer* VertexBuffer)
 	{
-		check(!IsInRenderingThread());
-
-		ENQUEUE_UNIQUE_RENDER_COMMAND_TWOPARAMETER(
-			InitArrowVertexFactory,
-			FArrowVertexFactory*, VertexFactory, this,
-			const FArrowVertexBuffer*, VertexBuffer, VertexBuffer,
-			{
+		if (IsInRenderingThread())
+		{
 			// Initialize the vertex factory's stream components.
 			DataType NewData;
 			NewData.PositionComponent = STRUCTMEMBER_VERTEXSTREAMCOMPONENT(VertexBuffer, FDynamicMeshVertex, Position, VET_Float3);
@@ -77,8 +72,26 @@ public:
 				);
 			NewData.TangentBasisComponents[0] = STRUCTMEMBER_VERTEXSTREAMCOMPONENT(VertexBuffer, FDynamicMeshVertex, TangentX, VET_PackedNormal);
 			NewData.TangentBasisComponents[1] = STRUCTMEMBER_VERTEXSTREAMCOMPONENT(VertexBuffer, FDynamicMeshVertex, TangentZ, VET_PackedNormal);
-			VertexFactory->SetData(NewData);
-		});
+			SetData(NewData);
+		}
+		else
+		{
+			ENQUEUE_UNIQUE_RENDER_COMMAND_TWOPARAMETER(
+				InitArrowVertexFactory,
+				FArrowVertexFactory*, VertexFactory, this,
+				const FArrowVertexBuffer*, VertexBuffer, VertexBuffer,
+				{
+				// Initialize the vertex factory's stream components.
+				DataType NewData;
+				NewData.PositionComponent = STRUCTMEMBER_VERTEXSTREAMCOMPONENT(VertexBuffer, FDynamicMeshVertex, Position, VET_Float3);
+				NewData.TextureCoordinates.Add(
+					FVertexStreamComponent(VertexBuffer, STRUCT_OFFSET(FDynamicMeshVertex, TextureCoordinate), sizeof(FDynamicMeshVertex), VET_Float2)
+					);
+				NewData.TangentBasisComponents[0] = STRUCTMEMBER_VERTEXSTREAMCOMPONENT(VertexBuffer, FDynamicMeshVertex, TangentX, VET_PackedNormal);
+				NewData.TangentBasisComponents[1] = STRUCTMEMBER_VERTEXSTREAMCOMPONENT(VertexBuffer, FDynamicMeshVertex, TangentZ, VET_PackedNormal);
+				VertexFactory->SetData(NewData);
+			});
+		}
 	}
 };
 
