Maintain reference to anim notify state event data so we can correctly end the notify and clean up correctly.
#ttp 331600 - AnimNotifyStates don't call NotifyEnd when being deleted.
#ttp 331736 - EDITOR: Crash: Deleting a Notify State when the red Line is inside the Track Crashes the Editor

diff --git a/Engine/Source/Runtime/Engine/Classes/Animation/AnimInstance.h b/Engine/Source/Runtime/Engine/Classes/Animation/AnimInstance.h
--- a/Engine/Source/Runtime/Engine/Classes/Animation/AnimInstance.h
+++ b/Engine/Source/Runtime/Engine/Classes/Animation/AnimInstance.h
@@ -472,8 +472,10 @@ public:
 	/** Animation Notifies that has been triggered in the latest tick **/
 	TArray<const struct FAnimNotifyEvent *> AnimNotifies;
 
-	/** Currently Active AnimNotifyState */
-	TArray<const struct FAnimNotifyEvent *> ActiveAnimNotifyState;
+	/** Currently Active AnimNotifyState, stored as a copy of the event as we need to
+		call NotifyEnd on the event after a deletion in the editor. After this the event
+		is removed correctly. */
+	TArray<FAnimNotifyEvent> ActiveAnimNotifyState;
 
 	/** AnimMontage instances that are running currently 
 	 * - only one is primarily active, but the other ones are blending out 
diff --git a/Engine/Source/Runtime/Engine/Classes/Animation/AnimSequenceBase.h b/Engine/Source/Runtime/Engine/Classes/Animation/AnimSequenceBase.h
--- a/Engine/Source/Runtime/Engine/Classes/Animation/AnimSequenceBase.h
+++ b/Engine/Source/Runtime/Engine/Classes/Animation/AnimSequenceBase.h
@@ -111,6 +111,11 @@ struct FAnimNotifyEvent
 	{
 		return Notify != NULL || NotifyStateClass != NULL;
 	}
+
+	bool operator ==(const FAnimNotifyEvent& Other)
+	{
+		return((Notify && Notify == Other.Notify) || (NotifyStateClass && NotifyStateClass == Other.NotifyStateClass));
+	}
 };
 
 /**
diff --git a/Engine/Source/Runtime/Engine/Private/Animation/AnimInstance.cpp b/Engine/Source/Runtime/Engine/Private/Animation/AnimInstance.cpp
--- a/Engine/Source/Runtime/Engine/Private/Animation/AnimInstance.cpp
+++ b/Engine/Source/Runtime/Engine/Private/Animation/AnimInstance.cpp
@@ -687,9 +687,9 @@ void UAnimInstance::DisplayDebug(class UCanvas* Canvas, const FDebugDisplayInfo&
 		{
 			FIndenter NotifyIndent(Indent);
 
-			const FAnimNotifyEvent* NotifyState = ActiveAnimNotifyState[NotifyIndex];
+			const FAnimNotifyEvent& NotifyState = ActiveAnimNotifyState[NotifyIndex];
 
-			FString NotifyEntry = FString::Printf(TEXT("%i) %s Class: %s Dur:%.3f"), NotifyIndex, *NotifyState->NotifyName.ToString(), *NotifyState->NotifyStateClass->GetName(), NotifyState->Duration);
+			FString NotifyEntry = FString::Printf(TEXT("%i) %s Class: %s Dur:%.3f"), NotifyIndex, *NotifyState.NotifyName.ToString(), *NotifyState.NotifyStateClass->GetName(), NotifyState.Duration);
 			Canvas->DrawText(RenderFont, NotifyEntry, Indent, YPos, 1.f, 1.f, RenderInfo);
 			YPos += YL;
 		}
@@ -1058,13 +1058,10 @@ void UAnimInstance::TriggerAnimNotifies(float DeltaSeconds)
 	USkeletalMeshComponent * SkelMeshComp = GetSkelMeshComponent();
 
 	// Array that will replace the 'ActiveAnimNotifyState' at the end of this function.
-	TArray<const FAnimNotifyEvent *> NewActiveAnimNotifyState;
+	TArray<FAnimNotifyEvent> NewActiveAnimNotifyState;
 	// AnimNotifyState freshly added that need their 'NotifyBegin' event called.
 	TArray<const FAnimNotifyEvent *> NotifyStateBeginEvent;
 
-	// Remove NULL entries.
-	ActiveAnimNotifyState.RemoveSwap(NULL);
-
 	for (int32 Index=0; Index<AnimNotifies.Num(); Index++)
 	{
 		const FAnimNotifyEvent * AnimNotifyEvent = AnimNotifies[Index];
@@ -1072,12 +1069,12 @@ void UAnimInstance::TriggerAnimNotifies(float DeltaSeconds)
 		// AnimNotifyState
 		if( AnimNotifyEvent->NotifyStateClass )
 		{
-			if( !ActiveAnimNotifyState.RemoveSingleSwap(AnimNotifyEvent) )
+			if( !ActiveAnimNotifyState.RemoveSingleSwap(*AnimNotifyEvent) )
 			{
 				// Queue up calls to 'NotifyBegin', so they happen after 'NotifyEnd'.
 				NotifyStateBeginEvent.Add(AnimNotifyEvent);
 			}
-			NewActiveAnimNotifyState.Add(AnimNotifyEvent);
+			NewActiveAnimNotifyState.Add(*AnimNotifyEvent);
 			continue;
 		}
 
@@ -1124,8 +1121,8 @@ void UAnimInstance::TriggerAnimNotifies(float DeltaSeconds)
 	// Send end notification to AnimNotifyState not active anymore.
 	for(int32 Index=0; Index<ActiveAnimNotifyState.Num(); Index++)
 	{
-		const FAnimNotifyEvent * AnimNotifyEvent = ActiveAnimNotifyState[Index];
-		AnimNotifyEvent->NotifyStateClass->NotifyEnd(SkelMeshComp, Cast<UAnimSequence>(AnimNotifyEvent->NotifyStateClass->GetOuter()));
+		const FAnimNotifyEvent& AnimNotifyEvent = ActiveAnimNotifyState[Index];
+		AnimNotifyEvent.NotifyStateClass->NotifyEnd(SkelMeshComp, Cast<UAnimSequence>(AnimNotifyEvent.NotifyStateClass->GetOuter()));
 	}
 
 	// Call 'NotifyBegin' event on freshly added AnimNotifyState.
@@ -1136,13 +1133,13 @@ void UAnimInstance::TriggerAnimNotifies(float DeltaSeconds)
 	}
 
 	// Switch our arrays.
-	ActiveAnimNotifyState = NewActiveAnimNotifyState;
+	ActiveAnimNotifyState = MoveTemp(NewActiveAnimNotifyState);
 
 	// Tick currently active AnimNotifyState
 	for(int32 Index=0; Index<ActiveAnimNotifyState.Num(); Index++)
 	{
-		const FAnimNotifyEvent * AnimNotifyEvent = ActiveAnimNotifyState[Index];
-		AnimNotifyEvent->NotifyStateClass->NotifyTick(SkelMeshComp, Cast<UAnimSequence>(AnimNotifyEvent->NotifyStateClass->GetOuter()), DeltaSeconds);
+		const FAnimNotifyEvent& AnimNotifyEvent = ActiveAnimNotifyState[Index];
+		AnimNotifyEvent.NotifyStateClass->NotifyTick(SkelMeshComp, Cast<UAnimSequence>(AnimNotifyEvent.NotifyStateClass->GetOuter()), DeltaSeconds);
 	}
 }
 
