Rename SplineComponent SplineXDir to SplineUpDir
Add BP API to SplineComponent for getting num points, moving existing points, finding location/tangent at a point
Add BP API to SplineMeshComponent for setting scale/roll/offset
Remove SplineMeshComponent.h from Engine.h

diff --git a/Engine/Config/BaseEngine.ini b/Engine/Config/BaseEngine.ini
--- a/Engine/Config/BaseEngine.ini
+++ b/Engine/Config/BaseEngine.ini
@@ -1217,6 +1217,9 @@ bUseStreamingPause=false
 +ActiveClassRedirects=(OldClassName="SlateWidgetStyleAsset", NewClassName="/Script/SlateCore.SlateWidgetStyleAsset")
 +ActiveClassRedirects=(OldClassName="SlateWidgetStyleContainerBase", NewClassName="/Script/SlateCore.SlateWidgetStyleContainerBase")
 
++TaggedPropertyRedirects=(ClassName="SplineMeshComponent",OldPropertyName="SplineXDir",NewPropertyName="SplineUpDir")
+
+
 
 [CoreUObject.Metadata]
 ; Note: UnrealHeaderTool should be rerun after making changes to MetadataRedirects to catch any keys specified in class headers
diff --git a/Engine/Source/Editor/UnrealEd/Private/Fbx/FbxMainExport.cpp b/Engine/Source/Editor/UnrealEd/Private/Fbx/FbxMainExport.cpp
--- a/Engine/Source/Editor/UnrealEd/Private/Fbx/FbxMainExport.cpp
+++ b/Engine/Source/Editor/UnrealEd/Private/Fbx/FbxMainExport.cpp
@@ -50,6 +50,7 @@
 #include "StaticMeshResources.h"
 #include "LandscapeDataAccess.h"
 #include "Landscape/LandscapeProxy.h"
+#include "Components/SplineMeshComponent.h"
 
 #include "FbxExporter.h"
 #include "RawMesh.h"
diff --git a/Engine/Source/Editor/UnrealEd/Private/FbxExporter.h b/Engine/Source/Editor/UnrealEd/Private/FbxExporter.h
--- a/Engine/Source/Editor/UnrealEd/Private/FbxExporter.h
+++ b/Engine/Source/Editor/UnrealEd/Private/FbxExporter.h
@@ -36,6 +36,7 @@
 
 class ALandscapeProxy;
 struct FAnimControlTrackKey;
+class USplineMeshComponent;
 
 namespace UnFbx
 {
diff --git a/Engine/Source/Editor/UnrealEd/Private/Lightmass/Lightmass.cpp b/Engine/Source/Editor/UnrealEd/Private/Lightmass/Lightmass.cpp
--- a/Engine/Source/Editor/UnrealEd/Private/Lightmass/Lightmass.cpp
+++ b/Engine/Source/Editor/UnrealEd/Private/Lightmass/Lightmass.cpp
@@ -14,6 +14,7 @@
 #include "Runtime/Engine/Classes/Matinee/InterpGroupInst.h"
 #include "Runtime/Engine/Classes/Matinee/InterpTrackMove.h"
 #include "Runtime/Engine/Classes/Matinee/InterpTrackInstMove.h"
+#include "Runtime/Engine/Classes/Components/SplineMeshComponent.h"
 #include "Lightmass/PrecomputedVisibilityVolume.h"
 #include "Lightmass/PrecomputedVisibilityOverrideVolume.h"
 
@@ -1468,7 +1469,7 @@ void FLightmassExporter::WriteMeshInstances( int32 Channel )
 			const USplineMeshComponent* SplineComponent = CastChecked<USplineMeshComponent>(SMLightingMesh->Component);
 			SMInstanceMeshData.bIsSplineMesh = true;
 			Copy(*SplineParams, SMInstanceMeshData.SplineParameters);
-			SMInstanceMeshData.SplineParameters.SplineUpDir = SplineComponent->SplineXDir;
+			SMInstanceMeshData.SplineParameters.SplineUpDir = SplineComponent->SplineUpDir;
 			SMInstanceMeshData.SplineParameters.bSmoothInterpRollScale = SplineComponent->bSmoothInterpRollScale;
 			SMInstanceMeshData.SplineParameters.MeshMinZ = MeshBounds.Origin[SplineComponent->ForwardAxis] - MeshBounds.BoxExtent[SplineComponent->ForwardAxis];
 			SMInstanceMeshData.SplineParameters.MeshRangeZ = 2.f * MeshBounds.BoxExtent[SplineComponent->ForwardAxis];
diff --git a/Engine/Source/Runtime/Engine/Classes/Components/SplineComponent.h b/Engine/Source/Runtime/Engine/Classes/Components/SplineComponent.h
--- a/Engine/Source/Runtime/Engine/Classes/Components/SplineComponent.h
+++ b/Engine/Source/Runtime/Engine/Classes/Components/SplineComponent.h
@@ -44,6 +44,22 @@ class ENGINE_API USplineComponent : public USceneComponent
 	UFUNCTION(BlueprintCallable, Category=Spline)
 	void SetSplineWorldPoints(const TArray<FVector>& Points);
 
+	/** Move an existing point to a new world location */
+	UFUNCTION(BlueprintCallable, Category = Spline)
+	void SetWorldLocationAtSplinePoint(int32 PointIndex, const FVector& InLocation);
+
+	/** Get the number of points that make up this spline */
+	UFUNCTION(BlueprintCallable, Category=Spline)
+	int32 GetNumSplinePoints() const;
+
+	/** Get the world location at spline point */
+	UFUNCTION(BlueprintCallable, Category=Spline)
+	FVector GetWorldLocationAtSplinePoint(int32 PointIndex) const;
+
+	/** Get local location and tangent at a spline point */
+	UFUNCTION(BlueprintCallable, Category=Spline)
+	void GetLocalLocationAndTangentAtSplinePoint(int32 PointIndex, FVector& LocalLocation, FVector& LocalTangent) const;
+
 	/** Returns total length along this spline */
 	UFUNCTION(BlueprintCallable, Category=Spline) 
 	float GetSplineLength() const;
@@ -76,6 +92,9 @@ class ENGINE_API USplineComponent : public USceneComponent
 	UFUNCTION(BlueprintCallable, Category=Spline)
 	FRotator GetWorldRotationAtTime(float Time, bool bUseConstantVelocity = false) const;
 
+
+
+
 	/** Walk through keys and set time for each one */
 	void RefreshSplineInputs();
 
diff --git a/Engine/Source/Runtime/Engine/Classes/Components/SplineMeshComponent.h b/Engine/Source/Runtime/Engine/Classes/Components/SplineMeshComponent.h
--- a/Engine/Source/Runtime/Engine/Classes/Components/SplineMeshComponent.h
+++ b/Engine/Source/Runtime/Engine/Classes/Components/SplineMeshComponent.h
@@ -41,7 +41,7 @@ struct FSplineMeshParams
 	float StartRoll;
 
 	/** Starting offset of the mesh from the spline, in component space */
-	UPROPERTY(EditAnywhere, Category=SplineMesh)
+	UPROPERTY(EditAnywhere, Category=SplineMesh, AdvancedDisplay)
 	FVector2D StartOffset;
 
 	/** End location of spline, in component space */
@@ -88,11 +88,11 @@ class ENGINE_API USplineMeshComponent : public UStaticMeshComponent, public IInt
 
 	/** Spline that is used to deform mesh */
 	UPROPERTY(EditAnywhere, Category=SplineMesh, meta=(ShowOnlyInnerProperties))
-	struct FSplineMeshParams SplineParams;
+	FSplineMeshParams SplineParams;
 
 	/** Axis (in component space) that is used to determine X axis for co-ordinates along spline */
 	UPROPERTY(EditAnywhere, Category=SplineMesh)
-	FVector SplineXDir;
+	FVector SplineUpDir;
 
 	/** If true, will use smooth interpolation (ease in/out) for Scale, Roll, and Offset along this section of spline. If false, uses linear */
 	UPROPERTY(EditAnywhere, Category=SplineMesh, AdvancedDisplay)
@@ -150,18 +150,52 @@ class ENGINE_API USplineMeshComponent : public UStaticMeshComponent, public IInt
 	/** Called when spline params are changed, to notify render thread and possibly collision */
 	void MarkSplineParamsDirty();
 
+	/** Set the start position of spline in local space */
 	UFUNCTION(BlueprintCallable, Category=SplineMesh)
 	void SetStartPosition(FVector StartPos);
 
+	/** Set the start tangent vector of spline in local space */
 	UFUNCTION(BlueprintCallable, Category = SplineMesh)
 	void SetStartTangent(FVector StartTangent);
 
+	/** Set the end position of spline in local space */
 	UFUNCTION(BlueprintCallable, Category = SplineMesh)
 	void SetEndPosition(FVector EndPos);
 
+	/** Set the end tangent vector of spline in local space */
 	UFUNCTION(BlueprintCallable, Category = SplineMesh)
 	void SetEndTangent(FVector EndTangent);
 
+	/** Set the start and end, position and tangent, all in local space */
+	UFUNCTION(BlueprintCallable, Category = SplineMesh)
+	void SetStartAndEnd(FVector StartPos, FVector StartTangent, FVector EndPos, FVector EndTangent);
+
+	/** Set the start scaling */
+	UFUNCTION(BlueprintCallable, Category = SplineMesh)
+	void SetStartScale(FVector2D StartScale = FVector2D(1,1));
+
+	/** Set the start roll */
+	UFUNCTION(BlueprintCallable, Category = SplineMesh)
+	void SetStartRoll(float StartRoll);
+
+	/** Set the start offset */
+	UFUNCTION(BlueprintCallable, Category = SplineMesh)
+	void SetStartOffset(FVector2D StartOffset);
+
+	/** Set the end scaling */
+	UFUNCTION(BlueprintCallable, Category = SplineMesh)
+	void SetEndScale(FVector2D EndScale = FVector2D(1,1));
+
+	/** Set the end roll */
+	UFUNCTION(BlueprintCallable, Category = SplineMesh)
+	void SetEndRoll(float EndRoll);
+
+	/** Set the end offset */
+	UFUNCTION(BlueprintCallable, Category = SplineMesh)
+	void SetEndOffset(FVector2D EndOffset);
+
+
+
 	// Destroys the body setup, used to clear collision if the mesh goes missing
 	void DestroyBodySetup();
 #if WITH_EDITOR
diff --git a/Engine/Source/Runtime/Engine/Private/Landscape/LandscapeSplines.cpp b/Engine/Source/Runtime/Engine/Private/Landscape/LandscapeSplines.cpp
--- a/Engine/Source/Runtime/Engine/Private/Landscape/LandscapeSplines.cpp
+++ b/Engine/Source/Runtime/Engine/Private/Landscape/LandscapeSplines.cpp
@@ -1527,7 +1527,7 @@ void ULandscapeSplineSegment::UpdateSplinePoints(bool bUpdateCollision)
 			MeshComponent->SplineParams.EndRoll = RollEnd;
 			MeshComponent->SplineParams.EndOffset = OffsetEnd;
 
-			MeshComponent->SplineXDir = FVector(0,0,1); // Up, to be consistent between joined meshes. We rotate it to horizontal using roll if using Z Forward X Up or X Forward Y Up
+			MeshComponent->SplineUpDir = FVector(0,0,1); // Up, to be consistent between joined meshes. We rotate it to horizontal using roll if using Z Forward X Up or X Forward Y Up
 			MeshComponent->ForwardAxis = MeshEntry->ForwardAxis;
 
 			if (USplineMeshComponent::GetAxisValue(MeshEntry->Scale, MeshEntry->ForwardAxis) < 0)
diff --git a/Engine/Source/Runtime/Engine/Private/SplineComponent.cpp b/Engine/Source/Runtime/Engine/Private/SplineComponent.cpp
--- a/Engine/Source/Runtime/Engine/Private/SplineComponent.cpp
+++ b/Engine/Source/Runtime/Engine/Private/SplineComponent.cpp
@@ -210,6 +210,42 @@ void USplineComponent::SetSplineWorldPoints(const TArray<FVector>& Points)
 	UpdateSplineReparamTable();
 }
 
+void USplineComponent::SetWorldLocationAtSplinePoint(int32 PointIndex, const FVector& InLocation)
+{
+	if((PointIndex >= 0) && (PointIndex < SplineInfo.Points.Num()))
+	{
+		SplineInfo.Points[PointIndex].OutVal = InLocation;
+
+		SplineInfo.AutoSetTangents();
+		UpdateSplineReparamTable();
+	}
+}
+
+int32 USplineComponent::GetNumSplinePoints() const
+{
+	return SplineInfo.Points.Num();
+}
+
+FVector USplineComponent::GetWorldLocationAtSplinePoint(int32 PointIndex) const
+{
+	FVector LocalLocation(0,0,0);
+	if ((PointIndex >= 0) && (PointIndex < SplineInfo.Points.Num()))
+	{
+		LocalLocation = SplineInfo.Points[PointIndex].OutVal;
+	}
+	return ComponentToWorld.TransformPosition(LocalLocation);
+}
+
+void USplineComponent::GetLocalLocationAndTangentAtSplinePoint(int32 PointIndex, FVector& LocalLocation, FVector& LocalTangent) const
+{
+	LocalTangent = FVector(0, 0, 0);
+	LocalLocation = FVector(0, 0, 0);
+	if ((PointIndex >= 0) && (PointIndex < SplineInfo.Points.Num()))
+	{
+		LocalTangent = SplineInfo.Points[PointIndex].LeaveTangent;
+		LocalLocation = SplineInfo.Points[PointIndex].OutVal;
+	}
+}
 
 float USplineComponent::GetSplineLength() const
 {
diff --git a/Engine/Source/Runtime/Engine/Private/SplineMeshComponent.cpp b/Engine/Source/Runtime/Engine/Private/SplineMeshComponent.cpp
--- a/Engine/Source/Runtime/Engine/Private/SplineMeshComponent.cpp
+++ b/Engine/Source/Runtime/Engine/Private/SplineMeshComponent.cpp
@@ -53,7 +53,7 @@ void FSplineMeshVertexFactoryShaderParameters::SetMesh(FShader* Shader, const FV
 		SetShaderValue(Shader->GetVertexShader(), SplineEndScaleParam, SplineParams.EndScale);
 		SetShaderValue(Shader->GetVertexShader(), SplineEndOffsetParam, SplineParams.EndOffset);
 
-		SetShaderValue(Shader->GetVertexShader(), SplineUpDirParam, SplineProxy->SplineXDir);
+		SetShaderValue(Shader->GetVertexShader(), SplineUpDirParam, SplineProxy->SplineUpDir);
 		SetShaderValue(Shader->GetVertexShader(), SmoothInterpRollScaleParam, SplineProxy->bSmoothInterpRollScale);
 
 		SetShaderValue(Shader->GetVertexShader(), SplineMeshMinZParam, SplineProxy->SplineMeshMinZ);
@@ -200,7 +200,7 @@ USplineMeshComponent::USplineMeshComponent(const class FPostConstructInitializeP
 	SetCollisionProfileName(UCollisionProfile::NoCollision_ProfileName);
 	bHasCustomNavigableGeometry = EHasCustomNavigableGeometry::Yes;
 
-	SplineXDir.Z = 1.0f;
+	SplineUpDir.Z = 1.0f;
 
 	// Default to useful length and scale
 	SplineParams.StartTangent = FVector(100.f, 0.f, 0.f);
@@ -235,6 +235,51 @@ void USplineMeshComponent::SetEndTangent(FVector EndTangent)
 	MarkSplineParamsDirty();
 }
 
+void USplineMeshComponent::SetStartAndEnd(FVector StartPos, FVector StartTangent, FVector EndPos, FVector EndTangent)
+{
+	SplineParams.StartPos = StartPos;
+	SplineParams.StartTangent = StartTangent;
+	SplineParams.EndPos = EndPos;
+	SplineParams.EndTangent = EndTangent;
+	MarkSplineParamsDirty();
+}
+
+void USplineMeshComponent::SetStartScale(FVector2D StartScale)
+{
+	SplineParams.StartScale = StartScale;
+	MarkSplineParamsDirty();
+}
+
+void USplineMeshComponent::SetStartRoll(float StartRoll)
+{
+	SplineParams.StartRoll = StartRoll;
+	MarkSplineParamsDirty();
+}
+
+void USplineMeshComponent::SetStartOffset(FVector2D StartOffset)
+{
+	SplineParams.StartOffset = StartOffset;
+	MarkSplineParamsDirty();
+}
+
+void USplineMeshComponent::SetEndScale(FVector2D EndScale)
+{
+	SplineParams.EndScale = EndScale;
+	MarkSplineParamsDirty();
+}
+
+void USplineMeshComponent::SetEndRoll(float EndRoll)
+{
+	SplineParams.EndRoll = EndRoll;
+	MarkSplineParamsDirty();
+}
+
+void USplineMeshComponent::SetEndOffset(FVector2D EndOffset)
+{
+	SplineParams.EndOffset = EndOffset;
+	MarkSplineParamsDirty();
+}
+
 void USplineMeshComponent::MarkSplineParamsDirty()
 {
 	MarkRenderStateDirty();
@@ -387,7 +432,7 @@ FTransform USplineMeshComponent::CalcSliceTransform(const float DistanceAlong) c
 	const FVector SplineDir = SplineEvalDir( SplineParams.StartPos, SplineParams.StartTangent, SplineParams.EndPos, SplineParams.EndTangent, Alpha );
 
 	// Find base frenet frame
-	const FVector BaseXVec = (SplineXDir ^ SplineDir).SafeNormal();
+	const FVector BaseXVec = (SplineUpDir ^ SplineDir).SafeNormal();
 	const FVector BaseYVec = (SplineDir ^ BaseXVec).SafeNormal();
 
 	// Offset the spline by the desired amount
diff --git a/Engine/Source/Runtime/Engine/Private/SplineMeshSceneProxy.h b/Engine/Source/Runtime/Engine/Private/SplineMeshSceneProxy.h
--- a/Engine/Source/Runtime/Engine/Private/SplineMeshSceneProxy.h
+++ b/Engine/Source/Runtime/Engine/Private/SplineMeshSceneProxy.h
@@ -161,7 +161,7 @@ public:
 
 		// Copy spline params from component
 		SplineParams = InComponent->SplineParams;
-		SplineXDir = InComponent->SplineXDir;
+		SplineUpDir = InComponent->SplineUpDir;
 		bSmoothInterpRollScale = InComponent->bSmoothInterpRollScale;
 		ForwardAxis = InComponent->ForwardAxis;
 
@@ -249,7 +249,7 @@ public:
 	/** Parameters that define the spline, used to deform mesh */
 	FSplineMeshParams SplineParams;
 	/** Axis (in component space) that is used to determine X axis for co-ordinates along spline */
-	FVector SplineXDir;
+	FVector SplineUpDir;
 	/** Smoothly (cubic) interpolate the Roll and Scale params over spline. */
 	bool bSmoothInterpRollScale;
 	/** Chooses the forward axis for the spline mesh orientation */
diff --git a/Engine/Source/Runtime/Engine/Public/Engine.h b/Engine/Source/Runtime/Engine/Public/Engine.h
--- a/Engine/Source/Runtime/Engine/Public/Engine.h
+++ b/Engine/Source/Runtime/Engine/Public/Engine.h
@@ -173,7 +173,6 @@
 #include "Components/StaticMeshComponent.h"
 #include "Components/InstancedStaticMeshComponent.h"
 #include "Interfaces/Interface_CollisionDataProvider.h"
-#include "Components/SplineMeshComponent.h"
 #include "Components/ModelComponent.h"
 #include "AI/Navigation/NavLinkRenderingComponent.h"
 #include "AI/Navigation/NavMeshRenderingComponent.h"
