Temporarily disable deprecated warning around [AVFWriter finishWriting] - its replacement is asynchronous and requires some re-working.

diff --git a/Engine/Source/Runtime/Engine/Private/AVIWriter.cpp b/Engine/Source/Runtime/Engine/Private/AVIWriter.cpp
--- a/Engine/Source/Runtime/Engine/Private/AVIWriter.cpp
+++ b/Engine/Source/Runtime/Engine/Private/AVIWriter.cpp
@@ -593,7 +593,11 @@ public:
 		{
 			SCOPED_AUTORELEASE_POOL;
 			[AVFWriterInputRef markAsFinished];
+			// @TODO: Remove usage of deprecated finishWriting selector
+#pragma clang diagnostic push
+#pragma clang diagnostic ignored "-Wdeprecated-declarations"
 			BOOL OK = [AVFWriterRef finishWriting];
+#pragma clang diagnostic pop
 			check(OK);
 			[AVFWriterInputRef release];
 			[AVFWriterRef release];
