Shadow variable fixes

diff --git a/Engine/Source/Developer/LogVisualizer/Private/SLogBar.cpp b/Engine/Source/Developer/LogVisualizer/Private/SLogBar.cpp
--- a/Engine/Source/Developer/LogVisualizer/Private/SLogBar.cpp
+++ b/Engine/Source/Developer/LogVisualizer/Private/SLogBar.cpp
@@ -96,11 +96,13 @@ int32 SLogBar::OnPaint( const FGeometry& AllottedGeometry, const FSlateRect& MyC
 	while (EntryIndex < Entries.Num())
 	{
 		float CurrentStartX, CurrentEndX;
-		TSharedPtr<FVisLogEntry> Entry = Entries[EntryIndex];
-		if (!Entry.IsValid() || !CalculateEntryGeometry(Entry.Get(), AllottedGeometry, CurrentStartX, CurrentEndX))
 		{
-			EntryIndex++;
-			continue;
+			TSharedPtr<FVisLogEntry> Entry = Entries[EntryIndex];
+			if (!Entry.IsValid() || !CalculateEntryGeometry(Entry.Get(), AllottedGeometry, CurrentStartX, CurrentEndX))
+			{
+				EntryIndex++;
+				continue;
+			}
 		}
 
 		// find bar width, connect all contiguous bars to draw them as one geometry (rendering optimization)
@@ -166,7 +168,7 @@ int32 SLogBar::OnPaint( const FGeometry& AllottedGeometry, const FSlateRect& MyC
 	
 	if (bShouldDrawSelection)
 	{
-		int32 EntryIndex = CurrentEntryIndex.IsBound() ? CurrentEntryIndex.Execute() : INDEX_NONE;
+		EntryIndex = CurrentEntryIndex.IsBound() ? CurrentEntryIndex.Execute() : INDEX_NONE;
 		if (Entries.IsValidIndex(EntryIndex) == true)
 		{
 			TSharedPtr<FVisLogEntry> Entry = Entries[EntryIndex];
diff --git a/Engine/Source/Developer/LogVisualizer/Private/SLogVisualizer.cpp b/Engine/Source/Developer/LogVisualizer/Private/SLogVisualizer.cpp
--- a/Engine/Source/Developer/LogVisualizer/Private/SLogVisualizer.cpp
+++ b/Engine/Source/Developer/LogVisualizer/Private/SLogVisualizer.cpp
@@ -688,7 +688,7 @@ void SLogVisualizer::IncrementCurrentLogIndex(int32 IncrementBy)
 	}
 }
 
-void SLogVisualizer::AddLog(int32 Index, const FActorsVisLog* Log)
+void SLogVisualizer::AddLog(int32 LogIndex, const FActorsVisLog* Log)
 {
 	if (Log->Entries.Num() == 0)
 	{
@@ -731,7 +731,7 @@ void SLogVisualizer::AddLog(int32 Index, const FActorsVisLog* Log)
 	}
 
 	LogsList.Add(MakeShareable(new FLogsListItem(Log->Name.ToString()
-		, StartTimestamp, EndTimestamp, Index)));
+		, StartTimestamp, EndTimestamp, LogIndex)));
 }
 
 void SLogVisualizer::DoFullUpdate()
diff --git a/Engine/Source/Editor/DistCurveEditor/Private/SDistributionCurveEditor.cpp b/Engine/Source/Editor/DistCurveEditor/Private/SDistributionCurveEditor.cpp
--- a/Engine/Source/Editor/DistCurveEditor/Private/SDistributionCurveEditor.cpp
+++ b/Engine/Source/Editor/DistCurveEditor/Private/SDistributionCurveEditor.cpp
@@ -1258,18 +1258,18 @@ void SDistributionCurveEditor::ScaleTimeCommitted(const FText& CommentText, ETex
 
 		struct Local
 		{
-			static void ScaleCurveTime(FCurveEdEntry& Entry, float ScaleByValue)
+			static void ScaleCurveTime(FCurveEdEntry& Entry, float InScaleByValue)
 			{
 				FCurveEdInterface* EdInterface = UInterpCurveEdSetup::GetCurveEdInterfacePointer(Entry);
 				if(EdInterface)
 				{
 					// For each key
-					if (ScaleByValue >= 1.0f)
+					if (InScaleByValue >= 1.0f)
 					{
 						for (int32 KeyIndex = EdInterface->GetNumKeys() - 1; KeyIndex >= 0; KeyIndex--)
 						{
 							float InVal = EdInterface->GetKeyIn(KeyIndex);
-							EdInterface->SetKeyIn(KeyIndex, InVal * ScaleByValue);
+							EdInterface->SetKeyIn(KeyIndex, InVal * InScaleByValue);
 						}
 					}
 					else
@@ -1277,7 +1277,7 @@ void SDistributionCurveEditor::ScaleTimeCommitted(const FText& CommentText, ETex
 						for (int32 KeyIndex = 0; KeyIndex < EdInterface->GetNumKeys(); KeyIndex++)
 						{
 							float InVal = EdInterface->GetKeyIn(KeyIndex);
-							EdInterface->SetKeyIn(KeyIndex, InVal * ScaleByValue);
+							EdInterface->SetKeyIn(KeyIndex, InVal * InScaleByValue);
 						}
 					}
 				}
@@ -1321,7 +1321,7 @@ void SDistributionCurveEditor::ScaleValueCommitted(const FText& CommentText, ETe
 
 		struct Local
 		{
-			static void ScaleCurveValue(FCurveEdEntry& Entry, int32 SubCurve, float ScaleByValue)
+			static void ScaleCurveValue(FCurveEdEntry& Entry, int32 SubCurve, float InScaleByValue)
 			{
 				FCurveEdInterface* EdInterface = UInterpCurveEdSetup::GetCurveEdInterfacePointer(Entry);
 				if(EdInterface)
@@ -1335,7 +1335,7 @@ void SDistributionCurveEditor::ScaleValueCommitted(const FText& CommentText, ETe
 						for (int32 KeyIndex = 0; KeyIndex < EdInterface->GetNumKeys(); KeyIndex++)
 						{
 							float OutVal = EdInterface->GetKeyOut(SubCurve, KeyIndex);
-							EdInterface->SetKeyOut(SubCurve, KeyIndex, OutVal * ScaleByValue);
+							EdInterface->SetKeyOut(SubCurve, KeyIndex, OutVal * InScaleByValue);
 						}
 					}
 					else
@@ -1347,7 +1347,7 @@ void SDistributionCurveEditor::ScaleValueCommitted(const FText& CommentText, ETe
 							for (int32 KeyIndex = 0; KeyIndex < EdInterface->GetNumKeys(); KeyIndex++)
 							{
 								float OutVal = EdInterface->GetKeyOut(SubIndex, KeyIndex);
-								EdInterface->SetKeyOut(SubIndex, KeyIndex, OutVal * ScaleByValue);
+								EdInterface->SetKeyOut(SubIndex, KeyIndex, OutVal * InScaleByValue);
 							}
 						}
 					}
diff --git a/Engine/Source/Editor/Layers/Private/LayerViewModel.cpp b/Engine/Source/Editor/Layers/Private/LayerViewModel.cpp
--- a/Engine/Source/Editor/Layers/Private/LayerViewModel.cpp
+++ b/Engine/Source/Editor/Layers/Private/LayerViewModel.cpp
@@ -352,9 +352,9 @@ void FLayerViewModel::SelectActorsOfSpecificType( const TWeakObjectPtr< UClass >
 
 	struct Local
 	{
-		static bool ActorIsOfClass( const TWeakObjectPtr< AActor >& Actor, const TWeakObjectPtr< UClass > Class )
+		static bool ActorIsOfClass( const TWeakObjectPtr< AActor >& Actor, const TWeakObjectPtr< UClass > InClass )
 		{
-			return Actor->GetClass() == Class.Get();
+			return Actor->GetClass() == InClass.Get();
 		}
 	};
 
diff --git a/Engine/Source/Editor/MeshPaint/Private/MeshPaintEdMode.cpp b/Engine/Source/Editor/MeshPaint/Private/MeshPaintEdMode.cpp
--- a/Engine/Source/Editor/MeshPaint/Private/MeshPaintEdMode.cpp
+++ b/Engine/Source/Editor/MeshPaint/Private/MeshPaintEdMode.cpp
@@ -1735,7 +1735,7 @@ void FEdModeMeshPaint::PaintMeshVertices(
 
 
 /** Paints mesh texture */
-void FEdModeMeshPaint::PaintMeshTexture( UStaticMeshComponent* StaticMeshComponent, const FMeshPaintParameters& Params, const bool bShouldApplyPaint, FStaticMeshLODResources& LODModel, const FVector& ComponentSpaceCameraPosition, const FMatrix& ComponentToWorldMatrix, const float ComponentSpaceSquaredBrushRadius, const FVector& ComponentSpaceBrushPosition )
+void FEdModeMeshPaint::PaintMeshTexture( UStaticMeshComponent* StaticMeshComponent, const FMeshPaintParameters& Params, const bool bShouldApplyPaint, FStaticMeshLODResources& InLODModel, const FVector& ComponentSpaceCameraPosition, const FMatrix& ComponentToWorldMatrix, const float ComponentSpaceSquaredBrushRadius, const FVector& ComponentSpaceBrushPosition )
 {
 	UTexture2D* TargetTexture2D = GetSelectedTexture();
 	
@@ -1753,7 +1753,7 @@ void FEdModeMeshPaint::PaintMeshTexture( UStaticMeshComponent* StaticMeshCompone
 
 
 		// Make sure we're dealing with triangle lists
-		FIndexArrayView Indices = LODModel.IndexBuffer.GetArrayView();
+		FIndexArrayView Indices = InLODModel.IndexBuffer.GetArrayView();
 		const uint32 NumIndexBufferIndices = Indices.Num();
 		check( NumIndexBufferIndices % 3 == 0 );
 		const uint32 NumTriangles = NumIndexBufferIndices / 3;
diff --git a/Engine/Source/Editor/PropertyEditor/Private/PropertyEditorModule.cpp b/Engine/Source/Editor/PropertyEditor/Private/PropertyEditorModule.cpp
--- a/Engine/Source/Editor/PropertyEditor/Private/PropertyEditorModule.cpp
+++ b/Engine/Source/Editor/PropertyEditor/Private/PropertyEditorModule.cpp
@@ -60,9 +60,9 @@ const FPropertyTypeLayoutCallback& FPropertyTypeLayoutCallbackList::Find( const
 		FPropertyTypeLayoutCallback* Callback =
 			IdentifierList.FindByPredicate
 			(
-				[&]( const FPropertyTypeLayoutCallback& Callback )
+				[&]( const FPropertyTypeLayoutCallback& InCallback )
 				{
-					return Callback.PropertyTypeIdentifier->IsPropertyTypeCustomized( PropertyHandle );
+					return InCallback.PropertyTypeIdentifier->IsPropertyTypeCustomized( PropertyHandle );
 				}
 			);
 
diff --git a/Engine/Source/Editor/SceneOutliner/Private/SSceneOutliner.cpp b/Engine/Source/Editor/SceneOutliner/Private/SSceneOutliner.cpp
--- a/Engine/Source/Editor/SceneOutliner/Private/SSceneOutliner.cpp
+++ b/Engine/Source/Editor/SceneOutliner/Private/SSceneOutliner.cpp
@@ -1751,10 +1751,10 @@ namespace SceneOutliner
 		}
 
 		// Add a mini scene outliner for choosing an existing folder
-		FSceneOutlinerInitializationOptions InitOptions;
-		InitOptions.bShowHeaderRow = false;
-		InitOptions.bFocusSearchBoxWhenOpened = true;
-		InitOptions.bOnlyShowFolders = true;
+		FSceneOutlinerInitializationOptions MiniSceneOutlinerInitOptions;
+		MiniSceneOutlinerInitOptions.bShowHeaderRow = false;
+		MiniSceneOutlinerInitOptions.bFocusSearchBoxWhenOpened = true;
+		MiniSceneOutlinerInitOptions.bOnlyShowFolders = true;
 			
 		{
 			// Don't show any folders that are a child of any of the selected folders
@@ -1780,8 +1780,8 @@ namespace SceneOutliner
 					}
 
 					auto HasSubFolders = [](TWeakPtr<TOutlinerTreeItem> WeakItem){
-						auto Item = WeakItem.Pin();
-						return Item.IsValid() && Item->Type == TOutlinerTreeItem::Folder;
+						auto TreeItem = WeakItem.Pin();
+						return TreeItem.IsValid() && TreeItem->Type == TOutlinerTreeItem::Folder;
 					};
 
 					auto* ParentFolder = FolderToTreeItemMap.Find(Folder);
@@ -1795,8 +1795,8 @@ namespace SceneOutliner
 			if (ExcludedParents->Num())
 			{
 				// Add a filter if necessary
-				auto FilterOutChildFolders = [](FName Path, TSharedRef<TSet<FName>> ExcludedParents){
-					for (const auto& Parent : *ExcludedParents)
+				auto FilterOutChildFolders = [](FName Path, TSharedRef<TSet<FName>> InExcludedParents){
+					for (const auto& Parent : *InExcludedParents)
 					{
 						if (Path == Parent || FActorFolders::PathIsChildOf(Path.ToString(), Parent.ToString()))
 						{
@@ -1807,7 +1807,7 @@ namespace SceneOutliner
 					return true;
 				};
 
-				InitOptions.Filters->AddFilterPredicate(SceneOutliner::FFolderFilterPredicate::CreateStatic(FilterOutChildFolders, ExcludedParents));
+				MiniSceneOutlinerInitOptions.Filters->AddFilterPredicate(SceneOutliner::FFolderFilterPredicate::CreateStatic(FilterOutChildFolders, ExcludedParents));
 			}
 
 			// Actor selector to allow the user to choose a folder
@@ -1817,7 +1817,7 @@ namespace SceneOutliner
 				+ SVerticalBox::Slot()
 				.MaxHeight(400.0f)
 				[
-					SNew(SceneOutliner::SSceneOutliner, InitOptions)
+					SNew(SceneOutliner::SSceneOutliner, MiniSceneOutlinerInitOptions)
 					.IsEnabled(FSlateApplication::Get().GetNormalExecutionAttribute())
 					.OnItemPickedDelegate(FOnSceneOutlinerItemPicked::CreateSP(this, &SSceneOutliner::MoveSelectionTo))
 				];
diff --git a/Engine/Source/Editor/UnrealEd/Private/Commandlets/FixupRedirectsCommandlet.cpp b/Engine/Source/Editor/UnrealEd/Private/Commandlets/FixupRedirectsCommandlet.cpp
--- a/Engine/Source/Editor/UnrealEd/Private/Commandlets/FixupRedirectsCommandlet.cpp
+++ b/Engine/Source/Editor/UnrealEd/Private/Commandlets/FixupRedirectsCommandlet.cpp
@@ -114,9 +114,6 @@ int32 UFixupRedirectsCommandlet::Main( const FString& Params )
 		// load all string asset reference targets, and add fake redirectors for them
 		GRedirectCollector.ResolveStringAssetReference();
 
-		// a list of the redirectors that can't be cleaned because a read-only package references them
-		TArray<FString> RedirectorsThatCantBeCleaned;
-
 		for (int32 RedirIndex = 0; RedirIndex < GRedirectCollector.Redirections.Num(); RedirIndex++)
 		{
 			FRedirection& Redir = GRedirectCollector.Redirections[RedirIndex];
diff --git a/Engine/Source/Editor/UnrealEd/Private/CookOnTheFlyServer.cpp b/Engine/Source/Editor/UnrealEd/Private/CookOnTheFlyServer.cpp
--- a/Engine/Source/Editor/UnrealEd/Private/CookOnTheFlyServer.cpp
+++ b/Engine/Source/Editor/UnrealEd/Private/CookOnTheFlyServer.cpp
@@ -893,9 +893,9 @@ void UCookOnTheFlyServer::HandleNetworkFileServerFileRequest( const FString& Fil
 
 	if (!bIsCookable)
 	{
-		TArray<FFilePlatformRequest> UnsolicitedFileRequests;
-		UnsolicitedCookedPackages.DequeueAll(UnsolicitedFileRequests);
-		for (const FFilePlatformRequest &Request : UnsolicitedFileRequests)
+		TArray<FFilePlatformRequest> UnsolicitedRequests;
+		UnsolicitedCookedPackages.DequeueAll(UnsolicitedRequests);
+		for (const FFilePlatformRequest &Request : UnsolicitedRequests)
 		{
 			if ( Request.Platformname == PlatformFname )
 			{
@@ -924,9 +924,9 @@ void UCookOnTheFlyServer::HandleNetworkFileServerFileRequest( const FString& Fil
 
 	UE_LOG( LogCookOnTheFly, Display, TEXT("Cook complete %s"), *FileRequest.Filename)
 
-	TArray<FFilePlatformRequest> UnsolicitedFileRequests;
-	UnsolicitedCookedPackages.DequeueAll(UnsolicitedFileRequests);
-	for (const FFilePlatformRequest &Request : UnsolicitedFileRequests)
+	TArray<FFilePlatformRequest> UnsolicitedRequests;
+	UnsolicitedCookedPackages.DequeueAll(UnsolicitedRequests);
+	for (const FFilePlatformRequest &Request : UnsolicitedRequests)
 	{
 		if ( Request.Filename == FileRequest.Filename )
 		{
diff --git a/Engine/Source/Runtime/AIModule/Private/EnvironmentQuery/EnvQueryTypes.cpp b/Engine/Source/Runtime/AIModule/Private/EnvironmentQuery/EnvQueryTypes.cpp
--- a/Engine/Source/Runtime/AIModule/Private/EnvironmentQuery/EnvQueryTypes.cpp
+++ b/Engine/Source/Runtime/AIModule/Private/EnvironmentQuery/EnvQueryTypes.cpp
@@ -120,16 +120,18 @@ FText FEnvTraceData::ToText(FEnvTraceData::EDescriptionMode DescMode) const
 		FNumberFormattingOptions NumberFormatOptions;
 		NumberFormatOptions.MaximumFractionalDigits = 2;
 
-		FFormatNamedArguments Args;
-		Args.Add(TEXT("ExtentX"), FText::AsNumber(ExtentX, &NumberFormatOptions));
-		Args.Add(TEXT("ExtentY"), FText::AsNumber(ExtentY, &NumberFormatOptions));
-		Args.Add(TEXT("ExtentZ"), FText::AsNumber(ExtentZ, &NumberFormatOptions));
-
-		Desc = (TraceShape == EEnvTraceShape::Line) ? LOCTEXT("Line", "line") :
-			(TraceShape == EEnvTraceShape::Sphere) ? FText::Format(LOCTEXT("SphereWithRadius", "sphere (radius: {ExtentX})"), Args) :
-			(TraceShape == EEnvTraceShape::Capsule) ? FText::Format(LOCTEXT("CasuleWithRadiusHalfHeight", "capsule (radius: {ExtentX}, half height: {ExtentZ})"), Args) :
-			(TraceShape == EEnvTraceShape::Box) ? FText::Format(LOCTEXT("BoxWithExtents", "box (extent: {ExtentX} {ExtentY} {ExtentZ)"), Args) :
-			LOCTEXT("Unknown", "unknown");
+		{
+			FFormatNamedArguments Args;
+			Args.Add(TEXT("ExtentX"), FText::AsNumber(ExtentX, &NumberFormatOptions));
+			Args.Add(TEXT("ExtentY"), FText::AsNumber(ExtentY, &NumberFormatOptions));
+			Args.Add(TEXT("ExtentZ"), FText::AsNumber(ExtentZ, &NumberFormatOptions));
+
+			Desc = (TraceShape == EEnvTraceShape::Line) ? LOCTEXT("Line", "line") :
+				(TraceShape == EEnvTraceShape::Sphere) ? FText::Format(LOCTEXT("SphereWithRadius", "sphere (radius: {ExtentX})"), Args) :
+				(TraceShape == EEnvTraceShape::Capsule) ? FText::Format(LOCTEXT("CasuleWithRadiusHalfHeight", "capsule (radius: {ExtentX}, half height: {ExtentZ})"), Args) :
+				(TraceShape == EEnvTraceShape::Box) ? FText::Format(LOCTEXT("BoxWithExtents", "box (extent: {ExtentX} {ExtentY} {ExtentZ)"), Args) :
+				LOCTEXT("Unknown", "unknown");
+		}
 
 		if (DescMode == FEnvTraceData::Brief)
 		{
diff --git a/Engine/Source/Runtime/AIModule/Private/Navigation/CrowdManager.cpp b/Engine/Source/Runtime/AIModule/Private/Navigation/CrowdManager.cpp
--- a/Engine/Source/Runtime/AIModule/Private/Navigation/CrowdManager.cpp
+++ b/Engine/Source/Runtime/AIModule/Private/Navigation/CrowdManager.cpp
@@ -751,12 +751,14 @@ void UCrowdManager::DestroyCrowdManager()
 
 void UCrowdManager::DrawDebugCorners(const struct dtCrowdAgent* CrowdAgent) const
 {
-	FVector P0 = Recast2UnrealPoint(CrowdAgent->npos);
-	for (int32 Idx = 0; Idx < CrowdAgent->ncorners; Idx++)
 	{
-		FVector P1 = Recast2UnrealPoint(&CrowdAgent->cornerVerts[Idx * 3]);
-		DrawDebugLine(GetWorld(), P0 + CrowdDebugDrawing::Offset, P1 + CrowdDebugDrawing::Offset, CrowdDebugDrawing::Corner, false, -1.0f, SDPG_World, 2.0f);
-		P0 = P1;
+		FVector P0 = Recast2UnrealPoint(CrowdAgent->npos);
+		for (int32 Idx = 0; Idx < CrowdAgent->ncorners; Idx++)
+		{
+			FVector P1 = Recast2UnrealPoint(&CrowdAgent->cornerVerts[Idx * 3]);
+			DrawDebugLine(GetWorld(), P0 + CrowdDebugDrawing::Offset, P1 + CrowdDebugDrawing::Offset, CrowdDebugDrawing::Corner, false, -1.0f, SDPG_World, 2.0f);
+			P0 = P1;
+		}
 	}
 
 	if (CrowdAgent->ncorners && (CrowdAgent->cornerFlags[CrowdAgent->ncorners - 1] & DT_STRAIGHTPATH_OFFMESH_CONNECTION))
@@ -946,12 +948,14 @@ void UCrowdManager::DebugTick() const
 
 			if (CrowdAgent && LogOwner)
 			{
-				FVector P0 = Recast2UnrealPoint(CrowdAgent->npos);
-				for (int32 Idx = 0; Idx < CrowdAgent->ncorners; Idx++)
 				{
-					FVector P1 = Recast2UnrealPoint(&CrowdAgent->cornerVerts[Idx * 3]);
-					UE_VLOG_SEGMENT(LogOwner, LogEngineCrowdFollowing, Log, P0 + CrowdDebugDrawing::Offset, P1 + CrowdDebugDrawing::Offset, CrowdDebugDrawing::Corner, TEXT(""));
-					P0 = P1;
+					FVector P0 = Recast2UnrealPoint(CrowdAgent->npos);
+					for (int32 Idx = 0; Idx < CrowdAgent->ncorners; Idx++)
+					{
+						FVector P1 = Recast2UnrealPoint(&CrowdAgent->cornerVerts[Idx * 3]);
+						UE_VLOG_SEGMENT(LogOwner, LogEngineCrowdFollowing, Log, P0 + CrowdDebugDrawing::Offset, P1 + CrowdDebugDrawing::Offset, CrowdDebugDrawing::Corner, TEXT(""));
+						P0 = P1;
+					}
 				}
 
 				if (CrowdAgent->ncorners && (CrowdAgent->cornerFlags[CrowdAgent->ncorners - 1] & DT_STRAIGHTPATH_OFFMESH_CONNECTION))
diff --git a/Engine/Source/Runtime/Engine/Private/ChildActorComponent.cpp b/Engine/Source/Runtime/Engine/Private/ChildActorComponent.cpp
--- a/Engine/Source/Runtime/Engine/Private/ChildActorComponent.cpp
+++ b/Engine/Source/Runtime/Engine/Private/ChildActorComponent.cpp
@@ -182,9 +182,9 @@ void UChildActorComponent::DestroyChildActor()
 			// World may be NULL during shutdown
 			if(World != NULL)
 			{
-				UClass* ChildActorClass = ChildActor->GetClass();
-				const FString ObjectBaseName = FString::Printf(TEXT("DESTROYED_%s_CHILDACTOR"), *ChildActorClass->GetName());
-				ChildActor->Rename(*MakeUniqueObjectName(GetOuter(), ChildActorClass, *ObjectBaseName).ToString(), NULL, REN_DoNotDirty);
+				UClass* ChildClass = ChildActor->GetClass();
+				const FString ObjectBaseName = FString::Printf(TEXT("DESTROYED_%s_CHILDACTOR"), *ChildClass->GetName());
+				ChildActor->Rename(*MakeUniqueObjectName(GetOuter(), ChildClass, *ObjectBaseName).ToString(), NULL, REN_DoNotDirty);
 				World->DestroyActor(ChildActor);
 			}
 		}
diff --git a/Engine/Source/Runtime/Engine/Private/ColorVertexBuffer.cpp b/Engine/Source/Runtime/Engine/Private/ColorVertexBuffer.cpp
--- a/Engine/Source/Runtime/Engine/Private/ColorVertexBuffer.cpp
+++ b/Engine/Source/Runtime/Engine/Private/ColorVertexBuffer.cpp
@@ -309,7 +309,7 @@ void FColorVertexBuffer::GetVertexColors( TArray<FColor>& OutColors )
 }
 
 /** Load from raw color array */
-void FColorVertexBuffer::InitFromColorArray( const FColor *InColors, const uint32 Count, const uint32 Stride )
+void FColorVertexBuffer::InitFromColorArray( const FColor *InColors, const uint32 Count, const uint32 InStride )
 {
 	check( Count > 0 );
 
@@ -329,7 +329,7 @@ void FColorVertexBuffer::InitFromColorArray( const FColor *InColors, const uint3
 		{
 			*Dst++ = *(const FColor*)Src;
 
-			Src += Stride;
+			Src += InStride;
 		}
 	}
 
diff --git a/Engine/Source/Runtime/Engine/Private/ContentStreaming.cpp b/Engine/Source/Runtime/Engine/Private/ContentStreaming.cpp
--- a/Engine/Source/Runtime/Engine/Private/ContentStreaming.cpp
+++ b/Engine/Source/Runtime/Engine/Private/ContentStreaming.cpp
@@ -341,10 +341,10 @@ struct FStreamingTexture
 	 * @param MipCount	Number of mip-maps to account for
 	 * @return			Total amount of memory used for the specified mip-maps, in bytes
 	 */
-	int32 GetSize( int32 MipCount ) const
+	int32 GetSize( int32 InMipCount ) const
 	{
-		check( MipCount >= 0 && MipCount <= MAX_TEXTURE_MIP_COUNT );
-		return TextureSizes[ MipCount - 1 ];
+		check(InMipCount >= 0 && InMipCount <= MAX_TEXTURE_MIP_COUNT);
+		return TextureSizes[ InMipCount - 1 ];
 	}
 
 	/**
@@ -848,9 +848,9 @@ struct FTextureStreamingStats
  */
 struct FStreamingContext
 {
-	FStreamingContext( bool bProcessEverything, UTexture2D* IndividualStreamingTexture, bool bCollectTextureStats )
+	FStreamingContext( bool bProcessEverything, UTexture2D* IndividualStreamingTexture, bool bInCollectTextureStats )
 	{
-		Reset( bProcessEverything, IndividualStreamingTexture, bCollectTextureStats );
+		Reset( bProcessEverything, IndividualStreamingTexture, bInCollectTextureStats );
 	}
 
 	/**
@@ -4099,11 +4099,13 @@ void FStreamingManagerTexture::CalcWantedMips( FStreamingTexture& StreamingTextu
 		if ( !WantedMips.IsHandled() || bShouldAlsoUseLastRenderTime )
 		{
 			// Try the handler for textures that has recently lost instance locations.
-			float HandlerDistance = FLT_MAX;
-			FFloatMipLevel HandlerWantedMips = GetWantedMipsForOrphanedTexture( StreamingTexture, HandlerDistance );
-			WantedMips = FMath::Max( WantedMips, HandlerWantedMips );
-			MinDistance = FMath::Min( MinDistance, HandlerDistance );
-			STAT( PerfectWantedMips = FMath::Max( PerfectWantedMips, HandlerWantedMips ) );
+			{
+				float HandlerDistance = FLT_MAX;
+				FFloatMipLevel HandlerWantedMips = GetWantedMipsForOrphanedTexture(StreamingTexture, HandlerDistance);
+				WantedMips = FMath::Max(WantedMips, HandlerWantedMips);
+				MinDistance = FMath::Min(MinDistance, HandlerDistance);
+				STAT(PerfectWantedMips = FMath::Max(PerfectWantedMips, HandlerWantedMips));
+			}
 
 			// Still wasn't handled? Use the LastRenderTime handler as last resort.
 			if (!WantedMips.IsHandled() || bShouldAlsoUseLastRenderTime )
@@ -4962,7 +4964,7 @@ void FStreamingManagerTexture::DumpTextureGroupStats( bool bDetailedStats )
  *
  * @param InvestigateTextureName	Partial name to match textures against
  */
-void FStreamingManagerTexture::InvestigateTexture( const FString& InvestigateTextureName )
+void FStreamingManagerTexture::InvestigateTexture( const FString& InInvestigateTextureName )
 {
 	bTriggerInvestigateTexture = false;
 #if !(UE_BUILD_SHIPPING || UE_BUILD_TEST)
@@ -4970,7 +4972,7 @@ void FStreamingManagerTexture::InvestigateTexture( const FString& InvestigateTex
 	{
 		FStreamingTexture& StreamingTexture = StreamingTextures[TextureIndex];
 		FString TextureName = StreamingTexture.Texture->GetFullName();
-		if ( TextureName.Contains( InvestigateTextureName) )
+		if (TextureName.Contains(InInvestigateTextureName))
 		{
 			UTexture2D* Texture2D = StreamingTexture.Texture;
 			ETextureStreamingType StreamType = StreamingTexture.GetStreamingType();
@@ -5097,9 +5099,9 @@ void FStreamingManagerTexture::InvestigateTexture( const FString& InvestigateTex
 									}
 								}
 								int32 IntWantedMipCount = WantedMipCount.ComputeMip(&StreamingTexture, ThreadSettings.MipBias, false);
-								int32 WantedMipIndex = FMath::Max(Texture2D->GetNumMips() - IntWantedMipCount, 0);
+								int32 WantedMip = FMath::Max(Texture2D->GetNumMips() - IntWantedMipCount, 0);
 								UE_LOG(LogContentStreaming, Log, TEXT("Static: Wanted=%dx%d, Distance=%.1f, TexelFactor=%.2f, Radius=%5.1f, Position=(%d,%d,%d)%s"),
-									Texture2D->PlatformData->Mips[WantedMipIndex].SizeX, Texture2D->PlatformData->Mips[WantedMipIndex].SizeY,
+									Texture2D->PlatformData->Mips[WantedMip].SizeX, Texture2D->PlatformData->Mips[WantedMip].SizeY,
 									MinDistance, TexelFactor, Radius, int32(Center.X), int32(Center.Y), int32(Center.Z),
 									bInside ? TEXT(" [all mips]") : TEXT(""));
 							}
diff --git a/Engine/Source/Runtime/Engine/Private/DataTable.cpp b/Engine/Source/Runtime/Engine/Private/DataTable.cpp
--- a/Engine/Source/Runtime/Engine/Private/DataTable.cpp
+++ b/Engine/Source/Runtime/Engine/Private/DataTable.cpp
@@ -345,12 +345,12 @@ UProperty* UDataTable::FindTableProperty(const FName& PropertyName) const
 }
 
 /** Get array of UProperties that corresponds to columns in the table */
-TArray<UProperty*> UDataTable::GetTablePropertyArray(const FString& FirstRowString, UStruct* RowStruct, TArray<FString>& OutProblems)
+TArray<UProperty*> UDataTable::GetTablePropertyArray(const FString& FirstRowString, UStruct* InRowStruct, TArray<FString>& OutProblems)
 {
 	TArray<UProperty*> ColumnProps;
 
 	// Get list of all expected properties from the struct
-	TArray<FName> ExpectedPropNames = GetStructPropertyNames(RowStruct);	
+	TArray<FName> ExpectedPropNames = GetStructPropertyNames(InRowStruct);
 
 	// Find the column names from first row
 	TArray<FString> ColumnNameStrings;
@@ -372,11 +372,11 @@ TArray<UProperty*> UDataTable::GetTablePropertyArray(const FString& FirstRowStri
 			}
 			else
 			{
-				UProperty* ColumnProp = FindField<UProperty>(RowStruct, PropName);
+				UProperty* ColumnProp = FindField<UProperty>(InRowStruct, PropName);
 				// Didn't find a property with this name, problem..
 				if(ColumnProp == NULL)
 				{
-					OutProblems.Add(FString::Printf(TEXT("Cannot find Property for column '%s' in struct '%s'."), *PropName.ToString(), *RowStruct->GetName()));
+					OutProblems.Add(FString::Printf(TEXT("Cannot find Property for column '%s' in struct '%s'."), *PropName.ToString(), *InRowStruct->GetName()));
 				}
 				// Found one!
 				else
diff --git a/Engine/Source/Runtime/Engine/Private/FloatingPawnMovement.cpp b/Engine/Source/Runtime/Engine/Private/FloatingPawnMovement.cpp
--- a/Engine/Source/Runtime/Engine/Private/FloatingPawnMovement.cpp
+++ b/Engine/Source/Runtime/Engine/Private/FloatingPawnMovement.cpp
@@ -89,8 +89,8 @@ void UFloatingPawnMovement::ApplyControlInputToVelocity(float DeltaTime)
 	const FVector ControlAcceleration = GetInputVector().ClampMaxSize(1.f);
 
 	const float AnalogInputModifier = (ControlAcceleration.SizeSquared() > 0.f ? ControlAcceleration.Size() : 0.f);
-	const float MaxSpeed = GetModifiedMaxSpeed() * AnalogInputModifier;
-	const bool bExceedingMaxSpeed = IsExceedingMaxSpeed(MaxSpeed);
+	const float MaxPawnSpeed = GetModifiedMaxSpeed() * AnalogInputModifier;
+	const bool bExceedingMaxSpeed = IsExceedingMaxSpeed(MaxPawnSpeed);
 
 	if (AnalogInputModifier > 0.f && !bExceedingMaxSpeed)
 	{
@@ -110,15 +110,15 @@ void UFloatingPawnMovement::ApplyControlInputToVelocity(float DeltaTime)
 			Velocity = Velocity.SafeNormal() * VelSize;
 
 			// Don't allow braking to lower us below max speed if we started above it.
-			if (bExceedingMaxSpeed && Velocity.SizeSquared() < FMath::Square(MaxSpeed))
+			if (bExceedingMaxSpeed && Velocity.SizeSquared() < FMath::Square(MaxPawnSpeed))
 			{
-				Velocity = OldVelocity.SafeNormal() * MaxSpeed;
+				Velocity = OldVelocity.SafeNormal() * MaxPawnSpeed;
 			}
 		}
 	}
 
 	// Apply acceleration and clamp velocity magnitude.
-	const float NewMaxSpeed = (IsExceedingMaxSpeed(MaxSpeed)) ? Velocity.Size() : MaxSpeed;
+	const float NewMaxSpeed = (IsExceedingMaxSpeed(MaxPawnSpeed)) ? Velocity.Size() : MaxPawnSpeed;
 	Velocity += ControlAcceleration * FMath::Abs(Acceleration) * DeltaTime;
 	Velocity = Velocity.ClampMaxSize(NewMaxSpeed);
 
diff --git a/Engine/Source/Runtime/Engine/Private/GameSession.cpp b/Engine/Source/Runtime/Engine/Private/GameSession.cpp
--- a/Engine/Source/Runtime/Engine/Private/GameSession.cpp
+++ b/Engine/Source/Runtime/Engine/Private/GameSession.cpp
@@ -261,7 +261,7 @@ void AGameSession::ReturnToMainMenuHost()
 	}
 }
 
-bool AGameSession::TravelToSession(int32 ControllerId, FName SessionName)
+bool AGameSession::TravelToSession(int32 ControllerId, FName InSessionName)
 {
 	UWorld* World = GetWorld();
 	IOnlineSubsystem* OnlineSub = Online::GetSubsystem(World);
@@ -269,7 +269,7 @@ bool AGameSession::TravelToSession(int32 ControllerId, FName SessionName)
 	{
 		FString URL;
 		IOnlineSessionPtr SessionInt = OnlineSub->GetSessionInterface();
-		if (SessionInt.IsValid() && SessionInt->GetResolvedConnectString(SessionName, URL))
+		if (SessionInt.IsValid() && SessionInt->GetResolvedConnectString(InSessionName, URL))
 		{
 			APlayerController* PC = UGameplayStatics::GetPlayerController(World, ControllerId);
 			if (PC)
@@ -280,7 +280,7 @@ bool AGameSession::TravelToSession(int32 ControllerId, FName SessionName)
 		}
 		else
 		{
-			UE_LOG(LogGameSession, Warning, TEXT("Failed to resolve session connect string for %s"), *SessionName.ToString());
+			UE_LOG(LogGameSession, Warning, TEXT("Failed to resolve session connect string for %s"), *InSessionName.ToString());
 		}
 	}
 
diff --git a/Engine/Source/Runtime/Engine/Private/GameViewportClient.cpp b/Engine/Source/Runtime/Engine/Private/GameViewportClient.cpp
--- a/Engine/Source/Runtime/Engine/Private/GameViewportClient.cpp
+++ b/Engine/Source/Runtime/Engine/Private/GameViewportClient.cpp
@@ -1185,7 +1185,7 @@ void UGameViewportClient::PeekTravelFailureMessages(UWorld* InWorld, ETravelFail
 	UE_LOG(LogNet, Warning, TEXT("Travel Failure: [%s]: %s"), ETravelFailure::ToString(FailureType), *ErrorString);
 }
 
-void UGameViewportClient::PeekNetworkFailureMessages(UWorld *World, UNetDriver *NetDriver, ENetworkFailure::Type FailureType, const FString& ErrorString)
+void UGameViewportClient::PeekNetworkFailureMessages(UWorld *InWorld, UNetDriver *NetDriver, ENetworkFailure::Type FailureType, const FString& ErrorString)
 {
 	UE_LOG(LogNet, Warning, TEXT("Network Failure: %s[%s]: %s"), NetDriver ? *NetDriver->NetDriverName.ToString() : TEXT("NULL"), ENetworkFailure::ToString(FailureType), *ErrorString);
 }
diff --git a/Engine/Source/Runtime/Engine/Private/NavMeshRenderingComponent.cpp b/Engine/Source/Runtime/Engine/Private/NavMeshRenderingComponent.cpp
--- a/Engine/Source/Runtime/Engine/Private/NavMeshRenderingComponent.cpp
+++ b/Engine/Source/Runtime/Engine/Private/NavMeshRenderingComponent.cpp
@@ -308,13 +308,13 @@ void UNavMeshRenderingComponent::GatherData(struct FNavMeshSceneProxyData* Curre
 						{
 							for(int k = 0; k < Polys.Num(); ++k)
 							{
-								uint32 PolyIndex = 0;
-								uint32 TileIndex = 0;
-								NavMesh->GetPolyTileIndex(Polys[k].Ref, PolyIndex, TileIndex);
+								uint32 NavPolyIndex = 0;
+								uint32 NavTileIndex = 0;
+								NavMesh->GetPolyTileIndex(Polys[k].Ref, NavPolyIndex, NavTileIndex);
 
 								CurrentData->DebugLabels.Add(FNavMeshSceneProxyData::FDebugText(
 									/*Location*/Polys[k].Center + CurrentData->NavMeshDrawOffset
-									, /*Text*/FString::Printf(TEXT("[%X:%X]"), TileIndex, PolyIndex)
+									, /*Text*/FString::Printf(TEXT("[%X:%X]"), NavTileIndex, NavPolyIndex)
 									));
 							}
 						}
diff --git a/Engine/Source/Runtime/Engine/Private/PhysicsEngine/PhysSubstepTasks.cpp b/Engine/Source/Runtime/Engine/Private/PhysicsEngine/PhysSubstepTasks.cpp
--- a/Engine/Source/Runtime/Engine/Private/PhysicsEngine/PhysSubstepTasks.cpp
+++ b/Engine/Source/Runtime/Engine/Private/PhysicsEngine/PhysSubstepTasks.cpp
@@ -164,11 +164,11 @@ void FPhysSubstepTask::ApplyTorques(const FPhysTarget & PhysTarget, FBodyInstanc
 }
 
 /** Interpolates kinematic actor transform - Assumes caller has obtained writer lock */
-void FPhysSubstepTask::InterpolateKinematicActor(const FPhysTarget & PhysTarget, FBodyInstance * BodyInstance, float Alpha)
+void FPhysSubstepTask::InterpolateKinematicActor(const FPhysTarget & PhysTarget, FBodyInstance * BodyInstance, float InAlpha)
 {
 #if WITH_PHYSX
 	PxRigidDynamic * PRigidDynamic = BodyInstance->GetPxRigidDynamic();
-	Alpha = FMath::Clamp(Alpha, 0.f, 1.f);
+	InAlpha = FMath::Clamp(InAlpha, 0.f, 1.f);
 
 	/** Interpolate kinematic actors */
 	if (PhysTarget.bKinematicTarget)
@@ -181,8 +181,8 @@ void FPhysSubstepTask::InterpolateKinematicActor(const FPhysTarget & PhysTarget,
 			const FTransform & StartTM = KinematicTarget.OriginalTM;
 			FTransform InterTM = FTransform::Identity;
 
-			InterTM.SetLocation(FMath::Lerp(StartTM.GetLocation(), TargetTM.GetLocation(), Alpha));
-			InterTM.SetRotation(FMath::Lerp(StartTM.GetRotation(), TargetTM.GetRotation(), Alpha));
+			InterTM.SetLocation(FMath::Lerp(StartTM.GetLocation(), TargetTM.GetLocation(), InAlpha));
+			InterTM.SetRotation(FMath::Lerp(StartTM.GetRotation(), TargetTM.GetRotation(), InAlpha));
 
 			const PxTransform PNewPose = U2PTransform(InterTM);
 			check(PNewPose.isValid());
@@ -192,7 +192,7 @@ void FPhysSubstepTask::InterpolateKinematicActor(const FPhysTarget & PhysTarget,
 #endif
 }
 
-void FPhysSubstepTask::SubstepInterpolation(float Alpha)
+void FPhysSubstepTask::SubstepInterpolation(float InAlpha)
 {
 #if WITH_PHYSX
 #if WITH_APEX
@@ -222,11 +222,11 @@ void FPhysSubstepTask::SubstepInterpolation(float Alpha)
 
 		ApplyForces(PhysTarget, BodyInstance);
 		ApplyTorques(PhysTarget, BodyInstance);
-		InterpolateKinematicActor(PhysTarget, BodyInstance, Alpha);
+		InterpolateKinematicActor(PhysTarget, BodyInstance, InAlpha);
 	}
 
 	/** Final substep */
-	if (Alpha >= 1.f)
+	if (InAlpha >= 1.f)
 	{
 		Targets.Empty(Targets.Num());
 	}
diff --git a/Engine/Source/Runtime/Engine/Private/SkeletalMeshComponentPhysics.cpp b/Engine/Source/Runtime/Engine/Private/SkeletalMeshComponentPhysics.cpp
--- a/Engine/Source/Runtime/Engine/Private/SkeletalMeshComponentPhysics.cpp
+++ b/Engine/Source/Runtime/Engine/Private/SkeletalMeshComponentPhysics.cpp
@@ -3722,17 +3722,16 @@ void USkeletalMeshComponent::DrawClothingCollisionVolumes(FPrimitiveDrawInterfac
 	for(int32 AssetIdx=0; AssetIdx < SkeletalMesh->ClothingAssets.Num(); AssetIdx++)
 	{
 		TArray<FApexClothCollisionVolumeData>& Collisions = SkeletalMesh->ClothingAssets[AssetIdx].ClothCollisionVolumes;
-		int32 NumCollisions = Collisions.Num();
 		int32 SphereCount = 0;
 
-		for(int32 i=0; i<NumCollisions; i++)
+		for (const FApexClothCollisionVolumeData& Collision : Collisions)
 		{
-			if(Collisions[i].BoneIndex < 0)
+			if(Collision.BoneIndex < 0)
 			{
 				continue;
 			}
 
-			FName BoneName = SkeletalMesh->ClothingAssets[AssetIdx].ApexClothingAsset->GetConvertedBoneName(Collisions[i].BoneIndex);
+			FName BoneName = SkeletalMesh->ClothingAssets[AssetIdx].ApexClothingAsset->GetConvertedBoneName(Collision.BoneIndex);
 			
 			int32 BoneIndex = GetBoneIndex(BoneName);
 
@@ -3743,9 +3742,9 @@ void USkeletalMeshComponent::DrawClothingCollisionVolumes(FPrimitiveDrawInterfac
 
 			FMatrix BoneMat = GetBoneMatrix(BoneIndex);
 
-			FMatrix LocalToWorld = Collisions[i].LocalPose * BoneMat;
+			FMatrix LocalToWorld = Collision.LocalPose * BoneMat;
 
-			if(Collisions[i].IsCapsule())
+			if(Collision.IsCapsule())
 			{
 				FColor CapsuleColor = Colors[AssetIdx % 3];
 				if (SphereCount >= MaxSphereCollisions)
@@ -3753,23 +3752,23 @@ void USkeletalMeshComponent::DrawClothingCollisionVolumes(FPrimitiveDrawInterfac
 					CapsuleColor = GrayColor; // Draw in gray to show that these collisions are ignored
 				}
 
-				const int32 CapsuleSides = FMath::Clamp<int32>(Collisions[i].CapsuleRadius/4.f, 16, 64);
-				float CapsuleHalfHeight = (Collisions[i].CapsuleHeight*0.5f+Collisions[i].CapsuleRadius);
-				float CapsuleRadius = Collisions[i].CapsuleRadius*2.f;
+				const int32 CapsuleSides = FMath::Clamp<int32>(Collision.CapsuleRadius/4.f, 16, 64);
+				float CapsuleHalfHeight = (Collision.CapsuleHeight*0.5f+Collision.CapsuleRadius);
+				float CapsuleRadius = Collision.CapsuleRadius*2.f;
 				// swapped Y-axis and Z-axis to convert apex coordinate to UE coordinate
-				DrawWireCapsule(PDI, LocalToWorld.GetOrigin(), LocalToWorld.GetUnitAxis(EAxis::X), LocalToWorld.GetUnitAxis(EAxis::Z), LocalToWorld.GetUnitAxis(EAxis::Y), CapsuleColor, Collisions[i].CapsuleRadius, CapsuleHalfHeight, CapsuleSides, SDPG_World);
+				DrawWireCapsule(PDI, LocalToWorld.GetOrigin(), LocalToWorld.GetUnitAxis(EAxis::X), LocalToWorld.GetUnitAxis(EAxis::Z), LocalToWorld.GetUnitAxis(EAxis::Y), CapsuleColor, Collision.CapsuleRadius, CapsuleHalfHeight, CapsuleSides, SDPG_World);
 				SphereCount += 2; // 1 capsule takes 2 spheres
 			}
 			else // convex
 			{
-				int32 NumVerts = Collisions[i].BoneVertices.Num();
+				int32 NumVerts = Collision.BoneVertices.Num();
 
 				TArray<FVector> TransformedVerts;
 				TransformedVerts.AddUninitialized(NumVerts);
 
 				for(int32 VertIdx=0; VertIdx < NumVerts; VertIdx++)
 				{
-					TransformedVerts[VertIdx] = BoneMat.TransformPosition(Collisions[i].BoneVertices[VertIdx]);
+					TransformedVerts[VertIdx] = BoneMat.TransformPosition(Collision.BoneVertices[VertIdx]);
 				}
 
 				// just draw all connected wires to check convex shape
@@ -4069,7 +4068,6 @@ void USkeletalMeshComponent::LoadClothingVisualizationInfo(int32 AssetIndex)
 					verify(NxParameterized::getParamArraySize(*PhysicalMeshParams, "physicalMesh.boneIndices", BoneIndicesCount));
 					check(BoneIndicesCount == BoneWeightsCount);
 
-					NxParameterized::Handle BoneWeightHandle(*PhysicalMeshParams);
 					for (int32 VertexIndex = 0; VertexIndex < VertexCount; ++VertexIndex)
 					{
 						for(uint8 WeightIndex=0; WeightIndex < MaxBoneWeights; ++WeightIndex)
@@ -4348,7 +4346,6 @@ void USkeletalMeshComponent::DrawClothingPhysicalMeshWire(FPrimitiveDrawInterfac
 		FClothVisualizationInfo& VisualInfo = SkeletalMesh->ClothingAssets[AssetIdx].ClothVisualizationInfos[PhysicalMeshLOD];
 
 		uint32 NumPhysicalMeshVerts = VisualInfo.ClothPhysicalMeshVertices.Num();
-		const physx::PxVec3* SimulVertices = NULL;
 
 		bool bUseSimulatedResult = false;
 		// skinning for fixed vertices
diff --git a/Engine/Source/Runtime/Engine/Private/UnrealEngine.cpp b/Engine/Source/Runtime/Engine/Private/UnrealEngine.cpp
--- a/Engine/Source/Runtime/Engine/Private/UnrealEngine.cpp
+++ b/Engine/Source/Runtime/Engine/Private/UnrealEngine.cpp
@@ -7240,8 +7240,7 @@ void UEngine::HandleNetworkFailure(UWorld *World, UNetDriver *NetDriver, ENetwor
 		// If this net driver has already been unregistered with this world, then don't handle it.
 		if (World)
 		{
-			UNetDriver * NetDriver = FindNamedNetDriver(World, NetDriverName);
-			if (!NetDriver)
+			if (!FindNamedNetDriver(World, NetDriverName))
 			{
 				// This netdriver has already been destroyed (probably waiting for GC)
 				return;
diff --git a/Engine/Source/Runtime/Engine/Private/Vehicles/WheeledVehicleMovementComponent.cpp b/Engine/Source/Runtime/Engine/Private/Vehicles/WheeledVehicleMovementComponent.cpp
--- a/Engine/Source/Runtime/Engine/Private/Vehicles/WheeledVehicleMovementComponent.cpp
+++ b/Engine/Source/Runtime/Engine/Private/Vehicles/WheeledVehicleMovementComponent.cpp
@@ -680,9 +680,9 @@ void UWheeledVehicleMovementComponent::UpdateDrag(float DeltaTime)
 		FVector GlobalForwardVector = UpdatedComponent->GetForwardVector();
 		FVector DragVector = -GlobalForwardVector;
 		float SpeedSquared = ForwardSpeed * ForwardSpeed;
-		float DragArea = ChassisHeight * ChassisWidth;
+		float ChassisDragArea = ChassisHeight * ChassisWidth;
 		float AirDensity = 1.25 / (100 * 100 * 100); //kg/cm^3
-		float DragMag = 0.5f * AirDensity * SpeedSquared * DragCoefficient * DragArea;
+		float DragMag = 0.5f * AirDensity * SpeedSquared * DragCoefficient * ChassisDragArea;
 		DebugDragMagnitude = DragMag;
 		DragVector *= DragMag;
 		FBodyInstance * BodyInstance = UpdatedComponent->GetBodyInstance();
diff --git a/Engine/Source/Runtime/Engine/Public/SkeletalMeshTypes.h b/Engine/Source/Runtime/Engine/Public/SkeletalMeshTypes.h
--- a/Engine/Source/Runtime/Engine/Public/SkeletalMeshTypes.h
+++ b/Engine/Source/Runtime/Engine/Public/SkeletalMeshTypes.h
@@ -652,7 +652,7 @@ struct FApexClothCollisionVolumeData
 		LocalPose.SetIdentity();
 	}
 
-	bool IsCapsule()
+	bool IsCapsule() const
 	{
 		return (ConvexVerticesCount == 0);
 	}
diff --git a/Engine/Source/Runtime/GameplayAbilities/Private/AbilitySystemComponent_Abilities.cpp b/Engine/Source/Runtime/GameplayAbilities/Private/AbilitySystemComponent_Abilities.cpp
--- a/Engine/Source/Runtime/GameplayAbilities/Private/AbilitySystemComponent_Abilities.cpp
+++ b/Engine/Source/Runtime/GameplayAbilities/Private/AbilitySystemComponent_Abilities.cpp
@@ -106,14 +106,14 @@ void UAbilitySystemComponent::CancelAbilitiesWithTags(const FGameplayTagContaine
 
 	struct local
 	{
-		static void CancelAbilitiesWithTags(const FGameplayTagContainer Tags, TArray<UGameplayAbility*> Abilities, const FGameplayAbilityActorInfo* ActorInfo, const FGameplayAbilityActivationInfo ActivationInfo, UGameplayAbility * Ignore)
+		static void CancelAbilitiesWithTags(const FGameplayTagContainer InTags, TArray<UGameplayAbility*> Abilities, const FGameplayAbilityActorInfo* InActorInfo, const FGameplayAbilityActivationInfo InActivationInfo, UGameplayAbility* InIgnore)
 		{
 			for (int32 idx=0; idx < Abilities.Num(); ++idx)
 			{
 				UGameplayAbility *Ability = Abilities[idx];
-				if (Ability && (Ability != Ignore) && Ability->AbilityTags.MatchesAny(Tags, false))
+				if (Ability && (Ability != InIgnore) && Ability->AbilityTags.MatchesAny(InTags, false))
 				{
-					Ability->CancelAbility(ActorInfo, ActivationInfo);
+					Ability->CancelAbility(InActorInfo, InActivationInfo);
 					if (!Ability->HasAnyFlags(RF_ClassDefaultObject) && Ability->GetInstancingPolicy() == EGameplayAbilityInstancingPolicy::InstancedPerExecution)
 					{
 						Ability->MarkPendingKill();
diff --git a/Engine/Source/Runtime/GameplayAbilities/Private/GameplayEffect.cpp b/Engine/Source/Runtime/GameplayAbilities/Private/GameplayEffect.cpp
--- a/Engine/Source/Runtime/GameplayAbilities/Private/GameplayEffect.cpp
+++ b/Engine/Source/Runtime/GameplayAbilities/Private/GameplayEffect.cpp
@@ -1419,16 +1419,16 @@ float FActiveGameplayEffectsContainer::GetGameplayEffectDuration(FActiveGameplay
 	return UGameplayEffect::INFINITE_DURATION;
 }
 
-void FActiveGameplayEffectsContainer::OnDurationAggregatorDirty(FAggregator* Aggregator, UAbilitySystemComponent* Owner, FActiveGameplayEffectHandle Handle)
+void FActiveGameplayEffectsContainer::OnDurationAggregatorDirty(FAggregator* Aggregator, UAbilitySystemComponent* InOwner, FActiveGameplayEffectHandle Handle)
 {
 	for (FActiveGameplayEffect& Effect : GameplayEffects)
 	{
 		if (Effect.Handle == Handle)
 		{
-			FTimerManager& TimerManager = Owner->GetWorld()->GetTimerManager();
+			FTimerManager& TimerManager = InOwner->GetWorld()->GetTimerManager();
 			float CurrDuration = Effect.GetDuration();
 			float TimeRemaining = CurrDuration - TimerManager.GetTimerElapsed(Effect.DurationHandle);
-			FTimerDelegate Delegate = FTimerDelegate::CreateUObject(Owner, &UAbilitySystemComponent::CheckDurationExpired, Handle);
+			FTimerDelegate Delegate = FTimerDelegate::CreateUObject(InOwner, &UAbilitySystemComponent::CheckDurationExpired, Handle);
 			TimerManager.SetTimer(Effect.DurationHandle, Delegate, CurrDuration, false, FMath::Max(TimeRemaining, 0.f));
 
 			return;
diff --git a/Engine/Source/Runtime/GameplayTags/Private/GameplayTagsManager.cpp b/Engine/Source/Runtime/GameplayTags/Private/GameplayTagsManager.cpp
--- a/Engine/Source/Runtime/GameplayTags/Private/GameplayTagsManager.cpp
+++ b/Engine/Source/Runtime/GameplayTags/Private/GameplayTagsManager.cpp
@@ -181,9 +181,9 @@ int32 UGameplayTagsManager::InsertTagIntoNodeArray(FName Tag, TWeakPtr<FGameplay
 		TSharedPtr<FGameplayTagNode> TagNode = MakeShareable(new FGameplayTagNode(Tag, ParentNode, CategoryDescription));
 		InsertionIdx = NodeArray.Add(TagNode);
 
-		FGameplayTag Tag = FGameplayTag(TagNode->GetCompleteTag());
-		GameplayTagMap.Add(TagNode->GetCompleteTag(), Tag);
-		GameplayTagNodeMap.Add(Tag, TagNode);
+		FGameplayTag GameplayTag = FGameplayTag(TagNode->GetCompleteTag());
+		GameplayTagMap.Add(TagNode->GetCompleteTag(), GameplayTag);
+		GameplayTagNodeMap.Add(GameplayTag, TagNode);
 	}
 
 	return InsertionIdx;
