Fixing invalid assert. Added a log for when the invalid case happens.

diff --git a/Engine/Source/Runtime/Engine/Private/PackageMapClient.cpp b/Engine/Source/Runtime/Engine/Private/PackageMapClient.cpp
--- a/Engine/Source/Runtime/Engine/Private/PackageMapClient.cpp
+++ b/Engine/Source/Runtime/Engine/Private/PackageMapClient.cpp
@@ -348,7 +348,16 @@ bool UPackageMapClient::SerializeNewActor(FArchive& Ar, class UActorChannel *Cha
 				// Everything worked, assign the guids
 				for ( int32 i = 0; i < Subobjs.Num(); i++ )
 				{
-					check( !GuidCache->NetGUIDLookup.Contains( Subobjs[i] ) );		// We should NOT have a guid yet
+					FNetworkGUID OldNetGUID = GuidCache->NetGUIDLookup.FindRef( Subobjs[i] );
+
+					if ( OldNetGUID.IsValid() )
+					{
+						if ( OldNetGUID != SubObjGuids[i] )
+						{
+							UE_LOG( LogNetPackageMap, Warning, TEXT( "SerializeNewActor: Changing NetGUID on subobject. Name: %s, OldGUID: %s, NewGUID: %s" ), *Subobjs[i]->GetPathName(), *OldNetGUID.ToString(), *SubObjGuids[i].ToString() );
+						}
+					}
+
 					GuidCache->AssignNetGUID( Subobjs[i], SubObjGuids[i] );
 
 					UE_LOG( LogNetPackageMap, Log, TEXT( "SerializeNewActor: [Loading] Assigned NetGUID <%s> to subobject %s" ), *SubObjGuids[i].ToString(), *Subobjs[i]->GetPathName() );
