In the Mac implementation of AVIWriter use an FEvent to wait for [AVAssetWriter finishWritingWithCompletionHandler:] to complete rather than using the deprecated finishWriting version.

diff --git a/Engine/Source/Runtime/Engine/Private/AVIWriter.cpp b/Engine/Source/Runtime/Engine/Private/AVIWriter.cpp
--- a/Engine/Source/Runtime/Engine/Private/AVIWriter.cpp
+++ b/Engine/Source/Runtime/Engine/Private/AVIWriter.cpp
@@ -593,12 +593,15 @@ public:
 		{
 			SCOPED_AUTORELEASE_POOL;
 			[AVFWriterInputRef markAsFinished];
-			// @TODO: Remove usage of deprecated finishWriting selector
-#pragma clang diagnostic push
-#pragma clang diagnostic ignored "-Wdeprecated-declarations"
-			BOOL OK = [AVFWriterRef finishWriting];
-#pragma clang diagnostic pop
-			check(OK);
+			// This will finish asynchronously and then destroy the relevant objects.
+			// We must wait for this to complete.
+			FEvent* Event = FPlatformProcess::CreateSynchEvent(true);
+			[AVFWriterRef finishWritingWithCompletionHandler:^{
+				check(AVFWriterRef.status == AVAssetWriterStatusCompleted);
+				Event->Trigger();
+			}];
+			Event->Wait(~0u);
+			delete Event;
 			[AVFWriterInputRef release];
 			[AVFWriterRef release];
 			[AVFPixelBufferAdaptorRef release];
