UE4 - merge first pass at parallel rendering to main

diff --git a/Engine/Config/ConsoleVariables.ini b/Engine/Config/ConsoleVariables.ini
--- a/Engine/Config/ConsoleVariables.ini
+++ b/Engine/Config/ConsoleVariables.ini
@@ -20,4 +20,5 @@
 [Startup]
 
 ; Uncomment to get detailed logs on shader compiles and the opportunity to retry on errors
-;r.ShaderDevelopmentMode=1
\ No newline at end of file
+r.ShaderDevelopmentMode=1
+r.DumpShaderDebugInfo=1
\ No newline at end of file
diff --git a/Engine/Shaders/DeferredShadingCommon.usf b/Engine/Shaders/DeferredShadingCommon.usf
--- a/Engine/Shaders/DeferredShadingCommon.usf
+++ b/Engine/Shaders/DeferredShadingCommon.usf
@@ -512,17 +512,6 @@ void ApplyDBufferData(FDBufferData DBufferData, inout float3 WorldNormal, inout
 #endif
 }
 
-// Resolved GBuffer textures
-Texture2D GBufferATexture;
-SamplerState GBufferATextureSampler;
-Texture2D GBufferBTexture;
-SamplerState GBufferBTextureSampler;
-Texture2D GBufferCTexture;
-SamplerState GBufferCTextureSampler;
-Texture2D GBufferDTexture;
-SamplerState GBufferDTextureSampler;
-Texture2D GBufferETexture;
-SamplerState GBufferETextureSampler;
 Texture2D DBufferATexture;
 SamplerState DBufferATextureSampler;
 Texture2D DBufferBTexture;
@@ -543,24 +532,6 @@ SamplerState CustomDepthTextureSampler;
 		#define FMultisampledGBuffer Texture2DMS<float4,MSAA_SAMPLE_COUNT>
 	#endif
 
-	/** Parameters bound and set by FDeferredPixelShaderParameters */
-	// Unresolved multisampled GBuffer textures
-	FMultisampledGBuffer GBufferATextureMS;
-	FMultisampledGBuffer GBufferBTextureMS;
-	FMultisampledGBuffer GBufferCTextureMS;
-	FMultisampledGBuffer GBufferDTextureMS;
-	FMultisampledGBuffer GBufferETextureMS;
-	FMultisampledGBuffer DBufferATextureMS;
-	FMultisampledGBuffer DBufferBTextureMS;
-	FMultisampledGBuffer DBufferCTextureMS;
-	FMultisampledGBuffer ScreenSpaceAOTextureMS;
-
-	// can be renamed, but then GBufferATexture should be renamed first
-	Texture2D<float4> GBufferATextureNonMS;
-	Texture2D<float4> GBufferBTextureNonMS;
-	Texture2D<float4> GBufferCTextureNonMS;
-	Texture2D<float4> GBufferDTextureNonMS;
-	Texture2D<float4> GBufferETextureNonMS;
 	Texture2D<float4> DBufferATextureNonMS;
 	Texture2D<float4> DBufferBTextureNonMS;
 	Texture2D<float2> DBufferCTextureNonMS;
@@ -570,14 +541,14 @@ SamplerState CustomDepthTextureSampler;
 	// @param PixelPos - integer pixel pos (from left top)
 	FGBufferData GetGBufferDataUint(uint2 PixelPos, bool bGetNormalizedNormal = true)
 	{
-		float4 GBufferA = GBufferATextureNonMS.Load(int3(PixelPos, 0));
-		float4 GBufferB = GBufferBTextureNonMS.Load(int3(PixelPos, 0));
-		float4 GBufferC = GBufferCTextureNonMS.Load(int3(PixelPos, 0));
-		float4 GBufferD = GBufferDTextureNonMS.Load(int3(PixelPos, 0));
+		float4 GBufferA = GBuffers.GBufferATextureNonMS.Load(int3(PixelPos, 0));
+		float4 GBufferB = GBuffers.GBufferBTextureNonMS.Load(int3(PixelPos, 0));
+		float4 GBufferC = GBuffers.GBufferCTextureNonMS.Load(int3(PixelPos, 0));
+		float4 GBufferD = GBuffers.GBufferDTextureNonMS.Load(int3(PixelPos, 0));
 		float CustomNativeDepth = CustomDepthTextureNonMS.Load(int3(PixelPos, 0)).r;
 
 		#if ALLOW_STATIC_LIGHTING
-			float4 GBufferE = GBufferETextureNonMS.Load(int3(PixelPos, 0));
+			float4 GBufferE = GBuffers.GBufferETextureNonMS.Load(int3(PixelPos, 0));
 		#else
 			float4 GBufferE = 1;
 		#endif
@@ -624,14 +595,14 @@ SamplerState CustomDepthTextureSampler;
 // @param UV - UV space in the GBuffer textures (BufferSize resolution)
 FGBufferData GetGBufferData(float2 UV, bool bGetNormalizedNormal = true)
 {
-	float4 GBufferA = Texture2DSampleLevel(GBufferATexture, GBufferATextureSampler, UV, 0);
-	float4 GBufferB = Texture2DSampleLevel(GBufferBTexture, GBufferBTextureSampler, UV, 0);
-	float4 GBufferC = Texture2DSampleLevel(GBufferCTexture, GBufferCTextureSampler, UV, 0);
-	float4 GBufferD = Texture2DSampleLevel(GBufferDTexture, GBufferDTextureSampler, UV, 0);
+	float4 GBufferA = Texture2DSampleLevel(GBuffers.GBufferATexture, GBuffers.GBufferATextureSampler, UV, 0);
+	float4 GBufferB = Texture2DSampleLevel(GBuffers.GBufferBTexture, GBuffers.GBufferBTextureSampler, UV, 0);
+	float4 GBufferC = Texture2DSampleLevel(GBuffers.GBufferCTexture, GBuffers.GBufferCTextureSampler, UV, 0);
+	float4 GBufferD = Texture2DSampleLevel(GBuffers.GBufferDTexture, GBuffers.GBufferDTextureSampler, UV, 0);
 	float CustomNativeDepth = Texture2DSampleLevel(CustomDepthTexture, CustomDepthTextureSampler, UV, 0).r;
 
 	#if ALLOW_STATIC_LIGHTING
-		float4 GBufferE = Texture2DSampleLevel(GBufferETexture, GBufferETextureSampler, UV, 0);
+		float4 GBufferE = Texture2DSampleLevel(GBuffers.GBufferETexture, GBuffers.GBufferETextureSampler, UV, 0);
 	#else
 		float4 GBufferE = 1;
 	#endif
@@ -644,7 +615,7 @@ FGBufferData GetGBufferData(float2 UV, bool bGetNormalizedNormal = true)
 // Minimal path for just the lighting model, used to branch around unlit pixels (skybox)
 uint GetShadingModelId(float2 UV)
 {
-	float4 GBufferB = Texture2DSampleLevel(GBufferBTexture, GBufferBTextureSampler, UV, 0);
+	float4 GBufferB = Texture2DSampleLevel(GBuffers.GBufferBTexture, GBuffers.GBufferBTextureSampler, UV, 0);
 	return DecodeShadingModelId(GBufferB.a);
 }
 
@@ -675,14 +646,14 @@ FScreenSpaceData GetScreenSpaceData(float2 UV, bool bGetNormalizedNormal = true)
 #if FEATURE_LEVEL >= FEATURE_LEVEL_SM5
 	FGBufferData GetGBufferDataMS(int2 IntUV, uint SampleIndex, bool bGetNormalizedNormal = true)
 	{
-		float4 GBufferA = GBufferATextureMS.Load(IntUV, SampleIndex);
-		float4 GBufferB = GBufferBTextureMS.Load(IntUV, SampleIndex);
-		float4 GBufferC = GBufferCTextureMS.Load(IntUV, SampleIndex);
-		float4 GBufferD = GBufferDTextureMS.Load(IntUV, SampleIndex);
+		float4 GBufferA = GBuffers.GBufferATextureMS.Load(IntUV, SampleIndex);
+		float4 GBufferB = GBuffers.GBufferBTextureMS.Load(IntUV, SampleIndex);
+		float4 GBufferC = GBuffers.GBufferCTextureMS.Load(IntUV, SampleIndex);
+		float4 GBufferD = GBuffers.GBufferDTextureMS.Load(IntUV, SampleIndex);
 		float CustomNativeDepth = CustomDepthTextureNonMS.Load(int3(IntUV, 0)).r;
 
 		#if ALLOW_STATIC_LIGHTING
-			float4 GBufferE = GBufferETextureMS.Load(IntUV, SampleIndex);
+			float4 GBufferE = GBuffers.GBufferETextureMS.Load(IntUV, SampleIndex);
 		#else
 			float4 GBufferE = 1;
 		#endif
@@ -698,7 +669,7 @@ FScreenSpaceData GetScreenSpaceData(float2 UV, bool bGetNormalizedNormal = true)
 		float2 SurfaceDimensions;
 		float NumSurfaceSamples;
 		// assuming all GBuffers share the same size
-		GBufferCTextureMS.GetDimensions(SurfaceDimensions.x, SurfaceDimensions.y, NumSurfaceSamples);
+		GBuffers.GBufferCTextureMS.GetDimensions(SurfaceDimensions.x, SurfaceDimensions.y, NumSurfaceSamples);
 
 		int2 IntUV = (int2)trunc(UV * SurfaceDimensions);
 
diff --git a/Engine/Shaders/PostProcessTonemap.usf b/Engine/Shaders/PostProcessTonemap.usf
--- a/Engine/Shaders/PostProcessTonemap.usf
+++ b/Engine/Shaders/PostProcessTonemap.usf
@@ -9,9 +9,6 @@
 #include "TonemapCommon.usf"
 
 //
-Texture2D BloomDirtMask;
-SamplerState BloomDirtMaskSampler;
-
 // 64x64 tiling noise texture, optimized to have mostly high frequency content
 Texture2D NoiseTexture;
 SamplerState NoiseTextureSampler;
@@ -21,9 +18,6 @@ float4 ColorScale0;
 // xyz:Bloom1Tint.rgb, w:unused
 float4 ColorScale1;
 
-// xyz:BloomDirtMaskTint.rgb, w:ExposureScale (useful if eyeadaptation is locked)
-float4 BloomDirtMaskTint;
-
 // to scale UV (with border) to NormaizedUV (without border)
 float2 TexScale;
 
@@ -208,14 +202,14 @@ void MainPS(
 
 		#if USE_BLOOM
 			float4 CombinedBloom = Texture2DSample(PostprocessInput1, PostprocessInput1Sampler, UV);
-			float3 BloomDirtMaskColor = Texture2DSample(BloomDirtMask, BloomDirtMaskSampler, ScreenSpacePos * float2(0.5, -0.5f) + 0.5f).rgb * BloomDirtMaskTint.rgb; 
-			LinearColor += CombinedBloom.rgb * (ColorScale1.rgb + BloomDirtMaskColor);
+			float3 BloomDirtMaskColor = Texture2DSample(BloomDirtMask.Mask, BloomDirtMask.MaskSampler, ScreenSpacePos * float2(0.5, -0.5f) + 0.5f).rgb * BloomDirtMask.Tint.rgb; 
+			LinearColor += CombinedBloom.rgb * (ColorScale1.rgb + BloomDirtMaskColor); 
 		#endif
 
 		float ExposureScale = InExposureScaleVignette.x;
 
 		#if NO_EYEADAPTATION_EXPOSURE_FIX
-			ExposureScale = BloomDirtMaskTint.w;
+			ExposureScale = BloomDirtMask.Tint.w;
 		#endif
 
 		LinearColor *= ExposureScale;
diff --git a/Engine/Shaders/ShaderVersion.usf b/Engine/Shaders/ShaderVersion.usf
--- a/Engine/Shaders/ShaderVersion.usf
+++ b/Engine/Shaders/ShaderVersion.usf
@@ -2,4 +2,4 @@
 // Each time the console command is executed it generates a new GUID. As this file is included
 // in common.usf (which should be included in any shader) it allows to invalidate the shader DDC.
 // 
-// GUID = {000C0006-0014-0005-0035-02555806A618} test2
\ No newline at end of file
+// GUID = {0C2B0080-9297-4D12-9954-EFA679F1122A}
diff --git a/Engine/Source/Developer/ShaderCompilerCommon/Private/ShaderCompilerCommon.cpp b/Engine/Source/Developer/ShaderCompilerCommon/Private/ShaderCompilerCommon.cpp
new file mode 100644
--- /dev/null
+++ b/Engine/Source/Developer/ShaderCompilerCommon/Private/ShaderCompilerCommon.cpp
@@ -0,0 +1,159 @@
+// Copyright 1998-2014 Epic Games, Inc. All Rights Reserved.
+
+/*=============================================================================
+	ShaderCompilerCommon.cpp: Common functionality for platform Shader Compilers.
+=============================================================================*/
+
+#include "ShaderCompilerCommon.h"
+#include "ModuleManager.h"
+
+IMPLEMENT_MODULE(FDefaultModuleImpl, ShaderCompilerCommon);
+
+void BuildResourceTableTokenStream(const TArray<uint32>& InResourceMap, int32 MaxBoundResourceTable, TArray<uint32>& OutTokenStream)
+{
+	// First we sort the resource map.
+	TArray<uint32> SortedResourceMap = InResourceMap;
+	SortedResourceMap.Sort();
+
+	// The token stream begins with a table that contains offsets per bound uniform buffer.
+	// This offset provides the start of the token stream.
+	OutTokenStream.AddZeroed(MaxBoundResourceTable+1);
+	auto LastBufferIndex = FRHIResourceTableEntry::GetEndOfStreamToken();
+	for (int32 i = 0; i < SortedResourceMap.Num(); ++i)
+	{
+		auto BufferIndex = FRHIResourceTableEntry::GetUniformBufferIndex(SortedResourceMap[i]);
+		if (BufferIndex != LastBufferIndex)
+		{
+			// Store the offset for resources from this buffer.
+			OutTokenStream[BufferIndex] = OutTokenStream.Num();
+			LastBufferIndex = BufferIndex;
+		}
+		OutTokenStream.Add(SortedResourceMap[i]);
+	}
+
+	// Add a token to mark the end of the stream. Not needed if there are no bound resources.
+	if (OutTokenStream.Num())
+	{
+		OutTokenStream.Add(FRHIResourceTableEntry::GetEndOfStreamToken());
+	}
+}
+
+void BuildResourceTableMapping(
+	const TMap<FString,FResourceTableEntry>& ResourceTableMap,
+	const TMap<FString,uint32>& ResourceTableLayoutHashes,
+	TBitArray<>& UsedUniformBufferSlots,
+	FShaderParameterMap& ParameterMap,
+	FShaderResourceTable& OutSRT)
+{
+	check(OutSRT.ResourceTableBits == 0);
+	check(OutSRT.ResourceTableLayoutHashes.Num() == 0);
+
+	// Build resource table mapping
+	int32 MaxBoundResourceTable = -1;
+	TArray<uint32> ResourceTableSRVs;
+	TArray<uint32> ResourceTableSamplerStates;
+	TArray<uint32> ResourceTableUAVs;
+
+	for( auto MapIt = ResourceTableMap.CreateConstIterator(); MapIt; ++MapIt )
+	{
+		const FString& Name	= MapIt->Key;
+		const FResourceTableEntry& Entry = MapIt->Value;
+
+		uint16 BufferIndex, BaseIndex, Size;
+		if (ParameterMap.FindParameterAllocation( *Name, BufferIndex, BaseIndex, Size ) )
+		{
+			ParameterMap.RemoveParameterAllocation(*Name);
+
+			uint16 UniformBufferIndex = INDEX_NONE, UBBaseIndex, UBSize;
+			if (ParameterMap.FindParameterAllocation(*Entry.UniformBufferName, UniformBufferIndex, UBBaseIndex, UBSize) == false)
+			{
+				UniformBufferIndex = UsedUniformBufferSlots.FindAndSetFirstZeroBit();
+				ParameterMap.AddParameterAllocation(*Entry.UniformBufferName,UniformBufferIndex,0,0);
+			}
+
+			OutSRT.ResourceTableBits |= (1 << UniformBufferIndex);
+			MaxBoundResourceTable = FMath::Max<int32>(MaxBoundResourceTable, (int32)UniformBufferIndex);
+
+			while (OutSRT.ResourceTableLayoutHashes.Num() <= MaxBoundResourceTable)
+			{
+				OutSRT.ResourceTableLayoutHashes.Add(0);
+			}
+			OutSRT.ResourceTableLayoutHashes[UniformBufferIndex] = ResourceTableLayoutHashes.FindChecked(Entry.UniformBufferName);
+
+			auto ResourceMap = FRHIResourceTableEntry::Create(UniformBufferIndex, Entry.ResourceIndex, BaseIndex);
+			switch( Entry.Type )
+			{
+			case UBMT_TEXTURE:
+				OutSRT.TextureMap.Add(ResourceMap);
+				break;
+			case UBMT_SAMPLER:
+				OutSRT.SamplerMap.Add(ResourceMap);
+				break;
+			case UBMT_SRV:
+				OutSRT.ShaderResourceViewMap.Add(ResourceMap);
+				break;
+			case UBMT_UAV:
+				OutSRT.UnorderedAccessViewMap.Add(ResourceMap);
+				break;
+			default:
+				check(0);
+			}
+		}
+	}
+
+	OutSRT.MaxBoundResourceTable = MaxBoundResourceTable;
+}
+
+bool RemoveUniformBuffersFromSource(FString& SourceCode)
+{
+	static const FString StaticStructToken(TEXT("static const struct"));
+	int32 StaticStructTokenPos = SourceCode.Find(StaticStructToken, ESearchCase::CaseSensitive, ESearchDir::FromStart);
+	while (StaticStructTokenPos != INDEX_NONE)
+	{
+		static const FString CloseBraceSpaceToken(TEXT("} "));
+		int32 CloseBraceSpaceTokenPos = SourceCode.Find(CloseBraceSpaceToken, ESearchCase::CaseSensitive, ESearchDir::FromStart, StaticStructTokenPos + StaticStructToken.Len());
+		if (CloseBraceSpaceTokenPos == INDEX_NONE)
+		{
+			check(0);	//@todo-rco: ERROR
+			return false;
+		}
+
+		int32 NamePos = CloseBraceSpaceTokenPos + CloseBraceSpaceToken.Len();
+		static const FString SpaceEqualsToken(TEXT(" ="));
+		int32 SpaceEqualsTokenPos = SourceCode.Find(SpaceEqualsToken, ESearchCase::CaseSensitive, ESearchDir::FromStart, NamePos);
+		if (SpaceEqualsTokenPos == INDEX_NONE)
+		{
+			check(0);	//@todo-rco: ERROR
+			return false;
+		}
+
+		FString UniformBufferName = SourceCode.Mid(NamePos, SpaceEqualsTokenPos - NamePos);
+		check(UniformBufferName.Len() > 0);
+
+		static const FString CloseBraceSemicolorToken(TEXT("};"));
+		int32 CloseBraceSemicolonTokenPos = SourceCode.Find(CloseBraceSemicolorToken, ESearchCase::CaseSensitive, ESearchDir::FromStart, SpaceEqualsTokenPos + SpaceEqualsToken.Len());
+		if (CloseBraceSemicolonTokenPos == INDEX_NONE)
+		{
+			check(0);	//@todo-rco: ERROR
+			return false;
+		}
+
+		// Comment out this UB
+		auto& SourceCharArray = SourceCode.GetCharArray();
+		SourceCharArray[StaticStructTokenPos] = TCHAR('/');
+		SourceCharArray[StaticStructTokenPos + 1] = TCHAR('*');
+		SourceCharArray[CloseBraceSemicolonTokenPos] = TCHAR('*');
+		SourceCharArray[CloseBraceSemicolonTokenPos + 1] = TCHAR('/');
+
+		// Find & Replace this UB
+		FString UBSource = UniformBufferName + FString(TEXT("."));
+		FString UBDest = UniformBufferName + FString(TEXT("_"));
+		SourceCode.ReplaceInline(*UBSource, *UBDest, ESearchCase::CaseSensitive);
+
+
+		// Find next UB
+		StaticStructTokenPos = SourceCode.Find(StaticStructToken, ESearchCase::CaseSensitive, ESearchDir::FromStart, CloseBraceSemicolonTokenPos + 2);
+	}
+
+	return true;
+}
diff --git a/Engine/Source/Developer/ShaderCompilerCommon/Public/ShaderCompilerCommon.h b/Engine/Source/Developer/ShaderCompilerCommon/Public/ShaderCompilerCommon.h
new file mode 100644
--- /dev/null
+++ b/Engine/Source/Developer/ShaderCompilerCommon/Public/ShaderCompilerCommon.h
@@ -0,0 +1,62 @@
+// Copyright 1998-2014 Epic Games, Inc. All Rights Reserved.
+
+/*=============================================================================
+	ShaderCompilerCommon.cpp: Common functionality for platform Shader Compilers.
+=============================================================================*/
+
+#pragma once
+
+#include "Core.h"
+#include "ShaderCore.h"
+
+/**
+ * This function looks for resources specified in ResourceTableMap in the 
+ * parameter map, adds them to the resource table, and removes them from the
+ * parameter map. If a resource is used from a currently unmapped uniform 
+ * buffer we allocate a slot for it from UsedUniformBufferSlots.
+ */
+extern SHADERCOMPILERCOMMON_API void BuildResourceTableMapping(
+	const TMap<FString,FResourceTableEntry>& ResourceTableMap,
+	const TMap<FString,uint32>& ResourceTableLayoutHashes,
+	TBitArray<>& UsedUniformBufferSlots,
+	FShaderParameterMap& ParameterMap,
+	FShaderResourceTable& OutSRT
+	);
+
+/**
+ * Builds a token stream out of the resource map. The resource map is one
+ * of the arrays generated by BuildResourceTableMapping. The token stream
+ * is used at runtime to gather resources from tables and bind them to the
+ * appropriate slots.
+ */
+extern SHADERCOMPILERCOMMON_API void BuildResourceTableTokenStream(
+	const TArray<uint32>& InResourceMap,
+	int32 MaxBoundResourceTable,
+	TArray<uint32>& OutTokenStream
+	);
+
+// This function goes through the shader source and converts
+//		static struct
+//		{
+//		TYPE Member;
+//		[...]
+//		} UniformBuffer = { [...] };
+//		
+//		TYPE Function([...])
+//		{
+//			[...]UniformBuffer.Member[...]
+//		}
+//
+//into
+//
+//		/*atic struct
+//		{
+//		TYPE Member;
+//		[...]
+//		} UniformBuffer = { [...] */
+//		
+//		TYPE Function([...])
+//		{
+//		[...]UniformBuffer_Member[...]
+//		}
+extern SHADERCOMPILERCOMMON_API bool RemoveUniformBuffersFromSource(FString& SourceCode);
diff --git a/Engine/Source/Developer/ShaderCompilerCommon/ShaderCompilerCommon.Build.cs b/Engine/Source/Developer/ShaderCompilerCommon/ShaderCompilerCommon.Build.cs
new file mode 100644
--- /dev/null
+++ b/Engine/Source/Developer/ShaderCompilerCommon/ShaderCompilerCommon.Build.cs
@@ -0,0 +1,16 @@
+// Copyright 1998-2014 Epic Games, Inc. All Rights Reserved.
+
+using UnrealBuildTool;
+
+public class ShaderCompilerCommon : ModuleRules
+{
+	public ShaderCompilerCommon(TargetInfo Target)
+	{
+		PrivateDependencyModuleNames.AddRange(
+			new string[] {
+				"Core",
+				"ShaderCore",
+			}
+			);
+	}
+}
diff --git a/Engine/Source/Developer/ShaderFormatOpenGL/Private/OpenGLShaderCompiler.cpp b/Engine/Source/Developer/ShaderFormatOpenGL/Private/OpenGLShaderCompiler.cpp
--- a/Engine/Source/Developer/ShaderFormatOpenGL/Private/OpenGLShaderCompiler.cpp
+++ b/Engine/Source/Developer/ShaderFormatOpenGL/Private/OpenGLShaderCompiler.cpp
@@ -1,4 +1,4 @@
-// Copyright 1998-2014 Epic Games, Inc. All Rights Reserved..
+// Copyright 1998-2014 Epic Games, Inc. All Rights Reserved.
 // 
 #include "ShaderFormatOpenGL.h"
 #include "Core.h"
@@ -13,7 +13,8 @@
 	#include "Windows/MinWindows.h"
 #include "HideWindowsPlatformTypes.h"
 #endif
-#include "ShaderPreprocessor.h" 
+#include "ShaderPreprocessor.h"
+#include "ShaderCompilerCommon.h"
 #include "hlslcc.h"
 #include "glsl/ir_gen_glsl.h"
 #if PLATFORM_WINDOWS
@@ -44,7 +45,7 @@ DEFINE_LOG_CATEGORY_STATIC(LogOpenGLShaderCompiler, Log, All);
 
 
 #define ENABLE_IMAGINATION_COMPILER		1
- 
+
 #define MAX_SAMPLERS_PER_SHADER_GLSL_150 16
 #define MAX_SAMPLERS_PER_SHADER_GLSL_430 32
 
@@ -90,7 +91,7 @@ struct FPlatformOpenGLContext
 };
 
 /**
- * A dummy wndproc. 
+ * A dummy wndproc.
  */
 static LRESULT CALLBACK PlatformDummyGLWndproc(HWND hWnd, uint32 Message, WPARAM wParam, LPARAM lParam)
 {
@@ -414,6 +415,7 @@ static uint32 ParseNumber(const ANSICHAR* &Str)
  */
 static void BuildShaderOutput(
 	FShaderCompilerOutput& ShaderOutput,
+	const FShaderCompilerInput& ShaderInput, 
 	const ANSICHAR* InShaderSource,
 	int32 SourceLen,
 	GLSLVersion Version
@@ -424,6 +426,9 @@ static void BuildShaderOutput(
 	FShaderParameterMap& ParameterMap = ShaderOutput.ParameterMap;
 	EShaderFrequency Frequency = (EShaderFrequency)ShaderOutput.Target.Frequency;
 
+	TBitArray<> UsedUniformBufferSlots;
+	UsedUniformBufferSlots.Init(false,32);
+
 	// Write out the magic markers.
 	Header.GlslMarker = 0x474c534c;
 	switch (Frequency)
@@ -579,6 +584,7 @@ static void BuildShaderOutput(
 			verify(Match(ShaderSource, '('));
 			uint16 UBIndex = ParseNumber(ShaderSource);
 			check(UBIndex == Header.Bindings.NumUniformBuffers);
+			UsedUniformBufferSlots[UBIndex] = true;
 			verify(Match(ShaderSource, ')'));
 			ParameterMap.AddParameterAllocation(*BufferName, Header.Bindings.NumUniformBuffers++, 0, 0);
 
@@ -958,6 +964,24 @@ static void BuildShaderOutput(
 		}
 	}
 
+	// Build the SRT for this shader.
+	{
+		// Build the generic SRT for this shader.
+		FShaderResourceTable GenericSRT;
+		BuildResourceTableMapping(ShaderInput.Environment.ResourceTableMap, ShaderInput.Environment.ResourceTableLayoutHashes, UsedUniformBufferSlots, ShaderOutput.ParameterMap, GenericSRT);
+
+		// Copy over the bits indicating which resource tables are active.
+		Header.Bindings.ShaderResourceTable.ResourceTableBits = GenericSRT.ResourceTableBits;
+
+		Header.Bindings.ShaderResourceTable.ResourceTableLayoutHashes = GenericSRT.ResourceTableLayoutHashes;
+
+		// Now build our token streams.
+		BuildResourceTableTokenStream(GenericSRT.TextureMap, GenericSRT.MaxBoundResourceTable, Header.Bindings.ShaderResourceTable.TextureMap);
+		BuildResourceTableTokenStream(GenericSRT.ShaderResourceViewMap, GenericSRT.MaxBoundResourceTable, Header.Bindings.ShaderResourceTable.ShaderResourceViewMap);
+		BuildResourceTableTokenStream(GenericSRT.SamplerMap, GenericSRT.MaxBoundResourceTable, Header.Bindings.ShaderResourceTable.SamplerMap);
+		BuildResourceTableTokenStream(GenericSRT.UnorderedAccessViewMap, GenericSRT.MaxBoundResourceTable, Header.Bindings.ShaderResourceTable.UnorderedAccessViewMap);
+	}
+
 	const int32 MaxSamplers = GetFeatureLevelMaxTextureSamplers(GetMaxSupportedFeatureLevel((EShaderPlatform)ShaderOutput.Target.Platform));
 
 	if (Header.Bindings.NumSamplers > MaxSamplers)
@@ -1109,7 +1133,7 @@ static void PrecompileGLSLES2(FShaderCompilerOutput& ShaderOutput, const FShader
 				ShaderOutput.bSucceeded = true;
 				ShaderOutput.Target = ShaderInput.Target;
 
-				BuildShaderOutput(ShaderOutput, ShaderSource, SourceLen, GLSL_ES2);
+				BuildShaderOutput(ShaderOutput, ShaderInput, ShaderSource, SourceLen, GLSL_ES2);
 
 				// Parse the cycle count
 				const int32 CycleCountStringLength = FPlatformString::Strlen(TEXT("Cycle count: "));
@@ -1140,7 +1164,7 @@ static void PrecompileGLSLES2(FShaderCompilerOutput& ShaderOutput, const FShader
 			ShaderOutput.bSucceeded = true;
 			ShaderOutput.Target = ShaderInput.Target;
 
-			BuildShaderOutput(ShaderOutput, ShaderSource, SourceLen, GLSL_ES2);
+			BuildShaderOutput(ShaderOutput, ShaderInput, ShaderSource, SourceLen, GLSL_ES2);
 		}
 	}
 	else
@@ -1148,7 +1172,7 @@ static void PrecompileGLSLES2(FShaderCompilerOutput& ShaderOutput, const FShader
 		ShaderOutput.bSucceeded = true;
 		ShaderOutput.Target = ShaderInput.Target;
 
-		BuildShaderOutput(ShaderOutput, ShaderSource, SourceLen, GLSL_ES2);
+		BuildShaderOutput(ShaderOutput, ShaderInput, ShaderSource, SourceLen, GLSL_ES2);
 	}
 }
 
@@ -1203,6 +1227,7 @@ static void PrecompileShader(FShaderCompilerOutput& ShaderOutput, const FShaderC
 				ShaderOutput.Target = ShaderInput.Target;
 				BuildShaderOutput(
 					ShaderOutput,
+					ShaderInput, 
 					ShaderSource,
 					(int32)SourceLen,
 					Version
@@ -1450,6 +1475,14 @@ void CompileShader_Windows_OGL(const FShaderCompilerInput& Input,FShaderCompiler
 			return;
 		}
 
+/*
+		//@todo-rco: Remove me!
+		if (!RemoveUniformBuffersFromSource(PreprocessedShader))
+		{
+			check(0);
+		}
+
+*/
 		// Write out the preprocessed file and a batch file to compile it if requested (DumpDebugInfoPath is valid)
 		if (bDumpDebugInfo)
 		{
@@ -1485,15 +1518,19 @@ void CompileShader_Windows_OGL(const FShaderCompilerInput& Input,FShaderCompiler
 			}
 		}
 
+/*
+//@todo-rco: Remove me!
+CCFlags &= ~HLSLCC_NoPreprocess;
+*/
 		FGlslCodeBackend GlslBackEnd(CCFlags);
-		FGlslLanguageSpec GLslLanguageSpec(IsES2Platform(Version) && !IsPCES2Platform(Version));
+		FGlslLanguageSpec GlslLanguageSpec(IsES2Platform(Version) && !IsPCES2Platform(Version));
 		int32 Result = HlslCrossCompile(
 			TCHAR_TO_ANSI(*Input.SourceFilename),
 			TCHAR_TO_ANSI(*PreprocessedShader),
 			TCHAR_TO_ANSI(*Input.EntryPointName),
 			Frequency,
 			&GlslBackEnd,
-			&GLslLanguageSpec,
+			&GlslLanguageSpec,
 			CCFlags,
 			HlslCompilerTarget,
 			&GlslShaderSource,
@@ -1546,11 +1583,22 @@ void CompileShader_Windows_OGL(const FShaderCompilerInput& Input,FShaderCompiler
 #else // VALIDATE_GLSL_WITH_DRIVER
 			int32 SourceLen = FCStringAnsi::Strlen(GlslShaderSource);
 			Output.Target = Input.Target;
-			BuildShaderOutput(Output, GlslShaderSource, SourceLen, Version);
+				BuildShaderOutput(Output, Input, GlslShaderSource, SourceLen, Version);
 #endif // VALIDATE_GLSL_WITH_DRIVER
 		}
 		else
 		{
+			if (bDumpDebugInfo)
+			{
+				// Generate the batch file to help track down cross-compiler issues if necessary
+				const FString GLSLFile = (Input.DumpDebugInfoPath / TEXT("Output.glsl"));
+				const FString GLBatchFileContents = CreateCommandLineGLSLES2(GLSLFile, (Input.DumpDebugInfoPath / TEXT("Output.asm")), Version, Frequency, false);
+				if (!GLBatchFileContents.IsEmpty())
+				{
+					FFileHelper::SaveStringToFile(GLBatchFileContents, *(Input.DumpDebugInfoPath / TEXT("GLSLCompile.bat")));
+				}
+			}
+
 			FString Tmp = ANSI_TO_TCHAR(ErrorLog);
 			TArray<FString> ErrorLines;
 			Tmp.ParseIntoArray(&ErrorLines, TEXT("\n"), true);
diff --git a/Engine/Source/Developer/ShaderFormatOpenGL/Private/ShaderFormatOpenGL.cpp b/Engine/Source/Developer/ShaderFormatOpenGL/Private/ShaderFormatOpenGL.cpp
--- a/Engine/Source/Developer/ShaderFormatOpenGL/Private/ShaderFormatOpenGL.cpp
+++ b/Engine/Source/Developer/ShaderFormatOpenGL/Private/ShaderFormatOpenGL.cpp
@@ -21,13 +21,13 @@ class FShaderFormatGLSL : public IShaderFormat
 	enum 
 	{
 		/** Version for shader format, this becomes part of the DDC key. */
-		UE_SHADER_GLSL_150_VER = 21,
-		UE_SHADER_GLSL_150_MAC_VER = 1,
-		UE_SHADER_GLSL_430_VER = 51,
-		UE_SHADER_GLSL_ES2_VER = 10,
-		UE_SHADER_GLSL_150ES2_VER = 12,
-		UE_SHADER_GLSL_ES2_VER_WEBGL  = 11,
-		UE_SHADER_GLSL_ES2_IOS_VER  = 1,
+		UE_SHADER_GLSL_150_VER = 23,
+		UE_SHADER_GLSL_150_MAC_VER = 2,
+		UE_SHADER_GLSL_430_VER = 53,
+		UE_SHADER_GLSL_ES2_VER = 12,
+		UE_SHADER_GLSL_150ES2_VER = 13,
+		UE_SHADER_GLSL_ES2_VER_WEBGL  = 13,
+		UE_SHADER_GLSL_ES2_IOS_VER  = 3,
 	}; 
 
 	void CheckFormat(FName Format) const
@@ -37,7 +37,7 @@ class FShaderFormatGLSL : public IShaderFormat
 				Format == NAME_GLSL_430 || 
 				Format == NAME_GLSL_ES2 || 
 				Format == NAME_GLSL_150_ES2 ||
-				Format == NAME_GLSL_ES2_WEBGL ||
+                Format == NAME_GLSL_ES2_WEBGL ||
 				Format == NAME_GLSL_ES2_IOS
 			);
 	}
diff --git a/Engine/Source/Developer/ShaderFormatOpenGL/ShaderFormatOpenGL.Build.cs b/Engine/Source/Developer/ShaderFormatOpenGL/ShaderFormatOpenGL.Build.cs
--- a/Engine/Source/Developer/ShaderFormatOpenGL/ShaderFormatOpenGL.Build.cs
+++ b/Engine/Source/Developer/ShaderFormatOpenGL/ShaderFormatOpenGL.Build.cs
@@ -14,6 +14,7 @@ public class ShaderFormatOpenGL : ModuleRules
 			new string[] {
 				"Core",
 				"ShaderCore",
+				"ShaderCompilerCommon",
 				"ShaderPreprocessor"
 			}
 			);
diff --git a/Engine/Source/Developer/Windows/ShaderFormatD3D/Private/D3D11ShaderCompiler.cpp b/Engine/Source/Developer/Windows/ShaderFormatD3D/Private/D3D11ShaderCompiler.cpp
--- a/Engine/Source/Developer/Windows/ShaderFormatD3D/Private/D3D11ShaderCompiler.cpp
+++ b/Engine/Source/Developer/Windows/ShaderFormatD3D/Private/D3D11ShaderCompiler.cpp
@@ -2,6 +2,8 @@
 
 #include "ShaderFormatD3D.h"
 #include "ShaderPreprocessor.h"
+#include "ShaderCompilerCommon.h"
+#include "D3D11ShaderResources.h"
 
 DEFINE_LOG_CATEGORY_STATIC(LogD3D11ShaderCompiler, Log, All);
 
@@ -470,13 +472,8 @@ void CompileD3D11Shader(const FShaderCompilerInput& Input,FShaderCompilerOutput&
 		{
 			Output.bSucceeded = true;
 
-			uint32 NumShaderBytes = Shader->GetBufferSize();
-			Output.Code.Empty(NumShaderBytes);
-			Output.Code.AddUninitialized(NumShaderBytes);
-			FMemory::Memcpy(Output.Code.GetTypedData(),Shader->GetBufferPointer(),NumShaderBytes);
-
 			ID3D11ShaderReflection* Reflector = NULL;
-			Result = D3DReflect(Output.Code.GetData(),Output.Code.Num(),IID_ID3D11ShaderReflection, (void**) &Reflector);
+			Result = D3DReflect( Shader->GetBufferPointer(), Shader->GetBufferSize(), IID_ID3D11ShaderReflection, (void**) &Reflector );
 			if(FAILED(Result))
 			{
 				UE_LOG(LogD3D11ShaderCompiler, Fatal,TEXT("D3DReflect failed: Result=%08x"),Result);
@@ -489,6 +486,9 @@ void CompileD3D11Shader(const FShaderCompilerInput& Input,FShaderCompilerOutput&
 			bool bGlobalUniformBufferUsed = false;
 			uint32 NumSamplers = 0;
 
+			TBitArray<> UsedUniformBufferSlots;
+			UsedUniformBufferSlots.Init(false,32);
+
 			if (Input.Target.Frequency == SF_Vertex)
 			{
 				for (uint32 Index = 0; Index < ShaderDesc.OutputParameters; ++Index)
@@ -515,14 +515,10 @@ void CompileD3D11Shader(const FShaderCompilerInput& Input,FShaderCompilerOutput&
 					ID3D11ShaderReflectionConstantBuffer* ConstantBuffer = Reflector->GetConstantBufferByName(BindDesc.Name);
 					D3D11_SHADER_BUFFER_DESC CBDesc;
 					ConstantBuffer->GetDesc(&CBDesc);
+					bool bGlobalCB = (FCStringAnsi::Strcmp(CBDesc.Name, "$Globals") == 0);
 
-#if 0
-					if (CBDesc.Size > GConstantBufferSizes[CBIndex])
+					if (bGlobalCB)
 					{
-						UE_LOG(LogD3D11ShaderCompiler, Fatal,TEXT("Set GConstantBufferSizes[%d] to >= %d"), CBIndex, CBDesc.Size);
-					}
-#endif
-
 					// Track all of the variables in this constant buffer.
 					for (uint32 ConstantIndex = 0; ConstantIndex < CBDesc.Variables; ConstantIndex++)
 					{
@@ -531,10 +527,7 @@ void CompileD3D11Shader(const FShaderCompilerInput& Input,FShaderCompilerOutput&
 						Variable->GetDesc(&VariableDesc);
 						if (VariableDesc.uFlags & D3D10_SVF_USED)
 						{
-							if (FCStringAnsi::Strcmp(CBDesc.Name, "$Globals") == 0)
-							{
 								bGlobalUniformBufferUsed = true;
-							}
 
 							Output.ParameterMap.AddParameterAllocation(
 								ANSI_TO_TCHAR(VariableDesc.Name),
@@ -542,8 +535,21 @@ void CompileD3D11Shader(const FShaderCompilerInput& Input,FShaderCompilerOutput&
 								VariableDesc.StartOffset,
 								VariableDesc.Size
 								);
+								UsedUniformBufferSlots[CBIndex] = true;
 						}
 					}
+				}
+					else
+					{
+						// Track just the constant buffer itself.
+						Output.ParameterMap.AddParameterAllocation(
+							ANSI_TO_TCHAR(CBDesc.Name),
+							CBIndex,
+							0,
+							0
+							);
+						UsedUniformBufferSlots[CBIndex] = true;
+					}
 				}
 				else if (BindDesc.Type == D3D10_SIT_TEXTURE || BindDesc.Type == D3D10_SIT_SAMPLER)
 				{
@@ -604,16 +610,16 @@ void CompileD3D11Shader(const FShaderCompilerInput& Input,FShaderCompilerOutput&
 						 BindDesc.Type == D3D11_SIT_UAV_RWBYTEADDRESS || BindDesc.Type == D3D11_SIT_UAV_RWSTRUCTURED_WITH_COUNTER ||
 						 BindDesc.Type == D3D11_SIT_UAV_APPEND_STRUCTURED)
 				{
-					TCHAR OfficialName[1024];
-					FCString::Strcpy(OfficialName, ANSI_TO_TCHAR(BindDesc.Name));
-
-					Output.ParameterMap.AddParameterAllocation(
-						OfficialName,
-						0,
-						BindDesc.BindPoint,
-						1
-						);
-				}
+						TCHAR OfficialName[1024];
+						FCString::Strcpy(OfficialName, ANSI_TO_TCHAR(BindDesc.Name));
+
+						Output.ParameterMap.AddParameterAllocation(
+							OfficialName,
+							0,
+							BindDesc.BindPoint,
+							1
+							);
+					}
 				else if (BindDesc.Type == D3D11_SIT_STRUCTURED || BindDesc.Type == D3D11_SIT_BYTEADDRESS)
 				{
 					TCHAR OfficialName[1024];
@@ -631,10 +637,10 @@ void CompileD3D11Shader(const FShaderCompilerInput& Input,FShaderCompilerOutput&
 			// Strip shader reflection
 			D3D_SHADER_DATA ShaderData;
 			ShaderData.pBytecode = Shader->GetBufferPointer();
-			ShaderData.BytecodeLength = NumShaderBytes;
+			ShaderData.BytecodeLength = Shader->GetBufferSize();
 			TRefCountPtr<ID3DBlob> CompressedData;
 			Result = D3DStripShader(Shader->GetBufferPointer(),
-				NumShaderBytes,
+				Shader->GetBufferSize(),
 				D3DCOMPILER_STRIP_REFLECTION_DATA | D3DCOMPILER_STRIP_DEBUG_INFO | D3DCOMPILER_STRIP_TEST_BLOBS,
 				CompressedData.GetInitReference());
 
@@ -643,16 +649,34 @@ void CompileD3D11Shader(const FShaderCompilerInput& Input,FShaderCompilerOutput&
 				UE_LOG(LogD3D11ShaderCompiler, Fatal,TEXT("D3DStripShader failed: Result=%08x"),Result);
 			}
 
-			// Copy stripped data
-			NumShaderBytes = CompressedData->GetBufferSize();
-			Output.Code.Empty(NumShaderBytes + 1);
+			// Build the SRT for this shader.
+			FD3D11ShaderResourceTable SRT;
+			{
+				// Build the generic SRT for this shader.
+				FShaderResourceTable GenericSRT;
+				BuildResourceTableMapping(Input.Environment.ResourceTableMap, Input.Environment.ResourceTableLayoutHashes, UsedUniformBufferSlots, Output.ParameterMap, GenericSRT);
+
+				// At runtime textures are just SRVs, so combine them for the purposes of building token streams.
+				GenericSRT.ShaderResourceViewMap.Append(GenericSRT.TextureMap);
+
+				// Copy over the bits indicating which resource tables are active.
+				SRT.ResourceTableBits = GenericSRT.ResourceTableBits;
 
-			// Allocate room for the bytecode and bGlobalUniformBufferUsed
-			Output.Code.AddUninitialized(NumShaderBytes + 1);
-			FMemory::Memcpy(Output.Code.GetTypedData(), CompressedData->GetBufferPointer(), NumShaderBytes);
+				SRT.ResourceTableLayoutHashes = GenericSRT.ResourceTableLayoutHashes;
+
+				// Now build our token streams.
+				BuildResourceTableTokenStream(GenericSRT.ShaderResourceViewMap, GenericSRT.MaxBoundResourceTable, SRT.ShaderResourceViewMap);
+				BuildResourceTableTokenStream(GenericSRT.SamplerMap, GenericSRT.MaxBoundResourceTable, SRT.SamplerMap);
+				BuildResourceTableTokenStream(GenericSRT.UnorderedAccessViewMap, GenericSRT.MaxBoundResourceTable, SRT.UnorderedAccessViewMap);
+
+			}
+
+			FMemoryWriter Ar( Output.Code, true );
+			Ar << SRT;
+			Ar.Serialize( CompressedData->GetBufferPointer(), CompressedData->GetBufferSize() );
 
 			// Pack bGlobalUniformBufferUsed in the last byte
-			Output.Code[NumShaderBytes] = bGlobalUniformBufferUsed;
+			Output.Code.Add( bGlobalUniformBufferUsed );
 
 			// Set the number of instructions.
 			Output.NumInstructions = ShaderDesc.InstructionCount;
@@ -754,4 +778,4 @@ void CompileShader_Windows_ES2(const FShaderCompilerInput& Input,FShaderCompiler
 	FShaderCompilerDefinitions AdditionalDefines;
 	AdditionalDefines.SetDefine(TEXT("ES2_PROFILE"), 1);
 	CompileD3D11Shader(Input, Output, AdditionalDefines, WorkingDirectory);
-}
+}
\ No newline at end of file
diff --git a/Engine/Source/Developer/Windows/ShaderFormatD3D/Private/ShaderFormatD3D.cpp b/Engine/Source/Developer/Windows/ShaderFormatD3D/Private/ShaderFormatD3D.cpp
--- a/Engine/Source/Developer/Windows/ShaderFormatD3D/Private/ShaderFormatD3D.cpp
+++ b/Engine/Source/Developer/Windows/ShaderFormatD3D/Private/ShaderFormatD3D.cpp
@@ -15,9 +15,9 @@ class FShaderFormatD3D : public IShaderFormat
 	enum
 	{
 		/** Version for shader format, this becomes part of the DDC key. */
-		UE_SHADER_PCD3D_SM5_VER = 0,
-		UE_SHADER_PCD3D_SM4_VER = 0,
-		UE_SHADER_PCD3D_ES2_VER = 0
+		UE_SHADER_PCD3D_SM5_VER = 2,
+		UE_SHADER_PCD3D_SM4_VER = 2,
+		UE_SHADER_PCD3D_ES2_VER = 2
 	};
 
 	void CheckFormat(FName Format) const
diff --git a/Engine/Source/Developer/Windows/ShaderFormatD3D/ShaderFormatD3D.Build.cs b/Engine/Source/Developer/Windows/ShaderFormatD3D/ShaderFormatD3D.Build.cs
--- a/Engine/Source/Developer/Windows/ShaderFormatD3D/ShaderFormatD3D.Build.cs
+++ b/Engine/Source/Developer/Windows/ShaderFormatD3D/ShaderFormatD3D.Build.cs
@@ -7,12 +7,14 @@ public class ShaderFormatD3D : ModuleRules
 	public ShaderFormatD3D(TargetInfo Target)
 	{
 		PrivateIncludePathModuleNames.Add("TargetPlatform");
+		PrivateIncludePathModuleNames.Add("D3D11RHI"); 
 
 		PrivateDependencyModuleNames.AddRange(
 			new string[] {
 				"Core",
 				"ShaderCore",
 				"ShaderPreprocessor",
+				"ShaderCompilerCommon",
 			}
 			);
 
diff --git a/Engine/Source/Editor/MeshPaint/Private/MeshPaintEdMode.cpp b/Engine/Source/Editor/MeshPaint/Private/MeshPaintEdMode.cpp
--- a/Engine/Source/Editor/MeshPaint/Private/MeshPaintEdMode.cpp
+++ b/Engine/Source/Editor/MeshPaint/Private/MeshPaintEdMode.cpp
@@ -47,9 +47,9 @@ class FMeshPaintBatchedElementParameters : public FBatchedElementParameters
 public:
 
 	/** Binds vertex and pixel shaders for this element */
-	virtual void BindShaders_RenderThread( const FMatrix& InTransform, const float InGamma, const FMatrix& ColorWeights, const FTexture* Texture )
+	virtual void BindShaders_RenderThread(FRHICommandList* RHICmdList, const FMatrix& InTransform, const float InGamma, const FMatrix& ColorWeights, const FTexture* Texture ) OVERRIDE
 	{
-		MeshPaintRendering::SetMeshPaintShaders_RenderThread( InTransform, InGamma, ShaderParams );
+		MeshPaintRendering::SetMeshPaintShaders_RenderThread(RHICmdList, InTransform, InGamma, ShaderParams );
 	}
 
 public:
@@ -66,9 +66,9 @@ class FMeshPaintDilateBatchedElementParameters : public FBatchedElementParameter
 public:
 
 	/** Binds vertex and pixel shaders for this element */
-	virtual void BindShaders_RenderThread( const FMatrix& InTransform, const float InGamma, const FMatrix& ColorWeights, const FTexture* Texture )
+	virtual void BindShaders_RenderThread(FRHICommandList* RHICmdList, const FMatrix& InTransform, const float InGamma, const FMatrix& ColorWeights, const FTexture* Texture ) OVERRIDE
 	{
-		MeshPaintRendering::SetMeshPaintDilateShaders_RenderThread( InTransform, InGamma, ShaderParams );
+		MeshPaintRendering::SetMeshPaintDilateShaders_RenderThread(RHICmdList, InTransform, InGamma, ShaderParams );
 	}
 
 public:
@@ -302,26 +302,8 @@ bool FEdModeMeshPaint::CapturedMouseMove( FLevelEditorViewportClient* InViewport
 
 			bool bAnyPaintAbleActorsUnderCursor = false;
 
-			bool bIsTexturePaintMode = FMeshPaintSettings::Get().ResourceType == EMeshPaintResource::Texture;
-			if( bIsTexturePaintMode )
-			{
-				ENQUEUE_UNIQUE_RENDER_COMMAND(
-					TexturePaintBeginSceneCommand,
-				{
-					RHIBeginScene();
-				});
-			}
-		
 			DoPaint( View->ViewMatrices.ViewOrigin, MouseViewportRay.GetOrigin(), MouseViewportRay.GetDirection(), NULL, PaintAction, bVisualCueOnly, StrengthScale, bAnyPaintAbleActorsUnderCursor );
 
-			if( bIsTexturePaintMode )
-			{
-				ENQUEUE_UNIQUE_RENDER_COMMAND(
-					TexturePaintEndSceneCommand,
-				{
-					RHIEndScene();
-				});
-			}
 			return true;
 		}
 	}
diff --git a/Engine/Source/Editor/UnrealEd/Private/CubemapUnwrapUtils.cpp b/Engine/Source/Editor/UnrealEd/Private/CubemapUnwrapUtils.cpp
--- a/Engine/Source/Editor/UnrealEd/Private/CubemapUnwrapUtils.cpp
+++ b/Engine/Source/Editor/UnrealEd/Private/CubemapUnwrapUtils.cpp
@@ -5,7 +5,9 @@
 =============================================================================*/
 
 #include "UnrealEd.h"
+#include "ShaderParameterUtils.h"
 #include "CubemapUnwrapUtils.h" 
+#include "RHIStaticStates.h"
 
 IMPLEMENT_SHADER_TYPE(,FCubemapTexturePropertiesVS,TEXT("SimpleElementVertexShader"),TEXT("Main"),SF_Vertex);
 IMPLEMENT_SHADER_TYPE(template<>,FCubemapTexturePropertiesPS<false>,TEXT("SimpleElementPixelShader"),TEXT("CubemapTextureProperties"),SF_Pixel);
@@ -46,12 +48,6 @@ namespace CubemapHelpers
 		FCanvas* Canvas = new FCanvas(RenderTarget, NULL, 0, 0, 0);
 		Canvas->SetRenderTarget(RenderTarget);
 
-		ENQUEUE_UNIQUE_RENDER_COMMAND(
-			LongLatUnwrapBeginSceneCommand,
-		{
-			RHIBeginScene();
-		});
-
 		// Clear the render target to black
 		Canvas->Clear(FLinearColor(0, 0, 0, 0));
 
@@ -64,11 +60,6 @@ namespace CubemapHelpers
 		FlushRenderingCommands();
 		Canvas->SetRenderTarget(NULL);
 		FlushRenderingCommands();
-		ENQUEUE_UNIQUE_RENDER_COMMAND(
-			LongLatUnwrapEndSceneCommand,
-		{
-			RHIEndScene();
-		});
 		
 		int32 ImageBytes = CalculateImageBytes(LongLatDimensions.X, LongLatDimensions.Y, 0, TargetPixelFormat);
 
@@ -116,3 +107,58 @@ namespace CubemapHelpers
 		return GenerateLongLatUnwrap(CubeTarget->Resource, CubeTarget->SizeX, CubeTarget->GetFormat(), BitsOUT, SizeOUT, FormatOUT);
 	}
 }
+
+void FCubemapTexturePropertiesVS::SetParameters( FRHICommandList* RHICmdList, const FMatrix& TransformValue )
+{
+	SetShaderValue(RHICmdList, GetVertexShader(), Transform, TransformValue);
+}
+
+template<bool bHDROutput>
+void FCubemapTexturePropertiesPS<bHDROutput>::SetParameters( FRHICommandList* RHICmdList, const FTexture* Texture, const FMatrix& ColorWeightsValue, float MipLevel, float GammaValue )
+{
+	SetTextureParameter(RHICmdList, GetPixelShader(),CubeTexture,CubeTextureSampler,Texture);
+
+	FVector4 PackedProperties0Value(MipLevel, 0, 0, 0);
+	SetShaderValue(RHICmdList, GetPixelShader(), PackedProperties0, PackedProperties0Value);
+	SetShaderValue(RHICmdList, GetPixelShader(), ColorWeights, ColorWeightsValue);
+	SetShaderValue(RHICmdList, GetPixelShader(), Gamma, GammaValue);
+}
+
+template<typename TPixelShader>
+void FMipLevelBatchedElementParameters::BindShaders_RenderThread( FRHICommandList* RHICmdList, const FMatrix& InTransform, const float InGamma, const FMatrix& ColorWeights, const FTexture* Texture )
+{
+	TShaderMapRef<FCubemapTexturePropertiesVS> VertexShader(GetGlobalShaderMap());
+	TShaderMapRef<TPixelShader> PixelShader(GetGlobalShaderMap());
+	check(!RHICmdList);
+
+	static FGlobalBoundShaderState BoundShaderState;
+	SetGlobalBoundShaderState(BoundShaderState, GSimpleElementVertexDeclaration.VertexDeclarationRHI, *VertexShader, *PixelShader);
+
+	VertexShader->SetParameters(RHICmdList, InTransform);
+
+	RHISetBlendState(TStaticBlendState<>::GetRHI());
+
+	PixelShader->SetParameters(RHICmdList, Texture, ColorWeights, MipLevel, InGamma);
+}
+
+void FIESLightProfilePS::SetParameters( FRHICommandList* RHICmdList, const FTexture* Texture, float InBrightnessInLumens )
+{
+	FPixelShaderRHIParamRef ShaderRHI = GetPixelShader();
+	SetTextureParameter(RHICmdList, ShaderRHI, IESTexture, IESTextureSampler, Texture);
+
+	SetShaderValue(RHICmdList, ShaderRHI, BrightnessInLumens, InBrightnessInLumens);
+}
+
+void FIESLightProfileBatchedElementParameters::BindShaders_RenderThread( FRHICommandList* RHICmdList, const FMatrix& InTransform, const float InGamma, const FMatrix& ColorWeights, const FTexture* Texture )
+{
+	TShaderMapRef<FSimpleElementVS> VertexShader(GetGlobalShaderMap());
+	TShaderMapRef<FIESLightProfilePS> PixelShader(GetGlobalShaderMap());
+
+	check(!RHICmdList);
+	static FGlobalBoundShaderState BoundShaderState;
+	SetGlobalBoundShaderState(BoundShaderState, GSimpleElementVertexDeclaration.VertexDeclarationRHI, *VertexShader, *PixelShader);
+	RHISetBlendState(TStaticBlendState<>::GetRHI());
+
+	VertexShader->SetParameters(RHICmdList, InTransform);
+	PixelShader->SetParameters(RHICmdList, Texture, BrightnessInLumens);
+}
diff --git a/Engine/Source/Editor/UnrealEd/Private/EditorExporters.cpp b/Engine/Source/Editor/UnrealEd/Private/EditorExporters.cpp
--- a/Engine/Source/Editor/UnrealEd/Private/EditorExporters.cpp
+++ b/Engine/Source/Editor/UnrealEd/Private/EditorExporters.cpp
@@ -2556,7 +2556,6 @@ namespace MaterialExportUtils
 		check(InRenderTarget);
 		FTextureRenderTargetResource* RTResource = InRenderTarget->GameThread_GetRenderTargetResource();
 
-		RHIBeginScene();
 		{
 			// Create a canvas for the render target and clear it to black
 			FCanvas Canvas(RTResource, NULL, FApp::GetCurrentTime() - GStartTime, FApp::GetDeltaTime(), FApp::GetCurrentTime() - GStartTime);
@@ -2569,7 +2568,6 @@ namespace MaterialExportUtils
 			Canvas.SetRenderTarget(NULL);
 			FlushRenderingCommands();
 		}
-		RHIEndScene();
 				
 		bool bNormalmap = (InMaterialProperty == MP_Normal);
 		FReadSurfaceDataFlags ReadPixelFlags(bNormalmap ? RCM_SNorm : RCM_UNorm);
@@ -2696,13 +2694,6 @@ namespace MaterialExportUtils
 			FTextureRenderTargetResource* RenderTargetResource = RenderTargetTexture->GameThread_GetRenderTargetResource();
 			FSceneInterface* Scene = InLandscape->GetWorld()->Scene;
 			{
-				// Manually call RHIBeginScene since we are issuing draw calls outside of the main rendering function
-				ENQUEUE_UNIQUE_RENDER_COMMAND(
-					BeginCommand,
-				{
-					RHIBeginScene();
-				});
-
 				FSceneViewFamilyContext ViewFamily(
 					FSceneViewFamily::ConstructionValues(RenderTargetResource, Scene, FEngineShowFlags(ESFIM_Game))
 						.SetWorldTimes(FApp::GetCurrentTime() - GStartTime, FApp::GetDeltaTime(), FApp::GetCurrentTime() - GStartTime)
@@ -2782,12 +2773,6 @@ namespace MaterialExportUtils
 					FIntRect(0, 0, OutFlattenMaterial.DiffuseSize.X, OutFlattenMaterial.DiffuseSize.Y)
 					);
 			
-				ENQUEUE_UNIQUE_RENDER_COMMAND(
-					EndCommand,
-				{
-					RHIEndScene();
-				});
-
 				FlushRenderingCommands();
 			}
 
diff --git a/Engine/Source/Editor/UnrealEd/Private/Lightmass/LightmassRender.cpp b/Engine/Source/Editor/UnrealEd/Private/Lightmass/LightmassRender.cpp
--- a/Engine/Source/Editor/UnrealEd/Private/Lightmass/LightmassRender.cpp
+++ b/Engine/Source/Editor/UnrealEd/Private/Lightmass/LightmassRender.cpp
@@ -931,7 +931,7 @@ bool FLightmassMaterialRenderer::GenerateMaterialPropertyData(
 			{
 				// At this point, we can't just return false at failure since we have some clean-up to do...
 				Canvas->SetRenderTarget(RenderTarget->GameThread_GetRenderTargetResource());
-				RHIBeginScene();
+
 				// Clear the render target to black
 				// This is necessary because the below DrawTile doesn't write to the first column and first row
 				//@todo - figure out and fix DrawTile issues when rendering a full render target quad
@@ -943,7 +943,6 @@ bool FLightmassMaterialRenderer::GenerateMaterialPropertyData(
 				FlushRenderingCommands();
 				Canvas->SetRenderTarget(NULL);
 				FlushRenderingCommands();
-				RHIEndScene();
 
 				// Read in the data
 				//@todo UE4. Check the format! RenderTarget->Format
diff --git a/Engine/Source/Editor/UnrealEd/Private/MeshPaintRendering.cpp b/Engine/Source/Editor/UnrealEd/Private/MeshPaintRendering.cpp
--- a/Engine/Source/Editor/UnrealEd/Private/MeshPaintRendering.cpp
+++ b/Engine/Source/Editor/UnrealEd/Private/MeshPaintRendering.cpp
@@ -8,7 +8,8 @@
 #include "MeshPaintRendering.h"
 #include "GlobalShader.h"
 #include "ShaderParameters.h"
-
+#include "ShaderParameterUtils.h"
+#include "RHIStaticStates.h"
 
 namespace MeshPaintRendering
 {
@@ -42,9 +43,9 @@ namespace MeshPaintRendering
 			return bShaderHasOutdatedParameters;
 		}
 
-		void SetParameters( const FMatrix& InTransform )
+		void SetParameters(FRHICommandList* RHICmdList, const FMatrix& InTransform )
 		{
-			SetShaderValue( GetVertexShader(), TransformParameter, InTransform );
+			SetShaderValue(RHICmdList, GetVertexShader(), TransformParameter, InTransform );
 		}
 
 	private:
@@ -102,43 +103,44 @@ namespace MeshPaintRendering
 			return bShaderHasOutdatedParameters;
 		}
 
-		void SetParameters( const float InGamma, const FMeshPaintShaderParameters& InShaderParams )
+		void SetParameters(FRHICommandList* RHICmdList, const float InGamma, const FMeshPaintShaderParameters& InShaderParams )
 		{
 			const FPixelShaderRHIParamRef ShaderRHI = GetPixelShader();
 
 			SetTextureParameter(
+				RHICmdList, 
 				ShaderRHI,
 				CloneTextureParameter,
 				CloneTextureParameterSampler,
 				TStaticSamplerState< SF_Point, AM_Clamp, AM_Clamp, AM_Clamp >::GetRHI(),
 				InShaderParams.CloneTexture->GetRenderTargetResource()->TextureRHI );
 
-			SetShaderValue( ShaderRHI, WorldToBrushMatrixParameter, InShaderParams.WorldToBrushMatrix );
+			SetShaderValue(RHICmdList, ShaderRHI, WorldToBrushMatrixParameter, InShaderParams.WorldToBrushMatrix );
 
 			FVector4 BrushMetrics;
 			BrushMetrics.X = InShaderParams.BrushRadius;
 			BrushMetrics.Y = InShaderParams.BrushRadialFalloffRange;
 			BrushMetrics.Z = InShaderParams.BrushDepth;
 			BrushMetrics.W = InShaderParams.BrushDepthFalloffRange;
-			SetShaderValue( ShaderRHI, BrushMetricsParameter, BrushMetrics );
+			SetShaderValue(RHICmdList, ShaderRHI, BrushMetricsParameter, BrushMetrics );
 
 			FVector4 BrushStrength4( InShaderParams.BrushStrength, 0.0f, 0.0f, 0.0f );
-			SetShaderValue( ShaderRHI, BrushStrengthParameter, BrushStrength4 );
+			SetShaderValue(RHICmdList, ShaderRHI, BrushStrengthParameter, BrushStrength4 );
 
-			SetShaderValue( ShaderRHI, BrushColorParameter, InShaderParams.BrushColor );
+			SetShaderValue(RHICmdList, ShaderRHI, BrushColorParameter, InShaderParams.BrushColor );
 
 			FVector4 ChannelFlags;
 			ChannelFlags.X = InShaderParams.RedChannelFlag;
 			ChannelFlags.Y = InShaderParams.GreenChannelFlag;
 			ChannelFlags.Z = InShaderParams.BlueChannelFlag;
 			ChannelFlags.W = InShaderParams.AlphaChannelFlag;
-			SetShaderValue( ShaderRHI, ChannelFlagsParameter, ChannelFlags );
+			SetShaderValue(RHICmdList, ShaderRHI, ChannelFlagsParameter, ChannelFlags );
 			
 			float MaskVal = InShaderParams.GenerateMaskFlag ? 1.0f : 0.0f;
-			SetShaderValue( ShaderRHI, GenerateMaskFlagParameter, MaskVal );
+			SetShaderValue(RHICmdList, ShaderRHI, GenerateMaskFlagParameter, MaskVal );
 
 			// @todo MeshPaint
-			SetShaderValue( ShaderRHI, GammaParameter, InGamma );
+			SetShaderValue(RHICmdList, ShaderRHI, GammaParameter, InGamma );
 		}
 
 
@@ -204,9 +206,9 @@ namespace MeshPaintRendering
 			return bShaderHasOutdatedParameters;
 		}
 
-		void SetParameters( const FMatrix& InTransform )
+		void SetParameters(FRHICommandList* RHICmdList, const FMatrix& InTransform )
 		{
-			SetShaderValue( GetVertexShader(), TransformParameter, InTransform );
+			SetShaderValue(RHICmdList, GetVertexShader(), TransformParameter, InTransform );
 		}
 
 	private:
@@ -264,11 +266,12 @@ namespace MeshPaintRendering
 			return bShaderHasOutdatedParameters;
 		}
 
-		void SetParameters( const float InGamma, const FMeshPaintDilateShaderParameters& InShaderParams )
+		void SetParameters(FRHICommandList* RHICmdList, const float InGamma, const FMeshPaintDilateShaderParameters& InShaderParams )
 		{
 			const FPixelShaderRHIParamRef ShaderRHI = GetPixelShader();
 
 			SetTextureParameter(
+				RHICmdList, 
 				ShaderRHI,
 				Texture0Parameter,
 				Texture0ParameterSampler,
@@ -276,6 +279,7 @@ namespace MeshPaintRendering
 				InShaderParams.Texture0->GetRenderTargetResource()->TextureRHI );
 
 			SetTextureParameter(
+				RHICmdList, 
 				ShaderRHI,
 				Texture1Parameter,
 				Texture1ParameterSampler,
@@ -283,15 +287,16 @@ namespace MeshPaintRendering
 				InShaderParams.Texture1->GetRenderTargetResource()->TextureRHI );
 
 			SetTextureParameter(
+				RHICmdList, 
 				ShaderRHI,
 				Texture2Parameter,
 				Texture2ParameterSampler,
 				TStaticSamplerState< SF_Point, AM_Clamp, AM_Clamp, AM_Clamp >::GetRHI(),
 				InShaderParams.Texture2->GetRenderTargetResource()->TextureRHI );
 
-			SetShaderValue(ShaderRHI, WidthPixelOffsetParameter, InShaderParams.WidthPixelOffset );
-			SetShaderValue(ShaderRHI, HeightPixelOffsetParameter, InShaderParams.HeightPixelOffset );
-			SetShaderValue(ShaderRHI, GammaParameter, InGamma );
+			SetShaderValue(RHICmdList, ShaderRHI, WidthPixelOffsetParameter, InShaderParams.WidthPixelOffset );
+			SetShaderValue(RHICmdList, ShaderRHI, HeightPixelOffsetParameter, InShaderParams.HeightPixelOffset );
+			SetShaderValue(RHICmdList, ShaderRHI, GammaParameter, InGamma );
 		}
 
 
@@ -343,7 +348,7 @@ namespace MeshPaintRendering
 
 
 	/** Binds the mesh paint vertex and pixel shaders to the graphics device */
-	void SetMeshPaintShaders_RenderThread( const FMatrix& InTransform,
+	void SetMeshPaintShaders_RenderThread(FRHICommandList* RHICmdList, const FMatrix& InTransform,
 										   const float InGamma,
 										   const FMeshPaintShaderParameters& InShaderParams )
 	{
@@ -354,16 +359,16 @@ namespace MeshPaintRendering
 		SetGlobalBoundShaderState( BoundShaderState, GMeshPaintVertexDeclaration.VertexDeclarationRHI, *VertexShader, *PixelShader );
 
 		// Set vertex shader parameters
-		VertexShader->SetParameters( InTransform );
+		VertexShader->SetParameters(RHICmdList, InTransform );
 		
 		// Set pixel shader parameters
-		PixelShader->SetParameters( InGamma, InShaderParams );
+		PixelShader->SetParameters(RHICmdList, InGamma, InShaderParams );
 
 		// @todo MeshPaint: Make sure blending/color writes are setup so we can write to ALPHA if needed!
 	}
 
 	/** Binds the mesh paint vertex and pixel shaders to the graphics device */
-	void SetMeshPaintDilateShaders_RenderThread( const FMatrix& InTransform,
+	void SetMeshPaintDilateShaders_RenderThread(FRHICommandList* RHICmdList, const FMatrix& InTransform,
 												 const float InGamma,
 												 const FMeshPaintDilateShaderParameters& InShaderParams )
 	{
@@ -374,10 +379,10 @@ namespace MeshPaintRendering
 		SetGlobalBoundShaderState( BoundShaderState, GMeshPaintDilateVertexDeclaration.VertexDeclarationRHI, *VertexShader, *PixelShader );
 
 		// Set vertex shader parameters
-		VertexShader->SetParameters( InTransform );
+		VertexShader->SetParameters(RHICmdList, InTransform );
 
 		// Set pixel shader parameters
-		PixelShader->SetParameters( InGamma, InShaderParams );
+		PixelShader->SetParameters(RHICmdList, InGamma, InShaderParams );
 	}
 
 }
diff --git a/Engine/Source/Editor/UnrealEd/Private/NormalMapPreview.cpp b/Engine/Source/Editor/UnrealEd/Private/NormalMapPreview.cpp
--- a/Engine/Source/Editor/UnrealEd/Private/NormalMapPreview.cpp
+++ b/Engine/Source/Editor/UnrealEd/Private/NormalMapPreview.cpp
@@ -9,6 +9,8 @@
 #include "SimpleElementShaders.h"
 #include "GlobalShader.h"
 #include "ShaderParameters.h"
+#include "ShaderParameterUtils.h"
+#include "RHIStaticStates.h"
 
 /*------------------------------------------------------------------------------
 	Batched element shaders for previewing normal maps.
@@ -46,10 +48,10 @@ public:
 	 * Set shader parameters.
 	 * @param NormalMapTexture - The normal map texture to sample.
 	 */
-	void SetParameters(const FTexture* NormalMapTexture)
+	void SetParameters(FRHICommandList* RHICmdList, const FTexture* NormalMapTexture)
 	{
 		FPixelShaderRHIParamRef PixelShaderRHI = GetPixelShader();
-		SetTextureParameter(PixelShaderRHI,Texture,TextureSampler,NormalMapTexture);
+		SetTextureParameter(RHICmdList, PixelShaderRHI,Texture,TextureSampler,NormalMapTexture);
 	}
 
 	/**
@@ -73,6 +75,7 @@ IMPLEMENT_SHADER_TYPE(,FSimpleElementNormalMapPS,TEXT("SimpleElementNormalMapPix
 
 /** Binds vertex and pixel shaders for this element */
 void FNormalMapBatchedElementParameters::BindShaders_RenderThread(
+	FRHICommandList* RHICmdList, 
 	const FMatrix& InTransform,
 	const float InGamma,
 	const FMatrix& ColorWeights,
@@ -81,11 +84,13 @@ void FNormalMapBatchedElementParameters::BindShaders_RenderThread(
 	TShaderMapRef<FSimpleElementVS> VertexShader(GetGlobalShaderMap());
 	TShaderMapRef<FSimpleElementNormalMapPS> PixelShader(GetGlobalShaderMap());
 
+	check(!RHICmdList);
+
 	static FGlobalBoundShaderState BoundShaderState;
 	SetGlobalBoundShaderState(BoundShaderState, GSimpleElementVertexDeclaration.VertexDeclarationRHI, *VertexShader, *PixelShader);
 
-	VertexShader->SetParameters(InTransform);
-	PixelShader->SetParameters(Texture);
+	VertexShader->SetParameters(RHICmdList, InTransform);
+	PixelShader->SetParameters(RHICmdList, Texture);
 
 	RHISetBlendState(TStaticBlendState<>::GetRHI());
 }
diff --git a/Engine/Source/Editor/UnrealEd/Private/ObjectTools.cpp b/Engine/Source/Editor/UnrealEd/Private/ObjectTools.cpp
--- a/Engine/Source/Editor/UnrealEd/Private/ObjectTools.cpp
+++ b/Engine/Source/Editor/UnrealEd/Private/ObjectTools.cpp
@@ -3326,13 +3326,6 @@ namespace ThumbnailTools
 		}
 		check( RenderTargetResource != NULL );
 
-		// Manually call RHIBeginScene since we are issuing draw calls outside of the main rendering function
-		ENQUEUE_UNIQUE_RENDER_COMMAND(
-			BeginCommand,
-		{
-			RHIBeginScene();
-		});
-
 		// Create a canvas for the render target and clear it to black
 		FCanvas Canvas( RenderTargetResource, NULL, FApp::GetCurrentTime() - GStartTime, FApp::GetDeltaTime(), FApp::GetCurrentTime() - GStartTime );
 		Canvas.Clear( FLinearColor::Black );
@@ -3450,12 +3443,6 @@ namespace ThumbnailTools
 				}
 			}
 		}
-
-		ENQUEUE_UNIQUE_RENDER_COMMAND(
-			EndCommand,
-		{
-			RHIEndScene();
-		});
 	}
 
 
diff --git a/Engine/Source/Editor/UnrealEd/Private/Texture2DPreview.cpp b/Engine/Source/Editor/UnrealEd/Private/Texture2DPreview.cpp
--- a/Engine/Source/Editor/UnrealEd/Private/Texture2DPreview.cpp
+++ b/Engine/Source/Editor/UnrealEd/Private/Texture2DPreview.cpp
@@ -9,6 +9,8 @@
 #include "SimpleElementShaders.h"
 #include "GlobalShader.h"
 #include "ShaderParameters.h"
+#include "ShaderParameterUtils.h"
+#include "RHIStaticStates.h"
 
 /*------------------------------------------------------------------------------
 	Batched element shaders for previewing 2d textures.
@@ -39,15 +41,15 @@ public:
 		return IsFeatureLevelSupported(Platform, ERHIFeatureLevel::SM3);
 	}
 
-	void SetParameters(const FTexture* TextureValue, const FMatrix& ColorWeightsValue, float GammaValue, float MipLevel, bool bIsNormalMap)
+	void SetParameters(FRHICommandList* RHICmdList, const FTexture* TextureValue, const FMatrix& ColorWeightsValue, float GammaValue, float MipLevel, bool bIsNormalMap)
 	{
-		SetTextureParameter(GetPixelShader(),InTexture,InTextureSampler,TextureValue);
-		SetShaderValue(GetPixelShader(),ColorWeights,ColorWeightsValue);
+		SetTextureParameter(RHICmdList, GetPixelShader(),InTexture,InTextureSampler,TextureValue);
+		SetShaderValue(RHICmdList, GetPixelShader(),ColorWeights,ColorWeightsValue);
 		FVector4 PackedParametersValue(GammaValue, MipLevel, bIsNormalMap ? 1.0 : -1.0f, 0.0f);
-		SetShaderValue(GetPixelShader(), PackedParameters, PackedParametersValue);
+		SetShaderValue(RHICmdList, GetPixelShader(), PackedParameters, PackedParametersValue);
 
-		SetShaderValue(GetPixelShader(),TextureComponentReplicate,TextureValue->bGreyScaleFormat ? FLinearColor(1,0,0,0) : FLinearColor(0,0,0,0));
-		SetShaderValue(GetPixelShader(),TextureComponentReplicateAlpha,TextureValue->bGreyScaleFormat ? FLinearColor(1,0,0,0) : FLinearColor(0,0,0,1));
+		SetShaderValue(RHICmdList, GetPixelShader(),TextureComponentReplicate,TextureValue->bGreyScaleFormat ? FLinearColor(1,0,0,0) : FLinearColor(0,0,0,0));
+		SetShaderValue(RHICmdList, GetPixelShader(),TextureComponentReplicateAlpha,TextureValue->bGreyScaleFormat ? FLinearColor(1,0,0,0) : FLinearColor(0,0,0,1));
 	}
 
 	virtual bool Serialize(FArchive& Ar)
@@ -75,6 +77,7 @@ IMPLEMENT_SHADER_TYPE(,FSimpleElementTexture2DPreviewPS,TEXT("SimpleElementTextu
 
 /** Binds vertex and pixel shaders for this element */
 void FBatchedElementTexture2DPreviewParameters::BindShaders_RenderThread(
+	FRHICommandList* RHICmdList, 
 	const FMatrix& InTransform,
 	const float InGamma,
 	const FMatrix& ColorWeights,
@@ -82,12 +85,13 @@ void FBatchedElementTexture2DPreviewParameters::BindShaders_RenderThread(
 {
 	TShaderMapRef<FSimpleElementVS> VertexShader(GetGlobalShaderMap());
 	TShaderMapRef<FSimpleElementTexture2DPreviewPS> PixelShader(GetGlobalShaderMap());
+	check(!RHICmdList);
 
 	static FGlobalBoundShaderState BoundShaderState;
 	SetGlobalBoundShaderState(BoundShaderState, GSimpleElementVertexDeclaration.VertexDeclarationRHI, *VertexShader, *PixelShader);
 
-	VertexShader->SetParameters(InTransform);
-	PixelShader->SetParameters(Texture, ColorWeights, InGamma, MipLevel, bIsNormalMap);
+	VertexShader->SetParameters(RHICmdList, InTransform);
+	PixelShader->SetParameters(RHICmdList, Texture, ColorWeights, InGamma, MipLevel, bIsNormalMap);
 
 	if( bIsSingleChannelFormat )
 	{
diff --git a/Engine/Source/Editor/UnrealEd/Public/CubemapUnwrapUtils.h b/Engine/Source/Editor/UnrealEd/Public/CubemapUnwrapUtils.h
--- a/Engine/Source/Editor/UnrealEd/Public/CubemapUnwrapUtils.h
+++ b/Engine/Source/Editor/UnrealEd/Public/CubemapUnwrapUtils.h
@@ -51,10 +51,7 @@ public:
 	}
 	FCubemapTexturePropertiesVS() {}
 
-	void SetParameters(const FMatrix& TransformValue)
-	{
-		SetShaderValue(GetVertexShader(), Transform, TransformValue);
-	}
+	void SetParameters(FRHICommandList* RHICmdList, const FMatrix& TransformValue);
 
 	virtual bool Serialize(FArchive& Ar)
 	{
@@ -90,15 +87,7 @@ public:
 	}
 	FCubemapTexturePropertiesPS() {}
 
-	void SetParameters(const FTexture* Texture, const FMatrix& ColorWeightsValue, float MipLevel, float GammaValue)
-	{
-		SetTextureParameter(GetPixelShader(),CubeTexture,CubeTextureSampler,Texture);
-
-		FVector4 PackedProperties0Value(MipLevel, 0, 0, 0);
-		SetShaderValue(GetPixelShader(), PackedProperties0, PackedProperties0Value);
-		SetShaderValue(GetPixelShader(), ColorWeights, ColorWeightsValue);
-		SetShaderValue(GetPixelShader(), Gamma, GammaValue);
-	}
+	void SetParameters(FRHICommandList* RHICmdList, const FTexture* Texture, const FMatrix& ColorWeightsValue, float MipLevel, float GammaValue);
 
 	virtual bool Serialize(FArchive& Ar)
 	{
@@ -130,33 +119,20 @@ public:
 	}
 
 	/** Binds vertex and pixel shaders for this element */
-	virtual void BindShaders_RenderThread(const FMatrix& InTransform, const float InGamma, const FMatrix& ColorWeights, const FTexture* Texture)
+	virtual void BindShaders_RenderThread(FRHICommandList* RHICmdList, const FMatrix& InTransform, const float InGamma, const FMatrix& ColorWeights, const FTexture* Texture) OVERRIDE
 	{
 		if(bHDROutput)
 		{
-			BindShaders_RenderThread<FCubemapTexturePropertiesPS<true> >(InTransform, InGamma, ColorWeights, Texture);
+			BindShaders_RenderThread<FCubemapTexturePropertiesPS<true> >(RHICmdList, InTransform, InGamma, ColorWeights, Texture);
 		}
 		else
 		{
-			BindShaders_RenderThread<FCubemapTexturePropertiesPS<false> >(InTransform, InGamma, ColorWeights, Texture);
+			BindShaders_RenderThread<FCubemapTexturePropertiesPS<false> >(RHICmdList, InTransform, InGamma, ColorWeights, Texture);
 		}
 	}
 
 private:
-	template<typename TPixelShader> void BindShaders_RenderThread(const FMatrix& InTransform, const float InGamma, const FMatrix& ColorWeights, const FTexture* Texture)
-	{
-		TShaderMapRef<FCubemapTexturePropertiesVS> VertexShader(GetGlobalShaderMap());
-		TShaderMapRef<TPixelShader> PixelShader(GetGlobalShaderMap());
-
-		static FGlobalBoundShaderState BoundShaderState;
-		SetGlobalBoundShaderState(BoundShaderState, GSimpleElementVertexDeclaration.VertexDeclarationRHI, *VertexShader, *PixelShader);
-
-		VertexShader->SetParameters(InTransform);
-
-		RHISetBlendState(TStaticBlendState<>::GetRHI());
-
-		PixelShader->SetParameters(Texture, ColorWeights, MipLevel, InGamma);
-	}
+	template<typename TPixelShader> void BindShaders_RenderThread(FRHICommandList* RHICmdList, const FMatrix& InTransform, const float InGamma, const FMatrix& ColorWeights, const FTexture* Texture);
 
 	bool bHDROutput;
 	/** Parameters that need to be passed to the shader */
@@ -187,13 +163,7 @@ public:
 		BrightnessInLumens.Bind(Initializer.ParameterMap,TEXT("BrightnessInLumens"));
 	}
 
-	void SetParameters(const FTexture* Texture, float InBrightnessInLumens)
-	{
-		FPixelShaderRHIParamRef ShaderRHI = GetPixelShader();
-		SetTextureParameter(ShaderRHI, IESTexture, IESTextureSampler, Texture);
-
-		SetShaderValue(ShaderRHI, BrightnessInLumens, InBrightnessInLumens);
-	}
+	void SetParameters(FRHICommandList* RHICmdList, const FTexture* Texture, float InBrightnessInLumens);
 
 	virtual bool Serialize(FArchive& Ar) OVERRIDE
 	{
@@ -219,18 +189,7 @@ public:
 	}
 
 	/** Binds vertex and pixel shaders for this element */
-	virtual void BindShaders_RenderThread( const FMatrix& InTransform, const float InGamma, const FMatrix& ColorWeights, const FTexture* Texture )
-	{
-		TShaderMapRef<FSimpleElementVS> VertexShader(GetGlobalShaderMap());
-		TShaderMapRef<FIESLightProfilePS> PixelShader(GetGlobalShaderMap());
-
-		static FGlobalBoundShaderState BoundShaderState;
-		SetGlobalBoundShaderState(BoundShaderState, GSimpleElementVertexDeclaration.VertexDeclarationRHI, *VertexShader, *PixelShader);
-		RHISetBlendState(TStaticBlendState<>::GetRHI());
-
-		VertexShader->SetParameters(InTransform);
-		PixelShader->SetParameters(Texture, BrightnessInLumens);
-	}
+	virtual void BindShaders_RenderThread(FRHICommandList* RHICmdList, const FMatrix& InTransform, const float InGamma, const FMatrix& ColorWeights, const FTexture* Texture ) OVERRIDE;
 
 private:
 	float BrightnessInLumens;
diff --git a/Engine/Source/Editor/UnrealEd/Public/MeshPaintRendering.h b/Engine/Source/Editor/UnrealEd/Public/MeshPaintRendering.h
--- a/Engine/Source/Editor/UnrealEd/Public/MeshPaintRendering.h
+++ b/Engine/Source/Editor/UnrealEd/Public/MeshPaintRendering.h
@@ -57,12 +57,12 @@ namespace MeshPaintRendering
 
 
 	/** Binds the mesh paint vertex and pixel shaders to the graphics device */
-	void UNREALED_API SetMeshPaintShaders_RenderThread( const FMatrix& InTransform,
+	void UNREALED_API SetMeshPaintShaders_RenderThread(FRHICommandList* RHICmdList, const FMatrix& InTransform,
 										   const float InGamma,
 										   const FMeshPaintShaderParameters& InShaderParams );
 
 	/** Binds the mesh paint dilation vertex and pixel shaders to the graphics device */
-	void UNREALED_API SetMeshPaintDilateShaders_RenderThread( const FMatrix& InTransform,
+	void UNREALED_API SetMeshPaintDilateShaders_RenderThread(FRHICommandList* RHICmdList, const FMatrix& InTransform,
 												const float InGamma,
 												const FMeshPaintDilateShaderParameters& InShaderParams );
 
diff --git a/Engine/Source/Editor/UnrealEd/Public/NormalMapPreview.h b/Engine/Source/Editor/UnrealEd/Public/NormalMapPreview.h
--- a/Engine/Source/Editor/UnrealEd/Public/NormalMapPreview.h
+++ b/Engine/Source/Editor/UnrealEd/Public/NormalMapPreview.h
@@ -15,6 +15,7 @@ class UNREALED_API FNormalMapBatchedElementParameters : public FBatchedElementPa
 {
 	/** Binds vertex and pixel shaders for this element */
 	virtual void BindShaders_RenderThread(
+		FRHICommandList* RHICmdList, 
 		const FMatrix& InTransform,
 		const float InGamma,
 		const FMatrix& ColorWeights,
diff --git a/Engine/Source/Editor/UnrealEd/Public/Texture2DPreview.h b/Engine/Source/Editor/UnrealEd/Public/Texture2DPreview.h
--- a/Engine/Source/Editor/UnrealEd/Public/Texture2DPreview.h
+++ b/Engine/Source/Editor/UnrealEd/Public/Texture2DPreview.h
@@ -22,7 +22,7 @@ public:
 	}
 
 	/** Binds vertex and pixel shaders for this element */
-	virtual void BindShaders_RenderThread( const FMatrix& InTransform, const float InGamma, const FMatrix& ColorWeights, const FTexture* Texture );
+	virtual void BindShaders_RenderThread(FRHICommandList* RHICmdList, const FMatrix& InTransform, const float InGamma, const FMatrix& ColorWeights, const FTexture* Texture ) OVERRIDE;
 
 private:
 
diff --git a/Engine/Source/Runtime/CoreUObject/Private/UObject/TextureAllocations.cpp b/Engine/Source/Runtime/CoreUObject/Private/UObject/TextureAllocations.cpp
--- a/Engine/Source/Runtime/CoreUObject/Private/UObject/TextureAllocations.cpp
+++ b/Engine/Source/Runtime/CoreUObject/Private/UObject/TextureAllocations.cpp
@@ -4,6 +4,8 @@
 
 #if WITH_ENGINE
 
+#include "Runtime/RHI/Public/RHIDefinitions.h"	// For TexCreate* flags
+
 /*-----------------------------------------------------------------------------
 	Linker file package summary texture allocation info
 -----------------------------------------------------------------------------*/
@@ -79,7 +81,7 @@ FTexture2DResourceMem* FTextureAllocations::FindAndRemove( int32 SizeX, int32 Si
 FTextureAllocations::FTextureType* FTextureAllocations::FindTextureType( int32 SizeX, int32 SizeY, int32 NumMips, uint32 Format, uint32 TexCreateFlags )
 {
 	FTexture2DResourceMem* ResourceMem = NULL;
-	const uint32 FlagMask = ~((1<<9) | (1<<10)); // TexCreate_AllowFailure | TexCreate_DisableAutoDefrag (from ETextureCreateFlags in RHI.h)
+	const uint32 FlagMask = ~(TexCreate_AllowFailure | TexCreate_DisableAutoDefrag); // TexCreate_AllowFailure | TexCreate_DisableAutoDefrag (from ETextureCreateFlags in RHI.h)
 	for ( int32 TypeIndex=0; TypeIndex < TextureTypes.Num(); ++TypeIndex )
 	{
 		FTextureType& TextureType = TextureTypes[ TypeIndex ];
diff --git a/Engine/Source/Runtime/Engine/Classes/Engine/LocalPlayer.h b/Engine/Source/Runtime/Engine/Classes/Engine/LocalPlayer.h
--- a/Engine/Source/Runtime/Engine/Classes/Engine/LocalPlayer.h
+++ b/Engine/Source/Runtime/Engine/Classes/Engine/LocalPlayer.h
@@ -265,5 +265,8 @@ public:
 	 * @param InOffset		Offset applied to current world origin
 	 */
 	void ApplyWorldOffset(FVector InOffset);
+
+	/** Locked view state needs access to GetViewPoint. */
+	friend class FLockedViewState;
 };
 
diff --git a/Engine/Source/Runtime/Engine/Classes/Engine/Texture.h b/Engine/Source/Runtime/Engine/Classes/Engine/Texture.h
--- a/Engine/Source/Runtime/Engine/Classes/Engine/Texture.h
+++ b/Engine/Source/Runtime/Engine/Classes/Engine/Texture.h
@@ -505,6 +505,12 @@ public:
 	/** The texture's resource, can be NULL */
 	class FTextureResource*	Resource;
 
+	/** Stable RHI texture reference that refers to the current RHI texture. Note this is manually refcounted! */
+	class FRHITextureReference* TextureRefRHI;
+
+	/** Release fence to know when resources have been freed on the rendering thread. */
+	FRenderCommandFence ReleaseFence;
+
 	/** delegate type for texture save events ( Params: UTexture* TextureToSave ) */
 	DECLARE_MULTICAST_DELEGATE_OneParam(FOnTextureSaved, class UTexture*);
 	/** triggered before a texture is being saved */
diff --git a/Engine/Source/Runtime/Engine/Classes/Engine/Texture2D.h b/Engine/Source/Runtime/Engine/Classes/Engine/Texture2D.h
--- a/Engine/Source/Runtime/Engine/Classes/Engine/Texture2D.h
+++ b/Engine/Source/Runtime/Engine/Classes/Engine/Texture2D.h
@@ -243,6 +243,12 @@ public:
 	/** Sets the minimum number of mips that must be resident in memory (cannot be streamed). */
 	static void SetMinTextureResidentMipCount(int32 InMinTextureResidentMipCount);
 
+	/** Retrieves the last render time used for streaming. */
+	float GetLastRenderTimeForStreaming();
+
+	/** Invalidates the last render time used for streaming. */
+	void InvalidateLastRenderTimeForStreaming();
+
 	/**
 	 * Get mip data starting with the specified mip index.
 	 * @param FirstMipToLoad - The first mip index to cache.
diff --git a/Engine/Source/Runtime/Engine/Private/BatchedElements.cpp b/Engine/Source/Runtime/Engine/Private/BatchedElements.cpp
--- a/Engine/Source/Runtime/Engine/Private/BatchedElements.cpp
+++ b/Engine/Source/Runtime/Engine/Private/BatchedElements.cpp
@@ -5,6 +5,7 @@
 #include "BatchedElements.h"
 #include "GlobalShader.h"
 #include "ShaderParameters.h"
+#include "RHIStaticStates.h"
 #include "SimpleElementShaders.h"
 
 DEFINE_LOG_CATEGORY_STATIC(LogBatchedElements, Log, All);
@@ -355,6 +356,9 @@ void FBatchedElements::PrepareShaders(
 	FTexture2DRHIRef DepthTexture
 	) const
 {
+	//@todo-rco: RHIPacketList
+	FRHICommandList* RHICmdList = nullptr;
+
 	// used to mask individual channels and desaturate
 	FMatrix ColorWeights( FPlane(1, 0, 0, 0), FPlane(0, 1, 0, 0), FPlane(0, 0, 1, 0), FPlane(0, 0, 0, 0) );
 	
@@ -422,7 +426,7 @@ void FBatchedElements::PrepareShaders(
 	if( BatchedElementParameters != NULL )
 	{
 		// Use the vertex/pixel shader that we were given
-		BatchedElementParameters->BindShaders_RenderThread( Transform, Gamma, ColorWeights, Texture );
+		BatchedElementParameters->BindShaders_RenderThread(RHICmdList, Transform, Gamma, ColorWeights, Texture );
 	}
 	else
 	{
@@ -434,7 +438,7 @@ void FBatchedElements::PrepareShaders(
 			SetGlobalBoundShaderState(HitTestingBoundShaderState, GSimpleElementVertexDeclaration.VertexDeclarationRHI, 
 				*VertexShader, *HitTestingPixelShader);
 
-		    HitTestingPixelShader->SetParameters(Texture);
+			HitTestingPixelShader->SetParameters(RHICmdList, Texture);
 		    SetHitTestingBlendState( BlendMode);
 		    
 	    }
@@ -450,8 +454,8 @@ void FBatchedElements::PrepareShaders(
 				SetGlobalBoundShaderState( MaskedBoundShaderState, GSimpleElementVertexDeclaration.VertexDeclarationRHI, 
 					*VertexShader, *MaskedPixelShader);
 
-				MaskedPixelShader->SetEditorCompositingParameters( View, DepthTexture );
-				MaskedPixelShader->SetParameters(Texture,Gamma,GBatchedElementAlphaRefVal / 255.0f,BlendMode);
+				MaskedPixelShader->SetEditorCompositingParameters(RHICmdList, View, DepthTexture );
+				MaskedPixelShader->SetParameters(RHICmdList, Texture,Gamma,GBatchedElementAlphaRefVal / 255.0f,BlendMode);
 		    }
 		    // render distance field elements
 		    else if (
@@ -489,6 +493,7 @@ void FBatchedElements::PrepareShaders(
 				    );
 			    
 			    DistanceFieldPixelShader->SetParameters(
+					RHICmdList, 
 				    Texture,
 				    Gamma,
 				    AlphaRefVal / 255.0f,
@@ -507,7 +512,7 @@ void FBatchedElements::PrepareShaders(
 				SetGlobalBoundShaderState( ColorChannelMaskShaderState, GSimpleElementVertexDeclaration.VertexDeclarationRHI, 
 					*VertexShader, *ColorChannelMaskPixelShader );
 			
-				ColorChannelMaskPixelShader->SetParameters( Texture, ColorWeights, GammaToUse );
+				ColorChannelMaskPixelShader->SetParameters(RHICmdList, Texture, ColorWeights, GammaToUse );
 			}
 			else
 		    {
@@ -519,8 +524,8 @@ void FBatchedElements::PrepareShaders(
 					SetGlobalBoundShaderState( SimpleBoundShaderState, GSimpleElementVertexDeclaration.VertexDeclarationRHI, 
 						*VertexShader, *RegularPixelShader );
 
-					RegularPixelShader->SetEditorCompositingParameters(View, DepthTexture);
-				    RegularPixelShader->SetParameters(Texture);
+					RegularPixelShader->SetEditorCompositingParameters(RHICmdList, View, DepthTexture);
+				    RegularPixelShader->SetParameters(RHICmdList, Texture);
 			    }
 			    else
 			    {
@@ -528,14 +533,14 @@ void FBatchedElements::PrepareShaders(
 					SetGlobalBoundShaderState( RegularBoundShaderState, GSimpleElementVertexDeclaration.VertexDeclarationRHI, 
 						*VertexShader, *RegularPixelShader );
 
-				    RegularPixelShader->SetParameters(Texture,Gamma,BlendMode);
-					RegularPixelShader->SetEditorCompositingParameters(View, DepthTexture);
+				    RegularPixelShader->SetParameters(RHICmdList, Texture,Gamma,BlendMode);
+					RegularPixelShader->SetEditorCompositingParameters(RHICmdList, View, DepthTexture);
 			    }
 		    }
 	    }
 
 		// Set the simple element vertex shader parameters
-		VertexShader->SetParameters(Transform, bSwitchVerticalAxis);
+		VertexShader->SetParameters(RHICmdList, Transform, bSwitchVerticalAxis);
 	}
 }
 
diff --git a/Engine/Source/Runtime/Engine/Private/ClearReplacementShaders.cpp b/Engine/Source/Runtime/Engine/Private/ClearReplacementShaders.cpp
--- a/Engine/Source/Runtime/Engine/Private/ClearReplacementShaders.cpp
+++ b/Engine/Source/Runtime/Engine/Private/ClearReplacementShaders.cpp
@@ -3,6 +3,7 @@
 
 #include "EnginePrivate.h"
 #include "ClearReplacementShaders.h"
+#include "ShaderParameterUtils.h"
 
 IMPLEMENT_SHADER_TYPE(,FClearReplacementVS,TEXT("ClearReplacementShaders"),TEXT("ClearVS"),SF_Vertex);
 IMPLEMENT_SHADER_TYPE(,FClearReplacementPS,TEXT("ClearReplacementShaders"),TEXT("ClearPS"),SF_Pixel);
@@ -10,3 +11,19 @@ IMPLEMENT_SHADER_TYPE(,FClearReplacementPS,TEXT("ClearReplacementShaders"),TEXT(
 IMPLEMENT_SHADER_TYPE(,FClearTexture2DReplacementCS,TEXT("ClearReplacementShaders"),TEXT("ClearTexture2DCS"),SF_Compute);
 
 IMPLEMENT_SHADER_TYPE(,FClearBufferReplacementCS,TEXT("ClearReplacementShaders"),TEXT("ClearBufferCS"),SF_Compute);
+
+void FClearTexture2DReplacementCS::SetParameters( FRHICommandList* RHICmdList, FUnorderedAccessViewRHIParamRef TextureRW, FLinearColor Value )
+{
+	check(!RHICmdList);
+	FComputeShaderRHIParamRef ComputeShaderRHI = GetComputeShader();
+	SetShaderValue(RHICmdList, ComputeShaderRHI, ClearColor, Value);
+	RHISetUAVParameter(ComputeShaderRHI, ClearTextureRW.GetBaseIndex(), TextureRW);
+}
+
+void FClearBufferReplacementCS::SetParameters( FRHICommandList* RHICmdList, FUnorderedAccessViewRHIParamRef BufferRW, uint32 Dword )
+{
+	check(!RHICmdList);
+	FComputeShaderRHIParamRef ComputeShaderRHI = GetComputeShader();
+	SetShaderValue(RHICmdList, ComputeShaderRHI, ClearDword, Dword);
+	RHISetUAVParameter(ComputeShaderRHI, ClearBufferRW.GetBaseIndex(), BufferRW);
+}
diff --git a/Engine/Source/Runtime/Engine/Private/ContentStreaming.cpp b/Engine/Source/Runtime/Engine/Private/ContentStreaming.cpp
--- a/Engine/Source/Runtime/Engine/Private/ContentStreaming.cpp
+++ b/Engine/Source/Runtime/Engine/Private/ContentStreaming.cpp
@@ -257,7 +257,8 @@ struct FStreamingTexture
 		MaxAllowedMips = MipCount;	//FMath::Max(Texture->ResidentMips, Texture->RequestedMips);
 		STAT( MaxAllowedOptimalMips = MaxAllowedMips );
 		STAT( MostResidentMips = FMath::Max(MostResidentMips, Texture->ResidentMips) );
-		LastRenderTime = (FApp::GetCurrentTime() > Texture->Resource->LastRenderTime) ? float( FApp::GetCurrentTime() - Texture->Resource->LastRenderTime ) : 0.0f;
+		float LastRenderTimeForTexture = Texture->GetLastRenderTimeForStreaming();
+		LastRenderTime = (FApp::GetCurrentTime() > LastRenderTimeForTexture) ? float( FApp::GetCurrentTime() - LastRenderTimeForTexture ) : 0.0f;
 		MinDistance = MAX_STREAMINGDISTANCE;
 		bForceFullyLoad = Texture->ShouldMipLevelsBeForcedResident() || (ForceLoadRefCount > 0);
 		TextureLODBias = Texture->GetCachedLODBias();
@@ -2198,7 +2199,7 @@ void FStreamingManagerTexture::CancelForcedResources()
 				StreamingTexture.InstanceRemovedTimestamp = -FLT_MAX;
 				if ( StreamingTexture.Texture->Resource )
 				{
-					StreamingTexture.Texture->Resource->LastRenderTime = -FLT_MAX;
+					StreamingTexture.Texture->InvalidateLastRenderTimeForStreaming();
 				}
 #if STREAMING_LOG_CANCELFORCED
 				UE_LOG(LogContentStreaming, Log, TEXT("Canceling forced texture: %s (had %.1f seconds left)"), *StreamingTexture.Texture->GetFullName(), TimeLeft );
@@ -2299,8 +2300,6 @@ void Renderthread_StreamOutTextureData( TArray<FTextureSortElement>* InCandidate
 {
 	*bSucceeded = false;
 
-	RHIBeginScene();
-
 	// only for debugging?
 	FTextureMemoryStats OldStats;
 	RHIGetTextureMemoryStats(OldStats);
@@ -2362,8 +2361,6 @@ void Renderthread_StreamOutTextureData( TArray<FTextureSortElement>* InCandidate
 
 	// Return the result.
 	*bSucceeded = SavedMemory >= RequiredMemorySize;
-
-	RHIEndScene();
 }
 
 /**
diff --git a/Engine/Source/Runtime/Engine/Private/GPUSkinCache.cpp b/Engine/Source/Runtime/Engine/Private/GPUSkinCache.cpp
--- a/Engine/Source/Runtime/Engine/Private/GPUSkinCache.cpp
+++ b/Engine/Source/Runtime/Engine/Private/GPUSkinCache.cpp
@@ -11,6 +11,7 @@
 #include "GPUSkinVertexFactory.h"
 #include "SkeletalRenderGPUSkin.h"
 #include "GPUSkinCache.h"
+#include "ShaderParameterUtils.h"
 
 #define GPUSKINCACHE_FRAMES	3
 
@@ -77,20 +78,22 @@ public:
 	{
 	}
 
-	void SetParameters(uint32 InputStreamStride, uint32 InputStreamFloatOffset, uint32 InputStreamVertexCount, uint32 OutputBufferFloatOffset, FBoneBufferTypeRef BoneBuffer, FUniformBufferRHIRef UniformBuffer, FShaderResourceViewRHIRef VertexBufferSRV, FRWBuffer& SkinBuffer, const FVector& MeshOrigin, const FVector& MeshExtension)
+	void SetParameters(FRHICommandList* RHICmdList, uint32 InputStreamStride, uint32 InputStreamFloatOffset, uint32 InputStreamVertexCount, uint32 OutputBufferFloatOffset, FBoneBufferTypeRef BoneBuffer, FUniformBufferRHIRef UniformBuffer, FShaderResourceViewRHIRef VertexBufferSRV, FRWBuffer& SkinBuffer, const FVector& MeshOrigin, const FVector& MeshExtension)
 	{
 		FComputeShaderRHIParamRef ComputeShaderRHI = GetComputeShader();
-		SetShaderValue(ComputeShaderRHI, SkinMeshOriginParameter, MeshOrigin);
-		SetShaderValue(ComputeShaderRHI, SkinMeshExtensionParameter, MeshExtension);
+		SetShaderValue(RHICmdList, ComputeShaderRHI, SkinMeshOriginParameter, MeshOrigin);
+		SetShaderValue(RHICmdList, ComputeShaderRHI, SkinMeshExtensionParameter, MeshExtension);
 
-		SetShaderValue(ComputeShaderRHI, SkinInputStreamStride, InputStreamStride);
-		SetShaderValue(ComputeShaderRHI, SkinInputStreamVertexCount, InputStreamVertexCount);
-		SetShaderValue(ComputeShaderRHI, SkinInputStreamFloatOffset, InputStreamFloatOffset);
-		SetShaderValue(ComputeShaderRHI, SkinOutputBufferFloatOffset, OutputBufferFloatOffset);
+		SetShaderValue(RHICmdList, ComputeShaderRHI, SkinInputStreamStride, InputStreamStride);
+		SetShaderValue(RHICmdList, ComputeShaderRHI, SkinInputStreamVertexCount, InputStreamVertexCount);
+		SetShaderValue(RHICmdList, ComputeShaderRHI, SkinInputStreamFloatOffset, InputStreamFloatOffset);
+		SetShaderValue(RHICmdList, ComputeShaderRHI, SkinOutputBufferFloatOffset, OutputBufferFloatOffset);
+
+		check(!RHICmdList);
 
 		if (UniformBuffer)
 		{
-			SetUniformBufferParameter(ComputeShaderRHI, GetUniformBufferParameter<GPUSkinCacheBonesUniformShaderParameters>(), UniformBuffer);
+			SetUniformBufferParameter(RHICmdList, ComputeShaderRHI, GetUniformBufferParameter<GPUSkinCacheBonesUniformShaderParameters>(), UniformBuffer);
 		}
 		else
 		{
@@ -108,8 +111,10 @@ public:
 			);
 	}
 
-	void UnsetParameters()
+	void UnsetParameters(FRHICommandList* RHICmdList)
 	{
+		check(!RHICmdList);
+
 		FComputeShaderRHIParamRef ComputeShaderRHI = GetComputeShader();
 		FShaderResourceViewRHIParamRef NullSRV = FShaderResourceViewRHIParamRef();
 #if GPUSKIN_USE_DATA_BUFFERS
@@ -352,7 +357,9 @@ int32 FGPUSkinCache::StartCacheMesh(int32 Key, const FVertexFactory* VertexFacto
 	}
 
 	auto& ShaderData = GPUVertexFactory->GetShaderData();
-	DispatchSkinCacheProcess(InputStreamFloatOffset, ECSInfo->StreamOffset, ShaderData.GetBoneBuffer(), ShaderData.GetUniformBuffer(), VBInfo, StreamStrides[0], ElementVertexCount, ShaderData.MeshOrigin, ShaderData.MeshExtension, bExtraBoneInfluences);
+	//@todo-rco: RHIPacketList
+	FRHICommandList* RHICmdList = nullptr;
+	DispatchSkinCacheProcess(RHICmdList, InputStreamFloatOffset, ECSInfo->StreamOffset, ShaderData.GetBoneBuffer(), ShaderData.GetUniformBuffer(), VBInfo, StreamStrides[0], ElementVertexCount, ShaderData.MeshOrigin, ShaderData.MeshExtension, bExtraBoneInfluences);
 
 	CacheCurrentFloatOffset += ProcessedFloatCount;
 	CachedChunksThisFrameCount++;
@@ -380,7 +387,7 @@ bool FGPUSkinCache::IsElementProcessed(int32 Key)
 	return false;
 }
 
-bool FGPUSkinCache::SetVertexStreamFromCache(int32 Key, FShader* Shader, const FVertexFactory* VertexFactory, uint32 BaseVertexIndex, bool VelocityPass, FShaderParameter PreviousStreamFloatOffset, FShaderParameter PreviousStreamFloatStride, FShaderResourceParameter PreviousStreamBuffer)
+bool FGPUSkinCache::SetVertexStreamFromCache(FRHICommandList* RHICmdList, int32 Key, FShader* Shader, const FVertexFactory* VertexFactory, uint32 BaseVertexIndex, bool VelocityPass, FShaderParameter PreviousStreamFloatOffset, FShaderParameter PreviousStreamFloatStride, FShaderResourceParameter PreviousStreamBuffer)
 {
 	if (Key >= 0 && Key < CachedElements.Num())
 	{
@@ -395,6 +402,8 @@ bool FGPUSkinCache::SetVertexStreamFromCache(int32 Key, FShader* Shader, const F
 				return false;
 			}
 
+			check(!RHICmdList);
+
 			RHISetStreamSource(0, SkinCacheBufferRW.Buffer, CacheInfo.StreamStride, CompensatedOffset);
 
 			FVertexShaderRHIParamRef VertexShaderRHI = Shader->GetVertexShader();
@@ -408,12 +417,12 @@ bool FGPUSkinCache::SetVertexStreamFromCache(int32 Key, FShader* Shader, const F
 
 				if (PreviousStreamFloatOffset.IsBound())
 				{
-					SetShaderValue(VertexShaderRHI, PreviousStreamFloatOffset, PreviousCompensatedOffset/4);
+					SetShaderValue(RHICmdList, VertexShaderRHI, PreviousStreamFloatOffset, PreviousCompensatedOffset/4);
 				}
 
 				if (PreviousStreamFloatStride.IsBound())
 				{
-					SetShaderValue(VertexShaderRHI, PreviousStreamFloatStride, CacheInfo.StreamStride/4);
+					SetShaderValue(RHICmdList, VertexShaderRHI, PreviousStreamFloatStride, CacheInfo.StreamStride/4);
 				}
 
 				if (PreviousStreamBuffer.IsBound())
@@ -469,7 +478,7 @@ FGPUSkinCache::FElementCacheStatusInfo* FGPUSkinCache::FindEvictableCacheStatusI
 	}
 }
 
-void FGPUSkinCache::DispatchSkinCacheProcess(uint32 InputStreamFloatOffset, uint32 OutputBufferFloatOffset, FBoneBufferTypeRef BoneBuffer, FUniformBufferRHIRef UniformBuffer, const FVertexBufferInfo* VBInfo, uint32 VertexStride, uint32 VertexCount, const FVector& MeshOrigin, const FVector& MeshExtension, bool bUseExtraBoneInfluences)
+void FGPUSkinCache::DispatchSkinCacheProcess(FRHICommandList* RHICmdList, uint32 InputStreamFloatOffset, uint32 OutputBufferFloatOffset, FBoneBufferTypeRef BoneBuffer, FUniformBufferRHIRef UniformBuffer, const FVertexBufferInfo* VBInfo, uint32 VertexStride, uint32 VertexCount, const FVector& MeshOrigin, const FVector& MeshExtension, bool bUseExtraBoneInfluences)
 {
 	if (GRHIFeatureLevel < ERHIFeatureLevel::SM5)
 	{
@@ -483,20 +492,21 @@ void FGPUSkinCache::DispatchSkinCacheProcess(uint32 InputStreamFloatOffset, uint
 	if (bUseExtraBoneInfluences)
 	{
 		TShaderMapRef<FGPUSkinCacheCS<true> > SkinCacheCS(GetGlobalShaderMap());
+		SkinCacheCS->SetParameters(RHICmdList, VertexStride, InputStreamFloatOffset, VertexCount, OutputBufferFloatOffset, BoneBuffer, UniformBuffer, VBInfo->VertexBufferSRV, SkinCacheBufferRW, MeshOrigin, MeshExtension);
+		check(!RHICmdList);
 
 		RHISetComputeShader(SkinCacheCS->GetComputeShader());		
-		SkinCacheCS->SetParameters(VertexStride, InputStreamFloatOffset, VertexCount, OutputBufferFloatOffset, BoneBuffer, UniformBuffer, VBInfo->VertexBufferSRV, SkinCacheBufferRW, MeshOrigin, MeshExtension);
 		RHIDispatchComputeShader(VertexCountAlign64, 1, 1);
-		SkinCacheCS->UnsetParameters();
+		SkinCacheCS->UnsetParameters(RHICmdList);
 	}
 	else
 	{
 		TShaderMapRef<FGPUSkinCacheCS<false> > SkinCacheCS(GetGlobalShaderMap());
+		SkinCacheCS->SetParameters(RHICmdList, VertexStride, InputStreamFloatOffset, VertexCount, OutputBufferFloatOffset, BoneBuffer, UniformBuffer, VBInfo->VertexBufferSRV, SkinCacheBufferRW, MeshOrigin, MeshExtension);
 
 		RHISetComputeShader(SkinCacheCS->GetComputeShader());		
-		SkinCacheCS->SetParameters(VertexStride, InputStreamFloatOffset, VertexCount, OutputBufferFloatOffset, BoneBuffer, UniformBuffer, VBInfo->VertexBufferSRV, SkinCacheBufferRW, MeshOrigin, MeshExtension);
 		RHIDispatchComputeShader(VertexCountAlign64, 1, 1);
-		SkinCacheCS->UnsetParameters();
+		SkinCacheCS->UnsetParameters(RHICmdList);
 	}
 }
 
diff --git a/Engine/Source/Runtime/Engine/Private/GPUSkinCache.h b/Engine/Source/Runtime/Engine/Private/GPUSkinCache.h
--- a/Engine/Source/Runtime/Engine/Private/GPUSkinCache.h
+++ b/Engine/Source/Runtime/Engine/Private/GPUSkinCache.h
@@ -42,7 +42,7 @@ public:
 	//	Call End*
 	int32	StartCacheMesh(int32 Key, const class FVertexFactory* VertexFactory, const FVertexFactory* TargetVertexFactory, const struct FSkelMeshChunk& BatchElement, const FSkeletalMeshObjectGPUSkin* Skin, bool bExtraBoneInfluences);
 
-	bool	SetVertexStreamFromCache(int32 Key, FShader* Shader, const FVertexFactory* VertexFactory, uint32 BaseVertexIndex, bool VelocityPass, FShaderParameter PreviousStreamFloatOffset, FShaderParameter PreviousStreamFloatStride, FShaderResourceParameter PreviousStreamBuffer);
+	bool	SetVertexStreamFromCache(FRHICommandList* RHICmdList, int32 Key, FShader* Shader, const FVertexFactory* VertexFactory, uint32 BaseVertexIndex, bool VelocityPass, FShaderParameter PreviousStreamFloatOffset, FShaderParameter PreviousStreamFloatStride, FShaderResourceParameter PreviousStreamBuffer);
 
 	struct FElementCacheStatusInfo
 	{
@@ -113,7 +113,7 @@ private:
 
 	FElementCacheStatusInfo* FindEvictableCacheStatusInfo();
 
-	void	DispatchSkinCacheProcess(uint32 InputStreamFloatOffset, uint32 OutputBufferFloatOffset, FBoneBufferTypeRef BoneBuffer, FUniformBufferRHIRef UniformBuffer, const FVertexBufferInfo* VBInfo, uint32 VertexStride, uint32 VertexCount, const FVector& MeshOrigin, const FVector& MeshExtension, bool bUseExtraBoneInfluences);
+	void	DispatchSkinCacheProcess(FRHICommandList* RHICmdList, uint32 InputStreamFloatOffset, uint32 OutputBufferFloatOffset, FBoneBufferTypeRef BoneBuffer, FUniformBufferRHIRef UniformBuffer, const FVertexBufferInfo* VBInfo, uint32 VertexStride, uint32 VertexCount, const FVector& MeshOrigin, const FVector& MeshExtension, bool bUseExtraBoneInfluences);
 
 	bool	Initialized : 1;
 
diff --git a/Engine/Source/Runtime/Engine/Private/GPUSkinVertexFactory.cpp b/Engine/Source/Runtime/Engine/Private/GPUSkinVertexFactory.cpp
--- a/Engine/Source/Runtime/Engine/Private/GPUSkinVertexFactory.cpp
+++ b/Engine/Source/Runtime/Engine/Private/GPUSkinVertexFactory.cpp
@@ -9,6 +9,7 @@
 #include "SkeletalRenderGPUSkin.h"	// FPreviousPerBoneMotionBlur
 #include "GPUSkinCache.h"
 #include "ShaderParameters.h"
+#include "ShaderParameterUtils.h"
 
 IMPLEMENT_UNIFORM_BUFFER_STRUCT(FAPEXClothUniformShaderParameters,TEXT("APEXClothParam"));
 
@@ -196,7 +197,7 @@ void FGPUBaseSkinVertexFactory::ShaderDataType::UpdateBoneData()
 	{
 		check(NumBones * BoneMatrices.GetTypeSize() <= sizeof(GBoneUniformStruct));
 		FMemory::Memcpy(&GBoneUniformStruct, BoneMatrices.GetTypedData(), NumBones * sizeof(BoneMatrices[0]));
-		UniformBuffer = RHICreateUniformBuffer(&GBoneUniformStruct, sizeof(GBoneUniformStruct), UniformBuffer_MultiUse);
+		UniformBuffer = RHICreateUniformBuffer(&GBoneUniformStruct, FBoneMatricesUniformShaderParameters::StaticStruct.GetLayout(), UniformBuffer_MultiFrame);
 	}
 }
 
@@ -425,19 +426,23 @@ public:
 	/**
 	* Set any shader data specific to this vertex factory
 	*/
-	virtual void SetMesh(FShader* Shader,const FVertexFactory* VertexFactory,const FSceneView& View,const FMeshBatchElement& BatchElement,uint32 DataFlags) const OVERRIDE
+	virtual void SetMesh(FRHICommandList* RHICmdList, FShader* Shader,const FVertexFactory* VertexFactory,const FSceneView& View,const FMeshBatchElement& BatchElement,uint32 DataFlags) const OVERRIDE
 	{
+		check(!RHICmdList);
+
 		if(Shader->GetVertexShader())
 		{
 			const auto FeatureLevel = View.GetFeatureLevel();
 			const FGPUBaseSkinVertexFactory::ShaderDataType& ShaderData = ((const FGPUBaseSkinVertexFactory*)VertexFactory)->GetShaderData();
 
 			SetShaderValue(
+				RHICmdList, 
 				Shader->GetVertexShader(), 
 				MeshOriginParameter, 
 				ShaderData.MeshOrigin
 				);
 			SetShaderValue(
+				RHICmdList, 
 				Shader->GetVertexShader(), 
 				MeshExtensionParameter, 
 				ShaderData.MeshExtension
@@ -456,7 +461,7 @@ public:
 			}
 			else
 			{
-				SetUniformBufferParameter(Shader->GetVertexShader(), Shader->GetUniformBufferParameter<FBoneMatricesUniformShaderParameters>(), ShaderData.GetUniformBuffer());
+				SetUniformBufferParameter(RHICmdList, Shader->GetVertexShader(), Shader->GetUniformBufferParameter<FBoneMatricesUniformShaderParameters>(), ShaderData.GetUniformBuffer());
 			}
 
 			bool bLocalPerBoneMotionBlur = false;
@@ -491,6 +496,7 @@ public:
 					BoneIndexOffsetValue[3] = 0;
 
 					SetShaderValue(
+						RHICmdList, 
 						Shader->GetVertexShader(), 
 						BoneIndexOffset, 
 						BoneIndexOffsetValue
@@ -510,7 +516,7 @@ public:
 				}
 			}
 
-			SetShaderValue(Shader->GetVertexShader(), PerBoneMotionBlur, bLocalPerBoneMotionBlur);
+			SetShaderValue(RHICmdList, Shader->GetVertexShader(), PerBoneMotionBlur, bLocalPerBoneMotionBlur);
 		}
 	}
 
@@ -567,13 +573,13 @@ public:
 	/**
 	* Set any shader data specific to this vertex factory
 	*/
-	virtual void SetMesh(FShader* Shader,const FVertexFactory* VertexFactory,const FSceneView& View,const FMeshBatchElement& BatchElement,uint32 DataFlags) const OVERRIDE
+	virtual void SetMesh(FRHICommandList* RHICmdList, FShader* Shader,const FVertexFactory* VertexFactory,const FSceneView& View,const FMeshBatchElement& BatchElement,uint32 DataFlags) const OVERRIDE
 	{
-		FGPUSkinVertexFactoryShaderParameters::SetMesh(Shader, VertexFactory, View, BatchElement, DataFlags);
+		FGPUSkinVertexFactoryShaderParameters::SetMesh(RHICmdList, Shader, VertexFactory, View, BatchElement, DataFlags);
 
 		bool bIsGPUCached = false;
 
-		if (GEnableGPUSkinCache && GGPUSkinCache.SetVertexStreamFromCache(BatchElement.UserIndex, Shader, VertexFactory, BatchElement.MinVertexIndex, GPrevPerBoneMotionBlur.IsLocked(), GPUSkinCacheStreamFloatOffset, GPUSkinCacheStreamStride, GPUSkinCacheStreamBuffer))
+		if (GEnableGPUSkinCache && GGPUSkinCache.SetVertexStreamFromCache(RHICmdList, BatchElement.UserIndex, Shader, VertexFactory, BatchElement.MinVertexIndex, GPrevPerBoneMotionBlur.IsLocked(), GPUSkinCacheStreamFloatOffset, GPUSkinCacheStreamStride, GPUSkinCacheStreamBuffer))
 		{
 			bIsGPUCached = true;
 		}
@@ -705,12 +711,12 @@ public:
 		Ar << ClothBlendWeightParameter;
 	}
 
-	virtual void SetMesh(FShader* Shader,const FVertexFactory* VertexFactory,const FSceneView& View,const FMeshBatchElement& BatchElement,uint32 DataFlags) const OVERRIDE
+	virtual void SetMesh(FRHICommandList* RHICmdList, FShader* Shader,const FVertexFactory* VertexFactory,const FSceneView& View,const FMeshBatchElement& BatchElement,uint32 DataFlags) const OVERRIDE
 	{
 		if(Shader->GetVertexShader())
 		{
 			// Call regular GPU skinning shader parameters
-			FGPUSkinVertexFactoryShaderParameters::SetMesh(Shader, VertexFactory, View, BatchElement, DataFlags);
+			FGPUSkinVertexFactoryShaderParameters::SetMesh(RHICmdList, Shader, VertexFactory, View, BatchElement, DataFlags);
 			const auto* GPUSkinVertexFactory = (const FGPUBaseSkinVertexFactory*)VertexFactory;
 			// A little hacky; problem is we can't upcast from FGPUBaseSkinVertexFactory to FGPUBaseSkinAPEXClothVertexFactory as they are unrelated; a nice solution would be
 			// to use virtual inheritance, but that requires RTTI and complicates things further...
@@ -718,7 +724,8 @@ public:
 				? ((const TGPUSkinAPEXClothVertexFactory<true>*)GPUSkinVertexFactory)->GetClothShaderData()
 				: ((const TGPUSkinAPEXClothVertexFactory<false>*)GPUSkinVertexFactory)->GetClothShaderData();
 
-			SetUniformBufferParameter(Shader->GetVertexShader(),Shader->GetUniformBufferParameter<FAPEXClothUniformShaderParameters>(),ClothShaderData.GetClothUniformBuffer());
+			SetUniformBufferParameter(RHICmdList, Shader->GetVertexShader(),Shader->GetUniformBufferParameter<FAPEXClothUniformShaderParameters>(),ClothShaderData.GetClothUniformBuffer());
+			check(!RHICmdList);
 
 			// we tell the shader where to pickup the data
 			if(ClothSimulPositionsParameter.IsBound())
@@ -729,7 +736,7 @@ public:
 				RHISetShaderTexture(Shader->GetVertexShader(), ClothSimulPositionsParameter.GetBaseIndex(), ClothShaderData.GetClothSimulPositionBuffer());
 #endif
 			}
-
+			
 			if(ClothSimulNormalsParameter.IsBound())
 			{
 #if GPUSKIN_USE_DATA_BUFFERS
@@ -740,6 +747,7 @@ public:
 			}
 
 			SetShaderValue(
+				RHICmdList,
 				Shader->GetVertexShader(),
 				ClothBlendWeightParameter,
 				ClothShaderData.ClothBlendWeight
@@ -860,7 +868,7 @@ void FGPUBaseSkinAPEXClothVertexFactory::ClothShaderType::UpdateClothUniformBuff
 		}
 	}
 		
-	APEXClothUniformBuffer = TUniformBufferRef<FAPEXClothUniformShaderParameters>::CreateUniformBufferImmediate(ClothUniformShaderParameters, UniformBuffer_SingleUse);
+	APEXClothUniformBuffer = TUniformBufferRef<FAPEXClothUniformShaderParameters>::CreateUniformBufferImmediate(ClothUniformShaderParameters, UniformBuffer_SingleFrame);
 
 }
 /**
