UE4 - merge first pass at parallel rendering to main

diff --git a/Engine/Source/Runtime/Engine/Private/GPUSort.cpp b/Engine/Source/Runtime/Engine/Private/GPUSort.cpp
--- a/Engine/Source/Runtime/Engine/Private/GPUSort.cpp
+++ b/Engine/Source/Runtime/Engine/Private/GPUSort.cpp
@@ -9,6 +9,7 @@
 #include "RenderCore.h"
 #include "UniformBuffer.h"
 #include "ShaderParameters.h"
+#include "ShaderParameterUtils.h"
 #include "GlobalShader.h"
 
 DEFINE_LOG_CATEGORY_STATIC(LogGPUSort, Log, All);
@@ -255,10 +256,10 @@ public:
 	/**
 	 * Set parameters for this shader.
 	 */
-	void SetParameters( const FRadixSortUniformBufferRef& UniformBuffer )
+	void SetParameters(FRHICommandList* RHICmdList, const FRadixSortUniformBufferRef& UniformBuffer )
 	{
 		FComputeShaderRHIParamRef ComputeShaderRHI = GetComputeShader();
-		SetUniformBufferParameter( ComputeShaderRHI, GetUniformBufferParameter<FRadixSortParameters>(), UniformBuffer );
+		SetUniformBufferParameter(RHICmdList, ComputeShaderRHI, GetUniformBufferParameter<FRadixSortParameters>(), UniformBuffer );
 	}
 
 	/**
@@ -353,10 +354,11 @@ public:
 	/**
 	 * Set parameters for this shader.
 	 */
-	void SetParameters( FShaderResourceViewRHIParamRef InKeysSRV, FRadixSortUniformBufferRef& RadixSortUniformBuffer, FShaderResourceViewRHIParamRef RadixSortParameterBufferSRV )
+	void SetParameters(FRHICommandList* RHICmdList, FShaderResourceViewRHIParamRef InKeysSRV, FRadixSortUniformBufferRef& RadixSortUniformBuffer, FShaderResourceViewRHIParamRef RadixSortParameterBufferSRV )
 	{
 		FComputeShaderRHIParamRef ComputeShaderRHI = GetComputeShader();
-		SetUniformBufferParameter( ComputeShaderRHI, GetUniformBufferParameter<FRadixSortParameters>(), RadixSortUniformBuffer );
+		SetUniformBufferParameter(RHICmdList, ComputeShaderRHI, GetUniformBufferParameter<FRadixSortParameters>(), RadixSortUniformBuffer );
+		check(!RHICmdList);
 		if ( InKeys.IsBound() )
 		{
 			RHISetShaderResourceViewParameter( ComputeShaderRHI, InKeys.GetBaseIndex(), InKeysSRV );
@@ -460,11 +462,13 @@ public:
 	/**
 	 * Set parameters for this shader.
 	 */
-	void SetParameters( FShaderResourceViewRHIParamRef InOffsetsSRV )
+	void SetParameters(FRHICommandList* RHICmdList, FShaderResourceViewRHIParamRef InOffsetsSRV )
 	{
 		FComputeShaderRHIParamRef ComputeShaderRHI = GetComputeShader();
 		if ( InOffsets.IsBound() )
 		{
+			check(!RHICmdList);
+
 			RHISetShaderResourceViewParameter( ComputeShaderRHI, InOffsets.GetBaseIndex(), InOffsetsSRV );
 		}
 	}
@@ -576,6 +580,7 @@ public:
 	 * Set parameters for this shader.
 	 */
 	void SetParameters(
+		FRHICommandList* RHICmdList, 
 		FShaderResourceViewRHIParamRef InKeysSRV,
 		FShaderResourceViewRHIParamRef InValuesSRV,
 		FShaderResourceViewRHIParamRef InOffsetsSRV,
@@ -583,7 +588,8 @@ public:
 		FShaderResourceViewRHIParamRef RadixSortParameterBufferSRV )
 	{
 		FComputeShaderRHIParamRef ComputeShaderRHI = GetComputeShader();
-		SetUniformBufferParameter( ComputeShaderRHI, GetUniformBufferParameter<FRadixSortParameters>(), RadixSortUniformBuffer );
+		SetUniformBufferParameter(RHICmdList, ComputeShaderRHI, GetUniformBufferParameter<FRadixSortParameters>(), RadixSortUniformBuffer );
+		check(!RHICmdList);
 		if ( RadixSortParameterBuffer.IsBound() )
 		{
 			RHISetShaderResourceViewParameter( ComputeShaderRHI, RadixSortParameterBuffer.GetBaseIndex(), RadixSortParameterBufferSRV );
@@ -735,6 +741,9 @@ int32 SortGPUBuffers( FGPUSortBuffers SortBuffers, int32 BufferIndex, uint32 Key
 	check( UpsweepCS->RequiresConstantBufferWorkaround() == DownsweepCS->RequiresConstantBufferWorkaround() );
 	const bool bUseConstantBufferWorkaround = UpsweepCS->RequiresConstantBufferWorkaround();
 
+	//@todo-rco: RHIPacketList
+	FRHICommandList* RHICmdList = nullptr;
+
 	// Execute each pass as needed.
 	uint32 PassBits = DIGIT_COUNT - 1;
 	for ( int32 PassIndex = 0; PassIndex < PassCount; ++PassIndex )
@@ -751,20 +760,20 @@ int32 SortGPUBuffers( FGPUSortBuffers SortBuffers, int32 BufferIndex, uint32 Key
 			}
 			else
 			{
-				SortUniformBufferRef = FRadixSortUniformBufferRef::CreateUniformBufferImmediate( SortParameters, UniformBuffer_SingleUse );
+				SortUniformBufferRef = FRadixSortUniformBufferRef::CreateUniformBufferImmediate( SortParameters, UniformBuffer_SingleDraw );
 			}
 
 			// Clear the offsets buffer.
 			RHISetComputeShader(ClearOffsetsCS->GetComputeShader());
 			ClearOffsetsCS->SetOutput( GSortOffsetBuffers.BufferUAVs[0] );
-			DispatchComputeShader( *ClearOffsetsCS, 1, 1 ,1 );
+			DispatchComputeShader(RHICmdList, *ClearOffsetsCS, 1, 1 ,1 );
 			ClearOffsetsCS->UnbindBuffers();
 
 			// Phase 1: Scan upsweep to compute per-digit totals.
 			RHISetComputeShader(UpsweepCS->GetComputeShader());
 			UpsweepCS->SetOutput( GSortOffsetBuffers.BufferUAVs[0] );
-			UpsweepCS->SetParameters( SortBuffers.RemoteKeySRVs[BufferIndex], SortUniformBufferRef, GRadixSortParametersBuffer.SortParametersBufferSRV );
-			DispatchComputeShader( *UpsweepCS, GroupCount, 1, 1 );
+			UpsweepCS->SetParameters(RHICmdList, SortBuffers.RemoteKeySRVs[BufferIndex], SortUniformBufferRef, GRadixSortParametersBuffer.SortParametersBufferSRV );
+			DispatchComputeShader(RHICmdList, *UpsweepCS, GroupCount, 1, 1 );
 			UpsweepCS->UnbindBuffers();
 
 			if (bDebugOffsets)
@@ -776,8 +785,8 @@ int32 SortGPUBuffers( FGPUSortBuffers SortBuffers, int32 BufferIndex, uint32 Key
 			// Phase 2: Parallel prefix scan on the offsets buffer.
 			RHISetComputeShader(SpineCS->GetComputeShader());
 			SpineCS->SetOutput( GSortOffsetBuffers.BufferUAVs[1] );
-			SpineCS->SetParameters( GSortOffsetBuffers.BufferSRVs[0] );
-			DispatchComputeShader( *SpineCS, 1, 1, 1 );
+			SpineCS->SetParameters(RHICmdList, GSortOffsetBuffers.BufferSRVs[0] );
+			DispatchComputeShader(RHICmdList, *SpineCS, 1, 1, 1 );
 			SpineCS->UnbindBuffers();
 
 			if (bDebugOffsets)
@@ -789,8 +798,8 @@ int32 SortGPUBuffers( FGPUSortBuffers SortBuffers, int32 BufferIndex, uint32 Key
 			// Phase 3: Downsweep to compute final offsets and scatter keys.
 			RHISetComputeShader(DownsweepCS->GetComputeShader());
 			DownsweepCS->SetOutput( SortBuffers.RemoteKeyUAVs[BufferIndex ^ 0x1], SortBuffers.RemoteValueUAVs[BufferIndex ^ 0x1] );
-			DownsweepCS->SetParameters( SortBuffers.RemoteKeySRVs[BufferIndex], SortBuffers.RemoteValueSRVs[BufferIndex], GSortOffsetBuffers.BufferSRVs[1], SortUniformBufferRef, GRadixSortParametersBuffer.SortParametersBufferSRV );
-			DispatchComputeShader( *DownsweepCS, GroupCount, 1, 1 );
+			DownsweepCS->SetParameters(RHICmdList, SortBuffers.RemoteKeySRVs[BufferIndex], SortBuffers.RemoteValueSRVs[BufferIndex], GSortOffsetBuffers.BufferSRVs[1], SortUniformBufferRef, GRadixSortParametersBuffer.SortParametersBufferSRV );
+			DispatchComputeShader(RHICmdList, *DownsweepCS, GroupCount, 1, 1 );
 			DownsweepCS->UnbindBuffers();
 
 			// Flip buffers.
diff --git a/Engine/Source/Runtime/Engine/Private/InstancedStaticMesh.cpp b/Engine/Source/Runtime/Engine/Private/InstancedStaticMesh.cpp
--- a/Engine/Source/Runtime/Engine/Private/InstancedStaticMesh.cpp
+++ b/Engine/Source/Runtime/Engine/Private/InstancedStaticMesh.cpp
@@ -7,7 +7,7 @@
 #include "EnginePrivate.h"
 #include "PhysicsPublic.h"
 #include "ShaderParameters.h"
-#include "StaticLighting.h"
+#include "ShaderParameterUtils.h"
 
 #include "Components/InteractiveFoliageComponent.h"
 #include "Components/SplineMeshComponent.h"
@@ -512,12 +512,13 @@ class FInstancedStaticMeshVertexFactoryShaderParameters : public FVertexFactoryS
 		InstancingFadeOutParamsParameter.Bind(ParameterMap,TEXT("InstancingFadeOutParams"));
 	}
 
-	virtual void SetMesh(FShader* VertexShader,const class FVertexFactory* VertexFactory,const class FSceneView& View,const struct FMeshBatchElement& BatchElement,uint32 DataFlags) const OVERRIDE
+	virtual void SetMesh(FRHICommandList* RHICmdList, FShader* VertexShader,const class FVertexFactory* VertexFactory,const class FSceneView& View,const struct FMeshBatchElement& BatchElement,uint32 DataFlags) const OVERRIDE
 	{
 		if( InstancedViewTranslationParameter.IsBound() )
 		{
 			FVector4 InstancedViewTranslation(View.ViewMatrices.PreViewTranslation, 0.f);
 			SetShaderValue(
+				RHICmdList, 
 				VertexShader->GetVertexShader(),
 				InstancedViewTranslationParameter,
 				InstancedViewTranslation
@@ -551,7 +552,7 @@ class FInstancedStaticMeshVertexFactoryShaderParameters : public FVertexFactoryS
 				InstancingFadeOutParams.Z = InstancingUserData->bRenderSelected ? 1.f : 0.f;
 				InstancingFadeOutParams.W = InstancingUserData->bRenderUnselected ? 1.f : 0.f;
 			}
-			SetShaderValue( VertexShader->GetVertexShader(), InstancingFadeOutParamsParameter, InstancingFadeOutParams );
+			SetShaderValue(RHICmdList, VertexShader->GetVertexShader(), InstancingFadeOutParamsParameter, InstancingFadeOutParams );
 		}
 
 		// upload SpeedTree buffer, if available
@@ -560,7 +561,7 @@ class FInstancedStaticMeshVertexFactoryShaderParameters : public FVertexFactoryS
 			FUniformBufferRHIParamRef SpeedTreeUniformBuffer = View.Family->Scene->GetSpeedTreeUniformBuffer(VertexFactory);
 			if (SpeedTreeUniformBuffer != NULL)
 			{
-				SetUniformBufferParameter(VertexShader->GetVertexShader(), VertexShader->GetUniformBufferParameter<FSpeedTreeUniformParameters>(), SpeedTreeUniformBuffer);
+				SetUniformBufferParameter(RHICmdList, VertexShader->GetVertexShader(), VertexShader->GetUniformBufferParameter<FSpeedTreeUniformParameters>(), SpeedTreeUniformBuffer);
 			}
 		}
 	}
@@ -1235,7 +1236,7 @@ void UInstancedStaticMeshComponent::InitInstanceBody(int32 InstanceIdx, FBodyIns
 
 	UBodySetup* BodySetup = GetBodySetup();
 	check(BodySetup);
-	
+
 	// Get transform of the instance
 	FTransform InstanceTransform = FTransform(PerInstanceSMData[InstanceIdx].Transform) * ComponentToWorld;
 	
diff --git a/Engine/Source/Runtime/Engine/Private/Landscape/LandscapeRender.cpp b/Engine/Source/Runtime/Engine/Private/Landscape/LandscapeRender.cpp
--- a/Engine/Source/Runtime/Engine/Private/Landscape/LandscapeRender.cpp
+++ b/Engine/Source/Runtime/Engine/Private/Landscape/LandscapeRender.cpp
@@ -8,6 +8,8 @@ LandscapeRender.cpp: New terrain rendering
 #include "Materials/MaterialExpressionTextureCoordinate.h"
 #include "Materials/MaterialExpressionLandscapeLayerCoords.h"
 #include "ShaderParameters.h"
+#include "ShaderParameterUtils.h"
+#include "RHIStaticStates.h"
 #include "Landscape/LandscapeRender.h"
 #include "Landscape/LandscapeEdit.h"
 #include "LevelUtils.h"
@@ -1068,7 +1070,7 @@ void FLandscapeComponentSceneProxy::DrawStaticElements(FStaticPrimitiveDrawInter
 	{
 		int32 LodSubsectionSizeVerts = SubsectionSizeVerts >> LOD;
 
-		if (ForcedLOD < 0 && NumSubsections > 1)
+		if( ForcedLOD < 0 && NumSubsections > 1 )
 		{
 			// Per-subsection batch elements
 			for( int32 SubY=0;SubY<NumSubsections;SubY++ )
@@ -1611,7 +1613,7 @@ template <typename INDEX_TYPE>
 void FLandscapeSharedBuffers::CreateIndexBuffers()
 {
 	if (GRHIFeatureLevel == ERHIFeatureLevel::ES2)
-	{
+{
 		if (!bVertexScoresComputed)
 		{
 			bVertexScoresComputed = ComputeVertexScores();
@@ -1642,11 +1644,11 @@ void FLandscapeSharedBuffers::CreateIndexBuffers()
 			// ES2 version
 			float MipRatio = (float)SubsectionSizeQuads / (float)LodSubsectionSizeQuads; // Morph current MIP to base MIP
 
-			for (int32 SubY = 0; SubY < NumSubsections; SubY++)
+		for( int32 SubY=0;SubY<NumSubsections;SubY++ )
+		{
+			for( int32 SubX=0;SubX<NumSubsections;SubX++ )
 			{
-				for (int32 SubX = 0; SubX < NumSubsections; SubX++)
-				{
-					TArray<INDEX_TYPE> SubIndices;
+				TArray<INDEX_TYPE> SubIndices;
 					SubIndices.Empty(FMath::Square(LodSubsectionSizeQuads) * 6);
 
 					int32& MaxIndex = IndexRanges[Mip].MaxIndex[SubX][SubY];
@@ -1655,9 +1657,9 @@ void FLandscapeSharedBuffers::CreateIndexBuffers()
 					MinIndex = MAX_int32;
 
 					for (int32 y = 0; y < LodSubsectionSizeQuads; y++)
-					{
+				{
 						for (int32 x = 0; x < LodSubsectionSizeQuads; x++)
-						{
+					{
 							int32 x0 = FMath::RoundToInt((float)x * MipRatio);
 							int32 y0 = FMath::RoundToInt((float)y * MipRatio);
 							int32 x1 = FMath::RoundToInt((float)(x + 1) * MipRatio);
@@ -1672,7 +1674,7 @@ void FLandscapeSharedBuffers::CreateIndexBuffers()
 							uint64 Key10 = V10.MakeKey();
 							uint64 Key11 = V11.MakeKey();
 							uint64 Key01 = V01.MakeKey();
-							
+
 							INDEX_TYPE i00;
 							INDEX_TYPE i10;
 							INDEX_TYPE i11;
@@ -1683,9 +1685,9 @@ void FLandscapeSharedBuffers::CreateIndexBuffers()
 							{
 								i00 = VertexCount++;
 								VertexMap.Add(Key00, i00);
-							}
-							else
-							{
+						}
+						else
+						{
 								i00 = *KeyPtr;
 							}
 
@@ -1698,27 +1700,27 @@ void FLandscapeSharedBuffers::CreateIndexBuffers()
 							else
 							{
 								i10 = *KeyPtr;
-							}
-							
+						}
+
 							KeyPtr = VertexMap.Find(Key11);
 							if (KeyPtr == NULL)
 							{
 								i11 = VertexCount++;
 								VertexMap.Add(Key11, i11);
-							}
+					}
 							else
 							{
 								i11 = *KeyPtr;
-							}
+				}
 
 							KeyPtr = VertexMap.Find(Key01);
 							if (KeyPtr == NULL)
-							{
+				{
 								i01 = VertexCount++;
 								VertexMap.Add(Key01, i01);
 							}
 							else
-							{
+					{
 								i01 = *KeyPtr;
 							}
 
@@ -1799,11 +1801,11 @@ void FLandscapeSharedBuffers::CreateIndexBuffers()
 					MaxIndexFull = FMath::Max<int32>(MaxIndexFull, MaxIndex);
 					MinIndexFull = FMath::Min<int32>(MinIndexFull, MinIndex);
 				}
-			}
+	}
 
 			check(MinIndexFull <= (uint32)((INDEX_TYPE)(~(INDEX_TYPE)0)));
 			check(NewIndices.Num() == ExpectedNumIndices);
-		}
+	}
 
 		// Create and init new index buffer with index data
 		auto IndexBuffer = new FRawStaticIndexBuffer16or32<INDEX_TYPE>(false);
@@ -1990,7 +1992,7 @@ public:
 	/**
 	* Set any shader data specific to this vertex factory
 	*/
-	virtual void SetMesh(FShader* VertexShader,const class FVertexFactory* VertexFactory,const class FSceneView& View,const struct FMeshBatchElement& BatchElement,uint32 DataFlags) const OVERRIDE
+	virtual void SetMesh(FRHICommandList* RHICmdList, FShader* VertexShader,const class FVertexFactory* VertexFactory,const class FSceneView& View,const struct FMeshBatchElement& BatchElement,uint32 DataFlags) const OVERRIDE
 	{
 		SCOPE_CYCLE_COUNTER(STAT_LandscapeVFDrawTime);
 
@@ -1998,11 +2000,12 @@ public:
 		check(BatchElementParams);
 
 		const FLandscapeComponentSceneProxy* SceneProxy = BatchElementParams->SceneProxy;
-		SetUniformBufferParameter(VertexShader->GetVertexShader(),VertexShader->GetUniformBufferParameter<FLandscapeUniformShaderParameters>(),*BatchElementParams->LandscapeUniformShaderParametersResource);
+		SetUniformBufferParameter(RHICmdList, VertexShader->GetVertexShader(),VertexShader->GetUniformBufferParameter<FLandscapeUniformShaderParameters>(),*BatchElementParams->LandscapeUniformShaderParametersResource);
 
 		if( HeightmapTextureParameter.IsBound() )
 		{
 			SetTextureParameter(
+				RHICmdList, 
 				VertexShader->GetVertexShader(),
 				HeightmapTextureParameter,
 				HeightmapTextureParameterSampler,
@@ -2019,7 +2022,7 @@ public:
 				SceneProxy->HeightmapTexture->GetNumMips() - FMath::Min(SceneProxy->HeightmapTexture->ResidentMips, SceneProxy->HeightmapTexture->RequestedMips),
 				0.f // Reserved
 				);
-			SetShaderValue(VertexShader->GetVertexShader(), LodBiasParameter, LodBias);
+			SetShaderValue(RHICmdList, VertexShader->GetVertexShader(), LodBiasParameter, LodBias);
 		}
 
 		// Calculate LOD params
@@ -2049,12 +2052,12 @@ public:
 
 		if( SectionLodsParameter.IsBound() )
 		{
-			SetShaderValue(VertexShader->GetVertexShader(), SectionLodsParameter, fCurrentLODs);
+			SetShaderValue(RHICmdList, VertexShader->GetVertexShader(), SectionLodsParameter, fCurrentLODs);
 		}
 
 		if( NeighborSectionLodParameter.IsBound() )
 		{
-			SetShaderValue(VertexShader->GetVertexShader(), NeighborSectionLodParameter, CurrentNeighborLODs);
+			SetShaderValue(RHICmdList, VertexShader->GetVertexShader(), NeighborSectionLodParameter, CurrentNeighborLODs);
 		}
 
 		if( LodValuesParameter.IsBound() )
@@ -2066,7 +2069,7 @@ public:
 				(float)((SceneProxy->SubsectionSizeVerts >> BatchElementParams->CurrentLOD) - 1),
 				1.f/(float)((SceneProxy->SubsectionSizeVerts >> BatchElementParams->CurrentLOD) - 1) );
 
-			SetShaderValue(VertexShader->GetVertexShader(),LodValuesParameter,LodValues);
+			SetShaderValue(RHICmdList, VertexShader->GetVertexShader(),LodValuesParameter,LodValues);
 		}
 	}
 
@@ -2094,7 +2097,7 @@ public:
 	/**
 	* Set any shader data specific to this vertex factory
 	*/
-	virtual void SetMesh(FShader* VertexShader,const class FVertexFactory* VertexFactory,const class FSceneView& View,const struct FMeshBatchElement& BatchElement,uint32 DataFlags) const OVERRIDE
+	virtual void SetMesh(FRHICommandList* RHICmdList, FShader* VertexShader,const class FVertexFactory* VertexFactory,const class FSceneView& View,const struct FMeshBatchElement& BatchElement,uint32 DataFlags) const OVERRIDE
 	{
 		SCOPE_CYCLE_COUNTER(STAT_LandscapeVFDrawTime);
 
@@ -2102,11 +2105,12 @@ public:
 		check(BatchElementParams);
 
 		const FLandscapeComponentSceneProxy* SceneProxy = BatchElementParams->SceneProxy;
-		SetUniformBufferParameter(VertexShader->GetVertexShader(),VertexShader->GetUniformBufferParameter<FLandscapeUniformShaderParameters>(),*BatchElementParams->LandscapeUniformShaderParametersResource);
+		SetUniformBufferParameter(RHICmdList, VertexShader->GetVertexShader(),VertexShader->GetUniformBufferParameter<FLandscapeUniformShaderParameters>(),*BatchElementParams->LandscapeUniformShaderParametersResource);
 
 		if( HeightmapTextureParameter.IsBound() )
 		{
 			SetTextureParameter(
+				RHICmdList, 
 				VertexShader->GetVertexShader(),
 				HeightmapTextureParameter,
 				HeightmapTextureParameterSampler,
@@ -2123,7 +2127,7 @@ public:
 				SceneProxy->HeightmapTexture->GetNumMips() - FMath::Min(SceneProxy->HeightmapTexture->ResidentMips, SceneProxy->HeightmapTexture->RequestedMips),
 				SceneProxy->XYOffsetmapTexture->GetNumMips() - FMath::Min(SceneProxy->XYOffsetmapTexture->ResidentMips, SceneProxy->XYOffsetmapTexture->RequestedMips)
 				);
-			SetShaderValue(VertexShader->GetVertexShader(), LodBiasParameter, LodBias);
+			SetShaderValue(RHICmdList, VertexShader->GetVertexShader(), LodBiasParameter, LodBias);
 		}
 
 		// Calculate LOD params
@@ -2153,12 +2157,12 @@ public:
 
 		if( SectionLodsParameter.IsBound() )
 		{
-			SetShaderValue(VertexShader->GetVertexShader(), SectionLodsParameter, fCurrentLODs);
+			SetShaderValue(RHICmdList, VertexShader->GetVertexShader(), SectionLodsParameter, fCurrentLODs);
 		}
 
 		if( NeighborSectionLodParameter.IsBound() )
 		{
-			SetShaderValue(VertexShader->GetVertexShader(), NeighborSectionLodParameter, CurrentNeighborLODs);
+			SetShaderValue(RHICmdList, VertexShader->GetVertexShader(), NeighborSectionLodParameter, CurrentNeighborLODs);
 		}
 
 		if( LodValuesParameter.IsBound() )
@@ -2170,12 +2174,13 @@ public:
 				(float)((SceneProxy->SubsectionSizeVerts >> BatchElementParams->CurrentLOD) - 1),
 				1.f/(float)((SceneProxy->SubsectionSizeVerts >> BatchElementParams->CurrentLOD) - 1) );
 
-			SetShaderValue(VertexShader->GetVertexShader(),LodValuesParameter,LodValues);
+			SetShaderValue(RHICmdList, VertexShader->GetVertexShader(),LodValuesParameter,LodValues);
 		}
 
 		if( XYOffsetTextureParameter.IsBound() && SceneProxy->XYOffsetmapTexture )
 		{
 			SetTextureParameter(
+				RHICmdList, 
 				VertexShader->GetVertexShader(),
 				XYOffsetTextureParameter,
 				XYOffsetTextureParameterSampler,
@@ -2214,7 +2219,7 @@ void FLandscapeVertexFactoryPixelShaderParameters::Serialize(FArchive& Ar)
 /**
 * Set any shader data specific to this vertex factory
 */
-void FLandscapeVertexFactoryPixelShaderParameters::SetMesh(FShader* PixelShader,const class FVertexFactory* VertexFactory,const class FSceneView& View,const struct FMeshBatchElement& BatchElement,uint32 DataFlags) const
+void FLandscapeVertexFactoryPixelShaderParameters::SetMesh(FRHICommandList* RHICmdList, FShader* PixelShader,const class FVertexFactory* VertexFactory,const class FSceneView& View,const struct FMeshBatchElement& BatchElement,uint32 DataFlags) const
 {
 	SCOPE_CYCLE_COUNTER(STAT_LandscapeVFDrawTime);
 
@@ -2222,12 +2227,13 @@ void FLandscapeVertexFactoryPixelShaderParameters::SetMesh(FShader* PixelShader,
 
 	if( LocalToWorldNoScalingParameter.IsBound() )
 	{
-		SetShaderValue(PixelShader->GetPixelShader(), LocalToWorldNoScalingParameter, *BatchElementParams->LocalToWorldNoScalingPtr);
+		SetShaderValue(RHICmdList, PixelShader->GetPixelShader(), LocalToWorldNoScalingParameter, *BatchElementParams->LocalToWorldNoScalingPtr);
 	}
 
 	if( NormalmapTextureParameter.IsBound() )
 	{
 		SetTextureParameter(
+			RHICmdList, 
 			PixelShader->GetPixelShader(),
 			NormalmapTextureParameter,
 			NormalmapTextureParameterSampler,
@@ -2446,7 +2452,7 @@ void ALandscapeProxy::ChangeLODDistanceFactor(float InLODDistanceFactor)
 			LODFactor = LODDistanceFactor;
 			break;
 	}
-
+	
 	if (LandscapeComponents.Num())
 	{
 		int32 CompNum = LandscapeComponents.Num();
diff --git a/Engine/Source/Runtime/Engine/Private/Landscape/LandscapeRenderMobile.cpp b/Engine/Source/Runtime/Engine/Private/Landscape/LandscapeRenderMobile.cpp
--- a/Engine/Source/Runtime/Engine/Private/Landscape/LandscapeRenderMobile.cpp
+++ b/Engine/Source/Runtime/Engine/Private/Landscape/LandscapeRenderMobile.cpp
@@ -6,6 +6,7 @@ LandscapeRenderMobile.cpp: Landscape Rendering without using vertex texture fetc
 
 #include "EnginePrivate.h"
 #include "ShaderParameters.h"
+#include "ShaderParameterUtils.h"
 #include "Landscape/LandscapeRender.h"
 #include "Landscape/LandscapeRenderMobile.h"
 
@@ -60,7 +61,7 @@ public:
 	/**
 	* Set any shader data specific to this vertex factory
 	*/
-	virtual void SetMesh(FShader* VertexShader,const class FVertexFactory* VertexFactory,const class FSceneView& View,const struct FMeshBatchElement& BatchElement,uint32 DataFlags) const OVERRIDE
+	virtual void SetMesh(FRHICommandList* RHICmdList, FShader* VertexShader,const class FVertexFactory* VertexFactory,const class FSceneView& View,const struct FMeshBatchElement& BatchElement,uint32 DataFlags) const OVERRIDE
 	{
 		SCOPE_CYCLE_COUNTER(STAT_LandscapeVFDrawTime);
 
@@ -68,7 +69,7 @@ public:
 		check(BatchElementParams);
 
 		const FLandscapeComponentSceneProxy* SceneProxy = BatchElementParams->SceneProxy;
-		SetUniformBufferParameter(VertexShader->GetVertexShader(),VertexShader->GetUniformBufferParameter<FLandscapeUniformShaderParameters>(),*BatchElementParams->LandscapeUniformShaderParametersResource);
+		SetUniformBufferParameter(RHICmdList, VertexShader->GetVertexShader(),VertexShader->GetUniformBufferParameter<FLandscapeUniformShaderParameters>(),*BatchElementParams->LandscapeUniformShaderParametersResource);
 
 		FVector CameraLocalPos3D = SceneProxy->WorldToLocal.TransformPosition(View.ViewMatrices.ViewOrigin); 
 		FVector2D CameraLocalPos = FVector2D(CameraLocalPos3D.X, CameraLocalPos3D.Y);
@@ -81,7 +82,7 @@ public:
 				CameraLocalPos3D.X + SceneProxy->SectionBase.X,
 				CameraLocalPos3D.Y + SceneProxy->SectionBase.Y 
 				);
-			SetShaderValue(VertexShader->GetVertexShader(), LodBiasParameter, LodBias);
+			SetShaderValue(RHICmdList, VertexShader->GetVertexShader(), LodBiasParameter, LodBias);
 		}
 
 		// Calculate LOD params
@@ -108,12 +109,12 @@ public:
 
 		if( SectionLodsParameter.IsBound() )
 		{
-			SetShaderValue(VertexShader->GetVertexShader(), SectionLodsParameter, fCurrentLODs);
+			SetShaderValue(RHICmdList, VertexShader->GetVertexShader(), SectionLodsParameter, fCurrentLODs);
 		}
 
 		if( NeighborSectionLodParameter.IsBound() )
 		{
-			SetShaderValue(VertexShader->GetVertexShader(), NeighborSectionLodParameter, CurrentNeighborLODs);
+			SetShaderValue(RHICmdList, VertexShader->GetVertexShader(), NeighborSectionLodParameter, CurrentNeighborLODs);
 		}
 
 		if( LodValuesParameter.IsBound() )
@@ -125,7 +126,7 @@ public:
 				(float)SceneProxy->SubsectionSizeQuads, 
 				1.f / (float)SceneProxy->SubsectionSizeQuads );
 
-			SetShaderValue(VertexShader->GetVertexShader(),LodValuesParameter,LodValues);
+			SetShaderValue(RHICmdList, VertexShader->GetVertexShader(),LodValuesParameter,LodValues);
 		}
 	}
 protected:
diff --git a/Engine/Source/Runtime/Engine/Private/Landscape/LandscapeSplines.cpp b/Engine/Source/Runtime/Engine/Private/Landscape/LandscapeSplines.cpp
--- a/Engine/Source/Runtime/Engine/Private/Landscape/LandscapeSplines.cpp
+++ b/Engine/Source/Runtime/Engine/Private/Landscape/LandscapeSplines.cpp
@@ -5,6 +5,7 @@
 =============================================================================*/
 
 #include "EnginePrivate.h"
+#include "ShaderParameterUtils.h"
 #include "StaticMeshResources.h"
 #include "AI/Navigation/RecastHelpers.h"
 #include "Landscape/LandscapeSplineProxies.h"
diff --git a/Engine/Source/Runtime/Engine/Private/LocalPlayer.cpp b/Engine/Source/Runtime/Engine/Private/LocalPlayer.cpp
--- a/Engine/Source/Runtime/Engine/Private/LocalPlayer.cpp
+++ b/Engine/Source/Runtime/Engine/Private/LocalPlayer.cpp
@@ -13,6 +13,7 @@
 #include "RenderCore.h"
 #include "ColorList.h"
 #include "Slate.h"
+#include "UObjectAnnotation.h"
 
 #include "IHeadMountedDisplay.h"
 #include "SceneViewExtension.h"
@@ -280,9 +281,301 @@ void ULocalPlayer::FinishDestroy()
 	Super::FinishDestroy();
 }
 
+/**
+ * Singleton managing saved locked views and the current per-player state.
+ */
+class FLockedViewState
+{
+public:
+	/** Singleton accessor. */
+	static FLockedViewState& Get()
+	{
+		static FLockedViewState State;
+		return State;
+	}
+
+	/**
+	 * Retrieves the locked view point for the given player.
+	 * @param Player			The player for which to retrieve the locked view point.
+	 * @param OutViewLocation	The location at which the view was locked.
+	 * @param OutViewRotation	The rotation at which the view was locked.
+	 * @param OutFOV			The FOV at which the view was locked.
+	 * @returns true if the view is locked, false if it is not.
+	 */
+	bool GetViewPoint(ULocalPlayer* Player, FVector& OutViewLocation, FRotator& OutViewRotation, float& OutFOV)
+	{
+		FPlayerState PlayerState = PlayerStates.GetAnnotation(Player);
+		if (PlayerState.bLocked)
+		{
+			OutViewLocation = PlayerState.ViewPoint.Location;
+			OutViewRotation = PlayerState.ViewPoint.Rotation;
+			OutFOV = PlayerState.ViewPoint.FOV;
+			return true;
+		}
+		return false;
+	}
+
+	/**
+	 * Returns true if the player's viewpoint is locked.
+	 */
+	bool IsViewLocked(ULocalPlayer* Player)
+	{
+		FPlayerState PlayerState = PlayerStates.GetAnnotation(Player);
+		return PlayerState.bLocked;
+	}
+
+	/**
+	 * Forces the player's view to be unlocked.
+	 */
+	void UnlockView(ULocalPlayer* Player)
+	{
+		PlayerStates.RemoveAnnotation(Player);
+	}
+
+	/**
+	 * Processes a LockView console command.
+	 * @param Player	The player for which the command was given.
+	 * @param Args		Arguments to the LockView command.
+	 */
+	void LockView(ULocalPlayer* Player, const TArray<FString>& Args)
+	{
+		bool bPrintHelp = false;
+		bool bShouldLockView = false;
+		FPlayerState PlayerState = PlayerStates.GetAnnotation(Player);
+
+		// ? as only arg == display help.
+		if (Args.Num() == 1 && Args[0] == TEXT("?"))
+		{
+			bPrintHelp = true;
+		}
+		// No args == toggle view locking.
+		else if (Args.Num() == 0)
+		{
+			if (PlayerState.bLocked)
+			{
+				PlayerStates.RemoveAnnotation(Player);
+			}
+			else
+			{
+				FMinimalViewInfo MinViewInfo;
+				Player->GetViewPoint(MinViewInfo);
+				PlayerState.ViewPoint.Location = MinViewInfo.Location;
+				PlayerState.ViewPoint.Rotation = MinViewInfo.Rotation;
+				PlayerState.ViewPoint.FOV = MinViewInfo.FOV;
+				bShouldLockView = true;
+			}
+		}
+		// One arg == lock view at named location.
+		else if (Args.Num() == 1)
+		{
+			FName ViewName(*Args[0]);
+			if (Viewpoints.Contains(ViewName))
+			{
+				PlayerState.ViewPoint = Viewpoints.FindRef(ViewName);
+			}
+			else
+			{
+				FMinimalViewInfo MinViewInfo;
+				Player->GetViewPoint(MinViewInfo);
+				PlayerState.ViewPoint.Location = MinViewInfo.Location;
+				PlayerState.ViewPoint.Rotation = MinViewInfo.Rotation;
+				PlayerState.ViewPoint.FOV = MinViewInfo.FOV;
+				Viewpoints.Add(ViewName,PlayerState.ViewPoint);
+			}
+			bShouldLockView = true;
+		}
+		// Six args == specify explicit location
+		else if (Args.Num() == 6)
+		{
+			bool bAnyEmpty = false;
+			for (int32 i = 0; i < Args.Num(); ++i)
+			{
+				bAnyEmpty |= Args[0].Len() == 0;
+			}
+			if (bAnyEmpty)
+			{
+				bPrintHelp = true;
+			}
+			else
+			{
+				PlayerState.ViewPoint = GetViewPointFromStrings(&Args[0],Args.Num());
+				bShouldLockView = true;
+			}
+		}
+		// Seven args == specify an explicit location and store it.
+		else if (Args.Num() == 7)
+		{
+			bool bAnyEmpty = false;
+			for (int32 i = 0; i < Args.Num(); ++i)
+			{
+				bAnyEmpty |= Args[0].Len() == 0;
+			}
+			if (bAnyEmpty)
+			{
+				bPrintHelp = true;
+			}
+			else
+			{
+				FName ViewName(*Args[0]);
+				PlayerState.ViewPoint = GetViewPointFromStrings(&Args[1],Args.Num() - 1);
+				Viewpoints.Add(ViewName,PlayerState.ViewPoint);
+				bShouldLockView = true;
+			}
+		}
+		// Anything else: unrecognized. Print help.
+		else
+		{
+			bPrintHelp = true;
+		}
+
+		if (bShouldLockView)
+		{
+			PlayerState.bLocked = true;
+			PlayerStates.AddAnnotation(Player, PlayerState);
+
+			// Also copy to the clipboard.
+			FString ViewPointString = ViewPointToString(PlayerState.ViewPoint);
+			FPlatformMisc::ClipboardCopy(*ViewPointString);
+		}
+
+		if (bPrintHelp)
+		{
+			UE_LOG(LogConsoleResponse,Display,
+				TEXT("Locks the player view and rendering time.\n")
+				TEXT("r.LockView ?\n")
+				TEXT("   Displays this message.\n")
+				TEXT("r.LockView\n")
+				TEXT("   Toggles whether the view is currently locked.\n")
+				TEXT("r.LockView <name>\n")
+				TEXT("   Locks the view at the named location. If there is no stored view with that name the current view is stored with that name.\n")
+				TEXT("r.LockView x y z pitch yaw roll\n")
+				TEXT("   Locks the view at the specified location and rotation.\n")
+				TEXT("r.LockView <name> x y z pitch yaw roll\n")
+				TEXT("   Locks the view at the specified location and rotation and stores it with the specified name.\n")
+				);
+		}
+	}
+
+private:
+	/**
+	 * Information stored for a given viewpoint.
+	 */
+	struct FViewPoint
+	{
+		FVector Location;
+		float FOV;
+		FRotator Rotation;
+	};
+
+	/** Viewpoints stored by name. */
+	TMap<FName,FViewPoint> Viewpoints;
+
+	/**
+	 * Per-player state attached to ULocalPlayer objects via a sparse UObject
+	 * annotation.
+	 */
+	struct FPlayerState
+	{
+		FViewPoint ViewPoint;
+		bool bLocked;
+
+		FPlayerState()
+		{
+			ViewPoint.Location = FVector::ZeroVector;
+			ViewPoint.FOV = 90.0f;
+			ViewPoint.Rotation = FRotator::ZeroRotator;
+			bLocked = false;
+		}
+
+		bool IsDefault() const
+		{
+			return bLocked == false
+				&& ViewPoint.Location == FVector::ZeroVector
+				&& ViewPoint.FOV == 90.0f
+				&& ViewPoint.Rotation == FRotator::ZeroRotator;
+		}
+	};
+	FUObjectAnnotationSparse<FPlayerState,/*bAutoRemove=*/true> PlayerStates;
+
+	/** Default constructor. */
+	FLockedViewState()
+	{
+	}
+
+	/**
+	 * Parses a viewpoint from an array of strings.
+	 *   WARNING: It is expected that the array has six entries!
+	 */
+	static FViewPoint GetViewPointFromStrings(const FString* Strings, int32 NumStrings)
+	{
+		FViewPoint ViewPoint;
+		if (NumStrings == 6)
+		{
+			ViewPoint.Location.X = FCString::Atof(*Strings[0]);
+			ViewPoint.Location.Y = FCString::Atof(*Strings[1]);
+			ViewPoint.Location.Z = FCString::Atof(*Strings[2]);
+			ViewPoint.Rotation.Pitch = FCString::Atof(*Strings[3]);
+			ViewPoint.Rotation.Yaw = FCString::Atof(*Strings[4]);
+			ViewPoint.Rotation.Roll = FCString::Atof(*Strings[5]);
+			ViewPoint.FOV = 90.0f;
+		}
+		return ViewPoint;
+	}
+
+	/**
+	 * Constructs a string from the view point.
+	 */
+	static FString ViewPointToString(const FViewPoint& ViewPoint)
+	{
+		return FString::Printf(TEXT("%f %f %f %f %f %f"),
+			ViewPoint.Location.X,
+			ViewPoint.Location.Y,
+			ViewPoint.Location.Z,
+			ViewPoint.Rotation.Pitch,
+			ViewPoint.Rotation.Yaw,
+			ViewPoint.Rotation.Roll
+			);
+	}
+
+	/**
+	 * Constructs a string representing all locked views and copies it to the
+	 * clipboard. Passing this string to r.LockViews will restore the state of
+	 * those locked views.
+	 */
+	static void CopyLockedViews()
+	{
+		FString LockedViewsStr;
+		FLockedViewState& This = FLockedViewState::Get();
+		bool bFirst = true;
+
+		for (TMap<FName,FViewPoint>::TConstIterator It(This.Viewpoints); It; ++It)
+		{
+			LockedViewsStr += FString::Printf(
+				TEXT("%s%s %s"),
+				bFirst ? TEXT("") : TEXT(";\n"),
+				*It.Key().ToString(),
+				*ViewPointToString(It.Value())
+				);
+			bFirst = false;
+		}
+		FPlatformMisc::ClipboardCopy(*LockedViewsStr);
+		UE_LOG(LogConsoleResponse,Display,TEXT("%s"),*LockedViewsStr);
+	}
+
+	static FAutoConsoleCommand CmdCopyLockedViews;
+};
+
+/** Console command to copy all named locked views to the clipboard. */
+FAutoConsoleCommand FLockedViewState::CmdCopyLockedViews(
+	TEXT("r.CopyLockedViews"),
+	TEXT("Copies all locked views in to a string that r.LockView will accept to reload them."),
+	FConsoleCommandDelegate::CreateStatic(FLockedViewState::CopyLockedViews)
+	);
+
 void ULocalPlayer::GetViewPoint(FMinimalViewInfo& OutViewInfo)
 {
-	if (PlayerController != NULL)
+	if (FLockedViewState::Get().GetViewPoint(this, OutViewInfo.Location, OutViewInfo.Rotation, OutViewInfo.FOV) == false
+		&& PlayerController != NULL)
 	{
 		if (PlayerController->PlayerCameraManager != NULL)
 		{
@@ -1090,6 +1383,32 @@ bool ULocalPlayer::Exec(UWorld* InWorld, const TCHAR* Cmd,FOutputDevice& Ar)
 	{
 		return HandleListMoveBodyCommand( Cmd, Ar );
 	}
+	else if( FParse::Command(&Cmd,TEXT("r.LockView")) )
+	{
+		TArray<FString> Lines;
+		FString CmdString(Cmd);
+		int32 NewLineIndex;
+		if (CmdString.FindChar(TEXT(';'),NewLineIndex))
+		{
+			CmdString.ParseIntoArray(&Lines,TEXT(";"),true);
+		}
+		else
+		{
+			Lines.Add(CmdString);
+		}
+
+		for (int32 LineIndex = 0; LineIndex < Lines.Num(); ++LineIndex)
+		{
+			TArray<FString> Args;
+			Lines[LineIndex].ParseIntoArrayWS(&Args);
+			FLockedViewState::Get().LockView(this,Args);
+		}
+		if (Lines.Num() > 1)
+		{
+			FLockedViewState::Get().UnlockView(this);
+		}
+		return true;
+	}
 #if WITH_PHYSX
 	// This will list all awake rigid bodies
 	else if( FParse::Command(&Cmd,TEXT("LISTAWAKEBODIES")) )
diff --git a/Engine/Source/Runtime/Engine/Private/LocalVertexFactory.cpp b/Engine/Source/Runtime/Engine/Private/LocalVertexFactory.cpp
--- a/Engine/Source/Runtime/Engine/Private/LocalVertexFactory.cpp
+++ b/Engine/Source/Runtime/Engine/Private/LocalVertexFactory.cpp
@@ -6,6 +6,7 @@
 
 #include "EnginePrivate.h"
 #include "SpeedTreeWind.h"
+#include "ShaderParameterUtils.h"
 
 class FSpeedTreeVertexFactoryShaderParameters : public FVertexFactoryShaderParameters
 {
@@ -21,25 +22,24 @@ public:
 		Ar << LODParameter;
 	}
 
-	virtual void SetMesh(FShader* Shader,const FVertexFactory* VertexFactory,const FSceneView& View,const FMeshBatchElement& BatchElement,uint32 DataFlags) const OVERRIDE
+	virtual void SetMesh(FRHICommandList* RHICmdList, FShader* Shader,const FVertexFactory* VertexFactory,const FSceneView& View,const FMeshBatchElement& BatchElement,uint32 DataFlags) const OVERRIDE
 	{
 		if (View.Family != NULL && View.Family->Scene != NULL)
 		{
 			FUniformBufferRHIParamRef SpeedTreeUniformBuffer = View.Family->Scene->GetSpeedTreeUniformBuffer(VertexFactory);
 			if (SpeedTreeUniformBuffer != NULL)
 			{
-				SetUniformBufferParameter(Shader->GetVertexShader(), Shader->GetUniformBufferParameter<FSpeedTreeUniformParameters>(), SpeedTreeUniformBuffer);
+				SetUniformBufferParameter(RHICmdList, Shader->GetVertexShader(), Shader->GetUniformBufferParameter<FSpeedTreeUniformParameters>(), SpeedTreeUniformBuffer);
 
 				if (LODParameter.IsBound())
 				{
 					//@todo - speedtree smooth LOD
 					FVector LODData(0, 0, 0);
-					SetShaderValue(Shader->GetVertexShader(), LODParameter, LODData);
+					SetShaderValue(RHICmdList, Shader->GetVertexShader(), LODParameter, LODData);
 				}
 			}
 		}
 	}
-
 	FShaderParameter LODParameter;
 };
 
diff --git a/Engine/Source/Runtime/Engine/Private/Materials/HLSLMaterialTranslator.h b/Engine/Source/Runtime/Engine/Private/Materials/HLSLMaterialTranslator.h
--- a/Engine/Source/Runtime/Engine/Private/Materials/HLSLMaterialTranslator.h
+++ b/Engine/Source/Runtime/Engine/Private/Materials/HLSLMaterialTranslator.h
@@ -454,7 +454,6 @@ public:
 			}
 
 			ResourcesString = TEXT("");
-			MaterialCompilationOutput.UniformExpressionSet.GetResourcesString(Platform,ResourcesString);
 
 			// Output the implementation for any custom expressions we will call below.
 			for(int32 ExpressionIndex = 0;ExpressionIndex < CustomExpressionImplementations.Num();ExpressionIndex++)
@@ -1040,7 +1039,7 @@ protected:
 				break;
 			default: UE_LOG(LogMaterial, Fatal,TEXT("Unrecognized texture material value type: %u"),(int32)CodeChunk.Type);
 			};
-			FCString::Sprintf(FormattedCode, TEXT("Material%s_%u"), BaseName, TextureInputIndex);
+			FCString::Sprintf(FormattedCode, TEXT("Material.%s_%u"), BaseName, TextureInputIndex);
 		}
 		else
 		{
diff --git a/Engine/Source/Runtime/Engine/Private/Materials/Material.cpp b/Engine/Source/Runtime/Engine/Private/Materials/Material.cpp
--- a/Engine/Source/Runtime/Engine/Private/Materials/Material.cpp
+++ b/Engine/Source/Runtime/Engine/Private/Materials/Material.cpp
@@ -594,7 +594,7 @@ void UMaterial::OverrideTexture( const UTexture* InTextureToOverride, UTexture*
 #if WITH_EDITOR
 	bool bShouldRecacheMaterialExpressions = false;
 	const bool bES2Preview = false;
-	ERHIFeatureLevel::Type FeatureLevelsToUpdate[2] = { GRHIFeatureLevel, ERHIFeatureLevel::ES2 };
+	ERHIFeatureLevel::Type FeatureLevelsToUpdate[2] = {GRHIFeatureLevel,ERHIFeatureLevel::ES2};
 	int32 NumFeatureLevelsToUpdate = bES2Preview ? 2 : 1;
 	
 	for (int32 i = 0; i < NumFeatureLevelsToUpdate; ++i)
@@ -1995,6 +1995,18 @@ void UMaterial::PostLoad()
 		//Don't compile shaders in post load for dev overhead materials.
 		if (FApp::CanEverRender() && !bIsMaterialEditorStatsMaterial)
 		{
+			// Before caching shader resources we have to make sure all referenced textures have been post loaded
+			// as we depend on their resources being valid.
+			RebuildExpressionTextureReferences();
+			for (int32 TextureIndex=0, NumTextures=ExpressionTextureReferences.Num(); TextureIndex < NumTextures; ++TextureIndex)
+			{
+				UTexture* Texture = ExpressionTextureReferences[TextureIndex];
+				if (Texture)
+				{
+					Texture->ConditionalPostLoad();
+				}
+			}
+
 			CacheResourceShadersForRendering(false);
 		}
 	}
diff --git a/Engine/Source/Runtime/Engine/Private/Materials/MaterialUniformExpressions.cpp b/Engine/Source/Runtime/Engine/Private/Materials/MaterialUniformExpressions.cpp
--- a/Engine/Source/Runtime/Engine/Private/Materials/MaterialUniformExpressions.cpp
+++ b/Engine/Source/Runtime/Engine/Private/Materials/MaterialUniformExpressions.cpp
@@ -171,20 +171,6 @@ FString FUniformExpressionSet::GetSummaryString() const
 		);
 }
 
-void FUniformExpressionSet::GetResourcesString(EShaderPlatform ShaderPlatform,FString& InputsString) const
-{
-	for(int32 TextureIndex = 0;TextureIndex < Uniform2DTextureExpressions.Num();TextureIndex++)
-	{
-		InputsString += FString::Printf(TEXT("Texture2D MaterialTexture2D_%i;\r\n"), TextureIndex);
-		InputsString += FString::Printf(TEXT("SamplerState MaterialTexture2D_%iSampler;\r\n"), TextureIndex);
-	}
-	for(int32 TextureIndex = 0;TextureIndex < UniformCubeTextureExpressions.Num();TextureIndex++)
-	{
-		InputsString += FString::Printf(TEXT("TextureCube MaterialTextureCube_%i;\r\n"),TextureIndex);
-		InputsString += FString::Printf(TEXT("SamplerState MaterialTextureCube_%iSampler;\r\n"),TextureIndex);
-	}
-}
-
 void FUniformExpressionSet::SetParameterCollections(const TArray<UMaterialParameterCollection*>& InCollections)
 {
 	ParameterCollections.Empty(InCollections.Num());
@@ -204,18 +190,56 @@ void FUniformExpressionSet::CreateBufferStruct()
 
 	if(UniformVectorExpressions.Num())
 	{
-		new(Members) FUniformBufferStruct::FMember(TEXT("VectorExpressions"),NextMemberOffset,UBMT_FLOAT32,EShaderPrecisionModifier::Half,1,4,UniformVectorExpressions.Num(),NULL);
+		new(Members) FUniformBufferStruct::FMember(TEXT("VectorExpressions"),TEXT(""),NextMemberOffset,UBMT_FLOAT32,EShaderPrecisionModifier::Half,1,4,UniformVectorExpressions.Num(),NULL);
 		const uint32 VectorArraySize = UniformVectorExpressions.Num() * sizeof(FVector4);
 		NextMemberOffset += VectorArraySize;
 	}
 
 	if(UniformScalarExpressions.Num())
 	{
-		new(Members) FUniformBufferStruct::FMember(TEXT("ScalarExpressions"),NextMemberOffset,UBMT_FLOAT32,EShaderPrecisionModifier::Half,1,4,(UniformScalarExpressions.Num() + 3) / 4,NULL);
+		new(Members) FUniformBufferStruct::FMember(TEXT("ScalarExpressions"),TEXT(""),NextMemberOffset,UBMT_FLOAT32,EShaderPrecisionModifier::Half,1,4,(UniformScalarExpressions.Num() + 3) / 4,NULL);
 		const uint32 ScalarArraySize = (UniformScalarExpressions.Num() + 3) / 4 * sizeof(FVector4);
 		NextMemberOffset += ScalarArraySize;
 	}
 
+	static FString Texture2DNames[128];
+	static FString Texture2DSamplerNames[128];
+	static FString TextureCubeNames[128];
+	static FString TextureCubeSamplerNames[128];
+	static bool bInitializedTextureNames = false;
+	if (!bInitializedTextureNames)
+	{
+		bInitializedTextureNames = true;
+		for (int32 i = 0; i < 128; ++i)
+		{
+			Texture2DNames[i] = FString::Printf(TEXT("Texture2D_%d"),i);
+			Texture2DSamplerNames[i] = FString::Printf(TEXT("Texture2D_%dSampler"),i);
+			TextureCubeNames[i] = FString::Printf(TEXT("TextureCube_%d"),i);
+			TextureCubeSamplerNames[i] = FString::Printf(TEXT("TextureCube_%dSampler"),i);
+		}
+	}
+
+	check(Uniform2DTextureExpressions.Num() <= 128);
+	check(UniformCubeTextureExpressions.Num() <= 128);
+
+	for (int32 i = 0; i < Uniform2DTextureExpressions.Num(); ++i)
+	{
+		check((NextMemberOffset & 0x7) == 0);
+		new(Members) FUniformBufferStruct::FMember(*Texture2DNames[i],TEXT("Texture2D"),NextMemberOffset,UBMT_TEXTURE,EShaderPrecisionModifier::Float,1,1,1,NULL);
+		NextMemberOffset += 8;
+		new(Members) FUniformBufferStruct::FMember(*Texture2DSamplerNames[i],TEXT("SamplerState"),NextMemberOffset,UBMT_SAMPLER,EShaderPrecisionModifier::Float,1,1,1,NULL);
+		NextMemberOffset += 8;
+	}
+
+	for (int32 i = 0; i < UniformCubeTextureExpressions.Num(); ++i)
+	{
+		check((NextMemberOffset & 0x7) == 0);
+		new(Members) FUniformBufferStruct::FMember(*TextureCubeNames[i],TEXT("TextureCube"),NextMemberOffset,UBMT_TEXTURE,EShaderPrecisionModifier::Float,1,1,1,NULL);
+		NextMemberOffset += 8;
+		new(Members) FUniformBufferStruct::FMember(*TextureCubeSamplerNames[i],TEXT("SamplerState"),NextMemberOffset,UBMT_SAMPLER,EShaderPrecisionModifier::Float,1,1,1,NULL);
+		NextMemberOffset += 8;
+	}
+
 	const uint32 StructSize = Align(NextMemberOffset,UNIFORM_BUFFER_STRUCT_ALIGNMENT);
 	UniformBufferStruct = new FUniformBufferStruct(
 		TEXT("MaterialUniforms"),
@@ -260,7 +284,44 @@ FUniformBufferRHIRef FUniformExpressionSet::CreateUniformBuffer(const FMaterialR
 			TempScalarBuffer[ScalarIndex] = VectorValue.R;
 		}
 
-		UniformBuffer = RHICreateUniformBuffer(TempBuffer,UniformBufferStruct->GetSize(),UniformBuffer_MultiUse);
+		void** ResourceTable = (void**)((uint8*)TempBuffer + UniformBufferStruct->GetLayout().ResourceOffset);
+		check(((UPTRINT)ResourceTable & 0x7) == 0);
+
+		// Cache 2D texture uniform expressions.
+		for(int32 ExpressionIndex = 0;ExpressionIndex < Uniform2DTextureExpressions.Num();ExpressionIndex++)
+		{
+			const UTexture* Value;
+			Uniform2DTextureExpressions[ExpressionIndex]->GetTextureValue(MaterialRenderContext,MaterialRenderContext.Material,Value);
+			if(Value && Value->Resource)
+			{
+				*ResourceTable++ = Value->TextureRefRHI;
+				*ResourceTable++ = Value->Resource->SamplerStateRHI;
+			}
+			else
+			{
+				*ResourceTable++ = GWhiteTexture->TextureRHI;
+				*ResourceTable++ = GWhiteTexture->SamplerStateRHI;
+			}
+		}
+
+		// Cache cube texture uniform expressions.
+		for(int32 ExpressionIndex = 0;ExpressionIndex < UniformCubeTextureExpressions.Num();ExpressionIndex++)
+		{
+			const UTexture* Value;
+			UniformCubeTextureExpressions[ExpressionIndex]->GetTextureValue(MaterialRenderContext,MaterialRenderContext.Material,Value);
+			if(Value && Value->Resource)
+			{
+				*ResourceTable++ = Value->TextureRefRHI;
+				*ResourceTable++ = Value->Resource->SamplerStateRHI;
+			}
+			else
+			{
+				*ResourceTable++ = GWhiteTexture->TextureRHI;
+				*ResourceTable++ = GWhiteTexture->SamplerStateRHI;
+			}
+		}
+
+		UniformBuffer = RHICreateUniformBuffer(TempBuffer,UniformBufferStruct->GetLayout(),UniformBuffer_MultiFrame);
 	}
 
 	return UniformBuffer;
diff --git a/Engine/Source/Runtime/Engine/Private/Materials/ParameterCollection.cpp b/Engine/Source/Runtime/Engine/Private/Materials/ParameterCollection.cpp
--- a/Engine/Source/Runtime/Engine/Private/Materials/ParameterCollection.cpp
+++ b/Engine/Source/Runtime/Engine/Private/Materials/ParameterCollection.cpp
@@ -325,7 +325,7 @@ void UMaterialParameterCollection::CreateBufferStruct()
 	uint32 NextMemberOffset = 0;
 
 	const uint32 NumVectors = FMath::DivideAndRoundUp(ScalarParameters.Num(), 4) + VectorParameters.Num();
-	new(Members) FUniformBufferStruct::FMember(TEXT("Vectors"),NextMemberOffset,UBMT_FLOAT32,EShaderPrecisionModifier::Half,1,4,NumVectors,NULL);
+	new(Members) FUniformBufferStruct::FMember(TEXT("Vectors"),TEXT(""),NextMemberOffset,UBMT_FLOAT32,EShaderPrecisionModifier::Half,1,4,NumVectors,NULL);
 	const uint32 VectorArraySize = NumVectors * sizeof(FVector4);
 	NextMemberOffset += VectorArraySize;
 
@@ -550,6 +550,10 @@ void FMaterialParameterCollectionInstanceResource::UpdateContents(const FGuid& I
 
 	if (InId != FGuid() && Data.Num() > 0)
 	{
-		UniformBuffer = RHICreateUniformBuffer(Data.GetData(), Data.GetTypeSize() * Data.Num(), UniformBuffer_MultiUse);
+		UniformBuffer.SafeRelease();
+		UniformBufferLayout.ConstantBufferSize = Data.GetTypeSize() * Data.Num();
+		UniformBufferLayout.ResourceOffset = 0;
+		check(UniformBufferLayout.Resources.Num() == 0);
+		UniformBuffer = RHICreateUniformBuffer(Data.GetData(), UniformBufferLayout, UniformBuffer_MultiFrame);
 	}
 }
\ No newline at end of file
diff --git a/Engine/Source/Runtime/Engine/Private/MeshParticleVertexFactory.cpp b/Engine/Source/Runtime/Engine/Private/MeshParticleVertexFactory.cpp
--- a/Engine/Source/Runtime/Engine/Private/MeshParticleVertexFactory.cpp
+++ b/Engine/Source/Runtime/Engine/Private/MeshParticleVertexFactory.cpp
@@ -35,12 +35,12 @@ public:
 		Ar << ParticleColor;
 	}
 
-	virtual void SetMesh(FShader* Shader,const FVertexFactory* VertexFactory,const FSceneView& View,const FMeshBatchElement& BatchElement,uint32 DataFlags) const OVERRIDE
+	virtual void SetMesh(FRHICommandList* RHICmdList, FShader* Shader,const FVertexFactory* VertexFactory,const FSceneView& View,const FMeshBatchElement& BatchElement,uint32 DataFlags) const OVERRIDE
 	{
 		const bool bInstanced = View.GetFeatureLevel() >= ERHIFeatureLevel::SM3;
 		FMeshParticleVertexFactory* MeshParticleVF = (FMeshParticleVertexFactory*)VertexFactory;
 		FVertexShaderRHIParamRef VertexShaderRHI = Shader->GetVertexShader();
-		SetUniformBufferParameter( VertexShaderRHI, Shader->GetUniformBufferParameter<FMeshParticleUniformParameters>(), MeshParticleVF->GetUniformBuffer() );
+		SetUniformBufferParameter(RHICmdList, VertexShaderRHI, Shader->GetUniformBufferParameter<FMeshParticleUniformParameters>(), MeshParticleVF->GetUniformBuffer() );
 
 		if (!bInstanced)
 		{
@@ -48,13 +48,13 @@ public:
 			const FMeshParticleInstanceVertex* Vertex = BatchParameters->InstanceBuffer + BatchElement.UserIndex;
 			const FMeshParticleInstanceVertexDynamicParameter* DynamicVertex = BatchParameters->DynamicParameterBuffer + BatchElement.UserIndex;
 
-			SetShaderValue(VertexShaderRHI, Transform1, Vertex->Transform[0]);
-			SetShaderValue(VertexShaderRHI, Transform2, Vertex->Transform[1]);
-			SetShaderValue(VertexShaderRHI, Transform3, Vertex->Transform[2]);
-			SetShaderValue(VertexShaderRHI, SubUVParams, FVector4((float)Vertex->SubUVParams[0], (float)Vertex->SubUVParams[1], (float)Vertex->SubUVParams[2], (float)Vertex->SubUVParams[3]));
-			SetShaderValue(VertexShaderRHI, SubUVLerp, Vertex->SubUVLerp);
-			SetShaderValue(VertexShaderRHI, DynamicParameter, FVector4(DynamicVertex->DynamicValue[0], DynamicVertex->DynamicValue[1], DynamicVertex->DynamicValue[2], DynamicVertex->DynamicValue[3]));
-			SetShaderValue(VertexShaderRHI, ParticleColor, FVector4(Vertex->Color.Component(0), Vertex->Color.Component(1), Vertex->Color.Component(2), Vertex->Color.Component(3)));
+			SetShaderValue(RHICmdList, VertexShaderRHI, Transform1, Vertex->Transform[0]);
+			SetShaderValue(RHICmdList, VertexShaderRHI, Transform2, Vertex->Transform[1]);
+			SetShaderValue(RHICmdList, VertexShaderRHI, Transform3, Vertex->Transform[2]);
+			SetShaderValue(RHICmdList, VertexShaderRHI, SubUVParams, FVector4((float)Vertex->SubUVParams[0], (float)Vertex->SubUVParams[1], (float)Vertex->SubUVParams[2], (float)Vertex->SubUVParams[3]));
+			SetShaderValue(RHICmdList, VertexShaderRHI, SubUVLerp, Vertex->SubUVLerp);
+			SetShaderValue(RHICmdList, VertexShaderRHI, DynamicParameter, FVector4(DynamicVertex->DynamicValue[0], DynamicVertex->DynamicValue[1], DynamicVertex->DynamicValue[2], DynamicVertex->DynamicValue[3]));
+			SetShaderValue(RHICmdList, VertexShaderRHI, ParticleColor, FVector4(Vertex->Color.Component(0), Vertex->Color.Component(1), Vertex->Color.Component(2), Vertex->Color.Component(3)));
 		}
 	}
 
diff --git a/Engine/Source/Runtime/Engine/Private/NiagaraComponent.cpp b/Engine/Source/Runtime/Engine/Private/NiagaraComponent.cpp
--- a/Engine/Source/Runtime/Engine/Private/NiagaraComponent.cpp
+++ b/Engine/Source/Runtime/Engine/Private/NiagaraComponent.cpp
@@ -126,7 +126,7 @@ private:
 				{
 					FParticleSpriteUniformParameters PerViewUniformParameters = UniformParameters;
 					PerViewUniformParameters.MacroUVParameters = FVector4(0.0f,0.0f,1.0f,1.0f);
-					*SpriteViewUniformBufferPtr = FParticleSpriteUniformBufferRef::CreateUniformBufferImmediate(PerViewUniformParameters, UniformBuffer_SingleUse);
+					*SpriteViewUniformBufferPtr = FParticleSpriteUniformBufferRef::CreateUniformBufferImmediate(PerViewUniformParameters, UniformBuffer_SingleFrame);
 				}
 			}
 
diff --git a/Engine/Source/Runtime/Engine/Private/Particles/ParticleBeamTrailVertexFactory.cpp b/Engine/Source/Runtime/Engine/Private/Particles/ParticleBeamTrailVertexFactory.cpp
--- a/Engine/Source/Runtime/Engine/Private/Particles/ParticleBeamTrailVertexFactory.cpp
+++ b/Engine/Source/Runtime/Engine/Private/Particles/ParticleBeamTrailVertexFactory.cpp
@@ -3,9 +3,10 @@
 /*=============================================================================
 	ParticleBeamTrailVertexFactory.cpp: Particle vertex factory implementation.
 =============================================================================*/
+
 #include "EnginePrivate.h"
 #include "ParticleDefinitions.h"
-#include "ShaderParameters.h"
+#include "ShaderParameterUtils.h"
 
 IMPLEMENT_UNIFORM_BUFFER_STRUCT(FParticleBeamTrailUniformParameters,TEXT("BeamTrailVF"));
 
@@ -23,10 +24,10 @@ public:
 	{
 	}
 
-	virtual void SetMesh(FShader* Shader,const FVertexFactory* VertexFactory,const FSceneView& View,const FMeshBatchElement& BatchElement,uint32 DataFlags) const OVERRIDE
+	virtual void SetMesh(FRHICommandList* RHICmdList, FShader* Shader,const FVertexFactory* VertexFactory,const FSceneView& View,const FMeshBatchElement& BatchElement,uint32 DataFlags) const OVERRIDE
 	{
 		FParticleBeamTrailVertexFactory* BeamTrailVF = (FParticleBeamTrailVertexFactory*)VertexFactory;
-		SetUniformBufferParameter( Shader->GetVertexShader(), Shader->GetUniformBufferParameter<FParticleBeamTrailUniformParameters>(), BeamTrailVF->GetBeamTrailUniformBuffer() );
+		SetUniformBufferParameter(RHICmdList, Shader->GetVertexShader(), Shader->GetUniformBufferParameter<FParticleBeamTrailUniformParameters>(), BeamTrailVF->GetBeamTrailUniformBuffer() );
 	}
 };
 
diff --git a/Engine/Source/Runtime/Engine/Private/Particles/ParticleCurveTexture.cpp b/Engine/Source/Runtime/Engine/Private/Particles/ParticleCurveTexture.cpp
--- a/Engine/Source/Runtime/Engine/Private/Particles/ParticleCurveTexture.cpp
+++ b/Engine/Source/Runtime/Engine/Private/Particles/ParticleCurveTexture.cpp
@@ -10,6 +10,8 @@
 #include "ParticleResources.h"
 #include "UniformBuffer.h"
 #include "ShaderParameters.h"
+#include "ShaderParameterUtils.h"
+#include "RHIStaticStates.h"
 #include "GlobalShader.h"
 #include "FXSystem.h"
 
@@ -68,15 +70,15 @@ public:
 	/**
 	 * Sets parameters for particle injection.
 	 */
-	void SetParameters( const FVector2D& CurveOffset )
+	void SetParameters(FRHICommandList* RHICmdList, const FVector2D& CurveOffset )
 	{
 		FParticleCurveInjectionParameters Parameters;
 		Parameters.PixelScale.X = 1.0f / GParticleCurveTextureSizeX;
 		Parameters.PixelScale.Y = 1.0f / GParticleCurveTextureSizeY;
 		Parameters.CurveOffset = CurveOffset;
-		FParticleCurveInjectionBufferRef UniformBuffer = FParticleCurveInjectionBufferRef::CreateUniformBufferImmediate( Parameters, UniformBuffer_SingleUse );
+		FParticleCurveInjectionBufferRef UniformBuffer = FParticleCurveInjectionBufferRef::CreateUniformBufferImmediate( Parameters, UniformBuffer_SingleDraw );
 		FVertexShaderRHIParamRef VertexShader = GetVertexShader();
-		SetUniformBufferParameter( VertexShader, GetUniformBufferParameter<FParticleCurveInjectionParameters>(), UniformBuffer );
+		SetUniformBufferParameter(RHICmdList, VertexShader, GetUniformBufferParameter<FParticleCurveInjectionParameters>(), UniformBuffer );
 	}
 };
 
@@ -172,6 +174,7 @@ TGlobalResource<FParticleCurveInjectionVertexDeclaration> GParticleCurveInjectio
  * @param InPendingCurves - Curves to be stored on the GPU.
  */
 static void InjectCurves(
+	FRHICommandList* RHICmdList, 
 	FTexture2DRHIParamRef CurveTextureRHI,
 	FTexture2DRHIParamRef CurveTextureTargetRHI,
 	TArray<FCurveSamples>& InPendingCurves )
@@ -184,6 +187,8 @@ static void InjectCurves(
 
 	FVertexBufferRHIParamRef ScratchVertexBufferRHI = GParticleScratchVertexBuffer.VertexBufferRHI;
 
+	check(!RHICmdList);
+
 	RHISetRenderTarget( CurveTextureTargetRHI, FTextureRHIParamRef() );
 	RHISetScissorRect( false, 0, 0, 0, 0 );
 	RHISetViewport( 0, 0, 0.0f, GParticleCurveTextureSizeX, GParticleCurveTextureSizeY, 1.0f );
@@ -231,7 +236,7 @@ static void InjectCurves(
 			0
 			);
 
-		VertexShader->SetParameters( CurveOffset );
+		VertexShader->SetParameters(RHICmdList, CurveOffset );
 
 		// Stream 0: New particles.
 		RHISetStreamSource(
@@ -588,7 +593,9 @@ void FParticleCurveTexture::SubmitPendingCurves()
 			FParticleCurveTexture*, ParticleCurveTexture, this,
 			TArray<FCurveSamples>, PendingCurves, PendingCurves,
 		{
+			//@todo-rco: RHIPacketList
 			InjectCurves(
+				nullptr,
 				ParticleCurveTexture->CurveTextureRHI,
 				ParticleCurveTexture->CurveTextureTargetRHI,
 				PendingCurves
diff --git a/Engine/Source/Runtime/Engine/Private/Particles/ParticleGpuSimulation.cpp b/Engine/Source/Runtime/Engine/Private/Particles/ParticleGpuSimulation.cpp
--- a/Engine/Source/Runtime/Engine/Private/Particles/ParticleGpuSimulation.cpp
+++ b/Engine/Source/Runtime/Engine/Private/Particles/ParticleGpuSimulation.cpp
@@ -13,6 +13,8 @@
 #include "ParticleResources.h"
 #include "UniformBuffer.h"
 #include "ShaderParameters.h"
+#include "ShaderParameterUtils.h"
+#include "RHIStaticStates.h"
 #include "ParticleDefinitions.h"
 #include "GlobalShader.h"
 #include "../VectorField.h"
@@ -495,7 +497,7 @@ public:
 		Ar << CurveTextureSampler;
 	}
 
-	virtual void SetMesh(FShader* Shader,const FVertexFactory* VertexFactory,const FSceneView& View,const FMeshBatchElement& BatchElement,uint32 DataFlags) const OVERRIDE;
+	virtual void SetMesh(FRHICommandList* RHICmdList, FShader* Shader,const FVertexFactory* VertexFactory,const FSceneView& View,const FMeshBatchElement& BatchElement,uint32 DataFlags) const OVERRIDE;
 
 	virtual uint32 GetSize() const { return sizeof(*this); }
 
@@ -642,24 +644,26 @@ public:
 /**
  * Set vertex factory shader parameters.
  */
-void FGPUSpriteVertexFactoryShaderParameters::SetMesh(FShader* Shader,const FVertexFactory* VertexFactory,const FSceneView& View,const FMeshBatchElement& BatchElement,uint32 DataFlags) const
+void FGPUSpriteVertexFactoryShaderParameters::SetMesh(FRHICommandList* RHICmdList, FShader* Shader,const FVertexFactory* VertexFactory,const FSceneView& View,const FMeshBatchElement& BatchElement,uint32 DataFlags) const
 {
 	FGPUSpriteVertexFactory* GPUVF = (FGPUSpriteVertexFactory*)VertexFactory;
 	FVertexShaderRHIParamRef VertexShader = Shader->GetVertexShader();
 	FSamplerStateRHIParamRef SamplerStatePoint = TStaticSamplerState<SF_Point>::GetRHI();
 	FSamplerStateRHIParamRef SamplerStateLinear = TStaticSamplerState<SF_Bilinear>::GetRHI();
-	SetUniformBufferParameter( VertexShader, Shader->GetUniformBufferParameter<FGPUSpriteEmitterUniformParameters>(), GPUVF->EmitterUniformBuffer );
-	SetUniformBufferParameter( VertexShader, Shader->GetUniformBufferParameter<FGPUSpriteEmitterDynamicUniformParameters>(), GPUVF->EmitterDynamicUniformBuffer );
+	SetUniformBufferParameter(RHICmdList, VertexShader, Shader->GetUniformBufferParameter<FGPUSpriteEmitterUniformParameters>(), GPUVF->EmitterUniformBuffer );
+	SetUniformBufferParameter(RHICmdList, VertexShader, Shader->GetUniformBufferParameter<FGPUSpriteEmitterDynamicUniformParameters>(), GPUVF->EmitterDynamicUniformBuffer );
 	if (ParticleIndices.IsBound())
 	{
+		check(!RHICmdList);
+
 		RHISetShaderResourceViewParameter(VertexShader, ParticleIndices.GetBaseIndex(), GPUVF->ParticleIndicesBuffer->VertexBufferSRV);
 	}
-	SetShaderValue(VertexShader, ParticleIndicesOffset, GPUVF->ParticleIndicesOffset);
-	SetTextureParameter( VertexShader, PositionTexture, PositionTextureSampler, SamplerStatePoint, GPUVF->PositionTextureRHI );
-	SetTextureParameter( VertexShader, PositionZWTexture, PositionZWTextureSampler, SamplerStatePoint, GPUVF->PositionZWTextureRHI );
-	SetTextureParameter( VertexShader, VelocityTexture, VelocityTextureSampler, SamplerStatePoint, GPUVF->VelocityTextureRHI );
-	SetTextureParameter( VertexShader, AttributesTexture, AttributesTextureSampler, SamplerStatePoint, GPUVF->AttributesTextureRHI );
-	SetTextureParameter( VertexShader, CurveTexture, CurveTextureSampler, SamplerStateLinear, GParticleCurveTexture.GetCurveTexture() );
+	SetShaderValue(RHICmdList, VertexShader, ParticleIndicesOffset, GPUVF->ParticleIndicesOffset);
+	SetTextureParameter(RHICmdList, VertexShader, PositionTexture, PositionTextureSampler, SamplerStatePoint, GPUVF->PositionTextureRHI );
+	SetTextureParameter(RHICmdList, VertexShader, PositionZWTexture, PositionZWTextureSampler, SamplerStatePoint, GPUVF->PositionZWTextureRHI );
+	SetTextureParameter(RHICmdList, VertexShader, VelocityTexture, VelocityTextureSampler, SamplerStatePoint, GPUVF->VelocityTextureRHI );
+	SetTextureParameter(RHICmdList, VertexShader, AttributesTexture, AttributesTextureSampler, SamplerStatePoint, GPUVF->AttributesTextureRHI );
+	SetTextureParameter(RHICmdList, VertexShader, CurveTexture, CurveTextureSampler, SamplerStateLinear, GParticleCurveTexture.GetCurveTexture() );
 }
 
 IMPLEMENT_VERTEX_FACTORY_TYPE(FGPUSpriteVertexFactory,"ParticleGPUSpriteVertexFactory",true,false,true,false,false);
@@ -801,8 +805,10 @@ public:
 	}
 
 	/** Set parameters. */
-	void SetParameters(FParticleShaderParamRef TileOffsetsRef)
+	void SetParameters(FRHICommandList* RHICmdList, FParticleShaderParamRef TileOffsetsRef)
 	{
+		check(!RHICmdList);
+
 		FVertexShaderRHIParamRef VertexShaderRHI = GetVertexShader();
 		if (TileOffsets.IsBound())
 		{
@@ -902,6 +908,7 @@ public:
 	 * Set parameters for this shader.
 	 */
 	void SetParameters(
+		FRHICommandList* RHICmdList, 
 		const FParticleStateTextures& TextureResources,
 		const FParticleAttributesTexture& InAttributesTexture,
 		const FParticleAttributesTexture& InRenderAttributesTexture,
@@ -913,16 +920,17 @@ public:
 		FPixelShaderRHIParamRef PixelShaderRHI = GetPixelShader();
 		FSamplerStateRHIParamRef SamplerStatePoint = TStaticSamplerState<SF_Point>::GetRHI();
 		FSamplerStateRHIParamRef SamplerStateLinear = TStaticSamplerState<SF_Bilinear,AM_Clamp,AM_Clamp,AM_Clamp>::GetRHI();
-		SetTextureParameter(PixelShaderRHI, PositionTexture, PositionTextureSampler, SamplerStatePoint, TextureResources.PositionTextureRHI);
-		SetTextureParameter(PixelShaderRHI, PositionZWTexture, PositionZWTextureSampler, SamplerStatePoint, TextureResources.PositionZWTextureRHI);
-		SetTextureParameter(PixelShaderRHI, VelocityTexture, VelocityTextureSampler, SamplerStatePoint, TextureResources.VelocityTextureRHI);
-		SetTextureParameter(PixelShaderRHI, AttributesTexture, AttributesTextureSampler, SamplerStatePoint, InAttributesTexture.TextureRHI);
-		SetTextureParameter(PixelShaderRHI, CurveTexture, CurveTextureSampler, SamplerStateLinear, GParticleCurveTexture.GetCurveTexture());
+		SetTextureParameter(RHICmdList, PixelShaderRHI, PositionTexture, PositionTextureSampler, SamplerStatePoint, TextureResources.PositionTextureRHI);
+		SetTextureParameter(RHICmdList, PixelShaderRHI, PositionZWTexture, PositionZWTextureSampler, SamplerStatePoint, TextureResources.PositionZWTextureRHI);
+		SetTextureParameter(RHICmdList, PixelShaderRHI, VelocityTexture, VelocityTextureSampler, SamplerStatePoint, TextureResources.VelocityTextureRHI);
+		SetTextureParameter(RHICmdList, PixelShaderRHI, AttributesTexture, AttributesTextureSampler, SamplerStatePoint, InAttributesTexture.TextureRHI);
+		SetTextureParameter(RHICmdList, PixelShaderRHI, CurveTexture, CurveTextureSampler, SamplerStateLinear, GParticleCurveTexture.GetCurveTexture());
 		if (bUseDepthBufferCollision)
 		{
 			check(CollisionView != NULL);
-			FGlobalShader::SetParameters(PixelShaderRHI,*CollisionView);
+			FGlobalShader::SetParameters(RHICmdList, PixelShaderRHI,*CollisionView);
 			SetTextureParameter(
+				RHICmdList, 
 				PixelShaderRHI,
 				SceneDepthTextureParameter,
 				SceneDepthTextureParameterSampler,
@@ -930,6 +938,7 @@ public:
 				SceneDepthTexture
 				);
 			SetTextureParameter(
+				RHICmdList, 
 				PixelShaderRHI,
 				GBufferATextureParameter,
 				GBufferATextureParameterSampler,
@@ -937,13 +946,14 @@ public:
 				GBufferATexture
 				);
 			SetTextureParameter(
+				RHICmdList, 
 				PixelShaderRHI,
 				RenderAttributesTexture,
 				RenderAttributesTextureSampler,
 				SamplerStatePoint,
 				InRenderAttributesTexture.TextureRHI
 				);
-			SetShaderValue(PixelShaderRHI, CollisionDepthBounds, FXConsoleVariables::GPUCollisionDepthBounds);
+			SetShaderValue(RHICmdList, PixelShaderRHI, CollisionDepthBounds, FXConsoleVariables::GPUCollisionDepthBounds);
 		}
 	}
 
@@ -951,26 +961,26 @@ public:
 	 * Set parameters for the vector fields sampled by this shader.
 	 * @param VectorFieldParameters -Parameters needed to sample local vector fields.
 	 */
-	void SetVectorFieldParameters(const FVectorFieldUniformBufferRef& UniformBuffer, const FTexture3DRHIParamRef VolumeTexturesRHI[])
+	void SetVectorFieldParameters(FRHICommandList* RHICmdList, const FVectorFieldUniformBufferRef& UniformBuffer, const FTexture3DRHIParamRef VolumeTexturesRHI[])
 	{
 		FPixelShaderRHIParamRef PixelShaderRHI = GetPixelShader();
 		FSamplerStateRHIParamRef SamplerStateLinear = TStaticSamplerState<SF_Bilinear,AM_Clamp,AM_Clamp,AM_Clamp>::GetRHI();
-		SetUniformBufferParameter(PixelShaderRHI, GetUniformBufferParameter<FVectorFieldUniformParameters>(), UniformBuffer);
+		SetUniformBufferParameter(RHICmdList, PixelShaderRHI, GetUniformBufferParameter<FVectorFieldUniformParameters>(), UniformBuffer);
 		for (int32 TextureIndex = 0; TextureIndex < MAX_VECTOR_FIELDS; ++TextureIndex)
 		{
- 			SetTextureParameter(PixelShaderRHI, VectorFieldTextures, VectorFieldTexturesSampler, SamplerStateLinear, VolumeTexturesRHI[TextureIndex], TextureIndex);
+			SetTextureParameter(RHICmdList, PixelShaderRHI, VectorFieldTextures, VectorFieldTexturesSampler, SamplerStateLinear, VolumeTexturesRHI[TextureIndex], TextureIndex);
 		}
 	}
 
 	/**
 	 * Set per-instance parameters for this shader.
 	 */
-	void SetInstanceParameters(float InDeltaSeconds, FUniformBufferRHIParamRef UniformBuffer, FUniformBufferRHIParamRef PerFrameUniformBuffer)
+	void SetInstanceParameters(FRHICommandList* RHICmdList, float InDeltaSeconds, FUniformBufferRHIParamRef UniformBuffer, FUniformBufferRHIParamRef PerFrameUniformBuffer)
 	{
 		FPixelShaderRHIParamRef PixelShaderRHI = GetPixelShader();
-		SetShaderValue(PixelShaderRHI, DeltaSeconds, InDeltaSeconds);
-		SetUniformBufferParameter(PixelShaderRHI, GetUniformBufferParameter<FParticleSimulationParameters>(), UniformBuffer);
-		SetUniformBufferParameter(PixelShaderRHI, GetUniformBufferParameter<FParticlePerFrameSimulationParameters>(), PerFrameUniformBuffer);
+		SetShaderValue(RHICmdList, PixelShaderRHI, DeltaSeconds, InDeltaSeconds);
+		SetUniformBufferParameter(RHICmdList, PixelShaderRHI, GetUniformBufferParameter<FParticleSimulationParameters>(), UniformBuffer);
+		SetUniformBufferParameter(RHICmdList, PixelShaderRHI, GetUniformBufferParameter<FParticlePerFrameSimulationParameters>(), PerFrameUniformBuffer);
 	}
 
 	/**
@@ -1202,6 +1212,7 @@ struct FSimulationCommandGPU
  */
 template <bool bUseDepthBufferCollision>
 void ExecuteSimulationCommands(
+	FRHICommandList* RHICmdList, 
 	const TArray<FSimulationCommandGPU>& SimulationCommands,
 	const FParticleStateTextures& TextureResources,
 	const FParticleAttributesTexture& AttributeTexture,
@@ -1225,16 +1236,17 @@ void ExecuteSimulationCommands(
 		0
 		);
 
-	PixelShader->SetParameters(TextureResources, AttributeTexture, RenderAttributeTexture, CollisionView, SceneDepthTexture, GBufferATexture);
+	PixelShader->SetParameters(RHICmdList, TextureResources, AttributeTexture, RenderAttributeTexture, CollisionView, SceneDepthTexture, GBufferATexture);
 
 	// Draw tiles to perform the simulation step.
 	const int32 CommandCount = SimulationCommands.Num();
 	for (int32 CommandIndex = 0; CommandIndex < CommandCount; ++CommandIndex)
 	{
 		const FSimulationCommandGPU& Command = SimulationCommands[CommandIndex];
-		VertexShader->SetParameters(Command.TileOffsetsRef);
-		PixelShader->SetInstanceParameters(Command.DeltaSeconds, Command.UniformBuffer, Command.PerFrameUniformBuffer);
+		VertexShader->SetParameters(RHICmdList, Command.TileOffsetsRef);
+		PixelShader->SetInstanceParameters(RHICmdList, Command.DeltaSeconds, Command.UniformBuffer, Command.PerFrameUniformBuffer);
 		PixelShader->SetVectorFieldParameters(
+			RHICmdList, 
 			Command.VectorFieldsUniformBuffer,
 			Command.VectorFieldTexturesRHI
 			);
@@ -1284,7 +1296,8 @@ void ClearTiles(const TArray<uint32>& Tiles)
 		const int32 TilesThisDrawCall = FMath::Min<int32>( TileCount, MaxTilesPerDrawCall );
 		const uint32* TilesPtr = Tiles.GetTypedData() + FirstTile;
 		BuildTileVertexBuffer( BufferParam, TilesPtr, TilesThisDrawCall );
-		VertexShader->SetParameters(ShaderParam);
+		//@todo-rco: RHIPacketList
+		VertexShader->SetParameters(nullptr, ShaderParam);
 		DrawAlignedParticleTiles( ComputeAlignedTileCount(TilesThisDrawCall) );
 		TileCount -= TilesThisDrawCall;
 		FirstTile += TilesThisDrawCall;
@@ -1374,14 +1387,14 @@ public:
 	/**
 	 * Sets parameters for particle injection.
 	 */
-	void SetParameters()
+	void SetParameters(FRHICommandList* RHICmdList)
 	{
 		FParticleInjectionParameters Parameters;
 		Parameters.PixelScale.X = 1.0f / GParticleSimulationTextureSizeX;
 		Parameters.PixelScale.Y = 1.0f / GParticleSimulationTextureSizeY;
-		FParticleInjectionBufferRef UniformBuffer = FParticleInjectionBufferRef::CreateUniformBufferImmediate( Parameters, UniformBuffer_SingleUse );
+		FParticleInjectionBufferRef UniformBuffer = FParticleInjectionBufferRef::CreateUniformBufferImmediate( Parameters, UniformBuffer_SingleDraw );
 		FVertexShaderRHIParamRef VertexShader = GetVertexShader();
-		SetUniformBufferParameter( VertexShader, GetUniformBufferParameter<FParticleInjectionParameters>(), UniformBuffer );
+		SetUniformBufferParameter(RHICmdList, VertexShader, GetUniformBufferParameter<FParticleInjectionParameters>(), UniformBuffer );
 	}
 };
 
@@ -1518,7 +1531,8 @@ void InjectNewParticles( const TArray<FNewParticle>& NewParticles )
 			0
 			);
 
-		VertexShader->SetParameters();
+		//@todo-rco: RHIPacketList
+		VertexShader->SetParameters(nullptr);
 
 		// Stream 0: New particles.
 		RHISetStreamSource(
@@ -1593,10 +1607,10 @@ public:
 	/**
 	 * Set parameters for this shader.
 	 */
-	void SetParameters( const FParticleSimVisualizeBufferRef& UniformBuffer )
+	void SetParameters(FRHICommandList* RHICmdList, const FParticleSimVisualizeBufferRef& UniformBuffer )
 	{
 		FVertexShaderRHIParamRef VertexShader = GetVertexShader();
-		SetUniformBufferParameter( VertexShader, GetUniformBufferParameter<FParticleSimVisualizeParameters>(), UniformBuffer );
+		SetUniformBufferParameter(RHICmdList, VertexShader, GetUniformBufferParameter<FParticleSimVisualizeParameters>(), UniformBuffer );
 	}
 };
 
@@ -1649,14 +1663,14 @@ public:
 	/**
 	 * Set parameters for this shader.
 	 */
-	void SetParameters( int32 InVisualizationMode, FTexture2DRHIParamRef PositionTextureRHI, FTexture2DRHIParamRef PositionZWTextureRHI, FTexture2DRHIParamRef CurveTextureRHI )
+	void SetParameters(FRHICommandList* RHICmdList, int32 InVisualizationMode, FTexture2DRHIParamRef PositionTextureRHI, FTexture2DRHIParamRef PositionZWTextureRHI, FTexture2DRHIParamRef CurveTextureRHI )
 	{
 		FPixelShaderRHIParamRef PixelShader = GetPixelShader();
-		SetShaderValue( PixelShader, VisualizationMode, InVisualizationMode );
+		SetShaderValue(RHICmdList, PixelShader, VisualizationMode, InVisualizationMode );
 		FSamplerStateRHIParamRef SamplerStatePoint = TStaticSamplerState<SF_Point>::GetRHI();
-		SetTextureParameter( PixelShader, PositionTexture, PositionTextureSampler, SamplerStatePoint, PositionTextureRHI );
-		SetTextureParameter( PixelShader, PositionZWTexture, PositionZWTextureSampler, SamplerStatePoint, PositionZWTextureRHI );
-		SetTextureParameter( PixelShader, CurveTexture, CurveTextureSampler, SamplerStatePoint, CurveTextureRHI );
+		SetTextureParameter(RHICmdList, PixelShader, PositionTexture, PositionTextureSampler, SamplerStatePoint, PositionTextureRHI );
+		SetTextureParameter(RHICmdList, PixelShader, PositionZWTexture, PositionZWTextureSampler, SamplerStatePoint, PositionZWTextureRHI );
+		SetTextureParameter(RHICmdList, PixelShader, CurveTexture, CurveTextureSampler, SamplerStatePoint, CurveTextureRHI );
 	}
 
 private:
@@ -1741,15 +1755,18 @@ static void VisualizeGPUSimulation(
 		*PixelShader
 		);
 
+	//@todo-rco: RHIPacketList
+	FRHICommandList* RHICmdList = nullptr;
+
 	// Parameters for the visualization.
 	FParticleSimVisualizeParameters Parameters;
 	Parameters.ScaleBias.X = 2.0f * DisplaySizeX / (float)TargetSize.X;
 	Parameters.ScaleBias.Y = 2.0f * DisplaySizeY / (float)TargetSize.Y;
 	Parameters.ScaleBias.Z = 2.0f * DisplayOffsetX / (float)TargetSize.X - 1.0f;
 	Parameters.ScaleBias.W = 2.0f * DisplayOffsetY / (float)TargetSize.Y - 1.0f;
-	FParticleSimVisualizeBufferRef UniformBuffer = FParticleSimVisualizeBufferRef::CreateUniformBufferImmediate( Parameters, UniformBuffer_SingleUse );
-	VertexShader->SetParameters(UniformBuffer);
-	PixelShader->SetParameters(VisualizationMode, StateTextures.PositionTextureRHI, StateTextures.PositionZWTextureRHI, CurveTextureRHI);
+	FParticleSimVisualizeBufferRef UniformBuffer = FParticleSimVisualizeBufferRef::CreateUniformBufferImmediate( Parameters, UniformBuffer_SingleDraw );
+	VertexShader->SetParameters(RHICmdList, UniformBuffer);
+	PixelShader->SetParameters(RHICmdList, VisualizationMode, StateTextures.PositionTextureRHI, StateTextures.PositionZWTextureRHI, CurveTextureRHI);
 
 	const int32 VertexStride = sizeof(FVector2D);
 	
@@ -1805,10 +1822,10 @@ static void BuildParticleVertexBuffer( FVertexBufferRHIParamRef VertexBufferRHI,
 
 				// on some platforms, union and/or bitfield writes to Locked memory are really slow, so use a forced int write instead
 				// and in fact one 32-bit write is faster than two uint16 writes (i.e. using .Encoded)
-  				FParticleIndex Temp;
-  				Temp.X = IndexX;
-  				Temp.Y = IndexY;
-  				*(uint32*)ParticleIndices = *(uint32*)&Temp;
+				FParticleIndex Temp;
+				Temp.X = IndexX;
+				Temp.Y = IndexY;
+				*(uint32*)ParticleIndices = *(uint32*)&Temp;
 
 				// move to next particle
 				ParticleIndices += Stride;
@@ -1895,11 +1912,13 @@ public:
 	/**
 	 * Set output buffers for this shader.
 	 */
-	void SetOutput( FUnorderedAccessViewRHIParamRef OutBoundsUAV )
+	void SetOutput(FRHICommandList* RHICmdList, FUnorderedAccessViewRHIParamRef OutBoundsUAV )
 	{
 		FComputeShaderRHIParamRef ComputeShaderRHI = GetComputeShader();
 		if ( OutBounds.IsBound() )
 		{
+			check(!RHICmdList);
+
 			RHISetUAVParameter( ComputeShaderRHI, OutBounds.GetBaseIndex(), OutBoundsUAV );
 		}
 	}
@@ -1908,14 +1927,17 @@ public:
 	 * Set input parameters.
 	 */
 	void SetParameters(
+		FRHICommandList* RHICmdList, 
 		FParticleBoundsUniformBufferRef& UniformBuffer,
 		FShaderResourceViewRHIParamRef InIndicesSRV,
 		FTexture2DRHIParamRef PositionTextureRHI,
 		FTexture2DRHIParamRef PositionZWTextureRHI
 		)
 	{
+		check(!RHICmdList);
+
 		FComputeShaderRHIParamRef ComputeShaderRHI = GetComputeShader();
-		SetUniformBufferParameter( ComputeShaderRHI, GetUniformBufferParameter<FParticleBoundsParameters>(), UniformBuffer );
+		SetUniformBufferParameter(RHICmdList, ComputeShaderRHI, GetUniformBufferParameter<FParticleBoundsParameters>(), UniformBuffer );
 		if ( InParticleIndices.IsBound() )
 		{
 			RHISetShaderResourceViewParameter( ComputeShaderRHI, InParticleIndices.GetBaseIndex(), InIndicesSRV );
@@ -1933,8 +1955,10 @@ public:
 	/**
 	 * Unbinds any buffers that have been bound.
 	 */
-	void UnbindBuffers()
+	void UnbindBuffers(FRHICommandList* RHICmdList)
 	{
+		check(!RHICmdList);
+
 		FComputeShaderRHIParamRef ComputeShaderRHI = GetComputeShader();
 		if ( InParticleIndices.IsBound() )
 		{
@@ -1976,6 +2000,7 @@ static bool AreBoundsValid( const FVector& Mins, const FVector& Maxs )
  * @param ParticleCount - The number of particles in the emitter.
  */
 static FBox ComputeParticleBounds(
+	FRHICommandList* RHICmdList, 
 	FShaderResourceViewRHIParamRef VertexBufferSRV,
 	FTexture2DRHIParamRef PositionTextureRHI,
 	FTexture2DRHIParamRef PositionZWTextureRHI,
@@ -1997,7 +2022,9 @@ static FBox ComputeParticleBounds(
 		Parameters.ChunksPerGroup = ChunkCount / GroupCount;
 		Parameters.ExtraChunkCount = ChunkCount % GroupCount;
 		Parameters.ParticleCount = ParticleCount;
-		UniformBuffer = FParticleBoundsUniformBufferRef::CreateUniformBufferImmediate( Parameters, UniformBuffer_SingleUse );
+		UniformBuffer = FParticleBoundsUniformBufferRef::CreateUniformBufferImmediate( Parameters, UniformBuffer_SingleFrame );
+
+		check(!RHICmdList);
 
 		// Create a buffer for storing bounds.
 		const int32 BufferSize = GroupCount * 2 * sizeof(FVector4);
@@ -2015,14 +2042,15 @@ static FBox ComputeParticleBounds(
 		RHISetComputeShader(ParticleBoundsCS->GetComputeShader());
 
 		// Dispatch shader to compute bounds.
-		ParticleBoundsCS->SetOutput( BoundsVertexBufferUAV );
-		ParticleBoundsCS->SetParameters( UniformBuffer, VertexBufferSRV, PositionTextureRHI, PositionZWTextureRHI );
+		ParticleBoundsCS->SetOutput(RHICmdList, BoundsVertexBufferUAV );
+		ParticleBoundsCS->SetParameters(RHICmdList, UniformBuffer, VertexBufferSRV, PositionTextureRHI, PositionZWTextureRHI );
 		DispatchComputeShader(
+			RHICmdList, 
 			*ParticleBoundsCS, 
 			GroupCount,
 			1,
 			1 );
-		ParticleBoundsCS->UnbindBuffers();
+		ParticleBoundsCS->UnbindBuffers(RHICmdList);
 
 		// Read back bounds.
 		FVector4* GroupBounds = (FVector4*)RHILockVertexBuffer( BoundsVertexBufferRHI, 0, BufferSize, RLM_ReadOnly );
@@ -2414,9 +2442,9 @@ public:
 	 */
 	virtual void InitRHI()
 	{
-		UniformBuffer = FGPUSpriteEmitterUniformBufferRef::CreateUniformBufferImmediate( UniformParameters, UniformBuffer_MultiUse );
+		UniformBuffer = FGPUSpriteEmitterUniformBufferRef::CreateUniformBufferImmediate( UniformParameters, UniformBuffer_MultiFrame );
 		EmitterSimulationResources.SimulationUniformBuffer =
-			FParticleSimulationBufferRef::CreateUniformBufferImmediate( SimulationParameters, UniformBuffer_MultiUse );
+			FParticleSimulationBufferRef::CreateUniformBufferImmediate( SimulationParameters, UniformBuffer_MultiFrame );
 	}
 
 	/**
@@ -2500,10 +2528,10 @@ struct FGPUSpriteDynamicEmitterData : FDynamicEmitterDataBase
 
 			// Create the per-frame uniform buffer.
 			Simulation->PerFrameSimulationUniformBuffer =
-				FParticlePerFrameSimulationBufferRef::CreateUniformBufferImmediate(PerFrameSimulationParameters, UniformBuffer_MultiUse);
+				FParticlePerFrameSimulationBufferRef::CreateUniformBufferImmediate(PerFrameSimulationParameters, UniformBuffer_SingleFrame);
 
 			// Create per-emitter uniform buffer for dynamic parameters
-			DynamicUniformBuffer = FGPUSpriteEmitterDynamicUniformBufferRef::CreateUniformBufferImmediate(EmitterDynamicParameters, UniformBuffer_MultiUse);
+			DynamicUniformBuffer = FGPUSpriteEmitterDynamicUniformBufferRef::CreateUniformBufferImmediate(EmitterDynamicParameters, UniformBuffer_SingleFrame);
 
 			// Local vector field parameters.
 			Simulation->LocalVectorField.Intensity = LocalVectorFieldIntensity;
@@ -3243,7 +3271,9 @@ public:
 			FComputeGPUSpriteBoundsCommand,
 			FGPUSpriteParticleEmitterInstance*, EmitterInstance, this,
 		{
+			//@todo-rco: RHIPacketList
 			EmitterInstance->ParticleBoundingBox = ComputeParticleBounds(
+				nullptr,
 				EmitterInstance->Simulation->VertexBuffer.VertexBufferSRV, 
 				EmitterInstance->FXSystem->GetParticleSimulationResources()->GetCurrentStateTextures().PositionTextureRHI,
 				EmitterInstance->FXSystem->GetParticleSimulationResources()->GetCurrentStateTextures().PositionZWTextureRHI,
@@ -4022,7 +4052,7 @@ void FFXSystem::SimulateGPUParticles(
 			VectorFieldParameters.IntensityAndTightness[Index] = FVector2D::ZeroVector;
 		}
 		VectorFieldParameters.Count = 0;
-		EmptyVectorFieldUniformBuffer = FVectorFieldUniformBufferRef::CreateUniformBufferImmediate(VectorFieldParameters, UniformBuffer_SingleUse);
+		EmptyVectorFieldUniformBuffer = FVectorFieldUniformBufferRef::CreateUniformBufferImmediate(VectorFieldParameters, UniformBuffer_SingleFrame);
 	}
 
 	// Gather simulation commands from all active simulations.
@@ -4104,7 +4134,7 @@ void FFXSystem::SimulateGPUParticles(
 						VectorFieldParameters.IntensityAndTightness[Index] = FVector2D::ZeroVector;
 						SimulationCommand->VectorFieldTexturesRHI[Index] = BlackVolumeTextureRHI;
 					}
-					SimulationCommand->VectorFieldsUniformBuffer = FVectorFieldUniformBufferRef::CreateUniformBufferImmediate(VectorFieldParameters, UniformBuffer_SingleUse);
+					SimulationCommand->VectorFieldsUniformBuffer = FVectorFieldUniformBufferRef::CreateUniformBufferImmediate(VectorFieldParameters, UniformBuffer_SingleFrame);
 				}
 			}
 		
@@ -4128,7 +4158,9 @@ void FFXSystem::SimulateGPUParticles(
 
 		if (Phase == EParticleSimulatePhase::Collision && CollisionView)
 		{
+			/// ?
 			ExecuteSimulationCommands<true>(
+				nullptr, 
 				SimulationCommands,
 				PrevStateTextures,
 				ParticleSimulationResources->SimulationAttributesTexture,
@@ -4140,7 +4172,9 @@ void FFXSystem::SimulateGPUParticles(
 		}
 		else
 		{
+			/// ?
 			ExecuteSimulationCommands<false>(
+				nullptr,
 				SimulationCommands,
 				PrevStateTextures,
 				ParticleSimulationResources->SimulationAttributesTexture,
diff --git a/Engine/Source/Runtime/Engine/Private/Particles/ParticleSortingGPU.cpp b/Engine/Source/Runtime/Engine/Private/Particles/ParticleSortingGPU.cpp
--- a/Engine/Source/Runtime/Engine/Private/Particles/ParticleSortingGPU.cpp
+++ b/Engine/Source/Runtime/Engine/Private/Particles/ParticleSortingGPU.cpp
@@ -10,6 +10,7 @@
 #include "ParticleHelper.h"
 #include "UniformBuffer.h"
 #include "ShaderParameters.h"
+#include "ShaderParameterUtils.h"
 #include "GlobalShader.h"
 #include "../GPUSort.h"
 
@@ -93,8 +94,10 @@ public:
 	/**
 	 * Set output buffers for this shader.
 	 */
-	void SetOutput( FUnorderedAccessViewRHIParamRef OutKeysUAV, FUnorderedAccessViewRHIParamRef OutIndicesUAV )
+	void SetOutput(FRHICommandList* RHICmdList, FUnorderedAccessViewRHIParamRef OutKeysUAV, FUnorderedAccessViewRHIParamRef OutIndicesUAV )
 	{
+		check(!RHICmdList);
+
 		FComputeShaderRHIParamRef ComputeShaderRHI = GetComputeShader();
 		if ( OutKeys.IsBound() )
 		{
@@ -110,14 +113,17 @@ public:
 	 * Set input parameters.
 	 */
 	void SetParameters(
+		FRHICommandList* RHICmdList,
 		FParticleKeyGenUniformBufferRef& UniformBuffer,
 		FShaderResourceViewRHIParamRef InIndicesSRV
 		)
 	{
 		FComputeShaderRHIParamRef ComputeShaderRHI = GetComputeShader();
-		SetUniformBufferParameter( ComputeShaderRHI, GetUniformBufferParameter<FParticleKeyGenParameters>(), UniformBuffer );
+		SetUniformBufferParameter(RHICmdList, ComputeShaderRHI, GetUniformBufferParameter<FParticleKeyGenParameters>(), UniformBuffer );
 		if ( InParticleIndices.IsBound() )
 		{
+			check(!RHICmdList);
+
 			RHISetShaderResourceViewParameter( ComputeShaderRHI, InParticleIndices.GetBaseIndex(), InIndicesSRV );
 		}
 	}
@@ -125,15 +131,19 @@ public:
 	/**
 	 * Set the texture from which particle positions can be read.
 	 */
-	void SetPositionTextures(FTexture2DRHIParamRef PositionTextureRHI, FTexture2DRHIParamRef PositionZWTextureRHI)
+	void SetPositionTextures(FRHICommandList* RHICmdList, FTexture2DRHIParamRef PositionTextureRHI, FTexture2DRHIParamRef PositionZWTextureRHI)
 	{
 		FComputeShaderRHIParamRef ComputeShaderRHI = GetComputeShader();
 		if (PositionTexture.IsBound())
 		{
+			check(!RHICmdList);
+
 			RHISetShaderTexture(ComputeShaderRHI, PositionTexture.GetBaseIndex(), PositionTextureRHI);
 		}
 		if (PositionZWTexture.IsBound())
 		{
+			check(!RHICmdList);
+
 			RHISetShaderTexture(ComputeShaderRHI, PositionZWTexture.GetBaseIndex(), PositionZWTextureRHI);
 		}
 	}
@@ -141,8 +151,10 @@ public:
 	/**
 	 * Unbinds any buffers that have been bound.
 	 */
-	void UnbindBuffers()
+	void UnbindBuffers(FRHICommandList* RHICmdList)
 	{
+		check(!RHICmdList);
+
 		FComputeShaderRHIParamRef ComputeShaderRHI = GetComputeShader();
 		if ( InParticleIndices.IsBound() )
 		{
@@ -183,6 +195,7 @@ IMPLEMENT_SHADER_TYPE(,FParticleSortKeyGenCS,TEXT("ParticleSortKeyGen"),TEXT("Ge
  * @returns the total number of particles being sorted.
  */
 static int32 GenerateParticleSortKeys(
+	FRHICommandList* RHICmdList,
 	FUnorderedAccessViewRHIParamRef KeyBufferUAV,
 	FUnorderedAccessViewRHIParamRef SortedVertexBufferUAV,
 	FTexture2DRHIParamRef PositionTextureRHI,
@@ -202,9 +215,8 @@ static int32 GenerateParticleSortKeys(
 	// Grab the shader, set output.
 	TShaderMapRef<FParticleSortKeyGenCS> KeyGenCS(GetGlobalShaderMap());
 	RHISetComputeShader(KeyGenCS->GetComputeShader());
-
-	KeyGenCS->SetOutput(KeyBufferUAV, SortedVertexBufferUAV);
-	KeyGenCS->SetPositionTextures(PositionTextureRHI, PositionZWTextureRHI);
+	KeyGenCS->SetOutput(RHICmdList, KeyBufferUAV, SortedVertexBufferUAV);
+	KeyGenCS->SetPositionTextures(RHICmdList, PositionTextureRHI, PositionZWTextureRHI);
 
 	// For each simulation, generate keys and store them in the sorting buffers.
 	const int32 SimulationCount = SimulationsToSort.Num();
@@ -223,18 +235,18 @@ static int32 GenerateParticleSortKeys(
 		KeyGenParameters.OutputOffset = TotalParticleCount;
 		KeyGenParameters.EmitterKey = SimulationIndex << 16;
 		KeyGenParameters.KeyCount = ParticleCount;
-		KeyGenUniformBuffer = FParticleKeyGenUniformBufferRef::CreateUniformBufferImmediate( KeyGenParameters, UniformBuffer_SingleUse );
+		KeyGenUniformBuffer = FParticleKeyGenUniformBufferRef::CreateUniformBufferImmediate( KeyGenParameters, UniformBuffer_SingleDraw );
 
 		// Dispatch.
-		KeyGenCS->SetParameters(KeyGenUniformBuffer, SortInfo.VertexBufferSRV);
-		DispatchComputeShader(*KeyGenCS, GroupCount, 1, 1);
+		KeyGenCS->SetParameters(RHICmdList, KeyGenUniformBuffer, SortInfo.VertexBufferSRV);
+		DispatchComputeShader(RHICmdList, *KeyGenCS, GroupCount, 1, 1);
 
 		// Update offset in to the buffer.
 		TotalParticleCount += ParticleCount;
 	}
 
 	// Clear the output buffer.
-	KeyGenCS->UnbindBuffers();
+	KeyGenCS->UnbindBuffers(RHICmdList);
 
 	return TotalParticleCount;
 }
@@ -337,7 +349,9 @@ int32 SortParticlesGPU(
 	}
 
 	// First generate keys for each emitter to be sorted.
+	//@todo-rco: RHIPacketList
 	const int32 TotalParticleCount = GenerateParticleSortKeys(
+		nullptr,
 		ParticleSortBuffers.GetKeyBufferUAV(),
 		ParticleSortBuffers.GetVertexBufferUAV(),
 		PositionTextureRHI,
diff --git a/Engine/Source/Runtime/Engine/Private/Particles/ParticleSystemRender.cpp b/Engine/Source/Runtime/Engine/Private/Particles/ParticleSystemRender.cpp
--- a/Engine/Source/Runtime/Engine/Private/Particles/ParticleSystemRender.cpp
+++ b/Engine/Source/Runtime/Engine/Private/Particles/ParticleSystemRender.cpp
@@ -1200,7 +1200,7 @@ void FDynamicSpriteEmitterData::PreRenderView(FParticleSystemSceneProxy* Proxy,
 				FVector2D ObjectMacroUVScales;
 				Proxy->GetObjectPositionAndScale(*View,ObjectNDCPosition, ObjectMacroUVScales);
 				PerViewUniformParameters.MacroUVParameters = FVector4(ObjectNDCPosition.X, ObjectNDCPosition.Y, ObjectMacroUVScales.X, ObjectMacroUVScales.Y);
-				*SpriteUniformBufferPtr = FParticleSpriteUniformBufferRef::CreateUniformBufferImmediate(PerViewUniformParameters, UniformBuffer_SingleUse);
+				*SpriteUniformBufferPtr = FParticleSpriteUniformBufferRef::CreateUniformBufferImmediate(PerViewUniformParameters, UniformBuffer_SingleFrame);
 			}
 		}
 	}
@@ -1691,7 +1691,7 @@ void FDynamicMeshEmitterData::CreateRenderThreadResources(const FParticleSystemS
 		const uint32 TexCoordWeight = (SourceData->SubUVDataOffset > 0) ? 1 : 0;
 		UniformParameters.TexCoordWeightA = TexCoordWeight;
 		UniformParameters.TexCoordWeightB = 1 - TexCoordWeight;
-		UniformBuffer = FMeshParticleUniformBufferRef::CreateUniformBufferImmediate( UniformParameters, UniformBuffer_SingleUse );
+		UniformBuffer = FMeshParticleUniformBufferRef::CreateUniformBufferImmediate( UniformParameters, UniformBuffer_SingleFrame );
 	}
 }
 
@@ -2633,7 +2633,7 @@ FParticleBeamTrailUniformBufferRef CreateBeamTrailUniformBuffer(
 	// Screen alignment.
 	UniformParameters.ScreenAlignment = FVector4( (float)SourceData->ScreenAlignment, 0.0f, 0.0f, 0.0f );
 
-	return FParticleBeamTrailUniformBufferRef::CreateUniformBufferImmediate( UniformParameters, UniformBuffer_SingleUse );
+	return FParticleBeamTrailUniformBufferRef::CreateUniformBufferImmediate( UniformParameters, UniformBuffer_SingleFrame );
 }
 
 int32 FDynamicBeam2EmitterData::Render(FParticleSystemSceneProxy* Proxy, FPrimitiveDrawInterface* PDI,const FSceneView* View)
diff --git a/Engine/Source/Runtime/Engine/Private/Particles/ParticleVertexFactory.cpp b/Engine/Source/Runtime/Engine/Private/Particles/ParticleVertexFactory.cpp
--- a/Engine/Source/Runtime/Engine/Private/Particles/ParticleVertexFactory.cpp
+++ b/Engine/Source/Runtime/Engine/Private/Particles/ParticleVertexFactory.cpp
@@ -3,10 +3,11 @@
 /*=============================================================================
 	ParticleVertexFactory.cpp: Particle vertex factory implementation.
 =============================================================================*/
+
 #include "EnginePrivate.h"
 #include "ParticleDefinitions.h"
 #include "ParticleResources.h"
-#include "ShaderParameters.h"
+#include "ShaderParameterUtils.h"
 
 IMPLEMENT_UNIFORM_BUFFER_STRUCT(FParticleSpriteUniformParameters,TEXT("SpriteVF"));
 
@@ -32,11 +33,11 @@ class FParticleSpriteVertexFactoryShaderParametersVS : public FParticleSpriteVer
 {
 public:
 
-	virtual void SetMesh(FShader* Shader,const FVertexFactory* VertexFactory,const FSceneView& View,const FMeshBatchElement& BatchElement,uint32 DataFlags) const OVERRIDE
+	virtual void SetMesh(FRHICommandList* RHICmdList, FShader* Shader,const FVertexFactory* VertexFactory,const FSceneView& View,const FMeshBatchElement& BatchElement,uint32 DataFlags) const OVERRIDE
 	{
 		FParticleSpriteVertexFactory* SpriteVF = (FParticleSpriteVertexFactory*)VertexFactory;
 		FVertexShaderRHIParamRef VertexShaderRHI = Shader->GetVertexShader();
-		SetUniformBufferParameter( VertexShaderRHI, Shader->GetUniformBufferParameter<FParticleSpriteUniformParameters>(), SpriteVF->GetSpriteUniformBuffer() );
+		SetUniformBufferParameter(RHICmdList, VertexShaderRHI, Shader->GetUniformBufferParameter<FParticleSpriteUniformParameters>(), SpriteVF->GetSpriteUniformBuffer() );
 	}
 };
 
@@ -44,11 +45,11 @@ class FParticleSpriteVertexFactoryShaderParametersPS : public FParticleSpriteVer
 {
 public:
 
-	virtual void SetMesh(FShader* Shader,const FVertexFactory* VertexFactory,const FSceneView& View,const FMeshBatchElement& BatchElement,uint32 DataFlags) const OVERRIDE
+	virtual void SetMesh(FRHICommandList* RHICmdList, FShader* Shader,const FVertexFactory* VertexFactory,const FSceneView& View,const FMeshBatchElement& BatchElement,uint32 DataFlags) const OVERRIDE
 	{
 		FParticleSpriteVertexFactory* SpriteVF = (FParticleSpriteVertexFactory*)VertexFactory;
 		FPixelShaderRHIParamRef PixelShaderRHI = Shader->GetPixelShader();
-		SetUniformBufferParameter( PixelShaderRHI, Shader->GetUniformBufferParameter<FParticleSpriteUniformParameters>(), SpriteVF->GetSpriteUniformBuffer() );
+		SetUniformBufferParameter(RHICmdList, PixelShaderRHI, Shader->GetUniformBufferParameter<FParticleSpriteUniformParameters>(), SpriteVF->GetSpriteUniformBuffer() );
 	}
 };
 
diff --git a/Engine/Source/Runtime/Engine/Private/SceneManagement.cpp b/Engine/Source/Runtime/Engine/Private/SceneManagement.cpp
--- a/Engine/Source/Runtime/Engine/Private/SceneManagement.cpp
+++ b/Engine/Source/Runtime/Engine/Private/SceneManagement.cpp
@@ -16,7 +16,7 @@ TUniformBufferRef<FPrimitiveUniformShaderParameters> CreatePrimitiveUniformBuffe
 	check(IsInRenderingThread());
 	return TUniformBufferRef<FPrimitiveUniformShaderParameters>::CreateUniformBufferImmediate(
 		GetPrimitiveUniformShaderParameters(LocalToWorld, BoundsOrigin, BoundsOrigin, Bounds, BoundsRadius, bReceivesDecals),
-		UniformBuffer_MultiUse
+		UniformBuffer_MultiFrame
 		);
 }
 */
diff --git a/Engine/Source/Runtime/Engine/Private/SceneView.cpp b/Engine/Source/Runtime/Engine/Private/SceneView.cpp
--- a/Engine/Source/Runtime/Engine/Private/SceneView.cpp
+++ b/Engine/Source/Runtime/Engine/Private/SceneView.cpp
@@ -170,6 +170,7 @@ FSceneView::FSceneView(const FSceneViewInitOptions& InitOptions)
 	, bIsViewInfo(false)
 	, bIsSceneCapture(false)
 	, bIsReflectionCapture(false)
+	, bIsLocked(false)
 	, bStaticSceneOnly(false)
 	, RenderingCompositePassContext(0)
 #if WITH_EDITOR
diff --git a/Engine/Source/Runtime/Engine/Private/ScreenRendering.cpp b/Engine/Source/Runtime/Engine/Private/ScreenRendering.cpp
--- a/Engine/Source/Runtime/Engine/Private/ScreenRendering.cpp
+++ b/Engine/Source/Runtime/Engine/Private/ScreenRendering.cpp
@@ -1,7 +1,7 @@
 // Copyright 1998-2014 Epic Games, Inc. All Rights Reserved.
 
 /*=============================================================================
-	ScreenRendering.cpp: D3D render target implementation.
+	ScreenRendering.cpp: Screen rendering implementation.
 =============================================================================*/
 
 #include "EnginePrivate.h"
@@ -27,7 +27,8 @@ void DrawScreenQuad( float X0, float Y0, float U0, float V0, float X1, float Y1,
 	TShaderMapRef<FScreenPS> ScreenPixelShader(GetGlobalShaderMap());
 	SetGlobalBoundShaderState(ScreenBoundShaderState, GScreenVertexDeclaration.VertexDeclarationRHI, *ScreenVertexShader, *ScreenPixelShader);
 
-	ScreenPixelShader->SetParameters(Texture);
+	//@todo-rco: RHIPacketList
+	ScreenPixelShader->SetParameters(nullptr, Texture);
 
 	// Generate the vertices used
 	FScreenVertex Vertices[4];
@@ -62,7 +63,8 @@ void DrawNormalizedScreenQuad( float X, float Y, float U0, float V0, float SizeX
 	TShaderMapRef<FScreenPS> ScreenPixelShader(GetGlobalShaderMap());
 	SetGlobalBoundShaderState(ScreenBoundShaderState, GScreenVertexDeclaration.VertexDeclarationRHI, *ScreenVertexShader, *ScreenPixelShader);
 
-	ScreenPixelShader->SetParameters(TStaticSamplerState<SF_Bilinear>::GetRHI(), TextureRHI);
+	/// ?
+	ScreenPixelShader->SetParameters(nullptr, TStaticSamplerState<SF_Bilinear>::GetRHI(), TextureRHI);
 
 	// Generate the vertices used
 	FScreenVertex Vertices[4];
diff --git a/Engine/Source/Runtime/Engine/Private/ShaderCompiler/ShaderCompiler.cpp b/Engine/Source/Runtime/Engine/Private/ShaderCompiler/ShaderCompiler.cpp
--- a/Engine/Source/Runtime/Engine/Private/ShaderCompiler/ShaderCompiler.cpp
+++ b/Engine/Source/Runtime/Engine/Private/ShaderCompiler/ShaderCompiler.cpp
@@ -1931,11 +1931,12 @@ void GlobalBeginCompileShader(
 	Input.ShaderFormat = LegacyShaderPlatformToShaderFormat(EShaderPlatform(Target.Platform));
 	Input.SourceFilename = SourceFilename;
 	Input.EntryPointName = FunctionName;
-	Input.ShaderDebugInfoDirectory = GShaderCompilingManager->GetAbsoluteShaderDebugInfoDirectory();
+	Input.DumpDebugInfoRootPath = GShaderCompilingManager->GetAbsoluteShaderDebugInfoDirectory();
 
 	if (GDumpShaderDebugInfo != 0)
 	{
-		Input.DumpDebugInfoPath = GShaderCompilingManager->GetAbsoluteShaderDebugInfoDirectory() / Input.ShaderFormat.ToString() / DebugGroupName;
+		Input.DumpDebugInfoRootPath = GShaderCompilingManager->GetAbsoluteShaderDebugInfoDirectory() / Input.ShaderFormat.ToString();
+		Input.DumpDebugInfoPath = Input.DumpDebugInfoRootPath / DebugGroupName;
 
 		if (VFType)
 		{
