Fix bone name case issues on export.

#TTP 337651 - FBX Import/Export bone names are not case sensitive

diff --git a/Engine/Source/Editor/UnrealEd/Private/Factories/SkeletalMeshImport.cpp b/Engine/Source/Editor/UnrealEd/Private/Factories/SkeletalMeshImport.cpp
--- a/Engine/Source/Editor/UnrealEd/Private/Factories/SkeletalMeshImport.cpp
+++ b/Engine/Source/Editor/UnrealEd/Private/Factories/SkeletalMeshImport.cpp
@@ -201,7 +201,7 @@ bool ProcessImportMeshSkeleton(FReferenceSkeleton& RefSkeleton, int32& SkeletalD
 	{
 		const VBone & BinaryBone = RefBonesBinary[ b ];
 		const FString BoneName = FSkeletalMeshImportData::FixupBoneName( BinaryBone.Name );
-		const FMeshBoneInfo BoneInfo(FName(*BoneName, FNAME_Add, true), BinaryBone.ParentIndex);
+		const FMeshBoneInfo BoneInfo(FName(*BoneName, FNAME_Add, true), BinaryBone.Name, BinaryBone.ParentIndex);
 		const FTransform BoneTransform(BinaryBone.BonePos.Transform);
 
 		if(RefSkeleton.FindBoneIndex(BoneInfo.Name) != INDEX_NONE)
diff --git a/Engine/Source/Editor/UnrealEd/Private/Fbx/FbxSkeletalMeshExport.cpp b/Engine/Source/Editor/UnrealEd/Private/Fbx/FbxSkeletalMeshExport.cpp
--- a/Engine/Source/Editor/UnrealEd/Private/Fbx/FbxSkeletalMeshExport.cpp
+++ b/Engine/Source/Editor/UnrealEd/Private/Fbx/FbxSkeletalMeshExport.cpp
@@ -63,7 +63,7 @@ FbxNode* FFbxExporter::CreateSkeleton(const USkeletalMesh& SkelMesh, TArray<FbxN
 		const FMeshBoneInfo & CurrentBone = RefSkeleton.GetRefBoneInfo()[BoneIndex];
 		const FTransform & BoneTransform = RefSkeleton.GetRefBonePose()[BoneIndex];
 
-		FbxString BoneName = Converter.ConvertToFbxString(CurrentBone.Name);
+		FbxString BoneName = Converter.ConvertToFbxString(CurrentBone.ExportName);
 
 
 		// Create the node's attributes
diff --git a/Engine/Source/Runtime/Core/Public/UObject/ObjectVersion.h b/Engine/Source/Runtime/Core/Public/UObject/ObjectVersion.h
--- a/Engine/Source/Runtime/Core/Public/UObject/ObjectVersion.h
+++ b/Engine/Source/Runtime/Core/Public/UObject/ObjectVersion.h
@@ -552,6 +552,9 @@ enum EUnrealEngineObjectUE4Version
 	VER_UE4_FTEXT_HISTORY,
 	// Shift comments to the left to contain expressions properly
 	VER_UE4_FIX_MATERIAL_COMMENTS,
+	// Bone names stored as FName mean that we cant guarrentee the correct case on export, now we store a separate string for export purposes only
+	VER_UE4_STORE_BONE_EXPORT_NAMES,
+
 	// -----<new versions can be added before this line>-------------------------------------------------
 	// - this needs to be the last line (see note below)
 	VER_UE4_AUTOMATIC_VERSION_PLUS_ONE,
diff --git a/Engine/Source/Runtime/Engine/Private/Animation/Skeleton.cpp b/Engine/Source/Runtime/Engine/Private/Animation/Skeleton.cpp
--- a/Engine/Source/Runtime/Engine/Private/Animation/Skeleton.cpp
+++ b/Engine/Source/Runtime/Engine/Private/Animation/Skeleton.cpp
@@ -162,7 +162,7 @@ void USkeleton::ConvertToFReferenceSkeleton()
 	for(int32 BoneIndex=0; BoneIndex<NumRefBones; BoneIndex++)
 	{
 		const FBoneNode & BoneNode = BoneTree[BoneIndex];
-		FMeshBoneInfo BoneInfo(BoneNode.Name_DEPRECATED, BoneNode.ParentIndex_DEPRECATED);
+		FMeshBoneInfo BoneInfo(BoneNode.Name_DEPRECATED, BoneNode.Name_DEPRECATED.ToString(), BoneNode.ParentIndex_DEPRECATED);
 		const FTransform & BoneTransform = RefLocalPoses_DEPRECATED[BoneIndex];
 
 		// All should be good. Parents before children, no duplicate bones?
diff --git a/Engine/Source/Runtime/Engine/Private/SkeletalMesh.cpp b/Engine/Source/Runtime/Engine/Private/SkeletalMesh.cpp
--- a/Engine/Source/Runtime/Engine/Private/SkeletalMesh.cpp
+++ b/Engine/Source/Runtime/Engine/Private/SkeletalMesh.cpp
@@ -672,6 +672,20 @@ FArchive &operator<<( FArchive& Ar, FMeshBoneInfo& F)
 		Ar << DummyColor;
 	}
 
+#if WITH_EDITORONLY_DATA
+	if (Ar.UE4Ver() >= VER_UE4_STORE_BONE_EXPORT_NAMES)
+	{
+		if(!Ar.IsCooking())
+		{
+			Ar << F.ExportName;
+		}
+	}
+	else
+	{
+		F.ExportName = F.Name.ToString();
+	}
+#endif
+
 	return Ar;
 }
 
diff --git a/Engine/Source/Runtime/Engine/Public/SkeletalMeshTypes.h b/Engine/Source/Runtime/Engine/Public/SkeletalMeshTypes.h
--- a/Engine/Source/Runtime/Engine/Public/SkeletalMeshTypes.h
+++ b/Engine/Source/Runtime/Engine/Public/SkeletalMeshTypes.h
@@ -113,19 +113,31 @@ struct FMeshBoneInfo
 {
 	// Bone's name.
 	FName Name;			
-	// 0/NULL if this is the root bone.  
+
+#if WITH_EDITORONLY_DATA
+	// Name used for export (this should be exact as FName may mess with case) 
+	FString ExportName;
+#endif
+
+	// 0/NULL if this is the root bone. 
 	int32 ParentIndex;	
 
 	FMeshBoneInfo() : Name(NAME_None), ParentIndex(INDEX_NONE) {}
 
-	FMeshBoneInfo(const FName & InName, int32 InParentIndex)
+	FMeshBoneInfo(const FName & InName, const FString& InExportName, int32 InParentIndex)
 	:	Name(InName)
 	,	ParentIndex(InParentIndex)
+#if WITH_EDITORONLY_DATA
+	, ExportName(InExportName)
+#endif
 	{}
 
 	FMeshBoneInfo(const FMeshBoneInfo & Other)
 		:	Name(Other.Name)
 		,	ParentIndex(Other.ParentIndex)
+#if WITH_EDITORONLY_DATA
+		, ExportName(Other.ExportName)
+#endif
 	{}
 
 	bool operator==( const FMeshBoneInfo& B ) const
