Fix position issues when applying transactions on blueprint actors with multiple scene components

diff --git a/Engine/Source/Runtime/Engine/Private/ActorEditor.cpp b/Engine/Source/Runtime/Engine/Private/ActorEditor.cpp
--- a/Engine/Source/Runtime/Engine/Private/ActorEditor.cpp
+++ b/Engine/Source/Runtime/Engine/Private/ActorEditor.cpp
@@ -235,11 +235,12 @@ AActor::FActorTransactionAnnotation::FActorTransactionAnnotation(const AActor* A
 
 		for (USceneComponent* AttachChild : RootComponent->AttachChildren)
 		{
-			if (AttachChild)
+			AActor* ChildOwner = (AttachChild ? AttachChild->GetOwner() : NULL);
+			if (ChildOwner != Actor)
 			{
 				// Save info about actor to reattach
 				FActorRootComponentReconstructionData::FAttachedActorInfo Info;
-				Info.Actor = AttachChild->GetOwner();
+				Info.Actor = ChildOwner;
 				Info.SocketName = AttachChild->AttachSocketName;
 				Info.RelativeTransform = AttachChild->GetRelativeTransform();
 				RootComponentData.AttachedToInfo.Add(Info);
