Use macro to consolidate how we compose guids in various places, improve comments

diff --git a/Engine/Source/Runtime/Engine/Private/PackageMapClient.cpp b/Engine/Source/Runtime/Engine/Private/PackageMapClient.cpp
--- a/Engine/Source/Runtime/Engine/Private/PackageMapClient.cpp
+++ b/Engine/Source/Runtime/Engine/Private/PackageMapClient.cpp
@@ -20,6 +20,11 @@ static const int GUID_PACKET_ACKED		= -1;
  */
 static const int INTERNAL_LOAD_OBJECT_RECURSION_LIMIT = 16;
 
+#define GET_NET_GUID_INDEX( GUID )					( GUID.Value >> 1 )
+#define GET_NET_GUID_STATIC_MOD( GUID )				( GUID.Value & 1 )
+#define COMPOSE_NET_GUID( Index, IsStatic )			( ( ( Index ) << 1 ) | ( IsStatic ) )
+#define COMPOSE_RELATIVE_NET_GUID( GUID, Index )	( COMPOSE_NET_GUID( GET_NET_GUID_INDEX( GUID ) + Index, GET_NET_GUID_STATIC_MOD( GUID ) ) )
+
 /*-----------------------------------------------------------------------------
 	UPackageMapClient implementation.
 -----------------------------------------------------------------------------*/
@@ -301,6 +306,8 @@ bool UPackageMapClient::SerializeNewActor(FArchive& Ar, class UActorChannel *Cha
 
 		if ( Ar.IsSaving() )
 		{
+			// We're not really saving anything necessary here except a checksum to make sure we're in sync
+			// It might be a worth thinking about turning this off in shipping, and trust this is working as intended to save bandwidth if needed
 			uint32 Checksum = 0;
 
 			for ( int32 i = 0; i < Subobjs.Num(); i++ )
@@ -309,7 +316,8 @@ bool UPackageMapClient::SerializeNewActor(FArchive& Ar, class UActorChannel *Cha
 
 				FNetworkGUID SubobjNetGUID = GuidCache->NetGUIDLookup.FindRef( SubObj );
 
-				check( SubobjNetGUID.Value == ( ( ( NetGUID.Value >> 1 ) + i + 1 ) << 1 ) );
+				// Make sure these sub objects have net guids that are relative to this owning actor
+				check( SubobjNetGUID.Value == COMPOSE_RELATIVE_NET_GUID( NetGUID, i + 1 ) );
 
 				// Evolve checksum so we can sanity check
 				Checksum = FCrc::MemCrc32( &SubobjNetGUID.Value, sizeof( SubobjNetGUID.Value ), Checksum );
@@ -322,6 +330,7 @@ bool UPackageMapClient::SerializeNewActor(FArchive& Ar, class UActorChannel *Cha
 		}
 		else
 		{
+			// Load the server checksum, then generate our own, and make sure they match
 			uint32 ServerChecksum = 0;
 			uint32 LocalChecksum = 0;
 			Ar << ServerChecksum;
@@ -330,7 +339,8 @@ bool UPackageMapClient::SerializeNewActor(FArchive& Ar, class UActorChannel *Cha
 
 			for ( int32 i = 0; i < Subobjs.Num(); i++ )
 			{
-				SubObjGuids.Add( ( ( ( NetGUID.Value >> 1 ) + i + 1 ) << 1 ) );
+				// Generate a guid that is relative to our owning actor
+				SubObjGuids.Add( COMPOSE_RELATIVE_NET_GUID( NetGUID, i + 1 ) );
 
 				// Evolve checksum so we can sanity check
 				LocalChecksum = FCrc::MemCrc32( &SubObjGuids[i].Value, sizeof( SubObjGuids[i].Value ), LocalChecksum );
@@ -1581,10 +1591,7 @@ FNetworkGUID FNetGUIDCache::AssignNewNetGUID( const UObject * Object )
 		}
 	}
 
-#define GET_NET_GUID_INDEX( GUID )				( GUID.Value >> 1 )
-#define GET_NET_GUID_STATIC_MOD( GUID )			( GUID.Value & 1 )
-#define COMPOSE_NET_GUID( Index, IsStatic )		( ( ( Index ) << 1 ) | ( IsStatic ) )
-#define ALLOC_NEW_NET_GUID( IsStatic )			( COMPOSE_NET_GUID( ++UniqueNetIDs[ IsStatic ], IsStatic ) )
+#define ALLOC_NEW_NET_GUID( IsStatic ) ( COMPOSE_NET_GUID( ++UniqueNetIDs[ IsStatic ], IsStatic ) )
 
 	// Generate new NetGUID and assign it
 	const int32 IsStatic = IsDynamicObject( Object ) ? 0 : 1;
@@ -1613,7 +1620,7 @@ FNetworkGUID FNetGUIDCache::AssignNewNetGUID( const UObject * Object )
 			const int32 SubIsStatic = IsDynamicObject( SubObj ) ? 0 : 1;
 			check( SubIsStatic == IsStatic );
 			const FNetworkGUID SubobjNetGUID( ALLOC_NEW_NET_GUID( SubIsStatic ) );
-			check( SubobjNetGUID.Value == COMPOSE_NET_GUID( GET_NET_GUID_INDEX( NewNetGuid ) + 1 + i, GET_NET_GUID_STATIC_MOD( NewNetGuid ) ) );
+			check( SubobjNetGUID.Value == COMPOSE_RELATIVE_NET_GUID( NewNetGuid, i + 1 ) );
 			AssignNetGUID( SubObj, SubobjNetGUID );
 		}
 	}
