Make IPv6 parsing more robust in FURL
Now handles port numbers and IPv4 shorthand format

diff --git a/Engine/Source/Runtime/Engine/Private/URL.cpp b/Engine/Source/Runtime/Engine/Private/URL.cpp
--- a/Engine/Source/Runtime/Engine/Private/URL.cpp
+++ b/Engine/Source/Runtime/Engine/Private/URL.cpp
@@ -265,12 +265,15 @@ FURL::FURL( FURL* Base, const TCHAR* TextURL, ETravelType Type )
 			const TCHAR* Dot = FCString::Strchr(URL,'.');
 			const int32 ExtLen = FPackageName::GetMapPackageExtension().Len();
 
-			if
-			(	(Dot)
-			&&	(Dot-URL>0)
-			&&	(FCString::Strnicmp( Dot, *FPackageName::GetMapPackageExtension(), FPackageName::GetMapPackageExtension().Len() )!=0 || FChar::IsAlnum(Dot[ExtLen]) )
-			&&	(FCString::Strnicmp( Dot+1,*UrlConfig.DefaultSaveExt, UrlConfig.DefaultSaveExt.Len() )!=0 || FChar::IsAlnum(Dot[UrlConfig.DefaultSaveExt.Len()+1]) )
-			&&	(FCString::Strnicmp( Dot+1,TEXT("demo"), 4 ) != 0 || FChar::IsAlnum(Dot[5])) )
+			const bool bIsHostnameWithDot =
+					(Dot)
+				&&	(Dot-URL>0)
+				&&	(FCString::Strnicmp( Dot, *FPackageName::GetMapPackageExtension(), FPackageName::GetMapPackageExtension().Len() )!=0 || FChar::IsAlnum(Dot[ExtLen]) )
+				&&	(FCString::Strnicmp( Dot+1,*UrlConfig.DefaultSaveExt, UrlConfig.DefaultSaveExt.Len() )!=0 || FChar::IsAlnum(Dot[UrlConfig.DefaultSaveExt.Len()+1]) )
+				&&	(FCString::Strnicmp( Dot+1,TEXT("demo"), 4 ) != 0 || FChar::IsAlnum(Dot[5]));
+
+			// Square bracket indicates an IPv6 address, but IPv6 addresses can contain dots also
+			if ( bIsHostnameWithDot || SquareBracket )
 			{
 				TCHAR* ss = URL;
 				URL     = FCString::Strchr(URL,'/');
@@ -278,35 +281,35 @@ FURL::FURL( FURL* Base, const TCHAR* TextURL, ETravelType Type )
 				{
 					*(URL++) = 0;
 				}
-				TCHAR* t = FCString::Strchr(ss,':');
+
+				// Skip past all the ':' characters in the IPv6 address to get to the port.
+				TCHAR* ClosingSquareBracket = FCString::Strchr(ss, ']');
+				
+				TCHAR* PortText = ss;
+				if( ClosingSquareBracket )
+				{
+					PortText = ClosingSquareBracket;
+				}
+				TCHAR* t = FCString::Strchr(PortText,':');
 				if( t )
 				{
 					// Port.
 					*(t++) = 0;
 					Port = FCString::Atoi( t );
 				}
-				Host = ss;
-				if( FCString::Stricmp(*Protocol,*UrlConfig.DefaultProtocol)==0 )
+
+				if(SquareBracket && ClosingSquareBracket)
 				{
-					Map = UGameMapsSettings::GetGameDefaultMap();
+					// Trim the brackets from the host address
+					*ClosingSquareBracket = 0;
+					Host = ss + 1;
 				}
 				else
 				{
-					Map = TEXT("");
+					// Normal IPv4 address
+					Host = ss;
 				}
-				FarHost = 1;
-			}
 
-			// Parse IPv6 host.
-			if(URL && SquareBracket)
-			{
-				const TCHAR* ss = URL;
-				URL = FCString::Strchr(URL, ']');
-				if (URL)
-				{
-					*URL = 0;
-					Host = SquareBracket + 1;
-				}
 				if( FCString::Stricmp(*Protocol,*UrlConfig.DefaultProtocol)==0 )
 				{
 					Map = UGameMapsSettings::GetGameDefaultMap();
