Add per poly skeletal mesh body setup.
- per poly works for static data, still need to update with animation

diff --git a/Engine/Source/Runtime/Engine/Classes/Components/SkeletalMeshComponent.h b/Engine/Source/Runtime/Engine/Classes/Components/SkeletalMeshComponent.h
--- a/Engine/Source/Runtime/Engine/Classes/Components/SkeletalMeshComponent.h
+++ b/Engine/Source/Runtime/Engine/Classes/Components/SkeletalMeshComponent.h
@@ -1,6 +1,8 @@
 // Copyright 1998-2014 Epic Games, Inc. All Rights Reserved.
 
 #pragma once
+
+#include "Interfaces/Interface_CollisionDataProvider.h"
 #include "Components/SkinnedMeshComponent.h"
 #include "PhysicsEngine/PhysicsConstraintComponent.h"
 #include "SkeletalMeshComponent.generated.h"
@@ -208,7 +210,7 @@ struct FSkeletalMeshComponentPreClothTickFunction : public FTickFunction
  * SkeletalMeshComponent is a mesh that supports skeletal animation and physics.
  */
 UCLASS(ClassGroup=(Rendering, Common), hidecategories=Object, config=Engine, editinlinenew, meta=(BlueprintSpawnableComponent))
-class ENGINE_API USkeletalMeshComponent : public USkinnedMeshComponent
+class ENGINE_API USkeletalMeshComponent : public USkinnedMeshComponent, public IInterface_CollisionDataProvider
 {
 	GENERATED_UCLASS_BODY()
 	
@@ -382,6 +384,30 @@ public:
 	UPROPERTY(EditAnywhere, AdvancedDisplay, BlueprintReadWrite, Category=Animation)
 	uint32 bPauseAnims:1;
 
+	/**
+	* Uses skinned data for collision data.
+	*/
+	UPROPERTY(EditAnywhere, AdvancedDisplay, BlueprintReadOnly, Category=SkeletalMesh)
+	uint32 bEnablePerPolyCollision:1;
+
+	/**
+	* Used for per poly collision. In 99% of cases you will be better off using a Physics Asset.
+	* This BodySetup is per instance because all modification of vertices is done in place */
+	UPROPERTY(transient)
+	class UBodySetup * BodySetup;
+
+	void CreateBodySetup();
+
+	// Begin Interface_CollisionDataProvider Interface
+	virtual bool GetPhysicsTriMeshData(struct FTriMeshCollisionData* CollisionData, bool InUseAllTriData) OVERRIDE;
+	virtual bool ContainsPhysicsTriMeshData(bool InUseAllTriData) const OVERRIDE;
+	virtual bool WantsNegXTriMesh() OVERRIDE
+	{
+		return true;
+	}
+	// End Interface_CollisionDataProvider Interface
+
+
 	/**
 	 * Misc 
 	 */
diff --git a/Engine/Source/Runtime/Engine/Private/PhysicsEngine/PhysAnim.cpp b/Engine/Source/Runtime/Engine/Private/PhysicsEngine/PhysAnim.cpp
--- a/Engine/Source/Runtime/Engine/Private/PhysicsEngine/PhysAnim.cpp
+++ b/Engine/Source/Runtime/Engine/Private/PhysicsEngine/PhysAnim.cpp
@@ -278,75 +278,84 @@ void USkeletalMeshComponent::UpdateKinematicBonesToPhysics(bool bTeleport)
 		UE_LOG(LogPhysics, Log, TEXT("USkeletalMeshComponent::UpdateKinematicBonesToPhysics : Non-uniform scale factor (%s) can cause physics to mismatch for %s  SkelMesh: %s"), *MeshScale3D.ToString(), *GetFullName(), SkeletalMesh ? *SkeletalMesh->GetFullName() : TEXT("NULL"));
 	}
 #endif
-	const UPhysicsAsset * const PhysicsAsset = GetPhysicsAsset();
-	if(PhysicsAsset && SkeletalMesh && Bodies.Num() > 0)
+
+	if (bEnablePerPolyCollision == false)
 	{
-#if !(UE_BUILD_SHIPPING || UE_BUILD_TEST)
-		if ( !ensure( PhysicsAsset->BodySetup.Num() == Bodies.Num()) )
+		const UPhysicsAsset * const PhysicsAsset = GetPhysicsAsset();
+		if (PhysicsAsset && SkeletalMesh && Bodies.Num() > 0)
 		{
-			// related to TTP 280315
-			UE_LOG(LogPhysics, Warning, TEXT("Mesh (%s) has PhysicsAsset(%s), and BodySetup(%d) and Bodies(%d) don't match"), 
-				*SkeletalMesh->GetName(), *PhysicsAsset->GetName(), PhysicsAsset->BodySetup.Num(), Bodies.Num());
-			return;
-		}
+#if !(UE_BUILD_SHIPPING || UE_BUILD_TEST)
+			if (!ensure(PhysicsAsset->BodySetup.Num() == Bodies.Num()))
+			{
+				// related to TTP 280315
+				UE_LOG(LogPhysics, Warning, TEXT("Mesh (%s) has PhysicsAsset(%s), and BodySetup(%d) and Bodies(%d) don't match"),
+					*SkeletalMesh->GetName(), *PhysicsAsset->GetName(), PhysicsAsset->BodySetup.Num(), Bodies.Num());
+				return;
+			}
 #endif
 
-		// Iterate over each body
-		for(int32 i = 0; i < Bodies.Num(); i++)
-		{
-			// If we have a physics body, and its kinematic...
-			FBodyInstance* BodyInst = Bodies[i];
-			check(BodyInst);
-
-			if (BodyInst->IsValidBodyInstance() && !BodyInst->IsInstanceSimulatingPhysics(true))
+			// Iterate over each body
+			for (int32 i = 0; i < Bodies.Num(); i++)
 			{
-				// Find the graphics bone index that corresponds to this physics body.
-				FName const BodyName = PhysicsAsset->BodySetup[i]->BoneName;
-				int32 const BoneIndex = SkeletalMesh->RefSkeleton.FindBoneIndex(BodyName);
+				// If we have a physics body, and its kinematic...
+				FBodyInstance* BodyInst = Bodies[i];
+				check(BodyInst);
 
-				// If we could not find it - warn.
-				if( BoneIndex == INDEX_NONE || BoneIndex >= SpaceBases.Num() )
-				{
-					UE_LOG(LogPhysics, Log,  TEXT("UpdateRBBones: WARNING: Failed to find bone '%s' need by PhysicsAsset '%s' in SkeletalMesh '%s'."), *BodyName.ToString(), *PhysicsAsset->GetName(), *SkeletalMesh->GetName() );
-				}
-				else
+				if (BodyInst->IsValidBodyInstance() && !BodyInst->IsInstanceSimulatingPhysics(true))
 				{
-					// update bone transform to world
-					FTransform BoneTransform = SpaceBases[BoneIndex] * CurrentLocalToWorld;
-						
-					// move body
-					BodyInst->SetBodyTransform(BoneTransform, bTeleport);
-						
-					// now update scale
-					// if uniform, we'll use BoneTranform
-					if ( MeshScale3D.IsUniform() )
+					// Find the graphics bone index that corresponds to this physics body.
+					FName const BodyName = PhysicsAsset->BodySetup[i]->BoneName;
+					int32 const BoneIndex = SkeletalMesh->RefSkeleton.FindBoneIndex(BodyName);
+
+					// If we could not find it - warn.
+					if (BoneIndex == INDEX_NONE || BoneIndex >= SpaceBases.Num())
 					{
-						// @todo UE4 should we update scale when it's simulated?
-						BodyInst->UpdateBodyScale(BoneTransform.GetScale3D());						
+						UE_LOG(LogPhysics, Log, TEXT("UpdateRBBones: WARNING: Failed to find bone '%s' need by PhysicsAsset '%s' in SkeletalMesh '%s'."), *BodyName.ToString(), *PhysicsAsset->GetName(), *SkeletalMesh->GetName());
 					}
 					else
 					{
-						// @note When you have non-uniform scale on mesh base,
-						// hierarchical bone transform can update scale too often causing performance issue
-						// So we just use mesh scale for all bodies when non-uniform
-						// This means physics representation won't be accurate, but
-						// it is performance friendly by preventing too frequent physics update
-						BodyInst->UpdateBodyScale(MeshScale3D);
+						// update bone transform to world
+						FTransform BoneTransform = SpaceBases[BoneIndex] * CurrentLocalToWorld;
+
+						// move body
+						BodyInst->SetBodyTransform(BoneTransform, bTeleport);
+
+						// now update scale
+						// if uniform, we'll use BoneTranform
+						if (MeshScale3D.IsUniform())
+						{
+							// @todo UE4 should we update scale when it's simulated?
+							BodyInst->UpdateBodyScale(BoneTransform.GetScale3D());
+						}
+						else
+						{
+							// @note When you have non-uniform scale on mesh base,
+							// hierarchical bone transform can update scale too often causing performance issue
+							// So we just use mesh scale for all bodies when non-uniform
+							// This means physics representation won't be accurate, but
+							// it is performance friendly by preventing too frequent physics update
+							BodyInst->UpdateBodyScale(MeshScale3D);
+						}
 					}
 				}
-			}
-			else
-			{
-				//make sure you have physics weight or blendphysics on, otherwise, you'll have inconsistent representation of bodies
-				// @todo make this to be kismet log? But can be too intrusive
-				if (!bBlendPhysics && BodyInst->PhysicsBlendWeight <= 0.f && BodyInst->BodySetup.IsValid())
+				else
 				{
-					UE_LOG(LogPhysics, Warning, TEXT("%s(Mesh %s, PhysicsAsset %s, Bone %s) is simulating, but no blending. "), 
-						*GetName(), *GetNameSafe(SkeletalMesh), *GetNameSafe(PhysicsAsset), *BodyInst->BodySetup.Get()->BoneName.ToString());
+					//make sure you have physics weight or blendphysics on, otherwise, you'll have inconsistent representation of bodies
+					// @todo make this to be kismet log? But can be too intrusive
+					if (!bBlendPhysics && BodyInst->PhysicsBlendWeight <= 0.f && BodyInst->BodySetup.IsValid())
+					{
+						UE_LOG(LogPhysics, Warning, TEXT("%s(Mesh %s, PhysicsAsset %s, Bone %s) is simulating, but no blending. "),
+							*GetName(), *GetNameSafe(SkeletalMesh), *GetNameSafe(PhysicsAsset), *BodyInst->BodySetup.Get()->BoneName.ToString());
+					}
 				}
 			}
 		}
 	}
+	else
+	{
+		//per poly update requires us to update all vertex positions
+	}
+
 
 }
 
diff --git a/Engine/Source/Runtime/Engine/Private/SkeletalMesh.cpp b/Engine/Source/Runtime/Engine/Private/SkeletalMesh.cpp
--- a/Engine/Source/Runtime/Engine/Private/SkeletalMesh.cpp
+++ b/Engine/Source/Runtime/Engine/Private/SkeletalMesh.cpp
@@ -4583,6 +4583,11 @@ void USkeletalMeshComponent::Serialize(FArchive& Ar)
 		}
 	}
 
+	if (bEnablePerPolyCollision)
+	{
+		Ar << BodySetup;
+	}
+
 	// Since we separated simulation vs blending
 	// if simulation is on when loaded, just set blendphysics to be true
 	if (BodyInstance.bSimulatePhysics)
diff --git a/Engine/Source/Runtime/Engine/Private/SkeletalMeshComponentPhysics.cpp b/Engine/Source/Runtime/Engine/Private/SkeletalMeshComponentPhysics.cpp
--- a/Engine/Source/Runtime/Engine/Private/SkeletalMeshComponentPhysics.cpp
+++ b/Engine/Source/Runtime/Engine/Private/SkeletalMeshComponentPhysics.cpp
@@ -334,23 +334,133 @@ int32 USkeletalMesh::GetClothAssetIndex(int32 SectionIndex)
 }
 #endif//#if WITH_APEX_CLOTHING
 
+void USkeletalMeshComponent::CreateBodySetup()
+{
+	if (BodySetup == NULL)
+	{
+		BodySetup = ConstructObject<UBodySetup>(UBodySetup::StaticClass(), this);
+		BodySetup->CollisionTraceFlag = CTF_UseComplexAsSimple;
+	}
+}
+
+bool USkeletalMeshComponent::GetPhysicsTriMeshData(FTriMeshCollisionData* CollisionData, bool bInUseAllTriData)
+{
+#if WITH_EDITORONLY_DATA
+
+	// Fail if no mesh or not per poly collision
+	if (!SkeletalMesh || !MeshObject || !bEnablePerPolyCollision)
+	{
+		return false;
+	}
+
+	const FStaticLODModel& Model = MeshObject->GetSkeletalMeshResource().LODModels[0];
+
+	{
+		// Copy all verts into collision vertex buffer.
+		CollisionData->Vertices.Empty();
+		CollisionData->Vertices.AddUninitialized(Model.NumVertices);
+		const uint32 NumChunks = Model.Chunks.Num();
+
+		for (uint32 ChunkIdx = 0; ChunkIdx < NumChunks; ++ChunkIdx)
+		{
+			const FSkelMeshChunk & Chunk = Model.Chunks[ChunkIdx];
+			{
+				//rigid
+				const uint32 RigidOffset = Chunk.GetRigidVertexBufferIndex();
+				const uint32 NumRigidVerts = Chunk.GetNumRigidVertices();
+				for (uint32 RigidIdx = 0; RigidIdx < NumRigidVerts; ++RigidIdx)
+				{
+					CollisionData->Vertices[RigidIdx + RigidOffset] = Chunk.RigidVertices[RigidIdx].Position;
+				}
+			}
+			{
+				//soft
+				const uint32 SoftOffset = Chunk.GetSoftVertexBufferIndex();
+				const uint32 NumSoftVerts = Chunk.GetNumSoftVertices();
+				for (uint32 SoftIdx = 0; SoftIdx < NumSoftVerts; ++SoftIdx)
+				{
+					CollisionData->Vertices[SoftIdx + SoftOffset] = Chunk.SoftVertices[SoftIdx].Position;
+				}
+			}
+			
+		}
+	}
+	
+	{
+		// Copy indices into collision index buffer
+		const FMultiSizeIndexContainer & IndexBufferContainer = Model.MultiSizeIndexContainer;
+
+		TArray<uint32> Indices;
+		IndexBufferContainer.GetIndexBuffer(Indices);
+
+		const uint32 NumTris = Indices.Num() / 3;
+		CollisionData->Indices.Empty();
+		CollisionData->Indices.Reserve(NumTris);
+
+		FTriIndices TriIndex;
+		for (int32 SectionIndex = 0; SectionIndex < Model.Sections.Num(); ++SectionIndex)
+		{
+			const FSkelMeshSection& Section = Model.Sections[SectionIndex];
+
+			const uint32 OnePastLastIndex = Section.BaseIndex + Section.NumTriangles * 3;
+
+			for (uint32 i = Section.BaseIndex; i < OnePastLastIndex; i += 3)
+			{
+				TriIndex.v0 = Indices[i];
+				TriIndex.v1 = Indices[i + 1];
+				TriIndex.v2 = Indices[i + 2];
+
+				CollisionData->Indices.Add(TriIndex);
+				CollisionData->MaterialIndices.Add(Section.MaterialIndex);
+			}
+		}
+	}
+
+	CollisionData->bFlipNormals = true;
+
+	// We only have a valid TriMesh if the CollisionData has vertices AND indices. For meshes with disabled section collision, it
+	// can happen that the indices will be empty, in which case we do not want to consider that as valid trimesh data
+	return CollisionData->Vertices.Num() > 0 && CollisionData->Indices.Num() > 0;
+#else // #if WITH_EDITORONLY_DATA
+	return false;
+#endif // #if WITH_EDITORONLY_DATA
+}
+
+bool USkeletalMeshComponent::ContainsPhysicsTriMeshData(bool InUseAllTriData) const
+{
+	return bEnablePerPolyCollision;
+}
+
 //
 //	USkeletalMeshComponent
 //
 UBodySetup* USkeletalMeshComponent::GetBodySetup()
 {
-	UPhysicsAsset * const PhysicsAsset = GetPhysicsAsset();
-	if(SkeletalMesh && PhysicsAsset)
+	if (bEnablePerPolyCollision == false)
 	{
-		for(int32 i=0; i<SkeletalMesh->RefSkeleton.GetNum(); i++)
+		UPhysicsAsset * const PhysicsAsset = GetPhysicsAsset();
+		if (SkeletalMesh && PhysicsAsset)
 		{
-			int32 BodyIndex = PhysicsAsset->FindBodyIndex( SkeletalMesh->RefSkeleton.GetBoneName(i) );
-			if(BodyIndex != INDEX_NONE)
+			for (int32 i = 0; i < SkeletalMesh->RefSkeleton.GetNum(); i++)
 			{
-				return PhysicsAsset->BodySetup[BodyIndex];
+				int32 BodyIndex = PhysicsAsset->FindBodyIndex(SkeletalMesh->RefSkeleton.GetBoneName(i));
+				if (BodyIndex != INDEX_NONE)
+				{
+					return PhysicsAsset->BodySetup[BodyIndex];
+				}
 			}
 		}
-	}	
+	}
+	else
+	{
+		if (BodySetup == NULL)
+		{
+			CreateBodySetup();
+		}
+
+		return BodySetup;
+	}
+
 
 	return NULL;
 }
@@ -907,7 +1017,7 @@ void USkeletalMeshComponent::SetAllBodiesBelowSimulatePhysics( const FName& InBo
 	for(int32 i=0; i<BodyIndices.Num(); i++)
 	{
 		//UE_LOG(LogSkeletalMesh, Warning, TEXT( "ForceAllBodiesBelowUnfixed %s" ), *InAsset->BodySetup(BodyIndices(i))->BoneName.ToString() );
-		Bodies[BodyIndices[i]]->SetInstanceSimulatePhysics( bNewSimulate );
+		Bodies[BodyIndices[i]]->SetInstanceSimulatePhysics(bNewSimulate);
 	}
 }
 
@@ -1197,15 +1307,27 @@ void USkeletalMeshComponent::CreatePhysicsState()
 {
 	//	UE_LOG(LogSkeletalMesh, Warning, TEXT("Creating Physics State (%s : %s)"), *GetNameSafe(GetOuter()),  *GetName());
 	// Init physics
-	InitArticulated(World->GetPhysicsScene());
+	if (bEnablePerPolyCollision == false)
+	{
+		InitArticulated(World->GetPhysicsScene());
+		USceneComponent::CreatePhysicsState(); // Need to route CreatePhysicsState, skip PrimitiveComponent
+	}
+	else
+	{
+		CreateBodySetup();
+		Super::CreatePhysicsState();	//If we're doing per poly we'll use the body instance of the primitive component
+	}
 
-	USceneComponent::CreatePhysicsState(); // Need to route CreatePhysicsState, skip PrimitiveComponent
+	
 }
 
 
 void USkeletalMeshComponent::DestroyPhysicsState()
 {
-	TermArticulated();
+	if (bEnablePerPolyCollision == false)
+	{
+		TermArticulated();
+	}
 
 	Super::DestroyPhysicsState();
 }
@@ -2977,6 +3099,8 @@ void USkeletalMeshComponent::PreClothTick(float DeltaTime)
 		BlendInPhysics();
 	}
 
+	//TODO: move this into pre physics tick
+
 #if WITH_APEX_CLOTHING
 	// if skeletal mesh has clothing assets, call TickClothing
 	if (SkeletalMesh && SkeletalMesh->ClothingAssets.Num() > 0)
