Make cached AxisProperties map per player, rather than static

diff --git a/Engine/Source/Runtime/Engine/Classes/GameFramework/PlayerInput.h b/Engine/Source/Runtime/Engine/Classes/GameFramework/PlayerInput.h
--- a/Engine/Source/Runtime/Engine/Classes/GameFramework/PlayerInput.h
+++ b/Engine/Source/Runtime/Engine/Classes/GameFramework/PlayerInput.h
@@ -313,8 +313,7 @@ public:
 
 protected:
 	/** Internal structure for storing axis config data. */
-	static TMap<FKey,FInputAxisProperties> AxisProperties;
-
+	TMap<FKey,FInputAxisProperties> AxisProperties;
 	TMap<FName,FActionKeyDetails> ActionKeyMap;
 	TMap<FName,FAxisKeyDetails> AxisKeyMap;
 
diff --git a/Engine/Source/Runtime/Engine/Private/UserInterface/PlayerInput.cpp b/Engine/Source/Runtime/Engine/Private/UserInterface/PlayerInput.cpp
--- a/Engine/Source/Runtime/Engine/Private/UserInterface/PlayerInput.cpp
+++ b/Engine/Source/Runtime/Engine/Private/UserInterface/PlayerInput.cpp
@@ -21,8 +21,6 @@ const TArray<FInputActionKeyMapping> UPlayerInput::NoKeyMappings;
 TArray<FInputActionKeyMapping> UPlayerInput::EngineDefinedActionMappings;
 TArray<FInputAxisKeyMapping> UPlayerInput::EngineDefinedAxisMappings;
 
-TMap<FKey,FInputAxisProperties> UPlayerInput::AxisProperties;
-
 UPlayerInput::UPlayerInput(const class FPostConstructInitializeProperties& PCIP)
 	: Super(PCIP)
 {
