#UE4 Fix AI crash during seamless travel

diff --git a/Engine/Source/Runtime/Engine/Private/World.cpp b/Engine/Source/Runtime/Engine/Private/World.cpp
--- a/Engine/Source/Runtime/Engine/Private/World.cpp
+++ b/Engine/Source/Runtime/Engine/Private/World.cpp
@@ -4493,6 +4493,9 @@ UWorld* FSeamlessTravelHandler::Tick()
 				if (GameMode)
 				{
 					UNavigationSystem::InitializeForWorld(LoadedWorld, FNavigationSystem::GameMode);
+					
+					// Note that AI system will be created only if ai-system-creation conditions are met
+					LoadedWorld->CreateAISystem();
 
 					// inform the new GameMode so it can handle players that persisted
 					GameMode->PostSeamlessTravel();					
