Adding checks to make sure object it not NULL

diff --git a/Engine/Source/Runtime/Engine/Private/DataReplication.cpp b/Engine/Source/Runtime/Engine/Private/DataReplication.cpp
--- a/Engine/Source/Runtime/Engine/Private/DataReplication.cpp
+++ b/Engine/Source/Runtime/Engine/Private/DataReplication.cpp
@@ -156,6 +156,14 @@ void FObjectReplicator::InitWithObject( UObject * InObject, UNetConnection * InC
 
 	SetObject( InObject );
 
+	if ( GetObject() == NULL )
+	{
+		// This may seem weird that we're checking for NULL, but the SetObject above will wrap this object with TWeakObjectPtr
+		// If the object is pending kill, it will switch to NULL, we're just making sure we handle this invalid edge case
+		UE_LOG( LogNet, Error, TEXT( "InitWithObject: Object == NULL" ) );
+		return;
+	}
+
 	ObjectClass					= InObject->GetClass();
 	Connection					= InConnection;
 	RemoteFunctions				= NULL;
@@ -205,6 +213,12 @@ void FObjectReplicator::StartReplicating( class UActorChannel * InActorChannel )
 {
 	check( OwningChannel == NULL );
 
+	if ( GetObject() == NULL )
+	{
+		UE_LOG( LogNet, Error, TEXT( "StartReplicating: Object == NULL" ) );
+		return;
+	}
+
 	OwningChannel = InActorChannel;
 
 	// Cache off netGUID so if this object gets deleted we can close it
@@ -293,6 +307,12 @@ void FObjectReplicator::ReceivedNak( int32 NakPacketId )
 {
 	UObject * Object = GetObject();
 
+	if ( Object == NULL )
+	{
+		UE_LOG( LogNet, Error, TEXT( "ReceivedNak: Object == NULL" ) );
+		return;
+	}
+
 	if ( Object != NULL && ObjectClass != NULL )
 	{
 		RepLayout->ReceivedNak( RepState, NakPacketId );
@@ -350,8 +370,7 @@ void FObjectReplicator::ReceivedNak( int32 NakPacketId )
 
 bool FObjectReplicator::ReceivedBunch( FInBunch &Bunch, const FReplicationFlags & RepFlags, bool & bOutHasUnmapped )
 {
-	UObject *		Object		= GetObject();
-	UPackageMap *	PackageMap	= OwningChannel->Connection->PackageMap;
+	UObject * Object = GetObject();
 
 	if ( Object == NULL )
 	{
@@ -359,6 +378,8 @@ bool FObjectReplicator::ReceivedBunch( FInBunch &Bunch, const FReplicationFlags
 		return false;
 	}
 
+	UPackageMap * PackageMap = OwningChannel->Connection->PackageMap;
+
 	const bool bIsServer = ( OwningChannel->Connection->Driver->ServerConnection == NULL );
 
 	FClassNetCache * ClassCache = OwningChannel->Connection->Driver->NetCache->GetClassNetCache( ObjectClass );
@@ -926,6 +947,12 @@ bool FObjectReplicator::ReplicateProperties( FOutBunch & Bunch, FReplicationFlag
 {
 	UObject * Object = GetObject();
 
+	if ( Object == NULL )
+	{
+		UE_LOG( LogNet, Error, TEXT( "ReplicateProperties: Object == NULL" ) );
+		return false;
+	}
+
 	check( Object );
 	check( OwningChannel );
 	check( RepLayout.IsValid() );
