UE4 - fix some nonunity errors

diff --git a/Engine/Source/Runtime/Engine/Private/MeshParticleVertexFactory.cpp b/Engine/Source/Runtime/Engine/Private/MeshParticleVertexFactory.cpp
--- a/Engine/Source/Runtime/Engine/Private/MeshParticleVertexFactory.cpp
+++ b/Engine/Source/Runtime/Engine/Private/MeshParticleVertexFactory.cpp
@@ -7,6 +7,7 @@
 #include "EnginePrivate.h"
 #include "ParticleDefinitions.h"
 #include "MeshParticleVertexFactory.h"
+#include "ShaderParameterUtils.h"
 
 
 class FMeshParticleVertexFactoryShaderParameters : public FVertexFactoryShaderParameters
diff --git a/Engine/Source/Runtime/Engine/Public/ScreenRendering.h b/Engine/Source/Runtime/Engine/Public/ScreenRendering.h
--- a/Engine/Source/Runtime/Engine/Public/ScreenRendering.h
+++ b/Engine/Source/Runtime/Engine/Public/ScreenRendering.h
@@ -9,6 +9,7 @@
 #include "Shader.h"
 #include "ShaderParameters.h"
 #include "GlobalShader.h"
+#include "ShaderParameterUtils.h"
 
 struct FScreenVertex
 {
