UE4 - unity build compile fix

diff --git a/Engine/Source/Runtime/Engine/Private/ScreenRendering.cpp b/Engine/Source/Runtime/Engine/Private/ScreenRendering.cpp
--- a/Engine/Source/Runtime/Engine/Private/ScreenRendering.cpp
+++ b/Engine/Source/Runtime/Engine/Private/ScreenRendering.cpp
@@ -6,6 +6,7 @@
 
 #include "EnginePrivate.h"
 #include "ScreenRendering.h"
+#include "RHIStaticStates.h"
 
 /** Vertex declaration for screen-space rendering. */
 TGlobalResource<FScreenVertexDeclaration> GScreenVertexDeclaration;
