Moving the child actor component now properly moves the spawned child actor

diff --git a/Engine/Source/Runtime/Engine/Private/ActorConstruction.cpp b/Engine/Source/Runtime/Engine/Private/ActorConstruction.cpp
--- a/Engine/Source/Runtime/Engine/Private/ActorConstruction.cpp
+++ b/Engine/Source/Runtime/Engine/Private/ActorConstruction.cpp
@@ -147,6 +147,7 @@ void AActor::RerunConstructionScripts()
 		FTransform OldTransform = FTransform::Identity;
 		FName  SocketName;
 		AActor* Parent = NULL;
+		USceneComponent* ParentComponent = NULL;
 
 		// Struct to store info about attached actors
 		struct FAttachedActorInfo
@@ -232,7 +233,7 @@ void AActor::RerunConstructionScripts()
 						UE_LOG(LogActor, Warning, TEXT("RerunConstructionScripts: RootComponent (%s) attached to another component in this Actor (%s)."), *RootComponent->GetPathName(), *Parent->GetPathName());
 						Parent = NULL;
 					}
-
+					ParentComponent = RootComponent->AttachParent;
 					SocketName = RootComponent->AttachSocketName;
 					//detach it to remove any scaling 
 					RootComponent->DetachFromParent(true);
@@ -259,11 +260,14 @@ void AActor::RerunConstructionScripts()
 
 		if(Parent)
 		{
-			USceneComponent* ChildRoot =	GetRootComponent();
-			USceneComponent* ParentRoot =	Parent->GetRootComponent();
-			if(ChildRoot != NULL && ParentRoot != NULL)
+			USceneComponent* ChildRoot = GetRootComponent();
+			if (ParentComponent == NULL)
+			{
+				ParentComponent = Parent->GetRootComponent();
+			}
+			if (ChildRoot != NULL && ParentComponent != NULL)
 			{
-				ChildRoot->AttachTo( ParentRoot, SocketName, EAttachLocation::KeepWorldPosition );
+				ChildRoot->AttachTo(ParentComponent, SocketName, EAttachLocation::KeepWorldPosition);
 			}
 		}
 
