#ue4 - Expose control over max simulation iterations and max simulation time step for CharacterMovementComponent.

- Added a helper function GetSimulationTimeStep(float RemainingTime, int32 Iterations) that computes the timestep given the total remaining time and current iteration.
- During the final iteration, if there is not enough time within the time step limit, all remaining time is used. This is better than cutting the simulation short (which used to happen). A warning is printed in this case.

#TTP 317394

diff --git a/Engine/Source/Runtime/Engine/Classes/GameFramework/CharacterMovementComponent.h b/Engine/Source/Runtime/Engine/Classes/GameFramework/CharacterMovementComponent.h
--- a/Engine/Source/Runtime/Engine/Classes/GameFramework/CharacterMovementComponent.h
+++ b/Engine/Source/Runtime/Engine/Classes/GameFramework/CharacterMovementComponent.h
@@ -534,6 +534,36 @@ protected:
 
 public:
 
+	/**
+	 * Compute remaining time step given remaining time and current iterations.
+	 * The last iteration (limited by MaxSimulationIterations) always returns the remaining time, which may violate MaxSimulationTimeStep.
+	 *
+	 * @see MaxSimulationTimeStep
+	 * @see MaxSimulationIterations
+	 */
+	float GetSimulationTimeStep(float RemainingTime, int32 Iterations) const;
+
+	/**
+	 * Max time delta for each discrete simulation step.
+	 * Used primarily in the the more advanced movement modes that break up larger time steps (usually those applying gravity such as falling and walking).
+	 * Lowering this value can address issues with fast-moving objects or complex collision scenarios, at the cost of performance.
+	 *
+	 * WARNING: if (MaxSimulationTimeStep * MaxSimulationIterations) is too low for the min framerate, the last simulation step may exceed MaxSimulationTimeStep to complete the simulation.
+	 * @see MaxSimulationIterations
+	 */
+	UPROPERTY(Category="Character Movement", EditAnywhere, BlueprintReadWrite, AdvancedDisplay, meta=(ClampMin="0.02", ClampMax="0.50", UIMin="0.02", UIMax="0.50"))
+	float MaxSimulationTimeStep;
+
+	/**
+	 * Max number of iterations used for each discrete simulation step
+	 * Used primarily in the the more advanced movement modes that break up larger time steps (usually those applying gravity such as falling and walking).
+	 * Increasing this value can address issues with fast-moving objects or complex collision scenarios, at the cost of performance.
+	 *
+	 * WARNING: if (MaxSimulationTimeStep * MaxSimulationIterations) is too low for the min framerate, the last simulation step may exceed MaxSimulationTimeStep to complete the simulation.
+	 */
+	UPROPERTY(Category="Character Movement", EditAnywhere, BlueprintReadWrite, AdvancedDisplay, meta=(ClampMin="1", ClampMax="20", UIMin="1", UIMax="20"))
+	int32 MaxSimulationIterations;
+
 	/** Max Acceleration (rate of change of velocity) */
 	UPROPERTY(Category="Character Movement", EditAnywhere, BlueprintReadWrite)
 	float MaxAcceleration;
@@ -1458,6 +1488,9 @@ protected:
 
 public:
 
+	/** Minimum delta time considered when ticking. Delta times below this are not considered. This is a very small non-zero value to avoid potential divide-by-zero in simulation code. */
+	static const float MIN_TICK_TIME;
+
 	/** Minimum acceptable distance for Character capsule to float above floor when walking. */
 	static const float MIN_FLOOR_DIST;
 
diff --git a/Engine/Source/Runtime/Engine/Private/CharacterMovementComponent.cpp b/Engine/Source/Runtime/Engine/Private/CharacterMovementComponent.cpp
--- a/Engine/Source/Runtime/Engine/Private/CharacterMovementComponent.cpp
+++ b/Engine/Source/Runtime/Engine/Private/CharacterMovementComponent.cpp
@@ -18,8 +18,8 @@ DEFINE_LOG_CATEGORY_STATIC(LogCharacterMovement, Log, All);
 // MAGIC NUMBERS
 const float MAX_STEP_SIDE_Z = 0.08f;	// maximum z value for the normal on the vertical side of steps
 const float SWIMBOBSPEED = -80.f;
-const float MIN_TICK_TIME = 0.0002f;	// minimum delta time considered when ticking.
 
+const float UCharacterMovementComponent::MIN_TICK_TIME = 0.0002f;
 const float UCharacterMovementComponent::MIN_FLOOR_DIST = 1.9f;
 const float UCharacterMovementComponent::MAX_FLOOR_DIST = 2.4f;
 const float UCharacterMovementComponent::BRAKE_TO_STOP_VELOCITY = 10.f;
@@ -130,6 +130,9 @@ UCharacterMovementComponent::UCharacterMovementComponent(const class FPostConstr
 	MaxWalkSpeed = 600.0f;
 	MaxSwimSpeed = 300.0f;
 	
+	MaxSimulationTimeStep = 0.05f;
+	MaxSimulationIterations = 8;
+
 	CrouchedSpeedMultiplier = 0.5f;
 	MaxOutOfWaterStepHeight = 40.0f;
 	OutofWaterZ = 420.0f;
@@ -1587,7 +1590,7 @@ void UCharacterMovementComponent::UnCrouch(bool bClientSimulation)
 
 void UCharacterMovementComponent::StartNewPhysics(float deltaTime, int32 Iterations)
 {
-	if ( (deltaTime < MIN_TICK_TIME) || (Iterations > 7) || !HasValidData())
+	if ( (deltaTime < MIN_TICK_TIME) || (Iterations >= MaxSimulationIterations) || !HasValidData())
 		return;
 
 	if (UpdatedComponent->IsSimulatingPhysics())
@@ -2391,7 +2394,7 @@ void UCharacterMovementComponent::StartSwimming(FVector OldLocation, FVector Old
 	{
 		Velocity.Z = SWIMBOBSPEED - Velocity.Size2D() * 0.7f; //smooth bobbing
 	}
-	if ( (remainingTime >= MIN_TICK_TIME) && (Iterations < 8) && CharacterOwner )
+	if ( (remainingTime >= MIN_TICK_TIME) && (Iterations < MaxSimulationIterations) && CharacterOwner )
 	{
 		PhysSwimming(remainingTime, Iterations);
 	}
@@ -2530,12 +2533,10 @@ void UCharacterMovementComponent::PhysFalling(float deltaTime, int32 Iterations)
 	float remainingTime = deltaTime;
 	float timeTick = 0.1f;
 
-	while( (remainingTime >= MIN_TICK_TIME) && (Iterations < 8) )
+	while( (remainingTime >= MIN_TICK_TIME) && (Iterations < MaxSimulationIterations) )
 	{
 		Iterations++;
-		
-		// subdivide moves to be no longer than 0.05 seconds, no less than MIN_TICK_TIME.
-		const float timeTick = FMath::Max(MIN_TICK_TIME, (remainingTime > 0.05f) ? FMath::Min(0.05f, remainingTime * 0.5f) : remainingTime);
+		const float timeTick = GetSimulationTimeStep(remainingTime, Iterations);
 		remainingTime -= timeTick;
 		
 		const FVector OldLocation = CharacterOwner->GetActorLocation();
@@ -3016,11 +3017,10 @@ void UCharacterMovementComponent::PhysWalking(float deltaTime, int32 Iterations)
 	bool bTriedLedgeMove = false;
 	float remainingTime = deltaTime;
 
-	while ( (remainingTime >= MIN_TICK_TIME) && (Iterations < 8) && (CharacterOwner->Controller || bRunPhysicsWithNoController || HasRootMotion()) )
+	while ( (remainingTime >= MIN_TICK_TIME) && (Iterations < MaxSimulationIterations) && (CharacterOwner->Controller || bRunPhysicsWithNoController || HasRootMotion()) )
 	{
 		Iterations++;
-		// subdivide moves to be no longer than 0.05 seconds, no less than MIN_TICK_TIME.
-		const float timeTick = FMath::Max(MIN_TICK_TIME, (remainingTime > 0.05f) ? FMath::Min(0.05f, remainingTime * 0.5f) : remainingTime);
+		const float timeTick = GetSimulationTimeStep(remainingTime, Iterations);
 		remainingTime -= timeTick;
 		const FVector Delta = timeTick * DesiredMove;
 		const FVector subLoc = CharacterOwner->GetActorLocation();
@@ -4504,6 +4504,34 @@ float UCharacterMovementComponent::GetAnalogInputModifier() const
 }
 
 
+float UCharacterMovementComponent::GetSimulationTimeStep(float RemainingTime, int32 Iterations) const
+{
+	if (RemainingTime > MaxSimulationTimeStep)
+	{
+		if (Iterations < MaxSimulationIterations)
+		{
+			// Subdivide moves to be no longer than MaxSimulationTimeStep seconds
+			RemainingTime = FMath::Min(MaxSimulationTimeStep, RemainingTime * 0.5f);
+		}
+		else
+		{
+			// If this is the last iteration, just use all the remaining time. This is usually better than cutting things short, as the simulation won't move far enough otherwise.
+			// Print a throttled warning.
+#if !(UE_BUILD_SHIPPING || UE_BUILD_TEST)
+			static uint32 s_WarningCount = 0;
+			if ((s_WarningCount++ < 100) || (GFrameCounter & 15) == 0)
+			{
+				UE_LOG(LogCharacterMovement, Warning, TEXT("GetSimulationTimeStep() - Max iterations %d hit while remaining time %.6f > MaxSimulationTimeStep (%.3f) for '%s'"), MaxSimulationIterations, RemainingTime, MaxSimulationTimeStep, *GetNameSafe(CharacterOwner));
+			}
+#endif
+		}
+	}
+
+	// no less than MIN_TICK_TIME (to avoid potential divide-by-zero during simulation).
+	return FMath::Max(MIN_TICK_TIME, RemainingTime);
+}
+
+
 void UCharacterMovementComponent::SmoothCorrection(const FVector& OldLocation)
 {
 	if (!HasValidData() || GetNetMode() != NM_Client)
