Fixed AnimScale not being applied in camera shakes.

diff --git a/Engine/Source/Runtime/Engine/Private/Camera/CameraModifier_CameraShake.cpp b/Engine/Source/Runtime/Engine/Private/Camera/CameraModifier_CameraShake.cpp
--- a/Engine/Source/Runtime/Engine/Private/Camera/CameraModifier_CameraShake.cpp
+++ b/Engine/Source/Runtime/Engine/Private/Camera/CameraModifier_CameraShake.cpp
@@ -262,13 +262,15 @@ void UCameraModifier_CameraShake::ReinitShake(FCameraShakeInstance& ShakeInst, f
 			Duration = SourceShake->RandomAnimSegmentDuration;
 		}
 
+		float AnimScale = NewScale * SourceShake->AnimScale;
+
 		if (ShakeInst.AnimInst != nullptr)
 		{
-			ShakeInst.AnimInst->Update(SourceShake->AnimPlayRate, NewScale, SourceShake->AnimBlendInTime, SourceShake->AnimBlendOutTime, Duration);
+			ShakeInst.AnimInst->Update(SourceShake->AnimPlayRate, AnimScale, SourceShake->AnimBlendInTime, SourceShake->AnimBlendOutTime, Duration);
 		}
 		else
 		{
-			ShakeInst.AnimInst = CameraOwner->PlayCameraAnim(SourceShake->Anim, SourceShake->AnimPlayRate, NewScale, SourceShake->AnimBlendInTime, SourceShake->AnimBlendOutTime, bLoop, bRandomStart, Duration);
+			ShakeInst.AnimInst = CameraOwner->PlayCameraAnim(SourceShake->Anim, SourceShake->AnimPlayRate, AnimScale, SourceShake->AnimBlendInTime, SourceShake->AnimBlendOutTime, bLoop, bRandomStart, Duration);
 		}
 	}
 }
@@ -323,9 +325,11 @@ FCameraShakeInstance UCameraModifier_CameraShake::InitializeShake(TSubclassOf<cl
 			Duration = SourceShake->RandomAnimSegmentDuration;
 		}
 
-		if (Scale > 0.f)
+		float AnimScale = Scale * SourceShake->AnimScale;
+		
+		if (AnimScale > 0.f)
 		{
-			Inst.AnimInst = CameraOwner->PlayCameraAnim(SourceShake->Anim, SourceShake->AnimPlayRate, Scale, SourceShake->AnimBlendInTime, SourceShake->AnimBlendOutTime, bLoop, bRandomStart, Duration, PlaySpace, UserPlaySpaceRot);
+			Inst.AnimInst = CameraOwner->PlayCameraAnim(SourceShake->Anim, SourceShake->AnimPlayRate, AnimScale, SourceShake->AnimBlendInTime, SourceShake->AnimBlendOutTime, bLoop, bRandomStart, Duration, PlaySpace, UserPlaySpaceRot);
 		}
 	}
 
