Fixed CameraAnim scaling not being properly applied to FOV

diff --git a/Engine/Source/Runtime/Engine/Private/PlayerCameraManager.cpp b/Engine/Source/Runtime/Engine/Private/PlayerCameraManager.cpp
--- a/Engine/Source/Runtime/Engine/Private/PlayerCameraManager.cpp
+++ b/Engine/Source/Runtime/Engine/Private/PlayerCameraManager.cpp
@@ -325,7 +325,7 @@ void APlayerCameraManager::ApplyAnimToCamera(ACameraActor const* AnimatedCamActo
 	// fov
 	const float FOVMin = 5.f;
 	const float FOVMax = 170.f;
-	InOutPOV.FOV += AnimatedCamActor->CameraComponent->FieldOfView - AnimInst->InitialFOV;
+	InOutPOV.FOV += (AnimatedCamActor->CameraComponent->FieldOfView - AnimInst->InitialFOV) * Scale;
 	InOutPOV.FOV = FMath::Clamp<float>(InOutPOV.FOV, FOVMin, FOVMax);
 
 	// postprocess
