#ue4 - Make UCharacterMovementComponent::SimulateMovement() work for characters that are not simulated proxies.

diff --git a/Engine/Source/Runtime/Engine/Private/CharacterMovementComponent.cpp b/Engine/Source/Runtime/Engine/Private/CharacterMovementComponent.cpp
--- a/Engine/Source/Runtime/Engine/Private/CharacterMovementComponent.cpp
+++ b/Engine/Source/Runtime/Engine/Private/CharacterMovementComponent.cpp
@@ -959,8 +959,11 @@ void UCharacterMovementComponent::SimulateMovement(float DeltaSeconds)
 		return;
 	}
 
+	const bool bIsSimulatedProxy = (CharacterOwner->Role == ROLE_SimulatedProxy);
+
 	// Workaround for replication not being updated initially
-	if (CharacterOwner->ReplicatedMovement.Location.IsZero() &&
+	if (bIsSimulatedProxy &&
+		CharacterOwner->ReplicatedMovement.Location.IsZero() &&
 		CharacterOwner->ReplicatedMovement.Rotation.IsZero() &&
 		CharacterOwner->ReplicatedMovement.LinearVelocity.IsZero())
 	{
@@ -976,24 +979,27 @@ void UCharacterMovementComponent::SimulateMovement(float DeltaSeconds)
 
 	FScopedMovementUpdate ScopedMovementUpdate(UpdatedComponent, bEnableScopedMovementUpdates ? EScopedUpdate::DeferredUpdates : EScopedUpdate::ImmediateUpdates);
 
-	// Handle network changes
-	if (bNetworkUpdateReceived)
+	if (bIsSimulatedProxy)
 	{
-		bNetworkUpdateReceived = false;
-		if (bNetworkMovementModeChanged)
+		// Handle network changes
+		if (bNetworkUpdateReceived)
 		{
-			bNetworkMovementModeChanged = false;
-			ApplyNetworkMovementMode(CharacterOwner->GetReplicatedMovementMode());
-		}
-		else if (bJustTeleported)
-		{
-			// Make sure floor is current. We will continue using the replicated base, if there was one.
-			bJustTeleported = false;
-			UpdateFloorFromAdjustment();
+			bNetworkUpdateReceived = false;
+			if (bNetworkMovementModeChanged)
+			{
+				bNetworkMovementModeChanged = false;
+				ApplyNetworkMovementMode(CharacterOwner->GetReplicatedMovementMode());
+			}
+			else if (bJustTeleported)
+			{
+				// Make sure floor is current. We will continue using the replicated base, if there was one.
+				bJustTeleported = false;
+				UpdateFloorFromAdjustment();
+			}
 		}
-	}
 
-	HandlePendingLaunch();
+		HandlePendingLaunch();
+	}
 
 	if (MovementMode == MOVE_None)
 	{
@@ -1009,7 +1015,8 @@ void UCharacterMovementComponent::SimulateMovement(float DeltaSeconds)
 	MoveSmooth(Velocity, DeltaSeconds, &StepDownResult);
 
 	// if simulated gravity, find floor and check if falling
-	if (!CharacterOwner->bSimGravityDisabled && (MovementMode == MOVE_Walking || MovementMode == MOVE_Falling))
+	const bool bEnableFloorCheck = (!CharacterOwner->bSimGravityDisabled || !bIsSimulatedProxy);
+	if (bEnableFloorCheck && (MovementMode == MOVE_Walking || MovementMode == MOVE_Falling))
 	{
 		const FVector CollisionCenter = UpdatedComponent->GetComponentLocation();
 		if (StepDownResult.bComputedFloor)
