Fixed FQuat::IsNormalized() not using same threshold as FTransformVectorized::IsRotationNormalized()
Also aligned normalization threshold value w/ PhysX, so we speak the same language.
Normalize LocalAtoms and SpaceBases transforms before they leave the internals of the animation system. Repetitive concatenation introduces errors (long bone chains transformed from local to world space). We don't want to pass that to PhysX, gameplay, etc.
Made vector and quaternion normalization thresholds explicit defines, so we can better manage them and make sure different code paths don't get out of sync.

diff --git a/Engine/Source/Runtime/Core/Public/Math/Quat.h b/Engine/Source/Runtime/Core/Public/Math/Quat.h
--- a/Engine/Source/Runtime/Core/Public/Math/Quat.h
+++ b/Engine/Source/Runtime/Core/Public/Math/Quat.h
@@ -809,7 +809,7 @@ FORCEINLINE void FQuat::Normalize(float Tolerance)
 
 FORCEINLINE bool FQuat::IsNormalized( ) const
 {
-	return (FMath::Abs(1.f - SizeSquared()) <= KINDA_SMALL_NUMBER);
+	return (FMath::Abs(1.f - SizeSquared()) < THRESH_QUAT_NORMALIZED);
 }
 
 
diff --git a/Engine/Source/Runtime/Core/Public/Math/TransformVectorized.h b/Engine/Source/Runtime/Core/Public/Math/TransformVectorized.h
--- a/Engine/Source/Runtime/Core/Public/Math/TransformVectorized.h
+++ b/Engine/Source/Runtime/Core/Public/Math/TransformVectorized.h
@@ -927,10 +927,9 @@ public:
 	FORCEINLINE bool IsRotationNormalized() const
 	{		
 		const VectorRegister TestValue = VectorAbs(VectorSubtract(VectorOne(), VectorDot4(Rotation, Rotation)));
-		return !VectorAnyGreaterThan(TestValue, GlobalVectorConstants::FloatOneHundredth);	
+		return !VectorAnyGreaterThan(TestValue, GlobalVectorConstants::ThreshQuatNormalized);
 	}
 
-
 	/**
 	 * Blends the Identity atom with a weighted source atom and accumulates that into a destination atom
 	 *
diff --git a/Engine/Source/Runtime/Core/Public/Math/UnrealMathUtility.h b/Engine/Source/Runtime/Core/Public/Math/UnrealMathUtility.h
--- a/Engine/Source/Runtime/Core/Public/Math/UnrealMathUtility.h
+++ b/Engine/Source/Runtime/Core/Public/Math/UnrealMathUtility.h
@@ -74,6 +74,9 @@ struct FLinearColor;
 #define THRESH_ZERO_NORM_SQUARED		(0.0001f)	/* Size of a unit normal that is considered "zero", squared */
 #define THRESH_VECTORS_ARE_PARALLEL		(0.02f)		/* Vectors are parallel if dot product varies less than this */
 
+#define THRESH_VECTOR_NORMALIZED		(0.01f)		/** Allowed error for a normalized vector (against squared magnitude) */
+#define THRESH_QUAT_NORMALIZED			(0.01f)		/** Allowed error for a normalized quaternion (against squared magnitude) */
+
 /*-----------------------------------------------------------------------------
 	Global functions.
 -----------------------------------------------------------------------------*/
diff --git a/Engine/Source/Runtime/Core/Public/Math/UnrealMathVectorConstants.h b/Engine/Source/Runtime/Core/Public/Math/UnrealMathVectorConstants.h
--- a/Engine/Source/Runtime/Core/Public/Math/UnrealMathVectorConstants.h
+++ b/Engine/Source/Runtime/Core/Public/Math/UnrealMathVectorConstants.h
@@ -13,6 +13,7 @@ namespace GlobalVectorConstants
 	static const VectorRegister FloatOneHalf = MakeVectorRegister( 0.5f, 0.5f, 0.5f, 0.5f );
 	static const VectorRegister KindaSmallNumber = MakeVectorRegister( KINDA_SMALL_NUMBER, KINDA_SMALL_NUMBER, KINDA_SMALL_NUMBER, KINDA_SMALL_NUMBER );
 	static const VectorRegister SmallNumber = MakeVectorRegister( SMALL_NUMBER, SMALL_NUMBER, SMALL_NUMBER, SMALL_NUMBER );
+	static const VectorRegister ThreshQuatNormalized = MakeVectorRegister( THRESH_QUAT_NORMALIZED, THRESH_QUAT_NORMALIZED, THRESH_QUAT_NORMALIZED, THRESH_QUAT_NORMALIZED );
 
 	/** This is to speed up Quaternion Inverse. Static variable to keep sign of inverse **/
 	static const VectorRegister QINV_SIGN_MASK = MakeVectorRegister( -1.f, -1.f, -1.f, 1.f );
diff --git a/Engine/Source/Runtime/Core/Public/Math/Vector.h b/Engine/Source/Runtime/Core/Public/Math/Vector.h
--- a/Engine/Source/Runtime/Core/Public/Math/Vector.h
+++ b/Engine/Source/Runtime/Core/Public/Math/Vector.h
@@ -1346,7 +1346,7 @@ FORCEINLINE bool FVector::Normalize(float Tolerance)
 
 FORCEINLINE bool FVector::IsNormalized() const
 {
-	return (FMath::Abs(1.f - SizeSquared()) <= 0.01f);
+	return (FMath::Abs(1.f - SizeSquared()) < THRESH_VECTOR_NORMALIZED);
 }
 
 FORCEINLINE void FVector::ToDirectionAndLength(FVector &OutDir, float &OutLength)
diff --git a/Engine/Source/Runtime/Engine/Private/Animation/AnimationRuntime.cpp b/Engine/Source/Runtime/Engine/Private/Animation/AnimationRuntime.cpp
--- a/Engine/Source/Runtime/Engine/Private/Animation/AnimationRuntime.cpp
+++ b/Engine/Source/Runtime/Engine/Private/Animation/AnimationRuntime.cpp
@@ -27,6 +27,14 @@ void FAnimationRuntime::NormalizeRotations(const FBoneContainer & RequiredBones,
 	}
 }
 
+void FAnimationRuntime::NormalizeRotations(FTransformArrayA2 & Atoms)
+{
+	for (int32 BoneIndex = 0; BoneIndex < Atoms.Num(); BoneIndex++)
+	{
+		Atoms[BoneIndex].NormalizeRotation();
+	}
+}
+
 void FAnimationRuntime::InitializeTransform(const FBoneContainer & RequiredBones, /*inout*/ FTransformArrayA2 & Atoms)
 {
 	check( Atoms.Num() == RequiredBones.GetNumBones() );
diff --git a/Engine/Source/Runtime/Engine/Private/PhysicsEngine/PhysAnim.cpp b/Engine/Source/Runtime/Engine/Private/PhysicsEngine/PhysAnim.cpp
--- a/Engine/Source/Runtime/Engine/Private/PhysicsEngine/PhysAnim.cpp
+++ b/Engine/Source/Runtime/Engine/Private/PhysicsEngine/PhysAnim.cpp
@@ -175,7 +175,16 @@ void USkeletalMeshComponent::BlendPhysicsBones( TArray<FBoneIndexType>& InRequir
 		else
 		{
 			const int32 ParentIndex	= SkeletalMesh->RefSkeleton.GetParentIndex(BoneIndex);
-			SpaceBases[BoneIndex]	= LocalAtoms[BoneIndex] * SpaceBases[ParentIndex];
+			SpaceBases[BoneIndex] = LocalAtoms[BoneIndex] * SpaceBases[ParentIndex];
+
+			/**
+			* Normalize rotations.
+			* We want to remove any loss of precision due to accumulation of error.
+			* i.e. A componentSpace transform is the accumulation of all of its local space parents. The further down the chain, the greater the error.
+			* SpaceBases are used by external systems, we feed this to PhysX, send this to gameplay through bone and socket queries, etc.
+			* So this is a good place to make sure all transforms are normalized.
+			*/
+			SpaceBases[BoneIndex].NormalizeRotation();
 		}
 
 		if (bUpdatePhysics && BodyInstance)
diff --git a/Engine/Source/Runtime/Engine/Private/SkeletalMeshComponent.cpp b/Engine/Source/Runtime/Engine/Private/SkeletalMeshComponent.cpp
--- a/Engine/Source/Runtime/Engine/Private/SkeletalMeshComponent.cpp
+++ b/Engine/Source/Runtime/Engine/Private/SkeletalMeshComponent.cpp
@@ -594,6 +594,15 @@ void USkeletalMeshComponent::FillSpaceBases(const TArray<FTransform>& SourceAtom
 		checkSlow( SpaceBases[BoneIndex].IsRotationNormalized() );
 		checkSlow( !SpaceBases[BoneIndex].ContainsNaN() );
 	}
+
+	/**
+	 * Normalize rotations.
+	 * We want to remove any loss of precision due to accumulation of error.
+	 * i.e. A componentSpace transform is the accumulation of all of its local space parents. The further down the chain, the greater the error.
+	 * SpaceBases are used by external systems, we feed this to PhysX, send this to gameplay through bone and socket queries, etc.
+	 * So this is a good place to make sure all transforms are normalized.
+	 */
+	NormalizeRotations(DestSpaceBases);
 }
 
 /** Takes sorted array Base and then adds any elements from sorted array Insert which is missing from it, preserving order. */
@@ -781,6 +790,9 @@ void USkeletalMeshComponent::EvaluateAnimation(TArray<FTransform>& OutLocalAtoms
 			if( EvaluationContext.Pose.Bones.Num() > 0 )
 			{
 				OutLocalAtoms = EvaluationContext.Pose.Bones;
+
+				// Make sure rotations are normalized to account for accumulation of errors.
+				FAnimationRuntime::NormalizeRotations(OutLocalAtoms);
 			}
 			else
 			{
diff --git a/Engine/Source/Runtime/Engine/Public/AnimationRuntime.h b/Engine/Source/Runtime/Engine/Public/AnimationRuntime.h
--- a/Engine/Source/Runtime/Engine/Public/AnimationRuntime.h
+++ b/Engine/Source/Runtime/Engine/Public/AnimationRuntime.h
@@ -16,6 +16,8 @@ class ENGINE_API FAnimationRuntime
 {
 public:
 	static void NormalizeRotations(const FBoneContainer & RequiredBones, /*inout*/ FTransformArrayA2& Atoms);
+	static void NormalizeRotations(FTransformArrayA2 & Atoms);
+
 	static void InitializeTransform(const FBoneContainer & RequiredBones, /*inout*/ FTransformArrayA2& Atoms);
 #if DO_GUARD_SLOW
 	static bool ContainsNaN(TArray<FBoneIndexType> & RequiredBoneIndices, FA2Pose & Pose);
