fix compile error

diff --git a/Engine/Source/Runtime/Engine/Private/SkeletalMeshComponent.cpp b/Engine/Source/Runtime/Engine/Private/SkeletalMeshComponent.cpp
--- a/Engine/Source/Runtime/Engine/Private/SkeletalMeshComponent.cpp
+++ b/Engine/Source/Runtime/Engine/Private/SkeletalMeshComponent.cpp
@@ -602,7 +602,7 @@ void USkeletalMeshComponent::FillSpaceBases(const TArray<FTransform>& SourceAtom
 	 * SpaceBases are used by external systems, we feed this to PhysX, send this to gameplay through bone and socket queries, etc.
 	 * So this is a good place to make sure all transforms are normalized.
 	 */
-	NormalizeRotations(DestSpaceBases);
+	FAnimationRuntime::NormalizeRotations(DestSpaceBases);
 }
 
 /** Takes sorted array Base and then adds any elements from sorted array Insert which is missing from it, preserving order. */
