Changes to engine content from Zen branch

diff --git a/Engine/Source/Runtime/Engine/Private/ArrowComponent.cpp b/Engine/Source/Runtime/Engine/Private/ArrowComponent.cpp
--- a/Engine/Source/Runtime/Engine/Private/ArrowComponent.cpp
+++ b/Engine/Source/Runtime/Engine/Private/ArrowComponent.cpp
@@ -188,7 +188,7 @@ public:
 			if (ZoomFactor != 0.0f)
 			{
 				const float Radius = View->WorldToScreen(Origin).W * (ScreenSize / ZoomFactor);
-				if (Radius < 1.0f)
+				if (Radius < 1.0f && Radius > 0)
 				{
 					ViewScale *= Radius;
 				}
