Physics: Change FPhysicsScene2D to use UWorld::GetDefaultGravityZ() instead of a hardcoded -10

diff --git a/Engine/Source/Runtime/Engine/Private/PhysicsEngine2D/Box2DIntegration.cpp b/Engine/Source/Runtime/Engine/Private/PhysicsEngine2D/Box2DIntegration.cpp
--- a/Engine/Source/Runtime/Engine/Private/PhysicsEngine2D/Box2DIntegration.cpp
+++ b/Engine/Source/Runtime/Engine/Private/PhysicsEngine2D/Box2DIntegration.cpp
@@ -102,9 +102,10 @@ void FBox2DVisualizer::DrawTransform(const b2Transform& xf)
 FPhysicsScene2D::FPhysicsScene2D(UWorld* AssociatedWorld)
 	: UnrealWorld(AssociatedWorld)
 {
-	const b2Vec2 DefaultGravity(0.0f, -10.0f);
+	const FVector DefaultGravityUnreal(0.0f, 0.0f, AssociatedWorld->GetDefaultGravityZ());
+	const b2Vec2 DefaultGravityBox = FPhysicsIntegration2D::ConvertUnrealVectorToBox(DefaultGravityUnreal);
 
-	World = new b2World(DefaultGravity);
+	World = new b2World(DefaultGravityBox);
 	World->SetDebugDraw(&FBox2DVisualizer::GlobalVisualizer);
 
 	StartPhysicsTickFunction.Target = this;
