Physics: Add UE_WITH_PHYSICS as a shorthand for (WITH_PHYSX || WITH_BOX2D)

diff --git a/Engine/Source/Runtime/Engine/Classes/PhysicsEngine/BodyInstance.h b/Engine/Source/Runtime/Engine/Classes/PhysicsEngine/BodyInstance.h
--- a/Engine/Source/Runtime/Engine/Classes/PhysicsEngine/BodyInstance.h
+++ b/Engine/Source/Runtime/Engine/Classes/PhysicsEngine/BodyInstance.h
@@ -5,6 +5,9 @@
 #include "PhysxUserData.h"
 #include "BodyInstance.generated.h"
 
+
+#define UE_WITH_PHYSICS (WITH_PHYSX || WITH_BOX2D)
+
 #if WITH_PHYSX
 namespace physx
 {
@@ -255,7 +258,7 @@ public:
 	typedef void* PhysXAggregateType;
 #endif
 
-#if WITH_PHYSX || WITH_BOX2D
+#if UE_WITH_PHYSICS
 	void InitBody(class UBodySetup* Setup, const FTransform& Transform, class UPrimitiveComponent* PrimComp, class FPhysScene* InRBScene, PhysXAggregateType InAggregate = NULL);
 #endif	//WITH_PHYSX
 
diff --git a/Engine/Source/Runtime/Engine/Private/Collision/WorldCollision.cpp b/Engine/Source/Runtime/Engine/Private/Collision/WorldCollision.cpp
--- a/Engine/Source/Runtime/Engine/Private/Collision/WorldCollision.cpp
+++ b/Engine/Source/Runtime/Engine/Private/Collision/WorldCollision.cpp
@@ -45,27 +45,27 @@ ECollisionChannel DefaultCollisionChannel = (ECollisionChannel) 0;
 
 bool UWorld::LineTraceTest(const FVector& Start, const FVector& End, ECollisionChannel TraceChannel, const struct FCollisionQueryParams& Params, const struct FCollisionResponseParams& ResponseParam) const
 {
-#if WITH_PHYSX
+#if UE_WITH_PHYSICS
 	return RaycastTest(this, Start, End, TraceChannel, Params, ResponseParam, FCollisionObjectQueryParams::DefaultObjectQueryParam);
-#endif // WITH_PHYSX
+#endif
 	return false;
 }
 
 
 bool UWorld::LineTraceSingle(struct FHitResult& OutHit, const FVector& Start, const FVector& End, ECollisionChannel TraceChannel, const struct FCollisionQueryParams& Params, const struct FCollisionResponseParams& ResponseParam) const
 {
-#if WITH_PHYSX
+#if UE_WITH_PHYSICS
 	return RaycastSingle(this, OutHit, Start, End, TraceChannel, Params, ResponseParam, FCollisionObjectQueryParams::DefaultObjectQueryParam);
 #else
 	OutHit.TraceStart = Start;
 	OutHit.TraceEnd = End;
 	return false;
-#endif // WITH_PHYSX
+#endif
 }
 
 bool UWorld::LineTraceMulti(TArray<struct FHitResult>& OutHits, const FVector& Start, const FVector& End, ECollisionChannel TraceChannel, const struct FCollisionQueryParams& Params, const struct FCollisionResponseParams& ResponseParam) const
 {
-#if WITH_PHYSX
+#if UE_WITH_PHYSICS
 	return RaycastMulti(this, OutHits, Start, End, TraceChannel, Params, ResponseParam, FCollisionObjectQueryParams::DefaultObjectQueryParam);
 #endif // WITH_PHYSX
 	return false;
@@ -109,6 +109,9 @@ bool UWorld::SweepTest(const FVector& Start, const FVector& End, const FQuat& Ro
 	}
 
 #endif //WITH_PHYSX
+
+	//@TODO: BOX2D: Implement UWorld::SweepTest
+
 	return false;
 }
 
@@ -149,6 +152,8 @@ bool UWorld::SweepSingle(struct FHitResult& OutHit, const FVector& Start, const
 
 #endif //WITH_PHYSX
 
+	//@TODO: BOX2D: Implement UWorld::SweepSingle
+
 	OutHit.TraceStart = Start;
 	OutHit.TraceEnd = End;
 	return false;
@@ -189,6 +194,9 @@ bool UWorld::SweepMulti( TArray<FHitResult>& OutHits, const FVector& Start, cons
 		ensure(false);
 	}
 #endif //WITH_PHYSX
+
+	//@TODO: BOX2D: Implement UWorld::SweepMulti
+
 	return false;
 }
 
@@ -224,6 +232,9 @@ bool UWorld::OverlapTest(const FVector& Pos, const FQuat& Rot,ECollisionChannel
 	}
 
 #endif //WITH_PHYSX
+
+	//@TODO: BOX2D: Implement UWorld::OverlapTest
+
 	return false;
 }
 
@@ -258,6 +269,9 @@ bool UWorld::OverlapSingle(struct FOverlapResult& OutOverlap, const FVector& Pos
 		ensure(false);
 	}
 #endif //WITH_PHYSX
+
+	//@TODO: BOX2D: Implement UWorld::OverlapSingle
+
 	return false;
 }
 
@@ -292,6 +306,9 @@ bool UWorld::OverlapMulti(TArray<struct FOverlapResult>& OutOverlaps, const FVec
 		ensure(false);
 	}
 #endif //WITH_PHYSX
+
+	//@TODO: BOX2D: Implement UWorld::OverlapMulti
+
 	return false;
 }
 
@@ -299,7 +316,7 @@ bool UWorld::OverlapMulti(TArray<struct FOverlapResult>& OutOverlaps, const FVec
 
 bool UWorld::LineTraceTest(const FVector& Start,const FVector& End,const struct FCollisionQueryParams& Params, const struct FCollisionObjectQueryParams& ObjectQueryParams) const
 {
-#if WITH_PHYSX
+#if UE_WITH_PHYSICS
 	return RaycastTest(this, Start, End, DefaultCollisionChannel, Params, FCollisionResponseParams::DefaultResponseParam, ObjectQueryParams);
 #endif
 	return false;
@@ -307,21 +324,21 @@ bool UWorld::LineTraceTest(const FVector& Start,const FVector& End,const struct
 
 bool UWorld::LineTraceSingle(struct FHitResult& OutHit,const FVector& Start,const FVector& End,const struct FCollisionQueryParams& Params, const struct FCollisionObjectQueryParams& ObjectQueryParams) const
 {
-#if WITH_PHYSX
+#if UE_WITH_PHYSICS
 	return RaycastSingle(this, OutHit, Start, End, DefaultCollisionChannel, Params, FCollisionResponseParams::DefaultResponseParam, ObjectQueryParams);
 #else
 	OutHit.TraceStart = Start;
 	OutHit.TraceEnd = End;
 	return false;
-#endif // WITH_PHYSX
+#endif
 }
 
-bool UWorld::LineTraceMulti(TArray<struct FHitResult>& OutHits,const FVector& Start,const FVector& End,const struct FCollisionQueryParams& Params, const struct FCollisionObjectQueryParams& ObjectQueryParams) const
+bool UWorld::LineTraceMulti(TArray<struct FHitResult>& OutHits, const FVector& Start, const FVector& End, const struct FCollisionQueryParams& Params, const struct FCollisionObjectQueryParams& ObjectQueryParams) const
 {
 	// object query returns true if any hit is found, not only blocking hit
-#if WITH_PHYSX
+#if UE_WITH_PHYSICS
 	RaycastMulti(this, OutHits, Start, End, DefaultCollisionChannel, Params, FCollisionResponseParams::DefaultResponseParam, ObjectQueryParams);
-#endif // WITH_PHYSX
+#endif
 	return (OutHits.Num() > 0);
 }
 
@@ -363,6 +380,9 @@ bool UWorld::SweepTest(const FVector& Start, const FVector& End, const FQuat& Ro
 	}
 
 #endif //WITH_PHYSX
+
+	//@TODO: BOX2D: Implement UWorld::SweepTest
+
 	return false;
 }
 
@@ -403,6 +423,8 @@ bool UWorld::SweepSingle(struct FHitResult& OutHit,const FVector& Start,const FV
 
 #endif //WITH_PHYSX
 
+	//@TODO: BOX2D: Implement UWorld::SweepSingle
+
 	OutHit.TraceStart = Start;
 	OutHit.TraceEnd = End;
 	return false;
@@ -449,6 +471,9 @@ bool UWorld::SweepMulti(TArray<struct FHitResult>& OutHits,const FVector& Start,
 		}
 	}
 #endif //WITH_PHYSX
+
+	//@TODO: BOX2D: Implement UWorld::SweepMulti
+
 	return (OutHits.Num() > 0);
 }
 
@@ -484,6 +509,9 @@ bool UWorld::OverlapTest(const FVector& Pos, const FQuat& Rot,const struct FColl
 	}
 
 #endif //WITH_PHYSX
+
+	//@TODO: BOX2D: Implement UWorld::OverlapTest
+
 	return false;
 }
 
@@ -518,6 +546,9 @@ bool UWorld::OverlapSingle(struct FOverlapResult& OutOverlap,const FVector& Pos,
 		ensure(false);
 	}
 #endif //WITH_PHYSX
+
+	//@TODO: BOX2D: Implement UWorld::OverlapSingle
+
 	return false;
 }
 
@@ -555,6 +586,9 @@ bool UWorld::OverlapMulti(TArray<struct FOverlapResult>& OutOverlaps,const FVect
 		ensure(false);
 	}
 #endif //WITH_PHYSX
+
+	//@TODO: BOX2D: Implement UWorld::OverlapMulti
+
 	return (OutOverlaps.Num() > 0);
 }
 
@@ -612,6 +646,9 @@ bool UWorld::ComponentOverlapTest(class UPrimitiveComponent* PrimComp, const FVe
 	}
 
 #endif //WITH_PHYSX
+
+	//@TODO: BOX2D: Implement UWorld::ComponentOverlapTest
+
 	return false;
 }
 
@@ -659,13 +696,15 @@ bool UWorld::ComponentSweepSingle(struct FHitResult& OutHit,class UPrimitiveComp
 	}
 	
 	ECollisionChannel TraceChannel = PrimComp->GetCollisionObjectType();
-#if WITH_PHYSX
+
 	// if extent is 0, do line trace
 	if (PrimComp->IsZeroExtent())
 	{
 		return RaycastSingle(this, OutHit, Start, End, TraceChannel, Params, FCollisionResponseParams(PrimComp->GetCollisionResponseToChannels()));
 	}
 
+	//@TODO: BOX2D: Implement UWorld::ComponentSweepSingle
+#if WITH_PHYSX
 	PxRigidActor* PRigidActor = PrimComp->BodyInstance.GetPxRigidActor();
 	if(PRigidActor == NULL)
 	{
@@ -714,6 +753,7 @@ bool UWorld::ComponentSweepSingle(struct FHitResult& OutHit,class UPrimitiveComp
 
 	return bHaveBlockingHit;
 #endif //WITH_PHYSX
+
 	return false;
 }
 
@@ -732,13 +772,13 @@ bool UWorld::ComponentSweepMulti(TArray<struct FHitResult>& OutHits, class UPrim
 
 	ECollisionChannel TraceChannel = PrimComp->GetCollisionObjectType();
 
-#if WITH_PHYSX
 	// if extent is 0, do line trace
 	if (PrimComp->IsZeroExtent())
 	{
 		return RaycastMulti(this, OutHits, Start, End, TraceChannel, Params, FCollisionResponseParams(PrimComp->GetCollisionResponseToChannels()));
 	}
 
+#if WITH_PHYSX
 	PxRigidActor* PRigidActor = PrimComp->BodyInstance.GetPxRigidActor();
 	if(PRigidActor == NULL)
 	{
@@ -798,6 +838,9 @@ bool UWorld::ComponentSweepMulti(TArray<struct FHitResult>& OutHits, class UPrim
 
 	return bHaveBlockingHit;
 #endif //WITH_PHYSX
+
+	//@TODO: BOX2D: Implement UWorld::ComponentSweepMulti
+
 	return false;
 }
 
diff --git a/Engine/Source/Runtime/Engine/Private/PhysicsEngine/BodyInstance.cpp b/Engine/Source/Runtime/Engine/Private/PhysicsEngine/BodyInstance.cpp
--- a/Engine/Source/Runtime/Engine/Private/PhysicsEngine/BodyInstance.cpp
+++ b/Engine/Source/Runtime/Engine/Private/PhysicsEngine/BodyInstance.cpp
@@ -731,7 +731,7 @@ void FBodyInstance::UpdatePhysicsFilterData()
 #endif
 }
 
-#if WITH_PHYSX || WITH_BOX2D
+#if UE_WITH_PHYSICS
 void FBodyInstance::InitBody(UBodySetup* Setup, const FTransform& Transform, UPrimitiveComponent* PrimComp, FPhysScene* InRBScene, PxAggregate* InAggregate)
 {
 	check(Setup);
@@ -1212,7 +1212,7 @@ void FBodyInstance::InitBody(UBodySetup* Setup, const FTransform& Transform, UPr
 	}
 #endif // WITH_PHYSX
 }
-#endif // WITH_PHYSX || WITH_BOX2D
+#endif // UE_WITH_PHYSICS
 
 //helper function for TermBody to avoid code duplication between scenes
 #if WITH_PHYSX
@@ -2730,7 +2730,7 @@ FString FBodyInstance::GetBodyDebugName() const
 /////////////////////////////////////////////////////////////////////////////////////////////////////////////////////////
 // COLLISION
 
-bool FBodyInstance::LineTrace(struct FHitResult& OutHit,const FVector& Start,const FVector& End, bool bTraceComplex, bool bReturnPhysicalMaterial)
+bool FBodyInstance::LineTrace(struct FHitResult& OutHit, const FVector& Start, const FVector& End, bool bTraceComplex, bool bReturnPhysicalMaterial)
 {
 	SCOPE_CYCLE_COUNTER(STAT_Collision_RaycastAny);
 
@@ -2772,7 +2772,7 @@ bool FBodyInstance::LineTrace(struct FHitResult& OutHit,const FVector& Start,con
 			PxFilterData ShapeFilter = PShape->getQueryFilterData();
 			const bool bShapeIsComplex = (ShapeFilter.word3 & EPDF_ComplexCollision) != 0;
 			const bool bShapeIsSimple = (ShapeFilter.word3 & EPDF_SimpleCollision) != 0;
-			if((bTraceComplex && bShapeIsComplex) || (!bTraceComplex && bShapeIsSimple))
+			if ((bTraceComplex && bShapeIsComplex) || (!bTraceComplex && bShapeIsSimple))
 			{
 				const int32 ArraySize = ARRAY_COUNT(PHits);
 				// only care about one hit - not closest hit			
@@ -2797,7 +2797,7 @@ bool FBodyInstance::LineTrace(struct FHitResult& OutHit,const FVector& Start,con
 			}
 		}
 
-		if(BestHitDistance < BIG_NUMBER)
+		if (BestHitDistance < BIG_NUMBER)
 		{
 			// we just like to make sure if the hit is made, set to test touch
 			PxFilterData QueryFilter;
diff --git a/Engine/Source/Runtime/Engine/Private/PrimitiveComponent.cpp b/Engine/Source/Runtime/Engine/Private/PrimitiveComponent.cpp
--- a/Engine/Source/Runtime/Engine/Private/PrimitiveComponent.cpp
+++ b/Engine/Source/Runtime/Engine/Private/PrimitiveComponent.cpp
@@ -360,10 +360,10 @@ void UPrimitiveComponent::CreatePhysicsState()
 				return;
 			}
 
-#if WITH_PHYSX || WITH_BOX2D
+#if UE_WITH_PHYSICS
 			// Create the body.
 			BodyInstance.InitBody(BodySetup, ComponentToWorld, this, World->GetPhysicsScene());
-#endif //WITH_PHYSX
+#endif //UE_WITH_PHYSICS
 
 #if WITH_EDITOR
 			// Make sure we have a valid body instance here. As we do not keep BIs with no collision shapes at all,
