Physics: Add initial support for 2D RaycastSingle queries, gated by UPhysicsSettings::bEnable2DPhysics (WIP)

diff --git a/Engine/Source/Runtime/Engine/Classes/PhysicsEngine/PhysicsSettings.h b/Engine/Source/Runtime/Engine/Classes/PhysicsEngine/PhysicsSettings.h
--- a/Engine/Source/Runtime/Engine/Classes/PhysicsEngine/PhysicsSettings.h
+++ b/Engine/Source/Runtime/Engine/Classes/PhysicsEngine/PhysicsSettings.h
@@ -57,23 +57,27 @@ class ENGINE_API UPhysicsSettings
 {
 	GENERATED_UCLASS_BODY()
 
-	/**Default gravity. */
+	/** Default gravity. */
 	UPROPERTY(config, EditAnywhere, Category = Constants)
 	float DefaultGravityZ;
 
-	/**Default Terminal Velocity. */
+	/** Default Terminal Velocity. */
 	UPROPERTY(config, EditAnywhere, Category = Constants)
 	float DefaultTerminalVelocity;
 
-	/**Enables the use of an async scene */
+	/** Enables the use of an async scene */
 	UPROPERTY(config, EditAnywhere, AdvancedDisplay, Category=Simulation)
 	bool bEnableAsyncScene;
 
+	/** Can 2D physics be used (Box2D)? */
+	UPROPERTY(config, EditAnywhere, Category = Simulation)
+	bool bEnable2DPhysics;
+
 	/** Friction combine mode, controls how friction is computed for multiple materials. */
 	UPROPERTY(config, EditAnywhere, Category=Simulation)
 	TEnumAsByte<EFrictionCombineMode::Type> FrictionCombineMode;
 
-	/**Max Physics Delta Time to be clamped. */
+	/** Max Physics Delta Time to be clamped. */
 	UPROPERTY(config, EditAnywhere, meta=(ClampMin="0.0013", UIMin = "0.0013", ClampMax="1.0", UIMax="1.0"), Category=Framerate)
 	float MaxPhysicsDeltaTime;
 
@@ -81,7 +85,7 @@ class ENGINE_API UPhysicsSettings
 	UPROPERTY(config, EditAnywhere, Category = Framerate)
 	bool bSubstepping;
 
-		/**Max delta time for an individual substep simulation. */
+	/** Max delta time for an individual substep simulation. */
 	UPROPERTY(config, EditAnywhere, meta = (ClampMin = "0.0013", UIMin = "0.0013", ClampMax = "1.0", UIMax = "1.0", editcondition = "bSubStepping"), Category=Framerate)
 	float MaxSubstepDeltaTime;
 
@@ -89,15 +93,15 @@ class ENGINE_API UPhysicsSettings
 	UPROPERTY(config, EditAnywhere, meta = (ClampMin = "1", UIMin = "1", ClampMax = "16", UIMax = "16", editcondition = "bSubstepping"), Category=Framerate)
 	int32 MaxSubsteps;
 
-	/**Physics delta time smoothing factor for sync scene. */
+	/** Physics delta time smoothing factor for sync scene. */
 	UPROPERTY(config, EditAnywhere, AdvancedDisplay, meta = (ClampMin = "0.0", UIMin = "0.0", ClampMax = "1.0", UIMax = "1.0"), Category = Framerate)
 	float SyncSceneSmoothingFactor;
 
-	/**Physics delta time smoothing factor for async scene. */
+	/** Physics delta time smoothing factor for async scene. */
 	UPROPERTY(config, EditAnywhere, AdvancedDisplay, meta = (ClampMin = "0.0", UIMin = "0.0", ClampMax = "1.0", UIMax = "1.0"), Category = Framerate)
 	float AsyncSceneSmoothingFactor;
 
-	/**Physics delta time initial average. */
+	/** Physics delta time initial average. */
 	UPROPERTY(config, EditAnywhere, AdvancedDisplay, meta = (ClampMin = "0.0013", UIMin = "1.0", ClampMax = "1.0", UIMax = "1.0"), Category = Framerate)
 	float InitialAverageFrameRate;
 
diff --git a/Engine/Source/Runtime/Engine/Private/Collision/PhysXCollision.cpp b/Engine/Source/Runtime/Engine/Private/Collision/PhysXCollision.cpp
--- a/Engine/Source/Runtime/Engine/Private/Collision/PhysXCollision.cpp
+++ b/Engine/Source/Runtime/Engine/Private/Collision/PhysXCollision.cpp
@@ -1,13 +1,362 @@
 // Copyright 1998-2014 Epic Games, Inc. All Rights Reserved.
 
-/*=============================================================================
-	WorldCollision.cpp: UWorld collision implementation
-=============================================================================*/
-
 #include "EnginePrivate.h"
 #include "PhysicsPublic.h"
 #include "Collision.h"
 
+#if WITH_BOX2D
+#include "../PhysicsEngine2D/Box2DIntegration.h"
+
+class FRaycastHelperBox2D : public b2RayCastCallback, public b2QueryCallback
+{
+public:
+	enum EResponse
+	{
+		ResponseNone = 0,
+		ResponseTouch = 1,
+		ResponseBlock = 2,
+	};
+
+	FRaycastHelperBox2D(ECollisionChannel InTraceChannel, bool bInTraceComplex, const FCollisionResponseContainer& InCollisionResponseContainer, const FCollisionObjectQueryParams& InObjectParam, const TArray<uint32, TInlineAllocator<1> >& InIgnoreActors, bool bInMultitrace)
+		: TraceChannel(InTraceChannel)
+		, CollisionResponseContainer(InCollisionResponseContainer)
+		, ObjectParam(InObjectParam)
+		, IgnoreActors(InIgnoreActors)
+		, BestResult(MAX_FLT)
+		, BestFixture(nullptr)
+		, QueryBlockingChannels(0)
+		, QueryTouchingChannels(0)
+		, bMultitrace(bInMultitrace)
+		, bTraceComplex(bInTraceComplex)
+		, bHasBlockingHit(false)
+	{
+		// Box2D is not thread safe
+		check(IsInGameThread());
+
+		if (!ObjectParam.IsValid())
+		{
+			for (int32 i = 0; i < ARRAY_COUNT(CollisionResponseContainer.EnumArray); i++)
+			{
+				if (CollisionResponseContainer.EnumArray[i] == ECR_Block)
+				{
+					QueryBlockingChannels |= CRC_TO_BITFIELD(i);
+				}
+				else if (CollisionResponseContainer.EnumArray[i] == ECR_Overlap)
+				{
+					QueryTouchingChannels |= CRC_TO_BITFIELD(i);
+				}
+			}
+
+		}
+	}
+
+	// b2RayCastCallback interface
+	virtual float32 ReportFixture(b2Fixture* Fixture, const b2Vec2& Point, const b2Vec2& Normal, float32 Fraction) OVERRIDE;
+	// End of b2RayCastCallback interface
+
+	// b2QueryCallback interface
+	virtual bool ReportFixture(b2Fixture* Fixture) OVERRIDE;
+	// End of b2QueryCallback interface
+
+	bool HasBlockingHit() const
+	{
+		return bHasBlockingHit;
+	}
+
+	float GetBestTime() const
+	{
+		return BestResult;
+	}
+
+	// Converts the query results to a hit result
+	void ConvertToHitResult(FHitResult& OutResult, const FVector& Start, const FVector& End, const FCollisionQueryParams& Params)
+	{
+		check(bHasBlockingHit);
+
+		//@TODO: BOX2D: How to deal with initial penetration (PxU32)(PHit.hadInitialOverlap());
+#if 0
+		if ((BestFixture != NULL) && PHit.hadInitialOverlap())
+		{
+			ConvertOverlappedShapeToImpactHit(PHit, StartLoc, EndLoc, OutResult, *BestFixture, QueryTM, QueryFilter, Params.bReturnPhysicalMaterial);
+			return;
+		}
+#endif
+		OutResult.bStartPenetrating = false;
+
+		check(BestFixture);
+		FBodyInstance* OwningBodyInstance = static_cast<FBodyInstance*>(BestFixture->GetBody()->GetUserData());
+		UPrimitiveComponent* OwningComponent = OwningBodyInstance->OwnerComponent.Get();
+
+
+		/** Set info in the HitResult (Actor, Component, PhysMaterial, BoneName, Item) based on the supplied shape and face index */
+		OutResult.PhysMaterial = NULL;
+
+		// Grab actor/component
+		if (OwningComponent)
+		{
+			OutResult.Actor = OwningComponent->GetOwner();
+			OutResult.Component = OwningComponent;
+
+			if (Params.bReturnPhysicalMaterial)
+			{
+				OutResult.PhysMaterial = OwningBodyInstance->GetSimplePhysicalMaterial();
+			}
+		}
+
+		// Body index / bone name
+		OutResult.Item = OwningBodyInstance->InstanceBodyIndex;
+		if (OwningBodyInstance->BodySetup.IsValid())
+		{
+			OutResult.BoneName = OwningBodyInstance->BodySetup->BoneName;
+		}
+	
+
+
+		OutResult.FaceIndex = INDEX_NONE;
+
+
+		// calculate the hit time
+		const float HitTime = BestResult;
+
+		// figure out where the the "safe" location for this shape is by moving from the startLoc toward the ImpactPoint
+		const FVector TraceDir = End - Start;
+		const FVector SafeLocationToFitShape = Start + (HitTime * TraceDir);
+
+		// Other info
+		OutResult.Location = SafeLocationToFitShape;
+		OutResult.ImpactPoint = SafeLocationToFitShape;
+		OutResult.Normal = FPhysicsIntegration2D::ConvertBoxVectorToUnreal(BestNormal).SafeNormal();
+		OutResult.ImpactNormal = OutResult.Normal;
+
+		OutResult.TraceStart = Start;
+		OutResult.TraceEnd = End;
+		OutResult.Time = HitTime;
+
+		// See if this is a 'blocking' hit
+		OutResult.bBlockingHit = (BestHitType == ResponseBlock);
+	}
+
+	void PointQuery(const b2World* BoxWorld, const b2Vec2& BoxStart, const FVector& Start, const FVector& End, const FCollisionQueryParams& Params, FHitResult& OutResult);
+
+
+	// Determines how the query and the fixture interact (either no interaction or touching/blocking each other)
+	EResponse CalcQueryHitType(b2Fixture* Fixture) const;
+private:
+	const ECollisionChannel TraceChannel;
+	const FCollisionResponseContainer& CollisionResponseContainer;
+	const FCollisionObjectQueryParams& ObjectParam;
+	const TArray<uint32, TInlineAllocator<1> >& IgnoreActors;
+
+	b2Vec2 BestPoint;
+	b2Vec2 BestNormal;
+	float BestResult;
+	b2Fixture* BestFixture;
+	uint32 QueryBlockingChannels;
+	uint32 QueryTouchingChannels;
+	EResponse BestHitType;
+
+	bool bMultitrace;
+	bool bTraceComplex;
+	bool bHasBlockingHit;
+};
+
+void FRaycastHelperBox2D::PointQuery(const b2World* BoxWorld, const b2Vec2& BoxStart, const FVector& Start, const FVector& End, const FCollisionQueryParams& Params, FHitResult& OutResult)
+{
+	// Run the query
+	b2AABB QueryBox;
+	QueryBox.lowerBound = BoxStart;
+	QueryBox.upperBound = BoxStart;
+	BoxWorld->QueryAABB(this, QueryBox);
+
+	// Deal with a hit
+	if (bHasBlockingHit)
+	{
+		//@TODO: BOX2D: How to deal with initial penetration (PxU32)(PHit.hadInitialOverlap());
+#if 0
+		if ((BestFixture != NULL) && PHit.hadInitialOverlap())
+		{
+			ConvertOverlappedShapeToImpactHit(PHit, StartLoc, EndLoc, OutResult, *BestFixture, QueryTM, QueryFilter, Params.bReturnPhysicalMaterial);
+			return;
+		}
+#endif
+		OutResult.bStartPenetrating = false;
+
+		check(BestFixture);
+		FBodyInstance* OwningBodyInstance = static_cast<FBodyInstance*>(BestFixture->GetBody()->GetUserData());
+		UPrimitiveComponent* OwningComponent = OwningBodyInstance->OwnerComponent.Get();
+
+		/** Set info in the HitResult (Actor, Component, PhysMaterial, BoneName, Item) based on the supplied shape and face index */
+		OutResult.PhysMaterial = NULL;
+
+		// Grab actor/component
+		if (OwningComponent)
+		{
+			OutResult.Actor = OwningComponent->GetOwner();
+			OutResult.Component = OwningComponent;
+
+			if (Params.bReturnPhysicalMaterial)
+			{
+				OutResult.PhysMaterial = OwningBodyInstance->GetSimplePhysicalMaterial();
+			}
+		}
+
+		// Body index / bone name
+		OutResult.Item = OwningBodyInstance->InstanceBodyIndex;
+		if (OwningBodyInstance->BodySetup.IsValid())
+		{
+			OutResult.BoneName = OwningBodyInstance->BodySetup->BoneName;
+		}
+
+
+		OutResult.FaceIndex = INDEX_NONE;
+
+		// Figure out 'when' the ray hit the fixture
+		const float StartTime = FPhysicsIntegration2D::ConvertUnrealVectorToPerpendicularDistance(Start);
+		const float EndTime = FPhysicsIntegration2D::ConvertUnrealVectorToPerpendicularDistance(End);
+		const float IntersectTime = (OwningComponent != nullptr) ? FPhysicsIntegration2D::ConvertUnrealVectorToPerpendicularDistance(OwningComponent->GetComponentLocation()) : StartTime;
+		const float HitTime = IntersectTime / (EndTime - StartTime);
+
+		// figure out where the the "safe" location for this shape is by moving from the startLoc toward the ImpactPoint
+		const FVector TraceDir = End - Start;
+		const FVector SafeLocationToFitShape = Start + (HitTime * TraceDir);
+
+		// Other info
+		OutResult.Location = SafeLocationToFitShape;
+		OutResult.ImpactPoint = SafeLocationToFitShape;
+		OutResult.Normal = (Start - End).SafeNormal();
+		OutResult.ImpactNormal = OutResult.Normal;
+
+		OutResult.TraceStart = Start;
+		OutResult.TraceEnd = End;
+		OutResult.Time = HitTime;
+
+		// See if this is a 'blocking' hit
+		OutResult.bBlockingHit = (BestHitType == ResponseBlock);
+	}
+}
+
+float32 FRaycastHelperBox2D::ReportFixture(b2Fixture* Fixture, const b2Vec2& Point, const b2Vec2& Normal, float32 Fraction)
+{
+	const EResponse ResponseMode = CalcQueryHitType(Fixture);
+	
+	switch (ResponseMode)
+	{
+	case ResponseBlock:
+		if (Fraction < BestResult)
+		{
+			BestFixture = Fixture;
+			BestResult = Fraction;
+			BestNormal = Normal;
+			BestPoint = Point;
+			BestHitType = ResponseMode;
+			bHasBlockingHit = true;
+		}
+		else
+		{
+			//@TODO: BOX2D: Can this happen?
+			ensure(false);
+		}
+		return Fraction;
+	case ResponseTouch:
+		return 1.0f;
+	default:
+	case ResponseNone:
+		return -1.0f;
+	}
+}
+
+bool FRaycastHelperBox2D::ReportFixture(b2Fixture* Fixture)
+{
+	const EResponse ResponseMode = CalcQueryHitType(Fixture);
+
+	switch (ResponseMode)
+	{
+	case ResponseBlock:
+		{
+			BestFixture = Fixture;
+			BestHitType = ResponseMode;
+			bHasBlockingHit = true;
+		}
+		return true;
+	case ResponseTouch:
+		return false;
+	default:
+	case ResponseNone:
+		return false;
+	}
+}
+
+FRaycastHelperBox2D::EResponse FRaycastHelperBox2D::CalcQueryHitType(b2Fixture* Fixture) const
+{
+	const b2Filter& FixtureFilterData = Fixture->GetFilterData();
+	const ECollisionChannel ShapeChannel = static_cast<ECollisionChannel>(FixtureFilterData.ObjectTypeAndFlags >> 24);
+	const uint32 ShapeBit = ECC_TO_BITFIELD(ShapeChannel);
+
+	EResponse Result = ResponseNone;
+
+	if (ObjectParam.IsValid())
+	{
+		// do I belong to one of objects of interest?
+		if ((ShapeBit & ObjectParam.GetQueryBitfield()) != 0)
+		{
+			Result = bMultitrace ? ResponseTouch : ResponseBlock;
+		}
+	}
+	else
+	{
+		// Then see if the channel wants to be blocked
+		const uint32 TraceBit = ECC_TO_BITFIELD(TraceChannel);
+
+		// check if the fixture has a hit or touch
+		EResponse TraceHitType = ResponseNone;
+		if ((TraceBit & FixtureFilterData.BlockingChannels) != 0)
+		{
+			TraceHitType = ResponseBlock;
+		}
+		else if ((TraceBit & FixtureFilterData.TouchingChannels) != 0)
+		{
+			TraceHitType = ResponseTouch;
+		}
+
+		// Check if the shape has a hit or touch
+		EResponse ShapeHitType = ResponseNone;
+		if ((ShapeBit & QueryBlockingChannels) != 0)
+		{
+			ShapeHitType = ResponseBlock;
+		}
+		else if ((ShapeBit & QueryTouchingChannels) != 0)
+		{
+			ShapeHitType = ResponseTouch;
+		}
+
+		// return minimum agreed-upon interaction
+		Result = FMath::Min(TraceHitType, ShapeHitType);
+	}
+
+	// Single traces should ignore touches
+	if (!bMultitrace && (Result == ResponseTouch))
+	{
+		Result = ResponseNone;
+	}
+
+	// Check to see if the actor ultimately owning the fixture should be ignored
+	if (Result != ResponseNone)
+	{
+		FBodyInstance* OwnerBodyInstance = static_cast<FBodyInstance*>(Fixture->GetBody()->GetUserData());
+		UPrimitiveComponent* OwnerComponent = OwnerBodyInstance->OwnerComponent.Get();
+		AActor* Owner = OwnerComponent ? OwnerComponent->GetOwner() : nullptr;
+
+		if ((Owner != nullptr) && IgnoreActors.Contains(Owner->GetUniqueID()))
+		{
+			Result = ResponseNone;
+		}
+	}
+
+	return Result;
+}
+
+#endif	//WITH_BOX2D
+
 #if WITH_PHYSX
 
 #include "../PhysicsEngine/PhysXSupport.h"
@@ -15,9 +364,6 @@
 #include "CollisionDebugDrawing.h"
 #include "CollisionConversions.h"
 
-#if ENABLE_COLLISION_ANALYZER
-#include "CollisionAnalyzerModule.h"
-#endif
 
 #define VERIFY_GEOMSWEEPSINGLE 0
 #define VERIFY_GEOMSWEEPMULTI 0
@@ -205,10 +551,15 @@ PxSceneQueryHitType::Enum FPxQueryFilterCallbackSweep::postFilter(const PxFilter
 	}
 }
 
+#endif // WITH_PHYSX
 
 //////////////////////////////////////////////////////////////////////////
+
+//@TODO: BOX2D: Can we break the collision analyzer's dependence on PhysX?
 #if ENABLE_COLLISION_ANALYZER
 
+#include "CollisionAnalyzerModule.h"
+
 bool GCollisionAnalyzerIsRecording = false;
 
 bool bSkipCapture = false;
@@ -306,9 +657,11 @@ void CaptureRaycast(const UWorld* World, const FVector& Start, const FVector& En
 //////////////////////////////////////////////////////////////////////////
 // RAYCAST
 
-bool RaycastTest(const UWorld * World, const FVector Start, const FVector End, ECollisionChannel TraceChannel, const struct FCollisionQueryParams& Params, const struct FCollisionResponseParams& ResponseParams, const struct FCollisionObjectQueryParams& ObjectParams)
+#if UE_WITH_PHYSICS
+
+bool RaycastTest(const UWorld* World, const FVector Start, const FVector End, ECollisionChannel TraceChannel, const struct FCollisionQueryParams& Params, const struct FCollisionResponseParams& ResponseParams, const struct FCollisionObjectQueryParams& ObjectParams)
 {
-	if(World == NULL || World->GetPhysicsScene() == NULL)
+	if ((World == NULL) || (World->GetPhysicsScene() == NULL))
 	{
 		return false;
 	}
@@ -319,36 +672,49 @@ bool RaycastTest(const UWorld * World, const FVector Start, const FVector End, E
 
 	FVector Delta = End - Start;
 	float DeltaMag = Delta.Size();
-	if(DeltaMag > KINDA_SMALL_NUMBER)
+	if (DeltaMag > KINDA_SMALL_NUMBER)
 	{
-		// Create filter data used to filter collisions
-		PxFilterData PFilter = CreateQueryFilterData(TraceChannel, Params.bTraceComplex, ResponseParams.CollisionResponse, ObjectParams, false);
-		PxSceneQueryFilterData PQueryFilterData(PFilter, PxSceneQueryFilterFlag::eSTATIC | PxSceneQueryFilterFlag::eDYNAMIC | PxSceneQueryFilterFlag::ePREFILTER);
-		PxSceneQueryFlags POutputFlags = PxHitFlags();
-		FPxQueryFilterCallback PQueryCallback(Params.IgnoreActors);
-		PQueryCallback.bSingleQuery = true;
+#if WITH_PHYSX
+		{
+			// Create filter data used to filter collisions
+			PxFilterData PFilter = CreateQueryFilterData(TraceChannel, Params.bTraceComplex, ResponseParams.CollisionResponse, ObjectParams, false);
+			PxSceneQueryFilterData PQueryFilterData(PFilter, PxSceneQueryFilterFlag::eSTATIC | PxSceneQueryFilterFlag::eDYNAMIC | PxSceneQueryFilterFlag::ePREFILTER);
+			PxSceneQueryFlags POutputFlags = PxHitFlags();
+			FPxQueryFilterCallback PQueryCallback(Params.IgnoreActors);
+			PQueryCallback.bSingleQuery = true;
 
-		FPhysScene* PhysScene = World->GetPhysicsScene();
+			FPhysScene* PhysScene = World->GetPhysicsScene();
 
-		// Enable scene locks, in case they are required
-		PxScene* SyncScene = PhysScene->GetPhysXScene(PST_Sync);
-		SCOPED_SCENE_READ_LOCK(SyncScene);
+			// Enable scene locks, in case they are required
+			PxScene* SyncScene = PhysScene->GetPhysXScene(PST_Sync);
+			SCOPED_SCENE_READ_LOCK(SyncScene);
 
-		const PxVec3 PDir = U2PVector(Delta/DeltaMag);
-		PxRaycastHit PHit;
-		bHaveBlockingHit = SyncScene->raycastSingle(U2PVector(Start), PDir, DeltaMag, POutputFlags, PHit, PQueryFilterData, &PQueryCallback);
+			const PxVec3 PDir = U2PVector(Delta / DeltaMag);
+			PxRaycastHit PHit;
+			bHaveBlockingHit = SyncScene->raycastSingle(U2PVector(Start), PDir, DeltaMag, POutputFlags, PHit, PQueryFilterData, &PQueryCallback);
 
-		// Test async scene if we have no blocking hit, and async tests are requested
-		if( !bHaveBlockingHit && Params.bTraceAsyncScene && PhysScene->HasAsyncScene())
+			// Test async scene if we have no blocking hit, and async tests are requested
+			if (!bHaveBlockingHit && Params.bTraceAsyncScene && PhysScene->HasAsyncScene())
+			{
+				PxScene* AsyncScene = PhysScene->GetPhysXScene(PST_Async);
+				SCOPED_SCENE_READ_LOCK(AsyncScene);
+				bHaveBlockingHit = AsyncScene->raycastSingle(U2PVector(Start), PDir, DeltaMag, POutputFlags, PHit, PQueryFilterData, &PQueryCallback);
+			}
+		}
+#endif // WITH_PHYSX
+
+#if WITH_BOX2D
+		//@TODO: BOX2D: Implement RaycastTest
+		if (GetDefault<UPhysicsSettings>()->bEnable2DPhysics)
 		{
-			PxScene* AsyncScene = PhysScene->GetPhysXScene(PST_Async);
-			SCOPED_SCENE_READ_LOCK(AsyncScene);			bHaveBlockingHit = AsyncScene->raycastSingle(U2PVector(Start), PDir, DeltaMag, POutputFlags, PHit, PQueryFilterData, &PQueryCallback);
+			
 		}
+#endif // WITH_BOX2D
 	}
 
 	TArray<FHitResult> Hits;
 #if !(UE_BUILD_SHIPPING || UE_BUILD_TEST)
-	if((World->DebugDrawTraceTag != NAME_None) && (World->DebugDrawTraceTag == Params.TraceTag))
+	if ((World->DebugDrawTraceTag != NAME_None) && (World->DebugDrawTraceTag == Params.TraceTag))
 	{
 		DrawLineTraces(World, Start, End, Hits, DebugLineLifetime);
 	}
@@ -358,7 +724,7 @@ bool RaycastTest(const UWorld * World, const FVector Start, const FVector End, E
 	return bHaveBlockingHit;
 }
 
-bool RaycastSingle(const UWorld * World, struct FHitResult& OutHit, const FVector Start, const FVector End, ECollisionChannel TraceChannel, const struct FCollisionQueryParams& Params, const struct FCollisionResponseParams& ResponseParams, const struct FCollisionObjectQueryParams& ObjectParams)
+bool RaycastSingle(const UWorld* World, struct FHitResult& OutHit, const FVector Start, const FVector End, ECollisionChannel TraceChannel, const struct FCollisionQueryParams& Params, const struct FCollisionResponseParams& ResponseParams, const struct FCollisionObjectQueryParams& ObjectParams)
 {
 	SCOPE_CYCLE_COUNTER(STAT_Collision_RaycastSingle);
 	STARTQUERYTIMER();
@@ -366,7 +732,7 @@ bool RaycastSingle(const UWorld * World, struct FHitResult& OutHit, const FVecto
 	OutHit.TraceStart = Start;
 	OutHit.TraceEnd = End;
 
-	if(World == NULL || World->GetPhysicsScene() == NULL)
+	if ((World == NULL) || (World->GetPhysicsScene() == NULL))
 	{
 		return false;
 	}
@@ -375,49 +741,90 @@ bool RaycastSingle(const UWorld * World, struct FHitResult& OutHit, const FVecto
 
 	FVector Delta = End - Start;
 	float DeltaMag = Delta.Size();
-	if(DeltaMag > KINDA_SMALL_NUMBER)
+	if (DeltaMag > KINDA_SMALL_NUMBER)
 	{
-		// Create filter data used to filter collisions
-		PxFilterData PFilter = CreateQueryFilterData(TraceChannel, Params.bTraceComplex, ResponseParams.CollisionResponse, ObjectParams, false);
-		PxSceneQueryFilterData PQueryFilterData(PFilter, PxSceneQueryFilterFlag::eSTATIC | PxSceneQueryFilterFlag::eDYNAMIC | PxSceneQueryFilterFlag::ePREFILTER);
-		PxSceneQueryFlags POutputFlags = PxSceneQueryFlag::ePOSITION | PxSceneQueryFlag::eNORMAL | PxSceneQueryFlag::eDISTANCE;
-		FPxQueryFilterCallback PQueryCallback(Params.IgnoreActors);
-		PQueryCallback.bSingleQuery = true; // this needs to be set for raycastSingle so we return eNONE instead of eTOUCH
-
-		PxVec3 PStart = U2PVector(Start);
-		PxVec3 PDir = U2PVector(Delta/DeltaMag);
-
-		FPhysScene* PhysScene = World->GetPhysicsScene();
-		PxScene* SyncScene = PhysScene->GetPhysXScene(PST_Sync);
-
-		// Enable scene locks, in case they are required
-		SCOPED_SCENE_READ_LOCK(SyncScene);
-		PxRaycastHit PHit;
-		bHaveBlockingHit = SyncScene->raycastSingle(U2PVector(Start), PDir, DeltaMag, POutputFlags, PHit, PQueryFilterData, &PQueryCallback);
-
-		// Test async scene if async tests are requested
-		if( Params.bTraceAsyncScene && PhysScene->HasAsyncScene() )
+#if WITH_PHYSX
 		{
-			PxScene* AsyncScene = PhysScene->GetPhysXScene(PST_Async);
-			SCOPED_SCENE_READ_LOCK(AsyncScene);
-			bool bHaveBlockingHitAsync;
-			PxRaycastHit PHitAsync;
-			bHaveBlockingHitAsync = AsyncScene->raycastSingle(U2PVector(Start), PDir, DeltaMag, POutputFlags, PHitAsync, PQueryFilterData, &PQueryCallback);
+			// Create filter data used to filter collisions
+			PxFilterData PFilter = CreateQueryFilterData(TraceChannel, Params.bTraceComplex, ResponseParams.CollisionResponse, ObjectParams, false);
+			PxSceneQueryFilterData PQueryFilterData(PFilter, PxSceneQueryFilterFlag::eSTATIC | PxSceneQueryFilterFlag::eDYNAMIC | PxSceneQueryFilterFlag::ePREFILTER);
+			PxSceneQueryFlags POutputFlags = PxSceneQueryFlag::ePOSITION | PxSceneQueryFlag::eNORMAL | PxSceneQueryFlag::eDISTANCE;
+			FPxQueryFilterCallback PQueryCallback(Params.IgnoreActors);
+			PQueryCallback.bSingleQuery = true; // this needs to be set for raycastSingle so we return eNONE instead of eTOUCH
+
+			PxVec3 PStart = U2PVector(Start);
+			PxVec3 PDir = U2PVector(Delta / DeltaMag);
+
+			FPhysScene* PhysScene = World->GetPhysicsScene();
+			PxScene* SyncScene = PhysScene->GetPhysXScene(PST_Sync);
+
+			// Enable scene locks, in case they are required
+			SCOPED_SCENE_READ_LOCK(SyncScene);
+			PxRaycastHit PHit;
+			bHaveBlockingHit = SyncScene->raycastSingle(U2PVector(Start), PDir, DeltaMag, POutputFlags, PHit, PQueryFilterData, &PQueryCallback);
+
+			// Test async scene if async tests are requested
+			if (Params.bTraceAsyncScene && PhysScene->HasAsyncScene())
+			{
+				PxScene* AsyncScene = PhysScene->GetPhysXScene(PST_Async);
+				SCOPED_SCENE_READ_LOCK(AsyncScene);
+				bool bHaveBlockingHitAsync;
+				PxRaycastHit PHitAsync;
+				bHaveBlockingHitAsync = AsyncScene->raycastSingle(U2PVector(Start), PDir, DeltaMag, POutputFlags, PHitAsync, PQueryFilterData, &PQueryCallback);
+
+				// If we have a blocking hit in the async scene and there was no sync blocking hit, or if the async blocking hit came first,
+				// then this becomes the blocking hit.  We can test the PxSceneQueryImpactHit::distance since the DeltaMag is the same for both queries.
+				if (bHaveBlockingHitAsync && (!bHaveBlockingHit || PHitAsync.distance < PHit.distance))
+				{
+					PHit = PHitAsync;
+					bHaveBlockingHit = true;
+				}
+			}
 
-			// If we have a blocking hit in the async scene and there was no sync blocking hit, or if the async blocking hit came first,
-			// then this becomes the blocking hit.  We can test the PxSceneQueryImpactHit::distance since the DeltaMag is the same for both queries.
-			if( bHaveBlockingHitAsync && (!bHaveBlockingHit || PHitAsync.distance < PHit.distance) )
+			if (bHaveBlockingHit) // If we got a hit, 
 			{
-				PHit = PHitAsync;
-				bHaveBlockingHit = true;
+				PxTransform PStartTM(U2PVector(Start));
+				ConvertQueryImpactHit(PHit, OutHit, DeltaMag, PFilter, Start, End, NULL, PStartTM, Params.bReturnFaceIndex, Params.bReturnPhysicalMaterial);
 			}
 		}
+#endif //WITH_PHYSX
 
-		if(bHaveBlockingHit) // If we got a hit, 
+#if WITH_BOX2D
+		if (GetDefault<UPhysicsSettings>()->bEnable2DPhysics)
 		{
-			PxTransform PStartTM(U2PVector(Start));
-			ConvertQueryImpactHit(PHit, OutHit, DeltaMag, PFilter, Start, End, NULL, PStartTM, Params.bReturnFaceIndex, Params.bReturnPhysicalMaterial);
+			if (const b2World* BoxWorld = FPhysicsIntegration2D::FindAssociatedWorld(const_cast<UWorld*>(World)))
+			{
+				FRaycastHelperBox2D BoxQuery(TraceChannel, Params.bTraceComplex, ResponseParams.CollisionResponse, ObjectParams, Params.IgnoreActors, /*bInMultitrace=*/ false);
+
+				const b2Vec2 BoxStart(FPhysicsIntegration2D::ConvertUnrealVectorToBox(Start));
+				const b2Vec2 BoxEnd(FPhysicsIntegration2D::ConvertUnrealVectorToBox(End));
+
+				// Now check the length of the 2D query, we may end up having to do an AABB test instead if it's a ray cast along the 'Z' axis of the 2D scene
+				const b2Vec2 BoxDelta = BoxEnd - BoxStart;
+				if (BoxDelta.LengthSquared() < SMALL_NUMBER)
+				{
+					// Trace into the scene, do a point query instead
+					//@TODO: BOX2D: Implement RaycastSingle
+					BoxQuery.PointQuery(BoxWorld, BoxStart, Start, End, Params, OutHit);
+				}
+				else
+				{
+					// Cast the ray
+					BoxWorld->RayCast(&BoxQuery, BoxStart, BoxEnd);
+
+					// Accept or merge the results
+					if (BoxQuery.HasBlockingHit())
+					{
+						if (!bHaveBlockingHit || (BoxQuery.GetBestTime() < OutHit.Time))
+						{
+							BoxQuery.ConvertToHitResult(OutHit, Start, End, Params);
+							bHaveBlockingHit = true;
+						}
+					}
+				}
+			}
 		}
+#endif // WITH_BOX2D
 	}
 
 
@@ -428,7 +835,7 @@ bool RaycastSingle(const UWorld * World, struct FHitResult& OutHit, const FVecto
 		Hits.Add(OutHit);
 	}
 #if !(UE_BUILD_SHIPPING || UE_BUILD_TEST)
-	if((World->DebugDrawTraceTag != NAME_None) && (World->DebugDrawTraceTag == Params.TraceTag))
+	if ((World->DebugDrawTraceTag != NAME_None) && (World->DebugDrawTraceTag == Params.TraceTag))
 	{
 		DrawLineTraces(World, Start, End, Hits, DebugLineLifetime);	
 	}
@@ -440,9 +847,9 @@ bool RaycastSingle(const UWorld * World, struct FHitResult& OutHit, const FVecto
 	return bHaveBlockingHit;
 }
 
-bool RaycastMulti(const UWorld * World, TArray<struct FHitResult>& OutHits, const FVector& Start, const FVector& End, ECollisionChannel TraceChannel, const struct FCollisionQueryParams& Params, const struct FCollisionResponseParams& ResponseParams, const struct FCollisionObjectQueryParams& ObjectParams)
+bool RaycastMulti(const UWorld* World, TArray<struct FHitResult>& OutHits, const FVector& Start, const FVector& End, ECollisionChannel TraceChannel, const struct FCollisionQueryParams& Params, const struct FCollisionResponseParams& ResponseParams, const struct FCollisionObjectQueryParams& ObjectParams)
 {
-	if(World == NULL || World->GetPhysicsScene() == NULL)
+	if ((World == NULL) || (World->GetPhysicsScene() == NULL))
 	{
 		return false;
 	}
@@ -456,8 +863,10 @@ bool RaycastMulti(const UWorld * World, TArray<struct FHitResult>& OutHits, cons
 
 	FVector Delta = End - Start;
 	float DeltaMag = Delta.Size();
-	if(DeltaMag > KINDA_SMALL_NUMBER)
+	if (DeltaMag > KINDA_SMALL_NUMBER)
 	{
+		//@TODO: BOX2D: Implement RaycastMulti
+#if WITH_PHYSX
 		// Create filter data used to filter collisions
 		PxFilterData PFilter = CreateQueryFilterData(TraceChannel, Params.bTraceComplex, ResponseParams.CollisionResponse, ObjectParams, true);
 		PxSceneQueryFilterData PQueryFilterData(PFilter, PxSceneQueryFilterFlag::eSTATIC | PxSceneQueryFilterFlag::eDYNAMIC | PxSceneQueryFilterFlag::ePREFILTER);
@@ -564,6 +973,7 @@ bool RaycastMulti(const UWorld * World, TArray<struct FHitResult>& OutHits, cons
 		}
 
 		bHaveBlockingHit = bBlockingHit;
+#endif // WITH_PHYSX
 	}
 
 #if !(UE_BUILD_SHIPPING || UE_BUILD_TEST)
@@ -578,9 +988,13 @@ bool RaycastMulti(const UWorld * World, TArray<struct FHitResult>& OutHits, cons
 	return bHaveBlockingHit;
 }
 
+#endif //UE_WITH_PHYSICS
+
 //////////////////////////////////////////////////////////////////////////
 // GEOM SWEEP
 
+#if WITH_PHYSX
+
 bool GeomSweepTest(const UWorld * World, const PxGeometry& PGeom, const PxQuat& PGeomRot, FVector Start, FVector End, ECollisionChannel TraceChannel, const struct FCollisionQueryParams& Params, const struct FCollisionResponseParams& ResponseParams, const struct FCollisionObjectQueryParams& ObjectParams)
 {
 	if(World == NULL || World->GetPhysicsScene() == NULL)
diff --git a/Engine/Source/Runtime/Engine/Private/Collision/PhysXCollision.h b/Engine/Source/Runtime/Engine/Private/Collision/PhysXCollision.h
--- a/Engine/Source/Runtime/Engine/Private/Collision/PhysXCollision.h
+++ b/Engine/Source/Runtime/Engine/Private/Collision/PhysXCollision.h
@@ -131,20 +131,6 @@ public:
 
 	virtual PxSceneQueryHitType::Enum postFilter(const PxFilterData& filterData, const PxSceneQueryHit& hit) OVERRIDE;
 };
-// RAYCAST
-
-/** Trace a ray against the world and return if a blocking hit is found */
-bool RaycastTest(const UWorld * World, const FVector Start, const FVector End, ECollisionChannel TraceChannel, const struct FCollisionQueryParams& Params, const struct FCollisionResponseParams& ResponseParams, const struct FCollisionObjectQueryParams& ObjectParams = FCollisionObjectQueryParams::DefaultObjectQueryParam);
-
-/** Trace a ray against the world and return the first blocking hit */
-bool RaycastSingle(const UWorld * World, struct FHitResult& OutHit, const FVector Start, const FVector End, ECollisionChannel TraceChannel, const struct FCollisionQueryParams& Params, const struct FCollisionResponseParams& ResponseParams, const struct FCollisionObjectQueryParams& ObjectParams = FCollisionObjectQueryParams::DefaultObjectQueryParam);
-
-/** 
- *  Trace a ray against the world and return touching hits and then first blocking hit
- *  Results are sorted, so a blocking hit (if found) will be the last element of the array
- *  Only the single closest blocking result will be generated, no tests will be done after that
- */
-bool RaycastMulti(const UWorld * World, TArray<struct FHitResult>& OutHits, const FVector& Start, const FVector& End, ECollisionChannel TraceChannel, const struct FCollisionQueryParams& Params, const struct FCollisionResponseParams& ResponseParams, const struct FCollisionObjectQueryParams& ObjectParams = FCollisionObjectQueryParams::DefaultObjectQueryParam);
 
 // GEOM OVERLAP
 
@@ -187,3 +173,23 @@ PxTransform ConvertToPhysXCapsulePose(const FTransform& GeomPose);
 PxFilterData CreateQueryFilterData(const uint8 MyChannel, const bool bTraceComplex, const FCollisionResponseContainer & InCollisionResponseContainer, const struct FCollisionObjectQueryParams & ObjectParam, const bool bMultitrace);
 
 #endif // WITH_PHYX
+
+
+#if UE_WITH_PHYSICS
+
+// RAYCAST
+
+/** Trace a ray against the world and return if a blocking hit is found */
+bool RaycastTest(const UWorld* World, const FVector Start, const FVector End, ECollisionChannel TraceChannel, const struct FCollisionQueryParams& Params, const struct FCollisionResponseParams& ResponseParams, const struct FCollisionObjectQueryParams& ObjectParams = FCollisionObjectQueryParams::DefaultObjectQueryParam);
+
+/** Trace a ray against the world and return the first blocking hit */
+bool RaycastSingle(const UWorld* World, struct FHitResult& OutHit, const FVector Start, const FVector End, ECollisionChannel TraceChannel, const struct FCollisionQueryParams& Params, const struct FCollisionResponseParams& ResponseParams, const struct FCollisionObjectQueryParams& ObjectParams = FCollisionObjectQueryParams::DefaultObjectQueryParam);
+
+/** 
+ *  Trace a ray against the world and return touching hits and then first blocking hit
+ *  Results are sorted, so a blocking hit (if found) will be the last element of the array
+ *  Only the single closest blocking result will be generated, no tests will be done after that
+ */
+bool RaycastMulti(const UWorld* World, TArray<struct FHitResult>& OutHits, const FVector& Start, const FVector& End, ECollisionChannel TraceChannel, const struct FCollisionQueryParams& Params, const struct FCollisionResponseParams& ResponseParams, const struct FCollisionObjectQueryParams& ObjectParams = FCollisionObjectQueryParams::DefaultObjectQueryParam);
+
+#endif
\ No newline at end of file
diff --git a/Engine/Source/Runtime/Engine/Private/PhysicsEngine/PhysicsSettings.cpp b/Engine/Source/Runtime/Engine/Private/PhysicsEngine/PhysicsSettings.cpp
--- a/Engine/Source/Runtime/Engine/Private/PhysicsEngine/PhysicsSettings.cpp
+++ b/Engine/Source/Runtime/Engine/Private/PhysicsEngine/PhysicsSettings.cpp
@@ -3,17 +3,18 @@
 #include "EnginePrivate.h"
 
 UPhysicsSettings::UPhysicsSettings(const class FPostConstructInitializeProperties& PCIP)
-: Super(PCIP),
-DefaultGravityZ(-980.f),
-DefaultTerminalVelocity(4000.f),
-bEnableAsyncScene(false),
-MaxPhysicsDeltaTime(1.f / 30.f),
-bSubstepping(false),
-MaxSubstepDeltaTime(1.f / 60.f),
-MaxSubsteps(6),
-SyncSceneSmoothingFactor(0.0f),
-AsyncSceneSmoothingFactor(0.99f),
-InitialAverageFrameRate(1.f / 60.f)
+	: Super(PCIP)
+	, DefaultGravityZ(-980.f)
+	, DefaultTerminalVelocity(4000.f)
+	, bEnableAsyncScene(false)
+	, bEnable2DPhysics(false)
+	, MaxPhysicsDeltaTime(1.f / 30.f)
+	, bSubstepping(false)
+	, MaxSubstepDeltaTime(1.f / 60.f)
+	, MaxSubsteps(6)
+	, SyncSceneSmoothingFactor(0.0f)
+	, AsyncSceneSmoothingFactor(0.99f)
+	, InitialAverageFrameRate(1.f / 60.f)
 {
 }
 
diff --git a/Engine/Source/Runtime/Engine/Private/PhysicsEngine2D/Box2DIntegration.h b/Engine/Source/Runtime/Engine/Private/PhysicsEngine2D/Box2DIntegration.h
--- a/Engine/Source/Runtime/Engine/Private/PhysicsEngine2D/Box2DIntegration.h
+++ b/Engine/Source/Runtime/Engine/Private/PhysicsEngine2D/Box2DIntegration.h
@@ -102,6 +102,11 @@ public:
 		return FVector(0.0f, FMath::RadiansToDegrees<float>(AngularVelocity), 0.0f);
 	}
 
+	static inline float ConvertUnrealVectorToPerpendicularDistance(const FVector& Vector)
+	{
+		return Vector.Y;
+	}
+
 	// Finds the scene associated with the specified UWorld
 	static TSharedPtr<FPhysicsScene2D> FindAssociatedScene(UWorld* Source);
 
