Engine: Change UKismetSystemLibrary::GetDisplayName to use UActorComponent::GetReadableName on components

diff --git a/Engine/Source/Runtime/Engine/Private/KismetSystemLibrary.cpp b/Engine/Source/Runtime/Engine/Private/KismetSystemLibrary.cpp
--- a/Engine/Source/Runtime/Engine/Private/KismetSystemLibrary.cpp
+++ b/Engine/Source/Runtime/Engine/Private/KismetSystemLibrary.cpp
@@ -31,6 +31,11 @@ FString UKismetSystemLibrary::GetDisplayName(const UObject* Object)
 	}
 #endif
 
+	if (const UActorComponent* Component = Cast<const UActorComponent>(Object))
+	{
+		return Component->GetReadableName();
+	}
+
 	return Object ? Object->GetName() : FString();
 }
 
