Physics: Add support for 2D RaycastTest queries, gated by UPhysicsSettings::bEnable2DPhysics

diff --git a/Engine/Source/Runtime/Engine/Private/Collision/PhysXCollision.cpp b/Engine/Source/Runtime/Engine/Private/Collision/PhysXCollision.cpp
--- a/Engine/Source/Runtime/Engine/Private/Collision/PhysXCollision.cpp
+++ b/Engine/Source/Runtime/Engine/Private/Collision/PhysXCollision.cpp
@@ -138,6 +138,29 @@ public:
 
 	void PointQuery(const b2World* BoxWorld, const b2Vec2& BoxStart, const FVector& Start, const FVector& End, const FCollisionQueryParams& Params, FHitResult& OutResult);
 
+	bool RayTest(const b2World* BoxWorld, const FVector& Start, const FVector& End)
+	{
+		const b2Vec2 BoxStart(FPhysicsIntegration2D::ConvertUnrealVectorToBox(Start));
+		const b2Vec2 BoxEnd(FPhysicsIntegration2D::ConvertUnrealVectorToBox(End));
+
+		// Now check the length of the 2D query, we may end up having to do an AABB test instead if it's a ray cast along the 'Z' axis of the 2D scene
+		const b2Vec2 BoxDelta = BoxEnd - BoxStart;
+		if (BoxDelta.LengthSquared() < SMALL_NUMBER)
+		{
+			// Trace into the scene, do a point query instead
+			b2AABB QueryBox;
+			QueryBox.lowerBound = BoxStart;
+			QueryBox.upperBound = BoxStart;
+			BoxWorld->QueryAABB(this, QueryBox);
+		}
+		else
+		{
+			// Cast the ray
+			BoxWorld->RayCast(this, BoxStart, BoxEnd);
+		}
+
+		return bHasBlockingHit;
+	}
 
 	// Determines how the query and the fixture interact (either no interaction or touching/blocking each other)
 	EResponse CalcQueryHitType(b2Fixture* Fixture) const;
@@ -704,10 +727,13 @@ bool RaycastTest(const UWorld* World, const FVector Start, const FVector End, EC
 #endif // WITH_PHYSX
 
 #if WITH_BOX2D
-		//@TODO: BOX2D: Implement RaycastTest
-		if (GetDefault<UPhysicsSettings>()->bEnable2DPhysics)
+		if (!bHaveBlockingHit && GetDefault<UPhysicsSettings>()->bEnable2DPhysics)
 		{
-			
+			if (const b2World* BoxWorld = FPhysicsIntegration2D::FindAssociatedWorld(const_cast<UWorld*>(World)))
+			{
+				FRaycastHelperBox2D BoxQuery(TraceChannel, Params.bTraceComplex, ResponseParams.CollisionResponse, ObjectParams, Params.IgnoreActors, /*bInMultitrace=*/ false);
+				bHaveBlockingHit = BoxQuery.RayTest(BoxWorld, Start, End);
+			}
 		}
 #endif // WITH_BOX2D
 	}
@@ -804,7 +830,6 @@ bool RaycastSingle(const UWorld* World, struct FHitResult& OutHit, const FVector
 				if (BoxDelta.LengthSquared() < SMALL_NUMBER)
 				{
 					// Trace into the scene, do a point query instead
-					//@TODO: BOX2D: Implement RaycastSingle
 					BoxQuery.PointQuery(BoxWorld, BoxStart, Start, End, Params, OutHit);
 				}
 				else
