Physics: Add support for adding forces/torque/impulses to 2D body instances (mostly untested)

diff --git a/Engine/Source/Runtime/Engine/Private/PhysicsEngine/BodyInstance.cpp b/Engine/Source/Runtime/Engine/Private/PhysicsEngine/BodyInstance.cpp
--- a/Engine/Source/Runtime/Engine/Private/PhysicsEngine/BodyInstance.cpp
+++ b/Engine/Source/Runtime/Engine/Private/PhysicsEngine/BodyInstance.cpp
@@ -2581,7 +2581,16 @@ void FBodyInstance::AddForce(const FVector& Force, bool bAllowSubstepping)
 	}
 #endif // WITH_PHYSX
 
-	//@TODO: BOX2D: Implement AddForce
+#if WITH_BOX2D
+	if (BodyInstancePtr)
+	{
+		if (!Force.IsNearlyZero())
+		{
+			const b2Vec2 Force2D = FPhysicsIntegration2D::ConvertUnrealVectorToBox(Force);
+			BodyInstancePtr->ApplyForceToCenter(Force2D, /*wake=*/ true);
+		}
+	}
+#endif
 }
 
 void FBodyInstance::AddForceAtPosition(const FVector& Force, const FVector& Position, bool bAllowSubstepping)
@@ -2596,7 +2605,17 @@ void FBodyInstance::AddForceAtPosition(const FVector& Force, const FVector& Posi
 	}
 #endif // WITH_PHYSX
 
-	//@TODO: BOX2D: Implement AddForceAtPosition
+#if WITH_BOX2D
+	if (BodyInstancePtr)
+	{
+		if (!Force.IsNearlyZero())
+		{
+			const b2Vec2 Force2D = FPhysicsIntegration2D::ConvertUnrealVectorToBox(Force);
+			const b2Vec2 Position2D = FPhysicsIntegration2D::ConvertUnrealVectorToBox(Position);
+			BodyInstancePtr->ApplyForce(Force2D, Position2D, /*wake=*/ true);
+		}
+	}
+#endif
 }
 
 void FBodyInstance::AddTorque(const FVector& Torque, bool bAllowSubstepping)
@@ -2611,7 +2630,16 @@ void FBodyInstance::AddTorque(const FVector& Torque, bool bAllowSubstepping)
 	}
 #endif // WITH_PHYSX
 
-	//@TODO: BOX2D: Implement AddTorque
+#if WITH_BOX2D
+	if (BodyInstancePtr)
+	{
+		const float Torque1D = FPhysicsIntegration2D::ConvertUnrealTorqueToBox(Torque);
+		if (!FMath::IsNearlyZero(Torque1D))
+		{
+			BodyInstancePtr->ApplyTorque(Torque1D, /*wake=*/ true);
+		}
+	}
+#endif
 }
 
 void FBodyInstance::AddImpulse(const FVector& Impulse, bool bVelChange)
@@ -2625,7 +2653,16 @@ void FBodyInstance::AddImpulse(const FVector& Impulse, bool bVelChange)
 	}
 #endif // WITH_PHYSX
 
-	//@TODO: BOX2D: Implement AddImpulse
+#if WITH_BOX2D
+	if (BodyInstancePtr)
+	{
+		if (!Impulse.IsNearlyZero())
+		{
+			const b2Vec2 Impulse2D = FPhysicsIntegration2D::ConvertUnrealVectorToBox(Impulse);
+			BodyInstancePtr->ApplyLinearImpulse(Impulse2D, BodyInstancePtr->GetWorldCenter(), /*wake=*/ true);
+		}
+	}
+#endif
 }
 
 void FBodyInstance::AddImpulseAtPosition(const FVector& Impulse, const FVector& Position)
@@ -2639,7 +2676,17 @@ void FBodyInstance::AddImpulseAtPosition(const FVector& Impulse, const FVector&
 	}
 #endif // WITH_PHYSX
 
-	//@TODO: BOX2D: Implement AddImpulseAtPosition
+#if WITH_BOX2D
+	if (BodyInstancePtr)
+	{
+		if (!Impulse.IsNearlyZero())
+		{
+			const b2Vec2 Impulse2D = FPhysicsIntegration2D::ConvertUnrealVectorToBox(Impulse);
+			const b2Vec2 Position2D = FPhysicsIntegration2D::ConvertUnrealVectorToBox(Position);
+			BodyInstancePtr->ApplyLinearImpulse(Impulse2D, Position2D, /*wake=*/ true);
+		}
+	}
+#endif
 }
 
 void FBodyInstance::SetInstanceNotifyRBCollision(bool bNewNotifyCollision)
@@ -2648,7 +2695,6 @@ void FBodyInstance::SetInstanceNotifyRBCollision(bool bNewNotifyCollision)
 	UpdatePhysicsFilterData();
 }
 
-
 void FBodyInstance::SetEnableGravity(bool bInGravityEnabled)
 {
 	if (bEnableGravity != bInGravityEnabled)
@@ -2677,6 +2723,48 @@ void FBodyInstance::SetEnableGravity(bool bInGravityEnabled)
 }
 
 
+#if WITH_BOX2D
+void AddRadialImpulseToBox2DBodyInstance(class b2Body* BodyInstancePtr, const FVector& Origin, float Radius, float Strength, ERadialImpulseFalloff Falloff, bool bImpulse, bool bMassIndependent)
+{
+	const b2Vec2 CenterOfMass2D = BodyInstancePtr->GetWorldCenter();
+	const b2Vec2 Origin2D = FPhysicsIntegration2D::ConvertUnrealVectorToBox(Origin);
+
+	// If the center of mass is outside of the blast radius, do nothing.
+	b2Vec2 Delta2D = CenterOfMass2D - Origin2D;
+	const float DistanceFromBlast = Delta2D.Normalize() * UnrealUnitsPerMeter;
+	if (DistanceFromBlast > Radius)
+	{
+		return;
+	}
+	
+	// If using linear falloff, scale with distance.
+	const float EffectiveStrength = (Falloff == RIF_Linear) ? (Strength * (1.0f - (DistanceFromBlast / Radius))) : Strength;
+
+	b2Vec2 ForceOrImpulse2D = Delta2D;
+	ForceOrImpulse2D *= EffectiveStrength;
+
+	if (bMassIndependent)
+	{
+		const float Mass = BodyInstancePtr->GetMass();
+		if (!FMath::IsNearlyZero(Mass))
+		{
+			ForceOrImpulse2D *= Mass;
+		}
+	}
+
+	if (bImpulse)
+	{
+		BodyInstancePtr->ApplyLinearImpulse(ForceOrImpulse2D, CenterOfMass2D, /*wake=*/ true);
+	}
+	else
+	{
+		BodyInstancePtr->ApplyForceToCenter(ForceOrImpulse2D, /*wake=*/ true);
+	}
+}
+#endif
+
+
+
 void FBodyInstance::AddRadialImpulseToBody(const FVector& Origin, float Radius, float Strength, uint8 Falloff, bool bVelChange)
 {
 #if WITH_PHYSX
@@ -2687,7 +2775,12 @@ void FBodyInstance::AddRadialImpulseToBody(const FVector& Origin, float Radius,
 	}
 #endif // WITH_PHYSX
 
-//@TODO: BOX2D: Implement AddRadialImpulseToBody
+#if WITH_BOX2D
+	if (BodyInstancePtr)
+	{
+		AddRadialImpulseToBox2DBodyInstance(BodyInstancePtr, Origin, Radius, Strength, static_cast<ERadialImpulseFalloff>(Falloff), /*bImpulse=*/ true, /*bMassIndependent=*/ bVelChange);
+	}
+#endif
 }
 
 void FBodyInstance::AddRadialForceToBody(const FVector& Origin, float Radius, float Strength, uint8 Falloff)
@@ -2700,10 +2793,14 @@ void FBodyInstance::AddRadialForceToBody(const FVector& Origin, float Radius, fl
 	}
 #endif // WITH_PHYSX
 
-	//@TODO: BOX2D: Implement AddradialForceToBody
+#if WITH_BOX2D
+	if (BodyInstancePtr)
+	{
+		AddRadialImpulseToBox2DBodyInstance(BodyInstancePtr, Origin, Radius, Strength, static_cast<ERadialImpulseFalloff>(Falloff), /*bImpulse=*/ false, /*bMassIndependent=*/ false);
+	}
+#endif
 }
 
-
 FString FBodyInstance::GetBodyDebugName() const
 {
 	FString DebugName;
diff --git a/Engine/Source/Runtime/Engine/Private/PhysicsEngine2D/Box2DIntegration.h b/Engine/Source/Runtime/Engine/Private/PhysicsEngine2D/Box2DIntegration.h
--- a/Engine/Source/Runtime/Engine/Private/PhysicsEngine2D/Box2DIntegration.h
+++ b/Engine/Source/Runtime/Engine/Private/PhysicsEngine2D/Box2DIntegration.h
@@ -102,6 +102,13 @@ public:
 		return FVector(0.0f, FMath::RadiansToDegrees<float>(AngularVelocity), 0.0f);
 	}
 
+	// Converts Unreal torque to torque about the 2D 'Z' axis in Nm
+	static inline float ConvertUnrealTorqueToBox(const FVector& Torque3D)
+	{
+		const FVector ProjectedTorque = Torque3D.ProjectOnTo(FVector(0.0f, -1.0f, 0.0f));
+		return ProjectedTorque.Size() / UnrealUnitsPerMeter;
+	}
+
 	static inline float ConvertUnrealVectorToPerpendicularDistance(const FVector& Vector)
 	{
 		return Vector.Y;
