Fixed FPrimitiveUniformShaderParameters GetPrimitiveUniformShaderParameters creating NaN when inverting Matrices with zero scale.
TTP #337967

diff --git a/Engine/Source/Runtime/Engine/Public/SceneManagement.h b/Engine/Source/Runtime/Engine/Public/SceneManagement.h
--- a/Engine/Source/Runtime/Engine/Public/SceneManagement.h
+++ b/Engine/Source/Runtime/Engine/Public/SceneManagement.h
@@ -1158,7 +1158,7 @@ inline FPrimitiveUniformShaderParameters GetPrimitiveUniformShaderParameters(
 {
 	FPrimitiveUniformShaderParameters Result;
 	Result.LocalToWorld = LocalToWorld;
-	Result.WorldToLocal = LocalToWorld.Inverse();
+	Result.WorldToLocal = LocalToWorld.InverseSafe();
 	Result.ObjectWorldPositionAndRadius = FVector4(WorldBounds.Origin, WorldBounds.SphereRadius);
 	Result.ObjectBounds = WorldBounds.BoxExtent;
 	Result.LocalObjectBoundsMin = LocalBounds.GetBoxExtrema(0); // 0 == minimum
