Fix for Pinch to Zoom

diff --git a/Engine/Source/Runtime/Engine/Private/Slate/SceneViewport.cpp b/Engine/Source/Runtime/Engine/Private/Slate/SceneViewport.cpp
--- a/Engine/Source/Runtime/Engine/Private/Slate/SceneViewport.cpp
+++ b/Engine/Source/Runtime/Engine/Private/Slate/SceneViewport.cpp
@@ -582,7 +582,7 @@ FReply FSceneViewport::OnTouchMoved( const FGeometry& MyGeometry, const FPointer
 		// Switch to the viewport clients world before processing input
 		FScopedConditionalWorldSwitcher WorldSwitcher( ViewportClient );
 
-		const FVector2D TouchPosition = MyGeometry.AbsoluteToLocal(TouchEvent.GetLastScreenSpacePosition());
+		const FVector2D TouchPosition = MyGeometry.AbsoluteToLocal(TouchEvent.GetScreenSpacePosition());
 
 		if( !ViewportClient->InputTouch( this, TouchEvent.GetUserIndex(), TouchEvent.GetPointerIndex(), ETouchType::Moved, TouchPosition, FDateTime::Now(), TouchEvent.GetTouchpadIndex()) )
 		{
