Removed inner worlds processing from level streaming

diff --git a/Engine/Source/Runtime/Engine/Classes/Engine/Level.h b/Engine/Source/Runtime/Engine/Classes/Engine/Level.h
--- a/Engine/Source/Runtime/Engine/Classes/Engine/Level.h
+++ b/Engine/Source/Runtime/Engine/Classes/Engine/Level.h
@@ -394,10 +394,6 @@ class ULevel : public ULevelBase, public IInterface_AssetUserData
 	uint32										bAlreadyRoutedActorInitialize:1;
 	/** Whether we already sorted the actor list.											*/
 	uint32										bAlreadySortedActorList:1;
-	/** Whether one or more streaming levels requested to show this level					*/
-	uint32										bHasShowRequest:1;
-	/** Whether one or more streaming levels requested to hide this level					*/
-	uint32										bHasHideRequest:1;
 	/** Whether this level is in the process of being associated with its world				*/
 	uint32										bIsAssociatingLevel:1;
 	/** Whether this level should be fully added to the world before rendering his components	*/
@@ -450,9 +446,6 @@ private:
 
 	TArray<FPendingAutoReceiveInputActor> PendingAutoReceiveInputActors;
 
-	/** Number of streaming levels referring this level	*/
-	int32 NumStreamingLevelRefs;
-
 public:
 	/** Called when a level package has been dirtied. */
 	ENGINE_API static FSimpleMulticastDelegate LevelDirtiedEvent;
@@ -672,11 +665,6 @@ public:
 	/** Increments number of steaming objects referring this level */
 	ENGINE_API void IncStreamingLevelRefs();
 	
-	/** Decrements number of steaming objects referring this level */
-	ENGINE_API void DecStreamingLevelRefs();
-
-	/** Returns number of steaming objects referring this level */
-	ENGINE_API int32 GetStreamingLevelRefs() const;
 	
 	// Begin IInterface_AssetUserData Interface
 	virtual void AddAssetUserData(UAssetUserData* InUserData) OVERRIDE;
diff --git a/Engine/Source/Runtime/Engine/Classes/Engine/LevelStreaming.h b/Engine/Source/Runtime/Engine/Classes/Engine/LevelStreaming.h
--- a/Engine/Source/Runtime/Engine/Classes/Engine/LevelStreaming.h
+++ b/Engine/Source/Runtime/Engine/Classes/Engine/LevelStreaming.h
@@ -97,7 +97,7 @@ class ULevelStreaming : public UObject
 
 	/** This streaming level was not found																						*/
 	uint32 bFailedToLoad:1;
-
+	
 	/** Whether this level should be visible in the Editor																		*/
 	UPROPERTY()
 	uint32 bShouldBeVisibleInEditor:1;
@@ -144,6 +144,9 @@ class ULevelStreaming : public UObject
 	UPROPERTY()
 	TArray<FString> Keywords;
 
+	/** Requested LOD */
+	int32 LevelLODIndex;
+
 	// Begin UObject Interface
 	virtual void PostLoad() OVERRIDE;
 	virtual void Serialize( FArchive& Ar ) OVERRIDE;
@@ -193,14 +196,8 @@ class ULevelStreaming : public UObject
 	class ULevel* GetLoadedLevel() const {	return LoadedLevel; }
 	
 	/** Sets the LoadedLevel value to NULL */
-	void ClearLoadedLevel() { SetLoadedLevel(NULL); }
+	void ClearLoadedLevel() { SetLoadedLevel(nullptr); }
 	
-	/** Gets current streaming level LOD index  */
-	int32 GetLODIndex(UWorld* PersistentWorld) const;	
-
-	/** Sets current streaming level LOD index  */
-	void SetLODIndex(UWorld* PersistentWorld, int32 LODIndex);	
-		
 #if WITH_EDITOR
 	/** Override Pre/PostEditUndo functions to handle editor transform */
 	virtual void PreEditUndo();
@@ -276,10 +273,10 @@ class ULevelStreaming : public UObject
 	
 private:
 	/** @return Name of the LOD level package used for loading.																		*/
-	FName GetLODPackageName(UWorld* PersistentWorld) const;
+	FName GetLODPackageName() const;
 
 	/** @return Name of the LOD package on disk to load to the new package named PackageName, Name_None otherwise					*/
-	FName GetLODPackageNameToLoad(UWorld* PersistentWorld) const;
+	FName GetLODPackageNameToLoad() const;
 
 	/** 
 	 * Try to find loaded level in memory, issue a loading request otherwise
@@ -294,38 +291,10 @@ private:
 	/** Sets the value of LoadedLevel */
 	void SetLoadedLevel(class ULevel* Level)
 	{ 
-		// Pending level should be unloaded or hidden at this point
-		check(PendingUnloadLevel == nullptr || PendingUnloadLevel->bIsVisible == false || Level == PendingUnloadLevel);
-
-		SetPendingUnloadLevel(LoadedLevel);
-
-		if (LoadedLevel)
-		{
-			LoadedLevel->DecStreamingLevelRefs();
-		}
-
+		// Pending level should be unloaded at this point
+		check(PendingUnloadLevel == nullptr);
+		PendingUnloadLevel = LoadedLevel;
 		LoadedLevel = Level;
-
-		if (LoadedLevel)
-		{
-			LoadedLevel->IncStreamingLevelRefs();
-		}
-	}
-
-	/** Sets the value of PendingUnloadLevel */
-	void SetPendingUnloadLevel(class ULevel* Level)
-	{ 
-		if (PendingUnloadLevel)
-		{
-			PendingUnloadLevel->DecStreamingLevelRefs();
-		}
-
-		PendingUnloadLevel = Level;
-
-		if (PendingUnloadLevel)
-		{
-			PendingUnloadLevel->IncStreamingLevelRefs();
-		}
 	}
 
 	void DiscardPendingUnloadLevel(UWorld* PersistentWorld);
diff --git a/Engine/Source/Runtime/Engine/Classes/Engine/World.h b/Engine/Source/Runtime/Engine/Classes/Engine/World.h
--- a/Engine/Source/Runtime/Engine/Classes/Engine/World.h
+++ b/Engine/Source/Runtime/Engine/Classes/Engine/World.h
@@ -868,12 +868,6 @@ public:
 	UPROPERTY()
 	class UWorldComposition* WorldComposition;
 	
-	/** Streaming level package name to LOD index. 
-	 *	LOD index stored in the persistent world to support consistent LOD between nested streaming levels which could be loaded during gameplay
-	 *  LOD changes affects all streaming levels referring the same level package
-	 */
-	TMap<FName, int32>		StreamingLevelsLOD;
-
 	/** Whether we currently flushing level streaming state */ 
 	bool bFlushingLevelStreaming;
 
diff --git a/Engine/Source/Runtime/Engine/Private/EngineUtils.cpp b/Engine/Source/Runtime/Engine/Private/EngineUtils.cpp
--- a/Engine/Source/Runtime/Engine/Private/EngineUtils.cpp
+++ b/Engine/Source/Runtime/Engine/Private/EngineUtils.cpp
@@ -335,7 +335,7 @@ TArray<FSubLevelStatus> GetSubLevelsStatus( UWorld* World )
 			ULevel* Level = LevelStreaming->GetLoadedLevel();
 			FSubLevelStatus LevelStatus = {};
 			LevelStatus.PackageName = LevelStreaming->PackageName;
-			LevelStatus.LODIndex	= LevelStreaming->GetLODIndex(World);
+			LevelStatus.LODIndex	= LevelStreaming->LevelLODIndex;
 
 			if( Level != NULL )
 			{
diff --git a/Engine/Source/Runtime/Engine/Private/GameMode.cpp b/Engine/Source/Runtime/Engine/Private/GameMode.cpp
--- a/Engine/Source/Runtime/Engine/Private/GameMode.cpp
+++ b/Engine/Source/Runtime/Engine/Private/GameMode.cpp
@@ -1300,7 +1300,7 @@ void AGameMode::ReplicateStreamingStatus(APlayerController* PC)
 						TheLevel->bShouldBeLoaded,
 						TheLevel->bShouldBeVisible,
 						TheLevel->bShouldBlockOnLoad,
-						TheLevel->GetLODIndex(GetWorld()));
+						TheLevel->LevelLODIndex);
 				}
 			}
 			PC->ClientFlushLevelStreaming();
diff --git a/Engine/Source/Runtime/Engine/Private/Level.cpp b/Engine/Source/Runtime/Engine/Private/Level.cpp
--- a/Engine/Source/Runtime/Engine/Private/Level.cpp
+++ b/Engine/Source/Runtime/Engine/Private/Level.cpp
@@ -1892,26 +1892,6 @@ void ULevel::PushPendingAutoReceiveInput(APlayerController* InPlayerController)
 	}
 }
 
-void ULevel::IncStreamingLevelRefs()
-{
-	if (NumStreamingLevelRefs == 0)
-	{
-		// Remove from pending unload list if any
-		FLevelStreamingGCHelper::CancelUnloadRequest(this);
-	}
-	NumStreamingLevelRefs++;
-}
-	
-void ULevel::DecStreamingLevelRefs()
-{
-	NumStreamingLevelRefs--;
-}
-
-int32 ULevel::GetStreamingLevelRefs() const
-{
-	return NumStreamingLevelRefs;
-}
-
 void ULevel::AddAssetUserData(UAssetUserData* InUserData)
 {
 	if(InUserData != NULL)
diff --git a/Engine/Source/Runtime/Engine/Private/LevelStreaming.cpp b/Engine/Source/Runtime/Engine/Private/LevelStreaming.cpp
--- a/Engine/Source/Runtime/Engine/Private/LevelStreaming.cpp
+++ b/Engine/Source/Runtime/Engine/Private/LevelStreaming.cpp
@@ -248,32 +248,16 @@ void ULevelStreaming::Serialize( FArchive& Ar )
 	}
 }
 
-int32 ULevelStreaming::GetLODIndex(UWorld* PersistentWorld) const
+FName ULevelStreaming::GetLODPackageName() const
 {
-	check(PersistentWorld);
-	auto Entry = PersistentWorld->StreamingLevelsLOD.Find(PackageName);
-	return Entry ? *Entry : INDEX_NONE;
+	return LODPackageNames.IsValidIndex(LevelLODIndex) ? LODPackageNames[LevelLODIndex] : PackageName;
 }
 
-void ULevelStreaming::SetLODIndex(UWorld* PersistentWorld, int32 LODIndex)
+FName ULevelStreaming::GetLODPackageNameToLoad() const
 {
-	check(PersistentWorld);
-	PersistentWorld->StreamingLevelsLOD.Add(PackageName, LODIndex);
-}
-
-FName ULevelStreaming::GetLODPackageName(UWorld* PersistentWorld) const
-{
-	int32 LODCurrentIdx = GetLODIndex(PersistentWorld);
-	return LODPackageNames.IsValidIndex(LODCurrentIdx) ? LODPackageNames[LODCurrentIdx] : PackageName;
-}
-
-FName ULevelStreaming::GetLODPackageNameToLoad(UWorld* PersistentWorld) const
-{
-	int32 LODCurrentIdx = GetLODIndex(PersistentWorld);
-
-	if (LODPackageNames.IsValidIndex(LODCurrentIdx))
+	if (LODPackageNames.IsValidIndex(LevelLODIndex))
 	{
-		return LODPackageNamesToLoad.IsValidIndex(LODCurrentIdx) ? LODPackageNamesToLoad[LODCurrentIdx] : NAME_None;
+		return LODPackageNamesToLoad.IsValidIndex(LevelLODIndex) ? LODPackageNamesToLoad[LevelLODIndex] : NAME_None;
 	}
 	else
 	{
@@ -292,7 +276,8 @@ void ULevelStreaming::DiscardPendingUnloadLevel(UWorld* PersistentWorld)
 
 		if (!PendingUnloadLevel->bIsVisible)
 		{
-			SetPendingUnloadLevel(NULL);
+			FLevelStreamingGCHelper::RequestUnload(PendingUnloadLevel);
+			PendingUnloadLevel = nullptr;
 		}
 	}
 }
@@ -312,8 +297,8 @@ bool ULevelStreaming::RequestLevel(UWorld* PersistentWorld, bool bAllowLevelLoad
 	}
 	
 	// Package name we want to load
-	FName DesiredPackageName = PersistentWorld->IsGameWorld() ? GetLODPackageName(PersistentWorld) : PackageName;
-	FName DesiredPackageNameToLoad = PersistentWorld->IsGameWorld() ? GetLODPackageNameToLoad(PersistentWorld) : PackageNameToLoad;
+	FName DesiredPackageName = PersistentWorld->IsGameWorld() ? GetLODPackageName() : PackageName;
+	FName DesiredPackageNameToLoad = PersistentWorld->IsGameWorld() ? GetLODPackageNameToLoad() : PackageNameToLoad;
 
 	// Check if currently loaded level is what we want right now
 	if (LoadedLevel != NULL && 
@@ -454,12 +439,9 @@ bool ULevelStreaming::RequestLevel(UWorld* PersistentWorld, bool bAllowLevelLoad
 				NULL, NAME_None, *PackageNameToLoadFrom
 				).SetPIEInstanceID(PIEInstanceID);
 
-			// Are we sublevel of a persistent level or a sublevel of a streamed-in level?
-			const bool bInnerLevel = PersistentWorld != GetOuter(); 
-
 			// streamingServer: server loads everything?
 			// Editor immediately blocks on load and we also block if background level streaming is disabled.
-			if (bBlockOnLoad || (!bInnerLevel && ShouldBeAlwaysLoaded()))
+			if (bBlockOnLoad || ShouldBeAlwaysLoaded())
 			{
 				// Finish all async loading.
 				FlushAsyncLoading( NAME_None );
@@ -669,13 +651,6 @@ void ULevelStreaming::BroadcastLevelLoadedStatus(UWorld* PersistentWorld, FName
 			}
 		}
 	}
-
-	// traverse inner worlds as well, skip persistent level
-	for (int32 LevelIdx = 1; LevelIdx < PersistentWorld->GetNumLevels(); ++LevelIdx)
-	{
-		UWorld* InnerWorld = Cast<UWorld>(PersistentWorld->GetLevel(LevelIdx)->GetOuter());
-		BroadcastLevelLoadedStatus(InnerWorld, LevelPackageName, bLoaded);
-	}
 }
 	
 void ULevelStreaming::BroadcastLevelVisibleStatus(UWorld* PersistentWorld, FName LevelPackageName, bool bVisible)
@@ -694,13 +669,6 @@ void ULevelStreaming::BroadcastLevelVisibleStatus(UWorld* PersistentWorld, FName
 			}
 		}
 	}
-
-	// traverse inner worlds as well, skip persistent level
-	for (int32 LevelIdx = 1; LevelIdx < PersistentWorld->GetNumLevels(); ++LevelIdx)
-	{
-		UWorld* InnerWorld = Cast<UWorld>(PersistentWorld->GetLevel(LevelIdx)->GetOuter());
-		BroadcastLevelVisibleStatus(InnerWorld, LevelPackageName, bVisible);
-	}
 }
 
 void ULevelStreaming::RenameForPIE(int32 PIEInstanceID)
@@ -837,6 +805,7 @@ ULevelStreaming::ULevelStreaming(const class FPostConstructInitializeProperties&
 	LevelTransform = FTransform::Identity;
 	MinTimeBetweenVolumeUnloadRequests = 2.0f;
 	bDrawOnLevelStatusMap = true;
+	LevelLODIndex = INDEX_NONE;
 }
 
 #if WITH_EDITOR
diff --git a/Engine/Source/Runtime/Engine/Private/LevelTick.cpp b/Engine/Source/Runtime/Engine/Private/LevelTick.cpp
--- a/Engine/Source/Runtime/Engine/Private/LevelTick.cpp
+++ b/Engine/Source/Runtime/Engine/Private/LevelTick.cpp
@@ -636,7 +636,7 @@ void UWorld::ProcessLevelStreamingVolumes(FVector* OverrideViewLocation)
 			bool bOriginalShouldBeLoaded						= LevelStreamingObject->bShouldBeLoaded;
 			bool bOriginalShouldBeVisible						= LevelStreamingObject->bShouldBeVisible;
 			bool bOriginalShouldBlockOnLoad						= LevelStreamingObject->bShouldBlockOnLoad;
-			int32 bOriginalLODIndex								= LevelStreamingObject->GetLODIndex(this);
+			int32 bOriginalLODIndex								= LevelStreamingObject->LevelLODIndex;
 
 			if( bShouldBeLoaded || bShouldAffectLoading )
 			{
@@ -665,7 +665,7 @@ void UWorld::ProcessLevelStreamingVolumes(FVector* OverrideViewLocation)
 				if( bOriginalShouldBeLoaded		!= LevelStreamingObject->bShouldBeLoaded
 				||	bOriginalShouldBeVisible	!= LevelStreamingObject->bShouldBeVisible 
 				||	bOriginalShouldBlockOnLoad	!= LevelStreamingObject->bShouldBlockOnLoad
-				||  bOriginalLODIndex			!= LevelStreamingObject->GetLODIndex(this))
+				||  bOriginalLODIndex			!= LevelStreamingObject->LevelLODIndex)
 				{
 					for( FConstPlayerControllerIterator Iterator = GetPlayerControllerIterator(); Iterator; ++Iterator )
 					{
@@ -675,7 +675,7 @@ void UWorld::ProcessLevelStreamingVolumes(FVector* OverrideViewLocation)
 								LevelStreamingObject->bShouldBeLoaded, 
 								LevelStreamingObject->bShouldBeVisible,
 								LevelStreamingObject->bShouldBlockOnLoad,
-								LevelStreamingObject->GetLODIndex(this));
+								LevelStreamingObject->LevelLODIndex);
 					}
 				}
 			}
diff --git a/Engine/Source/Runtime/Engine/Private/PlayerController.cpp b/Engine/Source/Runtime/Engine/Private/PlayerController.cpp
--- a/Engine/Source/Runtime/Engine/Private/PlayerController.cpp
+++ b/Engine/Source/Runtime/Engine/Private/PlayerController.cpp
@@ -167,7 +167,7 @@ void APlayerController::ClientUpdateLevelStreamingStatus_Implementation(FName Pa
 						LevelStreamingObject->bShouldBeLoaded		= bNewShouldBeLoaded;
 						LevelStreamingObject->bShouldBeVisible		= bNewShouldBeVisible;
 						LevelStreamingObject->bShouldBlockOnLoad	= bNewShouldBlockOnLoad;
-						LevelStreamingObject->SetLODIndex(GetWorld(), LODIndex);
+						LevelStreamingObject->LevelLODIndex			= LODIndex;
 					}
 					else
 					{
diff --git a/Engine/Source/Runtime/Engine/Private/UnrealEngine.cpp b/Engine/Source/Runtime/Engine/Private/UnrealEngine.cpp
--- a/Engine/Source/Runtime/Engine/Private/UnrealEngine.cpp
+++ b/Engine/Source/Runtime/Engine/Private/UnrealEngine.cpp
@@ -9190,7 +9190,7 @@ bool UEngine::CommitMapChange( FWorldContext &Context )
 									StreamingLevel->bShouldBeLoaded, 
 									StreamingLevel->bShouldBeVisible,
 									StreamingLevel->bShouldBlockOnLoad,
-									StreamingLevel->GetLODIndex(Context.World()));							
+									StreamingLevel->LevelLODIndex);							
 						}
 					}
 #endif // WITH_SERVER_CODE
@@ -9227,7 +9227,7 @@ bool UEngine::CommitMapChange( FWorldContext &Context )
 						{
 							LevelStreamingObject->bShouldBeLoaded	= Context.PendingLevelStreamingStatusUpdates[i].bShouldBeLoaded;
 							LevelStreamingObject->bShouldBeVisible	= Context.PendingLevelStreamingStatusUpdates[i].bShouldBeVisible;
-							LevelStreamingObject->SetLODIndex(Context.World(), Context.PendingLevelStreamingStatusUpdates[i].LODIndex);
+							LevelStreamingObject->LevelLODIndex		= Context.PendingLevelStreamingStatusUpdates[i].LODIndex;
 						}
 						else
 						{
diff --git a/Engine/Source/Runtime/Engine/Private/World.cpp b/Engine/Source/Runtime/Engine/Private/World.cpp
--- a/Engine/Source/Runtime/Engine/Private/World.cpp
+++ b/Engine/Source/Runtime/Engine/Private/World.cpp
@@ -1742,14 +1742,6 @@ void UWorld::RemoveFromWorld( ULevel* Level )
 			WorldComposition->OnLevelRemovedFromWorld(Level);
 		}
 
-		// Clear inner level streaming references
-		UWorld* LevelWorld = Cast<UWorld>(Level->GetOuter());
-		for (ULevelStreaming* StreamingLevel : LevelWorld->StreamingLevels)
-		{
-			StreamingLevel->ClearLoadedLevel();
-			StreamingLevel->SetPendingUnloadLevel(NULL);
-		}
-		
 		// Make sure level always has OwningWorld in the editor
 		if ( IsGameWorld() )
 		{
@@ -2124,7 +2116,7 @@ void UWorld::UpdateLevelStreamingInner( UWorld* PersistentWorld, FSceneViewFamil
 		if (bShouldBeLoaded)
 		{
 			const bool bBlockOnLoad = (!IsGameWorld() || !GEngine->bUseBackgroundLevelStreaming || bShouldBlockOnLoad);
-			// Try to obtain level
+			// Request to load or duplicate existing level
 			StreamingLevel->RequestLevel(PersistentWorld, bAllowLevelLoadRequests, bBlockOnLoad);
 		}
 		
@@ -2138,8 +2130,6 @@ void UWorld::UpdateLevelStreamingInner( UWorld* PersistentWorld, FSceneViewFamil
 			// Update loaded level visibility
 			if (bShouldBeVisible)
 			{
-				Level->bHasShowRequest = true;
-
 				// Add loaded level to a world if it's not there yet
 				if (!Level->bIsVisible)
 				{
@@ -2155,16 +2145,20 @@ void UWorld::UpdateLevelStreamingInner( UWorld* PersistentWorld, FSceneViewFamil
 			}
 			else
 			{
-				Level->bHasHideRequest = true;
 				// Discard previous LOD level
 				StreamingLevel->DiscardPendingUnloadLevel(PersistentWorld);
+				// Hide loaded level
+				if (Level->bIsVisible)
+				{
+					PersistentWorld->RemoveFromWorld(Level);
+				}
 			}
 
 			if (!bShouldBeLoaded)
 			{
-				if (!Level->bIsVisible && 
-					!Level->HasVisibilityRequestPending())
+				if (!Level->bIsVisible && !PersistentWorld->IsVisibilityRequestPending())
 				{
+					StreamingLevel->DiscardPendingUnloadLevel(PersistentWorld);
 					StreamingLevel->ClearLoadedLevel();
 					StreamingLevel->DiscardPendingUnloadLevel(PersistentWorld);
 				}
@@ -2213,24 +2207,8 @@ void UWorld::UpdateLevelStreaming( FSceneViewFamily* ViewFamily )
 	// Store current count of async loading packages
 	const int32	NumAsyncLoadingPackages = GetNumAsyncPackages(); 
 	
-	// In the editor update only streaming levels in the persistent world
-	// "Levels Browser" does not currently correctly supports multiple streaming levels referring to the same level package
-	if (!IsGameWorld())
-	{
-		UpdateLevelStreamingInner(this, ViewFamily);
-	}
-	else
-	{
-		// Make a copy of the current levels list 
-		// because levels count can be changed in the loop below
-		TArray<ULevel*> VisibleLevels = Levels;
-		for (int32 LevelIndex = 0; LevelIndex < VisibleLevels.Num(); ++LevelIndex)
-		{
-			// Update streaming objects in all visible worlds including Persistent world
-			UWorld* InnerWorld = Cast<UWorld>(VisibleLevels[LevelIndex]->GetOuter());
-			InnerWorld->UpdateLevelStreamingInner(this, ViewFamily);
-		}
-	}
+	// Update level streaming objects state
+	UpdateLevelStreamingInner(this, ViewFamily);
 	
 	// Force initial loading to be "bShouldBlockOnLoad".
 	const bool bLevelsHaveLoadRequestPending = NumAsyncLoadingPackages < GetNumAsyncPackages();
@@ -2240,40 +2218,6 @@ void UWorld::UpdateLevelStreaming( FSceneViewFamily* ViewFamily )
 		FlushAsyncLoading( NAME_None );
 	}
 	
-	TArray<ULevel*> VisibleLevels = Levels;
-	for (int LevelIndex = 1; LevelIndex < VisibleLevels.Num(); ++LevelIndex)
-	{
-		ULevel* Level = VisibleLevels[LevelIndex];
-		
-		// Hide level if we have request to hide it and no request to keep it visible
-		if (Level->bIsVisible &&
-			Level->bHasHideRequest && !Level->bHasShowRequest)
-		{
-			RemoveFromWorld(Level);
-		}
-
-		// Clear visibility request flags for next update
-		Level->bHasShowRequest = false;
-		Level->bHasHideRequest = false;
-				
-		// Process abandoned levels (visible levels with no streaming level objects referring them)
-		if (Level->GetStreamingLevelRefs() <= 0)
-		{
-			if (Level->bIsVisible)
-			{
-				// Hide abandoned level
-				RemoveFromWorld(Level);
-			}
-
-			if (!Level->bIsVisible &&
-				!Level->HasVisibilityRequestPending())
-			{
-				// It's safe to request unload for hidden abandoned level
-				FLevelStreamingGCHelper::RequestUnload(Level);
-			}
-		}
-	}
-	
 	// In case more levels has been requested to unload, force GC on next tick 
 	if (NumLevelsPendingPurge < FLevelStreamingGCHelper::GetNumLevelsPendingPurge())
 	{
diff --git a/Engine/Source/Runtime/Engine/Private/WorldComposition.cpp b/Engine/Source/Runtime/Engine/Private/WorldComposition.cpp
--- a/Engine/Source/Runtime/Engine/Private/WorldComposition.cpp
+++ b/Engine/Source/Runtime/Engine/Private/WorldComposition.cpp
@@ -570,7 +570,7 @@ void UWorldComposition::CommitTileStreamingState(UWorld* PersistenWorld, int32 T
 	// Quit early in case state is not going to be changed
 	if (StreamingLevel->bShouldBeLoaded == bShouldBeLoaded &&
 		StreamingLevel->bShouldBeVisible == bShouldBeVisible &&
-		StreamingLevel->GetLODIndex(PersistenWorld) == LODIdx)
+		StreamingLevel->LevelLODIndex == LODIdx)
 	{
 		return;
 	}
@@ -591,9 +591,9 @@ void UWorldComposition::CommitTileStreamingState(UWorld* PersistenWorld, int32 T
 	}
 
 	// Commit new state
-	StreamingLevel->bShouldBeLoaded = bShouldBeLoaded;
-	StreamingLevel->bShouldBeVisible = bShouldBeVisible;
-	StreamingLevel->SetLODIndex(PersistenWorld, LODIdx);
+	StreamingLevel->bShouldBeLoaded		= bShouldBeLoaded;
+	StreamingLevel->bShouldBeVisible	= bShouldBeVisible;
+	StreamingLevel->LevelLODIndex		= LODIdx;
 }
 
 
