LOD sublevels does not replicate visibility state to server

diff --git a/Engine/Source/Runtime/Engine/Classes/Engine/Level.h b/Engine/Source/Runtime/Engine/Classes/Engine/Level.h
--- a/Engine/Source/Runtime/Engine/Classes/Engine/Level.h
+++ b/Engine/Source/Runtime/Engine/Classes/Engine/Level.h
@@ -398,7 +398,8 @@ class ULevel : public ULevelBase, public IInterface_AssetUserData
 	uint32										bIsAssociatingLevel:1;
 	/** Whether this level should be fully added to the world before rendering his components	*/
 	uint32										bRequireFullVisibilityToRender:1;
-		
+	/** Whether this level is specific to client, visibility state will not be replicated to server	*/
+	uint32										bClientOnlyVisible:1;
 	/** Current index into actors array for updating components.							*/
 	int32										CurrentActorIndexForUpdateComponents;
 
diff --git a/Engine/Source/Runtime/Engine/Private/LevelStreaming.cpp b/Engine/Source/Runtime/Engine/Private/LevelStreaming.cpp
--- a/Engine/Source/Runtime/Engine/Private/LevelStreaming.cpp
+++ b/Engine/Source/Runtime/Engine/Private/LevelStreaming.cpp
@@ -484,6 +484,8 @@ void ULevelStreaming::AsyncLevelLoadComplete( const FString& InPackageName, UPac
 				if (LODPackageNames.Num() > 0)
 				{
 					Level->bRequireFullVisibilityToRender = true;
+					// LOD levels should not be visible on server
+					Level->bClientOnlyVisible = LODPackageNames.Contains(InLoadedPackage->GetFName());
 				}
 			
 				// In the editor levels must be in the levels array regardless of whether they are visible or not
diff --git a/Engine/Source/Runtime/Engine/Private/NetConnection.cpp b/Engine/Source/Runtime/Engine/Private/NetConnection.cpp
--- a/Engine/Source/Runtime/Engine/Private/NetConnection.cpp
+++ b/Engine/Source/Runtime/Engine/Private/NetConnection.cpp
@@ -1503,7 +1503,7 @@ void UNetConnection::HandleClientPlayer( APlayerController *PC, UNetConnection*
 		if (LevelStreaming != NULL)
 		{
 			const ULevel* Level = LevelStreaming->GetLoadedLevel();
-			if ( Level != NULL && Level->bIsVisible )
+			if ( Level != NULL && Level->bIsVisible && !Level->bClientOnlyVisible )
 			{
 				// Remap packagename for PIE networking before sending out to server
 				FName PackageName = Level->GetOutermost()->GetFName();
diff --git a/Engine/Source/Runtime/Engine/Private/World.cpp b/Engine/Source/Runtime/Engine/Private/World.cpp
--- a/Engine/Source/Runtime/Engine/Private/World.cpp
+++ b/Engine/Source/Runtime/Engine/Private/World.cpp
@@ -1638,19 +1638,22 @@ void UWorld::AddToWorld( ULevel* Level, const FTransform& LevelTransform )
 		CurrentLevelPendingVisibility					= NULL;
 
 		// notify server that the client has finished making this level visible
-		for (FLocalPlayerIterator It(GEngine, this); It; ++It)
+		if (!Level->bClientOnlyVisible)
 		{
-			if (It->PlayerController != NULL)
+			for (FLocalPlayerIterator It(GEngine, this); It; ++It)
 			{
-				// Remap packagename for PIE networking before sending out to server
-				FName PackageName = Level->GetOutermost()->GetFName();
-				FString PackageNameStr = PackageName.ToString();
-				if (GEngine->NetworkRemapPath(this, PackageNameStr, false))
+				if (It->PlayerController != NULL)
 				{
-					PackageName = FName(*PackageNameStr);
-				}
+					// Remap packagename for PIE networking before sending out to server
+					FName PackageName = Level->GetOutermost()->GetFName();
+					FString PackageNameStr = PackageName.ToString();
+					if (GEngine->NetworkRemapPath(this, PackageNameStr, false))
+					{
+						PackageName = FName(*PackageNameStr);
+					}
 
-				It->PlayerController->ServerUpdateLevelVisibility(PackageName, true);
+					It->PlayerController->ServerUpdateLevelVisibility(PackageName, true);
+				}
 			}
 		}
 
@@ -1720,19 +1723,22 @@ void UWorld::RemoveFromWorld( ULevel* Level )
 		Level->ClearLevelComponents();
 
 		// notify server that the client has removed this level
-		for (FLocalPlayerIterator It(GEngine, this); It; ++It)
+		if (!Level->bClientOnlyVisible)
 		{
-			if (It->PlayerController != NULL)
+			for (FLocalPlayerIterator It(GEngine, this); It; ++It)
 			{
-				// Remap packagename for PIE networking before sending out to server
-				FName PackageName = Level->GetOutermost()->GetFName();
-				FString PackageNameStr = PackageName.ToString();
-				if (GEngine->NetworkRemapPath(this, PackageNameStr, false))
+				if (It->PlayerController != NULL)
 				{
-					PackageName = FName(*PackageNameStr);
-				}
+					// Remap packagename for PIE networking before sending out to server
+					FName PackageName = Level->GetOutermost()->GetFName();
+					FString PackageNameStr = PackageName.ToString();
+					if (GEngine->NetworkRemapPath(this, PackageNameStr, false))
+					{
+						PackageName = FName(*PackageNameStr);
+					}
 
-				It->PlayerController->ServerUpdateLevelVisibility(PackageName, false);
+					It->PlayerController->ServerUpdateLevelVisibility(PackageName, false);
+				}
 			}
 		}
 		
