Fixed RescaleSimpleCollision so that ConvexElems position is updated when the build scale is increased (previously just scaled the verts).

#ttp 336728 - EDITOR: Static Mesh Editor, BuildScale3D does not update the the position of convex elements

#branch UE4

#change Convex elem position is now scaled when the build scale is modified.

reviwed by [REDACTED_PERSON]

diff --git a/Engine/Source/Runtime/Engine/Private/PhysicsEngine/BodySetup.cpp b/Engine/Source/Runtime/Engine/Private/PhysicsEngine/BodySetup.cpp
--- a/Engine/Source/Runtime/Engine/Private/PhysicsEngine/BodySetup.cpp
+++ b/Engine/Source/Runtime/Engine/Private/PhysicsEngine/BodySetup.cpp
@@ -528,6 +528,12 @@ void UBodySetup::RescaleSimpleCollision( FVector BuildScale )
 		{
 			FKConvexElem* ConvexElem = &(AggGeom.ConvexElems[i]);
 
+			FTransform ConvexTrans = ConvexElem->GetTransform();
+			FVector ConvexLoc = ConvexTrans.GetLocation();
+			ConvexLoc *= ScaleMultiplier3D;
+			ConvexTrans.SetLocation(ConvexLoc);
+			ConvexElem->SetTransform(ConvexTrans);
+
 			TArray<FVector>& Vertices = ConvexElem->VertexData;
 			for (int32 VertIndex = 0; VertIndex < Vertices.Num(); ++VertIndex)
 			{
