Added Structs to store streamed audio chunks for runtime streaming or in DDC.

Followed pattern set out by texture streaming so that each chunk of audio data resides in its own bulkdata struct. It is currently possible to split audio into chunks when SoundWaves are marked for streaming but there is no way of doing this exposed at present.
Changed the parameters of FAudioDevice::GetRuntimeFormat so that the relevant SoundWave must be passed in, to allow for different formats for individual sounds/streaming options.
USoundWave::FreeResources no longer resets the NumChannels as it is unnecessary and causes sounds to be unable to play after the OGG data is flushed when attempting to switch to OPUS.

diff --git a/Engine/Source/Developer/AudioFormatOpus/Private/AudioFormatOpus.cpp b/Engine/Source/Developer/AudioFormatOpus/Private/AudioFormatOpus.cpp
--- a/Engine/Source/Developer/AudioFormatOpus/Private/AudioFormatOpus.cpp
+++ b/Engine/Source/Developer/AudioFormatOpus/Private/AudioFormatOpus.cpp
@@ -232,6 +232,67 @@ public:
 		return CompressedDataStore.Num();
 	}
 
+	virtual bool SplitDataForStreaming(const TArray<uint8>& SrcBuffer, TArray<TArray<uint8>>& OutBuffers) const OVERRIDE
+	{
+		// 16K chunks - don't really have an idea of what's best for loading from disc yet
+		const int32 kMaxChunkSizeBytes = 16384;
+		if (SrcBuffer.Num() == 0)
+		{
+			return false;
+		}
+
+		uint32 ReadOffset = 0;
+		uint32 WriteOffset = 0;
+		uint16 ProcessedFrames = 0;
+		const uint8* LockedSrc = SrcBuffer.GetTypedData();
+
+		// Read True Sample Count, Number of channels and Frames to Encode first
+		uint32 TrueSampleCount = *((uint32*)(LockedSrc + ReadOffset));
+		ReadOffset += sizeof(uint32);
+		uint8 NumChannels = *(LockedSrc + ReadOffset);
+		ReadOffset += sizeof(uint8);
+		uint16 SerializedFrames = *((uint16*)(LockedSrc + ReadOffset));
+		ReadOffset += sizeof(uint16);
+
+		// Should always be able to store basic info in a single chunk
+		check(ReadOffset - WriteOffset < kMaxChunkSizeBytes)
+
+		while (ProcessedFrames < SerializedFrames)
+		{
+			uint16 FrameSize = *((uint16*)(LockedSrc + ReadOffset));
+
+			if ( (ReadOffset + sizeof(uint16) + FrameSize) - WriteOffset >= kMaxChunkSizeBytes)
+			{
+				// need to write some data
+				int32 OutIndex = OutBuffers.Num();
+				int32 BytesToCopy = ReadOffset - WriteOffset;
+				OutBuffers.Add(TArray<uint8>());
+				OutBuffers[OutIndex].Empty(BytesToCopy);
+				OutBuffers[OutIndex].AddUninitialized(BytesToCopy);
+				void* NewChunkData = OutBuffers[OutIndex].GetTypedData();
+				FMemory::Memcpy(NewChunkData, LockedSrc+WriteOffset, BytesToCopy);
+				WriteOffset += BytesToCopy;
+			}
+
+			ReadOffset += sizeof(uint16) + FrameSize;
+			ProcessedFrames++;
+		}
+		if (WriteOffset < ReadOffset)
+		{
+			// need to write some data
+			int32 OutIndex = OutBuffers.Num();
+			int32 BytesToCopy = ReadOffset - WriteOffset;
+			OutBuffers.Add(TArray<uint8>());
+			OutBuffers[OutIndex].Empty(BytesToCopy);
+			OutBuffers[OutIndex].AddUninitialized(BytesToCopy);
+			void* NewChunkData = OutBuffers[OutIndex].GetTypedData();
+			FMemory::Memcpy(NewChunkData, LockedSrc + WriteOffset, BytesToCopy);
+			WriteOffset += BytesToCopy;
+		}
+
+		return true;
+	}
+
 	void Destroy(OpusEncoder* Encoder) const
 	{
 #if USE_UE4_MEM_ALLOC
diff --git a/Engine/Source/Developer/MovieSceneTools/Private/AudioTrackEditor.cpp b/Engine/Source/Developer/MovieSceneTools/Private/AudioTrackEditor.cpp
--- a/Engine/Source/Developer/MovieSceneTools/Private/AudioTrackEditor.cpp
+++ b/Engine/Source/Developer/MovieSceneTools/Private/AudioTrackEditor.cpp
@@ -172,7 +172,7 @@ void FAudioThumbnail::GenerateWaveformPreview(TArray<uint8>& OutData, TRange<flo
 	{
 		// @todo Sequencer optimize - We might want to generate the data when we generate the texture
 		// and then discard the data afterwards, though that might be a perf hit traded for better memory usage
-		SoundWave->InitAudioResource(GEngine->GetAudioDevice()->GetRuntimeFormat());
+		SoundWave->InitAudioResource(GEngine->GetAudioDevice()->GetRuntimeFormat(SoundWave));
 		FAsyncAudioDecompress TempDecompress(SoundWave);
 		TempDecompress.StartSynchronousTask();
 	}
diff --git a/Engine/Source/Developer/TargetPlatform/Public/Interfaces/IAudioFormat.h b/Engine/Source/Developer/TargetPlatform/Public/Interfaces/IAudioFormat.h
--- a/Engine/Source/Developer/TargetPlatform/Public/Interfaces/IAudioFormat.h
+++ b/Engine/Source/Developer/TargetPlatform/Public/Interfaces/IAudioFormat.h
@@ -100,6 +100,16 @@ public:
 	 */
 	virtual int32 Recompress( FName Format, const TArray<uint8>& SrcBuffer, FSoundQualityInfo& QualityInfo, TArray<uint8>& OutBuffer ) const = 0;
 
+	/**
+	 * Splits compressed data into chunks suitable for streaming audio
+	 *
+	 * @param	SrcBuffer	pre-compressed data as an array of bytes
+	 * @param	OutBuffers	Array of buffers that contain the chunks the original data was split into
+	 *
+	 * @return	bool	Whether bulk data could be split for streaming
+	 */
+	virtual bool SplitDataForStreaming(const TArray<uint8>& SrcBuffer, TArray<TArray<uint8>>& OutBuffers) const {return false;}
+
 public:
 
 	/**
diff --git a/Engine/Source/Developer/Windows/WindowsTargetPlatform/Private/GenericWindowsTargetPlatform.h b/Engine/Source/Developer/Windows/WindowsTargetPlatform/Private/GenericWindowsTargetPlatform.h
--- a/Engine/Source/Developer/Windows/WindowsTargetPlatform/Private/GenericWindowsTargetPlatform.h
+++ b/Engine/Source/Developer/Windows/WindowsTargetPlatform/Private/GenericWindowsTargetPlatform.h
@@ -156,6 +156,13 @@ public:
 	virtual FName GetWaveFormat( class USoundWave* Wave ) const OVERRIDE
 	{
 		static FName NAME_OGG(TEXT("OGG"));
+		static FName NAME_OPUS(TEXT("OPUS"));
+
+		if (Wave->IsStreaming())
+		{
+			return NAME_OPUS;
+		}
+
 		return NAME_OGG;
 	}
 #endif //WITH_ENGINE
diff --git a/Engine/Source/Runtime/Android/AndroidAudio/Private/AndroidAudioBuffer.cpp b/Engine/Source/Runtime/Android/AndroidAudio/Private/AndroidAudioBuffer.cpp
--- a/Engine/Source/Runtime/Android/AndroidAudio/Private/AndroidAudioBuffer.cpp
+++ b/Engine/Source/Runtime/Android/AndroidAudio/Private/AndroidAudioBuffer.cpp
@@ -53,7 +53,7 @@ FSLESSoundBuffer* FSLESSoundBuffer::CreateQueuedBuffer( FSLESAudioDevice* AudioD
 	
 	Buffer->DecompressionState = AudioDevice->CreateCompressedAudioInfo(InWave);
 
-	InWave->InitAudioResource( AudioDevice->GetRuntimeFormat() );
+	InWave->InitAudioResource( AudioDevice->GetRuntimeFormat(InWave) );
 	
 	if( Buffer->DecompressionState->ReadCompressedInfo( InWave->ResourceData, InWave->ResourceSize, &QualityInfo ) )
 	{
diff --git a/Engine/Source/Runtime/Android/AndroidAudio/Public/AndroidAudioDevice.h b/Engine/Source/Runtime/Android/AndroidAudio/Public/AndroidAudioDevice.h
--- a/Engine/Source/Runtime/Android/AndroidAudio/Public/AndroidAudioDevice.h
+++ b/Engine/Source/Runtime/Android/AndroidAudio/Public/AndroidAudioDevice.h
@@ -235,7 +235,7 @@ public:
 	FSLESAudioDevice() {} 
 	virtual ~FSLESAudioDevice() {} 
 
-	virtual FName GetRuntimeFormat() OVERRIDE
+	virtual FName GetRuntimeFormat(USoundWave* SoundWave) OVERRIDE
 	{
 		static FName NAME_OGG(TEXT("OGG"));
 		return NAME_OGG;
diff --git a/Engine/Source/Runtime/Engine/Classes/Sound/SoundWave.h b/Engine/Source/Runtime/Engine/Classes/Sound/SoundWave.h
--- a/Engine/Source/Runtime/Engine/Classes/Sound/SoundWave.h
+++ b/Engine/Source/Runtime/Engine/Classes/Sound/SoundWave.h
@@ -25,6 +25,84 @@ enum EDecompressionType
 	DTYPE_MAX,
 };
 
+/**
+ * A chunk of streamed audio.
+ */
+struct FStreamedAudioChunk
+{
+	/** Bulk data if stored in the package. */
+	FByteBulkData BulkData;
+
+	/** Default constructor. */
+	FStreamedAudioChunk()
+	{
+	}
+
+	/** Serialization. */
+	void Serialize(FArchive& Ar, UObject* Owner, int32 ChunkIndex);
+
+#if WITH_EDITORONLY_DATA
+	/** Key if stored in the derived data cache. */
+	FString DerivedDataKey;
+
+	/**
+	 * Place chunk data in the derived data cache associated with the provided
+	 * key.
+	 */
+	void StoreInDerivedDataCache(const FString& InDerivedDataKey);
+#endif // #if WITH_EDITORONLY_DATA
+};
+
+/**
+ * Platform-specific data used streaming audio at runtime.
+ */
+USTRUCT()
+struct FStreamedAudioPlatformData
+{
+	GENERATED_USTRUCT_BODY()
+
+	/** Number of audio chunks. */
+	int32 NumChunks;
+	/** Format in which audio chunks are stored. */
+	FName AudioFormat;
+	/** audio data. */
+	TIndirectArray<struct FStreamedAudioChunk> Chunks;
+
+#if WITH_EDITORONLY_DATA
+	/** The key associated with this derived data. */
+	FString DerivedDataKey;
+	/** Async cache task if one is outstanding. */
+	struct FStreamedAudioAsyncCacheDerivedDataTask* AsyncTask;
+#endif
+
+	/** Default constructor. */
+	FStreamedAudioPlatformData();
+
+	/** Destructor. */
+	~FStreamedAudioPlatformData();
+
+	/**
+	 * Try to load audio chunks from the derived data cache.
+	 * @param FirstChunkToLoad - The first Chunk index to load.
+	 * @param OutChunkData -	Must point to an array of pointers with at least
+	 *						Audio.Chunks.Num() - FirstChunkToLoad + 1 entries. Upon
+	 *						return those pointers will contain audio chunk data.
+	 * @returns true if all requested chunks have been loaded.
+	 */
+	bool TryLoadChunks(int32 FirstChunkToLoad, void** OutChunkData);
+
+	/** Serialization. */
+	void Serialize(FArchive& Ar, class USoundWave* Owner);
+
+#if WITH_EDITORONLY_DATA
+	void Cache(class USoundWave& InSoundWave, FName AudioFormatName, uint32 InFlags);
+	void FinishCache();
+	ENGINE_API bool TryInlineChunkData();
+	bool AreDerivedChunksAvailable() const;
+#endif
+
+};
+
 UCLASS(hidecategories=Object, editinlinenew, BlueprintType, dependson=USoundGroups)
 class ENGINE_API USoundWave : public USoundBase
 {
@@ -38,6 +116,9 @@ class ENGINE_API USoundWave : public USoundBase
 	UPROPERTY(EditAnywhere, Category=SoundWave )
 	uint32 bLooping:1;
 
+	UPROPERTY()
+	uint32 bStreaming:1;
+
 	/** Set to true for programmatically-generated, streamed audio. */
 	uint32 bProcedural:1;
 
@@ -160,6 +241,12 @@ public:
 	/** Cache the total used memory recorded for this SoundWave to keep INC/DEC consistent */
 	int32 TrackedMemoryUsage;
 
+	/** The streaming derived data for this sound on this platform. */
+	FStreamedAudioPlatformData *RunningPlatformData;
+
+	/* cooked streaming platform data for this sound */
+	TMap<FString, FStreamedAudioPlatformData*> CookedPlatformData;
+
 	// Begin UObject interface. 
 	virtual void Serialize( FArchive& Ar ) OVERRIDE;
 	virtual void PostInitProperties() OVERRIDE;
@@ -261,7 +348,45 @@ public:
 	 */ 
 	void InvalidateCompressedData();
 
+	/**
+	 * Checks whether sound has been categorised as streaming
+	 */
+	bool IsStreaming() const;
+
+	/**
+	 * Attempts to update the cached platform data after any changes that might affect it
+	 */
+	void UpdatePlatformData();
 
+	void CleanupCachedRunningPlatformData();
+	void CleanupCachedCookedPlatformData();
+
+	/**
+	 * Serializes cooked platform data.
+	 */
+	void SerializeCookedPlatformData(class FArchive& Ar);
+
+#if WITH_EDITORONLY_DATA
+	/**
+	 * Caches platform data for the sound.
+	 */
+	void CachePlatformData(bool bAsyncCache = false);
+
+	/**
+	 * Begins caching platform data in the background for the platform requested
+	 */
+	virtual void BeginCacheForCookedPlatformData(  const ITargetPlatform *TargetPlatform ) OVERRIDE;
+
+	/**
+	 * Begins caching platform data in the background.
+	 */
+	void BeginCachePlatformData();
+
+	/**
+	 * Blocks on async cache tasks and prepares platform data for use.
+	 */
+	void FinishCachePlatformData();
+#endif
 };
 
 
diff --git a/Engine/Source/Runtime/Engine/Private/AudioDerivedData.cpp b/Engine/Source/Runtime/Engine/Private/AudioDerivedData.cpp
--- a/Engine/Source/Runtime/Engine/Private/AudioDerivedData.cpp
+++ b/Engine/Source/Runtime/Engine/Private/AudioDerivedData.cpp
@@ -6,6 +6,555 @@
 
 DEFINE_LOG_CATEGORY_STATIC(LogAudioDerivedData, Log, All);
 
+/*------------------------------------------------------------------------------
+	Derived data key generation.
+------------------------------------------------------------------------------*/
+
+#if WITH_EDITORONLY_DATA
+
+// If you want to bump this version, generate a new guid using
+// VS->Tools->Create GUID and paste it here.
+#define STREAMEDAUDIO_DERIVEDDATA_VER		TEXT("B9C1C21AEBE443B28C3B5546C8D5D742")
+
+/**
+ * Computes the derived data key suffix for a SoundWave's Streamed Audio.
+ * @param SoundWave - The SoundWave for which to compute the derived data key.
+ * @param AudioFormatName - The audio format we're creating the key for
+ * @param OutKeySuffix - The derived data key suffix.
+ */
+static void GetStreamedAudioDerivedDataKeySuffix(
+	const USoundWave& SoundWave,
+	FName AudioFormatName,
+	FString& OutKeySuffix
+	)
+{
+	uint16 Version = 0;
+
+	// get the version for this soundwave's platform format
+	ITargetPlatformManagerModule* TPM = GetTargetPlatformManager();
+	if (TPM)
+	{
+		const IAudioFormat* AudioFormat = TPM->FindAudioFormat(AudioFormatName);
+		if (AudioFormat)
+		{
+			Version = AudioFormat->GetVersion(AudioFormatName);
+		}
+	}
+	
+	// build the key
+	OutKeySuffix = FString::Printf(TEXT("%s_%d_%s"),
+		*AudioFormatName.ToString(),
+		Version,
+		*SoundWave.CompressedDataGuid.ToString()
+		);
+}
+
+/**
+ * Constructs a derived data key from the key suffix.
+ * @param KeySuffix - The key suffix.
+ * @param OutKey - The full derived data key.
+ */
+static void GetStreamedAudioDerivedDataKeyFromSuffix(const FString& KeySuffix, FString& OutKey)
+{
+	OutKey = FDerivedDataCacheInterface::BuildCacheKey(
+		TEXT("STREAMEDAUDIO"),
+		STREAMEDAUDIO_DERIVEDDATA_VER,
+		*KeySuffix
+		);
+}
+
+/**
+ * Constructs the derived data key for an individual audio chunk.
+ * @param KeySuffix - The key suffix.
+ * @param ChunkIndex - The chunk index.
+ * @param OutKey - The full derived data key for the audio chunk.
+ */
+static void GetStreamedAudioDerivedChunkKey(
+	int32 ChunkIndex,
+	const FString& KeySuffix,
+	FString& OutKey
+	)
+{
+	OutKey = FDerivedDataCacheInterface::BuildCacheKey(
+		TEXT("STREAMEDAUDIO"),
+		STREAMEDAUDIO_DERIVEDDATA_VER,
+		*FString::Printf(TEXT("%s_CHUNK%u"), *KeySuffix, ChunkIndex)
+		);
+}
+
+/**
+ * Computes the derived data key for Streamed Audio.
+ * @param SoundWave - The soundwave for which to compute the derived data key.
+ * @param AudioFormatName - The audio format we're creating the key for
+ * @param OutKey - The derived data key.
+ */
+static void GetStreamedAudioDerivedDataKey(
+	const USoundWave& SoundWave,
+	FName AudioFormatName,
+	FString& OutKey
+	)
+{
+	FString KeySuffix;
+	GetStreamedAudioDerivedDataKeySuffix(SoundWave, AudioFormatName, KeySuffix);
+	GetStreamedAudioDerivedDataKeyFromSuffix(KeySuffix, OutKey);
+}
+
+/**
+ * Gets Wave format for a SoundWave on the current running platform
+ * @param SoundWave - The SoundWave to get format for.
+ */
+static FName GetWaveFormatForRunningPlatform(USoundWave& SoundWave)
+{
+	// Compress to whatever format the active target platform wants
+	ITargetPlatformManagerModule* TPM = GetTargetPlatformManager();
+	if (TPM)
+	{
+		ITargetPlatform* CurrentPlatform = NULL;
+		const TArray<ITargetPlatform*>& Platforms = TPM->GetActiveTargetPlatforms();
+
+		check(Platforms.Num());
+
+		CurrentPlatform = Platforms[0];
+
+		for (int32 Index = 1; Index < Platforms.Num(); Index++)
+		{
+			if (Platforms[Index]->IsRunningPlatform())
+			{
+				CurrentPlatform = Platforms[Index];
+				break;
+			}
+		}
+
+		check(CurrentPlatform != NULL);
+
+		return CurrentPlatform->GetWaveFormat(&SoundWave);
+	}
+
+	return NAME_None;
+}
+
+
+/**
+ * Stores derived data in the DDC.
+ * @param DerivedData - The data to store in the DDC.
+ * @param DerivedDataKeySuffix - The key suffix at which to store derived data.
+ */
+static void PutDerivedDataInCache(
+	FStreamedAudioPlatformData* DerivedData,
+	const FString& DerivedDataKeySuffix
+	)
+{
+	TArray<uint8> RawDerivedData;
+	FString DerivedDataKey;
+
+	// Build the key with which to cache derived data.
+	GetStreamedAudioDerivedDataKeyFromSuffix(DerivedDataKeySuffix, DerivedDataKey);
+
+	FString LogString;
+	if (UE_LOG_ACTIVE(LogAudio,Verbose))
+	{
+		LogString = FString::Printf(
+			TEXT("Storing Streamed Audio in DDC:\n  Key: %s\n  Format: %s\n"),
+			*DerivedDataKey,
+			*DerivedData->AudioFormat.ToString()
+			);
+	}
+
+	// Write out individual chunks to the derived data cache.
+	const int32 ChunkCount = DerivedData->Chunks.Num();
+	for (int32 ChunkIndex = 0; ChunkIndex < ChunkCount; ++ChunkIndex)
+	{
+		FString ChunkDerivedDataKey;
+		FStreamedAudioChunk& Chunk = DerivedData->Chunks[ChunkIndex];
+		GetStreamedAudioDerivedChunkKey(ChunkIndex, DerivedDataKeySuffix, ChunkDerivedDataKey);
+
+		if (UE_LOG_ACTIVE(LogAudio,Verbose))
+		{
+			LogString += FString::Printf(TEXT("  Chunk%d %d bytes %s\n"),
+				ChunkIndex,
+				Chunk.BulkData.GetBulkDataSize(),
+				*ChunkDerivedDataKey
+				);
+		}
+
+		Chunk.StoreInDerivedDataCache(ChunkDerivedDataKey);
+	}
+
+	// Store derived data.
+	FMemoryWriter Ar(RawDerivedData, /*bIsPersistent=*/ true);
+	DerivedData->Serialize(Ar, NULL);
+	GetDerivedDataCacheRef().Put(*DerivedDataKey, RawDerivedData);
+	UE_LOG(LogAudio,Verbose,TEXT("%s  Derived Data: %d bytes"),*LogString,RawDerivedData.Num());
+}
+
+#endif // #if WITH_EDITORONLY_DATA
+
+#if WITH_EDITORONLY_DATA
+
+namespace EStreamedAudioCacheFlags
+{
+	enum Type
+	{
+		None			= 0x0,
+		Async			= 0x1,
+		ForceRebuild	= 0x2,
+		AllowAsyncBuild	= 0x4,
+		ForDDCBuild		= 0x8,
+	};
+};
+
+/**
+ * Worker used to cache streamed audio derived data.
+ */
+class FStreamedAudioCacheDerivedDataWorker : public FNonAbandonableTask
+{
+	/** Where to store derived data. */
+	FStreamedAudioPlatformData* DerivedData;
+	/** The SoundWave for which derived data is being cached. */
+	USoundWave& SoundWave;
+	/** Audio Format Name */
+	FName AudioFormatName;
+	/** Derived data key suffix. */
+	FString KeySuffix;
+	/** Streamed Audio cache flags. */
+	uint32 CacheFlags;
+	/** true if caching has succeeded. */
+	bool bSucceeded;
+
+	/** Build the streamed audio. This function is safe to call from any thread. */
+	void BuildStreamedAudio()
+	{
+		GetStreamedAudioDerivedDataKeySuffix(SoundWave, AudioFormatName, KeySuffix);
+
+		DerivedData->Chunks.Empty();
+		
+		ITargetPlatformManagerModule* TPM = GetTargetPlatformManager();
+		const IAudioFormat* AudioFormat = NULL;
+		if (TPM)
+		{
+			AudioFormat = TPM->FindAudioFormat(AudioFormatName);
+		}
+
+		if (AudioFormat)
+		{
+			DerivedData->AudioFormat = AudioFormatName;
+
+			FByteBulkData* CompressedData = SoundWave.GetCompressedData(AudioFormatName);
+			TArray<uint8> CompressedBuffer;
+			CompressedBuffer.Empty(CompressedData->GetBulkDataSize());
+			CompressedBuffer.AddUninitialized(CompressedData->GetBulkDataSize());
+			void* BufferData = CompressedBuffer.GetTypedData();
+			CompressedData->GetCopy(&BufferData, false);
+			TArray<TArray<uint8>> ChunkBuffers;
+
+			if (AudioFormat->SplitDataForStreaming(CompressedBuffer, ChunkBuffers))
+			{
+				for (int32 ChunkIndex = 0; ChunkIndex < ChunkBuffers.Num(); ++ChunkIndex)
+				{
+					FStreamedAudioChunk* NewChunk = new(DerivedData->Chunks) FStreamedAudioChunk();
+					NewChunk->BulkData.Lock(LOCK_READ_WRITE);
+					void* NewChunkData = NewChunk->BulkData.Realloc(ChunkBuffers[ChunkIndex].Num());
+					FMemory::Memcpy(NewChunkData, ChunkBuffers[ChunkIndex].GetTypedData(), ChunkBuffers[ChunkIndex].Num());
+					NewChunk->BulkData.Unlock();
+				}
+			}
+			else
+			{
+				// Could not split so copy compressed data into a single chunk
+				FStreamedAudioChunk* NewChunk = new(DerivedData->Chunks) FStreamedAudioChunk();
+				NewChunk->BulkData.Lock(LOCK_READ_WRITE);
+				void* NewChunkData = NewChunk->BulkData.Realloc(CompressedBuffer.Num());
+				FMemory::Memcpy(NewChunkData, CompressedBuffer.GetTypedData(), CompressedBuffer.Num());
+				NewChunk->BulkData.Unlock();
+			}
+
+			DerivedData->NumChunks = DerivedData->Chunks.Num();
+
+			// Store it in the cache.
+			PutDerivedDataInCache(DerivedData, KeySuffix);
+		}
+		
+		if (DerivedData->Chunks.Num())
+		{
+			bSucceeded = true;
+		}
+		else
+		{
+			UE_LOG(LogAudio, Warning, TEXT("Failed to build %s derived data for %s"),
+				*AudioFormatName.GetPlainNameString(),
+				*SoundWave.GetPathName()
+				);
+		}
+	}
+
+public:
+	/** Initialization constructor. */
+	FStreamedAudioCacheDerivedDataWorker(
+		FStreamedAudioPlatformData* InDerivedData,
+		USoundWave* InSoundWave,
+		FName InAudioFormatName,
+		uint32 InCacheFlags
+		)
+		: DerivedData(InDerivedData)
+		, SoundWave(*InSoundWave)
+		, AudioFormatName(InAudioFormatName)
+		, CacheFlags(InCacheFlags)
+		, bSucceeded(false)
+	{
+	}
+
+	/** Does the work to cache derived data. Safe to call from any thread. */
+	void DoWork()
+	{
+		TArray<uint8> RawDerivedData;
+		bool bForceRebuild = (CacheFlags & EStreamedAudioCacheFlags::ForceRebuild) != 0;
+		bool bForDDC = (CacheFlags & EStreamedAudioCacheFlags::ForDDCBuild) != 0;
+		bool bAllowAsyncBuild = (CacheFlags & EStreamedAudioCacheFlags::AllowAsyncBuild) != 0;
+
+		if (!bForceRebuild && GetDerivedDataCacheRef().GetSynchronous(*DerivedData->DerivedDataKey, RawDerivedData))
+		{
+			FMemoryReader Ar(RawDerivedData, /*bIsPersistent=*/ true);
+			DerivedData->Serialize(Ar, NULL);
+			bSucceeded = true;
+			if (bForDDC)
+			{
+				bSucceeded = DerivedData->TryLoadChunks(0,NULL);
+			}
+			else
+			{
+				bSucceeded = DerivedData->AreDerivedChunksAvailable();
+			}
+		}
+		else if (bAllowAsyncBuild)
+		{
+			BuildStreamedAudio();
+		}
+	}
+
+	/** Finalize work. Must be called ONLY by the game thread! */
+	void Finalize()
+	{
+		check(IsInGameThread());
+		if (!bSucceeded)
+		{
+			BuildStreamedAudio();
+		}
+	}
+
+	/** Interface for FAsyncTask. */
+	static const TCHAR* Name()
+	{
+		return TEXT("FStreamedAudioAsyncCacheDerivedDataTask");
+	}
+};
+
+struct FStreamedAudioAsyncCacheDerivedDataTask : public FAsyncTask<FStreamedAudioCacheDerivedDataWorker>
+{
+	FStreamedAudioAsyncCacheDerivedDataTask(
+		FStreamedAudioPlatformData* InDerivedData,
+		USoundWave* InSoundWave,
+		FName InAudioFormatName,
+		uint32 InCacheFlags
+		)
+		: FAsyncTask<FStreamedAudioCacheDerivedDataWorker>(
+			InDerivedData,
+			InSoundWave,
+			InAudioFormatName,
+			InCacheFlags
+			)
+	{
+	}
+};
+
+void FStreamedAudioPlatformData::Cache(USoundWave& InSoundWave, FName AudioFormatName, uint32 InFlags)
+{
+	// Flush any existing async task and ignore results.
+	FinishCache();
+
+	uint32 Flags = InFlags;
+
+	static bool bForDDC = FString(FCommandLine::Get()).Contains(TEXT("DerivedDataCache"));
+	if (bForDDC)
+	{
+		Flags |= EStreamedAudioCacheFlags::ForDDCBuild;
+	}
+
+	bool bForceRebuild = (Flags & EStreamedAudioCacheFlags::ForceRebuild) != 0;
+	bool bAsync = !bForDDC && (Flags & EStreamedAudioCacheFlags::Async) != 0;
+	GetStreamedAudioDerivedDataKey(InSoundWave, AudioFormatName, DerivedDataKey);
+
+	if (bAsync && !bForceRebuild)
+	{
+		AsyncTask = new FStreamedAudioAsyncCacheDerivedDataTask(this, &InSoundWave, AudioFormatName, Flags);
+		AsyncTask->StartBackgroundTask();
+	}
+	else
+	{
+		FStreamedAudioCacheDerivedDataWorker Worker(this, &InSoundWave, AudioFormatName, Flags);
+		Worker.DoWork();
+		Worker.Finalize();
+	}
+}
+
+void FStreamedAudioPlatformData::FinishCache()
+{
+	if (AsyncTask)
+	{
+		{
+			AsyncTask->EnsureCompletion();
+		}
+		FStreamedAudioCacheDerivedDataWorker& Worker = AsyncTask->GetTask();
+		Worker.Finalize();
+		delete AsyncTask;
+		AsyncTask = NULL;
+	}
+}
+
+/**
+ * Executes async DDC gets for chunks stored in the derived data cache.
+ * @param Chunks - Chunks to retrieve.
+ * @param FirstChunkToLoad - Index of the first chunk to retrieve.
+ * @param OutHandles - Handles to the asynchronous DDC gets.
+ */
+static void BeginLoadDerivedChunks(TIndirectArray<FStreamedAudioChunk>& Chunks, int32 FirstChunkToLoad, TArray<uint32>& OutHandles)
+{
+	FDerivedDataCacheInterface& DDC = GetDerivedDataCacheRef();
+	OutHandles.AddZeroed(Chunks.Num());
+	for (int32 ChunkIndex = FirstChunkToLoad; ChunkIndex < Chunks.Num(); ++ChunkIndex)
+	{
+		const FStreamedAudioChunk& Chunk = Chunks[ChunkIndex];
+		if (!Chunk.DerivedDataKey.IsEmpty())
+		{
+			OutHandles[ChunkIndex] = DDC.GetAsynchronous(*Chunk.DerivedDataKey);
+		}
+	}
+}
+
+#endif //WITH_EDITORONLY_DATA
+
+FStreamedAudioPlatformData::FStreamedAudioPlatformData()
+#if WITH_EDITORONLY_DATA
+: AsyncTask(NULL)
+#endif // #if WITH_EDITORONLY_DATA
+{
+}
+
+FStreamedAudioPlatformData::~FStreamedAudioPlatformData()
+{
+#if WITH_EDITORONLY_DATA
+	if (AsyncTask)
+	{
+		AsyncTask->EnsureCompletion();
+		delete AsyncTask;
+		AsyncTask = NULL;
+	}
+#endif
+}
+
+bool FStreamedAudioPlatformData::TryLoadChunks(int32 FirstChunkToLoad, void** OutChunkData)
+{
+	int32 NumChunksCached = 0;
+
+#if WITH_EDITORONLY_DATA
+	TArray<uint8> TempData;
+	TArray<uint32> AsyncHandles;
+	FDerivedDataCacheInterface& DDC = GetDerivedDataCacheRef();
+	BeginLoadDerivedChunks(Chunks, FirstChunkToLoad, AsyncHandles);
+#endif // #if WITH_EDITORONLY_DATA
+
+	// Load remaining chunks (if any) from bulk data.
+	for (int32 ChunkIndex = FirstChunkToLoad; ChunkIndex < Chunks.Num(); ++ChunkIndex)
+	{
+		FStreamedAudioChunk& Chunk = Chunks[ChunkIndex];
+		if (Chunk.BulkData.GetBulkDataSize() > 0)
+		{
+			if (OutChunkData)
+			{
+				OutChunkData[ChunkIndex - FirstChunkToLoad] = FMemory::Malloc(Chunk.BulkData.GetBulkDataSize());
+				Chunk.BulkData.GetCopy(&OutChunkData[ChunkIndex - FirstChunkToLoad]);
+			}
+			NumChunksCached++;
+		}
+	}
+
+#if WITH_EDITORONLY_DATA
+	// Wait for async DDC gets.
+	for (int32 ChunkIndex = FirstChunkToLoad; ChunkIndex < Chunks.Num(); ++ChunkIndex)
+	{
+		FStreamedAudioChunk& Chunk = Chunks[ChunkIndex];
+		if (Chunk.DerivedDataKey.IsEmpty() == false)
+		{
+			uint32 AsyncHandle = AsyncHandles[ChunkIndex];
+			DDC.WaitAsynchronousCompletion(AsyncHandle);
+			if (DDC.GetAsynchronousResults(AsyncHandle, TempData))
+			{
+				FMemoryReader Ar(TempData, /*bIsPersistent=*/ true);
+				NumChunksCached++;
+
+				if (OutChunkData)
+				{
+					OutChunkData[ChunkIndex - FirstChunkToLoad] = FMemory::Malloc(TempData.Num());
+					Ar.Serialize(OutChunkData[ChunkIndex - FirstChunkToLoad], TempData.Num());
+				}
+			}
+			TempData.Reset();
+		}
+	}
+#endif // #if WITH_EDITORONLY_DATA
+
+	if (NumChunksCached != (Chunks.Num() - FirstChunkToLoad))
+	{
+		// Unable to cache all chunks. Release memory for those that were cached.
+		for (int32 ChunkIndex = FirstChunkToLoad; ChunkIndex < Chunks.Num(); ++ChunkIndex)
+		{
+			if (OutChunkData && OutChunkData[ChunkIndex - FirstChunkToLoad])
+			{
+				FMemory::Free(OutChunkData[ChunkIndex - FirstChunkToLoad]);
+				OutChunkData[ChunkIndex - FirstChunkToLoad] = NULL;
+			}
+		}
+		return false;
+	}
+
+	return true;
+}
+
+#if WITH_EDITORONLY_DATA
+bool FStreamedAudioPlatformData::AreDerivedChunksAvailable() const
+{
+	bool bChunksAvailable = true;
+	FDerivedDataCacheInterface& DDC = GetDerivedDataCacheRef();
+	for (int32 ChunkIndex = 0; bChunksAvailable && ChunkIndex < Chunks.Num(); ++ChunkIndex)
+	{
+		const FStreamedAudioChunk& Chunk = Chunks[ChunkIndex];
+		if (Chunk.DerivedDataKey.IsEmpty() == false)
+		{
+			bChunksAvailable = DDC.CachedDataProbablyExists(*Chunk.DerivedDataKey);
+		}
+	}
+	return bChunksAvailable;
+}
+#endif // #if WITH_EDITORONLY_DATA
+
+void FStreamedAudioPlatformData::Serialize(FArchive& Ar, USoundWave* Owner)
+{
+	Ar << NumChunks;
+	Ar << AudioFormat;
+
+	if (Ar.IsLoading())
+	{
+		Chunks.Empty(NumChunks);
+		for (int32 ChunkIndex = 0; ChunkIndex < NumChunks; ++ChunkIndex)
+		{
+			new(Chunks) FStreamedAudioChunk();
+		}
+	}
+	for (int32 ChunkIndex = 0; ChunkIndex < NumChunks; ++ChunkIndex)
+	{
+		Chunks[ChunkIndex].Serialize(Ar, Owner, ChunkIndex);
+	}
+}
+
 /**
  * Helper class to display a status update message in the editor.
  */
@@ -251,3 +800,161 @@ bool FDerivedAudioDataCompressor::Build(TArray<uint8>& OutData)
 	return OutData.Num() > 0;
 }
 
+/*---------------------------------------
+	USoundWave Derived Data functions
+---------------------------------------*/
+
+void USoundWave::CleanupCachedRunningPlatformData()
+{
+	if (RunningPlatformData != NULL)
+	{
+		delete RunningPlatformData;
+		RunningPlatformData = NULL;
+	}
+}
+
+void USoundWave::CleanupCachedCookedPlatformData()
+{
+	for (auto It : CookedPlatformData)
+	{
+		delete It.Value;
+	}
+
+	CookedPlatformData.Empty();
+}
+
+void USoundWave::SerializeCookedPlatformData(FArchive& Ar)
+{
+	if (IsTemplate())
+	{
+		return;
+	}
+
+#if WITH_EDITORONLY_DATA
+	if (Ar.IsCooking() && Ar.IsPersistent())
+	{
+		if (!Ar.CookingTarget()->IsServerOnly())
+		{
+			FName PlatformFormat = Ar.CookingTarget()->GetWaveFormat(this);
+			FString DerivedDataKey;
+			GetStreamedAudioDerivedDataKey(*this, PlatformFormat, DerivedDataKey);
+
+			FStreamedAudioPlatformData *PlatformDataToSave = CookedPlatformData.FindRef(DerivedDataKey);
+
+			if (PlatformDataToSave == NULL)
+			{
+				PlatformDataToSave = new FStreamedAudioPlatformData();
+				PlatformDataToSave->Cache(*this, PlatformFormat, EStreamedAudioCacheFlags::Async);
+
+				CookedPlatformData.Add(DerivedDataKey, PlatformDataToSave);
+			}
+
+			PlatformDataToSave->FinishCache();
+			PlatformDataToSave->Serialize(Ar, this);
+		}
+	}
+	else
+#endif // #if WITH_EDITORONLY_DATA
+	{
+		CleanupCachedRunningPlatformData();
+		check(RunningPlatformData == NULL);
+
+		RunningPlatformData = new FStreamedAudioPlatformData();
+		RunningPlatformData->Serialize(Ar, this);
+	}
+}
+
+#if WITH_EDITORONLY_DATA
+void USoundWave::CachePlatformData(bool bAsyncCache)
+{
+	FString DerivedDataKey;
+	FName AudioFormat = GetWaveFormatForRunningPlatform(*this);
+	GetStreamedAudioDerivedDataKey(*this, AudioFormat, DerivedDataKey);
+
+	if (RunningPlatformData == NULL || RunningPlatformData->DerivedDataKey != DerivedDataKey)
+	{
+		if (RunningPlatformData == NULL)
+		{
+			RunningPlatformData = new FStreamedAudioPlatformData();
+		}
+		RunningPlatformData->Cache(*this, AudioFormat, bAsyncCache ? EStreamedAudioCacheFlags::Async : EStreamedAudioCacheFlags::None);
+	}
+}
+
+void USoundWave::BeginCachePlatformData()
+{
+	CachePlatformData(true);
+
+	// enable caching in postload for derived data cache commandlet and cook by the book
+	ITargetPlatformManagerModule* TPM = GetTargetPlatformManager();
+	if (TPM && (TPM->RestrictFormatsToRuntimeOnly() == false))
+	{
+		ITargetPlatformManagerModule* TPM = GetTargetPlatformManager();
+		TArray<ITargetPlatform*> Platforms = TPM->GetActiveTargetPlatforms();
+		// Cache for all the audio formats that the cooking target requires
+		for (int32 FormatIndex = 0; FormatIndex < Platforms.Num(); FormatIndex++)
+		{
+			BeginCacheForCookedPlatformData(Platforms[FormatIndex]);
+		}
+	}
+}
+
+void USoundWave::BeginCacheForCookedPlatformData(const ITargetPlatform *TargetPlatform)
+{
+	if (IsStreaming())
+	{
+		// Retrieve format to cache for targetplatform.
+		FName PlatformFormat = TargetPlatform->GetWaveFormat(this);
+
+		uint32 CacheFlags = EStreamedAudioCacheFlags::Async;
+
+		// If source data is resident in memory then allow the streamed audio to be built
+		// in a background thread.
+		if (GetCompressedData(PlatformFormat)->IsBulkDataLoaded())
+		{
+			CacheFlags |= EStreamedAudioCacheFlags::AllowAsyncBuild;
+		}
+
+		// find format data by comparing derived data keys.
+		FString DerivedDataKey;
+		GetStreamedAudioDerivedDataKeySuffix(*this, PlatformFormat, DerivedDataKey);
+
+		FStreamedAudioPlatformData *PlatformData = CookedPlatformData.FindRef(DerivedDataKey);
+
+		if (PlatformData == NULL)
+		{
+			PlatformData = new FStreamedAudioPlatformData();
+			PlatformData->Cache(
+				*this,
+				PlatformFormat,
+				CacheFlags
+				);
+			CookedPlatformData.Add(DerivedDataKey, PlatformData);
+		}
+	}
+
+	Super::BeginCacheForCookedPlatformData(TargetPlatform);
+}
+
+void USoundWave::FinishCachePlatformData()
+{
+	if (RunningPlatformData == NULL)
+	{
+		// begin cache never called
+		CachePlatformData();
+	}
+	else
+	{
+		// make sure async requests are finished
+		RunningPlatformData->FinishCache();
+	}
+
+#if DO_CHECK
+	FString DerivedDataKey;
+	FName AudioFormat = GetWaveFormatForRunningPlatform(*this);
+	GetStreamedAudioDerivedDataKey(*this, AudioFormat, DerivedDataKey);
+
+	check(RunningPlatformData->DerivedDataKey == DerivedDataKey);
+#endif
+}
+#endif //WITH_EDITORONLY_DATA
\ No newline at end of file
diff --git a/Engine/Source/Runtime/Engine/Private/AudioDevice.cpp b/Engine/Source/Runtime/Engine/Private/AudioDevice.cpp
--- a/Engine/Source/Runtime/Engine/Private/AudioDevice.cpp
+++ b/Engine/Source/Runtime/Engine/Private/AudioDevice.cpp
@@ -309,7 +309,7 @@ void FAudioDevice::GetSoundClassInfo( TMap<FName, FAudioClassInfo>& AudioClassIn
 		}
 
 #if !WITH_EDITOR
-		AudioClassInfo->SizeResident += SoundWave->GetCompressedDataSize(GetRuntimeFormat());
+		AudioClassInfo->SizeResident += SoundWave->GetCompressedDataSize(GetRuntimeFormat(SoundWave));
 		AudioClassInfo->NumResident++;
 #else
 		switch( SoundWave->DecompressionType )
@@ -321,7 +321,7 @@ void FAudioDevice::GetSoundClassInfo( TMap<FName, FAudioClassInfo>& AudioClassIn
 			break;
 
 		case DTYPE_RealTime:
-			AudioClassInfo->SizeRealTime += SoundWave->GetCompressedDataSize(GetRuntimeFormat());
+			AudioClassInfo->SizeRealTime += SoundWave->GetCompressedDataSize(GetRuntimeFormat(SoundWave));
 			AudioClassInfo->NumRealTime++;
 			break;
 
@@ -2301,7 +2301,7 @@ void FAudioDevice::Precache(USoundWave* SoundWave, bool bSynchronous, bool bTrac
 		}
 
 		// Grab the compressed audio data
-		SoundWave->InitAudioResource(GetRuntimeFormat());
+		SoundWave->InitAudioResource(GetRuntimeFormat(SoundWave));
 
 		if (SoundWave->AudioDecompressor == NULL && SoundWave->DecompressionType == DTYPE_Native)
 		{
@@ -2320,7 +2320,7 @@ void FAudioDevice::Precache(USoundWave* SoundWave, bool bSynchronous, bool bTrac
 			}
 
 			static FName NAME_OGG(TEXT("OGG"));
-			SoundWave->bDecompressedFromOgg = GetRuntimeFormat() == NAME_OGG;
+			SoundWave->bDecompressedFromOgg = GetRuntimeFormat(SoundWave) == NAME_OGG;
 
 			// the audio decompressor will track memory
 			bTrackMemory = false;
diff --git a/Engine/Source/Runtime/Engine/Private/SoundWave.cpp b/Engine/Source/Runtime/Engine/Private/SoundWave.cpp
--- a/Engine/Source/Runtime/Engine/Private/SoundWave.cpp
+++ b/Engine/Source/Runtime/Engine/Private/SoundWave.cpp
@@ -7,6 +7,46 @@
 #include "AudioDerivedData.h"
 #include "SubtitleManager.h"
 
+/*-----------------------------------------------------------------------------
+	FStreamedAudioChunk
+-----------------------------------------------------------------------------*/
+
+void FStreamedAudioChunk::Serialize(FArchive& Ar, UObject* Owner, int32 ChunkIndex)
+{
+	bool bCooked = Ar.IsCooking();
+	Ar << bCooked;
+
+	BulkData.Serialize(Ar, Owner, ChunkIndex);
+
+#if WITH_EDITORONLY_DATA
+	if (!bCooked)
+	{
+		Ar << DerivedDataKey;
+	}
+#endif // #if WITH_EDITORONLY_DATA
+}
+
+#if WITH_EDITORONLY_DATA
+void FStreamedAudioChunk::StoreInDerivedDataCache(const FString& InDerivedDataKey)
+{
+	int32 BulkDataSizeInBytes = BulkData.GetBulkDataSize();
+	check(BulkDataSizeInBytes > 0);
+
+	TArray<uint8> DerivedData;
+	FMemoryWriter Ar(DerivedData, /*bIsPersistent=*/ true);
+	Ar << BulkDataSizeInBytes;
+	{
+		void* BulkChunkData = BulkData.Lock(LOCK_READ_ONLY);
+		Ar.Serialize(BulkChunkData, BulkDataSizeInBytes);
+		BulkData.Unlock();
+	}
+
+	GetDerivedDataCacheRef().Put(*InDerivedDataKey, DerivedData);
+	DerivedDataKey = InDerivedDataKey;
+	BulkData.RemoveBulkData();
+}
+#endif // #if WITH_EDITORONLY_DATA
+
 USoundWave::USoundWave(const class FPostConstructInitializeProperties& PCIP)
 	: Super(PCIP)
 {
@@ -27,7 +67,7 @@ SIZE_T USoundWave::GetResourceSize(EResourceSizeMode::Type Mode)
 
 	if (GEngine && GEngine->GetAudioDevice())
 	{
-		CalculatedResourceSize += GetCompressedDataSize(GEngine->GetAudioDevice()->GetRuntimeFormat());
+		CalculatedResourceSize += GetCompressedDataSize(GEngine->GetAudioDevice()->GetRuntimeFormat(this));
 	}
 
 	return CalculatedResourceSize;
@@ -126,6 +166,21 @@ void USoundWave::Serialize( FArchive& Ar )
 	}
 
 	Ar << CompressedDataGuid;
+
+	if (IsStreaming())
+	{
+		if (bCooked)
+		{
+			SerializeCookedPlatformData(Ar);
+		}
+
+#if WITH_EDITORONLY_DATA	
+		if (Ar.IsLoading() && !Ar.IsTransacting() && !bCooked && !(GetOutermost()->PackageFlags & PKG_ReloadingForCooker))
+		{
+			BeginCachePlatformData();
+		}
+#endif // #if WITH_EDITORONLY_DATA
+	}
 }
 
 /**
@@ -242,6 +297,13 @@ void USoundWave::PostLoad()
 		}
 	}
 
+#if WITH_EDITORONLY_DATA
+	if (IsStreaming())
+	{
+		FinishCachePlatformData();
+	}
+#endif // #if WITH_EDITORONLY_DATA
+
 	INC_FLOAT_STAT_BY( STAT_AudioBufferTime, Duration );
 	INC_FLOAT_STAT_BY( STAT_AudioBufferTimeChannels, NumChannels * Duration );
 }
@@ -306,6 +368,7 @@ void USoundWave::CookerWillNeverCookAgain()
 	Super::CookerWillNeverCookAgain();
 	RawData.RemoveBulkData();
 	CompressedFormatData.FlushData();
+	CleanupCachedCookedPlatformData();
 }
 
 void USoundWave::PostEditChangeProperty(FPropertyChangedEvent& PropertyChangedEvent)
@@ -313,6 +376,7 @@ void USoundWave::PostEditChangeProperty(FPropertyChangedEvent& PropertyChangedEv
 	Super::PostEditChangeProperty(PropertyChangedEvent);
 
 	static FName CompressionQualityFName = FName( TEXT( "CompressionQuality" ) );
+	static FName StreamingFName = GET_MEMBER_NAME_CHECKED(USoundWave, bStreaming);
 
 	// Prevent constant re-compression of SoundWave while properties are being changed interactively
 	if (PropertyChangedEvent.ChangeType != EPropertyChangeType::Interactive)
@@ -323,6 +387,13 @@ void USoundWave::PostEditChangeProperty(FPropertyChangedEvent& PropertyChangedEv
 		{
 			InvalidateCompressedData();
 			FreeResources();
+			UpdatePlatformData();
+			MarkPackageDirty();
+		}
+		else if (PropertyThatChanged && PropertyThatChanged->GetFName() == StreamingFName)
+		{
+			FreeResources();
+			UpdatePlatformData();
 			MarkPackageDirty();
 		}
 	}
@@ -349,7 +420,6 @@ void USoundWave::FreeResources()
 		}
 	}
 
-	NumChannels = 0;
 	SampleRate = 0;
 	Duration = 0.0f;
 	ResourceID = 0;
@@ -401,6 +471,9 @@ void USoundWave::FinishDestroy()
 	FreeResources();
 
 	Super::FinishDestroy();
+
+	CleanupCachedRunningPlatformData();
+	CleanupCachedCookedPlatformData();
 }
 
 void USoundWave::Parse( FAudioDevice* AudioDevice, const UPTRINT NodeWaveInstanceHash, FActiveSound& ActiveSound, const FSoundParseParameters& ParseParams, TArray<FWaveInstance*>& WaveInstances )
@@ -544,3 +617,22 @@ float USoundWave::GetDuration()
 {
 	return ( bLooping ? INDEFINITELY_LOOPING_DURATION : Duration);
 }
+
+bool USoundWave::IsStreaming() const
+{
+	// TODO: add in check on whether it's part of a streaming SoundGroup
+	return bStreaming;
+}
+
+void USoundWave::UpdatePlatformData()
+{
+	if (IsStreaming())
+	{
+		// TODO: Check we aren't already streaming before doing this
+
+#if WITH_EDITORONLY_DATA
+		// Recache platform data if the source has changed.
+		CachePlatformData();
+#endif
+	}
+}
diff --git a/Engine/Source/Runtime/Engine/Public/AudioDevice.h b/Engine/Source/Runtime/Engine/Public/AudioDevice.h
--- a/Engine/Source/Runtime/Engine/Public/AudioDevice.h
+++ b/Engine/Source/Runtime/Engine/Public/AudioDevice.h
@@ -476,7 +476,7 @@ public:
 	 */
 	void DeactivateReverbEffect(FName TagName);
 
-	virtual FName GetRuntimeFormat() PURE_VIRTUAL(FAudioDevice::GetRuntimeFormat,return NAME_None;);
+	virtual FName GetRuntimeFormat(USoundWave* SoundWave) PURE_VIRTUAL(FAudioDevice::GetRuntimeFormat,return NAME_None;);
 
 	/** Whether this SoundWave has an associated info class to decompress it */
 	virtual bool HasCompressedAudioInfoClass(USoundWave* SoundWave) { return false; }
diff --git a/Engine/Source/Runtime/HTML5/HTML5Audio/Private/HTML5AudioBuffer.cpp b/Engine/Source/Runtime/HTML5/HTML5Audio/Private/HTML5AudioBuffer.cpp
--- a/Engine/Source/Runtime/HTML5/HTML5Audio/Private/HTML5AudioBuffer.cpp
+++ b/Engine/Source/Runtime/HTML5/HTML5Audio/Private/HTML5AudioBuffer.cpp
@@ -111,7 +111,7 @@ FALSoundBuffer* FALSoundBuffer::CreateNativeBuffer( FALAudioDevice* AudioDevice,
 	}
 	FWaveModInfo WaveInfo;
 
-	Wave->InitAudioResource(AudioDevice->GetRuntimeFormat());
+	Wave->InitAudioResource(AudioDevice->GetRuntimeFormat(Wave));
 
 	FALSoundBuffer* Buffer = NULL;
 
diff --git a/Engine/Source/Runtime/HTML5/HTML5Audio/Public/HTML5AudioDevice.h b/Engine/Source/Runtime/HTML5/HTML5Audio/Public/HTML5AudioDevice.h
--- a/Engine/Source/Runtime/HTML5/HTML5Audio/Public/HTML5AudioDevice.h
+++ b/Engine/Source/Runtime/HTML5/HTML5Audio/Public/HTML5AudioDevice.h
@@ -197,7 +197,7 @@ public:
 	FALAudioDevice() {} 
 	virtual ~FALAudioDevice() {} 
 
-	virtual FName GetRuntimeFormat() OVERRIDE
+	virtual FName GetRuntimeFormat(USoundWave* SoundWave) OVERRIDE
 	{
 		static FName NAME_OGG(TEXT("OGG"));
 		return NAME_OGG;
diff --git a/Engine/Source/Runtime/IOS/IOSAudio/Private/IOSAudioBuffer.cpp b/Engine/Source/Runtime/IOS/IOSAudio/Private/IOSAudioBuffer.cpp
--- a/Engine/Source/Runtime/IOS/IOSAudio/Private/IOSAudioBuffer.cpp
+++ b/Engine/Source/Runtime/IOS/IOSAudio/Private/IOSAudioBuffer.cpp
@@ -59,7 +59,7 @@ FIOSAudioSoundBuffer* FIOSAudioSoundBuffer::CreateNativeBuffer(FIOSAudioDevice*
 {
 	FWaveModInfo WaveInfo;
 
-	InWave->InitAudioResource(IOSAudioDevice->GetRuntimeFormat());
+	InWave->InitAudioResource(IOSAudioDevice->GetRuntimeFormat(InWave));
 	if (!InWave->ResourceData || InWave->ResourceSize <= 0 || !WaveInfo.ReadWaveInfo(InWave->ResourceData, InWave->ResourceSize))
 	{
 		InWave->RemoveAudioResource();
diff --git a/Engine/Source/Runtime/IOS/IOSAudio/Public/IOSAudioDevice.h b/Engine/Source/Runtime/IOS/IOSAudio/Public/IOSAudioDevice.h
--- a/Engine/Source/Runtime/IOS/IOSAudio/Public/IOSAudioDevice.h
+++ b/Engine/Source/Runtime/IOS/IOSAudio/Public/IOSAudioDevice.h
@@ -178,7 +178,7 @@ public:
 	FIOSAudioDevice();
 	virtual ~FIOSAudioDevice() { }
 	
-	virtual FName GetRuntimeFormat() OVERRIDE
+	virtual FName GetRuntimeFormat(USoundWave* SoundWave) OVERRIDE
 	{
 		static FName NAME_ADPCM(TEXT("ADPCM"));
 		return NAME_ADPCM;
diff --git a/Engine/Source/Runtime/Linux/ALAudio/Private/ALAudioBuffer.cpp b/Engine/Source/Runtime/Linux/ALAudio/Private/ALAudioBuffer.cpp
--- a/Engine/Source/Runtime/Linux/ALAudio/Private/ALAudioBuffer.cpp
+++ b/Engine/Source/Runtime/Linux/ALAudio/Private/ALAudioBuffer.cpp
@@ -126,7 +126,7 @@ FALSoundBuffer* FALSoundBuffer::CreateNativeBuffer( FALAudioDevice* AudioDevice,
 	}
 
 	FWaveModInfo WaveInfo;
-	Wave->InitAudioResource(AudioDevice->GetRuntimeFormat());
+	Wave->InitAudioResource(AudioDevice->GetRuntimeFormat(Wave));
 	FALSoundBuffer* Buffer = nullptr;
 	
 	// Find the existing buffer if any
diff --git a/Engine/Source/Runtime/Linux/ALAudio/Public/ALAudioDevice.h b/Engine/Source/Runtime/Linux/ALAudio/Public/ALAudioDevice.h
--- a/Engine/Source/Runtime/Linux/ALAudio/Public/ALAudioDevice.h
+++ b/Engine/Source/Runtime/Linux/ALAudio/Public/ALAudioDevice.h
@@ -195,7 +195,7 @@ public:
 	FALAudioDevice() {} 
 	virtual ~FALAudioDevice() {} 
 
-	virtual FName GetRuntimeFormat() OVERRIDE
+	virtual FName GetRuntimeFormat(USoundWave* SoundWave) OVERRIDE
 	{
 		static FName NAME_OGG(TEXT("OGG"));
 		return NAME_OGG;
diff --git a/Engine/Source/Runtime/Mac/CoreAudio/Private/CoreAudioBuffer.cpp b/Engine/Source/Runtime/Mac/CoreAudio/Private/CoreAudioBuffer.cpp
--- a/Engine/Source/Runtime/Mac/CoreAudio/Private/CoreAudioBuffer.cpp
+++ b/Engine/Source/Runtime/Mac/CoreAudio/Private/CoreAudioBuffer.cpp
@@ -149,7 +149,7 @@ FCoreAudioSoundBuffer* FCoreAudioSoundBuffer::CreateQueuedBuffer( FCoreAudioDevi
 
 	Buffer->DecompressionState = CoreAudioDevice->CreateCompressedAudioInfo(Wave);
 	
-	Wave->InitAudioResource( CoreAudioDevice->GetRuntimeFormat() );
+	Wave->InitAudioResource( CoreAudioDevice->GetRuntimeFormat(Wave) );
 	
 	if( Buffer->DecompressionState->ReadCompressedInfo( Wave->ResourceData, Wave->ResourceSize, &QualityInfo ) )
 	{
diff --git a/Engine/Source/Runtime/Mac/CoreAudio/Public/CoreAudioDevice.h b/Engine/Source/Runtime/Mac/CoreAudio/Public/CoreAudioDevice.h
--- a/Engine/Source/Runtime/Mac/CoreAudio/Public/CoreAudioDevice.h
+++ b/Engine/Source/Runtime/Mac/CoreAudio/Public/CoreAudioDevice.h
@@ -354,7 +354,7 @@ class FCoreAudioDevice : public FAudioDevice
 	AUNode GetMatrixMixerNode() const { return MatrixMixerNode; }
 	AudioUnit GetMatrixMixerUnit() const { return MatrixMixerUnit; }
 	
-	virtual FName GetRuntimeFormat() OVERRIDE
+	virtual FName GetRuntimeFormat(USoundWave* SoundWave) OVERRIDE
 	{
 		static FName NAME_OGG(TEXT("OGG"));
 		return NAME_OGG;
diff --git a/Engine/Source/Runtime/Windows/XAudio2/Private/XAudio2Buffer.cpp b/Engine/Source/Runtime/Windows/XAudio2/Private/XAudio2Buffer.cpp
--- a/Engine/Source/Runtime/Windows/XAudio2/Private/XAudio2Buffer.cpp
+++ b/Engine/Source/Runtime/Windows/XAudio2/Private/XAudio2Buffer.cpp
@@ -328,7 +328,7 @@ FXAudio2SoundBuffer* FXAudio2SoundBuffer::CreateQueuedBuffer( FXAudio2Device* XA
 
 	Buffer->DecompressionState = XAudio2Device->CreateCompressedAudioInfo(Wave);
 
-	Wave->InitAudioResource(XAudio2Device->GetRuntimeFormat());
+	Wave->InitAudioResource(XAudio2Device->GetRuntimeFormat(Wave));
 
 	if( Buffer->DecompressionState->ReadCompressedInfo( Wave->ResourceData, Wave->ResourceSize, &QualityInfo ) )
 	{
diff --git a/Engine/Source/Runtime/Windows/XAudio2/Private/XAudio2Device.cpp b/Engine/Source/Runtime/Windows/XAudio2/Private/XAudio2Device.cpp
--- a/Engine/Source/Runtime/Windows/XAudio2/Private/XAudio2Device.cpp
+++ b/Engine/Source/Runtime/Windows/XAudio2/Private/XAudio2Device.cpp
@@ -237,6 +237,10 @@ bool FXAudio2Device::HasCompressedAudioInfoClass(USoundWave* SoundWave)
 
 class ICompressedAudioInfo* FXAudio2Device::CreateCompressedAudioInfo(USoundWave* SoundWave)
 {
+	if (SoundWave->IsStreaming())
+	{
+		return new FOpusAudioInfo();
+	}
 #if WITH_OGGVORBIS
 	return new FVorbisAudioInfo();
 #else
@@ -499,7 +503,7 @@ void FXAudio2Device::TimeTest( FOutputDevice& Ar, const TCHAR* WaveAssetName )
 		}
 		
 		// If the wave loaded in fine, time the decompression
-		Wave->InitAudioResource(GetRuntimeFormat());
+		Wave->InitAudioResource(GetRuntimeFormat(Wave));
 
 		double Start = FPlatformTime::Seconds();
 
diff --git a/Engine/Source/Runtime/Windows/XAudio2/Public/XAudio2Device.h b/Engine/Source/Runtime/Windows/XAudio2/Public/XAudio2Device.h
--- a/Engine/Source/Runtime/Windows/XAudio2/Public/XAudio2Device.h
+++ b/Engine/Source/Runtime/Windows/XAudio2/Public/XAudio2Device.h
@@ -91,8 +91,14 @@ class FXAudio2Device : public FAudioDevice
 	/** Creates a new platform specific sound source */
 	virtual FSoundSource* CreateSoundSource() OVERRIDE;
 
-	virtual FName GetRuntimeFormat() OVERRIDE
+	virtual FName GetRuntimeFormat(USoundWave* SoundWave) OVERRIDE
 	{
+		static FName NAME_OPUS(TEXT("OPUS"));
+
+		if (SoundWave->IsStreaming())
+		{
+			return NAME_OPUS;
+		}
 #if WITH_OGGVORBIS
 		static FName NAME_OGG(TEXT("OGG"));
 		return NAME_OGG;
