Fixed NavLinkProxy not being marked as "hidden" #UE4

- TTP#337898

diff --git a/Engine/Source/Runtime/Engine/Private/AI/Navigation/NavLinkProxy.cpp b/Engine/Source/Runtime/Engine/Private/AI/Navigation/NavLinkProxy.cpp
--- a/Engine/Source/Runtime/Engine/Private/AI/Navigation/NavLinkProxy.cpp
+++ b/Engine/Source/Runtime/Engine/Private/AI/Navigation/NavLinkProxy.cpp
@@ -12,6 +12,8 @@ ANavLinkProxy::ANavLinkProxy(const class FPostConstructInitializeProperties& PCI
 	TSubobjectPtr<USceneComponent> SceneComponent = PCIP.CreateDefaultSubobject<USceneComponent>(this, TEXT("PositionComponent"));
 	RootComponent = SceneComponent;
 
+	bHidden = true;
+
 #if WITH_EDITORONLY_DATA
 	EdRenderComp = PCIP.CreateDefaultSubobject<UNavLinkRenderingComponent>(this, TEXT("EdRenderComp"));
 	EdRenderComp->PostPhysicsComponentTick.bCanEverTick = false;
