Fix floating point conversion and move max angular to advanced
TTP#332279

diff --git a/Engine/Source/Runtime/Engine/Classes/PhysicsEngine/BodyInstance.h b/Engine/Source/Runtime/Engine/Classes/PhysicsEngine/BodyInstance.h
--- a/Engine/Source/Runtime/Engine/Classes/PhysicsEngine/BodyInstance.h
+++ b/Engine/Source/Runtime/Engine/Classes/PhysicsEngine/BodyInstance.h
@@ -198,7 +198,7 @@ public:
 	float LinearDamping;
 
 	/** The maximum angular velocity for this instance */
-	UPROPERTY(EditAnywhere, BlueprintReadOnly, Category = Physics, meta=(editcondition="bSimulatePhysics"))
+	UPROPERTY(EditAnywhere, AdvancedDisplay, BlueprintReadOnly, Category = Physics, meta = (editcondition = "bSimulatePhysics"))
 	float MaxAngularVelocity;
 
 	/**	Influence of rigid body physics (blending) on the mesh's pose (0.0 == use only animation, 1.0 == use only physics) */
@@ -279,6 +279,7 @@ public:
 	 */
 	bool UpdateBodyScale(const FVector& InScale3D);
 
+
 	/** Returns the center of mass of this body (in world space) */
 	FVector GetCOMPosition() const;
 
diff --git a/Engine/Source/Runtime/Engine/Private/PhysicsEngine/BodyInstance.cpp b/Engine/Source/Runtime/Engine/Private/PhysicsEngine/BodyInstance.cpp
--- a/Engine/Source/Runtime/Engine/Private/PhysicsEngine/BodyInstance.cpp
+++ b/Engine/Source/Runtime/Engine/Private/PhysicsEngine/BodyInstance.cpp
@@ -1710,6 +1710,7 @@ bool FBodyInstance::UpdateBodyScale(const FVector& InScale3D)
 	return bSuccess;
 }
 
+
 void FBodyInstance::UpdateInstanceSimulatePhysics(bool bIgnoreParent)
 {
 	// In skeletal case, we need both our bone and skelcomponent flag to be true.
@@ -2552,11 +2553,15 @@ void FBodyInstance::SetMaxAngularVelocity(float NewMaxAngVel, bool bAddToCurrent
 		{
 			float RadOldMaxAngVel = PRigidDynamic->getMaxAngularVelocity();
 			RadNewMaxAngVel += RadOldMaxAngVel;
+
+			//doing this part so our UI stays in degrees and not lose precision from the conversion
+			float OldMaxAngVel = FMath::RadiansToDegrees(RadOldMaxAngVel);
+			NewMaxAngVel += OldMaxAngVel;
 		}
 
-		PRigidDynamic->setMaxAngularVelocity(RadNewMaxAngVel);
+		PRigidDynamic->setMaxAngularVelocity(NewMaxAngVel);
 
-		MaxAngularVelocity = FMath::RadiansToDegrees(RadNewMaxAngVel);
+		MaxAngularVelocity = NewMaxAngVel;
 	}
 	else
 	{
