Improved comments

diff --git a/Engine/Source/Runtime/Engine/Private/PackageMapClient.cpp b/Engine/Source/Runtime/Engine/Private/PackageMapClient.cpp
--- a/Engine/Source/Runtime/Engine/Private/PackageMapClient.cpp
+++ b/Engine/Source/Runtime/Engine/Private/PackageMapClient.cpp
@@ -1553,6 +1553,9 @@ FNetworkGUID FNetGUIDCache::AssignNewNetGUID( const UObject * Object )
 
 	const AActor * Actor = Cast< const AActor >( Object );
 
+	// So what is going on here, is we want to make sure the network guid's of stably named sub-objects get 
+	// initialized in sync with their owning actor. What this allows us to do, is not consume any extra bandwidth 
+	// when assigning these guid's to the sub-objects on each client
 	if ( Actor == NULL )
 	{
 		const AActor * OuterActor = GetOuterActor( Object );
@@ -1564,6 +1567,8 @@ FNetworkGUID FNetGUIDCache::AssignNewNetGUID( const UObject * Object )
 
 			if ( Subobjs.Contains( const_cast< UObject * >( Object ) ) )
 			{
+				// Since we always assign guids to components when the owning actor gets assigned, we assume that if we get here
+				// the owning actor must not have been assigned yet
 				check( !NetGUIDLookup.Contains( OuterActor ) );
 
 				// Assign our owning actor a guid first (which should also assign our net guid if things are working correctly)
