Remove unnecessary bSuppressNativeHUD

diff --git a/Engine/Source/Runtime/Engine/Classes/GameFramework/HUD.h b/Engine/Source/Runtime/Engine/Classes/GameFramework/HUD.h
--- a/Engine/Source/Runtime/Engine/Classes/GameFramework/HUD.h
+++ b/Engine/Source/Runtime/Engine/Classes/GameFramework/HUD.h
@@ -186,10 +186,6 @@ class ENGINE_API AHUD : public AActor
 	UPROPERTY()
 	uint32 bShowHitBoxDebugInfo:1;    
 
-	/** If true, native HUD will not draw.  Allows blueprinted HUDs to totally replace an existing HUD. */
-	UPROPERTY(EditAnywhere, BlueprintReadWrite, Category=HUD)
-	uint32 bSuppressNativeHUD:1;
-
 	/** If true, render actor overlays. */
 	UPROPERTY()
 	uint32 bShowOverlays:1;
diff --git a/Engine/Source/Runtime/Engine/Private/HUD.cpp b/Engine/Source/Runtime/Engine/Private/HUD.cpp
--- a/Engine/Source/Runtime/Engine/Private/HUD.cpp
+++ b/Engine/Source/Runtime/Engine/Private/HUD.cpp
@@ -125,13 +125,7 @@ void AHUD::PostRender()
 	}
 	else if ( bShowHUD )
 	{
-		if (!bSuppressNativeHUD)
-		{
-			DrawHUD();
-		}
-
-		// Kismet draw
-		ReceiveDrawHUD(Canvas->SizeX, Canvas->SizeY);
+		DrawHUD();
 		
 		ULocalPlayer* LocalPlayer = Cast<ULocalPlayer>(GetOwningPlayerController()->Player);
 
@@ -325,6 +319,9 @@ void AHUD::DrawHUD()
 		PlayerOwner->GetPlayerViewPoint(ViewPoint, ViewRotation);
 		DrawActorOverlays(ViewPoint, ViewRotation);
 	}
+
+	// Blueprint draw
+	ReceiveDrawHUD(Canvas->SizeX, Canvas->SizeY);
 }
 
 
