Exposed PlayerController.ClientPlayCameraAnim() function and PlayerController.PlayerCameraManager variable to Blueprints

diff --git a/Engine/Source/Runtime/Engine/Classes/GameFramework/PlayerController.h b/Engine/Source/Runtime/Engine/Classes/GameFramework/PlayerController.h
--- a/Engine/Source/Runtime/Engine/Classes/GameFramework/PlayerController.h
+++ b/Engine/Source/Runtime/Engine/Classes/GameFramework/PlayerController.h
@@ -62,7 +62,7 @@ class ENGINE_API APlayerController : public AController
 	// Camera/view related variables
 
 	/** Camera manager associated with this Player Controller. */
-	UPROPERTY(EditInline)
+	UPROPERTY(EditInline, BlueprintReadOnly, Category=PlayerController)
 	class APlayerCameraManager* PlayerCameraManager;
 
 	/** PlayerCamera class should be set for each game, otherwise Engine.PlayerCamera is used */
@@ -534,7 +534,7 @@ public:
 	 * @param Space - Animation play area
 	 * @param CustomPlaySpace - Matrix used when Space = CAPS_UserDefined
 	 */
-	UFUNCTION(unreliable, client)
+	UFUNCTION(unreliable, client, BlueprintCallable, Category = "Game|Feedback")
 	void ClientPlayCameraAnim(class UCameraAnim* AnimToPlay, float Scale=0, float Rate=0, float BlendInTime=0, float BlendOutTime=0, bool bLoop=false, bool bRandomStartTime=false, enum ECameraAnimPlaySpace::Type Space=ECameraAnimPlaySpace::CameraLocal, FRotator CustomPlaySpace=FRotator::ZeroRotator);
 
 	/** 
