Fix misidentification of which component instance data in the cache mapped to which component

diff --git a/Engine/Source/Runtime/Engine/Private/ComponentInstanceDataCache.cpp b/Engine/Source/Runtime/Engine/Private/ComponentInstanceDataCache.cpp
--- a/Engine/Source/Runtime/Engine/Private/ComponentInstanceDataCache.cpp
+++ b/Engine/Source/Runtime/Engine/Private/ComponentInstanceDataCache.cpp
@@ -57,7 +57,7 @@ bool FComponentInstanceDataBase::MatchesComponent(const UActorComponent* Compone
 						&& (SerializedComponent->GetClass() == Component->GetClass())
 						&& (++FoundSerializedComponentsOfType == SourceComponentTypeSerializedIndex))
 					{
-						bMatches = true;
+						bMatches = (SerializedComponent == Component);
 						break;
 					}
 				}
