#ue4 - Fix various networking issues with Character crouching

- bIsCrouched is now only replicated to simulated proxies.
- Character Crouch()/UnCrouch() functions only set the bWantsToCrouch flag on CharacterMovement, and do not immediately change the crouch state.
- Client network move combining is now disabled for a move where the capsule size differs from the pending move.
- Crouching while walking now uses a much simpler approach to adjusting the capsule size and finding a suitable location.

diff --git a/Engine/Source/Runtime/Engine/Classes/GameFramework/CharacterMovementComponent.h b/Engine/Source/Runtime/Engine/Classes/GameFramework/CharacterMovementComponent.h
--- a/Engine/Source/Runtime/Engine/Classes/GameFramework/CharacterMovementComponent.h
+++ b/Engine/Source/Runtime/Engine/Classes/GameFramework/CharacterMovementComponent.h
@@ -1537,6 +1537,8 @@ public:
 	FFindFloorResult StartFloor;
 	FRotator StartRotation;
 	FRotator StartControlRotation;
+	float StartCapsuleRadius;
+	float StartCapsuleHalfHeight;
 
 	// Information after the move has been performed
 	FVector SavedLocation;
diff --git a/Engine/Source/Runtime/Engine/Private/Character.cpp b/Engine/Source/Runtime/Engine/Private/Character.cpp
--- a/Engine/Source/Runtime/Engine/Private/Character.cpp
+++ b/Engine/Source/Runtime/Engine/Private/Character.cpp
@@ -262,10 +262,6 @@ void ACharacter::Crouch(bool bClientSimulation)
 		if (CanCrouch())
 		{
 			CharacterMovement->bWantsToCrouch = true;
-			if (CharacterMovement->CanCrouchInCurrentState())
-			{
-				CharacterMovement->Crouch(bClientSimulation);
-			}
 		}
 #if !(UE_BUILD_SHIPPING || UE_BUILD_TEST)
 		else if (!CharacterMovement->CanEverCrouch())
@@ -281,7 +277,6 @@ void ACharacter::UnCrouch(bool bClientSimulation)
 	if (CharacterMovement)
 	{
 		CharacterMovement->bWantsToCrouch = false;
-		CharacterMovement->UnCrouch(bClientSimulation);
 	}
 }
 
@@ -1129,8 +1124,7 @@ void ACharacter::GetLifetimeReplicatedProps( TArray< FLifetimeProperty > & OutLi
 	DOREPLIFETIME_CONDITION( ACharacter, RepRootMotion,		COND_SimulatedOnly );
 	DOREPLIFETIME_CONDITION( ACharacter, RelativeMovement,	COND_SimulatedOnly );
 	DOREPLIFETIME_CONDITION( ACharacter, ReplicatedMovementMode, COND_SimulatedOnly );
-
-	DOREPLIFETIME( ACharacter, bIsCrouched ); // TODO: used to be !bNetOwner, now it's going to everyone (see change @revision #94)
+	DOREPLIFETIME_CONDITION( ACharacter, bIsCrouched,		COND_SimulatedOnly );
 }
 
 void ACharacter::UpdateFromCompressedFlags(uint8 Flags)
diff --git a/Engine/Source/Runtime/Engine/Private/CharacterMovementComponent.cpp b/Engine/Source/Runtime/Engine/Private/CharacterMovementComponent.cpp
--- a/Engine/Source/Runtime/Engine/Private/CharacterMovementComponent.cpp
+++ b/Engine/Source/Runtime/Engine/Private/CharacterMovementComponent.cpp
@@ -1253,20 +1253,18 @@ void UCharacterMovementComponent::PerformMovement(float DeltaSeconds)
 	FVector OldVelocity = Velocity;
 	FVector OldLocation = CharacterOwner->GetActorLocation();
 
-	bool bAllowedToCrouch = CanCrouchInCurrentState();
-	if (!bAllowedToCrouch && IsCrouching())
+	// Check for a change in crouch state. Players toggle crouch by changing bWantsToCrouch.
+	const bool bAllowedToCrouch = CanCrouchInCurrentState();
+	if ((!bAllowedToCrouch || !bWantsToCrouch) && IsCrouching())
 	{
 		UnCrouch(false);
 	}
-	else if( bWantsToCrouch && bAllowedToCrouch ) 
+	else if (bWantsToCrouch && bAllowedToCrouch && !IsCrouching()) 
 	{
-		// players crouch by setting bWantsToCrouch to true
-		if( !IsCrouching() )
-		{
-			Crouch(false);
-		}
+		Crouch(false);
 	}
 
+	// Character::LaunchCharacter() has been deferred until now.
 	HandlePendingLaunch();
 
 	// If using RootMotion, tick animations before running physics.
@@ -1319,10 +1317,8 @@ void UCharacterMovementComponent::PerformMovement(float DeltaSeconds)
 		return;
 	}
 
-	// uncrouch if no longer desiring crouch
-	// or if not in walking or falling physics
-	bAllowedToCrouch = CanCrouchInCurrentState();
-	if (IsCrouching() && (!bAllowedToCrouch || !bWantsToCrouch))
+	// uncrouch if no longer allowed to be crouched
+	if (IsCrouching() && !CanCrouchInCurrentState())
 	{
 		UnCrouch(false);
 	}
@@ -1499,6 +1495,8 @@ void UCharacterMovementComponent::UnCrouch(bool bClientSimulation)
 		return;
 	}
 
+	const float CurrentCrouchedHalfHeight = CharacterOwner->CapsuleComponent->GetScaledCapsuleHalfHeight();
+
 	const float ComponentScale = CharacterOwner->CapsuleComponent->GetShapeScale();
 	const float OldUnscaledHalfHeight = CharacterOwner->CapsuleComponent->GetUnscaledCapsuleHalfHeight();
 	const float HalfHeightAdjust = DefaultCharacter->CapsuleComponent->GetUnscaledCapsuleHalfHeight() - OldUnscaledHalfHeight;
@@ -1513,72 +1511,89 @@ void UCharacterMovementComponent::UnCrouch(bool bClientSimulation)
 
 	if( !bClientSimulation )
 	{
-		UPrimitiveComponent* OldBaseComponent = CharacterOwner->GetMovementBase();
-		FFindFloorResult OldFloor = CurrentFloor;
-		SetBase(NULL,false);
-
 		// Try to stay in place and see if the larger capsule fits. We use a slightly taller capsule to avoid penetration.
 		static const FName NAME_CrouchTrace = FName(TEXT("CrouchTrace"));
-		const float SweepInflation = KINDA_SMALL_NUMBER;
+		const float SweepInflation = KINDA_SMALL_NUMBER * 10.f;
 		FCollisionQueryParams CapsuleParams(NAME_CrouchTrace, false, CharacterOwner);
 		FCollisionResponseParams ResponseParam;
 		InitCollisionParams(CapsuleParams, ResponseParam);
-		const FCollisionShape StandingCapsuleShape = GetPawnCapsuleCollisionShape(SHRINK_HeightCustom, -SweepInflation);
+		const FCollisionShape StandingCapsuleShape = GetPawnCapsuleCollisionShape(SHRINK_HeightCustom, -SweepInflation); // Shrink by negative amount, so actually grow it.
 		const ECollisionChannel CollisionChannel = UpdatedComponent->GetCollisionObjectType();
-		// Perf: Avoid this test when on the ground, because it almost always fails and we should just continue to the next step.
 		bool bEncroached = true;
+
 		if (!IsMovingOnGround())
 		{
+			// Expand in place
 			bEncroached = GetWorld()->OverlapTest(PawnLocation, FQuat::Identity, CollisionChannel, StandingCapsuleShape, CapsuleParams, ResponseParam);
-		}
-
-		if (bEncroached)
-		{
-			// Try adjusting capsule position to see if we can avoid encroachment.
-			if (ScaledHalfHeightAdjust > 0.f)
+		
+			if (bEncroached)
 			{
-				// Shrink to a short capsule, sweep down to base to find where that would hit something, and then try to stand up from there.
-				float PawnRadius, PawnHalfHeight;
-				CharacterOwner->CapsuleComponent->GetScaledCapsuleSize(PawnRadius, PawnHalfHeight);
-				const float ShrinkHalfHeight = PawnHalfHeight - PawnRadius;
-				const float TraceDist = PawnHalfHeight - ShrinkHalfHeight;
-				const FVector Down = FVector(0.f, 0.f, -TraceDist);
-
-				FHitResult Hit(1.f);
-				const FCollisionShape ShortCapsuleShape = GetPawnCapsuleCollisionShape(SHRINK_HeightCustom, ShrinkHalfHeight);
-				const bool bBlockingHit = GetWorld()->SweepSingle(Hit, PawnLocation, PawnLocation + Down, FQuat::Identity, CollisionChannel, ShortCapsuleShape, CapsuleParams);
-				if (Hit.bStartPenetrating)
-				{
-					bEncroached = true;
-				}
-				else
+				// Try adjusting capsule position to see if we can avoid encroachment.
+				if (ScaledHalfHeightAdjust > 0.f)
 				{
-					// Compute where the base of the sweep ended up, and see if we can stand there
-					const float DistanceToBase = (Hit.Time * TraceDist) + ShortCapsuleShape.Capsule.HalfHeight;
-					const FVector NewLoc = FVector(PawnLocation.X, PawnLocation.Y, PawnLocation.Z - DistanceToBase + PawnHalfHeight + SweepInflation + MIN_FLOOR_DIST / 2.f);
-					bEncroached = GetWorld()->OverlapTest(NewLoc, FQuat::Identity, CollisionChannel, StandingCapsuleShape, CapsuleParams);
-					if (!bEncroached)
+					// Shrink to a short capsule, sweep down to base to find where that would hit something, and then try to stand up from there.
+					float PawnRadius, PawnHalfHeight;
+					CharacterOwner->CapsuleComponent->GetScaledCapsuleSize(PawnRadius, PawnHalfHeight);
+					const float ShrinkHalfHeight = PawnHalfHeight - PawnRadius;
+					const float TraceDist = PawnHalfHeight - ShrinkHalfHeight;
+					const FVector Down = FVector(0.f, 0.f, -TraceDist);
+
+					FHitResult Hit(1.f);
+					const FCollisionShape ShortCapsuleShape = GetPawnCapsuleCollisionShape(SHRINK_HeightCustom, ShrinkHalfHeight);
+					const bool bBlockingHit = GetWorld()->SweepSingle(Hit, PawnLocation, PawnLocation + Down, FQuat::Identity, CollisionChannel, ShortCapsuleShape, CapsuleParams);
+					if (Hit.bStartPenetrating)
 					{
-						// Intentionally not using MoveUpdatedComponent, where a horizontal plane constraint would prevent the base of the capsule from staying at the same spot.
-						UpdatedComponent->MoveComponent(NewLoc - PawnLocation, CharacterOwner->GetActorRotation(), false);
-						const bool bBaseChanged = (Hit.Component.Get() != OldBaseComponent);
-						CurrentFloor.SetFromSweep(Hit, 0.f, IsWalkable(Hit));
-						SetBase(Hit.Component.Get(), bBaseChanged);
+						bEncroached = true;
+					}
+					else
+					{
+						// Compute where the base of the sweep ended up, and see if we can stand there
+						const float DistanceToBase = (Hit.Time * TraceDist) + ShortCapsuleShape.Capsule.HalfHeight;
+						const FVector NewLoc = FVector(PawnLocation.X, PawnLocation.Y, PawnLocation.Z - DistanceToBase + PawnHalfHeight + SweepInflation + MIN_FLOOR_DIST / 2.f);
+						bEncroached = GetWorld()->OverlapTest(NewLoc, FQuat::Identity, CollisionChannel, StandingCapsuleShape, CapsuleParams);
+						if (!bEncroached)
+						{
+							// Intentionally not using MoveUpdatedComponent, where a horizontal plane constraint would prevent the base of the capsule from staying at the same spot.
+							UpdatedComponent->MoveComponent(NewLoc - PawnLocation, CharacterOwner->GetActorRotation(), false);
+						}
 					}
 				}
 			}
+		}
+		else
+		{
+			// Moving on ground
+			// Expand while keeping base location the same.
+			FVector StandingLocation = PawnLocation + FVector(0.f, 0.f, StandingCapsuleShape.GetCapsuleHalfHeight() - CurrentCrouchedHalfHeight);
+			bEncroached = GetWorld()->OverlapTest(StandingLocation, FQuat::Identity, CollisionChannel, StandingCapsuleShape, CapsuleParams, ResponseParam);
 
-			// If still encroached then abort.
 			if (bEncroached)
 			{
-				CharacterOwner->CapsuleComponent->SetCapsuleSize(CharacterOwner->CapsuleComponent->GetUnscaledCapsuleRadius(), OldUnscaledHalfHeight, false);
-				CharacterOwner->CapsuleComponent->UpdateBounds(); // Update bounds again back to old value
-				CurrentFloor = OldFloor;
-				SetBase(OldBaseComponent,false);
-				return;
+				// Something might be just barely overhead, try moving down closer to the floor to avoid it.
+				const float MinFloorDist = KINDA_SMALL_NUMBER * 10.f;
+				if (CurrentFloor.bBlockingHit && CurrentFloor.FloorDist > MinFloorDist)
+				{
+					StandingLocation.Z -= CurrentFloor.FloorDist - MinFloorDist;
+					bEncroached = GetWorld()->OverlapTest(StandingLocation, FQuat::Identity, CollisionChannel, StandingCapsuleShape, CapsuleParams, ResponseParam);
+				}
+			}
+
+			if (!bEncroached)
+			{
+				// Commit the change in location.
+				UpdatedComponent->MoveComponent(StandingLocation - PawnLocation, CharacterOwner->GetActorRotation(), false);
+				bForceNextFloorCheck = true;
 			}
 		}
 
+		// If still encroached then abort.
+		if (bEncroached)
+		{
+			CharacterOwner->CapsuleComponent->SetCapsuleSize(CharacterOwner->CapsuleComponent->GetUnscaledCapsuleRadius(), OldUnscaledHalfHeight, false);
+			CharacterOwner->CapsuleComponent->UpdateBounds(); // Update bounds again back to old value
+			return;
+		}
+
 		CharacterOwner->bIsCrouched = false;
 	}	
 	else
@@ -1590,7 +1605,6 @@ void UCharacterMovementComponent::UnCrouch(bool bClientSimulation)
 	bUpdateOverlaps = true;
 	CharacterOwner->CapsuleComponent->SetCapsuleSize(DefaultCharacter->CapsuleComponent->GetUnscaledCapsuleRadius(), DefaultCharacter->CapsuleComponent->GetUnscaledCapsuleHalfHeight(), bUpdateOverlaps);
 
-	bForceNextFloorCheck = true;
 	AdjustProxyCapsuleSize();
 	CharacterOwner->OnEndCrouch( HalfHeightAdjust, ScaledHalfHeightAdjust );
 }
@@ -3231,13 +3245,12 @@ void UCharacterMovementComponent::AdjustFloorHeight()
 		}
 		else if (MoveDist > 0.f)
 		{
-			// If moving up, use the actual impact location, not the pulled back time/location.
-			const float FloorDistTime = FMath::Abs((InitialZ - AdjustHit.Location.Z) / MoveDist);
-			CurrentFloor.FloorDist += MoveDist * FloorDistTime;
+			const float CurrentZ = UpdatedComponent->GetComponentLocation().Z;
+			CurrentFloor.FloorDist += CurrentZ - InitialZ;
 		}
 		else
 		{
-			check(MoveDist < 0.f);
+			checkSlow(MoveDist < 0.f);
 			const float CurrentZ = UpdatedComponent->GetComponentLocation().Z;
 			CurrentFloor.FloorDist = CurrentZ - AdjustHit.Location.Z;
 			if (IsWalkable(AdjustHit))
@@ -6005,6 +6018,9 @@ void FSavedMove_Character::Clear()
 	StartFloor = FFindFloorResult();
 	StartRotation = FRotator::ZeroRotator;
 	StartControlRotation = FRotator::ZeroRotator;
+	StartCapsuleRadius = 0.f;
+	StartCapsuleHalfHeight = 0.f;
+
 	SavedLocation = FVector::ZeroVector;
 	SavedRotation = FRotator::ZeroRotator;
 	SavedRelativeLocation = FVector::ZeroVector;
@@ -6055,6 +6071,8 @@ void FSavedMove_Character::SetInitialPosition(ACharacter* Character)
 	{
 		StartRelativeLocation = Character->RelativeMovement.Location;
 	}
+
+	Character->CapsuleComponent->GetScaledCapsuleSize(StartCapsuleRadius, StartCapsuleHalfHeight);
 }
 
 void FSavedMove_Character::PostUpdate(ACharacter* Character, FSavedMove_Character::EPostUpdateMode PostUpdateMode)
@@ -6145,6 +6163,8 @@ bool FSavedMove_Character::CanCombineWith(const FSavedMovePtr& NewMove, ACharact
 			&& bWantsToCrouch == NewMove->bWantsToCrouch
 			&& StartBase == NewMove->StartBase
 			&& MovementMode == NewMove->MovementMode
+			&& StartCapsuleRadius == NewMove->StartCapsuleRadius
+			&& StartCapsuleHalfHeight == NewMove->StartCapsuleHalfHeight
 			&& (CustomTimeDilation == NewMove->CustomTimeDilation);
 	}
 	else
@@ -6154,6 +6174,8 @@ bool FSavedMove_Character::CanCombineWith(const FSavedMovePtr& NewMove, ACharact
 			&& bWantsToCrouch == NewMove->bWantsToCrouch
 			&& StartBase == NewMove->StartBase
 			&& MovementMode == NewMove->MovementMode
+			&& StartCapsuleRadius == NewMove->StartCapsuleRadius
+			&& StartCapsuleHalfHeight == NewMove->StartCapsuleHalfHeight
 			&& ((AccelNormal | NewMove->AccelNormal) > 0.99f)
 			&& (CustomTimeDilation == NewMove->CustomTimeDilation);
 	}
