#ue4 - Add "IsWalking" helper to CharacterMovementComponent.

- Mostly just a wrapper for Blueprint users who don't think to use IsMovingOnGround().

diff --git a/Engine/Source/Runtime/Engine/Classes/GameFramework/CharacterMovementComponent.h b/Engine/Source/Runtime/Engine/Classes/GameFramework/CharacterMovementComponent.h
--- a/Engine/Source/Runtime/Engine/Classes/GameFramework/CharacterMovementComponent.h
+++ b/Engine/Source/Runtime/Engine/Classes/GameFramework/CharacterMovementComponent.h
@@ -816,6 +816,10 @@ public:
 	virtual void AddRadialImpulse(const FVector& Origin, float Radius, float Strength, enum ERadialImpulseFalloff Falloff, bool bVelChange) OVERRIDE;
 	//END UMovementComponent Interface
 
+	/** @return true if the character is in the 'Walking' movement mode. */
+	UFUNCTION(BlueprintCallable, Category="Pawn|Components|CharacterMovement")
+	bool IsWalking() const;
+
 	//BEGIN UNavMovementComponent Interface
 	virtual void RequestDirectMove(const FVector& MoveVelocity, bool bForceMaxSpeed) OVERRIDE;
 	virtual bool CanStopPathFollowing() const OVERRIDE;
diff --git a/Engine/Source/Runtime/Engine/Private/CharacterMovementComponent.cpp b/Engine/Source/Runtime/Engine/Private/CharacterMovementComponent.cpp
--- a/Engine/Source/Runtime/Engine/Private/CharacterMovementComponent.cpp
+++ b/Engine/Source/Runtime/Engine/Private/CharacterMovementComponent.cpp
@@ -1904,6 +1904,11 @@ bool UCharacterMovementComponent::IsCrouching() const
 	return CharacterOwner && CharacterOwner->bIsCrouched;
 }
 
+bool UCharacterMovementComponent::IsWalking() const
+{
+	return IsMovingOnGround();
+}
+
 void UCharacterMovementComponent::CalcVelocity(float DeltaTime, float Friction, bool bFluid, float BrakingDeceleration)
 {
 	// Do not update velocity when using root motion
