Physics: Add WIP support for 2D RaycastMulti and fix some issues with 'into the screen' 2D RayCastSingle calls

diff --git a/Engine/Source/Runtime/Engine/Private/Collision/PhysXCollision.cpp b/Engine/Source/Runtime/Engine/Private/Collision/PhysXCollision.cpp
--- a/Engine/Source/Runtime/Engine/Private/Collision/PhysXCollision.cpp
+++ b/Engine/Source/Runtime/Engine/Private/Collision/PhysXCollision.cpp
@@ -6,8 +6,9 @@
 
 #if WITH_BOX2D
 #include "../PhysicsEngine2D/Box2DIntegration.h"
+#include "CollisionConversions.h"
 
-class FRaycastHelperBox2D : public b2RayCastCallback, public b2QueryCallback
+class FRaycastHelperBox2D_Base : public b2RayCastCallback, public b2QueryCallback
 {
 public:
 	enum EResponse
@@ -17,18 +18,14 @@ public:
 		ResponseBlock = 2,
 	};
 
-	FRaycastHelperBox2D(ECollisionChannel InTraceChannel, bool bInTraceComplex, const FCollisionResponseContainer& InCollisionResponseContainer, const FCollisionObjectQueryParams& InObjectParam, const TArray<uint32, TInlineAllocator<1> >& InIgnoreActors, bool bInMultitrace)
+	FRaycastHelperBox2D_Base(ECollisionChannel InTraceChannel, const struct FCollisionQueryParams& InParams, const FCollisionResponseContainer& InCollisionResponseContainer, const FCollisionObjectQueryParams& InObjectParam, bool bInMultitrace)
 		: TraceChannel(InTraceChannel)
+		, Params(InParams)
 		, CollisionResponseContainer(InCollisionResponseContainer)
 		, ObjectParam(InObjectParam)
-		, IgnoreActors(InIgnoreActors)
-		, BestResult(MAX_FLT)
-		, BestFixture(nullptr)
 		, QueryBlockingChannels(0)
 		, QueryTouchingChannels(0)
 		, bMultitrace(bInMultitrace)
-		, bTraceComplex(bInTraceComplex)
-		, bHasBlockingHit(false)
 	{
 		// Box2D is not thread safe
 		check(IsInGameThread());
@@ -46,48 +43,39 @@ public:
 					QueryTouchingChannels |= CRC_TO_BITFIELD(i);
 				}
 			}
-
 		}
 	}
 
-	// b2RayCastCallback interface
-	virtual float32 ReportFixture(b2Fixture* Fixture, const b2Vec2& Point, const b2Vec2& Normal, float32 Fraction) OVERRIDE;
-	// End of b2RayCastCallback interface
-
-	// b2QueryCallback interface
-	virtual bool ReportFixture(b2Fixture* Fixture) OVERRIDE;
-	// End of b2QueryCallback interface
-
-	bool HasBlockingHit() const
-	{
-		return bHasBlockingHit;
-	}
+protected:
+	const ECollisionChannel TraceChannel;
+	const struct FCollisionQueryParams& Params;
+	const FCollisionResponseContainer& CollisionResponseContainer;
+	const FCollisionObjectQueryParams& ObjectParam;
+	uint32 QueryBlockingChannels;
+	uint32 QueryTouchingChannels;
+	bool bMultitrace;
 
-	float GetBestTime() const
-	{
-		return BestResult;
-	}
+protected:
+	// Determines how the query and the fixture interact (either no interaction or touching/blocking each other)
+	EResponse CalcQueryHitType(b2Fixture* Fixture) const;
 
 	// Converts the query results to a hit result
-	void ConvertToHitResult(FHitResult& OutResult, const FVector& Start, const FVector& End, const FCollisionQueryParams& Params)
+	void ConvertToHitResult_Internal(FHitResult& OutResult, const FVector& Start, const FVector& End, const float InHitTime, b2Fixture* InFixture, const FVector& InNormal, const EResponse InHitType)
 	{
-		check(bHasBlockingHit);
-
 		//@TODO: BOX2D: How to deal with initial penetration (PxU32)(PHit.hadInitialOverlap());
 #if 0
-		if ((BestFixture != NULL) && PHit.hadInitialOverlap())
+		if ((InFixture != nullptr) && PHit.hadInitialOverlap())
 		{
-			ConvertOverlappedShapeToImpactHit(PHit, StartLoc, EndLoc, OutResult, *BestFixture, QueryTM, QueryFilter, Params.bReturnPhysicalMaterial);
+			ConvertOverlappedShapeToImpactHit(PHit, StartLoc, EndLoc, OutResult, *InFixture, QueryTM, QueryFilter, Params.bReturnPhysicalMaterial);
 			return;
 		}
 #endif
 		OutResult.bStartPenetrating = false;
 
-		check(BestFixture);
-		FBodyInstance* OwningBodyInstance = static_cast<FBodyInstance*>(BestFixture->GetBody()->GetUserData());
+		check(InFixture);
+		FBodyInstance* OwningBodyInstance = static_cast<FBodyInstance*>(InFixture->GetBody()->GetUserData());
 		UPrimitiveComponent* OwningComponent = OwningBodyInstance->OwnerComponent.Get();
 
-
 		/** Set info in the HitResult (Actor, Component, PhysMaterial, BoneName, Item) based on the supplied shape and face index */
 		OutResult.PhysMaterial = NULL;
 
@@ -109,274 +97,369 @@ public:
 		{
 			OutResult.BoneName = OwningBodyInstance->BodySetup->BoneName;
 		}
-	
-
 
 		OutResult.FaceIndex = INDEX_NONE;
 
-
-		// calculate the hit time
-		const float HitTime = BestResult;
-
 		// figure out where the the "safe" location for this shape is by moving from the startLoc toward the ImpactPoint
 		const FVector TraceDir = End - Start;
-		const FVector SafeLocationToFitShape = Start + (HitTime * TraceDir);
+		const FVector SafeLocationToFitShape = Start + (InHitTime * TraceDir);
 
 		// Other info
 		OutResult.Location = SafeLocationToFitShape;
 		OutResult.ImpactPoint = SafeLocationToFitShape;
-		OutResult.Normal = FPhysicsIntegration2D::ConvertBoxVectorToUnreal(BestNormal).SafeNormal();
+		OutResult.Normal = InNormal;
 		OutResult.ImpactNormal = OutResult.Normal;
 
 		OutResult.TraceStart = Start;
 		OutResult.TraceEnd = End;
-		OutResult.Time = HitTime;
+		OutResult.Time = InHitTime;
 
 		// See if this is a 'blocking' hit
-		OutResult.bBlockingHit = (BestHitType == ResponseBlock);
+		OutResult.bBlockingHit = (InHitType == ResponseBlock);
 	}
+};
+
+FRaycastHelperBox2D_Base::EResponse FRaycastHelperBox2D_Base::CalcQueryHitType(b2Fixture* Fixture) const
+{
+	const b2Filter& FixtureFilterData = Fixture->GetFilterData();
+	const ECollisionChannel ShapeChannel = static_cast<ECollisionChannel>(FixtureFilterData.ObjectTypeAndFlags >> 24);
+	const uint32 ShapeBit = ECC_TO_BITFIELD(ShapeChannel);
 
-	void PointQuery(const b2World* BoxWorld, const b2Vec2& BoxStart, const FVector& Start, const FVector& End, const FCollisionQueryParams& Params, FHitResult& OutResult);
+	EResponse Result = ResponseNone;
 
-	bool RayTest(const b2World* BoxWorld, const FVector& Start, const FVector& End)
+	if (ObjectParam.IsValid())
 	{
-		const b2Vec2 BoxStart(FPhysicsIntegration2D::ConvertUnrealVectorToBox(Start));
-		const b2Vec2 BoxEnd(FPhysicsIntegration2D::ConvertUnrealVectorToBox(End));
+		// do I belong to one of objects of interest?
+		if ((ShapeBit & ObjectParam.GetQueryBitfield()) != 0)
+		{
+			Result = bMultitrace ? ResponseTouch : ResponseBlock;
+		}
+	}
+	else
+	{
+		// Then see if the channel wants to be blocked
+		const uint32 TraceBit = ECC_TO_BITFIELD(TraceChannel);
 
-		// Now check the length of the 2D query, we may end up having to do an AABB test instead if it's a ray cast along the 'Z' axis of the 2D scene
-		const b2Vec2 BoxDelta = BoxEnd - BoxStart;
-		if (BoxDelta.LengthSquared() < SMALL_NUMBER)
+		// check if the fixture has a hit or touch
+		EResponse TraceHitType = ResponseNone;
+		if ((TraceBit & FixtureFilterData.BlockingChannels) != 0)
 		{
-			// Trace into the scene, do a point query instead
-			b2AABB QueryBox;
-			QueryBox.lowerBound = BoxStart;
-			QueryBox.upperBound = BoxStart;
-			BoxWorld->QueryAABB(this, QueryBox);
+			TraceHitType = ResponseBlock;
 		}
-		else
+		else if ((TraceBit & FixtureFilterData.TouchingChannels) != 0)
 		{
-			// Cast the ray
-			BoxWorld->RayCast(this, BoxStart, BoxEnd);
+			TraceHitType = ResponseTouch;
 		}
 
-		return bHasBlockingHit;
+		// Check if the shape has a hit or touch
+		EResponse ShapeHitType = ResponseNone;
+		if ((ShapeBit & QueryBlockingChannels) != 0)
+		{
+			ShapeHitType = ResponseBlock;
+		}
+		else if ((ShapeBit & QueryTouchingChannels) != 0)
+		{
+			ShapeHitType = ResponseTouch;
+		}
+
+		// return minimum agreed-upon interaction
+		Result = FMath::Min(TraceHitType, ShapeHitType);
 	}
 
-	// Determines how the query and the fixture interact (either no interaction or touching/blocking each other)
-	EResponse CalcQueryHitType(b2Fixture* Fixture) const;
-private:
-	const ECollisionChannel TraceChannel;
-	const FCollisionResponseContainer& CollisionResponseContainer;
-	const FCollisionObjectQueryParams& ObjectParam;
-	const TArray<uint32, TInlineAllocator<1> >& IgnoreActors;
+	// Single traces should ignore touches
+	if (!bMultitrace && (Result == ResponseTouch))
+	{
+		Result = ResponseNone;
+	}
 
-	b2Vec2 BestPoint;
-	b2Vec2 BestNormal;
-	float BestResult;
-	b2Fixture* BestFixture;
-	uint32 QueryBlockingChannels;
-	uint32 QueryTouchingChannels;
-	EResponse BestHitType;
+	// Check to see if the actor ultimately owning the fixture should be ignored
+	if (Result != ResponseNone)
+	{
+		FBodyInstance* OwnerBodyInstance = static_cast<FBodyInstance*>(Fixture->GetBody()->GetUserData());
+		UPrimitiveComponent* OwnerComponent = OwnerBodyInstance->OwnerComponent.Get();
+		AActor* Owner = OwnerComponent ? OwnerComponent->GetOwner() : nullptr;
 
-	bool bMultitrace;
-	bool bTraceComplex;
-	bool bHasBlockingHit;
-};
+		if ((Owner != nullptr) && Params.IgnoreActors.Contains(Owner->GetUniqueID()))
+		{
+			Result = ResponseNone;
+		}
+	}
 
-void FRaycastHelperBox2D::PointQuery(const b2World* BoxWorld, const b2Vec2& BoxStart, const FVector& Start, const FVector& End, const FCollisionQueryParams& Params, FHitResult& OutResult)
+	return Result;
+}
+
+class FRaycastHelperBox2D_RayCastBase : public FRaycastHelperBox2D_Base
 {
-	// Run the query
-	b2AABB QueryBox;
-	QueryBox.lowerBound = BoxStart;
-	QueryBox.upperBound = BoxStart;
-	BoxWorld->QueryAABB(this, QueryBox);
+public:
+	FRaycastHelperBox2D_RayCastBase(ECollisionChannel InTraceChannel, const struct FCollisionQueryParams& InParams, const FCollisionResponseContainer& InCollisionResponseContainer, const FCollisionObjectQueryParams& InObjectParam, bool bInMultitrace)
+		: FRaycastHelperBox2D_Base(InTraceChannel, InParams, InCollisionResponseContainer, InObjectParam, bInMultitrace)
+	{
+	}
 
-	// Deal with a hit
-	if (bHasBlockingHit)
+	// b2RayCastCallback interface
+	virtual float32 ReportFixture(b2Fixture* Fixture, const b2Vec2& Point, const b2Vec2& Normal, float32 Fraction) override
 	{
-		//@TODO: BOX2D: How to deal with initial penetration (PxU32)(PHit.hadInitialOverlap());
-#if 0
-		if ((BestFixture != NULL) && PHit.hadInitialOverlap())
+		const EResponse ResponseMode = CalcQueryHitType(Fixture);
+
+		if (ResponseMode != ResponseNone)
 		{
-			ConvertOverlappedShapeToImpactHit(PHit, StartLoc, EndLoc, OutResult, *BestFixture, QueryTM, QueryFilter, Params.bReturnPhysicalMaterial);
-			return;
+			ProcessHit(Fixture, Fraction, FPhysicsIntegration2D::ConvertBoxVectorToUnreal(Normal).SafeNormal(), ResponseMode);
 		}
-#endif
-		OutResult.bStartPenetrating = false;
 
-		check(BestFixture);
-		FBodyInstance* OwningBodyInstance = static_cast<FBodyInstance*>(BestFixture->GetBody()->GetUserData());
-		UPrimitiveComponent* OwningComponent = OwningBodyInstance->OwnerComponent.Get();
+		switch (ResponseMode)
+		{
+		case ResponseBlock:
+			return Fraction;
+		case ResponseTouch:
+			return 1.0f;
+		default:
+		case ResponseNone:
+			return -1.0f;
+		}
+	}
+	// End of b2RayCastCallback interface
 
-		/** Set info in the HitResult (Actor, Component, PhysMaterial, BoneName, Item) based on the supplied shape and face index */
-		OutResult.PhysMaterial = NULL;
+	// b2QueryCallback interface
+	virtual bool ReportFixture(b2Fixture* Fixture) override
+	{
+		const EResponse ResponseMode = CalcQueryHitType(Fixture);
 
-		// Grab actor/component
-		if (OwningComponent)
+		bool bReportHit = (ResponseMode == ResponseBlock);
+
+		if (bReportHit)
 		{
-			OutResult.Actor = OwningComponent->GetOwner();
-			OutResult.Component = OwningComponent;
+			// Figure out 'when' the ray hit the fixture
+			FBodyInstance* OwningBodyInstance = static_cast<FBodyInstance*>(Fixture->GetBody()->GetUserData());
+			UPrimitiveComponent* OwningComponent = OwningBodyInstance->OwnerComponent.Get();
+			const float IntersectTime = (OwningComponent != nullptr) ? FPhysicsIntegration2D::ConvertUnrealVectorToPerpendicularDistance(OwningComponent->GetComponentLocation()) : PointQueryStartTime;
+			const float Fraction = (IntersectTime - PointQueryStartTime) * PointQueryScaleFactor;
 
-			if (Params.bReturnPhysicalMaterial)
-			{
-				OutResult.PhysMaterial = OwningBodyInstance->GetSimplePhysicalMaterial();
-			}
+			ProcessHit(Fixture, Fraction, PointQueryNormal, ResponseMode);
 		}
 
-		// Body index / bone name
-		OutResult.Item = OwningBodyInstance->InstanceBodyIndex;
-		if (OwningBodyInstance->BodySetup.IsValid())
+		switch (ResponseMode)
 		{
-			OutResult.BoneName = OwningBodyInstance->BodySetup->BoneName;
+		case ResponseBlock:
+			return true;
+		case ResponseTouch:
+			return false;
+		default:
+		case ResponseNone:
+			return false;
 		}
+	}
+	// End of b2QueryCallback interface
 
+protected:
+	// Only initialized / used for point queries
+	float PointQueryStartTime;
+	float PointQueryScaleFactor;
+	FVector PointQueryNormal;
 
-		OutResult.FaceIndex = INDEX_NONE;
-
-		// Figure out 'when' the ray hit the fixture
-		const float StartTime = FPhysicsIntegration2D::ConvertUnrealVectorToPerpendicularDistance(Start);
-		const float EndTime = FPhysicsIntegration2D::ConvertUnrealVectorToPerpendicularDistance(End);
-		const float IntersectTime = (OwningComponent != nullptr) ? FPhysicsIntegration2D::ConvertUnrealVectorToPerpendicularDistance(OwningComponent->GetComponentLocation()) : StartTime;
-		const float HitTime = IntersectTime / (EndTime - StartTime);
-
-		// figure out where the the "safe" location for this shape is by moving from the startLoc toward the ImpactPoint
-		const FVector TraceDir = End - Start;
-		const FVector SafeLocationToFitShape = Start + (HitTime * TraceDir);
-
-		// Other info
-		OutResult.Location = SafeLocationToFitShape;
-		OutResult.ImpactPoint = SafeLocationToFitShape;
-		OutResult.Normal = (Start - End).SafeNormal();
-		OutResult.ImpactNormal = OutResult.Normal;
+protected:
+	virtual void ProcessHit(b2Fixture* Fixture, const float HitTime, const FVector& Normal, const EResponse ResponseMode) = 0;
 
-		OutResult.TraceStart = Start;
-		OutResult.TraceEnd = End;
-		OutResult.Time = HitTime;
+	void RayCastInternal(const b2World* BoxWorld, const FVector& Start, const FVector& End)
+	{
+		const b2Vec2 BoxStart(FPhysicsIntegration2D::ConvertUnrealVectorToBox(Start));
+		const b2Vec2 BoxEnd(FPhysicsIntegration2D::ConvertUnrealVectorToBox(End));
 
-		// See if this is a 'blocking' hit
-		OutResult.bBlockingHit = (BestHitType == ResponseBlock);
-	}
-}
+		// Now check the length of the 2D query, we may end up having to do an AABB test instead if it's a ray cast along the 'Z' axis of the 2D scene
+		const b2Vec2 BoxDelta = BoxEnd - BoxStart;
+		if (BoxDelta.LengthSquared() < SMALL_NUMBER)
+		{
+			// Tracing directly into the scene, so do a point query instead
+			PointQueryStartTime = FPhysicsIntegration2D::ConvertUnrealVectorToPerpendicularDistance(Start);
+			const float PointQueryEndTime = FPhysicsIntegration2D::ConvertUnrealVectorToPerpendicularDistance(End);
+			PointQueryScaleFactor = 1.0f / (PointQueryEndTime - PointQueryStartTime);
+			PointQueryNormal = (Start - End).SafeNormal();
 
-float32 FRaycastHelperBox2D::ReportFixture(b2Fixture* Fixture, const b2Vec2& Point, const b2Vec2& Normal, float32 Fraction)
-{
-	const EResponse ResponseMode = CalcQueryHitType(Fixture);
-	
-	switch (ResponseMode)
-	{
-	case ResponseBlock:
-		if (Fraction < BestResult)
-		{
-			BestFixture = Fixture;
-			BestResult = Fraction;
-			BestNormal = Normal;
-			BestPoint = Point;
-			BestHitType = ResponseMode;
-			bHasBlockingHit = true;
+			// Run the query
+			b2AABB QueryBox;
+			QueryBox.lowerBound = BoxStart;
+			QueryBox.upperBound = BoxStart;
+			BoxWorld->QueryAABB(this, QueryBox);
 		}
 		else
 		{
-			//@TODO: BOX2D: Can this happen?
-			ensure(false);
+			// Cast the ray
+			BoxWorld->RayCast(this, BoxStart, BoxEnd);
 		}
-		return Fraction;
-	case ResponseTouch:
-		return 1.0f;
-	default:
-	case ResponseNone:
-		return -1.0f;
 	}
-}
+};
 
-bool FRaycastHelperBox2D::ReportFixture(b2Fixture* Fixture)
+class FRaycastHelperBox2D_RayCastSingle : public FRaycastHelperBox2D_RayCastBase
 {
-	const EResponse ResponseMode = CalcQueryHitType(Fixture);
+public:
+	FRaycastHelperBox2D_RayCastSingle(ECollisionChannel InTraceChannel, const struct FCollisionQueryParams& InParams, const FCollisionResponseContainer& InCollisionResponseContainer, const FCollisionObjectQueryParams& InObjectParam)
+		: FRaycastHelperBox2D_RayCastBase(InTraceChannel, InParams, InCollisionResponseContainer, InObjectParam, /*bInMultitrace=*/ false)
+		, BestHitTime(MAX_FLT)
+		, BestFixture(nullptr)
+		, bHasBlockingHit(false)
+	{
+	}
 
-	switch (ResponseMode)
+	bool RayCastSingle(const b2World* BoxWorld, const FVector& Start, const FVector& End, bool bAlreadyHadHit, FHitResult& InOutResult)
 	{
-	case ResponseBlock:
+		RayCastInternal(BoxWorld, Start, End);
+
+		// Accept or merge the result
+		if (bHasBlockingHit)
 		{
-			BestFixture = Fixture;
-			BestHitType = ResponseMode;
-			bHasBlockingHit = true;
+			if (!bAlreadyHadHit || (BestHitTime < InOutResult.Time))
+			{
+				ConvertToHitResult_Internal(InOutResult, Start, End, BestHitTime, BestFixture, BestNormal, BestHitType);
+			}
 		}
-		return true;
-	case ResponseTouch:
-		return false;
-	default:
-	case ResponseNone:
-		return false;
-	}
-}
-
-FRaycastHelperBox2D::EResponse FRaycastHelperBox2D::CalcQueryHitType(b2Fixture* Fixture) const
-{
-	const b2Filter& FixtureFilterData = Fixture->GetFilterData();
-	const ECollisionChannel ShapeChannel = static_cast<ECollisionChannel>(FixtureFilterData.ObjectTypeAndFlags >> 24);
-	const uint32 ShapeBit = ECC_TO_BITFIELD(ShapeChannel);
 
-	EResponse Result = ResponseNone;
+		return bAlreadyHadHit || bHasBlockingHit;
+	}
 
-	if (ObjectParam.IsValid())
+protected:
+	// FRaycastHelperBox2D_RayCastBase interface
+	virtual void ProcessHit(b2Fixture* Fixture, const float HitTime, const FVector& Normal, const EResponse ResponseMode) override
 	{
-		// do I belong to one of objects of interest?
-		if ((ShapeBit & ObjectParam.GetQueryBitfield()) != 0)
+		if (ResponseMode == ResponseBlock)
 		{
-			Result = bMultitrace ? ResponseTouch : ResponseBlock;
+			if (HitTime < BestHitTime)
+			{
+				BestFixture = Fixture;
+				BestHitTime = HitTime;
+				BestNormal = Normal;
+				BestHitType = ResponseMode;
+				bHasBlockingHit = true;
+			}
 		}
 	}
-	else
+	// End of FRaycastHelperBox2D_RayCastBase interface
+
+protected:
+	FVector BestNormal;
+	float BestHitTime;
+	b2Fixture* BestFixture;
+	EResponse BestHitType;
+
+	bool bHasBlockingHit;
+};
+
+
+class FRaycastHelperBox2D_RayCastMultiple : public FRaycastHelperBox2D_RayCastBase
+{
+public:
+	FRaycastHelperBox2D_RayCastMultiple(ECollisionChannel InTraceChannel, const struct FCollisionQueryParams& InParams, const FCollisionResponseContainer& InCollisionResponseContainer, const FCollisionObjectQueryParams& InObjectParam, TArray<struct FHitResult>& InHitsArray, const FVector& InStart, const FVector& InEnd)
+		: FRaycastHelperBox2D_RayCastBase(InTraceChannel, InParams, InCollisionResponseContainer, InObjectParam, /*bInMultitrace=*/ true)
+		, HitsArray(InHitsArray)
+		, Start(InStart)
+		, End(InEnd)
 	{
-		// Then see if the channel wants to be blocked
-		const uint32 TraceBit = ECC_TO_BITFIELD(TraceChannel);
+	}
 
-		// check if the fixture has a hit or touch
-		EResponse TraceHitType = ResponseNone;
-		if ((TraceBit & FixtureFilterData.BlockingChannels) != 0)
+	bool RayCastMultiple(const b2World* BoxWorld)
+	{
+		const int32 PreCount = HitsArray.Num();
+
+		RayCastInternal(BoxWorld, Start, End);
+
+		if (PreCount != HitsArray.Num())
 		{
-			TraceHitType = ResponseBlock;
+			// Sort the newly added results from first to last hit
+			HitsArray.Sort(FCompareFHitResultTime());
 		}
-		else if ((TraceBit & FixtureFilterData.TouchingChannels) != 0)
+
+		return HitsArray.Num() > 0;
+	}
+
+protected:
+	TArray<struct FHitResult>& HitsArray;
+	const FVector& Start;
+	const FVector& End;
+
+protected:
+	// FRaycastHelperBox2D_RayCastBase interface
+	virtual void ProcessHit(b2Fixture* Fixture, const float HitTime, const FVector& Normal, const EResponse ResponseMode) override
+	{
+		if (ResponseMode == ResponseBlock)
 		{
-			TraceHitType = ResponseTouch;
+			FHitResult& NewResult = *new (HitsArray) FHitResult();
+			ConvertToHitResult_Internal(NewResult, Start, End, HitTime, Fixture, Normal, ResponseMode);
 		}
+	}
+	// End of FRaycastHelperBox2D_RayCastBase interface
+};
 
-		// Check if the shape has a hit or touch
-		EResponse ShapeHitType = ResponseNone;
-		if ((ShapeBit & QueryBlockingChannels) != 0)
+
+class FRaycastHelperBox2D_RayTest : public FRaycastHelperBox2D_Base
+{
+public:
+	FRaycastHelperBox2D_RayTest(ECollisionChannel InTraceChannel, const struct FCollisionQueryParams& InParams, const FCollisionResponseContainer& InCollisionResponseContainer, const FCollisionObjectQueryParams& InObjectParam)
+		: FRaycastHelperBox2D_Base(InTraceChannel, InParams, InCollisionResponseContainer, InObjectParam, /*bInMultitrace=*/ false)
+		, bHasBlockingHit(false)
+	{
+	}
+
+	bool RayTest(const b2World* BoxWorld, const FVector& Start, const FVector& End)
+	{
+		const b2Vec2 BoxStart(FPhysicsIntegration2D::ConvertUnrealVectorToBox(Start));
+		const b2Vec2 BoxEnd(FPhysicsIntegration2D::ConvertUnrealVectorToBox(End));
+
+		// Now check the length of the 2D query, we may end up having to do an AABB test instead if it's a ray cast along the 'Z' axis of the 2D scene
+		const b2Vec2 BoxDelta = BoxEnd - BoxStart;
+		if (BoxDelta.LengthSquared() < SMALL_NUMBER)
 		{
-			ShapeHitType = ResponseBlock;
+			// Trace into the scene, do a point query instead
+			b2AABB QueryBox;
+			QueryBox.lowerBound = BoxStart;
+			QueryBox.upperBound = BoxStart;
+			BoxWorld->QueryAABB(this, QueryBox);
 		}
-		else if ((ShapeBit & QueryTouchingChannels) != 0)
+		else
 		{
-			ShapeHitType = ResponseTouch;
+			// Cast the ray
+			BoxWorld->RayCast(this, BoxStart, BoxEnd);
 		}
 
-		// return minimum agreed-upon interaction
-		Result = FMath::Min(TraceHitType, ShapeHitType);
+		return bHasBlockingHit;
 	}
 
-	// Single traces should ignore touches
-	if (!bMultitrace && (Result == ResponseTouch))
+	// b2RayCastCallback interface
+	virtual float32 ReportFixture(b2Fixture* Fixture, const b2Vec2& Point, const b2Vec2& Normal, float32 Fraction) override
 	{
-		Result = ResponseNone;
+		const EResponse ResponseMode = CalcQueryHitType(Fixture);
+
+		switch (ResponseMode)
+		{
+		case ResponseBlock:
+			bHasBlockingHit = true;
+			return 0.0f;
+		case ResponseTouch:
+			return 1.0f;
+		default:
+		case ResponseNone:
+			return -1.0f;
+		}
 	}
+	// End of b2RayCastCallback interface
 
-	// Check to see if the actor ultimately owning the fixture should be ignored
-	if (Result != ResponseNone)
+	// b2QueryCallback interface
+	virtual bool ReportFixture(b2Fixture* Fixture) override
 	{
-		FBodyInstance* OwnerBodyInstance = static_cast<FBodyInstance*>(Fixture->GetBody()->GetUserData());
-		UPrimitiveComponent* OwnerComponent = OwnerBodyInstance->OwnerComponent.Get();
-		AActor* Owner = OwnerComponent ? OwnerComponent->GetOwner() : nullptr;
+		const EResponse ResponseMode = CalcQueryHitType(Fixture);
 
-		if ((Owner != nullptr) && IgnoreActors.Contains(Owner->GetUniqueID()))
+		if (ResponseMode == ResponseBlock)
 		{
-			Result = ResponseNone;
+			bHasBlockingHit = true;
+			return true;
+		}
+		else
+		{
+			return false;
 		}
 	}
+	// End of b2QueryCallback interface
 
-	return Result;
-}
+private:
+	bool bHasBlockingHit;
+};
 
 #endif	//WITH_BOX2D
 
@@ -485,9 +568,9 @@ PxSceneQueryHitType::Enum FPxQueryFilterCallback::CalcQueryHitType(const PxFilte
 	return PxSceneQueryHitType::eNONE;
 }
 
-PxSceneQueryHitType::Enum FPxQueryFilterCallback::preFilter(const PxFilterData& filterData, const PxShape* shape,    const PxRigidActor* actor, PxSceneQueryFlags& queryFlags)
+PxSceneQueryHitType::Enum FPxQueryFilterCallback::preFilter(const PxFilterData& filterData, const PxShape* shape, const PxRigidActor* actor, PxSceneQueryFlags& queryFlags)
 {
-	// Check if the shape is the right compexity for the trace 
+	// Check if the shape is the right complexity for the trace 
 	PxFilterData ShapeFilter = shape->getQueryFilterData();
 #define ENABLE_PREFILTER_LOGGING 0
 #if ENABLE_PREFILTER_LOGGING
@@ -497,13 +580,13 @@ PxSceneQueryHitType::Enum FPxQueryFilterCallback::preFilter(const PxFilterData&
 		FBodyInstance* BodyInst = FPhysxUserData::Get<FBodyInstance>(actor->userData);
 		if ( BodyInst && BodyInst->OwnerComponent.IsValid() )
 		{
-			UE_LOG(LogCollision, Log, TEXT("[PREFILTER] against %s[%s] : About to check "), 
+			UE_LOG(LogCollision, Warning, TEXT("[PREFILTER] against %s[%s] : About to check "),
 				(BodyInst->OwnerComponent.Get()->GetOwner())? *BodyInst->OwnerComponent.Get()->GetOwner()->GetName():TEXT("NO OWNER"),
 				*BodyInst->OwnerComponent.Get()->GetName());
 		}
 
-		UE_LOG(LogCollision, Log, TEXT("ShapeFilter : %x %x %x %x"), ShapeFilter.word0, ShapeFilter.word1, ShapeFilter.word2, ShapeFilter.word3);
-		UE_LOG(LogCollision, Log, TEXT("FilterData : %x %x %x %x"), filterData.word0, filterData.word1, filterData.word2, filterData.word3);
+		UE_LOG(LogCollision, Warning, TEXT("ShapeFilter : %x %x %x %x"), ShapeFilter.word0, ShapeFilter.word1, ShapeFilter.word2, ShapeFilter.word3);
+		UE_LOG(LogCollision, Warning, TEXT("FilterData : %x %x %x %x"), filterData.word0, filterData.word1, filterData.word2, filterData.word3);
 	}
 #endif // ENABLE_PREFILTER_LOGGING
 
@@ -731,7 +814,7 @@ bool RaycastTest(const UWorld* World, const FVector Start, const FVector End, EC
 		{
 			if (const b2World* BoxWorld = FPhysicsIntegration2D::FindAssociatedWorld(const_cast<UWorld*>(World)))
 			{
-				FRaycastHelperBox2D BoxQuery(TraceChannel, Params.bTraceComplex, ResponseParams.CollisionResponse, ObjectParams, Params.IgnoreActors, /*bInMultitrace=*/ false);
+				FRaycastHelperBox2D_RayTest BoxQuery(TraceChannel, Params, ResponseParams.CollisionResponse, ObjectParams);
 				bHaveBlockingHit = BoxQuery.RayTest(BoxWorld, Start, End);
 			}
 		}
@@ -820,33 +903,9 @@ bool RaycastSingle(const UWorld* World, struct FHitResult& OutHit, const FVector
 		{
 			if (const b2World* BoxWorld = FPhysicsIntegration2D::FindAssociatedWorld(const_cast<UWorld*>(World)))
 			{
-				FRaycastHelperBox2D BoxQuery(TraceChannel, Params.bTraceComplex, ResponseParams.CollisionResponse, ObjectParams, Params.IgnoreActors, /*bInMultitrace=*/ false);
+				FRaycastHelperBox2D_RayCastSingle BoxQuery(TraceChannel, Params, ResponseParams.CollisionResponse, ObjectParams);
 
-				const b2Vec2 BoxStart(FPhysicsIntegration2D::ConvertUnrealVectorToBox(Start));
-				const b2Vec2 BoxEnd(FPhysicsIntegration2D::ConvertUnrealVectorToBox(End));
-
-				// Now check the length of the 2D query, we may end up having to do an AABB test instead if it's a ray cast along the 'Z' axis of the 2D scene
-				const b2Vec2 BoxDelta = BoxEnd - BoxStart;
-				if (BoxDelta.LengthSquared() < SMALL_NUMBER)
-				{
-					// Trace into the scene, do a point query instead
-					BoxQuery.PointQuery(BoxWorld, BoxStart, Start, End, Params, OutHit);
-				}
-				else
-				{
-					// Cast the ray
-					BoxWorld->RayCast(&BoxQuery, BoxStart, BoxEnd);
-
-					// Accept or merge the results
-					if (BoxQuery.HasBlockingHit())
-					{
-						if (!bHaveBlockingHit || (BoxQuery.GetBestTime() < OutHit.Time))
-						{
-							BoxQuery.ConvertToHitResult(OutHit, Start, End, Params);
-							bHaveBlockingHit = true;
-						}
-					}
-				}
+				bHaveBlockingHit = BoxQuery.RayCastSingle(BoxWorld, Start, End, bHaveBlockingHit, OutHit);
 			}
 		}
 #endif // WITH_BOX2D
@@ -890,7 +949,6 @@ bool RaycastMulti(const UWorld* World, TArray<struct FHitResult>& OutHits, const
 	float DeltaMag = Delta.Size();
 	if (DeltaMag > KINDA_SMALL_NUMBER)
 	{
-		//@TODO: BOX2D: Implement RaycastMulti
 #if WITH_PHYSX
 		// Create filter data used to filter collisions
 		PxFilterData PFilter = CreateQueryFilterData(TraceChannel, Params.bTraceComplex, ResponseParams.CollisionResponse, ObjectParams, true);
@@ -992,13 +1050,28 @@ bool RaycastMulti(const UWorld* World, TArray<struct FHitResult>& OutHits, const
 			}
 		}
 
-		if(NumHits > 0)
+		if (NumHits > 0)
 		{
 			ConvertRaycastResults(NumHits, PHits, DeltaMag, PFilter, OutHits, Start, End, Params.bReturnFaceIndex, Params.bReturnPhysicalMaterial);
 		}
 
 		bHaveBlockingHit = bBlockingHit;
 #endif // WITH_PHYSX
+
+#if WITH_BOX2D && 0
+		if (GetDefault<UPhysicsSettings>()->bEnable2DPhysics)
+		{
+			if (const b2World* BoxWorld = FPhysicsIntegration2D::FindAssociatedWorld(const_cast<UWorld*>(World)))
+			{
+				FRaycastHelperBox2D_RayCastMultiple BoxQuery(TraceChannel, Params, ResponseParams.CollisionResponse, ObjectParams, OutHits, Start, End);
+
+				if (BoxQuery.RayCastMultiple(BoxWorld))
+				{
+					bHaveBlockingHit = true;
+				}
+			}
+		}
+#endif // WITH_BOX2D
 	}
 
 #if !(UE_BUILD_SHIPPING || UE_BUILD_TEST)
@@ -1020,7 +1093,7 @@ bool RaycastMulti(const UWorld* World, TArray<struct FHitResult>& OutHits, const
 
 #if WITH_PHYSX
 
-bool GeomSweepTest(const UWorld * World, const PxGeometry& PGeom, const PxQuat& PGeomRot, FVector Start, FVector End, ECollisionChannel TraceChannel, const struct FCollisionQueryParams& Params, const struct FCollisionResponseParams& ResponseParams, const struct FCollisionObjectQueryParams& ObjectParams)
+bool GeomSweepTest(const UWorld* World, const PxGeometry& PGeom, const PxQuat& PGeomRot, FVector Start, FVector End, ECollisionChannel TraceChannel, const struct FCollisionQueryParams& Params, const struct FCollisionResponseParams& ResponseParams, const struct FCollisionObjectQueryParams& ObjectParams)
 {
 	if(World == NULL || World->GetPhysicsScene() == NULL)
 	{
@@ -1187,7 +1260,7 @@ bool GeomSweepSingle(const UWorld * World, const PxGeometry& PGeom, const PxQuat
 }
 
 
-bool GeomSweepMulti(const UWorld * World, const PxGeometry& PGeom, const PxQuat& PGeomRot, TArray<FHitResult>& OutHits, FVector Start, FVector End, ECollisionChannel TraceChannel, const struct FCollisionQueryParams& Params, const struct FCollisionResponseParams& ResponseParams, const struct FCollisionObjectQueryParams& ObjectParams)
+bool GeomSweepMulti(const UWorld* World, const PxGeometry& PGeom, const PxQuat& PGeomRot, TArray<FHitResult>& OutHits, FVector Start, FVector End, ECollisionChannel TraceChannel, const struct FCollisionQueryParams& Params, const struct FCollisionResponseParams& ResponseParams, const struct FCollisionObjectQueryParams& ObjectParams)
 {
 	if(World == NULL || World->GetPhysicsScene() == NULL)
 	{
@@ -1336,7 +1409,7 @@ bool GeomSweepMulti(const UWorld * World, const PxGeometry& PGeom, const PxQuat&
 //////////////////////////////////////////////////////////////////////////
 // GEOM OVERLAP
 
-bool GeomOverlapTest(const UWorld * World, const PxGeometry& PGeom, const PxTransform& PGeomPose, ECollisionChannel TraceChannel, const struct FCollisionQueryParams& Params, const struct FCollisionResponseParams& ResponseParams, const struct FCollisionObjectQueryParams& ObjectParams)
+bool GeomOverlapTest(const UWorld* World, const PxGeometry& PGeom, const PxTransform& PGeomPose, ECollisionChannel TraceChannel, const struct FCollisionQueryParams& Params, const struct FCollisionResponseParams& ResponseParams, const struct FCollisionObjectQueryParams& ObjectParams)
 {
 	if(World == NULL || World->GetPhysicsScene() == NULL)
 	{
@@ -1402,7 +1475,7 @@ bool GeomOverlapTest(const UWorld * World, const PxGeometry& PGeom, const PxTran
 	return bHaveBlockingHit;
 }
 
-bool GeomOverlapSingle(const UWorld * World, const PxGeometry& PGeom, const PxTransform& PGeomPose, FOverlapResult& OutOverlap, ECollisionChannel TraceChannel, const struct FCollisionQueryParams& Params, const struct FCollisionResponseParams& ResponseParams, const struct FCollisionObjectQueryParams& ObjectParams)
+bool GeomOverlapSingle(const UWorld* World, const PxGeometry& PGeom, const PxTransform& PGeomPose, FOverlapResult& OutOverlap, ECollisionChannel TraceChannel, const struct FCollisionQueryParams& Params, const struct FCollisionResponseParams& ResponseParams, const struct FCollisionObjectQueryParams& ObjectParams)
 {
 	if(World == NULL || World->GetPhysicsScene() == NULL)
 	{
@@ -1467,7 +1540,7 @@ bool GeomOverlapSingle(const UWorld * World, const PxGeometry& PGeom, const PxTr
 	return bHaveBlockingHit;
 }
 
-bool GeomOverlapMulti(const UWorld * World, const PxGeometry& PGeom, const PxTransform& PGeomPose, TArray<FOverlapResult>& OutOverlaps, ECollisionChannel TraceChannel, const struct FCollisionQueryParams& Params, const struct FCollisionResponseParams& ResponseParams, const struct FCollisionObjectQueryParams& ObjectParams)
+bool GeomOverlapMulti(const UWorld* World, const PxGeometry& PGeom, const PxTransform& PGeomPose, TArray<FOverlapResult>& OutOverlaps, ECollisionChannel TraceChannel, const struct FCollisionQueryParams& Params, const struct FCollisionResponseParams& ResponseParams, const struct FCollisionObjectQueryParams& ObjectParams)
 {
 	if(World == NULL || World->GetPhysicsScene() == NULL)
 	{
