Physics: Expose FCompareFHitResultTime for use with Box2D

diff --git a/Engine/Source/Runtime/Engine/Private/Collision/CollisionConversions.cpp b/Engine/Source/Runtime/Engine/Private/Collision/CollisionConversions.cpp
--- a/Engine/Source/Runtime/Engine/Private/Collision/CollisionConversions.cpp
+++ b/Engine/Source/Runtime/Engine/Private/Collision/CollisionConversions.cpp
@@ -7,6 +7,7 @@
 #include "../PhysicsEngine/PhysXSupport.h"
 #include "PhysXCollision.h"
 #include "CollisionDebugDrawing.h"
+#include "CollisionConversions.h"
 
 // Forward declare, I don't want to move the entire function right now or we lose change history.
 static bool ConvertOverlappedShapeToImpactHit(const PxLocationHit& PHit, const FVector& StartLoc, const FVector& EndLoc, FHitResult& OutResult, const PxGeometry& Geom, const PxTransform& QueryTM, const PxFilterData& QueryFilter, bool bReturnPhysMat);
@@ -444,21 +445,6 @@ void ConvertQueryImpactHit(const PxLocationHit& PHit, FHitResult& OutResult, flo
 	}
 }
 
-struct FCompareFHitResultTime
-{
-	FORCEINLINE bool operator()( const FHitResult& A, const FHitResult& B ) const
-	{
-		if (A.Time == B.Time)
-		{
-			// Sort blocking hits after non-blocking hits, if they are at the same time. Also avoid swaps if they are the same.
-			// This is important so initial touches are reported before processing stops on the first blocking hit.
-			return (A.bBlockingHit == B.bBlockingHit) ? true : B.bBlockingHit;
-		}
-
-		return A.Time < B.Time;
-	}
-};
-
 void ConvertRaycastResults(int32 NumHits, PxRaycastHit* Hits, float CheckLength, const PxFilterData& QueryFilter, TArray<FHitResult>& OutHits, const FVector& StartLoc, const FVector& EndLoc, bool bReturnFaceIndex, bool bReturnPhysMat)
 {
 	PxTransform PStartTM(U2PVector(StartLoc));
diff --git a/Engine/Source/Runtime/Engine/Private/Collision/CollisionConversions.h b/Engine/Source/Runtime/Engine/Private/Collision/CollisionConversions.h
--- a/Engine/Source/Runtime/Engine/Private/Collision/CollisionConversions.h
+++ b/Engine/Source/Runtime/Engine/Private/Collision/CollisionConversions.h
@@ -71,3 +71,23 @@ bool ConvertOverlapResults(int32 NumOverlaps,  PxOverlapHit* POverlapResults, co
 
 
 #endif // WITH_PHYX
+
+
+#if UE_WITH_PHYSICS
+
+struct FCompareFHitResultTime
+{
+	FORCEINLINE bool operator()(const FHitResult& A, const FHitResult& B) const
+	{
+		if (A.Time == B.Time)
+		{
+			// Sort blocking hits after non-blocking hits, if they are at the same time. Also avoid swaps if they are the same.
+			// This is important so initial touches are reported before processing stops on the first blocking hit.
+			return (A.bBlockingHit == B.bBlockingHit) ? true : B.bBlockingHit;
+		}
+
+		return A.Time < B.Time;
+	}
+};
+
+#endif
\ No newline at end of file
