moved navigation system lock to more editor specific place
- unfortunately AActor::ReregisterAllComponents is called during initial map load just for registering and can't lock anything related to actual re-registering :/

diff --git a/Engine/Source/Editor/UnrealEd/Private/Editor.cpp b/Engine/Source/Editor/UnrealEd/Private/Editor.cpp
--- a/Engine/Source/Editor/UnrealEd/Private/Editor.cpp
+++ b/Engine/Source/Editor/UnrealEd/Private/Editor.cpp
@@ -2065,6 +2065,8 @@ void UEditorEngine::ApplyDeltaToActor(AActor* InActor,
 		}
 	}
 
+	FNavigationLockContext LockNavigationUpdates;
+
 	bool bTranslationOnly = true;
 
 	///////////////////
diff --git a/Engine/Source/Runtime/Engine/Private/Actor.cpp b/Engine/Source/Runtime/Engine/Private/Actor.cpp
--- a/Engine/Source/Runtime/Engine/Private/Actor.cpp
+++ b/Engine/Source/Runtime/Engine/Private/Actor.cpp
@@ -3059,8 +3059,6 @@ void AActor::MarkComponentsAsPendingKill()
 
 void AActor::ReregisterAllComponents()
 {
-	FNavigationLockContext NavUpdateLock(GetWorld());
-
 	UnregisterAllComponents();
 	RegisterAllComponents();
 }
