Fixed static lighting on Blueprint mesh components

Avoid resizing the LOD data array when we dont have vertex colours (this was clobbering the static lighting).

TTP# 337767 - REGRESSION: UE4: RENDERING: Static lighting on Blueprint mesh components is not working

diff --git a/Engine/Source/Runtime/Engine/Private/StaticMeshComponent.cpp b/Engine/Source/Runtime/Engine/Private/StaticMeshComponent.cpp
--- a/Engine/Source/Runtime/Engine/Private/StaticMeshComponent.cpp
+++ b/Engine/Source/Runtime/Engine/Private/StaticMeshComponent.cpp
@@ -46,7 +46,7 @@ public:
 
 		if (StaticMeshComponent != NULL)
 		{
-			StaticMeshComponent->SetLODDataCount(VertexColorLODs.Num(), VertexColorLODs.Num());
+			StaticMeshComponent->SetLODDataCount(VertexColorLODs.Num(), StaticMeshComponent->LODData.Num());
 
 			for (int32 LODDataIndex = 0; LODDataIndex < VertexColorLODs.Num(); ++LODDataIndex)
 			{
