Improving Opus format so that it has unique identifier

Added a text string at the start of the Opus File format to identify that we're dealing with the right kind of data (discovered this could be a problem when testing the Opus format and was passed OGG data in error).
Added error checking for OGG decompression (Read function could get stuck in an infinite loop if the seek function failed - again only discovered while testing Opus but makes it future proof).

diff --git a/Engine/Source/Developer/AudioFormatOpus/Private/AudioFormatOpus.cpp b/Engine/Source/Developer/AudioFormatOpus/Private/AudioFormatOpus.cpp
--- a/Engine/Source/Developer/AudioFormatOpus/Private/AudioFormatOpus.cpp
+++ b/Engine/Source/Developer/AudioFormatOpus/Private/AudioFormatOpus.cpp
@@ -27,7 +27,7 @@ class FAudioFormatOpus : public IAudioFormat
 	enum
 	{
 		/** Version for OPUS format, this becomes part of the DDC key. */
-		UE_AUDIO_OPUS_VER = 1,
+		UE_AUDIO_OPUS_VER = 2,
 	};
 
 public:
@@ -145,7 +145,9 @@ public:
 			FramesToEncode++;
 		}
 
-		// Store True Sample Count, Number of channels and Frames to Encode first
+		// Store Identifier, True Sample Count, Number of channels and Frames to Encode first
+		const char* OpusIdentifier = OPUS_ID_STRING;
+		CompressedData.Serialize((void*)OpusIdentifier, FCStringAnsi::Strlen( OpusIdentifier )+1);
 		CompressedData.Serialize(&TrueSampleCount, sizeof(uint32));
 		check(QualityInfo.NumChannels < MAX_uint8)
 		uint8 SerializedChannels = QualityInfo.NumChannels;
@@ -246,7 +248,12 @@ public:
 		uint16 ProcessedFrames = 0;
 		const uint8* LockedSrc = SrcBuffer.GetTypedData();
 
-		// Read True Sample Count, Number of channels and Frames to Encode first
+		// Read Identifier, True Sample Count, Number of channels and Frames to Encode first
+		if (FCStringAnsi::Strcmp((char*)LockedSrc, OPUS_ID_STRING) != 0)
+		{
+			return false;
+		}
+		ReadOffset += FCStringAnsi::Strlen(OPUS_ID_STRING) + 1;
 		uint32 TrueSampleCount = *((uint32*)(LockedSrc + ReadOffset));
 		ReadOffset += sizeof(uint32);
 		uint8 NumChannels = *(LockedSrc + ReadOffset);
diff --git a/Engine/Source/Runtime/Engine/Private/AudioDecompress.cpp b/Engine/Source/Runtime/Engine/Private/AudioDecompress.cpp
--- a/Engine/Source/Runtime/Engine/Private/AudioDecompress.cpp
+++ b/Engine/Source/Runtime/Engine/Private/AudioDecompress.cpp
@@ -178,8 +178,18 @@ bool FVorbisAudioInfo::ReadCompressedInfo( const uint8* InSrcBufferData, uint32
 		vorbis_info* vi = ov_info( &VFWrapper->vf, -1 );
 		QualityInfo->SampleRate = vi->rate;
 		QualityInfo->NumChannels = vi->channels;
-		QualityInfo->SampleDataSize = ov_pcm_total( &VFWrapper->vf, -1 ) * QualityInfo->NumChannels * sizeof( int16 );
-		QualityInfo->Duration = ( float )ov_time_total( &VFWrapper->vf, -1 );
+		ogg_int64_t PCMTotal = ov_pcm_total( &VFWrapper->vf, -1 );
+		if (PCMTotal >= 0)
+		{
+			QualityInfo->SampleDataSize = PCMTotal * QualityInfo->NumChannels * sizeof( int16 );
+			QualityInfo->Duration = ( float )ov_time_total( &VFWrapper->vf, -1 );
+		}
+		else if (PCMTotal == OV_EINVAL)
+		{
+			// indicates an error or that the bitstream is non-seekable
+			QualityInfo->SampleDataSize = 0;
+			QualityInfo->Duration = 0.0f;
+		}
 	}
 
 	return( true );
@@ -191,7 +201,7 @@ bool FVorbisAudioInfo::ReadCompressedInfo( const uint8* InSrcBufferData, uint32
  */
 void FVorbisAudioInfo::ExpandFile( uint8* DstBuffer, FSoundQualityInfo* QualityInfo )
 {
-	uint32		BytesRead, TotalBytesRead, BytesToRead;
+	uint32		TotalBytesRead, BytesToRead;
 
 	check( VFWrapper != NULL );
 	check( DstBuffer );
@@ -204,7 +214,14 @@ void FVorbisAudioInfo::ExpandFile( uint8* DstBuffer, FSoundQualityInfo* QualityI
 	char* Destination = ( char* )DstBuffer;
 	while( TotalBytesRead < BytesToRead )
 	{
-		BytesRead = ov_read( &VFWrapper->vf, Destination, BytesToRead - TotalBytesRead, 0, 2, 1, NULL );
+		long BytesRead = ov_read( &VFWrapper->vf, Destination, BytesToRead - TotalBytesRead, 0, 2, 1, NULL );
+
+		if (BytesRead < 0)
+		{
+			// indicates an error - fill remainder of buffer with zero
+			FMemory::Memzero(Destination, BytesToRead - TotalBytesRead);
+			return;
+		}
 
 		TotalBytesRead += BytesRead;
 		Destination += BytesRead;
@@ -225,7 +242,7 @@ void FVorbisAudioInfo::ExpandFile( uint8* DstBuffer, FSoundQualityInfo* QualityI
 bool FVorbisAudioInfo::ReadCompressedData( uint8* InDestination, bool bLooping, uint32 BufferSize )
 {
 	bool		bLooped;
-	uint32		BytesRead, TotalBytesRead;
+	uint32		TotalBytesRead;
 
 	SCOPE_CYCLE_COUNTER( STAT_VorbisDecompressTime );
 
@@ -239,14 +256,20 @@ bool FVorbisAudioInfo::ReadCompressedData( uint8* InDestination, bool bLooping,
 
 	while( TotalBytesRead < BufferSize )
 	{
-		BytesRead = ov_read( &VFWrapper->vf, Destination, BufferSize - TotalBytesRead, 0, 2, 1, NULL );
+		long BytesRead = ov_read( &VFWrapper->vf, Destination, BufferSize - TotalBytesRead, 0, 2, 1, NULL );
 		if( !BytesRead )
 		{
 			// We've reached the end
 			bLooped = true;
 			if( bLooping )
 			{
-				ov_pcm_seek_page( &VFWrapper->vf, 0 );
+				int Result = ov_pcm_seek_page( &VFWrapper->vf, 0 );
+				if (Result < 0)
+				{
+					// indicates an error - fill remainder of buffer with zero
+					FMemory::Memzero(Destination, BufferSize - TotalBytesRead);
+					return true;
+				}
 			}
 			else
 			{
@@ -256,6 +279,12 @@ bool FVorbisAudioInfo::ReadCompressedData( uint8* InDestination, bool bLooping,
 				BytesRead += BufferSize - TotalBytesRead;
 			}
 		}
+		else if (BytesRead < 0)
+		{
+			// indicates an error - fill remainder of buffer with zero
+			FMemory::Memzero(Destination, BufferSize - TotalBytesRead);
+			return false;
+		}
 
 		TotalBytesRead += BytesRead;
 		Destination += BytesRead;
@@ -393,7 +422,15 @@ bool FOpusAudioInfo::ReadCompressedInfo(const uint8* InSrcBufferData, uint32 InS
 	SrcBufferOffset = 0;
 	CurrentSampleCount = 0;
 
-	// Read True Sample Count, Number of channels and Frames to Encode first
+	// Read Identifier, True Sample Count, Number of channels and Frames to Encode first
+	if (FCStringAnsi::Strcmp((char*)InSrcBufferData, OPUS_ID_STRING) != 0)
+	{
+		return false;
+	}
+	else
+	{
+		SrcBufferOffset += FCStringAnsi::Strlen(OPUS_ID_STRING) + 1;
+	}
 	Read(&TrueSampleCount, sizeof(uint32));
 	Read(&NumChannels, sizeof(uint8));
 	uint16 SerializedFrames = 0;
diff --git a/Engine/Source/Runtime/Engine/Public/AudioDecompress.h b/Engine/Source/Runtime/Engine/Public/AudioDecompress.h
--- a/Engine/Source/Runtime/Engine/Public/AudioDecompress.h
+++ b/Engine/Source/Runtime/Engine/Public/AudioDecompress.h
@@ -138,6 +138,8 @@ public:
 };
 #endif
 
+#define OPUS_ID_STRING "UE4OPUS"
+
 /**
 * Helper class to parse opus data
 */
