UE4: Rearrange RHICmdList toward non-null cmd lists and the final "look"

diff --git a/Engine/Source/Editor/MeshPaint/Private/MeshPaintEdMode.cpp b/Engine/Source/Editor/MeshPaint/Private/MeshPaintEdMode.cpp
--- a/Engine/Source/Editor/MeshPaint/Private/MeshPaintEdMode.cpp
+++ b/Engine/Source/Editor/MeshPaint/Private/MeshPaintEdMode.cpp
@@ -47,7 +47,7 @@ class FMeshPaintBatchedElementParameters : public FBatchedElementParameters
 public:
 
 	/** Binds vertex and pixel shaders for this element */
-	virtual void BindShaders_RenderThread(FRHICommandList* RHICmdList, const FMatrix& InTransform, const float InGamma, const FMatrix& ColorWeights, const FTexture* Texture ) OVERRIDE
+	virtual void BindShaders_RenderThread(FRHICommandList& RHICmdList, const FMatrix& InTransform, const float InGamma, const FMatrix& ColorWeights, const FTexture* Texture ) OVERRIDE
 	{
 		MeshPaintRendering::SetMeshPaintShaders_RenderThread(RHICmdList, InTransform, InGamma, ShaderParams );
 	}
@@ -66,7 +66,7 @@ class FMeshPaintDilateBatchedElementParameters : public FBatchedElementParameter
 public:
 
 	/** Binds vertex and pixel shaders for this element */
-	virtual void BindShaders_RenderThread(FRHICommandList* RHICmdList, const FMatrix& InTransform, const float InGamma, const FMatrix& ColorWeights, const FTexture* Texture ) OVERRIDE
+	virtual void BindShaders_RenderThread(FRHICommandList& RHICmdList, const FMatrix& InTransform, const float InGamma, const FMatrix& ColorWeights, const FTexture* Texture ) OVERRIDE
 	{
 		MeshPaintRendering::SetMeshPaintDilateShaders_RenderThread(RHICmdList, InTransform, InGamma, ShaderParams );
 	}
diff --git a/Engine/Source/Editor/UnrealEd/Private/CubemapUnwrapUtils.cpp b/Engine/Source/Editor/UnrealEd/Private/CubemapUnwrapUtils.cpp
--- a/Engine/Source/Editor/UnrealEd/Private/CubemapUnwrapUtils.cpp
+++ b/Engine/Source/Editor/UnrealEd/Private/CubemapUnwrapUtils.cpp
@@ -108,13 +108,13 @@ namespace CubemapHelpers
 	}
 }
 
-void FCubemapTexturePropertiesVS::SetParameters( FRHICommandList* RHICmdList, const FMatrix& TransformValue )
+void FCubemapTexturePropertiesVS::SetParameters( FRHICommandList& RHICmdList, const FMatrix& TransformValue )
 {
 	SetShaderValue(RHICmdList, GetVertexShader(), Transform, TransformValue);
 }
 
 template<bool bHDROutput>
-void FCubemapTexturePropertiesPS<bHDROutput>::SetParameters( FRHICommandList* RHICmdList, const FTexture* Texture, const FMatrix& ColorWeightsValue, float MipLevel, float GammaValue )
+void FCubemapTexturePropertiesPS<bHDROutput>::SetParameters( FRHICommandList& RHICmdList, const FTexture* Texture, const FMatrix& ColorWeightsValue, float MipLevel, float GammaValue )
 {
 	SetTextureParameter(RHICmdList, GetPixelShader(),CubeTexture,CubeTextureSampler,Texture);
 
@@ -124,7 +124,7 @@ void FCubemapTexturePropertiesPS<bHDROutput>::SetParameters( FRHICommandList* RH
 	SetShaderValue(RHICmdList, GetPixelShader(), Gamma, GammaValue);
 }
 
-void FMipLevelBatchedElementParameters::BindShaders_RenderThread(FRHICommandList* RHICmdList, const FMatrix& InTransform, const float InGamma, const FMatrix& ColorWeights, const FTexture* Texture)
+void FMipLevelBatchedElementParameters::BindShaders_RenderThread(FRHICommandList& RHICmdList, const FMatrix& InTransform, const float InGamma, const FMatrix& ColorWeights, const FTexture* Texture)
 {
 	if(bHDROutput)
 	{
@@ -137,11 +137,11 @@ void FMipLevelBatchedElementParameters::BindShaders_RenderThread(FRHICommandList
 }
 
 template<typename TPixelShader>
-void FMipLevelBatchedElementParameters::BindShaders_RenderThread( FRHICommandList* RHICmdList, const FMatrix& InTransform, const float InGamma, const FMatrix& ColorWeights, const FTexture* Texture )
+void FMipLevelBatchedElementParameters::BindShaders_RenderThread( FRHICommandList& RHICmdList, const FMatrix& InTransform, const float InGamma, const FMatrix& ColorWeights, const FTexture* Texture )
 {
 	TShaderMapRef<FCubemapTexturePropertiesVS> VertexShader(GetGlobalShaderMap());
 	TShaderMapRef<TPixelShader> PixelShader(GetGlobalShaderMap());
-	check(!RHICmdList);
+	RHICmdList.CheckIsNull();
 
 	static FGlobalBoundShaderState BoundShaderState;
 	SetGlobalBoundShaderState(BoundShaderState, GSimpleElementVertexDeclaration.VertexDeclarationRHI, *VertexShader, *PixelShader);
@@ -153,7 +153,7 @@ void FMipLevelBatchedElementParameters::BindShaders_RenderThread( FRHICommandLis
 	PixelShader->SetParameters(RHICmdList, Texture, ColorWeights, MipLevel, InGamma);
 }
 
-void FIESLightProfilePS::SetParameters( FRHICommandList* RHICmdList, const FTexture* Texture, float InBrightnessInLumens )
+void FIESLightProfilePS::SetParameters( FRHICommandList& RHICmdList, const FTexture* Texture, float InBrightnessInLumens )
 {
 	FPixelShaderRHIParamRef ShaderRHI = GetPixelShader();
 	SetTextureParameter(RHICmdList, ShaderRHI, IESTexture, IESTextureSampler, Texture);
@@ -161,12 +161,12 @@ void FIESLightProfilePS::SetParameters( FRHICommandList* RHICmdList, const FText
 	SetShaderValue(RHICmdList, ShaderRHI, BrightnessInLumens, InBrightnessInLumens);
 }
 
-void FIESLightProfileBatchedElementParameters::BindShaders_RenderThread( FRHICommandList* RHICmdList, const FMatrix& InTransform, const float InGamma, const FMatrix& ColorWeights, const FTexture* Texture )
+void FIESLightProfileBatchedElementParameters::BindShaders_RenderThread( FRHICommandList& RHICmdList, const FMatrix& InTransform, const float InGamma, const FMatrix& ColorWeights, const FTexture* Texture )
 {
 	TShaderMapRef<FSimpleElementVS> VertexShader(GetGlobalShaderMap());
 	TShaderMapRef<FIESLightProfilePS> PixelShader(GetGlobalShaderMap());
 
-	check(!RHICmdList);
+	RHICmdList.CheckIsNull();
 	static FGlobalBoundShaderState BoundShaderState;
 	SetGlobalBoundShaderState(BoundShaderState, GSimpleElementVertexDeclaration.VertexDeclarationRHI, *VertexShader, *PixelShader);
 	RHISetBlendState(TStaticBlendState<>::GetRHI());
diff --git a/Engine/Source/Editor/UnrealEd/Private/MeshPaintRendering.cpp b/Engine/Source/Editor/UnrealEd/Private/MeshPaintRendering.cpp
--- a/Engine/Source/Editor/UnrealEd/Private/MeshPaintRendering.cpp
+++ b/Engine/Source/Editor/UnrealEd/Private/MeshPaintRendering.cpp
@@ -43,7 +43,7 @@ namespace MeshPaintRendering
 			return bShaderHasOutdatedParameters;
 		}
 
-		void SetParameters(FRHICommandList* RHICmdList, const FMatrix& InTransform )
+		void SetParameters(FRHICommandList& RHICmdList, const FMatrix& InTransform )
 		{
 			SetShaderValue(RHICmdList, GetVertexShader(), TransformParameter, InTransform );
 		}
@@ -103,7 +103,7 @@ namespace MeshPaintRendering
 			return bShaderHasOutdatedParameters;
 		}
 
-		void SetParameters(FRHICommandList* RHICmdList, const float InGamma, const FMeshPaintShaderParameters& InShaderParams )
+		void SetParameters(FRHICommandList& RHICmdList, const float InGamma, const FMeshPaintShaderParameters& InShaderParams )
 		{
 			const FPixelShaderRHIParamRef ShaderRHI = GetPixelShader();
 
@@ -206,7 +206,7 @@ namespace MeshPaintRendering
 			return bShaderHasOutdatedParameters;
 		}
 
-		void SetParameters(FRHICommandList* RHICmdList, const FMatrix& InTransform )
+		void SetParameters(FRHICommandList& RHICmdList, const FMatrix& InTransform )
 		{
 			SetShaderValue(RHICmdList, GetVertexShader(), TransformParameter, InTransform );
 		}
@@ -266,7 +266,7 @@ namespace MeshPaintRendering
 			return bShaderHasOutdatedParameters;
 		}
 
-		void SetParameters(FRHICommandList* RHICmdList, const float InGamma, const FMeshPaintDilateShaderParameters& InShaderParams )
+		void SetParameters(FRHICommandList& RHICmdList, const float InGamma, const FMeshPaintDilateShaderParameters& InShaderParams )
 		{
 			const FPixelShaderRHIParamRef ShaderRHI = GetPixelShader();
 
@@ -348,7 +348,7 @@ namespace MeshPaintRendering
 
 
 	/** Binds the mesh paint vertex and pixel shaders to the graphics device */
-	void SetMeshPaintShaders_RenderThread(FRHICommandList* RHICmdList, const FMatrix& InTransform,
+	void SetMeshPaintShaders_RenderThread(FRHICommandList& RHICmdList, const FMatrix& InTransform,
 										   const float InGamma,
 										   const FMeshPaintShaderParameters& InShaderParams )
 	{
@@ -368,7 +368,7 @@ namespace MeshPaintRendering
 	}
 
 	/** Binds the mesh paint vertex and pixel shaders to the graphics device */
-	void SetMeshPaintDilateShaders_RenderThread(FRHICommandList* RHICmdList, const FMatrix& InTransform,
+	void SetMeshPaintDilateShaders_RenderThread(FRHICommandList& RHICmdList, const FMatrix& InTransform,
 												 const float InGamma,
 												 const FMeshPaintDilateShaderParameters& InShaderParams )
 	{
diff --git a/Engine/Source/Editor/UnrealEd/Private/NormalMapPreview.cpp b/Engine/Source/Editor/UnrealEd/Private/NormalMapPreview.cpp
--- a/Engine/Source/Editor/UnrealEd/Private/NormalMapPreview.cpp
+++ b/Engine/Source/Editor/UnrealEd/Private/NormalMapPreview.cpp
@@ -48,7 +48,7 @@ public:
 	 * Set shader parameters.
 	 * @param NormalMapTexture - The normal map texture to sample.
 	 */
-	void SetParameters(FRHICommandList* RHICmdList, const FTexture* NormalMapTexture)
+	void SetParameters(FRHICommandList& RHICmdList, const FTexture* NormalMapTexture)
 	{
 		FPixelShaderRHIParamRef PixelShaderRHI = GetPixelShader();
 		SetTextureParameter(RHICmdList, PixelShaderRHI,Texture,TextureSampler,NormalMapTexture);
@@ -75,7 +75,7 @@ IMPLEMENT_SHADER_TYPE(,FSimpleElementNormalMapPS,TEXT("SimpleElementNormalMapPix
 
 /** Binds vertex and pixel shaders for this element */
 void FNormalMapBatchedElementParameters::BindShaders_RenderThread(
-	FRHICommandList* RHICmdList, 
+	FRHICommandList& RHICmdList, 
 	const FMatrix& InTransform,
 	const float InGamma,
 	const FMatrix& ColorWeights,
@@ -84,7 +84,7 @@ void FNormalMapBatchedElementParameters::BindShaders_RenderThread(
 	TShaderMapRef<FSimpleElementVS> VertexShader(GetGlobalShaderMap());
 	TShaderMapRef<FSimpleElementNormalMapPS> PixelShader(GetGlobalShaderMap());
 
-	check(!RHICmdList);
+	RHICmdList.CheckIsNull();
 
 	static FGlobalBoundShaderState BoundShaderState;
 	SetGlobalBoundShaderState(BoundShaderState, GSimpleElementVertexDeclaration.VertexDeclarationRHI, *VertexShader, *PixelShader);
diff --git a/Engine/Source/Editor/UnrealEd/Private/Texture2DPreview.cpp b/Engine/Source/Editor/UnrealEd/Private/Texture2DPreview.cpp
--- a/Engine/Source/Editor/UnrealEd/Private/Texture2DPreview.cpp
+++ b/Engine/Source/Editor/UnrealEd/Private/Texture2DPreview.cpp
@@ -41,7 +41,7 @@ public:
 		return IsFeatureLevelSupported(Platform, ERHIFeatureLevel::SM3);
 	}
 
-	void SetParameters(FRHICommandList* RHICmdList, const FTexture* TextureValue, const FMatrix& ColorWeightsValue, float GammaValue, float MipLevel, bool bIsNormalMap)
+	void SetParameters(FRHICommandList& RHICmdList, const FTexture* TextureValue, const FMatrix& ColorWeightsValue, float GammaValue, float MipLevel, bool bIsNormalMap)
 	{
 		SetTextureParameter(RHICmdList, GetPixelShader(),InTexture,InTextureSampler,TextureValue);
 		SetShaderValue(RHICmdList, GetPixelShader(),ColorWeights,ColorWeightsValue);
@@ -77,7 +77,7 @@ IMPLEMENT_SHADER_TYPE(,FSimpleElementTexture2DPreviewPS,TEXT("SimpleElementTextu
 
 /** Binds vertex and pixel shaders for this element */
 void FBatchedElementTexture2DPreviewParameters::BindShaders_RenderThread(
-	FRHICommandList* RHICmdList, 
+	FRHICommandList& RHICmdList, 
 	const FMatrix& InTransform,
 	const float InGamma,
 	const FMatrix& ColorWeights,
@@ -85,7 +85,7 @@ void FBatchedElementTexture2DPreviewParameters::BindShaders_RenderThread(
 {
 	TShaderMapRef<FSimpleElementVS> VertexShader(GetGlobalShaderMap());
 	TShaderMapRef<FSimpleElementTexture2DPreviewPS> PixelShader(GetGlobalShaderMap());
-	check(!RHICmdList);
+	RHICmdList.CheckIsNull();
 
 	static FGlobalBoundShaderState BoundShaderState;
 	SetGlobalBoundShaderState(BoundShaderState, GSimpleElementVertexDeclaration.VertexDeclarationRHI, *VertexShader, *PixelShader);
diff --git a/Engine/Source/Editor/UnrealEd/Public/CubemapUnwrapUtils.h b/Engine/Source/Editor/UnrealEd/Public/CubemapUnwrapUtils.h
--- a/Engine/Source/Editor/UnrealEd/Public/CubemapUnwrapUtils.h
+++ b/Engine/Source/Editor/UnrealEd/Public/CubemapUnwrapUtils.h
@@ -51,7 +51,7 @@ public:
 	}
 	FCubemapTexturePropertiesVS() {}
 
-	void SetParameters(FRHICommandList* RHICmdList, const FMatrix& TransformValue);
+	void SetParameters(FRHICommandList& RHICmdList, const FMatrix& TransformValue);
 
 	virtual bool Serialize(FArchive& Ar)
 	{
@@ -87,7 +87,7 @@ public:
 	}
 	FCubemapTexturePropertiesPS() {}
 
-	void SetParameters(FRHICommandList* RHICmdList, const FTexture* Texture, const FMatrix& ColorWeightsValue, float MipLevel, float GammaValue);
+	void SetParameters(FRHICommandList& RHICmdList, const FTexture* Texture, const FMatrix& ColorWeightsValue, float MipLevel, float GammaValue);
 
 	virtual bool Serialize(FArchive& Ar)
 	{
@@ -119,10 +119,10 @@ public:
 	}
 
 	/** Binds vertex and pixel shaders for this element */
-	virtual void BindShaders_RenderThread(FRHICommandList* RHICmdList, const FMatrix& InTransform, const float InGamma, const FMatrix& ColorWeights, const FTexture* Texture) OVERRIDE;
+	virtual void BindShaders_RenderThread(FRHICommandList& RHICmdList, const FMatrix& InTransform, const float InGamma, const FMatrix& ColorWeights, const FTexture* Texture) OVERRIDE;
 
 private:
-	template<typename TPixelShader> void BindShaders_RenderThread(FRHICommandList* RHICmdList, const FMatrix& InTransform, const float InGamma, const FMatrix& ColorWeights, const FTexture* Texture);
+	template<typename TPixelShader> void BindShaders_RenderThread(FRHICommandList& RHICmdList, const FMatrix& InTransform, const float InGamma, const FMatrix& ColorWeights, const FTexture* Texture);
 
 	bool bHDROutput;
 	/** Parameters that need to be passed to the shader */
@@ -153,7 +153,7 @@ public:
 		BrightnessInLumens.Bind(Initializer.ParameterMap,TEXT("BrightnessInLumens"));
 	}
 
-	void SetParameters(FRHICommandList* RHICmdList, const FTexture* Texture, float InBrightnessInLumens);
+	void SetParameters(FRHICommandList& RHICmdList, const FTexture* Texture, float InBrightnessInLumens);
 
 	virtual bool Serialize(FArchive& Ar) OVERRIDE
 	{
@@ -179,7 +179,7 @@ public:
 	}
 
 	/** Binds vertex and pixel shaders for this element */
-	virtual void BindShaders_RenderThread(FRHICommandList* RHICmdList, const FMatrix& InTransform, const float InGamma, const FMatrix& ColorWeights, const FTexture* Texture ) OVERRIDE;
+	virtual void BindShaders_RenderThread(FRHICommandList& RHICmdList, const FMatrix& InTransform, const float InGamma, const FMatrix& ColorWeights, const FTexture* Texture ) OVERRIDE;
 
 private:
 	float BrightnessInLumens;
diff --git a/Engine/Source/Editor/UnrealEd/Public/MeshPaintRendering.h b/Engine/Source/Editor/UnrealEd/Public/MeshPaintRendering.h
--- a/Engine/Source/Editor/UnrealEd/Public/MeshPaintRendering.h
+++ b/Engine/Source/Editor/UnrealEd/Public/MeshPaintRendering.h
@@ -57,12 +57,12 @@ namespace MeshPaintRendering
 
 
 	/** Binds the mesh paint vertex and pixel shaders to the graphics device */
-	void UNREALED_API SetMeshPaintShaders_RenderThread(FRHICommandList* RHICmdList, const FMatrix& InTransform,
+	void UNREALED_API SetMeshPaintShaders_RenderThread(FRHICommandList& RHICmdList, const FMatrix& InTransform,
 										   const float InGamma,
 										   const FMeshPaintShaderParameters& InShaderParams );
 
 	/** Binds the mesh paint dilation vertex and pixel shaders to the graphics device */
-	void UNREALED_API SetMeshPaintDilateShaders_RenderThread(FRHICommandList* RHICmdList, const FMatrix& InTransform,
+	void UNREALED_API SetMeshPaintDilateShaders_RenderThread(FRHICommandList& RHICmdList, const FMatrix& InTransform,
 												const float InGamma,
 												const FMeshPaintDilateShaderParameters& InShaderParams );
 
diff --git a/Engine/Source/Editor/UnrealEd/Public/NormalMapPreview.h b/Engine/Source/Editor/UnrealEd/Public/NormalMapPreview.h
--- a/Engine/Source/Editor/UnrealEd/Public/NormalMapPreview.h
+++ b/Engine/Source/Editor/UnrealEd/Public/NormalMapPreview.h
@@ -15,7 +15,7 @@ class UNREALED_API FNormalMapBatchedElementParameters : public FBatchedElementPa
 {
 	/** Binds vertex and pixel shaders for this element */
 	virtual void BindShaders_RenderThread(
-		FRHICommandList* RHICmdList, 
+		FRHICommandList& RHICmdList, 
 		const FMatrix& InTransform,
 		const float InGamma,
 		const FMatrix& ColorWeights,
diff --git a/Engine/Source/Editor/UnrealEd/Public/Texture2DPreview.h b/Engine/Source/Editor/UnrealEd/Public/Texture2DPreview.h
--- a/Engine/Source/Editor/UnrealEd/Public/Texture2DPreview.h
+++ b/Engine/Source/Editor/UnrealEd/Public/Texture2DPreview.h
@@ -22,7 +22,7 @@ public:
 	}
 
 	/** Binds vertex and pixel shaders for this element */
-	virtual void BindShaders_RenderThread(FRHICommandList* RHICmdList, const FMatrix& InTransform, const float InGamma, const FMatrix& ColorWeights, const FTexture* Texture ) OVERRIDE;
+	virtual void BindShaders_RenderThread(FRHICommandList& RHICmdList, const FMatrix& InTransform, const float InGamma, const FMatrix& ColorWeights, const FTexture* Texture ) OVERRIDE;
 
 private:
 
diff --git a/Engine/Source/Runtime/Engine/Private/BatchedElements.cpp b/Engine/Source/Runtime/Engine/Private/BatchedElements.cpp
--- a/Engine/Source/Runtime/Engine/Private/BatchedElements.cpp
+++ b/Engine/Source/Runtime/Engine/Private/BatchedElements.cpp
@@ -7,6 +7,8 @@
 #include "ShaderParameters.h"
 #include "RHIStaticStates.h"
 #include "SimpleElementShaders.h"
+#include "RHICommandList.h"
+
 
 DEFINE_LOG_CATEGORY_STATIC(LogBatchedElements, Log, All);
 
@@ -357,7 +359,7 @@ void FBatchedElements::PrepareShaders(
 	) const
 {
 	//@todo-rco: RHIPacketList
-	FRHICommandList* RHICmdList = nullptr;
+	FRHICommandList& RHICmdList = FRHICommandList::GetNullRef();
 
 	// used to mask individual channels and desaturate
 	FMatrix ColorWeights( FPlane(1, 0, 0, 0), FPlane(0, 1, 0, 0), FPlane(0, 0, 1, 0), FPlane(0, 0, 0, 0) );
diff --git a/Engine/Source/Runtime/Engine/Private/ClearReplacementShaders.cpp b/Engine/Source/Runtime/Engine/Private/ClearReplacementShaders.cpp
--- a/Engine/Source/Runtime/Engine/Private/ClearReplacementShaders.cpp
+++ b/Engine/Source/Runtime/Engine/Private/ClearReplacementShaders.cpp
@@ -12,17 +12,17 @@ IMPLEMENT_SHADER_TYPE(,FClearTexture2DReplacementCS,TEXT("ClearReplacementShader
 
 IMPLEMENT_SHADER_TYPE(,FClearBufferReplacementCS,TEXT("ClearReplacementShaders"),TEXT("ClearBufferCS"),SF_Compute);
 
-void FClearTexture2DReplacementCS::SetParameters( FRHICommandList* RHICmdList, FUnorderedAccessViewRHIParamRef TextureRW, FLinearColor Value )
+void FClearTexture2DReplacementCS::SetParameters( FRHICommandList& RHICmdList, FUnorderedAccessViewRHIParamRef TextureRW, FLinearColor Value )
 {
-	check(!RHICmdList);
+	RHICmdList.CheckIsNull();
 	FComputeShaderRHIParamRef ComputeShaderRHI = GetComputeShader();
 	SetShaderValue(RHICmdList, ComputeShaderRHI, ClearColor, Value);
 	RHISetUAVParameter(ComputeShaderRHI, ClearTextureRW.GetBaseIndex(), TextureRW);
 }
 
-void FClearBufferReplacementCS::SetParameters( FRHICommandList* RHICmdList, FUnorderedAccessViewRHIParamRef BufferRW, uint32 Dword )
+void FClearBufferReplacementCS::SetParameters( FRHICommandList& RHICmdList, FUnorderedAccessViewRHIParamRef BufferRW, uint32 Dword )
 {
-	check(!RHICmdList);
+	RHICmdList.CheckIsNull();
 	FComputeShaderRHIParamRef ComputeShaderRHI = GetComputeShader();
 	SetShaderValue(RHICmdList, ComputeShaderRHI, ClearDword, Dword);
 	RHISetUAVParameter(ComputeShaderRHI, ClearBufferRW.GetBaseIndex(), BufferRW);
diff --git a/Engine/Source/Runtime/Engine/Private/GPUSkinCache.cpp b/Engine/Source/Runtime/Engine/Private/GPUSkinCache.cpp
--- a/Engine/Source/Runtime/Engine/Private/GPUSkinCache.cpp
+++ b/Engine/Source/Runtime/Engine/Private/GPUSkinCache.cpp
@@ -78,7 +78,7 @@ public:
 	{
 	}
 
-	void SetParameters(FRHICommandList* RHICmdList, uint32 InputStreamStride, uint32 InputStreamFloatOffset, uint32 InputStreamVertexCount, uint32 OutputBufferFloatOffset, FBoneBufferTypeRef BoneBuffer, FUniformBufferRHIRef UniformBuffer, FShaderResourceViewRHIRef VertexBufferSRV, FRWBuffer& SkinBuffer, const FVector& MeshOrigin, const FVector& MeshExtension)
+	void SetParameters(FRHICommandList& RHICmdList, uint32 InputStreamStride, uint32 InputStreamFloatOffset, uint32 InputStreamVertexCount, uint32 OutputBufferFloatOffset, FBoneBufferTypeRef BoneBuffer, FUniformBufferRHIRef UniformBuffer, FShaderResourceViewRHIRef VertexBufferSRV, FRWBuffer& SkinBuffer, const FVector& MeshOrigin, const FVector& MeshExtension)
 	{
 		FComputeShaderRHIParamRef ComputeShaderRHI = GetComputeShader();
 		SetShaderValue(RHICmdList, ComputeShaderRHI, SkinMeshOriginParameter, MeshOrigin);
@@ -89,7 +89,7 @@ public:
 		SetShaderValue(RHICmdList, ComputeShaderRHI, SkinInputStreamFloatOffset, InputStreamFloatOffset);
 		SetShaderValue(RHICmdList, ComputeShaderRHI, SkinOutputBufferFloatOffset, OutputBufferFloatOffset);
 
-		check(!RHICmdList);
+		RHICmdList.CheckIsNull();
 
 		if (UniformBuffer)
 		{
@@ -111,9 +111,9 @@ public:
 			);
 	}
 
-	void UnsetParameters(FRHICommandList* RHICmdList)
+	void UnsetParameters(FRHICommandList& RHICmdList)
 	{
-		check(!RHICmdList);
+		RHICmdList.CheckIsNull();
 
 		FComputeShaderRHIParamRef ComputeShaderRHI = GetComputeShader();
 		FShaderResourceViewRHIParamRef NullSRV = FShaderResourceViewRHIParamRef();
@@ -358,7 +358,7 @@ int32 FGPUSkinCache::StartCacheMesh(int32 Key, const FVertexFactory* VertexFacto
 
 	auto& ShaderData = GPUVertexFactory->GetShaderData();
 	//@todo-rco: RHIPacketList
-	FRHICommandList* RHICmdList = nullptr;
+	FRHICommandList& RHICmdList = FRHICommandList::GetNullRef();
 	DispatchSkinCacheProcess(RHICmdList, InputStreamFloatOffset, ECSInfo->StreamOffset, ShaderData.GetBoneBuffer(), ShaderData.GetUniformBuffer(), VBInfo, StreamStrides[0], ElementVertexCount, ShaderData.MeshOrigin, ShaderData.MeshExtension, bExtraBoneInfluences);
 
 	CacheCurrentFloatOffset += ProcessedFloatCount;
@@ -387,7 +387,7 @@ bool FGPUSkinCache::IsElementProcessed(int32 Key)
 	return false;
 }
 
-bool FGPUSkinCache::SetVertexStreamFromCache(FRHICommandList* RHICmdList, int32 Key, FShader* Shader, const FVertexFactory* VertexFactory, uint32 BaseVertexIndex, bool VelocityPass, FShaderParameter PreviousStreamFloatOffset, FShaderParameter PreviousStreamFloatStride, FShaderResourceParameter PreviousStreamBuffer)
+bool FGPUSkinCache::SetVertexStreamFromCache(FRHICommandList& RHICmdList, int32 Key, FShader* Shader, const FVertexFactory* VertexFactory, uint32 BaseVertexIndex, bool VelocityPass, FShaderParameter PreviousStreamFloatOffset, FShaderParameter PreviousStreamFloatStride, FShaderResourceParameter PreviousStreamBuffer)
 {
 	if (Key >= 0 && Key < CachedElements.Num())
 	{
@@ -402,7 +402,7 @@ bool FGPUSkinCache::SetVertexStreamFromCache(FRHICommandList* RHICmdList, int32
 				return false;
 			}
 
-			check(!RHICmdList);
+			RHICmdList.CheckIsNull();
 
 			RHISetStreamSource(0, SkinCacheBufferRW.Buffer, CacheInfo.StreamStride, CompensatedOffset);
 
@@ -478,7 +478,7 @@ FGPUSkinCache::FElementCacheStatusInfo* FGPUSkinCache::FindEvictableCacheStatusI
 	}
 }
 
-void FGPUSkinCache::DispatchSkinCacheProcess(FRHICommandList* RHICmdList, uint32 InputStreamFloatOffset, uint32 OutputBufferFloatOffset, FBoneBufferTypeRef BoneBuffer, FUniformBufferRHIRef UniformBuffer, const FVertexBufferInfo* VBInfo, uint32 VertexStride, uint32 VertexCount, const FVector& MeshOrigin, const FVector& MeshExtension, bool bUseExtraBoneInfluences)
+void FGPUSkinCache::DispatchSkinCacheProcess(FRHICommandList& RHICmdList, uint32 InputStreamFloatOffset, uint32 OutputBufferFloatOffset, FBoneBufferTypeRef BoneBuffer, FUniformBufferRHIRef UniformBuffer, const FVertexBufferInfo* VBInfo, uint32 VertexStride, uint32 VertexCount, const FVector& MeshOrigin, const FVector& MeshExtension, bool bUseExtraBoneInfluences)
 {
 	if (GRHIFeatureLevel < ERHIFeatureLevel::SM5)
 	{
@@ -493,7 +493,7 @@ void FGPUSkinCache::DispatchSkinCacheProcess(FRHICommandList* RHICmdList, uint32
 	{
 		TShaderMapRef<FGPUSkinCacheCS<true> > SkinCacheCS(GetGlobalShaderMap());
 		SkinCacheCS->SetParameters(RHICmdList, VertexStride, InputStreamFloatOffset, VertexCount, OutputBufferFloatOffset, BoneBuffer, UniformBuffer, VBInfo->VertexBufferSRV, SkinCacheBufferRW, MeshOrigin, MeshExtension);
-		check(!RHICmdList);
+		RHICmdList.CheckIsNull();
 
 		RHISetComputeShader(SkinCacheCS->GetComputeShader());		
 		RHIDispatchComputeShader(VertexCountAlign64, 1, 1);
diff --git a/Engine/Source/Runtime/Engine/Private/GPUSkinCache.h b/Engine/Source/Runtime/Engine/Private/GPUSkinCache.h
--- a/Engine/Source/Runtime/Engine/Private/GPUSkinCache.h
+++ b/Engine/Source/Runtime/Engine/Private/GPUSkinCache.h
@@ -42,7 +42,7 @@ public:
 	//	Call End*
 	int32	StartCacheMesh(int32 Key, const class FVertexFactory* VertexFactory, const FVertexFactory* TargetVertexFactory, const struct FSkelMeshChunk& BatchElement, const FSkeletalMeshObjectGPUSkin* Skin, bool bExtraBoneInfluences);
 
-	bool	SetVertexStreamFromCache(FRHICommandList* RHICmdList, int32 Key, FShader* Shader, const FVertexFactory* VertexFactory, uint32 BaseVertexIndex, bool VelocityPass, FShaderParameter PreviousStreamFloatOffset, FShaderParameter PreviousStreamFloatStride, FShaderResourceParameter PreviousStreamBuffer);
+	bool	SetVertexStreamFromCache(FRHICommandList& RHICmdList, int32 Key, FShader* Shader, const FVertexFactory* VertexFactory, uint32 BaseVertexIndex, bool VelocityPass, FShaderParameter PreviousStreamFloatOffset, FShaderParameter PreviousStreamFloatStride, FShaderResourceParameter PreviousStreamBuffer);
 
 	struct FElementCacheStatusInfo
 	{
@@ -113,7 +113,7 @@ private:
 
 	FElementCacheStatusInfo* FindEvictableCacheStatusInfo();
 
-	void	DispatchSkinCacheProcess(FRHICommandList* RHICmdList, uint32 InputStreamFloatOffset, uint32 OutputBufferFloatOffset, FBoneBufferTypeRef BoneBuffer, FUniformBufferRHIRef UniformBuffer, const FVertexBufferInfo* VBInfo, uint32 VertexStride, uint32 VertexCount, const FVector& MeshOrigin, const FVector& MeshExtension, bool bUseExtraBoneInfluences);
+	void	DispatchSkinCacheProcess(FRHICommandList& RHICmdList, uint32 InputStreamFloatOffset, uint32 OutputBufferFloatOffset, FBoneBufferTypeRef BoneBuffer, FUniformBufferRHIRef UniformBuffer, const FVertexBufferInfo* VBInfo, uint32 VertexStride, uint32 VertexCount, const FVector& MeshOrigin, const FVector& MeshExtension, bool bUseExtraBoneInfluences);
 
 	bool	Initialized : 1;
 
diff --git a/Engine/Source/Runtime/Engine/Private/GPUSkinVertexFactory.cpp b/Engine/Source/Runtime/Engine/Private/GPUSkinVertexFactory.cpp
--- a/Engine/Source/Runtime/Engine/Private/GPUSkinVertexFactory.cpp
+++ b/Engine/Source/Runtime/Engine/Private/GPUSkinVertexFactory.cpp
@@ -426,9 +426,9 @@ public:
 	/**
 	* Set any shader data specific to this vertex factory
 	*/
-	virtual void SetMesh(FRHICommandList* RHICmdList, FShader* Shader,const FVertexFactory* VertexFactory,const FSceneView& View,const FMeshBatchElement& BatchElement,uint32 DataFlags) const OVERRIDE
+	virtual void SetMesh(FRHICommandList& RHICmdList, FShader* Shader,const FVertexFactory* VertexFactory,const FSceneView& View,const FMeshBatchElement& BatchElement,uint32 DataFlags) const OVERRIDE
 	{
-		check(!RHICmdList);
+		RHICmdList.CheckIsNull();
 
 		if(Shader->GetVertexShader())
 		{
@@ -573,7 +573,7 @@ public:
 	/**
 	* Set any shader data specific to this vertex factory
 	*/
-	virtual void SetMesh(FRHICommandList* RHICmdList, FShader* Shader,const FVertexFactory* VertexFactory,const FSceneView& View,const FMeshBatchElement& BatchElement,uint32 DataFlags) const OVERRIDE
+	virtual void SetMesh(FRHICommandList& RHICmdList, FShader* Shader,const FVertexFactory* VertexFactory,const FSceneView& View,const FMeshBatchElement& BatchElement,uint32 DataFlags) const OVERRIDE
 	{
 		FGPUSkinVertexFactoryShaderParameters::SetMesh(RHICmdList, Shader, VertexFactory, View, BatchElement, DataFlags);
 
@@ -711,7 +711,7 @@ public:
 		Ar << ClothBlendWeightParameter;
 	}
 
-	virtual void SetMesh(FRHICommandList* RHICmdList, FShader* Shader,const FVertexFactory* VertexFactory,const FSceneView& View,const FMeshBatchElement& BatchElement,uint32 DataFlags) const OVERRIDE
+	virtual void SetMesh(FRHICommandList& RHICmdList, FShader* Shader,const FVertexFactory* VertexFactory,const FSceneView& View,const FMeshBatchElement& BatchElement,uint32 DataFlags) const OVERRIDE
 	{
 		if(Shader->GetVertexShader())
 		{
@@ -725,7 +725,7 @@ public:
 				: ((const TGPUSkinAPEXClothVertexFactory<false>*)GPUSkinVertexFactory)->GetClothShaderData();
 
 			SetUniformBufferParameter(RHICmdList, Shader->GetVertexShader(),Shader->GetUniformBufferParameter<FAPEXClothUniformShaderParameters>(),ClothShaderData.GetClothUniformBuffer());
-			check(!RHICmdList);
+			RHICmdList.CheckIsNull();
 
 			// we tell the shader where to pickup the data
 			if(ClothSimulPositionsParameter.IsBound())
diff --git a/Engine/Source/Runtime/Engine/Private/GPUSort.cpp b/Engine/Source/Runtime/Engine/Private/GPUSort.cpp
--- a/Engine/Source/Runtime/Engine/Private/GPUSort.cpp
+++ b/Engine/Source/Runtime/Engine/Private/GPUSort.cpp
@@ -256,7 +256,7 @@ public:
 	/**
 	 * Set parameters for this shader.
 	 */
-	void SetParameters(FRHICommandList* RHICmdList, const FRadixSortUniformBufferRef& UniformBuffer )
+	void SetParameters(FRHICommandList& RHICmdList, const FRadixSortUniformBufferRef& UniformBuffer )
 	{
 		FComputeShaderRHIParamRef ComputeShaderRHI = GetComputeShader();
 		SetUniformBufferParameter(RHICmdList, ComputeShaderRHI, GetUniformBufferParameter<FRadixSortParameters>(), UniformBuffer );
@@ -354,11 +354,11 @@ public:
 	/**
 	 * Set parameters for this shader.
 	 */
-	void SetParameters(FRHICommandList* RHICmdList, FShaderResourceViewRHIParamRef InKeysSRV, FRadixSortUniformBufferRef& RadixSortUniformBuffer, FShaderResourceViewRHIParamRef RadixSortParameterBufferSRV )
+	void SetParameters(FRHICommandList& RHICmdList, FShaderResourceViewRHIParamRef InKeysSRV, FRadixSortUniformBufferRef& RadixSortUniformBuffer, FShaderResourceViewRHIParamRef RadixSortParameterBufferSRV )
 	{
 		FComputeShaderRHIParamRef ComputeShaderRHI = GetComputeShader();
 		SetUniformBufferParameter(RHICmdList, ComputeShaderRHI, GetUniformBufferParameter<FRadixSortParameters>(), RadixSortUniformBuffer );
-		check(!RHICmdList);
+		RHICmdList.CheckIsNull();
 		if ( InKeys.IsBound() )
 		{
 			RHISetShaderResourceViewParameter( ComputeShaderRHI, InKeys.GetBaseIndex(), InKeysSRV );
@@ -462,12 +462,12 @@ public:
 	/**
 	 * Set parameters for this shader.
 	 */
-	void SetParameters(FRHICommandList* RHICmdList, FShaderResourceViewRHIParamRef InOffsetsSRV )
+	void SetParameters(FRHICommandList& RHICmdList, FShaderResourceViewRHIParamRef InOffsetsSRV )
 	{
 		FComputeShaderRHIParamRef ComputeShaderRHI = GetComputeShader();
 		if ( InOffsets.IsBound() )
 		{
-			check(!RHICmdList);
+			RHICmdList.CheckIsNull();
 
 			RHISetShaderResourceViewParameter( ComputeShaderRHI, InOffsets.GetBaseIndex(), InOffsetsSRV );
 		}
@@ -580,7 +580,7 @@ public:
 	 * Set parameters for this shader.
 	 */
 	void SetParameters(
-		FRHICommandList* RHICmdList, 
+		FRHICommandList& RHICmdList, 
 		FShaderResourceViewRHIParamRef InKeysSRV,
 		FShaderResourceViewRHIParamRef InValuesSRV,
 		FShaderResourceViewRHIParamRef InOffsetsSRV,
@@ -589,7 +589,7 @@ public:
 	{
 		FComputeShaderRHIParamRef ComputeShaderRHI = GetComputeShader();
 		SetUniformBufferParameter(RHICmdList, ComputeShaderRHI, GetUniformBufferParameter<FRadixSortParameters>(), RadixSortUniformBuffer );
-		check(!RHICmdList);
+		RHICmdList.CheckIsNull();
 		if ( RadixSortParameterBuffer.IsBound() )
 		{
 			RHISetShaderResourceViewParameter( ComputeShaderRHI, RadixSortParameterBuffer.GetBaseIndex(), RadixSortParameterBufferSRV );
@@ -742,7 +742,7 @@ int32 SortGPUBuffers( FGPUSortBuffers SortBuffers, int32 BufferIndex, uint32 Key
 	const bool bUseConstantBufferWorkaround = UpsweepCS->RequiresConstantBufferWorkaround();
 
 	//@todo-rco: RHIPacketList
-	FRHICommandList* RHICmdList = nullptr;
+	FRHICommandList& RHICmdList = FRHICommandList::GetNullRef();
 
 	// Execute each pass as needed.
 	uint32 PassBits = DIGIT_COUNT - 1;
diff --git a/Engine/Source/Runtime/Engine/Private/InstancedStaticMesh.cpp b/Engine/Source/Runtime/Engine/Private/InstancedStaticMesh.cpp
--- a/Engine/Source/Runtime/Engine/Private/InstancedStaticMesh.cpp
+++ b/Engine/Source/Runtime/Engine/Private/InstancedStaticMesh.cpp
@@ -512,7 +512,7 @@ class FInstancedStaticMeshVertexFactoryShaderParameters : public FVertexFactoryS
 		InstancingFadeOutParamsParameter.Bind(ParameterMap,TEXT("InstancingFadeOutParams"));
 	}
 
-	virtual void SetMesh(FRHICommandList* RHICmdList, FShader* VertexShader,const class FVertexFactory* VertexFactory,const class FSceneView& View,const struct FMeshBatchElement& BatchElement,uint32 DataFlags) const OVERRIDE
+	virtual void SetMesh(FRHICommandList& RHICmdList, FShader* VertexShader,const class FVertexFactory* VertexFactory,const class FSceneView& View,const struct FMeshBatchElement& BatchElement,uint32 DataFlags) const OVERRIDE
 	{
 		if( InstancedViewTranslationParameter.IsBound() )
 		{
diff --git a/Engine/Source/Runtime/Engine/Private/Landscape/LandscapeRender.cpp b/Engine/Source/Runtime/Engine/Private/Landscape/LandscapeRender.cpp
--- a/Engine/Source/Runtime/Engine/Private/Landscape/LandscapeRender.cpp
+++ b/Engine/Source/Runtime/Engine/Private/Landscape/LandscapeRender.cpp
@@ -1992,7 +1992,7 @@ public:
 	/**
 	* Set any shader data specific to this vertex factory
 	*/
-	virtual void SetMesh(FRHICommandList* RHICmdList, FShader* VertexShader,const class FVertexFactory* VertexFactory,const class FSceneView& View,const struct FMeshBatchElement& BatchElement,uint32 DataFlags) const OVERRIDE
+	virtual void SetMesh(FRHICommandList& RHICmdList, FShader* VertexShader,const class FVertexFactory* VertexFactory,const class FSceneView& View,const struct FMeshBatchElement& BatchElement,uint32 DataFlags) const OVERRIDE
 	{
 		SCOPE_CYCLE_COUNTER(STAT_LandscapeVFDrawTime);
 
@@ -2097,7 +2097,7 @@ public:
 	/**
 	* Set any shader data specific to this vertex factory
 	*/
-	virtual void SetMesh(FRHICommandList* RHICmdList, FShader* VertexShader,const class FVertexFactory* VertexFactory,const class FSceneView& View,const struct FMeshBatchElement& BatchElement,uint32 DataFlags) const OVERRIDE
+	virtual void SetMesh(FRHICommandList& RHICmdList, FShader* VertexShader,const class FVertexFactory* VertexFactory,const class FSceneView& View,const struct FMeshBatchElement& BatchElement,uint32 DataFlags) const OVERRIDE
 	{
 		SCOPE_CYCLE_COUNTER(STAT_LandscapeVFDrawTime);
 
@@ -2219,7 +2219,7 @@ void FLandscapeVertexFactoryPixelShaderParameters::Serialize(FArchive& Ar)
 /**
 * Set any shader data specific to this vertex factory
 */
-void FLandscapeVertexFactoryPixelShaderParameters::SetMesh(FRHICommandList* RHICmdList, FShader* PixelShader,const class FVertexFactory* VertexFactory,const class FSceneView& View,const struct FMeshBatchElement& BatchElement,uint32 DataFlags) const
+void FLandscapeVertexFactoryPixelShaderParameters::SetMesh(FRHICommandList& RHICmdList, FShader* PixelShader,const class FVertexFactory* VertexFactory,const class FSceneView& View,const struct FMeshBatchElement& BatchElement,uint32 DataFlags) const
 {
 	SCOPE_CYCLE_COUNTER(STAT_LandscapeVFDrawTime);
 
diff --git a/Engine/Source/Runtime/Engine/Private/Landscape/LandscapeRenderMobile.cpp b/Engine/Source/Runtime/Engine/Private/Landscape/LandscapeRenderMobile.cpp
--- a/Engine/Source/Runtime/Engine/Private/Landscape/LandscapeRenderMobile.cpp
+++ b/Engine/Source/Runtime/Engine/Private/Landscape/LandscapeRenderMobile.cpp
@@ -61,7 +61,7 @@ public:
 	/**
 	* Set any shader data specific to this vertex factory
 	*/
-	virtual void SetMesh(FRHICommandList* RHICmdList, FShader* VertexShader,const class FVertexFactory* VertexFactory,const class FSceneView& View,const struct FMeshBatchElement& BatchElement,uint32 DataFlags) const OVERRIDE
+	virtual void SetMesh(FRHICommandList& RHICmdList, FShader* VertexShader,const class FVertexFactory* VertexFactory,const class FSceneView& View,const struct FMeshBatchElement& BatchElement,uint32 DataFlags) const OVERRIDE
 	{
 		SCOPE_CYCLE_COUNTER(STAT_LandscapeVFDrawTime);
 
diff --git a/Engine/Source/Runtime/Engine/Private/LocalVertexFactory.cpp b/Engine/Source/Runtime/Engine/Private/LocalVertexFactory.cpp
--- a/Engine/Source/Runtime/Engine/Private/LocalVertexFactory.cpp
+++ b/Engine/Source/Runtime/Engine/Private/LocalVertexFactory.cpp
@@ -22,7 +22,7 @@ public:
 		Ar << LODParameter;
 	}
 
-	virtual void SetMesh(FRHICommandList* RHICmdList, FShader* Shader,const FVertexFactory* VertexFactory,const FSceneView& View,const FMeshBatchElement& BatchElement,uint32 DataFlags) const OVERRIDE
+	virtual void SetMesh(FRHICommandList& RHICmdList, FShader* Shader,const FVertexFactory* VertexFactory,const FSceneView& View,const FMeshBatchElement& BatchElement,uint32 DataFlags) const OVERRIDE
 	{
 		if (View.Family != NULL && View.Family->Scene != NULL)
 		{
diff --git a/Engine/Source/Runtime/Engine/Private/MeshParticleVertexFactory.cpp b/Engine/Source/Runtime/Engine/Private/MeshParticleVertexFactory.cpp
--- a/Engine/Source/Runtime/Engine/Private/MeshParticleVertexFactory.cpp
+++ b/Engine/Source/Runtime/Engine/Private/MeshParticleVertexFactory.cpp
@@ -36,7 +36,7 @@ public:
 		Ar << ParticleColor;
 	}
 
-	virtual void SetMesh(FRHICommandList* RHICmdList, FShader* Shader,const FVertexFactory* VertexFactory,const FSceneView& View,const FMeshBatchElement& BatchElement,uint32 DataFlags) const OVERRIDE
+	virtual void SetMesh(FRHICommandList& RHICmdList, FShader* Shader,const FVertexFactory* VertexFactory,const FSceneView& View,const FMeshBatchElement& BatchElement,uint32 DataFlags) const OVERRIDE
 	{
 		const bool bInstanced = View.GetFeatureLevel() >= ERHIFeatureLevel::SM3;
 		FMeshParticleVertexFactory* MeshParticleVF = (FMeshParticleVertexFactory*)VertexFactory;
diff --git a/Engine/Source/Runtime/Engine/Private/Particles/ParticleCurveTexture.cpp b/Engine/Source/Runtime/Engine/Private/Particles/ParticleCurveTexture.cpp
--- a/Engine/Source/Runtime/Engine/Private/Particles/ParticleCurveTexture.cpp
+++ b/Engine/Source/Runtime/Engine/Private/Particles/ParticleCurveTexture.cpp
@@ -70,7 +70,7 @@ public:
 	/**
 	 * Sets parameters for particle injection.
 	 */
-	void SetParameters(FRHICommandList* RHICmdList, const FVector2D& CurveOffset )
+	void SetParameters(FRHICommandList& RHICmdList, const FVector2D& CurveOffset )
 	{
 		FParticleCurveInjectionParameters Parameters;
 		Parameters.PixelScale.X = 1.0f / GParticleCurveTextureSizeX;
@@ -174,7 +174,7 @@ TGlobalResource<FParticleCurveInjectionVertexDeclaration> GParticleCurveInjectio
  * @param InPendingCurves - Curves to be stored on the GPU.
  */
 static void InjectCurves(
-	FRHICommandList* RHICmdList, 
+	FRHICommandList& RHICmdList, 
 	FTexture2DRHIParamRef CurveTextureRHI,
 	FTexture2DRHIParamRef CurveTextureTargetRHI,
 	TArray<FCurveSamples>& InPendingCurves )
@@ -187,7 +187,7 @@ static void InjectCurves(
 
 	FVertexBufferRHIParamRef ScratchVertexBufferRHI = GParticleScratchVertexBuffer.VertexBufferRHI;
 
-	check(!RHICmdList);
+	RHICmdList.CheckIsNull();
 
 	RHISetRenderTarget( CurveTextureTargetRHI, FTextureRHIParamRef() );
 	RHISetScissorRect( false, 0, 0, 0, 0 );
@@ -595,7 +595,7 @@ void FParticleCurveTexture::SubmitPendingCurves()
 		{
 			//@todo-rco: RHIPacketList
 			InjectCurves(
-				nullptr,
+				FRHICommandList::GetNullRef(),
 				ParticleCurveTexture->CurveTextureRHI,
 				ParticleCurveTexture->CurveTextureTargetRHI,
 				PendingCurves
diff --git a/Engine/Source/Runtime/Engine/Private/Particles/ParticleGpuSimulation.cpp b/Engine/Source/Runtime/Engine/Private/Particles/ParticleGpuSimulation.cpp
--- a/Engine/Source/Runtime/Engine/Private/Particles/ParticleGpuSimulation.cpp
+++ b/Engine/Source/Runtime/Engine/Private/Particles/ParticleGpuSimulation.cpp
@@ -497,7 +497,7 @@ public:
 		Ar << CurveTextureSampler;
 	}
 
-	virtual void SetMesh(FRHICommandList* RHICmdList, FShader* Shader,const FVertexFactory* VertexFactory,const FSceneView& View,const FMeshBatchElement& BatchElement,uint32 DataFlags) const OVERRIDE;
+	virtual void SetMesh(FRHICommandList& RHICmdList, FShader* Shader,const FVertexFactory* VertexFactory,const FSceneView& View,const FMeshBatchElement& BatchElement,uint32 DataFlags) const OVERRIDE;
 
 	virtual uint32 GetSize() const { return sizeof(*this); }
 
@@ -644,7 +644,7 @@ public:
 /**
  * Set vertex factory shader parameters.
  */
-void FGPUSpriteVertexFactoryShaderParameters::SetMesh(FRHICommandList* RHICmdList, FShader* Shader,const FVertexFactory* VertexFactory,const FSceneView& View,const FMeshBatchElement& BatchElement,uint32 DataFlags) const
+void FGPUSpriteVertexFactoryShaderParameters::SetMesh(FRHICommandList& RHICmdList, FShader* Shader,const FVertexFactory* VertexFactory,const FSceneView& View,const FMeshBatchElement& BatchElement,uint32 DataFlags) const
 {
 	FGPUSpriteVertexFactory* GPUVF = (FGPUSpriteVertexFactory*)VertexFactory;
 	FVertexShaderRHIParamRef VertexShader = Shader->GetVertexShader();
@@ -654,7 +654,7 @@ void FGPUSpriteVertexFactoryShaderParameters::SetMesh(FRHICommandList* RHICmdLis
 	SetUniformBufferParameter(RHICmdList, VertexShader, Shader->GetUniformBufferParameter<FGPUSpriteEmitterDynamicUniformParameters>(), GPUVF->EmitterDynamicUniformBuffer );
 	if (ParticleIndices.IsBound())
 	{
-		check(!RHICmdList);
+		RHICmdList.CheckIsNull();
 
 		RHISetShaderResourceViewParameter(VertexShader, ParticleIndices.GetBaseIndex(), GPUVF->ParticleIndicesBuffer->VertexBufferSRV);
 	}
@@ -805,9 +805,9 @@ public:
 	}
 
 	/** Set parameters. */
-	void SetParameters(FRHICommandList* RHICmdList, FParticleShaderParamRef TileOffsetsRef)
+	void SetParameters(FRHICommandList& RHICmdList, FParticleShaderParamRef TileOffsetsRef)
 	{
-		check(!RHICmdList);
+		RHICmdList.CheckIsNull();
 
 		FVertexShaderRHIParamRef VertexShaderRHI = GetVertexShader();
 		if (TileOffsets.IsBound())
@@ -908,7 +908,7 @@ public:
 	 * Set parameters for this shader.
 	 */
 	void SetParameters(
-		FRHICommandList* RHICmdList, 
+		FRHICommandList& RHICmdList, 
 		const FParticleStateTextures& TextureResources,
 		const FParticleAttributesTexture& InAttributesTexture,
 		const FParticleAttributesTexture& InRenderAttributesTexture,
@@ -961,7 +961,7 @@ public:
 	 * Set parameters for the vector fields sampled by this shader.
 	 * @param VectorFieldParameters -Parameters needed to sample local vector fields.
 	 */
-	void SetVectorFieldParameters(FRHICommandList* RHICmdList, const FVectorFieldUniformBufferRef& UniformBuffer, const FTexture3DRHIParamRef VolumeTexturesRHI[])
+	void SetVectorFieldParameters(FRHICommandList& RHICmdList, const FVectorFieldUniformBufferRef& UniformBuffer, const FTexture3DRHIParamRef VolumeTexturesRHI[])
 	{
 		FPixelShaderRHIParamRef PixelShaderRHI = GetPixelShader();
 		FSamplerStateRHIParamRef SamplerStateLinear = TStaticSamplerState<SF_Bilinear,AM_Clamp,AM_Clamp,AM_Clamp>::GetRHI();
@@ -975,7 +975,7 @@ public:
 	/**
 	 * Set per-instance parameters for this shader.
 	 */
-	void SetInstanceParameters(FRHICommandList* RHICmdList, float InDeltaSeconds, FUniformBufferRHIParamRef UniformBuffer, FUniformBufferRHIParamRef PerFrameUniformBuffer)
+	void SetInstanceParameters(FRHICommandList& RHICmdList, float InDeltaSeconds, FUniformBufferRHIParamRef UniformBuffer, FUniformBufferRHIParamRef PerFrameUniformBuffer)
 	{
 		FPixelShaderRHIParamRef PixelShaderRHI = GetPixelShader();
 		SetShaderValue(RHICmdList, PixelShaderRHI, DeltaSeconds, InDeltaSeconds);
@@ -1212,7 +1212,7 @@ struct FSimulationCommandGPU
  */
 template <bool bUseDepthBufferCollision>
 void ExecuteSimulationCommands(
-	FRHICommandList* RHICmdList, 
+	FRHICommandList& RHICmdList, 
 	const TArray<FSimulationCommandGPU>& SimulationCommands,
 	const FParticleStateTextures& TextureResources,
 	const FParticleAttributesTexture& AttributeTexture,
@@ -1297,7 +1297,7 @@ void ClearTiles(const TArray<uint32>& Tiles)
 		const uint32* TilesPtr = Tiles.GetTypedData() + FirstTile;
 		BuildTileVertexBuffer( BufferParam, TilesPtr, TilesThisDrawCall );
 		//@todo-rco: RHIPacketList
-		VertexShader->SetParameters(nullptr, ShaderParam);
+		VertexShader->SetParameters(FRHICommandList::GetNullRef(), ShaderParam);
 		DrawAlignedParticleTiles( ComputeAlignedTileCount(TilesThisDrawCall) );
 		TileCount -= TilesThisDrawCall;
 		FirstTile += TilesThisDrawCall;
@@ -1387,7 +1387,7 @@ public:
 	/**
 	 * Sets parameters for particle injection.
 	 */
-	void SetParameters(FRHICommandList* RHICmdList)
+	void SetParameters(FRHICommandList& RHICmdList)
 	{
 		FParticleInjectionParameters Parameters;
 		Parameters.PixelScale.X = 1.0f / GParticleSimulationTextureSizeX;
@@ -1532,7 +1532,7 @@ void InjectNewParticles( const TArray<FNewParticle>& NewParticles )
 			);
 
 		//@todo-rco: RHIPacketList
-		VertexShader->SetParameters(nullptr);
+		VertexShader->SetParameters(FRHICommandList::GetNullRef());
 
 		// Stream 0: New particles.
 		RHISetStreamSource(
@@ -1607,7 +1607,7 @@ public:
 	/**
 	 * Set parameters for this shader.
 	 */
-	void SetParameters(FRHICommandList* RHICmdList, const FParticleSimVisualizeBufferRef& UniformBuffer )
+	void SetParameters(FRHICommandList& RHICmdList, const FParticleSimVisualizeBufferRef& UniformBuffer )
 	{
 		FVertexShaderRHIParamRef VertexShader = GetVertexShader();
 		SetUniformBufferParameter(RHICmdList, VertexShader, GetUniformBufferParameter<FParticleSimVisualizeParameters>(), UniformBuffer );
@@ -1663,7 +1663,7 @@ public:
 	/**
 	 * Set parameters for this shader.
 	 */
-	void SetParameters(FRHICommandList* RHICmdList, int32 InVisualizationMode, FTexture2DRHIParamRef PositionTextureRHI, FTexture2DRHIParamRef PositionZWTextureRHI, FTexture2DRHIParamRef CurveTextureRHI )
+	void SetParameters(FRHICommandList& RHICmdList, int32 InVisualizationMode, FTexture2DRHIParamRef PositionTextureRHI, FTexture2DRHIParamRef PositionZWTextureRHI, FTexture2DRHIParamRef CurveTextureRHI )
 	{
 		FPixelShaderRHIParamRef PixelShader = GetPixelShader();
 		SetShaderValue(RHICmdList, PixelShader, VisualizationMode, InVisualizationMode );
@@ -1756,7 +1756,7 @@ static void VisualizeGPUSimulation(
 		);
 
 	//@todo-rco: RHIPacketList
-	FRHICommandList* RHICmdList = nullptr;
+	FRHICommandList& RHICmdList = FRHICommandList::GetNullRef();
 
 	// Parameters for the visualization.
 	FParticleSimVisualizeParameters Parameters;
@@ -1912,12 +1912,12 @@ public:
 	/**
 	 * Set output buffers for this shader.
 	 */
-	void SetOutput(FRHICommandList* RHICmdList, FUnorderedAccessViewRHIParamRef OutBoundsUAV )
+	void SetOutput(FRHICommandList& RHICmdList, FUnorderedAccessViewRHIParamRef OutBoundsUAV )
 	{
 		FComputeShaderRHIParamRef ComputeShaderRHI = GetComputeShader();
 		if ( OutBounds.IsBound() )
 		{
-			check(!RHICmdList);
+			RHICmdList.CheckIsNull();
 
 			RHISetUAVParameter( ComputeShaderRHI, OutBounds.GetBaseIndex(), OutBoundsUAV );
 		}
@@ -1927,14 +1927,14 @@ public:
 	 * Set input parameters.
 	 */
 	void SetParameters(
-		FRHICommandList* RHICmdList, 
+		FRHICommandList& RHICmdList, 
 		FParticleBoundsUniformBufferRef& UniformBuffer,
 		FShaderResourceViewRHIParamRef InIndicesSRV,
 		FTexture2DRHIParamRef PositionTextureRHI,
 		FTexture2DRHIParamRef PositionZWTextureRHI
 		)
 	{
-		check(!RHICmdList);
+		RHICmdList.CheckIsNull();
 
 		FComputeShaderRHIParamRef ComputeShaderRHI = GetComputeShader();
 		SetUniformBufferParameter(RHICmdList, ComputeShaderRHI, GetUniformBufferParameter<FParticleBoundsParameters>(), UniformBuffer );
@@ -1955,9 +1955,9 @@ public:
 	/**
 	 * Unbinds any buffers that have been bound.
 	 */
-	void UnbindBuffers(FRHICommandList* RHICmdList)
+	void UnbindBuffers(FRHICommandList& RHICmdList)
 	{
-		check(!RHICmdList);
+		RHICmdList.CheckIsNull();
 
 		FComputeShaderRHIParamRef ComputeShaderRHI = GetComputeShader();
 		if ( InParticleIndices.IsBound() )
@@ -2000,7 +2000,7 @@ static bool AreBoundsValid( const FVector& Mins, const FVector& Maxs )
  * @param ParticleCount - The number of particles in the emitter.
  */
 static FBox ComputeParticleBounds(
-	FRHICommandList* RHICmdList, 
+	FRHICommandList& RHICmdList, 
 	FShaderResourceViewRHIParamRef VertexBufferSRV,
 	FTexture2DRHIParamRef PositionTextureRHI,
 	FTexture2DRHIParamRef PositionZWTextureRHI,
@@ -2024,7 +2024,7 @@ static FBox ComputeParticleBounds(
 		Parameters.ParticleCount = ParticleCount;
 		UniformBuffer = FParticleBoundsUniformBufferRef::CreateUniformBufferImmediate( Parameters, UniformBuffer_SingleFrame );
 
-		check(!RHICmdList);
+		RHICmdList.CheckIsNull();
 
 		// Create a buffer for storing bounds.
 		const int32 BufferSize = GroupCount * 2 * sizeof(FVector4);
@@ -3273,7 +3273,7 @@ public:
 		{
 			//@todo-rco: RHIPacketList
 			EmitterInstance->ParticleBoundingBox = ComputeParticleBounds(
-				nullptr,
+				FRHICommandList::GetNullRef(),
 				EmitterInstance->Simulation->VertexBuffer.VertexBufferSRV, 
 				EmitterInstance->FXSystem->GetParticleSimulationResources()->GetCurrentStateTextures().PositionTextureRHI,
 				EmitterInstance->FXSystem->GetParticleSimulationResources()->GetCurrentStateTextures().PositionZWTextureRHI,
@@ -4160,7 +4160,7 @@ void FFXSystem::SimulateGPUParticles(
 		{
 			/// ?
 			ExecuteSimulationCommands<true>(
-				nullptr, 
+				FRHICommandList::GetNullRef(),
 				SimulationCommands,
 				PrevStateTextures,
 				ParticleSimulationResources->SimulationAttributesTexture,
@@ -4174,7 +4174,7 @@ void FFXSystem::SimulateGPUParticles(
 		{
 			/// ?
 			ExecuteSimulationCommands<false>(
-				nullptr,
+				FRHICommandList::GetNullRef(),
 				SimulationCommands,
 				PrevStateTextures,
 				ParticleSimulationResources->SimulationAttributesTexture,
diff --git a/Engine/Source/Runtime/Engine/Private/Particles/ParticleSortingGPU.cpp b/Engine/Source/Runtime/Engine/Private/Particles/ParticleSortingGPU.cpp
--- a/Engine/Source/Runtime/Engine/Private/Particles/ParticleSortingGPU.cpp
+++ b/Engine/Source/Runtime/Engine/Private/Particles/ParticleSortingGPU.cpp
@@ -94,9 +94,9 @@ public:
 	/**
 	 * Set output buffers for this shader.
 	 */
-	void SetOutput(FRHICommandList* RHICmdList, FUnorderedAccessViewRHIParamRef OutKeysUAV, FUnorderedAccessViewRHIParamRef OutIndicesUAV )
+	void SetOutput(FRHICommandList& RHICmdList, FUnorderedAccessViewRHIParamRef OutKeysUAV, FUnorderedAccessViewRHIParamRef OutIndicesUAV )
 	{
-		check(!RHICmdList);
+		RHICmdList.CheckIsNull();
 
 		FComputeShaderRHIParamRef ComputeShaderRHI = GetComputeShader();
 		if ( OutKeys.IsBound() )
@@ -113,7 +113,7 @@ public:
 	 * Set input parameters.
 	 */
 	void SetParameters(
-		FRHICommandList* RHICmdList,
+		FRHICommandList& RHICmdList,
 		FParticleKeyGenUniformBufferRef& UniformBuffer,
 		FShaderResourceViewRHIParamRef InIndicesSRV
 		)
@@ -122,7 +122,7 @@ public:
 		SetUniformBufferParameter(RHICmdList, ComputeShaderRHI, GetUniformBufferParameter<FParticleKeyGenParameters>(), UniformBuffer );
 		if ( InParticleIndices.IsBound() )
 		{
-			check(!RHICmdList);
+			RHICmdList.CheckIsNull();
 
 			RHISetShaderResourceViewParameter( ComputeShaderRHI, InParticleIndices.GetBaseIndex(), InIndicesSRV );
 		}
@@ -131,18 +131,18 @@ public:
 	/**
 	 * Set the texture from which particle positions can be read.
 	 */
-	void SetPositionTextures(FRHICommandList* RHICmdList, FTexture2DRHIParamRef PositionTextureRHI, FTexture2DRHIParamRef PositionZWTextureRHI)
+	void SetPositionTextures(FRHICommandList& RHICmdList, FTexture2DRHIParamRef PositionTextureRHI, FTexture2DRHIParamRef PositionZWTextureRHI)
 	{
 		FComputeShaderRHIParamRef ComputeShaderRHI = GetComputeShader();
 		if (PositionTexture.IsBound())
 		{
-			check(!RHICmdList);
+			RHICmdList.CheckIsNull();
 
 			RHISetShaderTexture(ComputeShaderRHI, PositionTexture.GetBaseIndex(), PositionTextureRHI);
 		}
 		if (PositionZWTexture.IsBound())
 		{
-			check(!RHICmdList);
+			RHICmdList.CheckIsNull();
 
 			RHISetShaderTexture(ComputeShaderRHI, PositionZWTexture.GetBaseIndex(), PositionZWTextureRHI);
 		}
@@ -151,9 +151,9 @@ public:
 	/**
 	 * Unbinds any buffers that have been bound.
 	 */
-	void UnbindBuffers(FRHICommandList* RHICmdList)
+	void UnbindBuffers(FRHICommandList& RHICmdList)
 	{
-		check(!RHICmdList);
+		RHICmdList.CheckIsNull();
 
 		FComputeShaderRHIParamRef ComputeShaderRHI = GetComputeShader();
 		if ( InParticleIndices.IsBound() )
@@ -195,7 +195,7 @@ IMPLEMENT_SHADER_TYPE(,FParticleSortKeyGenCS,TEXT("ParticleSortKeyGen"),TEXT("Ge
  * @returns the total number of particles being sorted.
  */
 static int32 GenerateParticleSortKeys(
-	FRHICommandList* RHICmdList,
+	FRHICommandList& RHICmdList,
 	FUnorderedAccessViewRHIParamRef KeyBufferUAV,
 	FUnorderedAccessViewRHIParamRef SortedVertexBufferUAV,
 	FTexture2DRHIParamRef PositionTextureRHI,
@@ -351,7 +351,7 @@ int32 SortParticlesGPU(
 	// First generate keys for each emitter to be sorted.
 	//@todo-rco: RHIPacketList
 	const int32 TotalParticleCount = GenerateParticleSortKeys(
-		nullptr,
+		FRHICommandList::GetNullRef(),
 		ParticleSortBuffers.GetKeyBufferUAV(),
 		ParticleSortBuffers.GetVertexBufferUAV(),
 		PositionTextureRHI,
diff --git a/Engine/Source/Runtime/Engine/Private/Particles/ParticleVertexFactory.cpp b/Engine/Source/Runtime/Engine/Private/Particles/ParticleVertexFactory.cpp
--- a/Engine/Source/Runtime/Engine/Private/Particles/ParticleVertexFactory.cpp
+++ b/Engine/Source/Runtime/Engine/Private/Particles/ParticleVertexFactory.cpp
@@ -33,7 +33,7 @@ class FParticleSpriteVertexFactoryShaderParametersVS : public FParticleSpriteVer
 {
 public:
 
-	virtual void SetMesh(FRHICommandList* RHICmdList, FShader* Shader,const FVertexFactory* VertexFactory,const FSceneView& View,const FMeshBatchElement& BatchElement,uint32 DataFlags) const OVERRIDE
+	virtual void SetMesh(FRHICommandList& RHICmdList, FShader* Shader,const FVertexFactory* VertexFactory,const FSceneView& View,const FMeshBatchElement& BatchElement,uint32 DataFlags) const OVERRIDE
 	{
 		FParticleSpriteVertexFactory* SpriteVF = (FParticleSpriteVertexFactory*)VertexFactory;
 		FVertexShaderRHIParamRef VertexShaderRHI = Shader->GetVertexShader();
@@ -45,7 +45,7 @@ class FParticleSpriteVertexFactoryShaderParametersPS : public FParticleSpriteVer
 {
 public:
 
-	virtual void SetMesh(FRHICommandList* RHICmdList, FShader* Shader,const FVertexFactory* VertexFactory,const FSceneView& View,const FMeshBatchElement& BatchElement,uint32 DataFlags) const OVERRIDE
+	virtual void SetMesh(FRHICommandList& RHICmdList, FShader* Shader,const FVertexFactory* VertexFactory,const FSceneView& View,const FMeshBatchElement& BatchElement,uint32 DataFlags) const OVERRIDE
 	{
 		FParticleSpriteVertexFactory* SpriteVF = (FParticleSpriteVertexFactory*)VertexFactory;
 		FPixelShaderRHIParamRef PixelShaderRHI = Shader->GetPixelShader();
diff --git a/Engine/Source/Runtime/Engine/Private/ScreenRendering.cpp b/Engine/Source/Runtime/Engine/Private/ScreenRendering.cpp
--- a/Engine/Source/Runtime/Engine/Private/ScreenRendering.cpp
+++ b/Engine/Source/Runtime/Engine/Private/ScreenRendering.cpp
@@ -29,7 +29,7 @@ void DrawScreenQuad( float X0, float Y0, float U0, float V0, float X1, float Y1,
 	SetGlobalBoundShaderState(ScreenBoundShaderState, GScreenVertexDeclaration.VertexDeclarationRHI, *ScreenVertexShader, *ScreenPixelShader);
 
 	//@todo-rco: RHIPacketList
-	ScreenPixelShader->SetParameters(nullptr, Texture);
+	ScreenPixelShader->SetParameters(FRHICommandList::GetNullRef(), Texture);
 
 	// Generate the vertices used
 	FScreenVertex Vertices[4];
@@ -65,7 +65,7 @@ void DrawNormalizedScreenQuad( float X, float Y, float U0, float V0, float SizeX
 	SetGlobalBoundShaderState(ScreenBoundShaderState, GScreenVertexDeclaration.VertexDeclarationRHI, *ScreenVertexShader, *ScreenPixelShader);
 
 	/// ?
-	ScreenPixelShader->SetParameters(nullptr, TStaticSamplerState<SF_Bilinear>::GetRHI(), TextureRHI);
+	ScreenPixelShader->SetParameters(FRHICommandList::GetNullRef(), TStaticSamplerState<SF_Bilinear>::GetRHI(), TextureRHI);
 
 	// Generate the vertices used
 	FScreenVertex Vertices[4];
diff --git a/Engine/Source/Runtime/Engine/Private/SimpleElementShaders.cpp b/Engine/Source/Runtime/Engine/Private/SimpleElementShaders.cpp
--- a/Engine/Source/Runtime/Engine/Private/SimpleElementShaders.cpp
+++ b/Engine/Source/Runtime/Engine/Private/SimpleElementShaders.cpp
@@ -19,7 +19,7 @@ FSimpleElementVS::FSimpleElementVS(const ShaderMetaType::CompiledShaderInitializ
 	SwitchVerticalAxis.Bind(Initializer.ParameterMap,TEXT("SwitchVerticalAxis"), SPF_Optional);
 }
 
-void FSimpleElementVS::SetParameters(FRHICommandList* RHICmdList, const FMatrix& TransformValue, bool bSwitchVerticalAxis)
+void FSimpleElementVS::SetParameters(FRHICommandList& RHICmdList, const FMatrix& TransformValue, bool bSwitchVerticalAxis)
 {
 	SetShaderValue(RHICmdList, GetVertexShader(),Transform,TransformValue);
 	SetShaderValue(RHICmdList, GetVertexShader(),SwitchVerticalAxis, bSwitchVerticalAxis ? -1.0f : 1.0f);
@@ -54,7 +54,7 @@ FSimpleElementPS::FSimpleElementPS(const ShaderMetaType::CompiledShaderInitializ
 	ScreenToPixel.Bind(Initializer.ParameterMap,TEXT("ScreenToPixel"));
 }
 
-void FSimpleElementPS::SetEditorCompositingParameters(FRHICommandList* RHICmdList, const FSceneView* View, FTexture2DRHIRef DepthTexture )
+void FSimpleElementPS::SetEditorCompositingParameters(FRHICommandList& RHICmdList, const FSceneView* View, FTexture2DRHIRef DepthTexture )
 {
 	if( View )
 	{
@@ -85,7 +85,7 @@ void FSimpleElementPS::SetEditorCompositingParameters(FRHICommandList* RHICmdLis
 	SetTextureParameter(RHICmdList, GetPixelShader(), SceneDepthTextureNonMS, IsValidRef(DepthTexture) ? (FTextureRHIRef)DepthTexture : GWhiteTexture->TextureRHI);
 }
 
-void FSimpleElementPS::SetParameters(FRHICommandList* RHICmdList, const FTexture* TextureValue)
+void FSimpleElementPS::SetParameters(FRHICommandList& RHICmdList, const FTexture* TextureValue)
 {
 	SetTextureParameter(RHICmdList, GetPixelShader(),InTexture,InTextureSampler,TextureValue);
 	
@@ -112,7 +112,7 @@ FSimpleElementGammaPS::FSimpleElementGammaPS(const ShaderMetaType::CompiledShade
 	Gamma.Bind(Initializer.ParameterMap,TEXT("Gamma"));
 }
 
-void FSimpleElementGammaPS::SetParameters(FRHICommandList* RHICmdList, const FTexture* Texture,float GammaValue,ESimpleElementBlendMode BlendMode)
+void FSimpleElementGammaPS::SetParameters(FRHICommandList& RHICmdList, const FTexture* Texture,float GammaValue,ESimpleElementBlendMode BlendMode)
 {
 	FSimpleElementPS::SetParameters(RHICmdList, Texture);
 	SetShaderValue(RHICmdList, GetPixelShader(),Gamma,GammaValue);
@@ -131,7 +131,7 @@ FSimpleElementMaskedGammaPS::FSimpleElementMaskedGammaPS(const ShaderMetaType::C
 	ClipRef.Bind(Initializer.ParameterMap,TEXT("ClipRef"), SPF_Mandatory);
 }
 
-void FSimpleElementMaskedGammaPS::SetParameters(FRHICommandList* RHICmdList, const FTexture* Texture,float Gamma,float ClipRefValue,ESimpleElementBlendMode BlendMode)
+void FSimpleElementMaskedGammaPS::SetParameters(FRHICommandList& RHICmdList, const FTexture* Texture,float Gamma,float ClipRefValue,ESimpleElementBlendMode BlendMode)
 {
 	FSimpleElementGammaPS::SetParameters(RHICmdList, Texture,Gamma,BlendMode);
 	SetShaderValue(RHICmdList, GetPixelShader(),ClipRef,ClipRefValue);
@@ -177,7 +177,7 @@ FSimpleElementDistanceFieldGammaPS::FSimpleElementDistanceFieldGammaPS(const Sha
 * @param BlendMode - current batched element blend mode being rendered
 */
 void FSimpleElementDistanceFieldGammaPS::SetParameters(
-	FRHICommandList* RHICmdList, 
+	FRHICommandList& RHICmdList, 
 	const FTexture* Texture,
 	float Gamma,
 	float ClipRef,
@@ -239,7 +239,7 @@ FSimpleElementHitProxyPS::FSimpleElementHitProxyPS(const ShaderMetaType::Compile
 	InTextureSampler.Bind(Initializer.ParameterMap,TEXT("InTextureSampler"));
 }
 
-void FSimpleElementHitProxyPS::SetParameters(FRHICommandList* RHICmdList, const FTexture* TextureValue)
+void FSimpleElementHitProxyPS::SetParameters(FRHICommandList& RHICmdList, const FTexture* TextureValue)
 {
 	SetTextureParameter(RHICmdList, GetPixelShader(),InTexture,InTextureSampler,TextureValue);
 }
@@ -269,7 +269,7 @@ FSimpleElementColorChannelMaskPS::FSimpleElementColorChannelMaskPS(const ShaderM
 * @param ColorWeights - reference value to compare with alpha for killing pixels
 * @param Gamma - if gamma != 1.0 then a pow(color,Gamma) is applied
 */
-void FSimpleElementColorChannelMaskPS::SetParameters(FRHICommandList* RHICmdList, const FTexture* TextureValue, const FMatrix& ColorWeightsValue, float GammaValue)
+void FSimpleElementColorChannelMaskPS::SetParameters(FRHICommandList& RHICmdList, const FTexture* TextureValue, const FMatrix& ColorWeightsValue, float GammaValue)
 {
 	SetTextureParameter(RHICmdList, GetPixelShader(),InTexture,InTextureSampler,TextureValue);
 	SetShaderValue(RHICmdList, GetPixelShader(),ColorWeights,ColorWeightsValue);
diff --git a/Engine/Source/Runtime/Engine/Private/SplineMeshComponent.cpp b/Engine/Source/Runtime/Engine/Private/SplineMeshComponent.cpp
--- a/Engine/Source/Runtime/Engine/Private/SplineMeshComponent.cpp
+++ b/Engine/Source/Runtime/Engine/Private/SplineMeshComponent.cpp
@@ -34,7 +34,7 @@ void FSplineMeshVertexFactoryShaderParameters::Bind(const FShaderParameterMap& P
 	SplineMeshYParam.Bind(ParameterMap, TEXT("SplineMeshY"), SPF_Mandatory);
 }
 
-void FSplineMeshVertexFactoryShaderParameters::SetMesh(FRHICommandList* RHICmdList, FShader* Shader, const FVertexFactory* VertexFactory, const FSceneView& View, const FMeshBatchElement& BatchElement, uint32 DataFlags) const
+void FSplineMeshVertexFactoryShaderParameters::SetMesh(FRHICommandList& RHICmdList, FShader* Shader, const FVertexFactory* VertexFactory, const FSceneView& View, const FMeshBatchElement& BatchElement, uint32 DataFlags) const
 {
 	if (Shader->GetVertexShader())
 	{
diff --git a/Engine/Source/Runtime/Engine/Private/SplineMeshSceneProxy.h b/Engine/Source/Runtime/Engine/Private/SplineMeshSceneProxy.h
--- a/Engine/Source/Runtime/Engine/Private/SplineMeshSceneProxy.h
+++ b/Engine/Source/Runtime/Engine/Private/SplineMeshSceneProxy.h
@@ -69,7 +69,7 @@ class FSplineMeshVertexFactoryShaderParameters : public FVertexFactoryShaderPara
 {
 	void Bind(const FShaderParameterMap& ParameterMap) OVERRIDE;
 
-	void SetMesh(FRHICommandList* RHICmdList, FShader* Shader, const FVertexFactory* VertexFactory, const FSceneView& View, const FMeshBatchElement& BatchElement, uint32 DataFlags) const OVERRIDE;
+	void SetMesh(FRHICommandList& RHICmdList, FShader* Shader, const FVertexFactory* VertexFactory, const FSceneView& View, const FMeshBatchElement& BatchElement, uint32 DataFlags) const OVERRIDE;
 
 	void Serialize(FArchive& Ar) OVERRIDE
 	{
diff --git a/Engine/Source/Runtime/Engine/Private/VectorField.cpp b/Engine/Source/Runtime/Engine/Private/VectorField.cpp
--- a/Engine/Source/Runtime/Engine/Private/VectorField.cpp
+++ b/Engine/Source/Runtime/Engine/Private/VectorField.cpp
@@ -623,7 +623,7 @@ public:
 	 * @param NoiseVolumeTextureRHI - The volume texture to use to add noise to the vector field.
 	 */
 	void SetParameters(
-		FRHICommandList* RHICmdList, 
+		FRHICommandList& RHICmdList, 
 		const FCompositeAnimatedVectorFieldUniformBufferRef& UniformBuffer,
 		FTextureRHIParamRef AtlasTextureRHI,
 		FTextureRHIParamRef NoiseVolumeTextureRHI )
@@ -815,12 +815,12 @@ public:
 			CompositeCS->SetOutput(VolumeTextureUAV);
 			/// ?
 			CompositeCS->SetParameters(
-				nullptr,
+				FRHICommandList::GetNullRef(),
 				UniformBuffer,
 				AnimatedVectorField->Texture->Resource->TextureRHI,
 				NoiseVolumeTextureRHI );
 			DispatchComputeShader(
-				nullptr,
+				FRHICommandList::GetNullRef(),
 				*CompositeCS,
 				SizeX / THREADS_PER_AXIS,
 				SizeY / THREADS_PER_AXIS,
diff --git a/Engine/Source/Runtime/Engine/Private/VectorFieldVisualization.cpp b/Engine/Source/Runtime/Engine/Private/VectorFieldVisualization.cpp
--- a/Engine/Source/Runtime/Engine/Private/VectorFieldVisualization.cpp
+++ b/Engine/Source/Runtime/Engine/Private/VectorFieldVisualization.cpp
@@ -40,7 +40,7 @@ public:
 		Ar << VectorFieldTextureSampler;
 	}
 
-	virtual void SetMesh(FRHICommandList* RHICmdList, FShader* Shader,const FVertexFactory* VertexFactory,const FSceneView& View,const FMeshBatchElement& BatchElement,uint32 DataFlags) const OVERRIDE;
+	virtual void SetMesh(FRHICommandList& RHICmdList, FShader* Shader,const FVertexFactory* VertexFactory,const FSceneView& View,const FMeshBatchElement& BatchElement,uint32 DataFlags) const OVERRIDE;
 
 	virtual uint32 GetSize() const { return sizeof(*this); }
 
@@ -164,7 +164,7 @@ void FVectorFieldVisualizationVertexFactory::SetParameters(
 /**
  * Set vertex factory shader parameters.
  */
-void FVectorFieldVisualizationVertexFactoryShaderParameters::SetMesh(FRHICommandList* RHICmdList, FShader* Shader,const FVertexFactory* InVertexFactory,const FSceneView& View,const FMeshBatchElement& BatchElement,uint32 DataFlags) const
+void FVectorFieldVisualizationVertexFactoryShaderParameters::SetMesh(FRHICommandList& RHICmdList, FShader* Shader,const FVertexFactory* InVertexFactory,const FSceneView& View,const FMeshBatchElement& BatchElement,uint32 DataFlags) const
 {
 	FVectorFieldVisualizationVertexFactory* VertexFactory = (FVectorFieldVisualizationVertexFactory*)InVertexFactory;
 	FVertexShaderRHIParamRef VertexShaderRHI = Shader->GetVertexShader();
diff --git a/Engine/Source/Runtime/Engine/Public/BatchedElements.h b/Engine/Source/Runtime/Engine/Public/BatchedElements.h
--- a/Engine/Source/Runtime/Engine/Public/BatchedElements.h
+++ b/Engine/Source/Runtime/Engine/Public/BatchedElements.h
@@ -88,7 +88,7 @@ class FBatchedElementParameters
 public:
 
 	/** Binds vertex and pixel shaders for this element */
-	virtual void BindShaders_RenderThread(FRHICommandList* RHICmdList, const FMatrix& InTransform, const float InGamma, const FMatrix& ColorWeights, const FTexture* Texture ) = 0;
+	virtual void BindShaders_RenderThread(FRHICommandList& RHICmdList, const FMatrix& InTransform, const float InGamma, const FMatrix& ColorWeights, const FTexture* Texture ) = 0;
 
 };
 
diff --git a/Engine/Source/Runtime/Engine/Public/ClearReplacementShaders.h b/Engine/Source/Runtime/Engine/Public/ClearReplacementShaders.h
--- a/Engine/Source/Runtime/Engine/Public/ClearReplacementShaders.h
+++ b/Engine/Source/Runtime/Engine/Public/ClearReplacementShaders.h
@@ -83,7 +83,7 @@ public:
 		return bShaderHasOutdatedParameters;
 	}
 
-	void SetParameters(FRHICommandList* RHICmdList, FUnorderedAccessViewRHIParamRef TextureRW, FLinearColor Value);
+	void SetParameters(FRHICommandList& RHICmdList, FUnorderedAccessViewRHIParamRef TextureRW, FLinearColor Value);
 
 	static bool ShouldCache(EShaderPlatform Platform)
 	{
@@ -107,7 +107,7 @@ public:
 		ClearBufferRW.Bind(Initializer.ParameterMap, TEXT("ClearBufferRW"), SPF_Mandatory);
 	}
 
-	void SetParameters(FRHICommandList* RHICmdList, FUnorderedAccessViewRHIParamRef BufferRW, uint32 Dword);
+	void SetParameters(FRHICommandList& RHICmdList, FUnorderedAccessViewRHIParamRef BufferRW, uint32 Dword);
 
 	// FShader interface.
 	virtual bool Serialize(FArchive& Ar)
diff --git a/Engine/Source/Runtime/Engine/Public/FogRendering.h b/Engine/Source/Runtime/Engine/Public/FogRendering.h
--- a/Engine/Source/Runtime/Engine/Public/FogRendering.h
+++ b/Engine/Source/Runtime/Engine/Public/FogRendering.h
@@ -18,7 +18,7 @@ public:
 	 * Sets exponential height fog shader parameters.
 	 */
 	template<typename ShaderRHIParamRef>
-	void Set(FRHICommandList* RHICmdList, ShaderRHIParamRef Shader, const FSceneView* View) const
+	void Set(FRHICommandList& RHICmdList, ShaderRHIParamRef Shader, const FSceneView* View) const
 	{
 		const FViewInfo* ViewInfo = static_cast<const FViewInfo*>(View);
 		SetShaderValue(RHICmdList, Shader, ExponentialFogParameters, ViewInfo->ExponentialFogParameters);
@@ -50,7 +50,7 @@ public:
 	 * Sets height fog shader parameters.
 	 */
 	template<typename ShaderRHIParamRef>
-	void Set(FRHICommandList* RHICmdList, ShaderRHIParamRef Shader, const FSceneView* View) const
+	void Set(FRHICommandList& RHICmdList, ShaderRHIParamRef Shader, const FSceneView* View) const
 	{
 		// Set the fog constants.
 		ExponentialParameters.Set(RHICmdList, Shader, View);
diff --git a/Engine/Source/Runtime/Engine/Public/GlobalShader.h b/Engine/Source/Runtime/Engine/Public/GlobalShader.h
--- a/Engine/Source/Runtime/Engine/Public/GlobalShader.h
+++ b/Engine/Source/Runtime/Engine/Public/GlobalShader.h
@@ -118,7 +118,7 @@ public:
 	ENGINE_API FGlobalShader(const ShaderMetaType::CompiledShaderInitializerType& Initializer);
 	
 	template<typename ShaderRHIParamRef>
-	void SetParameters(FRHICommandList* RHICmdList, const ShaderRHIParamRef ShaderRHI,const FSceneView& View)
+	void SetParameters(FRHICommandList& RHICmdList, const ShaderRHIParamRef ShaderRHI,const FSceneView& View)
 	{
 		check(GetUniformBufferParameter<FViewUniformShaderParameters>().IsInitialized());
 		CheckShaderIsValid();
diff --git a/Engine/Source/Runtime/Engine/Public/Landscape/LandscapeRender.h b/Engine/Source/Runtime/Engine/Public/Landscape/LandscapeRender.h
--- a/Engine/Source/Runtime/Engine/Public/Landscape/LandscapeRender.h
+++ b/Engine/Source/Runtime/Engine/Public/Landscape/LandscapeRender.h
@@ -112,7 +112,7 @@ public:
 	/**
 	* Set any shader data specific to this vertex factory
 	*/
-	virtual void SetMesh(FRHICommandList* RHICmdList, FShader* PixelShader,const class FVertexFactory* VertexFactory,const class FSceneView& View,const struct FMeshBatchElement& BatchElement,uint32 DataFlags) const OVERRIDE;
+	virtual void SetMesh(FRHICommandList& RHICmdList, FShader* PixelShader,const class FVertexFactory* VertexFactory,const class FSceneView& View,const struct FMeshBatchElement& BatchElement,uint32 DataFlags) const OVERRIDE;
 	
 	virtual uint32 GetSize() const OVERRIDE
 	{
diff --git a/Engine/Source/Runtime/Engine/Public/ScreenRendering.h b/Engine/Source/Runtime/Engine/Public/ScreenRendering.h
--- a/Engine/Source/Runtime/Engine/Public/ScreenRendering.h
+++ b/Engine/Source/Runtime/Engine/Public/ScreenRendering.h
@@ -60,12 +60,12 @@ public:
 	}
 	FScreenPS() {}
 
-	void SetParameters(FRHICommandList* RHICmdList, const FTexture* Texture)
+	void SetParameters(FRHICommandList& RHICmdList, const FTexture* Texture)
 	{
 		SetTextureParameter(RHICmdList, GetPixelShader(),InTexture,InTextureSampler,Texture);
 	}
 
-	void SetParameters(FRHICommandList* RHICmdList, FSamplerStateRHIParamRef SamplerStateRHI, FTextureRHIParamRef TextureRHI)
+	void SetParameters(FRHICommandList& RHICmdList, FSamplerStateRHIParamRef SamplerStateRHI, FTextureRHIParamRef TextureRHI)
 	{
 		SetTextureParameter(RHICmdList, GetPixelShader(),InTexture,InTextureSampler,SamplerStateRHI,TextureRHI);
 	}
@@ -100,7 +100,7 @@ public:
 	FScreenVS() {}
 
 
-	void SetParameters(FRHICommandList* RHICmdList, const FSceneView& View)
+	void SetParameters(FRHICommandList& RHICmdList, const FSceneView& View)
 	{
 		FGlobalShader::SetParameters(RHICmdList, GetVertexShader(), View);
 	}
diff --git a/Engine/Source/Runtime/Engine/Public/SimpleElementShaders.h b/Engine/Source/Runtime/Engine/Public/SimpleElementShaders.h
--- a/Engine/Source/Runtime/Engine/Public/SimpleElementShaders.h
+++ b/Engine/Source/Runtime/Engine/Public/SimpleElementShaders.h
@@ -22,7 +22,7 @@ public:
 	FSimpleElementVS(const ShaderMetaType::CompiledShaderInitializerType& Initializer);
 	FSimpleElementVS() {}
 
-	/*ENGINE_API */void SetParameters(FRHICommandList* RHICmdList, const FMatrix& TransformValue, bool bSwitchVerticalAxis = false);
+	/*ENGINE_API */void SetParameters(FRHICommandList& RHICmdList, const FMatrix& TransformValue, bool bSwitchVerticalAxis = false);
 
 	virtual bool Serialize(FArchive& Ar);
 
@@ -53,9 +53,9 @@ public:
 	 * @param View			SceneView for view constants when compositing
 	 * @param DepthTexture	Depth texture to read from when depth testing for compositing.  If not set no compositing will occur
 	 */
-	void SetEditorCompositingParameters(FRHICommandList* RHICmdList, const FSceneView* View, FTexture2DRHIRef DepthTexture );
+	void SetEditorCompositingParameters(FRHICommandList& RHICmdList, const FSceneView* View, FTexture2DRHIRef DepthTexture );
 
-	void SetParameters(FRHICommandList* RHICmdList, const FTexture* TextureValue );
+	void SetParameters(FRHICommandList& RHICmdList, const FTexture* TextureValue );
 
 	virtual bool Serialize(FArchive& Ar);
 
@@ -82,7 +82,7 @@ public:
 	FSimpleElementGammaPS(const ShaderMetaType::CompiledShaderInitializerType& Initializer);
 	FSimpleElementGammaPS() {}
 
-	void SetParameters(FRHICommandList* RHICmdList, const FTexture* Texture,float GammaValue,ESimpleElementBlendMode BlendMode);
+	void SetParameters(FRHICommandList& RHICmdList, const FTexture* Texture,float GammaValue,ESimpleElementBlendMode BlendMode);
 
 	virtual bool Serialize(FArchive& Ar);
 
@@ -103,7 +103,7 @@ public:
 	FSimpleElementMaskedGammaPS(const ShaderMetaType::CompiledShaderInitializerType& Initializer);
 	FSimpleElementMaskedGammaPS() {}
 
-	void SetParameters(FRHICommandList* RHICmdList, const FTexture* Texture,float Gamma,float ClipRefValue,ESimpleElementBlendMode BlendMode);
+	void SetParameters(FRHICommandList& RHICmdList, const FTexture* Texture,float Gamma,float ClipRefValue,ESimpleElementBlendMode BlendMode);
 
 	virtual bool Serialize(FArchive& Ar);
 
@@ -156,7 +156,7 @@ public:
 	* @param BlendMode - current batched element blend mode being rendered
 	*/
 	void SetParameters(
-		FRHICommandList* RHICmdList, 
+		FRHICommandList& RHICmdList, 
 		const FTexture* Texture,
 		float Gamma,
 		float ClipRef,
@@ -215,7 +215,7 @@ public:
 	FSimpleElementHitProxyPS(const ShaderMetaType::CompiledShaderInitializerType& Initializer);
 	FSimpleElementHitProxyPS() {}
 
-	void SetParameters(FRHICommandList* RHICmdList, const FTexture* TextureValue);
+	void SetParameters(FRHICommandList& RHICmdList, const FTexture* TextureValue);
 
 	virtual bool Serialize(FArchive& Ar);
 
@@ -251,7 +251,7 @@ public:
 	* @param ColorWeights - reference value to compare with alpha for killing pixels
 	* @param Gamma - if gamma != 1.0 then a pow(color,Gamma) is applied
 	*/
-	void SetParameters(FRHICommandList* RHICmdList, const FTexture* TextureValue, const FMatrix& ColorWeightsValue, float GammaValue);
+	void SetParameters(FRHICommandList& RHICmdList, const FTexture* TextureValue, const FMatrix& ColorWeightsValue, float GammaValue);
 
 	virtual bool Serialize(FArchive& Ar);
 
diff --git a/Engine/Source/Runtime/RHI/Private/RHICommandList.cpp b/Engine/Source/Runtime/RHI/Private/RHICommandList.cpp
--- a/Engine/Source/Runtime/RHI/Private/RHICommandList.cpp
+++ b/Engine/Source/Runtime/RHI/Private/RHICommandList.cpp
@@ -46,6 +46,8 @@ static FAutoConsoleVariableRef CVarRHICmdMem(
 	);
 
 RHI_API FRHICommandListExecutor GRHICommandList;
+RHI_API FRHICommandList FRHICommandList::NullRHICommandList(0);
+
 
 //@todo-rco: Temp hack!
 static FRHICommandList GCommandList;
@@ -56,9 +58,9 @@ FRHICommandListExecutor::FRHICommandListExecutor() :
 {
 }
 
-void FRHICommandListExecutor::ExecuteAndFreeList(FRHICommandList* CmdList)
+void FRHICommandListExecutor::ExecuteAndFreeList(FRHICommandList& CmdList)
 {
-	if (!CmdList)
+	if (CmdList.IsNull())
 	{
 		// HACK!
 		return;
@@ -66,17 +68,16 @@ void FRHICommandListExecutor::ExecuteAndFreeList(FRHICommandList* CmdList)
 
 	SCOPE_CYCLE_COUNTER(STAT_RHICmdListExecuteTime);
 
-	check(CmdList);
-	check(CmdList->CanAddCommand());
+	check(CmdList.CanAddCommand());
 
 	FScopeLock Lock(&CriticalSection);
 	check(bLock);
 	bLock = false;
 
-	CmdList->State = FRHICommandList::Executing;
+	CmdList.State = FRHICommandList::Executing;
 
-	auto* CmdPtr = &CmdList->Mem[0];
-	auto* CmdTail = CmdList->GetTail();
+	auto* CmdPtr = CmdList.Memory;
+	auto* CmdTail = CmdList.GetTail();
 	INC_MEMORY_STAT_BY(STAT_RHICmdListMemory, (CmdTail - CmdPtr));
 
 	static IConsoleVariable* RHICmdListCVar = IConsoleManager::Get().FindConsoleVariable(TEXT("r.RHISkip"));
@@ -87,7 +88,7 @@ void FRHICommandListExecutor::ExecuteAndFreeList(FRHICommandList* CmdList)
 		static IConsoleVariable* RHIExecCVar = IConsoleManager::Get().FindConsoleVariable(TEXT("r.RHIExec"));
 		if (RHIExecCVar && RHIExecCVar->GetInt() != 0)
 		{
-			RHIExecuteCommandList(CmdList);
+			RHIExecuteCommandList(&CmdList);
 		}
 		else
 		{
@@ -262,15 +263,15 @@ void FRHICommandListExecutor::ExecuteAndFreeList(FRHICommandList* CmdList)
 				++NumCommands;
 			}
 
-			check(CmdList->NumCommands == NumCommands);
+			check(CmdList.NumCommands == NumCommands);
 		}
 	}
 
-	INC_DWORD_STAT_BY(STAT_RHICmdListCount, CmdList->NumCommands);
-	CmdList->Reset();
+	INC_DWORD_STAT_BY(STAT_RHICmdListCount, CmdList.NumCommands);
+	CmdList.Reset();
 }
 
-FRHICommandList* FRHICommandListExecutor::CreateList()
+FRHICommandList& FRHICommandListExecutor::CreateList()
 {
 	FScopeLock Lock(&CriticalSection);
 	check(!bLock);
@@ -278,7 +279,7 @@ FRHICommandList* FRHICommandListExecutor::CreateList()
 	check(GCommandList.State == FRHICommandList::Kicked);
 	check(GCommandList.NumCommands == 0);
 	GCommandList.State = FRHICommandList::Ready;
-	return &GCommandList;
+	return GCommandList;
 }
 
 void FRHICommandListExecutor::Verify()
diff --git a/Engine/Source/Runtime/RHI/Public/RHICommandList.h b/Engine/Source/Runtime/RHI/Public/RHICommandList.h
--- a/Engine/Source/Runtime/RHI/Public/RHICommandList.h
+++ b/Engine/Source/Runtime/RHI/Public/RHICommandList.h
@@ -24,211 +24,88 @@ enum ERHICommandType
 	ERCT_,
 };
 
-class FRHICommandList
+class RHI_API FRHICommandList
 {
+	static FRHICommandList NullRHICommandList;
 public:
 	enum
 	{
-		MemSize = 32 * 1024 * 1024,
-		Alignment = sizeof(void*),	// Use SIZE_T?
+		DefaultMemSize = 32 * 1024 * 1024,
+		Alignment = MIN_ALIGNMENT,
 	};
 
-	FRHICommandList()
+	static FORCEINLINE FRHICommandList& GetNullRef()
 	{
-		Reset();
+		return NullRHICommandList;
 	}
 
-	template <typename TShaderRHIParamRef>
-	FORCEINLINE static void SetShaderUniformBuffer(FRHICommandList* RHICmdList, TShaderRHIParamRef Shader, uint32 BaseIndex, FUniformBufferRHIParamRef UniformBufferRHI)
+	FRHICommandList(int32 Size = DefaultMemSize)
 	{
-		if (RHICmdList)
+		MemSize = Size;
+		Memory = nullptr;
+		if (Size)
 		{
-			RHICmdList->SetShaderUniformBuffer(Shader, BaseIndex, UniformBufferRHI);
+			Memory = (uint8*)FMemory::Malloc(MemSize, MIN_ALIGNMENT);
 		}
-		else
+		Reset();
+	}
+	~FRHICommandList()
+	{
+		if (!IsNull())
 		{
-			RHISetShaderUniformBuffer(Shader, BaseIndex, UniformBufferRHI);
+			FMemory::Free(Memory);
 		}
+		MemSize = 0;
+		Memory = (uint8*)0x1; // we want this to not bypass the cmd list; please just crash.
 	}
+
 	template <typename TShaderRHIParamRef>
 	inline void SetShaderUniformBuffer(TShaderRHIParamRef Shader, uint32 BaseIndex, FUniformBufferRHIParamRef UniformBufferRHI);
 
-	template <typename TShaderRHIParamRef>
-	FORCEINLINE static void SetShaderResourceViewParameter(FRHICommandList* RHICmdList, TShaderRHIParamRef Shader, uint32 SamplerIndex, FShaderResourceViewRHIParamRef SRV)
-	{
-		if (RHICmdList)
-		{
-			RHICmdList->SetShaderResourceViewParameter(Shader, SamplerIndex, SRV);
-		}
-		else
-		{
-			RHISetShaderResourceViewParameter(Shader, SamplerIndex, SRV);
-		}
-	}
 	template <typename TShaderRHIParamRef>
 	inline void SetShaderResourceViewParameter(TShaderRHIParamRef Shader, uint32 SamplerIndex, FShaderResourceViewRHIParamRef SRV);
 
-	template <typename TShaderRHIParamRef>
-	FORCEINLINE static void SetShaderParameter(FRHICommandList* RHICmdList, TShaderRHIParamRef Shader, uint32 BufferIndex, uint32 BaseIndex, uint32 NumBytes, const void* Value, bool bValueInStack = false)
-	{
-		if (RHICmdList)
-		{
-			RHICmdList->SetShaderParameter(Shader, BufferIndex, BaseIndex, NumBytes, Value, bValueInStack);
-		}
-		else
-		{
-			RHISetShaderParameter(Shader, BufferIndex, BaseIndex, NumBytes, Value);
-		}
-	}
 	template <typename TShaderRHIParamRef>
 	inline void SetShaderParameter(TShaderRHIParamRef Shader, uint32 BufferIndex, uint32 BaseIndex, uint32 NumBytes, const void* Value, bool bValueInStack = false);
 
-	template <typename TShaderRHIParamRef>
-	FORCEINLINE static void SetShaderTexture(FRHICommandList* RHICmdList, TShaderRHIParamRef Shader, uint32 TextureIndex, FTextureRHIParamRef Texture)
-	{
-		if (RHICmdList)
-		{
-			RHICmdList->SetShaderTexture(Shader, TextureIndex, Texture);
-		}
-		else
-		{
-			RHISetShaderTexture(Shader, TextureIndex, Texture);
-		}
-	}
 	template <typename TShaderRHIParamRef>
 	inline void SetShaderTexture(TShaderRHIParamRef Shader, uint32 TextureIndex, FTextureRHIParamRef Texture);
 
-	template <typename TShaderRHIParamRef>
-	FORCEINLINE static void SetShaderSampler(FRHICommandList* RHICmdList, TShaderRHIParamRef Shader, uint32 SamplerIndex, FSamplerStateRHIParamRef State)
-	{
-		if (RHICmdList)
-		{
-			RHICmdList->SetShaderSampler(Shader, SamplerIndex, State);
-		}
-		else
-		{
-			RHISetShaderSampler(Shader, SamplerIndex, State);
-		}
-	}
 	template <typename TShaderRHIParamRef>
 	inline void SetShaderSampler(TShaderRHIParamRef Shader, uint32 SamplerIndex, FSamplerStateRHIParamRef State);
 
-	FORCEINLINE static void SetUAVParameter(FRHICommandList* RHICmdList, FComputeShaderRHIParamRef Shader, uint32 UAVIndex, FUnorderedAccessViewRHIParamRef UAV)
-	{
-		if (RHICmdList)
-		{
-			RHICmdList->SetUAVParameter(Shader, UAVIndex, UAV);
-		}
-		else
-		{
-			RHISetUAVParameter(Shader, UAVIndex, UAV);
-		}
-	}
-	FORCEINLINE static void SetUAVParameter(FRHICommandList* RHICmdList, FComputeShaderRHIParamRef Shader, uint32 UAVIndex, FUnorderedAccessViewRHIParamRef UAV, uint32 InitialCount)
-	{
-		if (RHICmdList)
-		{
-			RHICmdList->SetUAVParameter(Shader, UAVIndex, UAV, InitialCount);
-		}
-		else
-		{
-			RHISetUAVParameter(Shader, UAVIndex, UAV, InitialCount);
-		}
-	}
 	inline void SetUAVParameter(FComputeShaderRHIParamRef Shader, uint32 UAVIndex, FUnorderedAccessViewRHIParamRef UAV);
 	inline void SetUAVParameter(FComputeShaderRHIParamRef Shader, uint32 UAVIndex, FUnorderedAccessViewRHIParamRef UAV, uint32 InitialCount);
 
-	FORCEINLINE static void SetBoundShaderState(FRHICommandList* RHICmdList, FBoundShaderStateRHIParamRef BoundShaderState)
-	{
-		if (RHICmdList)
-		{
-			RHICmdList->SetBoundShaderState(BoundShaderState);
-		}
-		else
-		{
-			RHISetBoundShaderState(BoundShaderState);
-		}
-	}
 	inline void SetBoundShaderState(FBoundShaderStateRHIParamRef BoundShaderState);
 
-	FORCEINLINE static void SetRasterizerState(FRHICommandList* RHICmdList, FRasterizerStateRHIParamRef State)
-	{
-		if (RHICmdList)
-		{
-			RHICmdList->SetRasterizerState(State);
-		}
-		else
-		{
-			RHISetRasterizerState(State);
-		}
-	}
 	inline void SetRasterizerState(FRasterizerStateRHIParamRef State);
 
-	FORCEINLINE static void SetBlendState(FRHICommandList* RHICmdList, FBlendStateRHIParamRef State, const FLinearColor& BlendFactor)
-	{
-		if (RHICmdList)
-		{
-			RHICmdList->SetBlendState(State, BlendFactor);
-		}
-		else
-		{
-			RHISetBlendState(State, BlendFactor);
-		}
-	}
-	FORCEINLINE static void SetBlendState(FRHICommandList* RHICmdList, FBlendStateRHIParamRef State)
-	{
-		SetBlendState(RHICmdList, State, FLinearColor::White);
-	}
 	inline void SetBlendState(FBlendStateRHIParamRef State, const FLinearColor& BlendFactor);
 	inline void SetBlendState(FBlendStateRHIParamRef State)
 	{
 		SetBlendState(State, FLinearColor::White);
 	}
 
-	FORCEINLINE static void DrawPrimitive(FRHICommandList* RHICmdList, uint32 PrimitiveType, uint32 BaseVertexIndex, uint32 NumPrimitives, uint32 NumInstances)
-	{
-		if (RHICmdList)
-		{
-			RHICmdList->DrawPrimitive(PrimitiveType, BaseVertexIndex, NumPrimitives, NumInstances);
-		}
-		else
-		{
-			RHIDrawPrimitive(PrimitiveType, BaseVertexIndex, NumPrimitives, NumInstances);
-		}
-	}
 	inline void DrawPrimitive(uint32 PrimitiveType, uint32 BaseVertexIndex, uint32 NumPrimitives, uint32 NumInstances);
+	inline void DrawIndexedPrimitive(FIndexBufferRHIParamRef IndexBuffer, uint32 PrimitiveType, int32 BaseVertexIndex, uint32 MinIndex, uint32 NumVertices, uint32 StartIndex, uint32 NumPrimitives, uint32 NumInstances);
+	inline void SetStreamSource(uint32 StreamIndex, FVertexBufferRHIParamRef VertexBuffer, uint32 Stride, uint32 Offset);
 
-	FORCEINLINE static void DrawIndexedPrimitive(FRHICommandList* RHICmdList, FIndexBufferRHIParamRef IndexBuffer, uint32 PrimitiveType, int32 BaseVertexIndex, uint32 MinIndex, uint32 NumVertices, uint32 StartIndex, uint32 NumPrimitives, uint32 NumInstances)
+	FORCEINLINE bool IsNull() const
 	{
-		if (RHICmdList)
-		{
-			RHICmdList->DrawIndexedPrimitive(IndexBuffer, PrimitiveType, BaseVertexIndex, MinIndex, NumVertices, StartIndex, NumPrimitives, NumInstances);
-		}
-		else
-		{
-			RHIDrawIndexedPrimitive(IndexBuffer, PrimitiveType, BaseVertexIndex, MinIndex, NumVertices, StartIndex, NumPrimitives, NumInstances);
-		}
+		return !Memory;
 	}
-	inline void DrawIndexedPrimitive(FIndexBufferRHIParamRef IndexBuffer, uint32 PrimitiveType, int32 BaseVertexIndex, uint32 MinIndex, uint32 NumVertices, uint32 StartIndex, uint32 NumPrimitives, uint32 NumInstances);
-
-	FORCEINLINE static void SetStreamSource(FRHICommandList* RHICmdList, uint32 StreamIndex, FVertexBufferRHIParamRef VertexBuffer, uint32 Stride, uint32 Offset)
+	FORCEINLINE void CheckIsNull() const
 	{
-		if (RHICmdList)
-		{
-			RHICmdList->SetStreamSource(StreamIndex, VertexBuffer, Stride, Offset);
-		}
-		else
-		{
-			RHISetStreamSource(StreamIndex, VertexBuffer, Stride, Offset);
-		}
+		check(IsNull());
 	}
-	inline void SetStreamSource(uint32 StreamIndex, FVertexBufferRHIParamRef VertexBuffer, uint32 Stride, uint32 Offset);
 
 
 	uint8* Alloc(SIZE_T InSize)
 	{
+		checkSlow(!IsNull());
 		InSize = (InSize + (Alignment - 1)) & ~(Alignment - 1);
-		check(Tail + InSize + Alignment <= &Mem[MemSize]);
+		check(Tail + InSize + Alignment <= &Memory[MemSize]);
 		uint8* New = Tail;
 		Tail += InSize;
 		return New;
@@ -236,16 +113,19 @@ public:
 
 	const uint8* GetHead() const
 	{
-		return &Mem[0];
+		checkSlow(!IsNull());
+		return &Memory[0];
 	}
 
 	const uint8* GetTail() const
 	{
+		checkSlow(!IsNull());
 		return Tail;
 	}
 
 protected:
-	uint8 Mem[MemSize];
+	uint8* Memory;
+	int32 MemSize;
 
 	enum EState
 	{
@@ -264,12 +144,13 @@ protected:
 
 	bool CanFitCommand(SIZE_T Size) const
 	{
-		return (Size + (Tail - Mem) <= MemSize);
+		return (Size + (Tail - Memory) <= MemSize);
 	}
 
 	template <typename TCmd>
 	FORCEINLINE TCmd* AddCommand()
 	{
+		checkSlow(!IsNull());
 		check(CanAddCommand());
 		check(CanFitCommand(sizeof(TCmd)));
 		TCmd* Cmd = (TCmd*)Alloc(sizeof(TCmd));
@@ -281,8 +162,12 @@ protected:
 	{
 		State = Kicked;
 		NumCommands = 0;
-		Tail = &Mem[0];
-		Tail = Align(Tail, Alignment);
+		Tail = Memory;
+	}
+
+	FORCEINLINE bool Bypass()
+	{
+		return IsNull();
 	}
 
 	friend class FRHICommandListExecutor;
@@ -292,8 +177,8 @@ class RHI_API FRHICommandListExecutor
 {
 public:
 	FRHICommandListExecutor();
-	FRHICommandList* CreateList();
-	void ExecuteAndFreeList(FRHICommandList* CmdList);
+	FRHICommandList& CreateList();
+	void ExecuteAndFreeList(FRHICommandList& CmdList);
 
 	void Verify();
 
diff --git a/Engine/Source/Runtime/RHI/Public/RHICommandList.inl b/Engine/Source/Runtime/RHI/Public/RHICommandList.inl
--- a/Engine/Source/Runtime/RHI/Public/RHICommandList.inl
+++ b/Engine/Source/Runtime/RHI/Public/RHICommandList.inl
@@ -198,6 +198,11 @@ struct FRHICommandSetStreamSource : public FRHICommand
 template <typename TShaderRHIParamRef>
 FORCEINLINE void FRHICommandList::SetShaderUniformBuffer(TShaderRHIParamRef Shader, uint32 BaseIndex, FUniformBufferRHIParamRef UniformBuffer)
 {
+	if (Bypass())
+	{
+		RHISetShaderUniformBuffer(Shader, BaseIndex, UniformBuffer);
+		return;
+	}
 	auto* Cmd = AddCommand<FRHICommandSetShaderUniformBuffer>();
 	Cmd->Set<TShaderRHIParamRef>(Shader, BaseIndex, UniformBuffer);
 	Shader->AddRef();
@@ -205,11 +210,25 @@ FORCEINLINE void FRHICommandList::SetShaderUniformBuffer(TShaderRHIParamRef Shad
 }
 
 template <typename TShaderRHIParamRef>
-FORCEINLINE void FRHICommandList::SetShaderResourceViewParameter(TShaderRHIParamRef Shader, uint32 SamplerIndex, FShaderResourceViewRHIParamRef SRV){check(0);}
+FORCEINLINE void FRHICommandList::SetShaderResourceViewParameter(TShaderRHIParamRef Shader, uint32 SamplerIndex, FShaderResourceViewRHIParamRef SRV)
+{
+	if (Bypass())
+	{
+		RHISetShaderResourceViewParameter(Shader, SamplerIndex, SRV);
+		return;
+	}
+	check(0);
+}
 
 template <typename TShaderRHIParamRef>
 FORCEINLINE void FRHICommandList::SetShaderParameter(TShaderRHIParamRef Shader, uint32 BufferIndex, uint32 BaseIndex, uint32 NumBytes, const void* NewValue, bool bValueInStack)
 {
+	if (Bypass())
+	{
+		RHISetShaderParameter(Shader, BufferIndex, BaseIndex, NumBytes, NewValue);
+		return;
+	}
+
 	if (bValueInStack)
 	{
 		auto* Blob = AddCommand<FRHICommandNopBlob>();
@@ -224,6 +243,12 @@ FORCEINLINE void FRHICommandList::SetShaderParameter(TShaderRHIParamRef Shader,
 template <typename TShaderRHIParamRef>
 FORCEINLINE void FRHICommandList::SetShaderTexture(TShaderRHIParamRef Shader, uint32 TextureIndex, FTextureRHIParamRef Texture)
 {
+	if (Bypass())
+	{
+		RHISetShaderTexture(Shader, TextureIndex, Texture);
+		return;
+	}
+
 	auto* Cmd = AddCommand<FRHICommandSetShaderTexture>();
 	Cmd->Set<TShaderRHIParamRef>(Shader, TextureIndex, Texture);
 	Shader->AddRef();
@@ -233,17 +258,46 @@ FORCEINLINE void FRHICommandList::SetShaderTexture(TShaderRHIParamRef Shader, ui
 template <typename TShaderRHIParamRef>
 FORCEINLINE void FRHICommandList::SetShaderSampler(TShaderRHIParamRef Shader, uint32 SamplerIndex, FSamplerStateRHIParamRef State)
 {
+	if (Bypass())
+	{
+		RHISetShaderSampler(Shader, SamplerIndex, State);
+		return;
+	}
+
 	auto* Cmd = AddCommand<FRHICommandSetShaderSampler>();
 	Cmd->Set<TShaderRHIParamRef>(Shader, SamplerIndex, State);
 	Shader->AddRef();
 	State->AddRef();
 }
 
-FORCEINLINE void FRHICommandList::SetUAVParameter(FComputeShaderRHIParamRef Shader, uint32 UAVIndex, FUnorderedAccessViewRHIParamRef UAV){check(0);}
-FORCEINLINE void FRHICommandList::SetUAVParameter(FComputeShaderRHIParamRef Shader, uint32 UAVIndex, FUnorderedAccessViewRHIParamRef UAV, uint32 InitialCount){check(0);}
+FORCEINLINE void FRHICommandList::SetUAVParameter(FComputeShaderRHIParamRef Shader, uint32 UAVIndex, FUnorderedAccessViewRHIParamRef UAV)
+{
+	if (Bypass())
+	{
+		RHISetUAVParameter(Shader, UAVIndex, UAV);
+		return;
+	}
+
+	check(0);
+}
+FORCEINLINE void FRHICommandList::SetUAVParameter(FComputeShaderRHIParamRef Shader, uint32 UAVIndex, FUnorderedAccessViewRHIParamRef UAV, uint32 InitialCount)
+{
+	if (Bypass())
+	{
+		RHISetUAVParameter(Shader, UAVIndex, UAV, InitialCount);
+		return;
+	}
+	check(0);
+}
 
 FORCEINLINE void FRHICommandList::SetRasterizerState(FRasterizerStateRHIParamRef State)
 {
+	if (Bypass())
+	{
+		RHISetRasterizerState(State);
+		return;
+	}
+
 	auto* Cmd = AddCommand<FRHICommandSetRasterizerState>();
 	Cmd->Set(State);
 	State->AddRef();
@@ -251,12 +305,24 @@ FORCEINLINE void FRHICommandList::SetRasterizerState(FRasterizerStateRHIParamRef
 
 FORCEINLINE void FRHICommandList::DrawPrimitive(uint32 PrimitiveType, uint32 BaseVertexIndex, uint32 NumPrimitives, uint32 NumInstances)
 {
+	if (Bypass())
+	{
+		RHIDrawPrimitive(PrimitiveType, BaseVertexIndex, NumPrimitives, NumInstances);
+		return;
+	}
+
 	auto* Cmd = AddCommand<FRHICommandDrawPrimitive>();
 	Cmd->Set(PrimitiveType, BaseVertexIndex, NumPrimitives, NumInstances);
 }
 
 FORCEINLINE void FRHICommandList::DrawIndexedPrimitive(FIndexBufferRHIParamRef IndexBuffer, uint32 PrimitiveType, int32 BaseVertexIndex, uint32 MinIndex, uint32 NumVertices, uint32 StartIndex, uint32 NumPrimitives, uint32 NumInstances)
 {
+	if (Bypass())
+	{
+		RHIDrawIndexedPrimitive(IndexBuffer, PrimitiveType, BaseVertexIndex, MinIndex, NumVertices, StartIndex, NumPrimitives, NumInstances);
+		return;
+	}
+
 	auto* Cmd = AddCommand<FRHICommandDrawIndexedPrimitive>();
 	Cmd->Set(IndexBuffer, PrimitiveType, BaseVertexIndex, MinIndex, NumVertices, StartIndex, NumPrimitives, NumInstances);
 	IndexBuffer->AddRef();
@@ -264,6 +330,12 @@ FORCEINLINE void FRHICommandList::DrawIndexedPrimitive(FIndexBufferRHIParamRef I
 
 FORCEINLINE void FRHICommandList::SetBoundShaderState(FBoundShaderStateRHIParamRef BoundShaderState)
 {
+	if (Bypass())
+	{
+		RHISetBoundShaderState(BoundShaderState);
+		return;
+	}
+
 	auto* Cmd = AddCommand<FRHICommandSetBoundShaderState>();
 	Cmd->Set(BoundShaderState);
 	BoundShaderState->AddRef();
@@ -271,6 +343,12 @@ FORCEINLINE void FRHICommandList::SetBoundShaderState(FBoundShaderStateRHIParamR
 
 FORCEINLINE void FRHICommandList::SetBlendState(FBlendStateRHIParamRef State, const FLinearColor& BlendFactor)
 {
+	if (Bypass())
+	{
+		RHISetBlendState(State, BlendFactor);
+		return;
+	}
+
 	auto* Cmd = AddCommand<FRHICommandSetBlendState>();
 	Cmd->Set(State, BlendFactor);
 	State->AddRef();
@@ -278,6 +356,12 @@ FORCEINLINE void FRHICommandList::SetBlendState(FBlendStateRHIParamRef State, co
 
 FORCEINLINE void FRHICommandList::SetStreamSource(uint32 StreamIndex, FVertexBufferRHIParamRef VertexBuffer, uint32 Stride, uint32 Offset)
 {
+	if (Bypass())
+	{
+		RHISetStreamSource(StreamIndex, VertexBuffer, Stride, Offset);
+		return;
+	}
+
 	auto* Cmd = AddCommand<FRHICommandSetStreamSource>();
 	Cmd->Set(StreamIndex, VertexBuffer, Stride, Offset);
 	VertexBuffer->AddRef();
