338163 SetWorldScale3D incorrectly using InverseTransformVector

diff --git a/Engine/Source/Runtime/Engine/Private/SceneComponent.cpp b/Engine/Source/Runtime/Engine/Private/SceneComponent.cpp
--- a/Engine/Source/Runtime/Engine/Private/SceneComponent.cpp
+++ b/Engine/Source/Runtime/Engine/Private/SceneComponent.cpp
@@ -431,7 +431,7 @@ void USceneComponent::SetWorldScale3D(FVector NewScale)
 	if(AttachParent != NULL && !bAbsoluteScale)
 	{
 		FTransform ParentToWorld = AttachParent->GetSocketTransform(AttachSocketName);
-		NewRelScale = ParentToWorld.InverseTransformVector(NewScale);
+		NewRelScale = NewScale * ParentToWorld.GetSafeScaleReciprocal(ParentToWorld.GetScale3D());
 	}
 
 	SetRelativeScale3D(NewRelScale);
