Rollback //depot/UE4/Engine/Source/Runtime/Engine/Private/Tests/AutomationCommon.cpp to revision 2

diff --git a/Engine/Source/Runtime/Engine/Private/Tests/AutomationCommon.cpp b/Engine/Source/Runtime/Engine/Private/Tests/AutomationCommon.cpp
--- a/Engine/Source/Runtime/Engine/Private/Tests/AutomationCommon.cpp
+++ b/Engine/Source/Runtime/Engine/Private/Tests/AutomationCommon.cpp
@@ -4,7 +4,6 @@
 #include "Slate.h"
 #include "AutomationCommon.h"
 #include "ImageUtils.h"
-#include "UnrealEd.h"
 
 ///////////////////////////////////////////////////////////////////////
 // Common Latent commands
@@ -43,21 +42,5 @@ bool FTakeActiveEditorScreenshotCommand::Update()
 bool FTakeEditorScreenshotCommand::Update()
 {
 	AutomationCommon::SaveWindowAsScreenshot(ScreenshotParameters.CurrentWindow.ToSharedRef(), ScreenshotParameters.ScreenshotName);
-	return true;
-}
-
-bool UndoRedoCommand::Update()
-{
-	if (bUndo == true)
-	{
-		//Undo
-		GEditor->UndoTransaction();
-	}
-	else
-	{
-		//Redo
-		GEditor->RedoTransaction();
-	}
-
 	return true;
 }
\ No newline at end of file
