Rollback //depot/UE4/Engine/Source/Runtime/Engine/Public/Tests/AutomationCommon.h to revision 5

diff --git a/Engine/Source/Runtime/Engine/Public/Tests/AutomationCommon.h b/Engine/Source/Runtime/Engine/Public/Tests/AutomationCommon.h
--- a/Engine/Source/Runtime/Engine/Public/Tests/AutomationCommon.h
+++ b/Engine/Source/Runtime/Engine/Public/Tests/AutomationCommon.h
@@ -98,9 +98,3 @@ DEFINE_ENGINE_LATENT_AUTOMATION_COMMAND_ONE_PARAMETER(FTakeActiveEditorScreensho
  * Take a screenshot of the active window
  */
 DEFINE_ENGINE_LATENT_AUTOMATION_COMMAND_ONE_PARAMETER(FTakeEditorScreenshotCommand, WindowScreenshotParameters, ScreenshotParameters);
-
-/**
-* Creates a latent command which the user can either undo or redo an action.
-* True will trigger Undo, False will trigger a redo
-*/
-DEFINE_ENGINE_LATENT_AUTOMATION_COMMAND_ONE_PARAMETER(UndoRedoCommand, bool, bUndo);
\ No newline at end of file
