Add simple skinning to update per poly mesh vertices

diff --git a/Engine/Source/Runtime/Engine/Classes/PhysicsEngine/BodyInstance.h b/Engine/Source/Runtime/Engine/Classes/PhysicsEngine/BodyInstance.h
--- a/Engine/Source/Runtime/Engine/Classes/PhysicsEngine/BodyInstance.h
+++ b/Engine/Source/Runtime/Engine/Classes/PhysicsEngine/BodyInstance.h
@@ -279,6 +279,7 @@ public:
 	 */
 	bool UpdateBodyScale(const FVector& InScale3D);
 
+	void UpdateTriMeshVertices(const TArray<FVector> & NewPositions, const TArray<uint32> & Indices);
 
 	/** Returns the center of mass of this body (in world space) */
 	FVector GetCOMPosition() const;
diff --git a/Engine/Source/Runtime/Engine/Classes/PhysicsEngine/BodySetup.h b/Engine/Source/Runtime/Engine/Classes/PhysicsEngine/BodySetup.h
--- a/Engine/Source/Runtime/Engine/Classes/PhysicsEngine/BodySetup.h
+++ b/Engine/Source/Runtime/Engine/Classes/PhysicsEngine/BodySetup.h
@@ -534,6 +534,13 @@ public:
 	 * @return								true on success, false on failure because of vertex count overflow.
 	 */
 	ENGINE_API void CreateFromModel(class UModel* InModel, bool bRemoveExisting);
+
+	/**
+	 * Converts the skinned data of a skeletal mesh into a tri mesh collision. This is used for per poly scene queries and is quite expensive.
+	 * In 99% of cases you should be fine using a physics asset created for the skeletal mesh
+	 * @param	InSkeletalMeshComponent		The skeletal mesh component we'll be grabbing the skinning information from
+	 */
+	ENGINE_API void UpdateTriMeshVertices(const TArray<FVector> & NewPositions, const TArray<uint32> & Indices);
 #endif // WITH_EDITOR
 
 	/**
diff --git a/Engine/Source/Runtime/Engine/Private/PhysicsEngine/BodyInstance.cpp b/Engine/Source/Runtime/Engine/Private/PhysicsEngine/BodyInstance.cpp
--- a/Engine/Source/Runtime/Engine/Private/PhysicsEngine/BodyInstance.cpp
+++ b/Engine/Source/Runtime/Engine/Private/PhysicsEngine/BodyInstance.cpp
@@ -309,6 +309,33 @@ TArray<PxShape*> FBodyInstance::GetAllShapes(int32& OutNumSyncShapes) const
 }
 #endif
 
+void FBodyInstance::UpdateTriMeshVertices(const TArray<FVector> & NewPositions, const TArray<uint32> & Indices)
+{
+#if WITH_PHYSX
+	
+	SCOPED_SCENE_WRITE_LOCK((RigidActorSync ? RigidActorSync->getScene() : NULL));
+	SCOPED_SCENE_WRITE_LOCK((RigidActorAsync ? RigidActorAsync->getScene() : NULL));
+
+	BodySetup->UpdateTriMeshVertices(NewPositions, Indices);
+
+	//after updating the vertices we must call setGeometry again to update any shapes referencing the mesh
+
+	int32 SyncShapeCount = 0;
+	TArray<PxShape *> PShapes = GetAllShapes(SyncShapeCount);
+	PxTriangleMeshGeometry PTriangleMeshGeometry;
+	for (int32 ShapeIdx = 0; ShapeIdx < PShapes.Num(); ShapeIdx++)
+	{
+		PxShape* PShape = PShapes[ShapeIdx];
+		if (PShape->getGeometryType() == PxGeometryType::eTRIANGLEMESH)
+		{
+			PShape->getTriangleMeshGeometry(PTriangleMeshGeometry);
+			PShape->setGeometry(PTriangleMeshGeometry);
+		}
+	}
+	
+#endif
+}
+
 void FBodyInstance::UpdatePhysicalMaterials()
 {
 	UPhysicalMaterial* SimplePhysMat = GetSimplePhysicalMaterial();
@@ -1710,7 +1737,6 @@ bool FBodyInstance::UpdateBodyScale(const FVector& InScale3D)
 	return bSuccess;
 }
 
-
 void FBodyInstance::UpdateInstanceSimulatePhysics(bool bIgnoreParent)
 {
 	// In skeletal case, we need both our bone and skelcomponent flag to be true.
diff --git a/Engine/Source/Runtime/Engine/Private/PhysicsEngine/BodySetup.cpp b/Engine/Source/Runtime/Engine/Private/PhysicsEngine/BodySetup.cpp
--- a/Engine/Source/Runtime/Engine/Private/PhysicsEngine/BodySetup.cpp
+++ b/Engine/Source/Runtime/Engine/Private/PhysicsEngine/BodySetup.cpp
@@ -732,6 +732,49 @@ void UBodySetup::PostLoad()
 	}
 }
 
+void UBodySetup::UpdateTriMeshVertices(const TArray<FVector> & NewPositions, const TArray<uint32> & Indices)
+{
+#if WITH_PHYSX
+	if (TriMesh)
+	{
+		const bool bNeedsRemap = true;	//depending on if we clean verts during cooking there can be a 1-1 map or not. Need to benchmark this
+		if (bNeedsRemap)
+		{
+			const PxU32 * RuntimeToOriginal = TriMesh->getTrianglesRemap();
+			bool b16BitIndices = TriMesh->getTriangleMeshFlags() & PxTriangleMeshFlag::e16_BIT_INDICES;
+			const void * PIndices = TriMesh->getTriangles();
+
+			uint32 PNumTris = TriMesh->getNbTriangles();
+			PxVec3 * PNewPositions = TriMesh->getVerticesForModification();
+
+			for (uint32 TriIdx = 0; TriIdx < PNumTris; ++TriIdx)
+			{
+				PxU32 OriginalStartIndex = RuntimeToOriginal[TriIdx] * 3;
+				for (int i = 0; i < 3; ++i)
+				{
+					PxU32 PIndex = TriIdx * 3 + i;
+					PxU32 PVertexIndex = b16BitIndices ? ((PxU16*)PIndices)[PIndex] : ((PxU32*)PIndices)[PIndex];
+					PxVec3 & PRuntimeVertex = PNewPositions[PVertexIndex];
+
+					PxU32 OriginalVertexIndex = Indices[OriginalStartIndex + i];
+					const FVector & OriginalVertex = NewPositions[OriginalVertexIndex];
+					PRuntimeVertex = U2PVector(OriginalVertex);
+				}
+			}
+		}
+		else
+		{
+			PxVec3 * PNewPositions = TriMesh->getVerticesForModification();
+			for (int32 i = 0; i < NewPositions.Num(); ++i)
+			{
+				PNewPositions[i] = U2PVector(NewPositions[i]);
+			}
+		}
+
+		TriMesh->refitBVH();
+	}
+#endif
+}
 
 FByteBulkData* UBodySetup::GetCookedData(FName Format)
 {
diff --git a/Engine/Source/Runtime/Engine/Private/PhysicsEngine/PhysAnim.cpp b/Engine/Source/Runtime/Engine/Private/PhysicsEngine/PhysAnim.cpp
--- a/Engine/Source/Runtime/Engine/Private/PhysicsEngine/PhysAnim.cpp
+++ b/Engine/Source/Runtime/Engine/Private/PhysicsEngine/PhysAnim.cpp
@@ -6,6 +6,7 @@
 
 #include "EnginePrivate.h"
 #include "AnimTree.h"
+#include "SkeletalRenderPublic.h"
 
 #if WITH_PHYSX
 	#include "PhysXSupport.h"
@@ -363,6 +364,21 @@ void USkeletalMeshComponent::UpdateKinematicBonesToPhysics(bool bTeleport)
 	else
 	{
 		//per poly update requires us to update all vertex positions
+		if (MeshObject)
+		{
+			const FStaticLODModel& Model = MeshObject->GetSkeletalMeshResource().LODModels[0];
+			TArray<FVector> NewPositions;
+			NewPositions.AddUninitialized(Model.NumVertices);
+			for (uint32 VertIndex = 0; VertIndex < Model.NumVertices; ++VertIndex)
+			{
+				NewPositions[VertIndex] = GetSkinnedVertexPosition(VertIndex);
+			}
+
+			TArray<uint32> Indices;
+			Model.MultiSizeIndexContainer.GetIndexBuffer(Indices);
+			BodyInstance.UpdateTriMeshVertices(NewPositions, Indices);
+			BodyInstance.SetBodyTransform(CurrentLocalToWorld, bTeleport);
+		}
 	}
 
 
