Refactor SK::SetInstanceSimulatePhysics

diff --git a/Engine/Source/Runtime/Engine/Classes/PhysicsEngine/BodyInstance.h b/Engine/Source/Runtime/Engine/Classes/PhysicsEngine/BodyInstance.h
--- a/Engine/Source/Runtime/Engine/Classes/PhysicsEngine/BodyInstance.h
+++ b/Engine/Source/Runtime/Engine/Classes/PhysicsEngine/BodyInstance.h
@@ -329,7 +329,7 @@ public:
 
 
 	/** Set this body to be fixed (kinematic) or not. */
-	void SetInstanceSimulatePhysics(bool bSimulate, bool bMaintainPhysicsBlending=false, bool bIgnoreOwner=false);
+	void SetInstanceSimulatePhysics(bool bSimulate, bool bMaintainPhysicsBlending=false);
 	/** Makes sure the current kinematic state matches the simulate flag */
 	void UpdateInstanceSimulatePhysics(bool bIgnoreOwner = false);
 	/** Returns true if this body is simulating, false if it is fixed (kinematic) */
diff --git a/Engine/Source/Runtime/Engine/Private/PhysicsEngine/BodyInstance.cpp b/Engine/Source/Runtime/Engine/Private/PhysicsEngine/BodyInstance.cpp
--- a/Engine/Source/Runtime/Engine/Private/PhysicsEngine/BodyInstance.cpp
+++ b/Engine/Source/Runtime/Engine/Private/PhysicsEngine/BodyInstance.cpp
@@ -1794,7 +1794,7 @@ bool FBodyInstance::IsDynamic() const
 	return false;
 }
 
-void FBodyInstance::SetInstanceSimulatePhysics(bool bSimulate, bool bMaintainPhysicsBlending, bool bIgnoreOwner)
+void FBodyInstance::SetInstanceSimulatePhysics(bool bSimulate, bool bMaintainPhysicsBlending)
 {
 #if !(UE_BUILD_SHIPPING || UE_BUILD_TEST)
 	if (bSimulate && OwnerComponent.IsValid())
@@ -1831,7 +1831,7 @@ void FBodyInstance::SetInstanceSimulatePhysics(bool bSimulate, bool bMaintainPhy
 		}
 	}
 
-	UpdateInstanceSimulatePhysics(bIgnoreOwner);
+	UpdateInstanceSimulatePhysics(true);
 }
 
 bool FBodyInstance::IsInstanceSimulatingPhysics(bool bIgnoreOwner)
@@ -1842,16 +1842,10 @@ bool FBodyInstance::IsInstanceSimulatingPhysics(bool bIgnoreOwner)
 
 bool FBodyInstance::ShouldInstanceSimulatingPhysics(bool bIgnoreOwner)
 {
-	// if I'm simulating or owner is simulating
-	if (BodySetup.IsValid() && BodySetup.Get()->PhysicsType == PhysType_Default && !bIgnoreOwner)
+	//If type is set to default inherit whatever the parent does
+	if (OwnerComponent != NULL && BodySetup.IsValid() && BodySetup.Get()->PhysicsType == PhysType_Default && !bIgnoreOwner)
 	{
-		// if derive from owner, and owner is simulating, this should simulate
-		if (OwnerComponent != NULL && OwnerComponent->BodyInstance.bSimulatePhysics)
-		{
-			return true;
-		}
-
-		// or else, it should look its own setting
+		return OwnerComponent->BodyInstance.bSimulatePhysics;
 	}
 
 	return bSimulatePhysics;
diff --git a/Engine/Source/Runtime/Engine/Private/Vehicles/WheeledVehicleMovementComponent.cpp b/Engine/Source/Runtime/Engine/Private/Vehicles/WheeledVehicleMovementComponent.cpp
--- a/Engine/Source/Runtime/Engine/Private/Vehicles/WheeledVehicleMovementComponent.cpp
+++ b/Engine/Source/Runtime/Engine/Private/Vehicles/WheeledVehicleMovementComponent.cpp
@@ -1244,7 +1244,7 @@ void UWheeledVehicleMovementComponent::FixupSkeletalMesh()
 
 							if (BodySetup->PhysicsType == PhysType_Default) 	//if they set it to unfixed we don't fixup because they are explicitely saying Unfixed
 							{
-								BodyInstance->SetInstanceSimulatePhysics(false, false, true);
+								BodyInstance->SetInstanceSimulatePhysics(false);
 							}
 
 							//and get rid of constraints on the wheels. TODO: right now we remove all wheel constraints, we probably only want to remove parent constraints
