Fix for stale OwnedComponents array after reparting Blueprints.

#ttp 335387 - Reparenting blueprint twice causes assert/crash in RegisterComponentWithWorld

diff --git a/Engine/Source/Runtime/Engine/Classes/GameFramework/Actor.h b/Engine/Source/Runtime/Engine/Classes/GameFramework/Actor.h
--- a/Engine/Source/Runtime/Engine/Classes/GameFramework/Actor.h
+++ b/Engine/Source/Runtime/Engine/Classes/GameFramework/Actor.h
@@ -2046,6 +2046,7 @@ private:
 	TArray<UActorComponent*> OwnedComponents;
 
 	friend class FBlueprintEditorUtils;
+	friend class UEngine;
 	void ResetOwnedComponents();
 
 public:
diff --git a/Engine/Source/Runtime/Engine/Private/UnrealEngine.cpp b/Engine/Source/Runtime/Engine/Private/UnrealEngine.cpp
--- a/Engine/Source/Runtime/Engine/Private/UnrealEngine.cpp
+++ b/Engine/Source/Runtime/Engine/Private/UnrealEngine.cpp
@@ -9501,6 +9501,7 @@ void UEngine::CopyPropertiesForUnrelatedObjects(UObject* OldObject, UObject* New
 	if(NewActor != NULL)
 	{
 		NewActor->SetRootComponent(SavedRootComponent);
+		NewActor->ResetOwnedComponents();
 	}
 
 	bool bDumpProperties = CVarDumpCopyPropertiesForUnrelatedObjects.GetValueOnGameThread() != 0;
