TTP 337416: Rendering: Particles do not render when scaled negatively
We reverse backface culling if the local-to-world matrix determinant is <0 - that's fine for meshes, but for sprites that don't have a dimension in Z, that breaks if the Z axis is scaled negatively; that changes the matrix determinant, but not the winding order of the sprites (because we don't transform the mesh by LocalToWorld, but just use it for the size of the particles); instead use absolute values for sprite size, and ignore the determinant

diff --git a/Engine/Source/Runtime/Engine/Private/Particles/ParticleSystemRender.cpp b/Engine/Source/Runtime/Engine/Private/Particles/ParticleSystemRender.cpp
--- a/Engine/Source/Runtime/Engine/Private/Particles/ParticleSystemRender.cpp
+++ b/Engine/Source/Runtime/Engine/Private/Particles/ParticleSystemRender.cpp
@@ -821,8 +821,8 @@ void FDynamicSpriteEmitterData::Init( bool bInSelected )
 FVector2D GetParticleSize(const FBaseParticle& Particle, const FDynamicSpriteEmitterReplayDataBase& Source)
 {
 	FVector2D Size;
-	Size.X = Particle.Size.X * Source.Scale.X;
-	Size.Y = Particle.Size.Y * Source.Scale.Y;
+	Size.X = FMath::Abs(Particle.Size.X * Source.Scale.X);
+	Size.Y = FMath::Abs(Particle.Size.Y * Source.Scale.Y);
 	if (Source.ScreenAlignment == PSA_Square || Source.ScreenAlignment == PSA_FacingCameraPosition)
 	{
 		Size.Y = Size.X;
@@ -1400,7 +1400,6 @@ int32 FDynamicSpriteEmitterData::Render(FParticleSystemSceneProxy* Proxy, FPrimi
 			}
 			BatchElement.MinVertexIndex = 0;
 			BatchElement.MaxVertexIndex = (ParticleCount * 4) - 1;
-			Mesh.ReverseCulling = Proxy->IsLocalToWorldDeterminantNegative();
 			Mesh.CastShadow = Proxy->GetCastShadow();
 			Mesh.DepthPriorityGroup = (ESceneDepthPriorityGroup)Proxy->GetDepthPriorityGroup(View);
 
