Add profling info to per poly skinning
Skip skinning when updating transforms instead of animation positions

diff --git a/Engine/Source/Runtime/Engine/Classes/Components/SkeletalMeshComponent.h b/Engine/Source/Runtime/Engine/Classes/Components/SkeletalMeshComponent.h
--- a/Engine/Source/Runtime/Engine/Classes/Components/SkeletalMeshComponent.h
+++ b/Engine/Source/Runtime/Engine/Classes/Components/SkeletalMeshComponent.h
@@ -931,7 +931,7 @@ public:
 	 *	Iterate over each physics body in the physics for this mesh, and for each 'kinematic' (ie fixed or default if owner isn't simulating) one, update its
 	 *	transform based on the animated transform.
 	 */
-	void UpdateKinematicBonesToPhysics(bool bTeleport);
+	void UpdateKinematicBonesToPhysics(bool bTeleport, bool bNeedsSkinning);
 	
 	/**
 	 * Look up all bodies for broken constraints.
diff --git a/Engine/Source/Runtime/Engine/Private/PhysicsEngine/BodySetup.cpp b/Engine/Source/Runtime/Engine/Private/PhysicsEngine/BodySetup.cpp
--- a/Engine/Source/Runtime/Engine/Private/PhysicsEngine/BodySetup.cpp
+++ b/Engine/Source/Runtime/Engine/Private/PhysicsEngine/BodySetup.cpp
@@ -7,6 +7,7 @@
 #include "EnginePrivate.h"
 #include "PhysicsPublic.h"
 #include "TargetPlatform.h"
+#include "AnimTree.h"
 
 #if WITH_PHYSX
 	#include "PhysXSupport.h"
@@ -734,6 +735,7 @@ void UBodySetup::PostLoad()
 
 void UBodySetup::UpdateTriMeshVertices(const TArray<FVector> & NewPositions, const TArray<uint32> & Indices)
 {
+	SCOPE_CYCLE_COUNTER(STAT_UpdateTriMeshVertices);
 #if WITH_PHYSX
 	if (TriMesh)
 	{
diff --git a/Engine/Source/Runtime/Engine/Private/PhysicsEngine/PhysAnim.cpp b/Engine/Source/Runtime/Engine/Private/PhysicsEngine/PhysAnim.cpp
--- a/Engine/Source/Runtime/Engine/Private/PhysicsEngine/PhysAnim.cpp
+++ b/Engine/Source/Runtime/Engine/Private/PhysicsEngine/PhysAnim.cpp
@@ -241,7 +241,7 @@ void USkeletalMeshComponent::BlendInPhysics()
 }
 
 
-void USkeletalMeshComponent::UpdateKinematicBonesToPhysics(bool bTeleport)
+void USkeletalMeshComponent::UpdateKinematicBonesToPhysics(bool bTeleport, bool bNeedsSkinning)
 {
 	SCOPE_CYCLE_COUNTER(STAT_UpdateRBBones);
 
@@ -366,17 +366,24 @@ void USkeletalMeshComponent::UpdateKinematicBonesToPhysics(bool bTeleport)
 		//per poly update requires us to update all vertex positions
 		if (MeshObject)
 		{
-			const FStaticLODModel& Model = MeshObject->GetSkeletalMeshResource().LODModels[0];
-			TArray<FVector> NewPositions;
-			NewPositions.AddUninitialized(Model.NumVertices);
-			for (uint32 VertIndex = 0; VertIndex < Model.NumVertices; ++VertIndex)
+			if (bNeedsSkinning)
 			{
-				NewPositions[VertIndex] = GetSkinnedVertexPosition(VertIndex);
-			}
+				const FStaticLODModel& Model = MeshObject->GetSkeletalMeshResource().LODModels[0];
+				TArray<FVector> NewPositions;
+				NewPositions.AddUninitialized(Model.NumVertices);
+				{
+					SCOPE_CYCLE_COUNTER(STAT_SkinPerPolyVertices);
+					for (uint32 VertIndex = 0; VertIndex < Model.NumVertices; ++VertIndex)
+					{
+						NewPositions[VertIndex] = GetSkinnedVertexPosition(VertIndex);
+					}
+				}
 
-			TArray<uint32> Indices;
-			Model.MultiSizeIndexContainer.GetIndexBuffer(Indices);
-			BodyInstance.UpdateTriMeshVertices(NewPositions, Indices);
+				TArray<uint32> Indices;
+				Model.MultiSizeIndexContainer.GetIndexBuffer(Indices);
+				BodyInstance.UpdateTriMeshVertices(NewPositions, Indices);
+			}
+			
 			BodyInstance.SetBodyTransform(CurrentLocalToWorld, bTeleport);
 		}
 	}
diff --git a/Engine/Source/Runtime/Engine/Private/SkeletalMeshComponent.cpp b/Engine/Source/Runtime/Engine/Private/SkeletalMeshComponent.cpp
--- a/Engine/Source/Runtime/Engine/Private/SkeletalMeshComponent.cpp
+++ b/Engine/Source/Runtime/Engine/Private/SkeletalMeshComponent.cpp
@@ -974,7 +974,7 @@ void USkeletalMeshComponent::PostAnimEvaluation(bool bDoInterpolation, bool bDup
 		InvalidateCachedBounds();
 
 		// update physics data from animated data
-		UpdateKinematicBonesToPhysics(false);
+		UpdateKinematicBonesToPhysics(false, true);
 		UpdateRBJointMotors();
 	}
 
diff --git a/Engine/Source/Runtime/Engine/Private/SkeletalMeshComponentPhysics.cpp b/Engine/Source/Runtime/Engine/Private/SkeletalMeshComponentPhysics.cpp
--- a/Engine/Source/Runtime/Engine/Private/SkeletalMeshComponentPhysics.cpp
+++ b/Engine/Source/Runtime/Engine/Private/SkeletalMeshComponentPhysics.cpp
@@ -1285,7 +1285,7 @@ void USkeletalMeshComponent::OnUpdateTransform(bool bSkipPhysicsMove)
 	// Always send new transform to physics
 	if(bPhysicsStateCreated && !bSkipPhysicsMove )
 	{
-		UpdateKinematicBonesToPhysics(false);
+		UpdateKinematicBonesToPhysics(false, false);
 	}
 
 #if WITH_APEX_CLOTHING
