missing files

diff --git a/Engine/Source/Runtime/Engine/Private/Animation/AnimInstance.cpp b/Engine/Source/Runtime/Engine/Private/Animation/AnimInstance.cpp
--- a/Engine/Source/Runtime/Engine/Private/Animation/AnimInstance.cpp
+++ b/Engine/Source/Runtime/Engine/Private/Animation/AnimInstance.cpp
@@ -44,6 +44,9 @@ DEFINE_STAT(STAT_AnimStateMachineUpdate);
 DEFINE_STAT(STAT_AnimStateMachineFindTransition);
 DEFINE_STAT(STAT_AnimStateMachineEvaluate);
 
+DEFINE_STAT(STAT_SkinPerPolyVertices);
+DEFINE_STAT(STAT_UpdateTriMeshVertices);
+
 // Define AnimNotify
 DEFINE_LOG_CATEGORY(LogAnimNotify);
 
diff --git a/Engine/Source/Runtime/Engine/Public/AnimTree.h b/Engine/Source/Runtime/Engine/Public/AnimTree.h
--- a/Engine/Source/Runtime/Engine/Public/AnimTree.h
+++ b/Engine/Source/Runtime/Engine/Public/AnimTree.h
@@ -34,6 +34,7 @@ DECLARE_CYCLE_STAT_EXTERN(TEXT("SkelComp UpdateTransform"), STAT_SkelCompUpdateT
 DECLARE_CYCLE_STAT_EXTERN(TEXT("MeshObject Update"), STAT_MeshObjectUpdate, STATGROUP_Anim, );
 DECLARE_CYCLE_STAT_EXTERN(TEXT("Update SkelMesh Bounds"), STAT_UpdateSkelMeshBounds, STATGROUP_Anim, );
 DECLARE_CYCLE_STAT_EXTERN(TEXT("BlendInPhysics"), STAT_BlendInPhysics, STATGROUP_Anim, );
-
+DECLARE_CYCLE_STAT_EXTERN(TEXT("SkinPerPolyVertices);"), STAT_SkinPerPolyVertices, STATGROUP_Anim, );
+DECLARE_CYCLE_STAT_EXTERN(TEXT("UpdateTriMeshVertices"), STAT_UpdateTriMeshVertices, STATGROUP_Anim, );
 
 
