UE4: Duplicating a UWorld now properly also duplicates it's lightmaps, shadowmaps, and level script blueprint classes, unless duplicating for PIE. TTP#337963

diff --git a/Engine/Source/Runtime/Engine/Classes/Engine/World.h b/Engine/Source/Runtime/Engine/Classes/Engine/World.h
--- a/Engine/Source/Runtime/Engine/Classes/Engine/World.h
+++ b/Engine/Source/Runtime/Engine/Classes/Engine/World.h
@@ -1628,6 +1628,7 @@ public:
 #if WITH_EDITOR
 	virtual bool Rename(const TCHAR* NewName = NULL, UObject* NewOuter = NULL, ERenameFlags Flags = REN_None) OVERRIDE;
 #endif
+	virtual void PostDuplicate(bool bDuplicateForPIE) OVERRIDE;
 
 	// End UObject Interface
 	
diff --git a/Engine/Source/Runtime/Engine/Private/World.cpp b/Engine/Source/Runtime/Engine/Private/World.cpp
--- a/Engine/Source/Runtime/Engine/Private/World.cpp
+++ b/Engine/Source/Runtime/Engine/Private/World.cpp
@@ -265,6 +265,77 @@ bool UWorld::Rename(const TCHAR* InName, UObject* NewOuter, ERenameFlags Flags)
 }
 #endif
 
+void UWorld::PostDuplicate(bool bDuplicateForPIE)
+{
+	Super::PostDuplicate(bDuplicateForPIE);
+
+	// If we are not duplicating for PIE, fix up objects that travel with the world.
+	// Note that these objects should really be inners of the world, so if they become inners later, most of this code should not be necessary
+	if ( !bDuplicateForPIE )
+	{
+		// Update the persistent level's owning world. This is needed for some initialization
+		if ( !PersistentLevel->OwningWorld )
+		{
+			PersistentLevel->OwningWorld = this;
+		}
+
+		// Update the current level as well
+		if ( !CurrentLevel )
+		{
+			CurrentLevel = PersistentLevel;
+		}
+
+		TMap<UObject*, UObject*> ReplacementMap;
+		UPackage* MyPackage = GetOutermost();
+
+		// Gather the textures
+		TArray<UTexture2D*> LightMapsAndShadowMaps;
+		GetLightMapsAndShadowMaps(LightMapsAndShadowMaps);
+
+		// Duplicate the textures, if any
+		for (auto* Tex : LightMapsAndShadowMaps)
+		{
+			if (Tex && Tex->GetOutermost() != MyPackage)
+			{
+				UObject* NewTex = StaticDuplicateObject(Tex, MyPackage, *Tex->GetName());
+				ReplacementMap.Add(Tex, NewTex);
+			}
+		}
+
+		// Duplicate the level script blueprint generated classes as well
+		const bool bDontCreate = true;
+		UBlueprint* LevelScriptBlueprint = PersistentLevel ? PersistentLevel->GetLevelScriptBlueprint(bDontCreate) : NULL;
+		if ( LevelScriptBlueprint )
+		{
+			UObject* OldGeneratedClass = LevelScriptBlueprint->GeneratedClass;
+			if ( OldGeneratedClass )
+			{
+				UObject* NewGeneratedClass = StaticDuplicateObject(OldGeneratedClass, MyPackage, *OldGeneratedClass->GetName());
+				ReplacementMap.Add(OldGeneratedClass, NewGeneratedClass);
+			}
+
+			UObject* OldSkeletonClass = LevelScriptBlueprint->SkeletonGeneratedClass;
+			if ( OldSkeletonClass )
+			{
+				UObject* NewSkeletonClass = StaticDuplicateObject(OldSkeletonClass, MyPackage, *OldSkeletonClass->GetName());
+				ReplacementMap.Add(OldSkeletonClass, NewSkeletonClass);
+			}
+		}
+
+		// Now replace references from the old textures/classes to the new ones, if any were duplicated
+		if (ReplacementMap.Num() > 0)
+		{
+			const bool bNullPrivateRefs = false;
+			const bool bIgnoreOuterRef = true;
+			const bool bIgnoreArchetypeRef = false;
+			FArchiveReplaceObjectRef<UObject> ReplaceAr(this, ReplacementMap, bNullPrivateRefs, bIgnoreOuterRef, bIgnoreArchetypeRef);
+		}
+
+		// Make sure PKG_ContainsMap is set
+		MyPackage->ThisContainsMap();
+	}
+}
+
 void UWorld::FinishDestroy()
 {
 	// Avoid cleanup if the world hasn't been initialized, e.g., the default object or a world object that got loaded
@@ -4961,22 +5032,49 @@ void UWorld::CopyGameState(AGameMode* FromGameMode, AGameState* FromGameState)
 
 void UWorld::GetLightMapsAndShadowMaps(TArray<UTexture2D*>& OutLightMapsAndShadowMaps)
 {
-	TArray<UObject*> WorldPackageObjects;
-	const bool bIncludeNestedObjects = false;
-	GetObjectsWithOuter(GetOutermost(), WorldPackageObjects, bIncludeNestedObjects);
-
-	for (auto ObjIt = WorldPackageObjects.CreateConstIterator(); ObjIt; ++ObjIt)
+	class FFindLightmapsArchive : public FArchiveUObject
 	{
-		UObject* CurrentObject = *ObjIt;
-		if ( ULightMapTexture2D* LightMap = Cast<ULightMapTexture2D>(CurrentObject) )
+		/** The array of textures discovered */
+		TArray<UTexture2D*>& TextureList;
+
+	public:
+		FFindLightmapsArchive(UObject* InSearch, TArray<UTexture2D*>& OutTextureList)
+			: TextureList(OutTextureList)
 		{
-			OutLightMapsAndShadowMaps.Add(LightMap);
+			ArIsObjectReferenceCollector = true;
+			ArIsModifyingWeakAndStrongReferences = true; // While we are not modifying them, we want to follow weak references as well
+
+			for (FObjectIterator It; It; ++It)
+			{
+				It->Mark(OBJECTMARK_TagExp);
+			}
+
+			*this << InSearch;
 		}
-		else if ( UShadowMapTexture2D* ShadowMap = Cast<UShadowMapTexture2D>(CurrentObject) )
+
+		FArchive& operator<<(class UObject*& Obj)
 		{
-			OutLightMapsAndShadowMaps.Add(ShadowMap);
+			// Don't check null references or objects already visited
+			if (Obj != NULL && Obj->HasAnyMarks(OBJECTMARK_TagExp))
+			{
+				if (Obj->IsA(ULightMapTexture2D::StaticClass()) || Obj->IsA(UShadowMapTexture2D::StaticClass()))
+				{
+					UTexture2D* Tex = Cast<UTexture2D>(Obj);
+					if ( ensure(Tex) )
+					{
+						TextureList.Add(Tex);
+					}
+				}
+
+				Obj->UnMark(OBJECTMARK_TagExp);
+				Obj->Serialize(*this);
+			}
+
+			return *this;
 		}
-	}
+	};
+
+	FFindLightmapsArchive FindArchive(this, OutLightMapsAndShadowMaps);
 }
 
 void UWorld::ChangeFeatureLevel(ERHIFeatureLevel::Type InFeatureLevel)
