- Build Fixes.

diff --git a/Engine/Source/Runtime/Engine/Private/PhysicsEngine/BodySetup.cpp b/Engine/Source/Runtime/Engine/Private/PhysicsEngine/BodySetup.cpp
--- a/Engine/Source/Runtime/Engine/Private/PhysicsEngine/BodySetup.cpp
+++ b/Engine/Source/Runtime/Engine/Private/PhysicsEngine/BodySetup.cpp
@@ -25,6 +25,8 @@
 	}
 #endif // WITH_PHYSX
 
+DECLARE_CYCLE_STAT(TEXT("UpdateTriMeshVertices "),STAT_UpdateTriMeshVertices,STATGROUP_Physics);
+
 DEFINE_LOG_CATEGORY(LogPhysics);
 UBodySetup::UBodySetup(const class FPostConstructInitializeProperties& PCIP)
 	: Super(PCIP)
diff --git a/Engine/Source/Runtime/Engine/Private/PhysicsEngine/PhysAnim.cpp b/Engine/Source/Runtime/Engine/Private/PhysicsEngine/PhysAnim.cpp
--- a/Engine/Source/Runtime/Engine/Private/PhysicsEngine/PhysAnim.cpp
+++ b/Engine/Source/Runtime/Engine/Private/PhysicsEngine/PhysAnim.cpp
@@ -12,6 +12,7 @@
 	#include "PhysXSupport.h"
 #endif // WITH_PHYSX
 
+DECLARE_CYCLE_STAT(TEXT("SkinPerPolyVertices "),STAT_SkinPerPolyVertices,STATGROUP_Physics);
 
 /** Used for drawing pre-phys skeleton if bShowPrePhysBones is true */
 static const FColor AnimSkelDrawColor(255, 64, 64);
