#ue4 - Automatically detect changes in position outside CharacterMovementComponent's normal updates and force floor updates when in walking mode.

- Fixes issues when bAlwaysCheckFloor is false (the default) and user abilities teleport the player without using the TeleportTo() functionality.

diff --git a/Engine/Source/Runtime/Engine/Classes/GameFramework/CharacterMovementComponent.h b/Engine/Source/Runtime/Engine/Classes/GameFramework/CharacterMovementComponent.h
--- a/Engine/Source/Runtime/Engine/Classes/GameFramework/CharacterMovementComponent.h
+++ b/Engine/Source/Runtime/Engine/Classes/GameFramework/CharacterMovementComponent.h
@@ -515,6 +515,12 @@ protected:
 	UPROPERTY()
 	FVector Acceleration;
 
+	/**
+	 * Location after last PerformMovement update. Used internally to detect changes in position from outside character movement to try to validate the current floor.
+	 */
+	UPROPERTY()
+	FVector LastUpdateLocation;
+
 	/** Accumulated impulse added this tick */
 	UPROPERTY()
 	FVector PendingImpulseToApply;
diff --git a/Engine/Source/Runtime/Engine/Private/CharacterMovementComponent.cpp b/Engine/Source/Runtime/Engine/Private/CharacterMovementComponent.cpp
--- a/Engine/Source/Runtime/Engine/Private/CharacterMovementComponent.cpp
+++ b/Engine/Source/Runtime/Engine/Private/CharacterMovementComponent.cpp
@@ -562,12 +562,10 @@ void UCharacterMovementComponent::OnMovementModeChanged(EMovementMode PreviousMo
 		Velocity.Z = 0.f;
 		bCrouchMovesCharacterDown = true;
 
+		// make sure we update our new floor/base on initial entry of the walking physics
 		FindFloor(UpdatedComponent->GetComponentLocation(), CurrentFloor, false);
-		AdjustFloorHeight();			
+		AdjustFloorHeight();
 		SetBase(CurrentFloor.HitResult.Component.Get());
-
-		// make sure we validate our new floor/base on initial entry of the walking physics, even if velocity is zero.
-		bForceNextFloorCheck = true;
 	}
 	else
 	{
@@ -1244,10 +1242,13 @@ void UCharacterMovementComponent::PerformMovement(float DeltaSeconds)
 		return;
 	}
 
+	// Force floor update if we've moved outside of CharacterMovement since last update.
+	bForceNextFloorCheck |= (IsMovingOnGround() && UpdatedComponent->GetComponentLocation() != LastUpdateLocation);
+
+	// Scoped updates can improve performance of multiple MoveComponent calls.
 	FScopedMovementUpdate ScopedMovementUpdate(UpdatedComponent, bEnableScopedMovementUpdates ? EScopedUpdate::DeferredUpdates : EScopedUpdate::ImmediateUpdates);
 
 	MaybeUpdateBasedMovement(DeltaSeconds);
-	
 	ApplyAccumulatedForces(DeltaSeconds);
 
 	FVector OldVelocity = Velocity;
@@ -1373,6 +1374,8 @@ void UCharacterMovementComponent::PerformMovement(float DeltaSeconds)
 
 	// update component velocity
 	UpdateComponentVelocity();
+
+	LastUpdateLocation = UpdatedComponent ? UpdatedComponent->GetComponentLocation() : FVector::ZeroVector;
 }
 
 
@@ -2828,7 +2831,7 @@ void UCharacterMovementComponent::StartFalling(int32 Iterations, float remaining
 	// start falling 
 	const float DesiredDist = Delta.Size();
 	const float ActualDist = (CharacterOwner->GetActorLocation() - subLoc).Size2D();
-	remainingTime = (DesiredDist == 0.f) 
+	remainingTime = (DesiredDist < KINDA_SMALL_NUMBER) 
 					? 0.f
 					: remainingTime + timeTick * (1.f - FMath::Min(1.f,ActualDist/DesiredDist));
 
@@ -3033,7 +3036,6 @@ void UCharacterMovementComponent::PhysWalking(float deltaTime, int32 Iterations)
 	FVector DesiredMove = Velocity;
 	DesiredMove.Z = 0.f;
 
-	//Perform the move
 	const FVector OldLocation = UpdatedComponent->GetComponentLocation();
 	const FFindFloorResult OldFloor = CurrentFloor;
 	UPrimitiveComponent* const OldBase = CharacterOwner->GetMovementBase();
@@ -3043,6 +3045,7 @@ void UCharacterMovementComponent::PhysWalking(float deltaTime, int32 Iterations)
 	bool bTriedLedgeMove = false;
 	float remainingTime = deltaTime;
 
+	// Perform the move
 	while ( (remainingTime >= MIN_TICK_TIME) && (Iterations < MaxSimulationIterations) && (CharacterOwner->Controller || bRunPhysicsWithNoController || HasRootMotion()) )
 	{
 		Iterations++;
