Initial support and usage of deferred guid resolves

* No longer load streamed levels in package map, instead defer them, and then resolve pending guids when we can

diff --git a/Engine/Source/Runtime/CoreUObject/Public/UObject/CoreNet.h b/Engine/Source/Runtime/CoreUObject/Public/UObject/CoreNet.h
--- a/Engine/Source/Runtime/CoreUObject/Public/UObject/CoreNet.h
+++ b/Engine/Source/Runtime/CoreUObject/Public/UObject/CoreNet.h
@@ -124,7 +124,8 @@ class COREUOBJECT_API UPackageMap : public UObject
 
 	virtual void		SetLocked( bool L ) { }
 
-	virtual UObject *	ResolvePathAndAssignNetGUID( FNetworkGUID & InOutNetGUID, const FString & PathName, const FString & FilenameOverride, UObject * ObjOuter, const bool bNoLoad = false ) { return NULL; }
+	virtual UObject *	ResolvePathAndAssignNetGUID( FNetworkGUID & InOutNetGUID, const FString & PathName, UObject * ObjOuter, const bool bNoLoad = false ) { return NULL; }
+	virtual void		UpdatePendingResolveGUIDs() { }
 
 	virtual bool		SerializeNewActor(FArchive & Ar, class UActorChannel * Channel, class AActor *& Actor) { return false; }
 
diff --git a/Engine/Source/Runtime/Engine/Classes/Engine/PackageMapClient.h b/Engine/Source/Runtime/Engine/Classes/Engine/PackageMapClient.h
--- a/Engine/Source/Runtime/Engine/Classes/Engine/PackageMapClient.h
+++ b/Engine/Source/Runtime/Engine/Classes/Engine/PackageMapClient.h
@@ -9,26 +9,55 @@
  *	Once ACK'd, objects are just serialized as <NetGUID>. The result is higher bandwidth usage upfront for new clients, but minimal bandwidth once
  *	things gets going.
  *
+ *	A further optimization is enabled by default. References will actually be serialized via:
+ *	<NetGUID, <(Outer *), Object Name> >. Where Outer * is a reference to the UObject's Outer.
+ *
  *	The main advantages from this are:
  *		-Flexibility. No precomputed net indices or large package lists need to be exchanged for UObject serialization.
  *		-Cross version communication. The name is all that is needed to exchange references.
- *		-Effeciency in that a very small % of UObjects will ever be serialized. Only Objects that serialized are assigned NetGUIDs.
+ *		-Efficiency in that a very small % of UObjects will ever be serialized. Only Objects that serialized are assigned NetGUIDs.
  */
 
 #pragma once
 #include "Net/DataBunch.h"
 #include "PackageMapClient.generated.h"
 
-class FDeferredResolvePath
+/**
+ * A guid that needs to eventually be resolved
+ * GUID's that need to be resolved, but can't yet due to:
+ * 1. They are streaming levels, or levels in general
+ * 	1a. We always expect clients to be in charge of actually loading levels, so we wait for that to happen before we resolve any guids within the level
+ * 2. We are asynchronously loading the object we will assign the net guid to, so we need to wait for the load to finish
+ * 	2a. This is work in progress, we don't have any official async loading code yet
+ * 3. We don't currently do this, but if we were depending on a dynamic outer, but have a stable name (stably named sub-object) we would need to defer the resolve
+ */
+class FPendingResolveGUID
 {
 public:
-	FDeferredResolvePath() {}
+	FPendingResolveGUID() : bNoLoad( false ), StartTimeInSeconds( 0 ) {}
 
-	FName			PathName;
-	FString			FilenameOverride;
-	FNetworkGUID	OuterGUID;
+	FString			PathName;				// Path name of object to load for this guid
+	FNetworkGUID	OuterGUID;				// OuterGUID that this object depends on
+	bool			bNoLoad;				// Whether or not we should load this object, or just wait for it (we wait for levels for example, never load them in package map)
+	double			StartTimeInSeconds;		// When this guid was first deferred (in seconds)
 };
 
+class FMarkSerializeNewActorScoped
+{
+public:
+	FMarkSerializeNewActorScoped( bool & InIsSerializingNewActor ) : IsSerializingNewActorRef( InIsSerializingNewActor )
+	{ 
+		IsSerializingNewActorRef = true; 
+	};
+
+	~FMarkSerializeNewActorScoped()
+	{
+		IsSerializingNewActorRef = false;
+	}
+
+	bool & IsSerializingNewActorRef;
+};
+	
 /** Stores an object with path associated with FNetworkGUID */
 class FNetGuidCacheObject
 {
@@ -71,12 +100,13 @@ class UPackageMapClient : public UPackageMap
 
 	UPackageMapClient( const class FPostConstructInitializeProperties & PCIP, UNetConnection * InConnection, TSharedPtr< FNetGUIDCache > InNetGUIDCache ) : 
 		UPackageMap( PCIP )
-	, Connection(InConnection)
+	,	Connection(InConnection)
 	{
-		GuidCache = InNetGUIDCache;
-		Locked = false;
-		ExportNetGUIDCount = 0;
+		GuidCache				= InNetGUIDCache;
+		Locked					= false;
+		ExportNetGUIDCount		= 0;
 		IsExportingNetGUIDBunch = false;
+		IsSerializingNewActor	= false;
 	}
 
 	virtual ~UPackageMapClient()
@@ -101,7 +131,7 @@ class UPackageMapClient : public UPackageMap
 
 	// UPackageMapClient Connection specific methods
 
-	virtual bool NetGUIDHasBeenAckd(FNetworkGUID NetGUID);
+	bool NetGUIDHasBeenAckd(FNetworkGUID NetGUID);
 
 	virtual void ReceivedNak( const int32 NakPacketId ) OVERRIDE;
 	virtual void ReceivedAck( const int32 AckPacketId ) OVERRIDE;
@@ -129,17 +159,19 @@ protected:
 	bool	ExportNetGUID( FNetworkGUID NetGUID, const UObject * Object, FString PathName, UObject * ObjOuter );
 	void	ExportNetGUIDHeader();
 
-	void	InternalWriteObject( FArchive& Ar, FNetworkGUID NetGUID, const UObject * Object, FString ObjectPathName, UObject* ObjectOuter );	
-	FNetworkGUID	InternalLoadObject( FArchive & Ar, UObject *& Object, const bool bIsOuterLevel, int InternalLoadObjectRecursionCount );
+	void			InternalWriteObject( FArchive& Ar, FNetworkGUID NetGUID, const UObject * Object, FString ObjectPathName, UObject * ObjectOuter );	
+	FNetworkGUID	InternalLoadObject( FArchive & Ar, UObject *& Object, const bool bIsInnerLevel, int InternalLoadObjectRecursionCount );
 
-	void		ResolvePathAndAssignNetGUID_Deferred( const FNetworkGUID & NetGUID, const FString & PathName, const FString & FilenameOverride, const FNetworkGUID & OuterGUID );
-	UObject *	ResolvePathAndAssignNetGUID( FNetworkGUID & InOutNetGUID, const FString & PathName, const FString & FilenameOverride, const FNetworkGUID & OuterGUID, const bool bNoLoad = false );
+	void		ResolvePathAndAssignNetGUID_Deferred( const FNetworkGUID & NetGUID, const FString & PathName, const FNetworkGUID & OuterGUID, const bool bNoLoad = false );
+	UObject *	ResolvePathAndAssignNetGUID( FNetworkGUID & InOutNetGUID, const FString & PathName, const FNetworkGUID & OuterGUID, const bool bNoLoad = false );
 
-	virtual UObject * ResolvePathAndAssignNetGUID( FNetworkGUID & InOutNetGUID, const FString & PathName, const FString & FilenameOverride, UObject * ObjOuter, const bool bNoLoad = false ) OVERRIDE;
-	
-	bool	ShouldSendFullPath(const UObject* Object, const FNetworkGUID &NetGUID);
+	virtual UObject * ResolvePathAndAssignNetGUID( FNetworkGUID & InOutNetGUID, const FString & PathName, UObject * ObjOuter, const bool bNoLoad = false ) OVERRIDE;
 
-	bool	IsNetGUIDAuthority();
+	virtual void UpdatePendingResolveGUIDs() override;
+
+	bool	ShouldSendFullPath(const UObject* Object, const FNetworkGUID &NetGUID);
+	
+	bool IsNetGUIDAuthority();
 
 	bool Locked;
 	bool IsExportingNetGUIDBunch;
@@ -148,20 +180,21 @@ protected:
 
 	bool ObjectLevelHasFinishedLoading(UObject* Obj);
 
-	TSet< FNetworkGUID >			CurrentExportNetGUIDs;				// Current list of NetGUIDs being written to the Export Bunch.
+	TSet< FNetworkGUID >						CurrentExportNetGUIDs;				// Current list of NetGUIDs being written to the Export Bunch.
 
-	TMap< FNetworkGUID, int32 >		NetGUIDAckStatus;
-	TArray< FNetworkGUID >			PendingAckGUIDs;					// Quick access to all GUID's that haven't been acked
-
-	// -----------------
+	TMap< FNetworkGUID, int32 >					NetGUIDAckStatus;
+	TArray< FNetworkGUID >						PendingAckGUIDs;					// Quick access to all GUID's that haven't been acked
 
 	// Bunches of NetGUID/path tables to send with the current content bunch
-	TArray<FOutBunch* >				ExportBunches;
-	FOutBunch *						CurrentExportBunch;
+	TArray<FOutBunch* >							ExportBunches;
+	FOutBunch *									CurrentExportBunch;
+
+	int32										ExportNetGUIDCount;
 
-	int32							ExportNetGUIDCount;
+	// A guid that needs to eventually be resolved
+	TMap< FNetworkGUID, FPendingResolveGUID >	PendingResolveGUIDs;
 
-	TMap< FNetworkGUID, FDeferredResolvePath >	DeferredResolvePaths;
+	TSharedPtr< FNetGUIDCache >					GuidCache;
 
-	TSharedPtr< FNetGUIDCache >		GuidCache;
+	bool										IsSerializingNewActor;
 };
diff --git a/Engine/Source/Runtime/Engine/Private/NetworkDriver.cpp b/Engine/Source/Runtime/Engine/Private/NetworkDriver.cpp
--- a/Engine/Source/Runtime/Engine/Private/NetworkDriver.cpp
+++ b/Engine/Source/Runtime/Engine/Private/NetworkDriver.cpp
@@ -99,13 +99,13 @@ void UNetDriver::PostInitProperties()
 		PacketSimulationSettings.RegisterCommands();
 		PacketSimulationSettings.ParseSettings(FCommandLine::Get());
 #endif
-		RoleProperty       = FindObjectChecked<UProperty>( AActor::StaticClass(), TEXT("Role"      ) );
-		RemoteRoleProperty = FindObjectChecked<UProperty>( AActor::StaticClass(), TEXT("RemoteRole") );
+		RoleProperty		= FindObjectChecked<UProperty>( AActor::StaticClass(), TEXT("Role"      ) );
+		RemoteRoleProperty	= FindObjectChecked<UProperty>( AActor::StaticClass(), TEXT("RemoteRole") );
 
 		GuidCache			= TSharedPtr< FNetGUIDCache >( new FNetGUIDCache( this ) );
 		NetCache			= TSharedPtr< FClassNetCacheMgr >( new FClassNetCacheMgr() );
 
-		ProfileStats	   = FParse::Param(FCommandLine::Get(),TEXT("profilestats"));
+		ProfileStats		= FParse::Param(FCommandLine::Get(),TEXT("profilestats"));
 	}
 	// By default we're the game net driver and any child ones must override this
 	NetDriverName = NAME_GameNetDriver;
@@ -303,6 +303,14 @@ void UNetDriver::TickFlush(float DeltaSeconds)
 		DrawNetDriverDebug();
 	}
 
+	// Update any pending resolves on the server package map
+	// We shouldn't need to do this as the server, since clients should only be sending us names that resolve immediately
+	// Do this before we try to resolve unmapped properties, since this may reduce the number of unmapped properties
+	if ( ServerConnection != NULL && ServerConnection->PackageMap != NULL )
+	{
+		ServerConnection->PackageMap->UpdatePendingResolveGUIDs();
+	}
+
 	// Update properties that are unmapped, try to hook up the object pointers if they exist now
 	for ( auto It = UnmappedReplicators.CreateIterator(); It; ++It )
 	{
diff --git a/Engine/Source/Runtime/Engine/Private/PackageMapClient.cpp b/Engine/Source/Runtime/Engine/Private/PackageMapClient.cpp
--- a/Engine/Source/Runtime/Engine/Private/PackageMapClient.cpp
+++ b/Engine/Source/Runtime/Engine/Private/PackageMapClient.cpp
@@ -25,6 +25,8 @@ static const int INTERNAL_LOAD_OBJECT_RECURSION_LIMIT = 16;
 #define COMPOSE_NET_GUID( Index, IsStatic )			( ( ( Index ) << 1 ) | ( IsStatic ) )
 #define COMPOSE_RELATIVE_NET_GUID( GUID, Index )	( COMPOSE_NET_GUID( GET_NET_GUID_INDEX( GUID ) + Index, GET_NET_GUID_STATIC_MOD( GUID ) ) )
 
+#define STRESS_PENDING_GUIDS ( 0 )
+
 /*-----------------------------------------------------------------------------
 	UPackageMapClient implementation.
 -----------------------------------------------------------------------------*/
@@ -90,7 +92,7 @@ bool UPackageMapClient::SerializeObject( FArchive& Ar, UClass* Class, UObject*&
 		FNetworkGUID NetGUID = InternalLoadObject( Ar, Object, false, 0 );
 
 		// Write out NetGUID to caller if necessary
-		if (OutNetGUID)
+		if ( OutNetGUID )
 		{
 			*OutNetGUID = NetGUID;
 		}	
@@ -100,14 +102,14 @@ bool UPackageMapClient::SerializeObject( FArchive& Ar, UClass* Class, UObject*&
 		// ----------------	
 
 		// NULL if we haven't finished loading the objects level yet
-		if (!ObjectLevelHasFinishedLoading(Object))
+		if ( !ObjectLevelHasFinishedLoading( Object ) )
 		{
 			UE_LOG(LogNetPackageMap, Log, TEXT("Using None instead of replicated reference to %s because the level it's in has not been made visible"), *Object->GetFullName());
 			Object = NULL;
 		}
 
 		// Check that we got the right class
-		if( Object && !Object->IsA(Class) )
+		if ( Object && !Object->IsA( Class ) )
 		{
 			UE_LOG(LogNetPackageMap, Warning, TEXT("Forged object: got %s, expecting %s"),*Object->GetFullName(),*Class->GetFullName());
 			Object = NULL;
@@ -162,7 +164,10 @@ bool UPackageMapClient::WriteObject( FArchive& Ar, UObject* ObjOuter, FNetworkGU
  */
 bool UPackageMapClient::SerializeNewActor(FArchive& Ar, class UActorChannel *Channel, class AActor*& Actor)
 {
-	UE_LOG(LogNetPackageMap, VeryVerbose, TEXT("SerializeNewActor START"));
+	UE_LOG( LogNetPackageMap, VeryVerbose, TEXT( "SerializeNewActor START" ) );
+
+	// This will mark IsSerializingNewActor true for the scope of the function
+	FMarkSerializeNewActorScoped SerializeNewActorScoped( IsSerializingNewActor );
 
 	uint8 bIsClosingChannel = 0;
 
@@ -179,6 +184,11 @@ bool UPackageMapClient::SerializeNewActor(FArchive& Ar, class UActorChannel *Cha
 	UObject *NewObj = Actor;
 	SerializeObject(Ar, AActor::StaticClass(), NewObj, &NetGUID);
 
+	if ( PendingResolveGUIDs.Contains( NetGUID ) )
+	{
+		UE_LOG( LogNetPackageMap, Error, TEXT( "UPackageMapClient::SerializeNewActor. Unresolved Actor GUID: Path: %s, NetGUID: %s." ), *PendingResolveGUIDs.FindChecked( NetGUID ).PathName, *NetGUID.ToString() );
+	}
+
 	if ( Ar.IsError() )
 	{
 		UE_LOG( LogNetPackageMap, Error, TEXT( "UPackageMapClient::SerializeNewActor: Ar.IsError after SerializeObject 1" ) );
@@ -228,6 +238,11 @@ bool UPackageMapClient::SerializeNewActor(FArchive& Ar, class UActorChannel *Cha
 		FNetworkGUID ArchetypeNetGUID;
 		SerializeObject(Ar, UObject::StaticClass(), Archetype, &ArchetypeNetGUID );
 
+		if ( PendingResolveGUIDs.Contains( ArchetypeNetGUID ) )
+		{
+			UE_LOG( LogNetPackageMap, Error, TEXT( "UPackageMapClient::SerializeNewActor. Unresolved Archetype GUID: Path: %s, NetGUID: %s." ), *PendingResolveGUIDs.FindChecked( ArchetypeNetGUID ).PathName, *ArchetypeNetGUID.ToString() );
+		}
+
 		// SerializeCompressedInitial
 		{			
 			// read/write compressed location
@@ -309,8 +324,8 @@ bool UPackageMapClient::SerializeNewActor(FArchive& Ar, class UActorChannel *Cha
 				check( !PendingAckGUIDs.Contains( SubobjNetGUID ) );
 			}
 
-				Ar << Checksum;
-			}
+			Ar << Checksum;
+		}
 		else
 		{
 			// Load the server checksum, then generate our own, and make sure they match
@@ -318,10 +333,10 @@ bool UPackageMapClient::SerializeNewActor(FArchive& Ar, class UActorChannel *Cha
 			uint32 LocalChecksum = 0;
 			Ar << ServerChecksum;
 
-				TArray< FNetworkGUID > SubObjGuids;
+			TArray< FNetworkGUID > SubObjGuids;
 
 			for ( int32 i = 0; i < Subobjs.Num(); i++ )
-				{
+			{
 				// Generate a guid that is relative to our owning actor
 				SubObjGuids.Add( COMPOSE_RELATIVE_NET_GUID( NetGUID, i + 1 ) );
 
@@ -330,39 +345,39 @@ bool UPackageMapClient::SerializeNewActor(FArchive& Ar, class UActorChannel *Cha
 				LocalChecksum = FCrc::StrCrc32( *Subobjs[i]->GetName(), LocalChecksum );
 
 				UE_LOG( LogNetPackageMap, Log, TEXT( "SerializeNewActor: [Loading] Assigned NetGUID <%s> to subobject %s" ), *SubObjGuids[i].ToString(), *Subobjs[i]->GetPathName() );
-					}
+			}
 
 			if ( LocalChecksum != ServerChecksum )
-					{
-					UE_LOG( LogNetPackageMap, Error, TEXT( "SerializeNewActor: Subobject checksum FAILED: %s" ), *Actor->GetFullName() );
-				}
-				else
+			{
+				UE_LOG( LogNetPackageMap, Error, TEXT( "SerializeNewActor: Subobject checksum FAILED: %s" ), *Actor->GetFullName() );
+			}
+			else
+			{
+				// Everything worked, assign the guids
+				for ( int32 i = 0; i < Subobjs.Num(); i++ )
 				{
-					// Everything worked, assign the guids
-					for ( int32 i = 0; i < Subobjs.Num(); i++ )
-					{
 					FNetworkGUID OldNetGUID = GuidCache->NetGUIDLookup.FindRef( Subobjs[i] );
 
 					if ( OldNetGUID.IsValid() )
-{
+					{
 						if ( OldNetGUID != SubObjGuids[i] )
-	{
+						{
 							UE_LOG( LogNetPackageMap, Warning, TEXT( "SerializeNewActor: Changing NetGUID on subobject. Name: %s, OldGUID: %s, NewGUID: %s" ), *Subobjs[i]->GetPathName(), *OldNetGUID.ToString(), *SubObjGuids[i].ToString() );
-	}
-	}
+						}
+					}
 
 					GuidCache->AssignNetGUID( Subobjs[i], SubObjGuids[i] );
 
 					UE_LOG( LogNetPackageMap, Log, TEXT( "SerializeNewActor: [Loading] Assigned NetGUID <%s> to subobject %s" ), *SubObjGuids[i].ToString(), *Subobjs[i]->GetPathName() );
-	}
-		}
+				}
+			}
 		}
 	}
 
 	UE_LOG(LogNetPackageMap, Log, TEXT("SerializeNewActor END: Finished Serializing Actor %s <%s> on Channel %d"), Actor ? *Actor->GetName() : TEXT("NULL"), *NetGUID.ToString(), Channel->ChIndex );
 
 	return SpawnedActor;
-		}
+}
 
 //--------------------------------------------------------------------
 //
@@ -376,10 +391,9 @@ struct FExportFlags
 	{
 		struct
 		{
-			uint8 bHasPath		: 1;
-			uint8 bIsLevel		: 1;
-			uint8 bIsPackage	: 1;
-			uint8 bIsRenamedPackage	: 1;
+			uint8 bHasPath			: 1;
+			uint8 bIsLevel			: 1;
+			uint8 bIsPackage		: 1;
 		};
 
 		uint8	Value;
@@ -422,58 +436,46 @@ void UPackageMapClient::InternalWriteObject( FArchive & Ar, FNetworkGUID NetGUID
 			ExportFlags.bHasPath = ObjectPathName.IsEmpty() ? 0 : 1;
 		}
 
-		const UPackage * Package = Cast< const UPackage >( Object );
-
 		ExportFlags.bIsLevel	= ( Cast< const ULevel >( Object ) != NULL ) ? 1 : 0;
-		ExportFlags.bIsPackage	= ( Package != NULL ) ? 1 : 0;
-
-		ExportFlags.bIsRenamedPackage = ( Package != NULL && Package->FileName != NAME_None && Package->FileName != Package->GetFName() ) ? 1 : 0;
+		ExportFlags.bIsPackage	= ( Cast< const UPackage >( Object ) != NULL ) ? 1 : 0;
 
 		Ar << ExportFlags.Value;
 	}
 
 	if ( ExportFlags.bHasPath )
 	{
-	if ( Object != NULL )
-	{
-		// If the object isn't NULL, expect an empty path name, then fill it out with the actual info
-		check( ObjectOuter == NULL );
-		check( ObjectPathName.IsEmpty() );
-		ObjectPathName = Object->GetName();
-		ObjectOuter = Object->GetOuter();
-	}
-	else
-	{
-		// If we don't have an object, expect an already filled out path name
-		check( ObjectOuter != NULL );
-		check( !ObjectPathName.IsEmpty() );
-	}
+		if ( Object != NULL )
+		{
+			// If the object isn't NULL, expect an empty path name, then fill it out with the actual info
+			check( ObjectOuter == NULL );
+			check( ObjectPathName.IsEmpty() );
+			ObjectPathName = Object->GetName();
+			ObjectOuter = Object->GetOuter();
+		}
+		else
+		{
+			// If we don't have an object, expect an already filled out path name
+			check( ObjectOuter != NULL );
+			check( !ObjectPathName.IsEmpty() );
+		}
 
-	// Serialize reference to outer. This is basically a form of compression.
+		// Serialize reference to outer. This is basically a form of compression.
 		FNetworkGUID OuterNetGUID = GuidCache->GetOrAssignNetGUID( ObjectOuter );
 
-	InternalWriteObject( Ar, OuterNetGUID, ObjectOuter, TEXT( "" ), NULL );
+		InternalWriteObject( Ar, OuterNetGUID, ObjectOuter, TEXT( "" ), NULL );
 
-	GEngine->NetworkRemapPath(Connection->Driver->GetWorld(), ObjectPathName, false);
+		GEngine->NetworkRemapPath(Connection->Driver->GetWorld(), ObjectPathName, false);
 
-	// Serialize Name of object
-	Ar << ObjectPathName;
+		// Serialize Name of object
+		Ar << ObjectPathName;
 
-		if ( ExportFlags.bIsRenamedPackage )
+		if ( IsExportingNetGUIDBunch )
 		{
-			const UPackage * Package = CastChecked< const UPackage >( Object );
-			FString FilenameOverride = Package->FileName.ToString();
-			Ar << FilenameOverride;
-			//UE_LOG( LogNetPackageMap, Warning, TEXT("UPackageMapClient::InternalWriteObject: Renamed package. Name: %s, Filename: %s"), *Package->GetName(), *FilenameOverride );
-		}
-
-	if ( IsExportingNetGUIDBunch )
-	{
-		CurrentExportNetGUIDs.Add( NetGUID );
+			CurrentExportNetGUIDs.Add( NetGUID );
 
-		int32& Count = NetGUIDExportCountMap.FindOrAdd(NetGUID);
-		Count++;
-	}
+			int32 & Count = NetGUIDExportCountMap.FindOrAdd( NetGUID );
+			Count++;
+		}
 	}
 }
 
@@ -483,7 +485,7 @@ void UPackageMapClient::InternalWriteObject( FArchive & Ar, FNetworkGUID NetGUID
 //
 //--------------------------------------------------------------------
 
-static void SanityCheckExport( const UObject * Object, const FNetworkGUID & NetGUID, const FString & ExpectedPathName, const FString & ExpectedFilename, const UObject * ExpectedOuter, const FExportFlags & ExportFlags )
+static void SanityCheckExport( const UObject * Object, const FNetworkGUID & NetGUID, const FString & ExpectedPathName, const UObject * ExpectedOuter, const FExportFlags & ExportFlags )
 {
 	if ( Object->GetName() != ExpectedPathName )
 	{
@@ -497,24 +499,19 @@ static void SanityCheckExport( const UObject * Object, const FNetworkGUID & NetG
 		UE_LOG( LogNetPackageMap, Warning, TEXT( "SanityCheckExport: Outer mismatch. Object: %s, NetGUID: %s, Current: %s, Expected: %s" ), *Object->GetPathName(), *NetGUID.ToString(), *CurrentOuterName, *ExpectedOuterName );
 	}
 
-	if ( ExportFlags.bIsPackage || ExportFlags.bIsRenamedPackage )
+	if ( ExportFlags.bIsPackage )
 	{
 		const UPackage * Package = Cast< const UPackage >( Object );
 
 		if ( Package == NULL )
 		{
 			UE_LOG( LogNetPackageMap, Warning, TEXT( "SanityCheckExport: Expected package. Object:%s, NetGUID: %s" ), *Object->GetPathName(), *NetGUID.ToString() );
-}
-
-		if ( ExportFlags.bIsRenamedPackage && Package->FileName.ToString() != ExpectedFilename )
-		{
-			UE_LOG( LogNetPackageMap, Warning, TEXT( "SanityCheckExport: Renamed package filename mismatch. Filename: %s, Expected: %s" ), *Package->FileName.ToString(), *ExpectedFilename );
 		}
 	}
 }
 
 /** Loads a UObject from an FArchive stream. Reads object path if there, and tries to load object if its not already loaded */
-FNetworkGUID UPackageMapClient::InternalLoadObject( FArchive & Ar, UObject *& Object, const bool bIsOuterLevel, const int InternalLoadObjectRecursionCount )
+FNetworkGUID UPackageMapClient::InternalLoadObject( FArchive & Ar, UObject *& Object, const bool bIsInnerLevel, const int InternalLoadObjectRecursionCount )
 {
 	if ( InternalLoadObjectRecursionCount > INTERNAL_LOAD_OBJECT_RECURSION_LIMIT ) 
 	{
@@ -543,6 +540,14 @@ FNetworkGUID UPackageMapClient::InternalLoadObject( FArchive & Ar, UObject *& Ob
 		return NetGUID;
 	}
 
+	if ( PendingResolveGUIDs.Contains( NetGUID ) )
+	{
+		// If this was previously on the deferred resolve list, go ahead and take this opportunity to 
+		// resolve any pending guids now (which may or may not resolve this guid)
+		// If this guid is resolved, it will also try and resolve any guids that have this guid as an outer as well
+		UpdatePendingResolveGUIDs();
+	}
+
 	// ----------------	
 	// Try to resolve NetGUID
 	// ----------------	
@@ -577,18 +582,11 @@ FNetworkGUID UPackageMapClient::InternalLoadObject( FArchive & Ar, UObject *& Ob
 	{
 		UObject * ObjOuter = NULL;
 
-		FNetworkGUID OuterGUID = InternalLoadObject( Ar, ObjOuter, ExportFlags.bIsLevel || bIsOuterLevel, InternalLoadObjectRecursionCount + 1 );
+		FNetworkGUID OuterGUID = InternalLoadObject( Ar, ObjOuter, ExportFlags.bIsLevel || bIsInnerLevel, InternalLoadObjectRecursionCount + 1 );
 
 		FString PathName;
 		Ar << PathName;
 
-		FString FilenameOverride;
-
-		if ( ExportFlags.bIsRenamedPackage )
-		{
-			Ar << FilenameOverride;
-		}
-
 		if ( Ar.IsError() )
 		{
 			UE_LOG( LogNetPackageMap, Error, TEXT( "InternalLoadObject: Failed to load path name" ) );
@@ -598,34 +596,46 @@ FNetworkGUID UPackageMapClient::InternalLoadObject( FArchive & Ar, UObject *& Ob
 
 		GEngine->NetworkRemapPath( Connection->Driver->GetWorld(), PathName, true );
 
-		if ( Object == NULL )
+		// Make note if the outer is still deferred (we then must defer the load of this guid)
+		const bool OuterIsDeferred = PendingResolveGUIDs.Contains( OuterGUID );
+
+#if 0//STRESS_PENDING_GUIDS == 0
+		if ( Object == NULL && !OuterIsDeferred )
 		{
 			// Try to resolve the name (we may have already loaded it), but DON'T load yet if we haven't already
-			Object = ResolvePathAndAssignNetGUID( NetGUID, PathName, FilenameOverride, OuterGUID, true );
+			Object = ResolvePathAndAssignNetGUID( NetGUID, PathName, OuterGUID, true );
 		}
+#endif
 
 		if ( Object != NULL )
 		{
 			// If we already have the object, just do some sanity checking and return
-			SanityCheckExport( Object, NetGUID, PathName, FilenameOverride, ObjOuter, ExportFlags );
+			SanityCheckExport( Object, NetGUID, PathName, ObjOuter, ExportFlags );
 			return NetGUID;
 		}
 
-		// If this is a level or our outer is being deferred, then defer this guid
-		const bool bCanDefer = false;//!NetGUID.IsDefault() && ( ExportFlags.bIsLevel || bIsOuterLevel || DeferredResolvePaths.Contains( OuterGUID ) );
+		const bool bIsForLevel = ExportFlags.bIsLevel || bIsInnerLevel;
+
+#if STRESS_PENDING_GUIDS == 1
+		// For testing, defer everything
+		const bool bShouldDefer = !NetGUID.IsDefault();
+#else
+		// If this is a level (or any of our inners are) or our outer is being deferred, then defer this guid
+		const bool bShouldDefer = !NetGUID.IsDefault() && ( bIsForLevel || OuterIsDeferred );
+#endif
 
-		if ( bCanDefer )
+		if ( bShouldDefer )
 		{
 			UE_LOG( LogNetPackageMap, Warning, TEXT( "InternalLoadObject: Deferring resolve. Path: %s, NetGUID: %s, Outer: %s" ), *PathName, *NetGUID.ToString(), ObjOuter != NULL ? *ObjOuter->GetFullName() : TEXT( "NULL" ) );
-			ResolvePathAndAssignNetGUID_Deferred( NetGUID, PathName, FilenameOverride, OuterGUID );
+			ResolvePathAndAssignNetGUID_Deferred( NetGUID, PathName, OuterGUID, bIsForLevel );
 		}
 		else
 		{
-			Object = ResolvePathAndAssignNetGUID( NetGUID, PathName, FilenameOverride, OuterGUID );
+			Object = ResolvePathAndAssignNetGUID( NetGUID, PathName, OuterGUID );
 		}
 	}
 
-	if ( Object == NULL && NetGUID.IsStatic() && !DeferredResolvePaths.Contains( NetGUID ) )
+	if ( Object == NULL && NetGUID.IsStatic() && !PendingResolveGUIDs.Contains( NetGUID ) )
 	{
 		UE_LOG( LogNetPackageMap, Warning, TEXT( "InternalLoadObject: Unable to resolve object. NetGUID: %s, HasPath: %s" ), *NetGUID.ToString(), ExportFlags.bHasPath ? TEXT( "TRUE" ) : TEXT( "FALSE" ) );
 	}
@@ -633,33 +643,85 @@ FNetworkGUID UPackageMapClient::InternalLoadObject( FArchive & Ar, UObject *& Ob
 	return NetGUID;
 }
 
-void UPackageMapClient::ResolvePathAndAssignNetGUID_Deferred( const FNetworkGUID & NetGUID, const FString & PathName, const FString & FilenameOverride, const FNetworkGUID & OuterGUID )
+void UPackageMapClient::ResolvePathAndAssignNetGUID_Deferred( const FNetworkGUID & NetGUID, const FString & PathName, const FNetworkGUID & OuterGUID, const bool bNoLoad )
 {
 	check( !NetGUID.IsDefault() );		// We can't defer default guids, doesn't make sense
 
-	if ( !DeferredResolvePaths.Contains( NetGUID ) )
+	if ( !PendingResolveGUIDs.Contains( NetGUID ) )
 	{
-		FDeferredResolvePath DeferredResolvePath;
+		FPendingResolveGUID PendingResolveGUID;
 
-		DeferredResolvePath.OuterGUID	= OuterGUID;
-		DeferredResolvePath.PathName	= FName( *PathName );
-		DeferredResolvePath.FilenameOverride	= FilenameOverride;
+		PendingResolveGUID.OuterGUID			= OuterGUID;
+		PendingResolveGUID.PathName				= PathName;
+		PendingResolveGUID.bNoLoad				= bNoLoad;
+		PendingResolveGUID.StartTimeInSeconds	= FPlatformTime::Seconds();
 
-		DeferredResolvePaths.Add( NetGUID, DeferredResolvePath );
+		PendingResolveGUIDs.Add( NetGUID, PendingResolveGUID );
 	}
 	else
 	{
 		// Sanity check
-		FDeferredResolvePath & DeferredResolvePath = DeferredResolvePaths.FindChecked( NetGUID );
+		FPendingResolveGUID & PendingResolveGUID = PendingResolveGUIDs.FindChecked( NetGUID );
 
-		if ( DeferredResolvePath.OuterGUID != OuterGUID )
+		if ( PendingResolveGUID.OuterGUID != OuterGUID )
 		{
-			UE_LOG( LogNetPackageMap, Warning, TEXT( "ResolvePathAndAssignNetGUID_Deferred: Deferred Outer mismatch. NetGUID: %s, Path: %s, Current: %s, Expected: %s" ), *NetGUID.ToString(), *PathName, *DeferredResolvePath.OuterGUID.ToString(), *OuterGUID.ToString() );
+			UE_LOG( LogNetPackageMap, Warning, TEXT( "ResolvePathAndAssignNetGUID_Deferred: Deferred Outer mismatch. NetGUID: %s, Path: %s, Current: %s, Expected: %s" ), *NetGUID.ToString(), *PathName, *PendingResolveGUID.OuterGUID.ToString(), *OuterGUID.ToString() );
 		}
 	}
 }
 
-UObject * UPackageMapClient::ResolvePathAndAssignNetGUID( FNetworkGUID & InOutNetGUID, const FString & PathName, const FString & FilenameOverride, UObject * ObjOuter, const bool bNoLoad )
+void UPackageMapClient::UpdatePendingResolveGUIDs()
+{
+#if STRESS_PENDING_GUIDS == 1
+	const double CurrentTimeInSeconds = FPlatformTime::Seconds();
+#endif
+
+	bool NeedToKeepResolving = true;
+
+	while ( NeedToKeepResolving )
+	{
+		NeedToKeepResolving = false;
+
+		for ( auto It = PendingResolveGUIDs.CreateIterator(); It; ++It )
+		{
+			const FPendingResolveGUID & PendingResolveGUID = It.Value();
+
+#if STRESS_PENDING_GUIDS == 1
+			if ( !NeedToKeepResolving && !IsSerializingNewActor && ( CurrentTimeInSeconds - PendingResolveGUID.StartTimeInSeconds < 15.0f ) )
+			{
+				continue;		// Wait 15 seconds for testing (unless we are in SerializeNewActor, which can't handle deferred guids yet
+			}
+#endif
+
+			if ( PendingResolveGUIDs.Contains( PendingResolveGUID.OuterGUID ) )
+			{
+				// We can't load yet if our outer hasn't been resolved yet
+				continue;
+			}
+
+			// First, see if we have an object for this GUID already
+			UObject * Object = GetObjectFromNetGUID( It.Key() );
+
+			if ( Object == NULL )
+			{
+				// If not, try to find the object by path
+				FNetworkGUID NetGUID = It.Key();
+				Object = ResolvePathAndAssignNetGUID( NetGUID, PendingResolveGUID.PathName, PendingResolveGUID.OuterGUID, PendingResolveGUID.bNoLoad );
+			}
+
+			if ( Object != NULL )
+			{
+				UE_LOG( LogNetPackageMap, Warning, TEXT( "UPackageMapClient::UpdatePendingResolveGUIDs. Resolved deferred NetGUID. Path: %s, NetGUID: %s. NumLeft: %i, IsSerializingNewActor: %s" ), *PendingResolveGUID.PathName, *It.Key().ToString(), PendingResolveGUIDs.Num() - 1, IsSerializingNewActor ? TEXT("YES") : TEXT("NO") );
+
+				// If we have an object, we can make another pass in case there were dependencies on this guid loading
+				NeedToKeepResolving = true;
+				It.RemoveCurrent();		// We don't need to check this one anymore
+			}
+		}
+	}
+}
+
+UObject * UPackageMapClient::ResolvePathAndAssignNetGUID( FNetworkGUID & InOutNetGUID, const FString & PathName, UObject * ObjOuter, const bool bNoLoad )
 {
 	if ( InOutNetGUID.IsDynamic() )
 	{
@@ -697,35 +759,6 @@ UObject * UPackageMapClient::ResolvePathAndAssignNetGUID( FNetworkGUID & InOutNe
 			// If we failed to load it as an object, try to load it as a package
 			//
 
-			if ( !FilenameOverride.IsEmpty() )
-			{
-				// Package was renamed, we need to load a package on disk into a renamed package in memory
-				if ( ObjOuter != NULL )
-				{
-					UE_LOG( LogNetPackageMap, Warning, TEXT( "UPackageMapClient::ResolvePathAndAssignNetGUID: Expected outer to be NULL. Name: %s, Override: %s, Outer: %s" ), *PathName, *FilenameOverride, *ObjOuter->GetName() );
-					return NULL;
-				}
-
-				// Create the package
-				UPackage * PackageOuter = CreatePackage( NULL, *PathName );
-
-				// Load the package off disk into this renamed package
-				Object = LoadPackage( PackageOuter, *FilenameOverride, LOAD_None );
-
-				UPackage * Package = Cast< UPackage >( Object );
-
-				if ( Package != NULL )
-				{
-					if ( Package->GetName() != PathName || Package->FileName.ToString() != FilenameOverride )
-					{
-						UE_LOG( LogNetPackageMap, Error, TEXT( "UPackageMapClient::ResolvePathAndAssignNetGUID: Renamed package mismatch. Name: %s, Filename: %s, ActualName: %s, ActualFile: %s" ), *PathName, *FilenameOverride, *Package->GetName(), *Package->FileName.ToString() );
-					}
-				}
-
-				UE_LOG( LogNetPackageMap, Log, TEXT( "UPackageMapClient::ResolvePathAndAssignNetGUID: Loaded renamed package. Name: %s. Found: %s, FileName: %s" ), *PathName, Object ? *Object->GetName() : TEXT( "NULL" ), *FilenameOverride );
-			}
-			else
-			{
 			UPackage * PackageOuter = Cast< UPackage >( ObjOuter );
 
 			// If we have an outer at this point, it only makes sense for it to be a package (or NULL)
@@ -740,7 +773,6 @@ UObject * UPackageMapClient::ResolvePathAndAssignNetGUID( FNetworkGUID & InOutNe
 			}
 		}
 	}
-	}
 
 	if ( Object )
 	{
@@ -762,24 +794,24 @@ UObject * UPackageMapClient::ResolvePathAndAssignNetGUID( FNetworkGUID & InOutNe
 	return Object;
 }
 
-UObject * UPackageMapClient::ResolvePathAndAssignNetGUID( FNetworkGUID & InOutNetGUID, const FString & PathName, const FString & FilenameOverride, const FNetworkGUID & OuterGUID, const bool bNoLoad )
+UObject * UPackageMapClient::ResolvePathAndAssignNetGUID( FNetworkGUID & InOutNetGUID, const FString & PathName, const FNetworkGUID & OuterGUID, const bool bNoLoad )
 {
 	UObject * ObjOuter = NULL;
 
 	// Resolve the outer if the guid is valid
 	if ( OuterGUID.IsValid() )
-{
+	{
 		ObjOuter = GetObjectFromNetGUID( OuterGUID );
 
 		// If we have a valid outer guid, we expect it to resolve
 		if ( ObjOuter == NULL )
-	{
-		UE_LOG( LogNetPackageMap, Warning, TEXT( "ResolvePathAndAssignNetGUID: Unable resolve outer. Path: %s, NetGUID: %s, OuterGUID: %s " ), *PathName, *InOutNetGUID.ToString(), *OuterGUID.ToString() );
-		return NULL;
-	}
+		{
+			UE_LOG( LogNetPackageMap, Warning, TEXT( "ResolvePathAndAssignNetGUID: Unable resolve outer. Path: %s, NetGUID: %s, OuterGUID: %s " ), *PathName, *InOutNetGUID.ToString(), *OuterGUID.ToString() );
+			return NULL;
+		}
 	}
 
-	UObject * Object = ResolvePathAndAssignNetGUID( InOutNetGUID, PathName, FilenameOverride, ObjOuter, bNoLoad );
+	UObject * Object = ResolvePathAndAssignNetGUID( InOutNetGUID, PathName, ObjOuter, bNoLoad );
 
 	if ( Object == NULL && !bNoLoad )
 	{
@@ -1202,7 +1234,7 @@ bool UPackageMapClient::ShouldSendFullPath(const UObject* Object, const FNetwork
 /**
  *	Prints debug info about this package map's state
  */
-void UPackageMapClient::LogDebugInfo(FOutputDevice& Ar)
+void UPackageMapClient::LogDebugInfo( FOutputDevice & Ar )
 {
 	for ( auto It = GuidCache->NetGUIDLookup.CreateIterator(); It; ++It )
 	{
@@ -1458,7 +1490,7 @@ bool FNetGUIDCache::SupportsObject( const UObject * Object )
 	// This should happen for dynamic subobjects.
 	FNetworkGUID NetGUID = NetGUIDLookup.FindRef( Object );
 
-	if (NetGUID.IsValid())
+	if ( NetGUID.IsValid() )
 	{
 		return true;
 	}
@@ -1486,27 +1518,27 @@ bool FNetGUIDCache::SupportsObject( const UObject * Object )
  *	referenced with a path name to the client.
  */
 bool FNetGUIDCache::IsDynamicObject( const UObject * Object )
-	{
+{
 	check( Object != NULL );
 	check( Object->IsSupportedForNetworking() );
 
 	// Any non net addressable object is dynamic
 	return !Object->IsFullNameStableForNetworking();
-	}
+}
 
 bool FNetGUIDCache::IsNetGUIDAuthority()
-	{
+{
 	return Driver == NULL || Driver->IsServer();
-		}
-		
+}
+
 /** Gets or assigns a new NetGUID to this object. Returns whether the object is fully mapped or not */
 FNetworkGUID FNetGUIDCache::GetOrAssignNetGUID( const UObject * Object )
-		{
+{
 	if ( !Object || !SupportsObject( Object ) )
-			{
+	{
 		// Null of unsupported object, leave as default NetGUID and just return mapped=true
 		return FNetworkGUID();
-			}
+	}
 
 	// ----------------
 	// Assign NetGUID if necessary
@@ -1514,31 +1546,31 @@ FNetworkGUID FNetGUIDCache::GetOrAssignNetGUID( const UObject * Object )
 	FNetworkGUID NetGUID = NetGUIDLookup.FindRef( Object );
 
 	if ( NetGUID.IsValid() )
-			{
+	{
 		return NetGUID;
-			}
+	}
 
 	return AssignNewNetGUID( Object );
-		}
+}
 
 static const AActor * GetOuterActor( const UObject * Object )
-		{
+{
 	Object = Object->GetOuter();
 
 	while ( Object != NULL )
-			{
+	{
 		const AActor * Actor = Cast< const AActor >( Object );
 
 		if ( Actor != NULL )
-				{
+		{
 			return Actor;
-				}
+		}
 
 		Object = Object->GetOuter();
-			}
+	}
 
 	return NULL;
-		}
+}
 
 /**
  *	Generate a new NetGUID for this object and assign it.
@@ -1547,7 +1579,9 @@ FNetworkGUID FNetGUIDCache::AssignNewNetGUID( const UObject * Object )
 {
 	if ( !IsNetGUIDAuthority() )
 	{
-		// We cannot make new NetGUIDs
+		// We cannot make or assign new NetGUIDs
+		// Generate a default GUID, which signifies we write the full path
+		// The server should detect this, and assign a full-time guid, and send that back to us
 		return FNetworkGUID::GetDefault();
 	}
 
@@ -1577,12 +1611,9 @@ FNetworkGUID FNetGUIDCache::AssignNewNetGUID( const UObject * Object )
 				return NetGUIDLookup.FindChecked( Object );
 			}
 		}
-}
+	}
 
-#define GET_NET_GUID_INDEX( GUID )				( GUID.Value >> 1 )
-#define GET_NET_GUID_STATIC_MOD( GUID )			( GUID.Value & 1 )
-#define COMPOSE_NET_GUID( Index, IsStatic )		( ( ( Index ) << 1 ) | ( IsStatic ) )
-#define ALLOC_NEW_NET_GUID( IsStatic )			( COMPOSE_NET_GUID( ++UniqueNetIDs[ IsStatic ], IsStatic ) )
+#define ALLOC_NEW_NET_GUID( IsStatic ) ( COMPOSE_NET_GUID( ++UniqueNetIDs[ IsStatic ], IsStatic ) )
 
 	// Generate new NetGUID and assign it
 	const int32 IsStatic = IsDynamicObject( Object ) ? 0 : 1;
@@ -1605,22 +1636,21 @@ FNetworkGUID FNetGUIDCache::AssignNewNetGUID( const UObject * Object )
 			check( GetOuterActor( SubObj ) == Actor );
 
 			if ( NetGUIDLookup.Contains( SubObj ) )		// We should NOT have a guid yet
-{
+			{
 				UE_LOG( LogNetPackageMap, Error, TEXT( "AssignNewNetGUID: Sub object already registered: Actor: %s, SubObj: %s, NetGuid: %s" ), *Actor->GetPathName(), *SubObj->GetPathName(), *NetGUIDLookup.FindChecked( SubObj ).ToString() );
 			}
 			const int32 SubIsStatic = IsDynamicObject( SubObj ) ? 0 : 1;
 			check( SubIsStatic == IsStatic );
 			const FNetworkGUID SubobjNetGUID( ALLOC_NEW_NET_GUID( SubIsStatic ) );
-			check( SubobjNetGUID.Value == COMPOSE_NET_GUID( GET_NET_GUID_INDEX( NewNetGuid ) + 1 + i, GET_NET_GUID_STATIC_MOD( NewNetGuid ) ) );
 			check( SubobjNetGUID.Value == COMPOSE_RELATIVE_NET_GUID( NewNetGuid, i + 1 ) );
 			AssignNetGUID( SubObj, SubobjNetGUID );
 		}
 	}
-	
+
 	return NewNetGuid;
 }
 
-/**	
+/**
  *	Assigns the given UObject the give NetNetworkGUID.
  */
 void FNetGUIDCache::AssignNetGUID( const UObject * Object, const FNetworkGUID & NewNetworkGUID )
@@ -1637,19 +1667,19 @@ void FNetGUIDCache::AssignNetGUID( const UObject * Object, const FNetworkGUID &
 	CacheObject.FullPath = Object->GetPathName();
 
 	// Remove any existing entries
-{
+	{
 		const FNetGuidCacheObject * ExistingCacheObjectPtr = ObjectLookup.Find( NewNetworkGUID );
 
 		if ( ExistingCacheObjectPtr )
-	{
+		{
 			const UObject * OldObject = ExistingCacheObjectPtr->Object.Get();
 
 			// If we find static guid in the cache, this means we didn't find it in GetObjectFromNetGUID
 			// This is a problem, since we should never get this far if we do find it there
 			if ( !NewNetworkGUID.IsDynamic() )
-		{
+			{
 				UE_LOG( LogNetPackageMap, Warning, TEXT( "AssignNetGUID: NewNetworkGUID was in ObjectLookup cache but is static: <%s> %s" ), *NewNetworkGUID.ToString(), Object ? *Object->GetPathName() : TEXT( "NULL" ) );
-		}
+			}
 
 			// If this net guid was found but the old object is NULL, this can happen due to:
 			//	1. Actor channel was closed locally (but we don't remove the net guid entry, since we can't know for sure if it will be referenced again)
@@ -1662,11 +1692,11 @@ void FNetGUIDCache::AssignNetGUID( const UObject * Object, const FNetworkGUID &
 			// If the object pointers are different, this can be a problem, 
 			//	since this should only be possible if something gets out of sync during the net guid exchange code
 			if ( OldObject != NULL && OldObject != Object )
-		{
+			{
 				UE_LOG( LogNetPackageMap, Warning, TEXT( "Reassigning NetGUID <%s> to %s (was assigned to object %s)" ), *NewNetworkGUID.ToString(), Object ? *Object->GetPathName() : TEXT( "NULL" ), OldObject ? *OldObject->GetPathName() : TEXT( "NULL" ) );
-		}
-		else
-		{
+			}
+			else
+			{
 				UE_LOG( LogNetPackageMap, Verbose, TEXT( "Reassigning NetGUID <%s> to %s (was assigned to object %s)" ), *NewNetworkGUID.ToString(), Object ? *Object->GetPathName() : TEXT( "NULL" ), OldObject ? *OldObject->GetPathName() : TEXT( "NULL" ) );
 			}
 
@@ -1702,26 +1732,26 @@ UObject * FNetGUIDCache::GetObjectFromNetGUID( const FNetworkGUID & NetGUID )
 	if ( !ensure( !NetGUID.IsDefault() ) )
 	{
 		return NULL;
-}
+	}
 
 	FNetGuidCacheObject * CacheObjectPtr = ObjectLookup.Find( NetGUID );
 
 	if ( CacheObjectPtr == NULL )
-{
+	{
 		// We don't have the object mapped yet
 		return NULL;
-}
+	}
 
 	UObject * Object = CacheObjectPtr->Object.Get();
 
 	if ( Object != NULL )
-{
+	{
 		check( Object->GetPathName() == CacheObjectPtr->FullPath );
 		return Object;
-}
+	}
 
 	if ( CacheObjectPtr->FullPath.IsEmpty() )
-{
+	{
 		// This probably isn't possible, but making it a warning
 		UE_LOG( LogNetPackageMap, Warning, TEXT( "GetObjectFromNetGUID: No full path for %s" ), *NetGUID.ToString() );
 		return NULL;
@@ -1753,14 +1783,14 @@ UObject * FNetGUIDCache::GetObjectFromNetGUID( const FNetworkGUID & NetGUID )
 			return NULL;
 		}
 		else
-	{
+		{
 			UE_LOG( LogNetPackageMap, Log, TEXT( "GetObjectFromNetGUID: Could not find Object for: NetGUID <%s, %s>"), *NetGUID.ToString(), *CacheObjectPtr->FullPath );
-	}
+		}
 
 		Object = StaticLoadObject( UObject::StaticClass(), NULL, *CacheObjectPtr->FullPath, NULL, LOAD_NoWarn );
 
 		if ( Object )
-	{
+		{
 			UE_LOG( LogNetPackageMap, Log, TEXT( "GetObjectFromNetGUID: StaticLoadObject. Found: %s" ), Object ? *Object->GetName() : TEXT( "NULL" ) );
 		}
 		else
diff --git a/Engine/Source/Runtime/Engine/Private/PendingNetGame.cpp b/Engine/Source/Runtime/Engine/Private/PendingNetGame.cpp
--- a/Engine/Source/Runtime/Engine/Private/PendingNetGame.cpp
+++ b/Engine/Source/Runtime/Engine/Private/PendingNetGame.cpp
@@ -325,7 +325,7 @@ void UPendingNetGame::NotifyControlMessage(UNetConnection* Connection, uint8 Mes
 			FNetworkGUID NetGUID;
 			FString Path;
 			FNetControlMessage<NMT_NetGUIDAssign>::Receive(Bunch, NetGUID, Path);
-			NetDriver->ServerConnection->PackageMap->ResolvePathAndAssignNetGUID( NetGUID, Path, TEXT(""), NULL );
+			NetDriver->ServerConnection->PackageMap->ResolvePathAndAssignNetGUID(NetGUID, Path, NULL);
 			break;
 		}
 		default:
diff --git a/Engine/Source/Runtime/Engine/Private/World.cpp b/Engine/Source/Runtime/Engine/Private/World.cpp
--- a/Engine/Source/Runtime/Engine/Private/World.cpp
+++ b/Engine/Source/Runtime/Engine/Private/World.cpp
@@ -1711,22 +1711,22 @@ void UWorld::AddToWorld( ULevel* Level, const FTransform& LevelTransform )
 		// notify server that the client has finished making this level visible
 		if (!Level->bClientOnlyVisible)
 		{
-		for (FLocalPlayerIterator It(GEngine, this); It; ++It)
-		{
-			if (It->PlayerController != NULL)
+			for (FLocalPlayerIterator It(GEngine, this); It; ++It)
 			{
-				// Remap packagename for PIE networking before sending out to server
-				FName PackageName = Level->GetOutermost()->GetFName();
-				FString PackageNameStr = PackageName.ToString();
-				if (GEngine->NetworkRemapPath(this, PackageNameStr, false))
+				if (It->PlayerController != NULL)
 				{
-					PackageName = FName(*PackageNameStr);
-				}
+					// Remap packagename for PIE networking before sending out to server
+					FName PackageName = Level->GetOutermost()->GetFName();
+					FString PackageNameStr = PackageName.ToString();
+					if (GEngine->NetworkRemapPath(this, PackageNameStr, false))
+					{
+						PackageName = FName(*PackageNameStr);
+					}
 
-				It->PlayerController->ServerUpdateLevelVisibility(PackageName, true);
+					It->PlayerController->ServerUpdateLevelVisibility(PackageName, true);
+				}
 			}
 		}
-		}
 
 		Level->InitializeRenderingResources();
 
@@ -1796,22 +1796,22 @@ void UWorld::RemoveFromWorld( ULevel* Level )
 		// notify server that the client has removed this level
 		if (!Level->bClientOnlyVisible)
 		{
-		for (FLocalPlayerIterator It(GEngine, this); It; ++It)
-		{
-			if (It->PlayerController != NULL)
+			for (FLocalPlayerIterator It(GEngine, this); It; ++It)
 			{
-				// Remap packagename for PIE networking before sending out to server
-				FName PackageName = Level->GetOutermost()->GetFName();
-				FString PackageNameStr = PackageName.ToString();
-				if (GEngine->NetworkRemapPath(this, PackageNameStr, false))
+				if (It->PlayerController != NULL)
 				{
-					PackageName = FName(*PackageNameStr);
-				}
+					// Remap packagename for PIE networking before sending out to server
+					FName PackageName = Level->GetOutermost()->GetFName();
+					FString PackageNameStr = PackageName.ToString();
+					if (GEngine->NetworkRemapPath(this, PackageNameStr, false))
+					{
+						PackageName = FName(*PackageNameStr);
+					}
 
-				It->PlayerController->ServerUpdateLevelVisibility(PackageName, false);
+					It->PlayerController->ServerUpdateLevelVisibility(PackageName, false);
+				}
 			}
 		}
-		}
 		
 		// Notify world composition: will place a level at original position
 		if (WorldComposition)
@@ -3409,7 +3409,7 @@ void UWorld::NotifyControlMessage(UNetConnection* Connection, uint8 MessageType,
 				FNetControlMessage<NMT_NetGUIDAssign>::Receive(Bunch, NetGUID, Path);
 
 				UE_LOG(LogNet, Verbose, TEXT("NMT_NetGUIDAssign  NetGUID %s. Path: %s. "), *NetGUID.ToString(), *Path );
-				Connection->PackageMap->ResolvePathAndAssignNetGUID( NetGUID, Path, TEXT(""), NULL );
+				Connection->PackageMap->ResolvePathAndAssignNetGUID( NetGUID, Path, NULL );
 				break;
 			}
 		}
