Oops, didn't mean to check in warning spam

diff --git a/Engine/Source/Runtime/Engine/Private/PackageMapClient.cpp b/Engine/Source/Runtime/Engine/Private/PackageMapClient.cpp
--- a/Engine/Source/Runtime/Engine/Private/PackageMapClient.cpp
+++ b/Engine/Source/Runtime/Engine/Private/PackageMapClient.cpp
@@ -599,7 +599,7 @@ FNetworkGUID UPackageMapClient::InternalLoadObject( FArchive & Ar, UObject *& Ob
 		// Make note if the outer is still deferred (we then must defer the load of this guid)
 		const bool OuterIsDeferred = PendingResolveGUIDs.Contains( OuterGUID );
 
-#if 0//STRESS_PENDING_GUIDS == 0
+#if STRESS_PENDING_GUIDS == 0
 		if ( Object == NULL && !OuterIsDeferred )
 		{
 			// Try to resolve the name (we may have already loaded it), but DON'T load yet if we haven't already
@@ -626,7 +626,7 @@ FNetworkGUID UPackageMapClient::InternalLoadObject( FArchive & Ar, UObject *& Ob
 
 		if ( bShouldDefer )
 		{
-			UE_LOG( LogNetPackageMap, Warning, TEXT( "InternalLoadObject: Deferring resolve. Path: %s, NetGUID: %s, Outer: %s" ), *PathName, *NetGUID.ToString(), ObjOuter != NULL ? *ObjOuter->GetFullName() : TEXT( "NULL" ) );
+			UE_LOG( LogNetPackageMap, Log, TEXT( "InternalLoadObject: Deferring resolve. Path: %s, NetGUID: %s, Outer: %s" ), *PathName, *NetGUID.ToString(), ObjOuter != NULL ? *ObjOuter->GetFullName() : TEXT( "NULL" ) );
 			ResolvePathAndAssignNetGUID_Deferred( NetGUID, PathName, OuterGUID, bIsForLevel );
 		}
 		else
@@ -711,7 +711,7 @@ void UPackageMapClient::UpdatePendingResolveGUIDs()
 
 			if ( Object != NULL )
 			{
-				UE_LOG( LogNetPackageMap, Warning, TEXT( "UPackageMapClient::UpdatePendingResolveGUIDs. Resolved deferred NetGUID. Path: %s, NetGUID: %s. NumLeft: %i, IsSerializingNewActor: %s" ), *PendingResolveGUID.PathName, *It.Key().ToString(), PendingResolveGUIDs.Num() - 1, IsSerializingNewActor ? TEXT("YES") : TEXT("NO") );
+				UE_LOG( LogNetPackageMap, Log, TEXT( "UPackageMapClient::UpdatePendingResolveGUIDs. Resolved deferred NetGUID. Path: %s, NetGUID: %s. NumLeft: %i, IsSerializingNewActor: %s" ), *PendingResolveGUID.PathName, *It.Key().ToString(), PendingResolveGUIDs.Num() - 1, IsSerializingNewActor ? TEXT("YES") : TEXT("NO") );
 
 				// If we have an object, we can make another pass in case there were dependencies on this guid loading
 				NeedToKeepResolving = true;
