UE4: CIS fix

diff --git a/Engine/Source/Runtime/Engine/Private/World.cpp b/Engine/Source/Runtime/Engine/Private/World.cpp
--- a/Engine/Source/Runtime/Engine/Private/World.cpp
+++ b/Engine/Source/Runtime/Engine/Private/World.cpp
@@ -302,6 +302,7 @@ void UWorld::PostDuplicate(bool bDuplicateForPIE)
 			}
 		}
 
+#if WITH_EDITOR
 		// Duplicate the level script blueprint generated classes as well
 		const bool bDontCreate = true;
 		UBlueprint* LevelScriptBlueprint = PersistentLevel ? PersistentLevel->GetLevelScriptBlueprint(bDontCreate) : NULL;
@@ -321,6 +322,7 @@ void UWorld::PostDuplicate(bool bDuplicateForPIE)
 				ReplacementMap.Add(OldSkeletonClass, NewSkeletonClass);
 			}
 		}
+#endif // WITH_EDITOR
 
 		// Now replace references from the old textures/classes to the new ones, if any were duplicated
 		if (ReplacementMap.Num() > 0)
