UE4: Fixed up some issues related to matching more light/shadow maps than desired. You can now specify which level you want.

diff --git a/Engine/Source/Editor/DetailCustomizations/Private/WorldSettingsDetails.cpp b/Engine/Source/Editor/DetailCustomizations/Private/WorldSettingsDetails.cpp
--- a/Engine/Source/Editor/DetailCustomizations/Private/WorldSettingsDetails.cpp
+++ b/Engine/Source/Editor/DetailCustomizations/Private/WorldSettingsDetails.cpp
@@ -64,6 +64,7 @@ FLightmapCustomNodeBuilder::~FLightmapCustomNodeBuilder()
 {
 	FEditorDelegates::OnLightingBuildKept.RemoveAll(this);
 	FEditorDelegates::MapChange.RemoveAll(this);
+	FEditorDelegates::NewCurrentLevel.RemoveAll(this);
 }
 
 
@@ -73,6 +74,7 @@ void FLightmapCustomNodeBuilder::SetOnRebuildChildren(FSimpleDelegate InOnRegene
 
 	FEditorDelegates::OnLightingBuildKept.AddSP(this, &FLightmapCustomNodeBuilder::HandleLightingBuildKept);
 	FEditorDelegates::MapChange.AddSP(this, &FLightmapCustomNodeBuilder::HandleMapChanged);
+	FEditorDelegates::NewCurrentLevel.AddSP(this, &FLightmapCustomNodeBuilder::HandleNewCurrentLevel);
 }
 
 
@@ -129,6 +131,12 @@ void FLightmapCustomNodeBuilder::HandleMapChanged(uint32 MapChangeFlags)
 }
 
 
+void FLightmapCustomNodeBuilder::HandleNewCurrentLevel()
+{
+	OnRegenerateChildren.ExecuteIfBound();
+}
+
+
 TSharedRef<ITableRow> FLightmapCustomNodeBuilder::MakeLightMapListViewWidget(TSharedPtr<FLightmapItem> LightMapItem, const TSharedRef<STableViewBase>& OwnerTable)
 {
 	if ( !ensure(LightMapItem.IsValid()) )
@@ -254,7 +262,7 @@ void FLightmapCustomNodeBuilder::RefreshLightmapItems()
 	if ( World )
 	{
 		TArray<UTexture2D*> LightMapsAndShadowMaps;
-		World->GetLightMapsAndShadowMaps(LightMapsAndShadowMaps);
+		World->GetLightMapsAndShadowMaps(World->GetCurrentLevel(), LightMapsAndShadowMaps);
 
 		for ( auto ObjIt = LightMapsAndShadowMaps.CreateConstIterator(); ObjIt; ++ObjIt )
 		{
diff --git a/Engine/Source/Editor/DetailCustomizations/Private/WorldSettingsDetails.h b/Engine/Source/Editor/DetailCustomizations/Private/WorldSettingsDetails.h
--- a/Engine/Source/Editor/DetailCustomizations/Private/WorldSettingsDetails.h
+++ b/Engine/Source/Editor/DetailCustomizations/Private/WorldSettingsDetails.h
@@ -109,6 +109,9 @@ private:
 	/** Handler for when the map changed or was rebuilt */
 	void HandleMapChanged(uint32 MapChangeFlags);
 
+	/** Handler for when the current level changes */
+	void HandleNewCurrentLevel();
+
 	/** Handler for light map list view widget creation */
 	TSharedRef<ITableRow> MakeLightMapListViewWidget(TSharedPtr<FLightmapItem> AssetItem, const TSharedRef<STableViewBase>& OwnerTable);
 
diff --git a/Engine/Source/Runtime/Engine/Classes/Engine/World.h b/Engine/Source/Runtime/Engine/Classes/Engine/World.h
--- a/Engine/Source/Runtime/Engine/Classes/Engine/World.h
+++ b/Engine/Source/Runtime/Engine/Classes/Engine/World.h
@@ -2412,8 +2412,8 @@ public:
 	/** Updates all physics constraint actor joint locations.  */
 	virtual void UpdateConstraintActors();
 
-	/** Gets all LightMaps and ShadowMaps associated with this world's persistent level */
-	void GetLightMapsAndShadowMaps(TArray<UTexture2D*>& OutLightMapsAndShadowMaps);
+	/** Gets all LightMaps and ShadowMaps associated with this world. Specify the level or leave null for persistent */
+	void GetLightMapsAndShadowMaps(ULevel* Level, TArray<UTexture2D*>& OutLightMapsAndShadowMaps);
 
 public:
 	static FString ConvertToPIEPackageName(const FString& PackageName, int32 PIEInstanceID);
diff --git a/Engine/Source/Runtime/Engine/Private/World.cpp b/Engine/Source/Runtime/Engine/Private/World.cpp
--- a/Engine/Source/Runtime/Engine/Private/World.cpp
+++ b/Engine/Source/Runtime/Engine/Private/World.cpp
@@ -215,7 +215,7 @@ bool UWorld::Rename(const TCHAR* InName, UObject* NewOuter, ERenameFlags Flags)
 {
 	// Rename LightMaps and ShadowMaps to the new location. Keep the old name, since they are not named after the world.
 	TArray<UTexture2D*> LightMapsAndShadowMaps;
-	GetLightMapsAndShadowMaps(LightMapsAndShadowMaps);
+	GetLightMapsAndShadowMaps(PersistentLevel, LightMapsAndShadowMaps);
 	for (auto* Tex : LightMapsAndShadowMaps)
 	{
 		if ( Tex && !Tex->Rename(*Tex->GetName(), NewOuter, Flags) )
@@ -290,7 +290,7 @@ void UWorld::PostDuplicate(bool bDuplicateForPIE)
 
 		// Gather the textures
 		TArray<UTexture2D*> LightMapsAndShadowMaps;
-		GetLightMapsAndShadowMaps(LightMapsAndShadowMaps);
+		GetLightMapsAndShadowMaps(PersistentLevel, LightMapsAndShadowMaps);
 
 		// Duplicate the textures, if any
 		for (auto* Tex : LightMapsAndShadowMaps)
@@ -5032,7 +5032,7 @@ void UWorld::CopyGameState(AGameMode* FromGameMode, AGameState* FromGameState)
 	GameState = FromGameState;
 }
 
-void UWorld::GetLightMapsAndShadowMaps(TArray<UTexture2D*>& OutLightMapsAndShadowMaps)
+void UWorld::GetLightMapsAndShadowMaps(ULevel* Level, TArray<UTexture2D*>& OutLightMapsAndShadowMaps)
 {
 	class FFindLightmapsArchive : public FArchiveUObject
 	{
@@ -5056,8 +5056,8 @@ void UWorld::GetLightMapsAndShadowMaps(TArray<UTexture2D*>& OutLightMapsAndShado
 
 		FArchive& operator<<(class UObject*& Obj)
 		{
-			// Don't check null references or objects already visited
-			if (Obj != NULL && Obj->HasAnyMarks(OBJECTMARK_TagExp))
+			// Don't check null references or objects already visited. Also, skip UWorlds as they will pull in more levels than desired
+			if (Obj != NULL && Obj->HasAnyMarks(OBJECTMARK_TagExp) && !Obj->IsA(UWorld::StaticClass()))
 			{
 				if (Obj->IsA(ULightMapTexture2D::StaticClass()) || Obj->IsA(UShadowMapTexture2D::StaticClass()))
 				{
@@ -5076,7 +5076,13 @@ void UWorld::GetLightMapsAndShadowMaps(TArray<UTexture2D*>& OutLightMapsAndShado
 		}
 	};
 
-	FFindLightmapsArchive FindArchive(this, OutLightMapsAndShadowMaps);
+	UObject* SearchObject = Level;
+	if ( !SearchObject )
+	{
+		SearchObject = PersistentLevel;
+	}
+
+	FFindLightmapsArchive FindArchive(SearchObject, OutLightMapsAndShadowMaps);
 }
 
 void UWorld::ChangeFeatureLevel(ERHIFeatureLevel::Type InFeatureLevel)
