Changes to reduce headers pulled in when including basic Engine headers directly (Actor.h, PrimitiveComponent.h, MaterialInstance.h etc).
Break pieces of SceneManagement.h into their own new headers (SceneTypes.h, SceneView.h, SceneInterface.h, PrimitiveViewRelevance.h, PrimitiveUniformShaderParameters.h, MeshBatch.h, FinalPostProcessSettings.h, BlendableManager.h)
Break StaticParameterSet and static param classes from MaterialShared.h into new StaticParameterSet.h
Move FLight/ShadowMapRef definition from Light/ShadowMap.h into SceneTypes.h, to avoid needing to include them globally
Move FPhysicsVolumeChanged delegate definition to SceneComponent.h, to avoid including Components.h
Also fix some non-unity compile errors in AI module

diff --git a/Engine/Source/Runtime/AIModule/Private/AIModulePrivate.h b/Engine/Source/Runtime/AIModule/Private/AIModulePrivate.h
--- a/Engine/Source/Runtime/AIModule/Private/AIModulePrivate.h
+++ b/Engine/Source/Runtime/AIModule/Private/AIModulePrivate.h
@@ -20,7 +20,6 @@
 #include "Components/ActorComponent.h"
 #include "GameFramework/Actor.h"
 #include "CollisionQueryParams.h"
-#include "WorldCollision.h"
 #include "Components/PrimitiveComponent.h"
 #include "GameFramework/PlayerController.h"
 #include "AI/Navigation/RecastNavMesh.h"
diff --git a/Engine/Source/Runtime/Engine/Classes/AI/Navigation/NavRelevantComponent.h b/Engine/Source/Runtime/Engine/Classes/AI/Navigation/NavRelevantComponent.h
--- a/Engine/Source/Runtime/Engine/Classes/AI/Navigation/NavRelevantComponent.h
+++ b/Engine/Source/Runtime/Engine/Classes/AI/Navigation/NavRelevantComponent.h
@@ -1,6 +1,7 @@
 // Copyright 1998-2014 Epic Games, Inc. All Rights Reserved.
 
 #pragma once
+#include "Components/ActorComponent.h"
 #include "NavRelevantComponent.generated.h"
 
 UCLASS()
diff --git a/Engine/Source/Runtime/Engine/Classes/AI/Navigation/RecastNavMesh.h b/Engine/Source/Runtime/Engine/Classes/AI/Navigation/RecastNavMesh.h
--- a/Engine/Source/Runtime/Engine/Classes/AI/Navigation/RecastNavMesh.h
+++ b/Engine/Source/Runtime/Engine/Classes/AI/Navigation/RecastNavMesh.h
@@ -37,6 +37,7 @@ class FPImplRecastNavMesh;
 class FRecastQueryFilter;
 class INavLinkCustomInterface;
 class UNavArea;
+class UNavigationSystem;
 
 UENUM()
 namespace ERecastPartitioning
diff --git a/Engine/Source/Runtime/Engine/Classes/Camera/CameraTypes.h b/Engine/Source/Runtime/Engine/Classes/Camera/CameraTypes.h
--- a/Engine/Source/Runtime/Engine/Classes/Camera/CameraTypes.h
+++ b/Engine/Source/Runtime/Engine/Classes/Camera/CameraTypes.h
@@ -2,7 +2,7 @@
 
 #pragma once
 
-#include "Runtime/Engine/Classes/Engine/Scene.h"
+#include "Engine/Scene.h"
 #include "CameraTypes.generated.h"
 
 //@TODO: Document
diff --git a/Engine/Source/Runtime/Engine/Classes/Components/PrimitiveComponent.h b/Engine/Source/Runtime/Engine/Classes/Components/PrimitiveComponent.h
--- a/Engine/Source/Runtime/Engine/Classes/Components/PrimitiveComponent.h
+++ b/Engine/Source/Runtime/Engine/Classes/Components/PrimitiveComponent.h
@@ -2,13 +2,16 @@
 
 #pragma once
 
-#include "Runtime/Engine/Classes/PhysicsEngine/BodyInstance.h"
-#include "PrimitiveSceneProxy.h"
 #include "StaticArray.h"
+#include "PhysicsEngine/BodyInstance.h"
 #include "Components/SceneComponent.h"
+#include "SceneTypes.h"
+#include "CollisionQueryParams.h"
 
 #include "PrimitiveComponent.generated.h"
 
+class FPrimitiveSceneProxy;
+struct FPrimitiveViewRelevance;
 
 /** Information about a vertex of a primitive's triangle. */
 struct FPrimitiveTriangleVertex
@@ -77,7 +80,7 @@ namespace EHasCustomNavigableGeometry
 
 /** Mirrored from Scene.h */
 USTRUCT()
-struct FMaterialRelevance
+struct ENGINE_API FMaterialRelevance
 {
 	GENERATED_USTRUCT_BODY()
 
@@ -113,37 +116,29 @@ struct FMaterialRelevance
 		,	bDisableDepthTest(false)
 		{}
 
-		/** Bitwise OR operator.  Sets any relevance bits which are present in either FMaterialRelevance. */
-		FMaterialRelevance& operator|=(const FMaterialRelevance& B)
-		{
-			bOpaque |= B.bOpaque;
-			bMasked |= B.bMasked;
-			bDistortion |= B.bDistortion;
-			bUsesSceneColor |= B.bUsesSceneColor;
-			bSeparateTranslucency |= B.bSeparateTranslucency;
-			bNormalTranslucency |= B.bNormalTranslucency;
-			bDisableDepthTest |= B.bDisableDepthTest;
-			return *this;
-		}
+	/** Bitwise OR operator.  Sets any relevance bits which are present in either FMaterialRelevance. */
+	FMaterialRelevance& operator|=(const FMaterialRelevance& B)
+	{
+		bOpaque |= B.bOpaque;
+		bMasked |= B.bMasked;
+		bDistortion |= B.bDistortion;
+		bUsesSceneColor |= B.bUsesSceneColor;
+		bSeparateTranslucency |= B.bSeparateTranslucency;
+		bNormalTranslucency |= B.bNormalTranslucency;
+		bDisableDepthTest |= B.bDisableDepthTest;
+		return *this;
+	}
 	
-		/** Binary bitwise OR operator. */
-		friend FMaterialRelevance operator|(const FMaterialRelevance& A,const FMaterialRelevance& B)
-		{
-			FMaterialRelevance Result(A);
-			Result |= B;
-			return Result;
-		}
-
-		/** Copies the material's relevance flags to a primitive's view relevance flags. */
-		void SetPrimitiveViewRelevance(FPrimitiveViewRelevance& OutViewRelevance) const
-		{
-			OutViewRelevance.bOpaqueRelevance = bOpaque;
-			OutViewRelevance.bMaskedRelevance = bMasked;
-			OutViewRelevance.bDistortionRelevance = bDistortion;
-			OutViewRelevance.bSceneColorRelevance = bUsesSceneColor;
-			OutViewRelevance.bSeparateTranslucencyRelevance = bSeparateTranslucency;
-			OutViewRelevance.bNormalTranslucencyRelevance = bNormalTranslucency;
-		}
+	/** Binary bitwise OR operator. */
+	friend FMaterialRelevance operator|(const FMaterialRelevance& A,const FMaterialRelevance& B)
+	{
+		FMaterialRelevance Result(A);
+		Result |= B;
+		return Result;
+	}
+
+	/** Copies the material's relevance flags to a primitive's view relevance flags. */
+	void SetPrimitiveViewRelevance(FPrimitiveViewRelevance& OutViewRelevance) const;
 };
 
 /** Information about the sprite category */
diff --git a/Engine/Source/Runtime/Engine/Classes/Components/SceneComponent.h b/Engine/Source/Runtime/Engine/Classes/Components/SceneComponent.h
--- a/Engine/Source/Runtime/Engine/Classes/Components/SceneComponent.h
+++ b/Engine/Source/Runtime/Engine/Classes/Components/SceneComponent.h
@@ -2,7 +2,7 @@
 
 #pragma once
 
-#include "Components.h"
+#include "ActorComponent.h"
 #include "SceneComponent.generated.h"
 
 /** Overlap info consisting of the primitive and the body that is overlapping */
@@ -71,6 +71,8 @@ FORCEINLINE EMoveComponentFlags operator&(EMoveComponentFlags Arg1,EMoveComponen
 FORCEINLINE void operator&=(EMoveComponentFlags& Dest,EMoveComponentFlags Arg)					{ Dest = EMoveComponentFlags(Dest & Arg); }
 FORCEINLINE void operator|=(EMoveComponentFlags& Dest,EMoveComponentFlags Arg)					{ Dest = EMoveComponentFlags(Dest | Arg); }
 
+DECLARE_DELEGATE_OneParam(FPhysicsVolumeChanged, class APhysicsVolume*);
+
 
 /** A SceneComponent has a transform and supports attachment, but has no rendering or collision capabilities. Useful as a 'dummy' component in the hierarchy to offset others. */
 UCLASS(ClassGroup=Utility, BlueprintType, HideCategories=(Trigger), meta=(BlueprintSpawnableComponent))
diff --git a/Engine/Source/Runtime/Engine/Classes/Components/StaticMeshComponent.h b/Engine/Source/Runtime/Engine/Classes/Components/StaticMeshComponent.h
--- a/Engine/Source/Runtime/Engine/Classes/Components/StaticMeshComponent.h
+++ b/Engine/Source/Runtime/Engine/Classes/Components/StaticMeshComponent.h
@@ -2,10 +2,13 @@
 
 
 #pragma once
-#include "LightMap.h"
-#include "ShadowMap.h"
+#include "SceneTypes.h"
+#include "Components/MeshComponent.h"
+#include "Runtime/RenderCore/Public/RenderUtils.h"
 #include "StaticMeshComponent.generated.h"
 
+class FColorVertexBuffer;
+
 /** Cached vertex information at the time the mesh was painted. */
 USTRUCT()
 struct FPaintedVertex
diff --git a/Engine/Source/Runtime/Engine/Classes/Debug/DebugDrawService.h b/Engine/Source/Runtime/Engine/Classes/Debug/DebugDrawService.h
--- a/Engine/Source/Runtime/Engine/Classes/Debug/DebugDrawService.h
+++ b/Engine/Source/Runtime/Engine/Classes/Debug/DebugDrawService.h
@@ -1,8 +1,11 @@
 // Copyright 1998-2014 Epic Games, Inc. All Rights Reserved.
 #pragma once
 
+#include "ShowFlags.h"
 #include "DebugDrawService.generated.h"
 
+class FSceneView;
+
 /** 
  * 
  */
diff --git a/Engine/Source/Runtime/Engine/Classes/Engine/Canvas.h b/Engine/Source/Runtime/Engine/Classes/Engine/Canvas.h
--- a/Engine/Source/Runtime/Engine/Classes/Engine/Canvas.h
+++ b/Engine/Source/Runtime/Engine/Classes/Engine/Canvas.h
@@ -8,6 +8,8 @@
 #include "WordWrapper.h"
 #include "Canvas.generated.h"
 
+class FSceneView;
+
 /**
  * Holds texture information with UV coordinates as well.
  */
diff --git a/Engine/Source/Runtime/Engine/Classes/Engine/Engine.h b/Engine/Source/Runtime/Engine/Classes/Engine/Engine.h
--- a/Engine/Source/Runtime/Engine/Classes/Engine/Engine.h
+++ b/Engine/Source/Runtime/Engine/Classes/Engine/Engine.h
@@ -5,10 +5,14 @@
 //=============================================================================
 
 #pragma once
+#include "World.h"
 #include "Engine.generated.h"
 
 class FScreenSaverInhibitor;
 class UDeviceProfileManager;
+class FViewport;
+class FCommonViewportClient;
+class FCanvas;
 
 UENUM()
 enum EFullyLoadPackageType
diff --git a/Engine/Source/Runtime/Engine/Classes/Engine/Font.h b/Engine/Source/Runtime/Engine/Classes/Engine/Font.h
--- a/Engine/Source/Runtime/Engine/Classes/Engine/Font.h
+++ b/Engine/Source/Runtime/Engine/Classes/Engine/Font.h
@@ -1,6 +1,7 @@
 // Copyright 1998-2014 Epic Games, Inc. All Rights Reserved.
 
 #pragma once
+#include "Engine/FontImportOptions.h"
 #include "Font.generated.h"
 
 /** this struct is serialized using native serialization so any changes to it require a package version bump */
diff --git a/Engine/Source/Runtime/Engine/Classes/Engine/StaticMesh.h b/Engine/Source/Runtime/Engine/Classes/Engine/StaticMesh.h
--- a/Engine/Source/Runtime/Engine/Classes/Engine/StaticMesh.h
+++ b/Engine/Source/Runtime/Engine/Classes/Engine/StaticMesh.h
@@ -1,6 +1,8 @@
 // Copyright 1998-2014 Epic Games, Inc. All Rights Reserved.
 
 #pragma once
+#include "Interfaces/Interface_CollisionDataProvider.h"
+#include "Interfaces/Interface_AssetUserData.h"
 #include "StaticMesh.generated.h"
 
 /** The maximum number of static mesh LODs allowed. */
diff --git a/Engine/Source/Runtime/Engine/Classes/Engine/World.h b/Engine/Source/Runtime/Engine/Classes/Engine/World.h
--- a/Engine/Source/Runtime/Engine/Classes/Engine/World.h
+++ b/Engine/Source/Runtime/Engine/Classes/Engine/World.h
@@ -10,6 +10,9 @@
 #include "WorldCollision.h"
 #include "PendingNetGame.h"
 #include "GameFramework/MusicTrackDataStructures.h"
+#include "EngineDefines.h"
+#include "Engine/LatentActionManager.h"
+#include "Runtime/RHI/Public/RHIDefinitions.h"
 #include "World.generated.h"
 
 class FPhysScene;
diff --git a/Engine/Source/Runtime/Engine/Classes/GameFramework/Actor.h b/Engine/Source/Runtime/Engine/Classes/GameFramework/Actor.h
--- a/Engine/Source/Runtime/Engine/Classes/GameFramework/Actor.h
+++ b/Engine/Source/Runtime/Engine/Classes/GameFramework/Actor.h
@@ -2,10 +2,11 @@
 
 #pragma once
 
-#include "Runtime/Engine/Classes/Engine/LatentActionManager.h"
-#include "Runtime/Engine/Classes/Camera/CameraTypes.h"
+#include "Runtime/InputCore/Classes/InputCoreTypes.h"
+#include "Runtime/RenderCore/Public/RenderCommandFence.h"
+
+#include "Camera/CameraTypes.h"
 #include "Components/SceneComponent.h"
-#include "RenderingThread.h"
 #include "ComponentInstanceDataCache.h"
 
 UENUM(BlueprintType)
diff --git a/Engine/Source/Runtime/Engine/Classes/GameFramework/PlayerController.h b/Engine/Source/Runtime/Engine/Classes/GameFramework/PlayerController.h
--- a/Engine/Source/Runtime/Engine/Classes/GameFramework/PlayerController.h
+++ b/Engine/Source/Runtime/Engine/Classes/GameFramework/PlayerController.h
@@ -2,17 +2,17 @@
 
 #pragma once
 
-#include "PlayerMuteList.h"
-#include "../Camera/PlayerCameraManager.h"
-#include "GameMode.h"
+#include "GameFramework/PlayerMuteList.h"
+#include "Camera/PlayerCameraManager.h"
+#include "GameFramework/GameMode.h"
 #include "Components/InputComponent.h"
-#include "../Engine/GameViewportClient.h"
-#include "SceneManagement.h"
-#include "ForceFeedbackEffect.h"
-#include "OnlineReplStructs.h"
+#include "GameFramework/ForceFeedbackEffect.h"
+#include "GameFramework/OnlineReplStructs.h"
 #include "GameFramework/Controller.h"
 #include "PlayerController.generated.h"
 
+class FPrimitiveComponentId;
+
 /** delegate used to override default viewport audio listener position calculated from camera */
 DECLARE_DELEGATE_ThreeParams(FGetAudioListenerPos, FVector& /*Location*/, FVector& /*ProjFront*/, FVector& /*ProjRight*/);
 
diff --git a/Engine/Source/Runtime/Engine/Classes/Kismet/GameplayStatics.h b/Engine/Source/Runtime/Engine/Classes/Kismet/GameplayStatics.h
--- a/Engine/Source/Runtime/Engine/Classes/Kismet/GameplayStatics.h
+++ b/Engine/Source/Runtime/Engine/Classes/Kismet/GameplayStatics.h
@@ -2,7 +2,8 @@
 
 #pragma once
 
-#include "../Sound/DialogueTypes.h"
+#include "Sound/DialogueTypes.h"
+#include "Engine/LatentActionManager.h"
 #include "GameplayStatics.generated.h"
 
 //
diff --git a/Engine/Source/Runtime/Engine/Classes/Kismet/KismetSystemLibrary.h b/Engine/Source/Runtime/Engine/Classes/Kismet/KismetSystemLibrary.h
--- a/Engine/Source/Runtime/Engine/Classes/Kismet/KismetSystemLibrary.h
+++ b/Engine/Source/Runtime/Engine/Classes/Kismet/KismetSystemLibrary.h
@@ -2,6 +2,7 @@
 
 #pragma once
 
+#include "Engine/LatentActionManager.h"
 #include "KismetSystemLibrary.generated.h"
 
 UENUM(BlueprintType)
diff --git a/Engine/Source/Runtime/Engine/Classes/Materials/MaterialInstance.h b/Engine/Source/Runtime/Engine/Classes/Materials/MaterialInstance.h
--- a/Engine/Source/Runtime/Engine/Classes/Materials/MaterialInstance.h
+++ b/Engine/Source/Runtime/Engine/Classes/Materials/MaterialInstance.h
@@ -2,12 +2,17 @@
 
 
 #pragma once
+#include "Materials/MaterialInterface.h"
+#include "StaticParameterSet.h"
 #include "MaterialInstance.generated.h"
 
 //
 // Forward declarations.
 //
 struct FMaterialInstanceBasePropertyOverrides;
+class FMaterialShaderMap;
+class FMaterialShaderMapId;
+class FSHAHash;
 
 /** Editable font parameter. */
 USTRUCT()
@@ -36,16 +41,7 @@ struct FFontParameterValue
 
 
 	typedef const UTexture* ValueType;
-	static ValueType GetValue(const FFontParameterValue& Parameter)
-	{
-		ValueType Value = NULL;
-		if (Parameter.FontValue && Parameter.FontValue->Textures.IsValidIndex(Parameter.FontPage))
-		{
-			// get the texture for the font page
-			Value = Parameter.FontValue->Textures[Parameter.FontPage];
-		}
-		return Value;
-	}
+	static ValueType GetValue(const FFontParameterValue& Parameter);
 };
 
 /** Editable scalar parameter. */
diff --git a/Engine/Source/Runtime/Engine/Classes/Materials/MaterialInterface.h b/Engine/Source/Runtime/Engine/Classes/Materials/MaterialInterface.h
--- a/Engine/Source/Runtime/Engine/Classes/Materials/MaterialInterface.h
+++ b/Engine/Source/Runtime/Engine/Classes/Materials/MaterialInterface.h
@@ -2,11 +2,14 @@
 
 
 #pragma once
-
-#include "MaterialShared.h"
+#include "Engine/BlendableInterface.h"
+#include "Runtime/RHI/Public/RHIDefinitions.h"
+#include "SceneTypes.h"
 #include "MaterialInterface.generated.h"
 
 class UMaterial;
+class FMaterialResource;
+class FMaterialCompiler;
 
 UENUM()
 enum EMaterialUsage
diff --git a/Engine/Source/Runtime/Engine/Classes/PhysicsEngine/BodyInstance.h b/Engine/Source/Runtime/Engine/Classes/PhysicsEngine/BodyInstance.h
--- a/Engine/Source/Runtime/Engine/Classes/PhysicsEngine/BodyInstance.h
+++ b/Engine/Source/Runtime/Engine/Classes/PhysicsEngine/BodyInstance.h
@@ -8,6 +8,8 @@
 
 #define UE_WITH_PHYSICS (WITH_PHYSX || WITH_BOX2D)
 
+struct FCollisionShape;
+
 #if WITH_PHYSX
 namespace physx
 {
diff --git a/Engine/Source/Runtime/Engine/Private/AI/Navigation/NavLinkProxy.cpp b/Engine/Source/Runtime/Engine/Private/AI/Navigation/NavLinkProxy.cpp
--- a/Engine/Source/Runtime/Engine/Private/AI/Navigation/NavLinkProxy.cpp
+++ b/Engine/Source/Runtime/Engine/Private/AI/Navigation/NavLinkProxy.cpp
@@ -4,6 +4,7 @@
 #include "ObjectEditorUtils.h"
 #include "NavigationModifier.h"
 #include "Navigation/PathFollowingComponent.h"
+#include "AI/Navigation/NavLinkCustomComponent.h"
 #include "AI/Navigation/NavLinkCustomInterface.h"
 #include "AI/Navigation/NavLinkProxy.h"
 
diff --git a/Engine/Source/Runtime/Engine/Private/AI/Navigation/NavigationSystem.cpp b/Engine/Source/Runtime/Engine/Private/AI/Navigation/NavigationSystem.cpp
--- a/Engine/Source/Runtime/Engine/Private/AI/Navigation/NavigationSystem.cpp
+++ b/Engine/Source/Runtime/Engine/Private/AI/Navigation/NavigationSystem.cpp
@@ -4,6 +4,7 @@
 #include "NavDataGenerator.h"
 #include "NavigationOctree.h"
 #include "AI/Navigation/NavMeshBoundsVolume.h"
+#include "AI/Navigation/NavRelevantComponent.h"
 
 #if WITH_RECAST
 #include "RecastNavMeshGenerator.h"
diff --git a/Engine/Source/Runtime/Engine/Private/Materials/MaterialInstance.cpp b/Engine/Source/Runtime/Engine/Private/Materials/MaterialInstance.cpp
--- a/Engine/Source/Runtime/Engine/Private/Materials/MaterialInstance.cpp
+++ b/Engine/Source/Runtime/Engine/Private/Materials/MaterialInstance.cpp
@@ -12,6 +12,7 @@
 #include "MaterialInstanceSupport.h"
 #include "MaterialShader.h"
 #include "TargetPlatform.h"
+#include "Engine/Font.h"
 
 /**
  * Cache uniform expressions for the given material.
@@ -116,6 +117,17 @@ bool UpdateParameterSet(TArray<ParameterType> &Parameters, UMaterial* ParentMate
 	return bChanged;
 }
 
+FFontParameterValue::ValueType FFontParameterValue::GetValue(const FFontParameterValue& Parameter)
+{
+	ValueType Value = NULL;
+	if (Parameter.FontValue && Parameter.FontValue->Textures.IsValidIndex(Parameter.FontPage))
+	{
+		// get the texture for the font page
+		Value = Parameter.FontValue->Textures[Parameter.FontPage];
+	}
+	return Value;
+}
+
 FMaterialInstanceResource::FMaterialInstanceResource(UMaterialInstance* InOwner,bool bInSelected,bool bInHovered)
 	: FMaterialRenderProxy(bInSelected, bInHovered)
 	, Parent(NULL)
diff --git a/Engine/Source/Runtime/Engine/Private/PrimitiveComponent.cpp b/Engine/Source/Runtime/Engine/Private/PrimitiveComponent.cpp
--- a/Engine/Source/Runtime/Engine/Private/PrimitiveComponent.cpp
+++ b/Engine/Source/Runtime/Engine/Private/PrimitiveComponent.cpp
@@ -32,6 +32,21 @@ static FAutoConsoleVariable CVarAllowCachedOverlaps(
 	TEXT("AllowCachedOverlaps"), 1,
 	TEXT("0: disable cached overlaps, 1: enable\n"));
 
+
+//////////////////////////////////////////////////////////////////////////
+// FMaterialRelevance
+
+/** Copies the material's relevance flags to a primitive's view relevance flags. */
+void FMaterialRelevance::SetPrimitiveViewRelevance(FPrimitiveViewRelevance& OutViewRelevance) const
+{
+	OutViewRelevance.bOpaqueRelevance = bOpaque;
+	OutViewRelevance.bMaskedRelevance = bMasked;
+	OutViewRelevance.bDistortionRelevance = bDistortion;
+	OutViewRelevance.bSceneColorRelevance = bUsesSceneColor;
+	OutViewRelevance.bSeparateTranslucencyRelevance = bSeparateTranslucency;
+	OutViewRelevance.bNormalTranslucencyRelevance = bNormalTranslucency;
+}
+
 ///////////////////////////////////////////////////////////////////////////////
 // PRIMITIVE COMPONENT
 ///////////////////////////////////////////////////////////////////////////////
diff --git a/Engine/Source/Runtime/Engine/Private/SceneUtils.cpp b/Engine/Source/Runtime/Engine/Private/SceneUtils.cpp
--- a/Engine/Source/Runtime/Engine/Private/SceneUtils.cpp
+++ b/Engine/Source/Runtime/Engine/Private/SceneUtils.cpp
@@ -1,6 +1,7 @@
 // Copyright 1998-2014 Epic Games, Inc. All Rights Reserved.
 
 #include "EnginePrivate.h"
+#include "DynamicRHI.h"
 
 #if WANTS_DRAW_MESH_EVENTS && PLATFORM_SUPPORTS_DRAW_MESH_EVENTS
 
@@ -9,4 +10,24 @@ bool FDrawEvent::IsInRenderingThread_Internal()
 	return IsInRenderingThread();
 }
 
+FDrawEvent::~FDrawEvent()
+{
+	if (bDrawEventHasBeenEmitted)
+	{
+		GDynamicRHI->PopEvent();
+	}
+}
+
+void FDrawEvent::Start(const FColor& Color, const TCHAR* Fmt, ...)
+{
+	check(IsInRenderingThread_Internal());
+	va_list ptr;
+	va_start(ptr, Fmt);
+	TCHAR TempStr[256];
+	// Build the string in the temp buffer
+	FCString::GetVarArgs(TempStr, ARRAY_COUNT(TempStr), ARRAY_COUNT(TempStr) - 1, Fmt, ptr);
+	GDynamicRHI->PushEvent(TempStr);
+	bDrawEventHasBeenEmitted = true;
+}
+
 #endif // WANTS_DRAW_MESH_EVENTS && PLATFORM_SUPPORTS_DRAW_MESH_EVENTS
\ No newline at end of file
diff --git a/Engine/Source/Runtime/Engine/Private/SceneView.cpp b/Engine/Source/Runtime/Engine/Private/SceneView.cpp
--- a/Engine/Source/Runtime/Engine/Private/SceneView.cpp
+++ b/Engine/Source/Runtime/Engine/Private/SceneView.cpp
@@ -5,6 +5,7 @@
 =============================================================================*/
 
 #include "EnginePrivate.h"
+#include "SceneView.h"
 #include "SceneManagement.h"
 #include "EngineModule.h"
 #include "StereoRendering.h"
diff --git a/Engine/Source/Runtime/Engine/Public/BlendableManager.h b/Engine/Source/Runtime/Engine/Public/BlendableManager.h
new file mode 100644
--- /dev/null
+++ b/Engine/Source/Runtime/Engine/Public/BlendableManager.h
@@ -0,0 +1,115 @@
+// Copyright 1998-2014 Epic Games, Inc. All Rights Reserved.
+
+#pragma once
+
+// used to blend IBlendableInterface object data, the member act as a header for a following data block
+struct FBlendableEntry
+{
+	// weight 0..1, 0 to disable this entry
+	float Weight;
+
+private:
+	// to ensure type safety
+	FName BlendableType;
+	// to be able to jump over data
+	uint32 DataSize;
+
+	// @return pointer to the next object or end (can be compared with container end pointer)
+	uint8* GetDataPtr() { check(this); return (uint8*)(this + 1); }
+	// @return next or end of the array
+	FBlendableEntry* GetNext() { return (FBlendableEntry*)(GetDataPtr() + DataSize); }
+
+	friend class FBlendableManager;
+};
+
+// Manager class which allows blending of FBlendableBase objects, stores a copy (for the render thread)
+class FBlendableManager
+{
+public:
+
+	FBlendableManager()
+	{
+		// can be increased if needed (to avoid reallocations and copies at the cost of more memory allocation)
+		Scratch.Reserve(10 * 1024);
+	}
+
+	// @param InWeight 0..1, excluding 0 as this is used to disable entries
+	// @param InData is copied with a memcpy
+	template <class T>
+	void PushBlendableData(float InWeight, const T& InData)
+	{
+		FName BlendableType = T::GetFName();
+
+		PushBlendableDataPtr(InWeight, BlendableType, (const uint8*)&InData, sizeof(T));
+	}
+
+	// find next that has the given type, O(n), n is number of elements after the given one or all if started with 0
+	// @param InIterator 0 to start iteration
+	// @return 0 if no further one of that type was found
+	template <class T>
+	T* IterateBlendables(FBlendableEntry*& InIterator) const
+	{
+		FName BlendableType = T::GetFName();
+
+		do
+		{
+			InIterator = GetNextBlendableEntryPtr(InIterator);
+
+		} while (InIterator && InIterator->Weight <= 0.0f && InIterator->BlendableType != BlendableType);
+
+		if (InIterator)
+		{
+			return (T*)InIterator->GetDataPtr();
+		}
+
+		return 0;
+	}
+
+private:
+
+	// stored multiple elements with a FBlendableEntry header and following data of a size specified in the header
+	TArray<uint8> Scratch;
+
+	// find next, no restriction on the type O(n), n is number of elements after the given one or all if started with 0
+	FBlendableEntry* GetNextBlendableEntryPtr(FBlendableEntry* InIterator = 0) const
+	{
+		if (!InIterator)
+		{
+			// start at the beginning
+			InIterator = (FBlendableEntry*)Scratch.GetData();
+		}
+		else
+		{
+			// go to next
+			InIterator = InIterator->GetNext();
+		}
+
+		// end reached?
+		if ((uint8*)InIterator == Scratch.GetData() + Scratch.Num())
+		{
+			// no more entries
+			return 0;
+		}
+
+		return InIterator;
+	}
+
+	// @param InWeight 0..1, excluding 0 as this is used to disable entries
+	// @param InData is copied
+	// @param InDataSize >0
+	void PushBlendableDataPtr(float InWeight, FName InBlendableType, const uint8* InData, uint32 InDataSize)
+	{
+		check(InWeight > 0.0f && InWeight <= 1.0f);
+		check(InData);
+		check(InDataSize);
+
+		uint32 OldSize = Scratch.AddUninitialized(sizeof(FBlendableEntry)+InDataSize);
+
+		FBlendableEntry* Dst = (FBlendableEntry*)&Scratch[OldSize];
+
+		Dst->Weight = InWeight;
+		Dst->BlendableType = InBlendableType;
+		Dst->DataSize = InDataSize;
+		memcpy(Dst->GetDataPtr(), InData, InDataSize);
+	}
+};
\ No newline at end of file
diff --git a/Engine/Source/Runtime/Engine/Public/Components.h b/Engine/Source/Runtime/Engine/Public/Components.h
--- a/Engine/Source/Runtime/Engine/Public/Components.h
+++ b/Engine/Source/Runtime/Engine/Public/Components.h
@@ -164,4 +164,3 @@ private:
 	FColorVertexBuffer( const FColorVertexBuffer &rhs );
 };
 
-DECLARE_DELEGATE_OneParam( FPhysicsVolumeChanged , class APhysicsVolume* );
diff --git a/Engine/Source/Runtime/Engine/Public/DebugRenderSceneProxy.h b/Engine/Source/Runtime/Engine/Public/DebugRenderSceneProxy.h
--- a/Engine/Source/Runtime/Engine/Public/DebugRenderSceneProxy.h
+++ b/Engine/Source/Runtime/Engine/Public/DebugRenderSceneProxy.h
@@ -10,6 +10,7 @@
 #define _INC_DEBUGRENDERSCENEPROXY
 
 #include "Debug/DebugDrawService.h"
+#include "PrimitiveSceneProxy.h"
 
 class FDebugRenderSceneProxy : public FPrimitiveSceneProxy
 {
diff --git a/Engine/Source/Runtime/Engine/Public/FinalPostProcessSettings.h b/Engine/Source/Runtime/Engine/Public/FinalPostProcessSettings.h
new file mode 100644
--- /dev/null
+++ b/Engine/Source/Runtime/Engine/Public/FinalPostProcessSettings.h
@@ -0,0 +1,174 @@
+// Copyright 1998-2014 Epic Games, Inc. All Rights Reserved.
+
+#pragma once
+
+#include "Engine/Scene.h"
+#include "BlendableManager.h"
+
+/** All blended postprocessing in one place, non lerpable data is stored in non merged form */
+class FFinalPostProcessSettings : public FPostProcessSettings
+{
+public:
+	FFinalPostProcessSettings()
+	: HighResScreenshotMaterial(NULL)
+	, HighResScreenshotMaskMaterial(NULL)
+	, HighResScreenshotCaptureRegionMaterial(NULL)
+	, bBufferVisualizationDumpRequired(false)
+	{
+		// to avoid reallocations we reserve some elements
+
+		ContributingCubemaps.Empty(8);
+		ContributingLUTs.Empty(8);
+
+		SetLUT(0);
+	}
+	
+	struct FCubemapEntry
+	{
+		// 0..
+		FLinearColor AmbientCubemapTintMulScaleValue;
+		// can be 0
+		class UTexture* AmbientCubemap;
+
+		FCubemapEntry() :
+			AmbientCubemapTintMulScaleValue(FLinearColor(0, 0, 0, 0)),
+			AmbientCubemap(NULL)
+		{}
+	};
+
+	struct FLUTBlenderEntry
+	{
+		// 0..1
+		float Weight;
+		// can be 0
+		class UTexture* LUTTexture;
+	};
+
+	// Updated existing one or creates a new one
+	// this allows to combined volumes for fading between them but also to darken/disable volumes in some areas
+	void UpdateEntry(const FCubemapEntry &Entry, float Weight)
+	{
+		bool Existing = false;
+		for(int32 i = 0; i < ContributingCubemaps.Num(); ++i)
+		{
+			FCubemapEntry& Local = ContributingCubemaps[i];
+			
+			if(Local.AmbientCubemap == Entry.AmbientCubemap)
+			{
+				Local.AmbientCubemapTintMulScaleValue = FMath::Lerp(Local.AmbientCubemapTintMulScaleValue, Entry.AmbientCubemapTintMulScaleValue, Weight);
+				Existing = true;
+			}
+			else
+			{
+				Local.AmbientCubemapTintMulScaleValue *= 1.0f - Weight;
+			}
+
+			if(Local.AmbientCubemapTintMulScaleValue.IsAlmostBlack())
+			{
+				ContributingCubemaps.RemoveAt(i);
+			}
+		}
+
+		if( !Existing )
+		{
+			// We didn't find the entry previously, so we need to Lerp up from zero.
+			FCubemapEntry WeightedEntry = Entry;
+			WeightedEntry.AmbientCubemapTintMulScaleValue *= Weight;
+			if(!WeightedEntry.AmbientCubemapTintMulScaleValue.IsAlmostBlack()) // Only bother to add it if the lerped value isn't near-black.
+			{
+				ContributingCubemaps.Push(WeightedEntry);
+			}
+		}
+	}
+
+	// @param InTexture - must not be 0
+	// @param Weight - 0..1
+	void LerpTo(UTexture* InTexture, float Weight)
+	{
+		check(InTexture);
+		check(Weight >= 0.0f && Weight <= 1.0f);
+
+		if(Weight > 254.0f/255.0f || !ContributingLUTs.Num())
+		{
+			SetLUT(InTexture);
+			return;
+		}
+
+		for(uint32 i = 0; i < (uint32)ContributingLUTs.Num(); ++i)
+		{
+			ContributingLUTs[i].Weight *= 1.0f - Weight;
+		}
+
+		uint32 ExistingIndex = FindIndex(InTexture);
+		if(ExistingIndex == 0xffffffff)
+		{
+			PushLUT(InTexture, Weight);
+		}
+		else
+		{
+			ContributingLUTs[ExistingIndex].Weight += Weight;
+		}
+	}
+	
+	/**
+	 * add a LUT(look up table) to the ones that are blended together
+	 * @param Texture - can be 0
+	 * @param Weight - 0..1
+	 */
+	void PushLUT(UTexture* Texture, float Weight)
+	{
+		check(Weight >= 0.0f && Weight <= 1.0f);
+
+		FLUTBlenderEntry Entry;
+
+		Entry.LUTTexture = Texture;
+		Entry.Weight = Weight;
+
+		ContributingLUTs.Add(Entry);
+	}
+
+	/** @return 0xffffffff if not found */
+	uint32 FindIndex(UTexture* Tex) const
+	{
+		for(uint32 i = 0; i < (uint32)ContributingLUTs.Num(); ++i)
+		{
+			if(ContributingLUTs[i].LUTTexture == Tex)
+			{
+				return i;
+			}
+		}
+
+		return 0xffffffff;
+	}
+
+	void SetLUT(UTexture* Texture)
+	{
+		// intentionally no deallocations
+		ContributingLUTs.Empty();
+
+		PushLUT(Texture, 1.0f);
+	}
+
+	// was not merged during blending unlike e.g. BloomThreshold 
+	TArray<FCubemapEntry> ContributingCubemaps;
+	TArray<FLUTBlenderEntry> ContributingLUTs;
+
+	// List of materials to use in the buffer visualization overview
+	TArray<UMaterialInterface*> BufferVisualizationOverviewMaterials;
+
+	// Material to use for rendering high res screenshot with mask. Post process expects this material to be set all the time.
+	UMaterialInterface* HighResScreenshotMaterial;
+
+	// Material to use for rendering just the high res screenshot mask. Post process expects this material to be set all the time.
+	UMaterialInterface* HighResScreenshotMaskMaterial;
+
+	// Material to use for rendering the high resolution screenshot capture region. Post processing only draws the region if this material is set.
+	UMaterialInstanceDynamic* HighResScreenshotCaptureRegionMaterial;
+
+	// Current buffer visualization dumping settings
+	bool bBufferVisualizationDumpRequired;
+	FString BufferVisualizationDumpBaseFilename;
+
+	// Maintains a container with IBlendableInterface objects and their data
+	FBlendableManager BlendableManager;
+};
\ No newline at end of file
diff --git a/Engine/Source/Runtime/Engine/Public/LightMap.h b/Engine/Source/Runtime/Engine/Public/LightMap.h
--- a/Engine/Source/Runtime/Engine/Public/LightMap.h
+++ b/Engine/Source/Runtime/Engine/Public/LightMap.h
@@ -118,8 +118,7 @@ private:
 	virtual void FinishCleanup();
 };
 
-/** A reference to a light-map. */
-typedef TRefCountPtr<FLightMap> FLightMapRef;
+
 
 /** Lightmap reference serializer */
 extern FArchive& operator<<(FArchive& Ar,FLightMap*& R);
diff --git a/Engine/Source/Runtime/Engine/Public/MaterialShared.h b/Engine/Source/Runtime/Engine/Public/MaterialShared.h
--- a/Engine/Source/Runtime/Engine/Public/MaterialShared.h
+++ b/Engine/Source/Runtime/Engine/Public/MaterialShared.h
@@ -13,6 +13,8 @@
 #include "RenderResource.h"
 #include "ShaderCompiler.h"
 #include "UniformBuffer.h"
+#include "SceneTypes.h"
+#include "StaticParameterSet.h"
 
 class FMeshMaterialShaderMap;
 class FMaterialShaderMap;
@@ -35,16 +37,7 @@ class UMaterialInterface;
 
 DECLARE_LOG_CATEGORY_EXTERN(LogMaterial,Log,Verbose);
 
-/** Quality levels that a material can be compiled for. */
-namespace EMaterialQualityLevel
-{
-	enum Type
-	{
-		Low,
-		High,
-		Num
-	};
-}
+
 
 /** Creates a string that represents the given quality level. */
 extern void GetMaterialQualityLevelName(EMaterialQualityLevel::Type InMaterialQualityLevel, FString& OutName);
@@ -287,175 +280,13 @@ public:
 	bool bNeedsGBuffer;
 };
 
-/**
-* Holds the information for a static switch parameter
-*/
-class FStaticSwitchParameter
-{
-public:
-	FName ParameterName;
-	bool Value;
-	bool bOverride;
-	FGuid ExpressionGUID;
-
-	FStaticSwitchParameter() :
-		ParameterName(TEXT("None")),
-		Value(false),
-		bOverride(false),
-		ExpressionGUID(0,0,0,0)
-	{ }
-
-	FStaticSwitchParameter(FName InName, bool InValue, bool InOverride, FGuid InGuid) :
-		ParameterName(InName),
-		Value(InValue),
-		bOverride(InOverride),
-		ExpressionGUID(InGuid)
-	{ }
-
-	friend FArchive& operator<<(FArchive& Ar,FStaticSwitchParameter& P)
-	{
-		Ar << P.ParameterName << P.Value << P.bOverride;
-		Ar << P.ExpressionGUID;
-		return Ar;
-	}
-};
-
-/**
-* Holds the information for a static component mask parameter
-*/
-class FStaticComponentMaskParameter
-{
-public:
-	FName ParameterName;
-	bool R, G, B, A;
-	bool bOverride;
-	FGuid ExpressionGUID;
-
-	FStaticComponentMaskParameter() :
-		ParameterName(TEXT("None")),
-		R(false),
-		G(false),
-		B(false),
-		A(false),
-		bOverride(false),
-		ExpressionGUID(0,0,0,0)
-	{ }
-
-	FStaticComponentMaskParameter(FName InName, bool InR, bool InG, bool InB, bool InA, bool InOverride, FGuid InGuid) :
-		ParameterName(InName),
-		R(InR),
-		G(InG),
-		B(InB),
-		A(InA),
-		bOverride(InOverride),
-		ExpressionGUID(InGuid)
-	{ }
-
-	friend FArchive& operator<<(FArchive& Ar,FStaticComponentMaskParameter& P)
-	{
-		Ar << P.ParameterName << P.R << P.G << P.B << P.A << P.bOverride;
-		Ar << P.ExpressionGUID;
-		return Ar;
-	}
-};
-
-/**
-* Holds the information for a static switch parameter
-*/
-class FStaticTerrainLayerWeightParameter
-{
-public:
-	FName ParameterName;
-	bool bOverride;
-	FGuid ExpressionGUID;
-
-	int32 WeightmapIndex;
-
-	FStaticTerrainLayerWeightParameter() :
-		ParameterName(TEXT("None")),
-		bOverride(false),
-		ExpressionGUID(0,0,0,0),
-		WeightmapIndex(INDEX_NONE)
-	{ }
-
-	FStaticTerrainLayerWeightParameter(FName InName, int32 InWeightmapIndex, bool InOverride, FGuid InGuid) :
-		ParameterName(InName),
-		bOverride(InOverride),
-		ExpressionGUID(InGuid),
-		WeightmapIndex(InWeightmapIndex)
-	{ }
-
-	friend FArchive& operator<<(FArchive& Ar,FStaticTerrainLayerWeightParameter& P)
-	{
-		Ar << P.ParameterName << P.WeightmapIndex << P.bOverride;
-		Ar << P.ExpressionGUID;
-		return Ar;
-	}
-};
-
-/** Contains all the information needed to identify a single permutation of static parameters. */
-class FStaticParameterSet
-{
-public:
-	/** An array of static switch parameters in this set */
-	TArray<FStaticSwitchParameter> StaticSwitchParameters;
-
-	/** An array of static component mask parameters in this set */
-	TArray<FStaticComponentMaskParameter> StaticComponentMaskParameters;
-
-	/** An array of terrain layer weight parameters in this set */
-	TArray<FStaticTerrainLayerWeightParameter> TerrainLayerWeightParameters;
-
-	FStaticParameterSet() {}
-
-	/** 
-	* Checks if this set contains any parameters
-	* 
-	* @return	true if this set has no parameters
-	*/
-	bool IsEmpty() const
-	{
-		return StaticSwitchParameters.Num() == 0 && StaticComponentMaskParameters.Num() == 0 && TerrainLayerWeightParameters.Num() == 0;
-	}
-
-	void Serialize(FArchive& Ar)
-	{
-		// Note: FStaticParameterSet is saved both in packages (UMaterialInstance) and the DDC (FMaterialShaderMap)
-		// Backwards compatibility only works with FStaticParameterSet's stored in packages.  
-		// You must bump MATERIALSHADERMAP_DERIVEDDATA_VER as well if changing the serialization of FStaticParameterSet.
 
-		Ar << StaticSwitchParameters << StaticComponentMaskParameters;
-		Ar << TerrainLayerWeightParameters;
-	}
 
-	/** Updates the hash state with the static parameter names and values. */
-	void UpdateHash(FSHA1& HashState) const;
 
-	/** 
-	* Indicates whether this set is equal to another, copying override settings.
-	* 
-	* @param ReferenceSet	The set to compare against
-	* @return				true if the sets are not equal
-	*/
-	bool ShouldMarkDirty(const FStaticParameterSet* ReferenceSet);
 
-	/** 
-	* Tests this set against another for equality, disregarding override settings.
-	* 
-	* @param ReferenceSet	The set to compare against
-	* @return				true if the sets are equal
-	*/
-	bool operator==(const FStaticParameterSet& ReferenceSet) const;
 
-	bool operator!=(const FStaticParameterSet& ReferenceSet) const
-	{
-		return !(*this == ReferenceSet);
-	}
 
-	FString GetSummaryString() const;
 
-	void AppendKeyString(FString& KeyString) const;
-};
 
 /** 
  * Usage options for a shader map.
@@ -891,39 +722,7 @@ private:
 	friend class FShaderCompilingManager;
 };
 
-//
-//	EMaterialProperty
-//
 
-enum EMaterialProperty
-{
-	MP_EmissiveColor = 0,
-	MP_Opacity,
-	MP_OpacityMask,
-	MP_DiffuseColor,
-	MP_SpecularColor,
-	MP_BaseColor,
-	MP_Metallic,
-	MP_Specular,
-	MP_Roughness,
-	MP_Normal,
-	MP_WorldPositionOffset,
-	MP_WorldDisplacement,
-	MP_TessellationMultiplier,
-	MP_SubsurfaceColor,
-	MP_AmbientOcclusion,
-	MP_Refraction,
-	MP_CustomizedUVs0,
-	MP_CustomizedUVs1,
-	MP_CustomizedUVs2,
-	MP_CustomizedUVs3,
-	MP_CustomizedUVs4,
-	MP_CustomizedUVs5,
-	MP_CustomizedUVs6,
-	MP_CustomizedUVs7,
-	MP_MaterialAttributes,
-	MP_MAX,
-};
 
 /** 
  * Enum that contains entries for the ways that material properties need to be compiled.
diff --git a/Engine/Source/Runtime/Engine/Public/MeshBatch.h b/Engine/Source/Runtime/Engine/Public/MeshBatch.h
new file mode 100644
--- /dev/null
+++ b/Engine/Source/Runtime/Engine/Public/MeshBatch.h
@@ -0,0 +1,160 @@
+// Copyright 1998-2014 Epic Games, Inc. All Rights Reserved.
+
+#pragma once
+
+#include "PrimitiveUniformShaderParameters.h"
+
+/**
+ * A batch mesh element definition.
+ */
+struct FMeshBatchElement
+{
+	const TUniformBuffer<FPrimitiveUniformShaderParameters>* PrimitiveUniformBufferResource;
+	TUniformBufferRef<FPrimitiveUniformShaderParameters> PrimitiveUniformBuffer;
+
+	const FIndexBuffer* IndexBuffer;
+	uint32 FirstIndex;
+	uint32 NumPrimitives;
+	uint32 NumInstances;
+	uint32 MinVertexIndex;
+	uint32 MaxVertexIndex;
+	int32 UserIndex;
+	void* UserData;
+	float MinScreenSize;
+	float MaxScreenSize;
+
+	/** 
+	 *	DynamicIndexData - pointer to user memory containing the index data.
+	 *	Used for rendering dynamic data directly.
+	 */
+	const void* DynamicIndexData;
+	uint16 DynamicIndexStride;
+	
+	FMeshBatchElement()
+	:	PrimitiveUniformBufferResource(NULL)
+	,	IndexBuffer(NULL)
+	,	NumInstances(1)
+	,	UserIndex(-1)
+	,	UserData(NULL)
+	,	MinScreenSize(0.0f)
+	,	MaxScreenSize(1.0f)
+	,	DynamicIndexData(NULL)
+	{
+	}
+};
+
+/**
+ * A batch of mesh elements, all with the same material and vertex buffer
+ */
+struct FMeshBatch
+{
+	TArray<FMeshBatchElement,TInlineAllocator<1> > Elements;
+
+	// used with DynamicVertexData
+	uint16 DynamicVertexStride;
+
+	/** LOD index of the mesh, used for fading LOD transitions. */
+	int8 LODIndex;
+
+	uint32 UseDynamicData : 1;
+	uint32 ReverseCulling : 1;
+	uint32 bDisableBackfaceCulling : 1;
+	uint32 CastShadow : 1;
+	uint32 bWireframe : 1;
+	// e.g. PT_TriangleList(default), PT_LineList, ..
+	uint32 Type : PT_NumBits;
+	// e.g. SDPG_World (default), SDPG_Foreground
+	uint32 DepthPriorityGroup : SDPG_NumBits;
+	uint32 bUseAsOccluder : 1;
+
+	// can be 0
+	const FLightCacheInterface* LCI;
+
+	/** 
+	 *	DynamicVertexData - pointer to user memory containing the vertex data.
+	 *	Used for rendering dynamic data directly.
+	 *  used with DynamicVertexStride
+	 */
+	const void* DynamicVertexData;
+
+	// can be 0
+	const FVertexFactory* VertexFactory;
+	// can be 0
+	const FMaterialRenderProxy* MaterialRenderProxy;
+
+
+	FORCEINLINE bool IsTranslucent(ERHIFeatureLevel::Type InFeatureLevel) const
+	{
+		// Note: blend mode does not depend on the feature level we are actually rendering in.
+		return MaterialRenderProxy && IsTranslucentBlendMode(MaterialRenderProxy->GetMaterial(InFeatureLevel)->GetBlendMode());
+	}
+
+	FORCEINLINE bool IsMasked(ERHIFeatureLevel::Type InFeatureLevel) const
+	{
+		// Note: blend mode does not depend on the feature level we are actually rendering in.
+		return MaterialRenderProxy && MaterialRenderProxy->GetMaterial(InFeatureLevel)->IsMasked();
+	}
+
+	/** Converts from an int32 index into a int8 */
+	static int8 QuantizeLODIndex(int32 NewLODIndex)
+	{
+		checkSlow(NewLODIndex >= SCHAR_MIN && NewLODIndex <= SCHAR_MAX);
+		return (int8)NewLODIndex;
+	}
+
+	/** 
+	* @return vertex stride specified for the mesh. 0 if not dynamic
+	*/
+	FORCEINLINE uint32 GetDynamicVertexStride(ERHIFeatureLevel::Type /*InFeatureLevel*/) const
+	{
+		if (UseDynamicData && DynamicVertexData)
+		{
+			return DynamicVertexStride;
+		}
+		else
+		{
+			return 0;
+		}
+	}
+
+	FORCEINLINE int32 GetNumPrimitives() const
+	{
+		int32 Count=0;
+		for( int32 ElementIdx=0;ElementIdx<Elements.Num();ElementIdx++ )
+		{
+			Count += Elements[ElementIdx].NumPrimitives;
+		}
+		return Count;
+	}
+
+#if DO_CHECK
+	FORCEINLINE void CheckUniformBuffers() const
+	{
+		for( int32 ElementIdx=0;ElementIdx<Elements.Num();ElementIdx++ )
+		{
+			check(IsValidRef(Elements[ElementIdx].PrimitiveUniformBuffer) || Elements[ElementIdx].PrimitiveUniformBufferResource != NULL);
+		}
+	}
+#endif	
+
+	/** Default constructor. */
+	FMeshBatch()
+	:	DynamicVertexStride(0)
+	,	LODIndex(INDEX_NONE)
+	,	UseDynamicData(false)
+	,	ReverseCulling(false)
+	,	bDisableBackfaceCulling(false)
+	,	CastShadow(true)
+	,	bWireframe(false)
+	,	Type(PT_TriangleList)
+	,	DepthPriorityGroup(SDPG_World)
+	,	bUseAsOccluder(true)
+	,	LCI(NULL)
+	,	DynamicVertexData(NULL)
+	,	VertexFactory(NULL)
+	,	MaterialRenderProxy(NULL)
+	{
+		// By default always add the first element.
+		new(Elements) FMeshBatchElement;
+	}
+};
\ No newline at end of file
diff --git a/Engine/Source/Runtime/Engine/Public/PrimitiveSceneProxy.h b/Engine/Source/Runtime/Engine/Public/PrimitiveSceneProxy.h
--- a/Engine/Source/Runtime/Engine/Public/PrimitiveSceneProxy.h
+++ b/Engine/Source/Runtime/Engine/Public/PrimitiveSceneProxy.h
@@ -6,8 +6,17 @@
 
 #pragma once
 
-#include "SceneManagement.h"
-#include "UniformBuffer.h"
+class FSimpleLightEntry;
+class HHitProxy;
+class FStaticPrimitiveDrawInterface;
+class FPrimitiveSceneInfo;
+class FLightSceneProxy;
+class FLightSceneInfo;
+
+
+#include "SceneView.h"
+#include "PrimitiveUniformShaderParameters.h"
+#include "PrimitiveViewRelevance.h"
 
 /** Data for a simple dynamic light. */
 class FSimpleLightEntry
@@ -31,9 +40,9 @@ class FSimpleLightArray
 {
 public:
 	/** Data per simple dynamic light instance, independent of view */
-	TArray<FSimpleLightEntry, SceneRenderingAllocator> InstanceData;
+	TArray<FSimpleLightEntry, TMemStackAllocator<>> InstanceData;
 	/** Per-view data for each light */
-	TArray<FSimpleLightPerViewEntry, SceneRenderingAllocator> PerViewData;
+	TArray<FSimpleLightPerViewEntry, TMemStackAllocator<>> PerViewData;
 
 public:
 
diff --git a/Engine/Source/Runtime/Engine/Public/PrimitiveUniformShaderParameters.h b/Engine/Source/Runtime/Engine/Public/PrimitiveUniformShaderParameters.h
new file mode 100644
--- /dev/null
+++ b/Engine/Source/Runtime/Engine/Public/PrimitiveUniformShaderParameters.h
@@ -0,0 +1,116 @@
+// Copyright 1998-2014 Epic Games, Inc. All Rights Reserved.
+
+#pragma once
+
+
+#include "UniformBuffer.h"
+#include "Runtime/ShaderCore/Public/ShaderParameters.h"
+
+
+/** The uniform shader parameters associated with a primitive. */
+BEGIN_UNIFORM_BUFFER_STRUCT(FPrimitiveUniformShaderParameters,ENGINE_API)
+	DECLARE_UNIFORM_BUFFER_STRUCT_MEMBER(FMatrix,LocalToWorld)
+	DECLARE_UNIFORM_BUFFER_STRUCT_MEMBER(FMatrix,WorldToLocal)
+	DECLARE_UNIFORM_BUFFER_STRUCT_MEMBER(FVector4,ObjectWorldPositionAndRadius)
+	DECLARE_UNIFORM_BUFFER_STRUCT_MEMBER(FVector,ObjectBounds)
+	DECLARE_UNIFORM_BUFFER_STRUCT_MEMBER_EX(float,LocalToWorldDeterminantSign,EShaderPrecisionModifier::Half)
+	DECLARE_UNIFORM_BUFFER_STRUCT_MEMBER(FVector,ActorWorldPosition)
+	DECLARE_UNIFORM_BUFFER_STRUCT_MEMBER_EX(float,DecalReceiverMask,EShaderPrecisionModifier::Half)
+	DECLARE_UNIFORM_BUFFER_STRUCT_MEMBER_EX(float,HasDistanceFieldRepresentation,EShaderPrecisionModifier::Half)
+	DECLARE_UNIFORM_BUFFER_STRUCT_MEMBER_EX(FVector4,ObjectOrientation,EShaderPrecisionModifier::Half)
+	DECLARE_UNIFORM_BUFFER_STRUCT_MEMBER_EX(FVector4,NonUniformScale,EShaderPrecisionModifier::Half)
+	DECLARE_UNIFORM_BUFFER_STRUCT_MEMBER_EX(FVector4,InvNonUniformScale,EShaderPrecisionModifier::Half)
+	DECLARE_UNIFORM_BUFFER_STRUCT_MEMBER(FVector, LocalObjectBoundsMin)		// This is used in a custom material function (ObjectLocalBounds.uasset)
+	DECLARE_UNIFORM_BUFFER_STRUCT_MEMBER(FVector, LocalObjectBoundsMax)		// This is used in a custom material function (ObjectLocalBounds.uasset)
+	DECLARE_UNIFORM_BUFFER_STRUCT_MEMBER(float,LpvBiasMultiplier)
+END_UNIFORM_BUFFER_STRUCT(FPrimitiveUniformShaderParameters)
+
+/** Initializes the primitive uniform shader parameters. */
+inline FPrimitiveUniformShaderParameters GetPrimitiveUniformShaderParameters(
+	const FMatrix& LocalToWorld,
+	FVector ActorPosition,
+	const FBoxSphereBounds& WorldBounds,
+	const FBoxSphereBounds& LocalBounds,
+	bool bReceivesDecals,
+	bool bHasDistanceFieldRepresentation,
+	float LpvBiasMultiplier = 1.0f
+)
+{
+	FPrimitiveUniformShaderParameters Result;
+	Result.LocalToWorld = LocalToWorld;
+	Result.WorldToLocal = LocalToWorld.InverseSafe();
+	Result.ObjectWorldPositionAndRadius = FVector4(WorldBounds.Origin, WorldBounds.SphereRadius);
+	Result.ObjectBounds = WorldBounds.BoxExtent;
+	Result.LocalObjectBoundsMin = LocalBounds.GetBoxExtrema(0); // 0 == minimum
+	Result.LocalObjectBoundsMax = LocalBounds.GetBoxExtrema(1); // 1 == maximum
+	Result.ObjectOrientation = LocalToWorld.GetUnitAxis( EAxis::Z );
+	Result.ActorWorldPosition = ActorPosition;
+	Result.LpvBiasMultiplier = LpvBiasMultiplier;
+
+	{
+		// Extract per axis scales from LocalToWorld transform
+		FVector4 WorldX = FVector4(LocalToWorld.M[0][0],LocalToWorld.M[0][1],LocalToWorld.M[0][2],0);
+		FVector4 WorldY = FVector4(LocalToWorld.M[1][0],LocalToWorld.M[1][1],LocalToWorld.M[1][2],0);
+		FVector4 WorldZ = FVector4(LocalToWorld.M[2][0],LocalToWorld.M[2][1],LocalToWorld.M[2][2],0);
+		float ScaleX = FVector(WorldX).Size();
+		float ScaleY = FVector(WorldY).Size();
+		float ScaleZ = FVector(WorldZ).Size();
+		Result.NonUniformScale = FVector4(ScaleX,ScaleY,ScaleZ,0);
+		Result.InvNonUniformScale = FVector4(
+			ScaleX > KINDA_SMALL_NUMBER ? 1.0f/ScaleX : 0.0f,
+			ScaleY > KINDA_SMALL_NUMBER ? 1.0f/ScaleY : 0.0f,
+			ScaleZ > KINDA_SMALL_NUMBER ? 1.0f/ScaleZ : 0.0f,
+			0.0f
+			);
+	}
+
+	Result.LocalToWorldDeterminantSign = FMath::FloatSelect(LocalToWorld.RotDeterminant(),1,-1);
+	Result.DecalReceiverMask = bReceivesDecals ? 1 : 0;
+	Result.HasDistanceFieldRepresentation = bHasDistanceFieldRepresentation ? 1 : 0;
+	return Result;
+}
+
+inline TUniformBufferRef<FPrimitiveUniformShaderParameters> CreatePrimitiveUniformBufferImmediate(
+	const FMatrix& LocalToWorld,
+	const FBoxSphereBounds& WorldBounds,
+	const FBoxSphereBounds& LocalBounds,
+	bool bReceivesDecals,
+	float LpvBiasMultiplier = 1.0f
+	)
+{
+	check(IsInRenderingThread());
+	return TUniformBufferRef<FPrimitiveUniformShaderParameters>::CreateUniformBufferImmediate(
+		GetPrimitiveUniformShaderParameters(LocalToWorld, WorldBounds.Origin, WorldBounds, LocalBounds, bReceivesDecals, false, LpvBiasMultiplier ),
+		UniformBuffer_MultiFrame
+		);
+}
+
+/**
+ * Primitive uniform buffer containing only identity transforms.
+ */
+class FIdentityPrimitiveUniformBuffer : public TUniformBuffer<FPrimitiveUniformShaderParameters>
+{
+public:
+
+	/** Default constructor. */
+	FIdentityPrimitiveUniformBuffer()
+	{
+		SetContents(GetPrimitiveUniformShaderParameters(
+			FMatrix(
+				FPlane(1.0f,0.0f,0.0f,0.0f),
+				FPlane(0.0f,1.0f,0.0f,0.0f),
+				FPlane(0.0f,0.0f,1.0f,0.0f),
+				FPlane(0.0f,0.0f,0.0f,1.0f)
+				),
+			FVector(0.0f,0.0f,0.0f),
+			FBoxSphereBounds(EForceInit::ForceInit),
+			FBoxSphereBounds(EForceInit::ForceInit),
+			true,
+			false,
+			1.0f		// LPV bias
+			));
+	}
+};
+
+/** Global primitive uniform buffer resource containing identity transformations. */
+extern TGlobalResource<FIdentityPrimitiveUniformBuffer> GIdentityPrimitiveUniformBuffer;
\ No newline at end of file
diff --git a/Engine/Source/Runtime/Engine/Public/PrimitiveViewRelevance.h b/Engine/Source/Runtime/Engine/Public/PrimitiveViewRelevance.h
new file mode 100644
--- /dev/null
+++ b/Engine/Source/Runtime/Engine/Public/PrimitiveViewRelevance.h
@@ -0,0 +1,108 @@
+// Copyright 1998-2014 Epic Games, Inc. All Rights Reserved.
+
+#pragma once
+
+
+/**
+ * The different types of relevance a primitive scene proxy can declare towards a particular scene view.
+ */
+struct FPrimitiveViewRelevance
+{
+	/** The primitive's static elements are rendered for the view. */
+	uint32 bStaticRelevance : 1; 
+	/** The primitive's dynamic elements are rendered for the view. */
+	uint32 bDynamicRelevance : 1;
+	/** The primitive is drawn. */
+	uint32 bDrawRelevance : 1;
+	/** The primitive is casting a shadow. */
+	uint32 bShadowRelevance : 1;
+	/** The primitive should render to the custom depth pass. */
+	uint32 bRenderCustomDepth : 1;
+	/** The primitive should render to the main depth pass. */
+	uint32 bRenderInMainPass : 1;
+	/** The primitive is drawn only in the editor and composited onto the scene after post processing */
+	uint32 bEditorPrimitiveRelevance : 1;
+	/** The primitive is drawn only in the editor and composited onto the scene after post processing using no depth testing */
+	uint32 bEditorNoDepthTestPrimitiveRelevance : 1;
+	/** The primitive needs PreRenderView to be called before rendering. */
+	uint32 bNeedsPreRenderView : 1;
+	/** The primitive should have GatherSimpleLights called on the proxy when gathering simple lights. */
+	uint32 bHasSimpleLights : 1;
+
+	/** The primitive has one or more opaque or masked elements. */
+	uint32 bOpaqueRelevance : 1;
+	/** The primitive has one or more masked elements. */
+	uint32 bMaskedRelevance : 1;
+	/** The primitive has one or more distortion elements. */
+	uint32 bDistortionRelevance : 1;
+	/** The primitive reads from scene color. */
+	uint32 bSceneColorRelevance : 1;
+	/** The primitive has one or more elements that have SeparateTranslucency. */
+	uint32 bSeparateTranslucencyRelevance : 1;
+	/** The primitive has one or more elements that have normal translucency. */
+	uint32 bNormalTranslucencyRelevance : 1;
+
+	/** 
+	 * Whether this primitive view relevance has been initialized this frame.  
+	 * Primitives that have not had ComputeRelevanceForView called on them (because they were culled) will not be initialized,
+	 * But we may still need to render them from other views like shadow passes, so this tracks whether we can reuse the cached relevance or not.
+	 */
+	uint32 bInitializedThisFrame : 1;
+
+	bool HasTranslucency() const 
+	{
+		return bSeparateTranslucencyRelevance || bNormalTranslucencyRelevance;
+	}
+
+	/** Initialization constructor. */
+	FPrimitiveViewRelevance():
+		bStaticRelevance(false),
+		bDynamicRelevance(false),
+		bDrawRelevance(false),
+		bShadowRelevance(false),
+		bRenderCustomDepth(false),
+		bRenderInMainPass(true),
+		bEditorPrimitiveRelevance(false),
+		bEditorNoDepthTestPrimitiveRelevance(false),
+		bNeedsPreRenderView(false),
+		bHasSimpleLights(false),
+		bOpaqueRelevance(true),
+		bMaskedRelevance(false),
+		bDistortionRelevance(false),
+		bSceneColorRelevance(false),
+		bSeparateTranslucencyRelevance(false),
+		bNormalTranslucencyRelevance(false),
+		bInitializedThisFrame(false)
+	{}
+
+	/** Bitwise OR operator.  Sets any relevance bits which are present in either FPrimitiveViewRelevance. */
+	FPrimitiveViewRelevance& operator|=(const FPrimitiveViewRelevance& B)
+	{
+		bStaticRelevance |= B.bStaticRelevance != 0;
+		bDynamicRelevance |= B.bDynamicRelevance != 0;
+		bDrawRelevance |= B.bDrawRelevance != 0;
+		bShadowRelevance |= B.bShadowRelevance != 0;
+		bOpaqueRelevance |= B.bOpaqueRelevance != 0;
+		bMaskedRelevance |= B.bMaskedRelevance != 0;
+		bDistortionRelevance |= B.bDistortionRelevance != 0;
+		bRenderCustomDepth |= B.bRenderCustomDepth != 0;
+		bRenderInMainPass |= B.bRenderInMainPass !=0;
+		bEditorPrimitiveRelevance |= B.bEditorPrimitiveRelevance !=0;
+		bEditorNoDepthTestPrimitiveRelevance |= B.bEditorNoDepthTestPrimitiveRelevance !=0;
+		bSceneColorRelevance |= B.bSceneColorRelevance != 0;
+		bNeedsPreRenderView |= B.bNeedsPreRenderView != 0;
+		bHasSimpleLights |= B.bHasSimpleLights != 0;
+		bSeparateTranslucencyRelevance |= B.bSeparateTranslucencyRelevance != 0;
+		bNormalTranslucencyRelevance |= B.bNormalTranslucencyRelevance != 0;
+		bInitializedThisFrame |= B.bInitializedThisFrame;
+		return *this;
+	}
+
+	/** Binary bitwise OR operator. */
+	friend FPrimitiveViewRelevance operator|(const FPrimitiveViewRelevance& A,const FPrimitiveViewRelevance& B)
+	{
+		FPrimitiveViewRelevance Result(A);
+		Result |= B;
+		return Result;
+	}
+};
\ No newline at end of file
diff --git a/Engine/Source/Runtime/Engine/Public/SceneInterface.h b/Engine/Source/Runtime/Engine/Public/SceneInterface.h
new file mode 100644
--- /dev/null
+++ b/Engine/Source/Runtime/Engine/Public/SceneInterface.h
@@ -0,0 +1,306 @@
+// Copyright 1998-2014 Epic Games, Inc. All Rights Reserved.
+
+#pragma once
+
+class FSkyLightSceneProxy;
+class FMaterial;
+class FMaterialShaderMap;
+class ULightComponent;
+class USkyLightComponent;
+
+/**
+ * An interface to the private scene manager implementation of a scene.  Use GetRendererModule().AllocateScene to create.
+ * The scene
+ */
+class FSceneInterface
+{
+public:
+
+	// FSceneInterface interface
+
+	/** 
+	 * Adds a new primitive component to the scene
+	 * 
+	 * @param Primitive - primitive component to add
+	 */
+	virtual void AddPrimitive(UPrimitiveComponent* Primitive) = 0;
+	/** 
+	 * Removes a primitive component from the scene
+	 * 
+	 * @param Primitive - primitive component to remove
+	 */
+	virtual void RemovePrimitive(UPrimitiveComponent* Primitive) = 0;
+	/** Called when a primitive is being unregistered and will not be immediately re-registered. */
+	virtual void ReleasePrimitive(UPrimitiveComponent* Primitive) = 0;
+	/** 
+	 * Updates the transform of a primitive which has already been added to the scene. 
+	 * 
+	 * @param Primitive - primitive component to update
+	 */
+	virtual void UpdatePrimitiveTransform(UPrimitiveComponent* Primitive) = 0;
+	/** Updates primitive attachment state. */
+	virtual void UpdatePrimitiveAttachment(UPrimitiveComponent* Primitive) = 0;
+	/** 
+	 * Adds a new light component to the scene
+	 * 
+	 * @param Light - light component to add
+	 */
+	virtual void AddLight(ULightComponent* Light) = 0;
+	/** 
+	 * Removes a light component from the scene
+	 * 
+	 * @param Light - light component to remove
+	 */
+	virtual void RemoveLight(ULightComponent* Light) = 0;
+	/** 
+	 * Adds a new light component to the scene which is currently invisible, but needed for editor previewing
+	 * 
+	 * @param Light - light component to add
+	 */
+	virtual void AddInvisibleLight(ULightComponent* Light) = 0;
+	virtual void SetSkyLight(FSkyLightSceneProxy* Light) = 0;
+	/** 
+	 * Adds a new decal component to the scene
+	 * 
+	 * @param Component - component to add
+	 */
+	virtual void AddDecal(UDecalComponent* Component) = 0;
+	/** 
+	 * Removes a decal component from the scene
+	 * 
+	 * @param Component - component to remove
+	 */
+	virtual void RemoveDecal(UDecalComponent* Component) = 0;
+	/** 
+	 * Updates the transform of a decal which has already been added to the scene. 
+	 *
+	 * @param Decal - Decal component to update
+	 */
+	virtual void UpdateDecalTransform(UDecalComponent* Component) = 0;
+
+	/** Adds a reflection capture to the scene. */
+	virtual void AddReflectionCapture(class UReflectionCaptureComponent* Component) {}
+
+	/** Removes a reflection capture from the scene. */
+	virtual void RemoveReflectionCapture(class UReflectionCaptureComponent* Component) {}
+
+	/** Reads back reflection capture data from the GPU.  Very slow operation that blocks the GPU and rendering thread many times. */
+	virtual void GetReflectionCaptureData(UReflectionCaptureComponent* Component, class FReflectionCaptureFullHDRDerivedData& OutDerivedData) {}
+
+	/** Updates a reflection capture's transform, and then re-captures the scene. */
+	virtual void UpdateReflectionCaptureTransform(class UReflectionCaptureComponent* Component) {}
+
+	/** 
+	 * Allocates reflection captures in the scene's reflection cubemap array and updates them by recapturing the scene.
+	 * Existing captures will only be updated.  Must be called from the game thread.
+	 */
+	virtual void AllocateReflectionCaptures(const TArray<UReflectionCaptureComponent*>& NewCaptures) {}
+	virtual void ReleaseReflectionCubemap(UReflectionCaptureComponent* CaptureComponent) {}
+
+	/** 
+	 * Updates the contents of the given sky capture by rendering the scene. 
+	 * This must be called on the game thread.
+	 */
+	virtual void UpdateSkyCaptureContents(const USkyLightComponent* CaptureComponent, bool bCaptureEmissiveOnly, FTexture* OutProcessedTexture, FSHVectorRGB3& OutIrradianceEnvironmentMap) {}
+
+	/** 
+	* Updates the contents of the given scene capture by rendering the scene. 
+	* This must be called on the game thread.
+	*/
+	virtual void UpdateSceneCaptureContents(class USceneCaptureComponent2D* CaptureComponent) {}
+	virtual void UpdateSceneCaptureContents(class USceneCaptureComponentCube* CaptureComponent) {}
+
+	virtual void AddPrecomputedLightVolume(const class FPrecomputedLightVolume* Volume) {}
+	virtual void RemovePrecomputedLightVolume(const class FPrecomputedLightVolume* Volume) {}
+
+	/** 
+	 * Updates the transform of a light which has already been added to the scene. 
+	 *
+	 * @param Light - light component to update
+	 */
+	virtual void UpdateLightTransform(ULightComponent* Light) = 0;
+	/** 
+	 * Updates the color and brightness of a light which has already been added to the scene. 
+	 *
+	 * @param Light - light component to update
+	 */
+	virtual void UpdateLightColorAndBrightness(ULightComponent* Light) = 0;
+
+	/** Updates the scene's dynamic skylight. */
+	virtual void UpdateDynamicSkyLight(const FLinearColor& UpperColor, const FLinearColor& LowerColor) {}
+
+	/** Sets the precomputed visibility handler for the scene, or NULL to clear the current one. */
+	virtual void SetPrecomputedVisibility(const class FPrecomputedVisibilityHandler* PrecomputedVisibilityHandler) {}
+
+	/** Sets the precomputed volume distance field for the scene, or NULL to clear the current one. */
+	virtual void SetPrecomputedVolumeDistanceField(const class FPrecomputedVolumeDistanceField* PrecomputedVolumeDistanceField) {}
+
+	/** Sets shader maps on the specified materials without blocking. */
+	virtual void SetShaderMapsOnMaterialResources(const TMap<FMaterial*, FMaterialShaderMap*>& MaterialsToUpdate) {}
+
+	/** Updates static draw lists for the given set of materials. */
+	virtual void UpdateStaticDrawListsForMaterials(const TArray<const FMaterial*>& Materials) {}
+
+	/** 
+	 * Adds a new exponential height fog component to the scene
+	 * 
+	 * @param FogComponent - fog component to add
+	 */	
+	virtual void AddExponentialHeightFog(class UExponentialHeightFogComponent* FogComponent) = 0;
+	/** 
+	 * Removes a exponential height fog component from the scene
+	 * 
+	 * @param FogComponent - fog component to remove
+	 */	
+	virtual void RemoveExponentialHeightFog(class UExponentialHeightFogComponent* FogComponent) = 0;
+
+	/** 
+	 * Adds a new atmospheric fog component to the scene
+	 * 
+	 * @param FogComponent - fog component to add
+	 */	
+	virtual void AddAtmosphericFog(class UAtmosphericFogComponent* FogComponent) = 0;
+
+	/** 
+	 * Removes a atmospheric fog component from the scene
+	 * 
+	 * @param FogComponent - fog component to remove
+	 */	
+	virtual void RemoveAtmosphericFog(class UAtmosphericFogComponent* FogComponent) = 0;
+
+	/**
+	 * Adds a wind source component to the scene.
+	 * @param WindComponent - The component to add.
+	 */
+	virtual void AddWindSource(class UWindDirectionalSourceComponent* WindComponent) = 0;
+	/**
+	 * Removes a wind source component from the scene.
+	 * @param WindComponent - The component to remove.
+	 */
+	virtual void RemoveWindSource(class UWindDirectionalSourceComponent* WindComponent) = 0;
+	/**
+	 * Accesses the wind source list.  Must be called in the rendering thread.
+	 * @return The list of wind sources in the scene.
+	 */
+	virtual const TArray<class FWindSourceSceneProxy*>& GetWindSources_RenderThread() const = 0;
+
+	/** Accesses wind parameters.  XYZ will contain wind direction * Strength, W contains wind speed. */
+	virtual FVector4 GetWindParameters(const FVector& Position) const = 0;
+
+	/** Same as GetWindParameters, but ignores point wind sources. */
+	virtual FVector4 GetDirectionalWindParameters() const = 0;
+
+	/** 
+	 * Adds a SpeedTree wind computation object to the scene.
+	 * @param StaticMesh - The SpeedTree static mesh whose wind to add.
+	 */
+	virtual void AddSpeedTreeWind(class FVertexFactory* VertexFactory, class UStaticMesh* StaticMesh) = 0;
+
+	/** 
+	 * Removes a SpeedTree wind computation object to the scene.
+	 * @param StaticMesh - The SpeedTree static mesh whose wind to remove.
+	 */
+	virtual void RemoveSpeedTreeWind(class FVertexFactory* VertexFactory, class UStaticMesh* StaticMesh) = 0;
+
+	/** Ticks the SpeedTree wind object and updates the uniform buffer. */
+	virtual void UpdateSpeedTreeWind(double CurrentTime) = 0;
+
+	/** 
+	 * Looks up the SpeedTree uniform buffer for the passed in vertex factory.
+	 * @param VertexFactory - The vertex factory registered for SpeedTree.
+	 */
+	virtual FUniformBufferRHIParamRef GetSpeedTreeUniformBuffer(const FVertexFactory* VertexFactory) = 0;
+
+	/**
+	 * Release this scene and remove it from the rendering thread
+	 */
+	virtual void Release() = 0;
+	/**
+	 * Retrieves the lights interacting with the passed in primitive and adds them to the out array.
+	 *
+	 * @param	Primitive				Primitive to retrieve interacting lights for
+	 * @param	RelevantLights	[out]	Array of lights interacting with primitive
+	 */
+	virtual void GetRelevantLights( UPrimitiveComponent* Primitive, TArray<const ULightComponent*>* RelevantLights ) const = 0;
+	/**
+	 * Indicates if hit proxies should be processed by this scene
+	 *
+	 * @return true if hit proxies should be rendered in this scene.
+	 */
+	virtual bool RequiresHitProxies() const = 0;
+	/**
+	 * Get the optional UWorld that is associated with this scene
+	 * 
+ 	 * @return UWorld instance used by this scene
+	 */
+	virtual class UWorld* GetWorld() const = 0;
+	/**
+	 * Return the scene to be used for rendering. Note that this can return NULL if rendering has
+	 * been disabled!
+	 */
+	virtual class FScene* GetRenderScene()
+	{
+		return NULL;
+	}
+	/**
+	 * Sets the FX system associated with the scene.
+	 */
+	virtual void SetFXSystem( class FFXSystemInterface* InFXSystem ) = 0;
+
+	/**
+	 * Get the FX system associated with the scene.
+	 */
+	virtual class FFXSystemInterface* GetFXSystem() = 0;
+
+	virtual void DumpUnbuiltLightIteractions( FOutputDevice& Ar ) const { }
+
+	/**
+	 * Dumps static mesh draw list stats to the log.
+	 */
+	virtual void DumpStaticMeshDrawListStats() const {}
+
+	/**
+	 * Request to clear the MB info. Game thread only
+	 */
+	virtual void SetClearMotionBlurInfoGameThread() {}
+
+	/** Updates the scene's list of parameter collection id's and their uniform buffers. */
+	virtual void UpdateParameterCollections(const TArray<class FMaterialParameterCollectionInstanceResource*>& InParameterCollections) {}
+
+	/**
+	 * Exports the scene.
+	 *
+	 * @param	Ar		The Archive used for exporting.
+	 **/
+	virtual void Export( FArchive& Ar ) const
+	{}
+
+	
+	/**
+	 * Shifts scene data by provided delta
+	 * Called on world origin changes
+	 * 
+	 * @param	InOffset	Delta to shift scene by
+	 */
+	virtual void ApplyWorldOffset(FVector InOffset) {}
+
+	/**
+	 * Notification that level was added to a world
+	 * 
+	 * @param	InLevelName		Level name
+	 */
+	virtual void OnLevelAddedToWorld(FName InLevelName) {}
+
+	/**
+	 * @return True if there are any lights in the scene
+	 */
+	virtual bool HasAnyLights() const = 0;
+
+	virtual bool IsEditorScene() const { return false; }
+
+	virtual ERHIFeatureLevel::Type GetFeatureLevel() const { return GRHIFeatureLevel; }
+	virtual void ChangeFeatureLevel(ERHIFeatureLevel::Type InFeatureLevel) {}
+
+protected:
+	virtual ~FSceneInterface() {}
+};
\ No newline at end of file
diff --git a/Engine/Source/Runtime/Engine/Public/SceneManagement.h b/Engine/Source/Runtime/Engine/Public/SceneManagement.h
--- a/Engine/Source/Runtime/Engine/Public/SceneManagement.h
+++ b/Engine/Source/Runtime/Engine/Public/SceneManagement.h
@@ -16,17 +16,22 @@
 #include "Engine/World.h"
 #include "RendererInterface.h"
 #include "Engine/GameViewportClient.h"
+#include "SceneTypes.h"
+#include "SceneView.h"
+#include "FinalPostProcessSettings.h"
+#include "PrimitiveUniformShaderParameters.h"
+#include "MeshBatch.h"
+#include "PrimitiveViewRelevance.h"
+#include "SceneInterface.h"
+
 
 // Forward declarations.
 class FSceneViewFamily;
-class FSceneInterface;
 class FLightSceneInfo;
 class ULightComponent;
 class UDecalComponent;
 class FIndexBuffer;
-class FVertexBuffer;
 class FVertexFactory;
-class FAtmosphericFogSceneInfo;
 class ISceneViewExtension;
 class HHitProxy;
 struct FDynamicMeshVertex;
@@ -34,101 +39,11 @@ struct FDynamicMeshVertex;
 
 DECLARE_LOG_CATEGORY_EXTERN(LogBufferVisualization, Log, All);
 
-/** 
- * Class used to identify UPrimitiveComponents on the rendering thread without having to pass the pointer around, 
- * Which would make it easy for people to access game thread variables from the rendering thread.
- */
-class FPrimitiveComponentId
-{
-public:
-
-	FPrimitiveComponentId() : Value(0)
-	{}
-
-	inline bool IsValid() const
-	{
-		return Value > 0;
-	}
-
-	inline bool operator==(FPrimitiveComponentId OtherId) const
-	{
-		return Value == OtherId.Value;
-	}
-
-	friend uint32 GetTypeHash( FPrimitiveComponentId Id )
-	{
-		return GetTypeHash(Id.Value);
-	}
 
-	uint64 Value;
-};
 
 // -----------------------------------------------------------------------------
 
-struct FViewMatrices
-{
-	FViewMatrices()
-	{
-		ProjMatrix.SetIdentity();
-		ViewMatrix.SetIdentity();
-		TranslatedViewProjectionMatrix.SetIdentity();
-		InvTranslatedViewProjectionMatrix.SetIdentity();
-		PreViewTranslation = FVector::ZeroVector;
-		ViewOrigin = FVector::ZeroVector;
-	}
-
-	/** ViewToClip : UE4 projection matrix projects such that clip space Z=1 is the near plane, and Z=0 is the infinite far plane. */
-	FMatrix		ProjMatrix;
-	// WorldToView..
-	FMatrix		ViewMatrix;
-	/** The view-projection transform, starting from world-space points translated by -ViewOrigin. */
-	FMatrix		TranslatedViewProjectionMatrix;
-	/** The inverse view-projection transform, ending with world-space points translated by -ViewOrigin. */
-	FMatrix		InvTranslatedViewProjectionMatrix;
-	/** The translation to apply to the world before TranslatedViewProjectionMatrix. Usually it is -ViewOrigin but with rereflections this can differ */
-	FVector		PreViewTranslation;
-	/** To support ortho and other modes this is redundant, in world space */
-	FVector		ViewOrigin;
-
-	//
-	// World = TranslatedWorld - PreViewTranslation
-	// TranslatedWorld = World + PreViewTranslation
-	// 
-
-// ----------------
-
-	/** @return true:perspective, false:orthographic */
-	inline bool IsPerspectiveProjection() const
-	{
-		return ProjMatrix.M[3][3] < 1.0f;
-	}
-
-	FMatrix GetViewProjMatrix() const
-	{
-		return ViewMatrix * ProjMatrix;
-	}
-
-	FMatrix GetViewRotationProjMatrix() const
-	{
-		return ViewMatrix.RemoveTranslation() * ProjMatrix;
-	}
-
-	FMatrix GetInvProjMatrix() const
-	{
-		return ProjMatrix.Inverse();
-	}
 
-	FMatrix GetInvViewMatrix() const
-	{
-		// can be optimized: it's not a perspective matrix so transpose would be enough
-		return ViewMatrix.Inverse();
-	}
-	
-	FMatrix GetInvViewProjMatrix() const
-	{
-		return GetInvProjMatrix() * GetInvViewMatrix();
-	}
-};
 
 /**
  * The scene manager's persistent view state.
@@ -293,24 +208,9 @@ private:
 	ELightInteractionType Type;
 };
 
-/**
- * The types of interactions between a light and a primitive.
- */
-enum ELightMapInteractionType
-{
-	LMIT_None	= 0,
-	LMIT_Texture= 2,
 
-	LMIT_NumBits= 3
-};
 
-enum EShadowMapInteractionType
-{
-	SMIT_None	= 0,
-	SMIT_Texture= 2,
 
-	SMIT_NumBits= 3
-};
 
 /** The number of coefficients that are stored for each light sample. */ 
 static const int32 NUM_STORED_LIGHTMAP_COEF = 4;
@@ -1127,268 +1027,8 @@ private:
 	bool bIsPointSource;
 };
 
-/** The uniform shader parameters associated with a primitive. */
-BEGIN_UNIFORM_BUFFER_STRUCT(FPrimitiveUniformShaderParameters,ENGINE_API)
-	DECLARE_UNIFORM_BUFFER_STRUCT_MEMBER(FMatrix,LocalToWorld)
-	DECLARE_UNIFORM_BUFFER_STRUCT_MEMBER(FMatrix,WorldToLocal)
-	DECLARE_UNIFORM_BUFFER_STRUCT_MEMBER(FVector4,ObjectWorldPositionAndRadius)
-	DECLARE_UNIFORM_BUFFER_STRUCT_MEMBER(FVector,ObjectBounds)
-	DECLARE_UNIFORM_BUFFER_STRUCT_MEMBER_EX(float,LocalToWorldDeterminantSign,EShaderPrecisionModifier::Half)
-	DECLARE_UNIFORM_BUFFER_STRUCT_MEMBER(FVector,ActorWorldPosition)
-	DECLARE_UNIFORM_BUFFER_STRUCT_MEMBER_EX(float,DecalReceiverMask,EShaderPrecisionModifier::Half)
-	DECLARE_UNIFORM_BUFFER_STRUCT_MEMBER_EX(float,HasDistanceFieldRepresentation,EShaderPrecisionModifier::Half)
-	DECLARE_UNIFORM_BUFFER_STRUCT_MEMBER_EX(FVector4,ObjectOrientation,EShaderPrecisionModifier::Half)
-	DECLARE_UNIFORM_BUFFER_STRUCT_MEMBER_EX(FVector4,NonUniformScale,EShaderPrecisionModifier::Half)
-	DECLARE_UNIFORM_BUFFER_STRUCT_MEMBER_EX(FVector4,InvNonUniformScale,EShaderPrecisionModifier::Half)
-	DECLARE_UNIFORM_BUFFER_STRUCT_MEMBER(FVector, LocalObjectBoundsMin)		// This is used in a custom material function (ObjectLocalBounds.uasset)
-	DECLARE_UNIFORM_BUFFER_STRUCT_MEMBER(FVector, LocalObjectBoundsMax)		// This is used in a custom material function (ObjectLocalBounds.uasset)
-	DECLARE_UNIFORM_BUFFER_STRUCT_MEMBER(float,LpvBiasMultiplier)
-END_UNIFORM_BUFFER_STRUCT(FPrimitiveUniformShaderParameters)
-
-/** Initializes the primitive uniform shader parameters. */
-inline FPrimitiveUniformShaderParameters GetPrimitiveUniformShaderParameters(
-	const FMatrix& LocalToWorld,
-	FVector ActorPosition,
-	const FBoxSphereBounds& WorldBounds,
-	const FBoxSphereBounds& LocalBounds,
-	bool bReceivesDecals,
-	bool bHasDistanceFieldRepresentation,
-	float LpvBiasMultiplier = 1.0f
-)
-{
-	FPrimitiveUniformShaderParameters Result;
-	Result.LocalToWorld = LocalToWorld;
-	Result.WorldToLocal = LocalToWorld.InverseSafe();
-	Result.ObjectWorldPositionAndRadius = FVector4(WorldBounds.Origin, WorldBounds.SphereRadius);
-	Result.ObjectBounds = WorldBounds.BoxExtent;
-	Result.LocalObjectBoundsMin = LocalBounds.GetBoxExtrema(0); // 0 == minimum
-	Result.LocalObjectBoundsMax = LocalBounds.GetBoxExtrema(1); // 1 == maximum
-	Result.ObjectOrientation = LocalToWorld.GetUnitAxis( EAxis::Z );
-	Result.ActorWorldPosition = ActorPosition;
-	Result.LpvBiasMultiplier = LpvBiasMultiplier;
-
-	{
-		// Extract per axis scales from LocalToWorld transform
-		FVector4 WorldX = FVector4(LocalToWorld.M[0][0],LocalToWorld.M[0][1],LocalToWorld.M[0][2],0);
-		FVector4 WorldY = FVector4(LocalToWorld.M[1][0],LocalToWorld.M[1][1],LocalToWorld.M[1][2],0);
-		FVector4 WorldZ = FVector4(LocalToWorld.M[2][0],LocalToWorld.M[2][1],LocalToWorld.M[2][2],0);
-		float ScaleX = FVector(WorldX).Size();
-		float ScaleY = FVector(WorldY).Size();
-		float ScaleZ = FVector(WorldZ).Size();
-		Result.NonUniformScale = FVector4(ScaleX,ScaleY,ScaleZ,0);
-		Result.InvNonUniformScale = FVector4(
-			ScaleX > KINDA_SMALL_NUMBER ? 1.0f/ScaleX : 0.0f,
-			ScaleY > KINDA_SMALL_NUMBER ? 1.0f/ScaleY : 0.0f,
-			ScaleZ > KINDA_SMALL_NUMBER ? 1.0f/ScaleZ : 0.0f,
-			0.0f
-			);
-	}
-
-	Result.LocalToWorldDeterminantSign = FMath::FloatSelect(LocalToWorld.RotDeterminant(),1,-1);
-	Result.DecalReceiverMask = bReceivesDecals ? 1 : 0;
-	Result.HasDistanceFieldRepresentation = bHasDistanceFieldRepresentation ? 1 : 0;
-	return Result;
-}
-
-inline TUniformBufferRef<FPrimitiveUniformShaderParameters> CreatePrimitiveUniformBufferImmediate(
-	const FMatrix& LocalToWorld,
-	const FBoxSphereBounds& WorldBounds,
-	const FBoxSphereBounds& LocalBounds,
-	bool bReceivesDecals,
-	float LpvBiasMultiplier = 1.0f
-	)
-{
-	check(IsInRenderingThread());
-	return TUniformBufferRef<FPrimitiveUniformShaderParameters>::CreateUniformBufferImmediate(
-		GetPrimitiveUniformShaderParameters(LocalToWorld, WorldBounds.Origin, WorldBounds, LocalBounds, bReceivesDecals, false, LpvBiasMultiplier ),
-		UniformBuffer_MultiFrame
-		);
-}
-
-/**
- * Primitive uniform buffer containing only identity transforms.
- */
-class FIdentityPrimitiveUniformBuffer : public TUniformBuffer<FPrimitiveUniformShaderParameters>
-{
-public:
-
-	/** Default constructor. */
-	FIdentityPrimitiveUniformBuffer()
-	{
-		SetContents(GetPrimitiveUniformShaderParameters(
-			FMatrix(
-				FPlane(1.0f,0.0f,0.0f,0.0f),
-				FPlane(0.0f,1.0f,0.0f,0.0f),
-				FPlane(0.0f,0.0f,1.0f,0.0f),
-				FPlane(0.0f,0.0f,0.0f,1.0f)
-				),
-			FVector(0.0f,0.0f,0.0f),
-			FBoxSphereBounds(EForceInit::ForceInit),
-			FBoxSphereBounds(EForceInit::ForceInit),
-			true,
-			false,
-			1.0f		// LPV bias
-			));
-	}
-};
-
-/** Global primitive uniform buffer resource containing identity transformations. */
-extern TGlobalResource<FIdentityPrimitiveUniformBuffer> GIdentityPrimitiveUniformBuffer;
-
-/**
- * A batch mesh element definition.
- */
-struct FMeshBatchElement
-{
-	const TUniformBuffer<FPrimitiveUniformShaderParameters>* PrimitiveUniformBufferResource;
-	TUniformBufferRef<FPrimitiveUniformShaderParameters> PrimitiveUniformBuffer;
-
-	const FIndexBuffer* IndexBuffer;
-	uint32 FirstIndex;
-	uint32 NumPrimitives;
-	uint32 NumInstances;
-	uint32 MinVertexIndex;
-	uint32 MaxVertexIndex;
-	int32 UserIndex;
-	void* UserData;
-	float MinScreenSize;
-	float MaxScreenSize;
-
-	/** 
-	 *	DynamicIndexData - pointer to user memory containing the index data.
-	 *	Used for rendering dynamic data directly.
-	 */
-	const void* DynamicIndexData;
-	uint16 DynamicIndexStride;
-	
-	FMeshBatchElement()
-	:	PrimitiveUniformBufferResource(NULL)
-	,	IndexBuffer(NULL)
-	,	NumInstances(1)
-	,	UserIndex(-1)
-	,	UserData(NULL)
-	,	MinScreenSize(0.0f)
-	,	MaxScreenSize(1.0f)
-	,	DynamicIndexData(NULL)
-	{
-	}
-};
-
-/**
- * A batch of mesh elements, all with the same material and vertex buffer
- */
-struct FMeshBatch
-{
-	TArray<FMeshBatchElement,TInlineAllocator<1> > Elements;
-
-	// used with DynamicVertexData
-	uint16 DynamicVertexStride;
-
-	/** LOD index of the mesh, used for fading LOD transitions. */
-	int8 LODIndex;
-
-	uint32 UseDynamicData : 1;
-	uint32 ReverseCulling : 1;
-	uint32 bDisableBackfaceCulling : 1;
-	uint32 CastShadow : 1;
-	uint32 bWireframe : 1;
-	// e.g. PT_TriangleList(default), PT_LineList, ..
-	uint32 Type : PT_NumBits;
-	// e.g. SDPG_World (default), SDPG_Foreground
-	uint32 DepthPriorityGroup : SDPG_NumBits;
-	uint32 bUseAsOccluder : 1;
-
-	// can be 0
-	const FLightCacheInterface* LCI;
-
-	/** 
-	 *	DynamicVertexData - pointer to user memory containing the vertex data.
-	 *	Used for rendering dynamic data directly.
-	 *  used with DynamicVertexStride
-	 */
-	const void* DynamicVertexData;
-
-	// can be 0
-	const FVertexFactory* VertexFactory;
-	// can be 0
-	const FMaterialRenderProxy* MaterialRenderProxy;
-
-
-	FORCEINLINE bool IsTranslucent(ERHIFeatureLevel::Type InFeatureLevel) const
-	{
-		// Note: blend mode does not depend on the feature level we are actually rendering in.
-		return MaterialRenderProxy && IsTranslucentBlendMode(MaterialRenderProxy->GetMaterial(InFeatureLevel)->GetBlendMode());
-	}
-
-	FORCEINLINE bool IsMasked(ERHIFeatureLevel::Type InFeatureLevel) const
-	{
-		// Note: blend mode does not depend on the feature level we are actually rendering in.
-		return MaterialRenderProxy && MaterialRenderProxy->GetMaterial(InFeatureLevel)->IsMasked();
-	}
-
-	/** Converts from an int32 index into a int8 */
-	static int8 QuantizeLODIndex(int32 NewLODIndex)
-	{
-		checkSlow(NewLODIndex >= SCHAR_MIN && NewLODIndex <= SCHAR_MAX);
-		return (int8)NewLODIndex;
-	}
-
-	/** 
-	* @return vertex stride specified for the mesh. 0 if not dynamic
-	*/
-	FORCEINLINE uint32 GetDynamicVertexStride(ERHIFeatureLevel::Type /*InFeatureLevel*/) const
-	{
-		if (UseDynamicData && DynamicVertexData)
-		{
-			return DynamicVertexStride;
-		}
-		else
-		{
-			return 0;
-		}
-	}
-
-	FORCEINLINE int32 GetNumPrimitives() const
-	{
-		int32 Count=0;
-		for( int32 ElementIdx=0;ElementIdx<Elements.Num();ElementIdx++ )
-		{
-			Count += Elements[ElementIdx].NumPrimitives;
-		}
-		return Count;
-	}
 
-#if DO_CHECK
-	FORCEINLINE void CheckUniformBuffers() const
-	{
-		for( int32 ElementIdx=0;ElementIdx<Elements.Num();ElementIdx++ )
-		{
-			check(IsValidRef(Elements[ElementIdx].PrimitiveUniformBuffer) || Elements[ElementIdx].PrimitiveUniformBufferResource != NULL);
-		}
-	}
-#endif	
 
-	/** Default constructor. */
-	FMeshBatch()
-	:	DynamicVertexStride(0)
-	,	LODIndex(INDEX_NONE)
-	,	UseDynamicData(false)
-	,	ReverseCulling(false)
-	,	bDisableBackfaceCulling(false)
-	,	CastShadow(true)
-	,	bWireframe(false)
-	,	Type(PT_TriangleList)
-	,	DepthPriorityGroup(SDPG_World)
-	,	bUseAsOccluder(true)
-	,	LCI(NULL)
-	,	DynamicVertexData(NULL)
-	,	VertexFactory(NULL)
-	,	MaterialRenderProxy(NULL)
-	{
-		// By default always add the first element.
-		new(Elements) FMeshBatchElement;
-	}
-};
 
 /**
  * An interface implemented by dynamic resources which need to be initialized and cleaned up by the rendering thread.
@@ -1508,109 +1148,7 @@ public:
 		) = 0;
 };
 
-/**
- * The different types of relevance a primitive scene proxy can declare towards a particular scene view.
- */
-struct FPrimitiveViewRelevance
-{
-	/** The primitive's static elements are rendered for the view. */
-	uint32 bStaticRelevance : 1; 
-	/** The primitive's dynamic elements are rendered for the view. */
-	uint32 bDynamicRelevance : 1;
-	/** The primitive is drawn. */
-	uint32 bDrawRelevance : 1;
-	/** The primitive is casting a shadow. */
-	uint32 bShadowRelevance : 1;
-	/** The primitive should render to the custom depth pass. */
-	uint32 bRenderCustomDepth : 1;
-	/** The primitive should render to the main depth pass. */
-	uint32 bRenderInMainPass : 1;
-	/** The primitive is drawn only in the editor and composited onto the scene after post processing */
-	uint32 bEditorPrimitiveRelevance : 1;
-	/** The primitive is drawn only in the editor and composited onto the scene after post processing using no depth testing */
-	uint32 bEditorNoDepthTestPrimitiveRelevance : 1;
-	/** The primitive needs PreRenderView to be called before rendering. */
-	uint32 bNeedsPreRenderView : 1;
-	/** The primitive should have GatherSimpleLights called on the proxy when gathering simple lights. */
-	uint32 bHasSimpleLights : 1;
-
-	/** The primitive has one or more opaque or masked elements. */
-	uint32 bOpaqueRelevance : 1;
-	/** The primitive has one or more masked elements. */
-	uint32 bMaskedRelevance : 1;
-	/** The primitive has one or more distortion elements. */
-	uint32 bDistortionRelevance : 1;
-	/** The primitive reads from scene color. */
-	uint32 bSceneColorRelevance : 1;
-	/** The primitive has one or more elements that have SeparateTranslucency. */
-	uint32 bSeparateTranslucencyRelevance : 1;
-	/** The primitive has one or more elements that have normal translucency. */
-	uint32 bNormalTranslucencyRelevance : 1;
-
-	/** 
-	 * Whether this primitive view relevance has been initialized this frame.  
-	 * Primitives that have not had ComputeRelevanceForView called on them (because they were culled) will not be initialized,
-	 * But we may still need to render them from other views like shadow passes, so this tracks whether we can reuse the cached relevance or not.
-	 */
-	uint32 bInitializedThisFrame : 1;
-
-	bool HasTranslucency() const 
-	{
-		return bSeparateTranslucencyRelevance || bNormalTranslucencyRelevance;
-	}
-
-	/** Initialization constructor. */
-	FPrimitiveViewRelevance():
-		bStaticRelevance(false),
-		bDynamicRelevance(false),
-		bDrawRelevance(false),
-		bShadowRelevance(false),
-		bRenderCustomDepth(false),
-		bRenderInMainPass(true),
-		bEditorPrimitiveRelevance(false),
-		bEditorNoDepthTestPrimitiveRelevance(false),
-		bNeedsPreRenderView(false),
-		bHasSimpleLights(false),
-		bOpaqueRelevance(true),
-		bMaskedRelevance(false),
-		bDistortionRelevance(false),
-		bSceneColorRelevance(false),
-		bSeparateTranslucencyRelevance(false),
-		bNormalTranslucencyRelevance(false),
-		bInitializedThisFrame(false)
-	{}
-
-	/** Bitwise OR operator.  Sets any relevance bits which are present in either FPrimitiveViewRelevance. */
-	FPrimitiveViewRelevance& operator|=(const FPrimitiveViewRelevance& B)
-	{
-		bStaticRelevance |= B.bStaticRelevance != 0;
-		bDynamicRelevance |= B.bDynamicRelevance != 0;
-		bDrawRelevance |= B.bDrawRelevance != 0;
-		bShadowRelevance |= B.bShadowRelevance != 0;
-		bOpaqueRelevance |= B.bOpaqueRelevance != 0;
-		bMaskedRelevance |= B.bMaskedRelevance != 0;
-		bDistortionRelevance |= B.bDistortionRelevance != 0;
-		bRenderCustomDepth |= B.bRenderCustomDepth != 0;
-		bRenderInMainPass |= B.bRenderInMainPass !=0;
-		bEditorPrimitiveRelevance |= B.bEditorPrimitiveRelevance !=0;
-		bEditorNoDepthTestPrimitiveRelevance |= B.bEditorNoDepthTestPrimitiveRelevance !=0;
-		bSceneColorRelevance |= B.bSceneColorRelevance != 0;
-		bNeedsPreRenderView |= B.bNeedsPreRenderView != 0;
-		bHasSimpleLights |= B.bHasSimpleLights != 0;
-		bSeparateTranslucencyRelevance |= B.bSeparateTranslucencyRelevance != 0;
-		bNormalTranslucencyRelevance |= B.bNormalTranslucencyRelevance != 0;
-		bInitializedThisFrame |= B.bInitializedThisFrame;
-		return *this;
-	}
 
-	/** Binary bitwise OR operator. */
-	friend FPrimitiveViewRelevance operator|(const FPrimitiveViewRelevance& A,const FPrimitiveViewRelevance& B)
-	{
-		FPrimitiveViewRelevance Result(A);
-		Result |= B;
-		return Result;
-	}
-};
 
 /**
  *	Helper structure for storing motion blur information for a primitive
@@ -1740,313 +1278,18 @@ private:
 	FMotionBlurInfo* FindMBInfoIndex(FPrimitiveComponentId ComponentId);
 };
 
-/**
- * An interface to the private scene manager implementation of a scene.  Use GetRendererModule().AllocateScene to create.
- * The scene
- */
-class FSceneInterface
-{
-public:
-
-	// FSceneInterface interface
-
-	/** 
-	 * Adds a new primitive component to the scene
-	 * 
-	 * @param Primitive - primitive component to add
-	 */
-	virtual void AddPrimitive(UPrimitiveComponent* Primitive) = 0;
-	/** 
-	 * Removes a primitive component from the scene
-	 * 
-	 * @param Primitive - primitive component to remove
-	 */
-	virtual void RemovePrimitive(UPrimitiveComponent* Primitive) = 0;
-	/** Called when a primitive is being unregistered and will not be immediately re-registered. */
-	virtual void ReleasePrimitive(UPrimitiveComponent* Primitive) = 0;
-	/** 
-	 * Updates the transform of a primitive which has already been added to the scene. 
-	 * 
-	 * @param Primitive - primitive component to update
-	 */
-	virtual void UpdatePrimitiveTransform(UPrimitiveComponent* Primitive) = 0;
-	/** Updates primitive attachment state. */
-	virtual void UpdatePrimitiveAttachment(UPrimitiveComponent* Primitive) = 0;
-	/** 
-	 * Adds a new light component to the scene
-	 * 
-	 * @param Light - light component to add
-	 */
-	virtual void AddLight(ULightComponent* Light) = 0;
-	/** 
-	 * Removes a light component from the scene
-	 * 
-	 * @param Light - light component to remove
-	 */
-	virtual void RemoveLight(ULightComponent* Light) = 0;
-	/** 
-	 * Adds a new light component to the scene which is currently invisible, but needed for editor previewing
-	 * 
-	 * @param Light - light component to add
-	 */
-	virtual void AddInvisibleLight(ULightComponent* Light) = 0;
-	virtual void SetSkyLight(FSkyLightSceneProxy* Light) = 0;
-	/** 
-	 * Adds a new decal component to the scene
-	 * 
-	 * @param Component - component to add
-	 */
-	virtual void AddDecal(UDecalComponent* Component) = 0;
-	/** 
-	 * Removes a decal component from the scene
-	 * 
-	 * @param Component - component to remove
-	 */
-	virtual void RemoveDecal(UDecalComponent* Component) = 0;
-	/** 
-	 * Updates the transform of a decal which has already been added to the scene. 
-	 *
-	 * @param Decal - Decal component to update
-	 */
-	virtual void UpdateDecalTransform(UDecalComponent* Component) = 0;
-
-	/** Adds a reflection capture to the scene. */
-	virtual void AddReflectionCapture(class UReflectionCaptureComponent* Component) {}
-
-	/** Removes a reflection capture from the scene. */
-	virtual void RemoveReflectionCapture(class UReflectionCaptureComponent* Component) {}
-
-	/** Reads back reflection capture data from the GPU.  Very slow operation that blocks the GPU and rendering thread many times. */
-	virtual void GetReflectionCaptureData(UReflectionCaptureComponent* Component, class FReflectionCaptureFullHDRDerivedData& OutDerivedData) {}
-
-	/** Updates a reflection capture's transform, and then re-captures the scene. */
-	virtual void UpdateReflectionCaptureTransform(class UReflectionCaptureComponent* Component) {}
-
-	/** 
-	 * Allocates reflection captures in the scene's reflection cubemap array and updates them by recapturing the scene.
-	 * Existing captures will only be updated.  Must be called from the game thread.
-	 */
-	virtual void AllocateReflectionCaptures(const TArray<UReflectionCaptureComponent*>& NewCaptures) {}
-	virtual void ReleaseReflectionCubemap(UReflectionCaptureComponent* CaptureComponent) {}
-
-	/** 
-	 * Updates the contents of the given sky capture by rendering the scene. 
-	 * This must be called on the game thread.
-	 */
-	virtual void UpdateSkyCaptureContents(const USkyLightComponent* CaptureComponent, bool bCaptureEmissiveOnly, FTexture* OutProcessedTexture, FSHVectorRGB3& OutIrradianceEnvironmentMap) {}
-
-	/** 
-	* Updates the contents of the given scene capture by rendering the scene. 
-	* This must be called on the game thread.
-	*/
-	virtual void UpdateSceneCaptureContents(class USceneCaptureComponent2D* CaptureComponent) {}
-	virtual void UpdateSceneCaptureContents(class USceneCaptureComponentCube* CaptureComponent) {}
-
-	virtual void AddPrecomputedLightVolume(const class FPrecomputedLightVolume* Volume) {}
-	virtual void RemovePrecomputedLightVolume(const class FPrecomputedLightVolume* Volume) {}
-
-	/** 
-	 * Updates the transform of a light which has already been added to the scene. 
-	 *
-	 * @param Light - light component to update
-	 */
-	virtual void UpdateLightTransform(ULightComponent* Light) = 0;
-	/** 
-	 * Updates the color and brightness of a light which has already been added to the scene. 
-	 *
-	 * @param Light - light component to update
-	 */
-	virtual void UpdateLightColorAndBrightness(ULightComponent* Light) = 0;
-
-	/** Updates the scene's dynamic skylight. */
-	virtual void UpdateDynamicSkyLight(const FLinearColor& UpperColor, const FLinearColor& LowerColor) {}
-
-	/** Sets the precomputed visibility handler for the scene, or NULL to clear the current one. */
-	virtual void SetPrecomputedVisibility(const class FPrecomputedVisibilityHandler* PrecomputedVisibilityHandler) {}
-
-	/** Sets the precomputed volume distance field for the scene, or NULL to clear the current one. */
-	virtual void SetPrecomputedVolumeDistanceField(const class FPrecomputedVolumeDistanceField* PrecomputedVolumeDistanceField) {}
-
-	/** Sets shader maps on the specified materials without blocking. */
-	virtual void SetShaderMapsOnMaterialResources(const TMap<FMaterial*, FMaterialShaderMap*>& MaterialsToUpdate) {}
-
-	/** Updates static draw lists for the given set of materials. */
-	virtual void UpdateStaticDrawListsForMaterials(const TArray<const FMaterial*>& Materials) {}
-
-	/** 
-	 * Adds a new exponential height fog component to the scene
-	 * 
-	 * @param FogComponent - fog component to add
-	 */	
-	virtual void AddExponentialHeightFog(class UExponentialHeightFogComponent* FogComponent) = 0;
-	/** 
-	 * Removes a exponential height fog component from the scene
-	 * 
-	 * @param FogComponent - fog component to remove
-	 */	
-	virtual void RemoveExponentialHeightFog(class UExponentialHeightFogComponent* FogComponent) = 0;
-
-	/** 
-	 * Adds a new atmospheric fog component to the scene
-	 * 
-	 * @param FogComponent - fog component to add
-	 */	
-	virtual void AddAtmosphericFog(class UAtmosphericFogComponent* FogComponent) = 0;
-
-	/** 
-	 * Removes a atmospheric fog component from the scene
-	 * 
-	 * @param FogComponent - fog component to remove
-	 */	
-	virtual void RemoveAtmosphericFog(class UAtmosphericFogComponent* FogComponent) = 0;
-
-	/**
-	 * Adds a wind source component to the scene.
-	 * @param WindComponent - The component to add.
-	 */
-	virtual void AddWindSource(class UWindDirectionalSourceComponent* WindComponent) = 0;
-	/**
-	 * Removes a wind source component from the scene.
-	 * @param WindComponent - The component to remove.
-	 */
-	virtual void RemoveWindSource(class UWindDirectionalSourceComponent* WindComponent) = 0;
-	/**
-	 * Accesses the wind source list.  Must be called in the rendering thread.
-	 * @return The list of wind sources in the scene.
-	 */
-	virtual const TArray<class FWindSourceSceneProxy*>& GetWindSources_RenderThread() const = 0;
-
-	/** Accesses wind parameters.  XYZ will contain wind direction * Strength, W contains wind speed. */
-	virtual FVector4 GetWindParameters(const FVector& Position) const = 0;
-
-	/** Same as GetWindParameters, but ignores point wind sources. */
-	virtual FVector4 GetDirectionalWindParameters() const = 0;
-
-	/** 
-	 * Adds a SpeedTree wind computation object to the scene.
-	 * @param StaticMesh - The SpeedTree static mesh whose wind to add.
-	 */
-	virtual void AddSpeedTreeWind(class FVertexFactory* VertexFactory, class UStaticMesh* StaticMesh) = 0;
-
-	/** 
-	 * Removes a SpeedTree wind computation object to the scene.
-	 * @param StaticMesh - The SpeedTree static mesh whose wind to remove.
-	 */
-	virtual void RemoveSpeedTreeWind(class FVertexFactory* VertexFactory, class UStaticMesh* StaticMesh) = 0;
 
-	/** Ticks the SpeedTree wind object and updates the uniform buffer. */
-	virtual void UpdateSpeedTreeWind(double CurrentTime) = 0;
 
-	/** 
-	 * Looks up the SpeedTree uniform buffer for the passed in vertex factory.
-	 * @param VertexFactory - The vertex factory registered for SpeedTree.
-	 */
-	virtual FUniformBufferRHIParamRef GetSpeedTreeUniformBuffer(const FVertexFactory* VertexFactory) = 0;
+/** 
+ * Enumeration for currently used translucent lighting volume cascades 
+ */
+enum ETranslucencyVolumeCascade
+{
+	TVC_Inner,
+	TVC_Outer,
 
-	/**
-	 * Release this scene and remove it from the rendering thread
-	 */
-	virtual void Release() = 0;
-	/**
-	 * Retrieves the lights interacting with the passed in primitive and adds them to the out array.
-	 *
-	 * @param	Primitive				Primitive to retrieve interacting lights for
-	 * @param	RelevantLights	[out]	Array of lights interacting with primitive
-	 */
-	virtual void GetRelevantLights( UPrimitiveComponent* Primitive, TArray<const ULightComponent*>* RelevantLights ) const = 0;
-	/**
-	 * Indicates if hit proxies should be processed by this scene
-	 *
-	 * @return true if hit proxies should be rendered in this scene.
-	 */
-	virtual bool RequiresHitProxies() const = 0;
-	/**
-	 * Get the optional UWorld that is associated with this scene
-	 * 
- 	 * @return UWorld instance used by this scene
-	 */
-	virtual class UWorld* GetWorld() const = 0;
-	/**
-	 * Return the scene to be used for rendering. Note that this can return NULL if rendering has
-	 * been disabled!
-	 */
-	virtual class FScene* GetRenderScene()
-	{
-		return NULL;
-	}
-	/**
-	 * Sets the FX system associated with the scene.
-	 */
-	virtual void SetFXSystem( class FFXSystemInterface* InFXSystem ) = 0;
-
-	/**
-	 * Get the FX system associated with the scene.
-	 */
-	virtual class FFXSystemInterface* GetFXSystem() = 0;
-
-	virtual void DumpUnbuiltLightIteractions( FOutputDevice& Ar ) const { }
-
-	/**
-	 * Dumps static mesh draw list stats to the log.
-	 */
-	virtual void DumpStaticMeshDrawListStats() const {}
-
-	/**
-	 * Request to clear the MB info. Game thread only
-	 */
-	virtual void SetClearMotionBlurInfoGameThread() {}
-
-	/** Updates the scene's list of parameter collection id's and their uniform buffers. */
-	virtual void UpdateParameterCollections(const TArray<class FMaterialParameterCollectionInstanceResource*>& InParameterCollections) {}
-
-	/**
-	 * Exports the scene.
-	 *
-	 * @param	Ar		The Archive used for exporting.
-	 **/
-	virtual void Export( FArchive& Ar ) const
-	{}
-
-	
-	/**
-	 * Shifts scene data by provided delta
-	 * Called on world origin changes
-	 * 
-	 * @param	InOffset	Delta to shift scene by
-	 */
-	virtual void ApplyWorldOffset(FVector InOffset) {}
-
-	/**
-	 * Notification that level was added to a world
-	 * 
-	 * @param	InLevelName		Level name
-	 */
-	virtual void OnLevelAddedToWorld(FName InLevelName) {}
-
-	/**
-	 * @return True if there are any lights in the scene
-	 */
-	virtual bool HasAnyLights() const = 0;
-
-	virtual bool IsEditorScene() const { return false; }
-
-	virtual ERHIFeatureLevel::Type GetFeatureLevel() const { return GRHIFeatureLevel; }
-	virtual void ChangeFeatureLevel(ERHIFeatureLevel::Type InFeatureLevel) {}
-
-protected:
-	virtual ~FSceneInterface() {}
-};
-
-/** 
- * Enumeration for currently used translucent lighting volume cascades 
- */
-enum ETranslucencyVolumeCascade
-{
-	TVC_Inner,
-	TVC_Outer,
-
-	TVC_MAX,
-};
+	TVC_MAX,
+};
 
 /** The uniform shader parameters associated with a view. */
 BEGIN_UNIFORM_BUFFER_STRUCT(FViewUniformShaderParameters,ENGINE_API)
@@ -2137,823 +1380,6 @@ BEGIN_UNIFORM_BUFFER_STRUCT(FViewUniformShaderParameters,ENGINE_API)
 	DECLARE_UNIFORM_BUFFER_STRUCT_MEMBER_ARRAY(FVector4,SkyIrradianceEnvironmentMap,[7])
 END_UNIFORM_BUFFER_STRUCT(FViewUniformShaderParameters)
 
-// used to blend IBlendableInterface object data, the member act as a header for a following data block
-struct FBlendableEntry
-{
-	// weight 0..1, 0 to disable this entry
-	float Weight;
-
-private:
-	// to ensure type safety
-	FName BlendableType;
-	// to be able to jump over data
-	uint32 DataSize;
-
-	// @return pointer to the next object or end (can be compared with container end pointer)
-	uint8* GetDataPtr() { check(this); return (uint8*)(this + 1); }
-	// @return next or end of the array
-	FBlendableEntry* GetNext() { return (FBlendableEntry*)(GetDataPtr() + DataSize); }
-
-	friend class FBlendableManager;
-};
-
-// Manager class which allows blending of FBlendableBase objects, stores a copy (for the render thread)
-class FBlendableManager
-{
-public:
-
-	FBlendableManager()
-	{
-		// can be increased if needed (to avoid reallocations and copies at the cost of more memory allocation)
-		Scratch.Reserve(10 * 1024);
-	}
-	
-	// @param InWeight 0..1, excluding 0 as this is used to disable entries
-	// @param InData is copied with a memcpy
-	template <class T>
-	void PushBlendableData(float InWeight, const T& InData)
-	{
-		FName BlendableType = T::GetFName();
-		
-		PushBlendableDataPtr(InWeight, BlendableType, (const uint8*)&InData, sizeof(T));
-	}
-
-	// find next that has the given type, O(n), n is number of elements after the given one or all if started with 0
-	// @param InIterator 0 to start iteration
-	// @return 0 if no further one of that type was found
-	template <class T>
-	T* IterateBlendables(FBlendableEntry*& InIterator) const
-	{
-		FName BlendableType = T::GetFName();
-
-		do
-		{
-			InIterator = GetNextBlendableEntryPtr(InIterator);
-
-		} while(InIterator && InIterator->Weight <= 0.0f && InIterator->BlendableType != BlendableType);
-
-		if(InIterator)
-		{
-			return (T*)InIterator->GetDataPtr();
-		}
-
-		return 0;
-	}
-
-private:
-
-	// stored multiple elements with a FBlendableEntry header and following data of a size specified in the header
-	TArray<uint8> Scratch;
-
-	// find next, no restriction on the type O(n), n is number of elements after the given one or all if started with 0
-	FBlendableEntry* GetNextBlendableEntryPtr(FBlendableEntry* InIterator = 0) const
-	{
-		if(!InIterator)
-		{
-			// start at the beginning
-			InIterator = (FBlendableEntry*)Scratch.GetData();
-		}
-		else
-		{
-			// go to next
-			InIterator = InIterator->GetNext();
-		}
-
-		// end reached?
-		if((uint8*)InIterator == Scratch.GetData() + Scratch.Num())
-		{
-			// no more entries
-			return 0;
-		}
-
-		return InIterator;
-	}
-
-	// @param InWeight 0..1, excluding 0 as this is used to disable entries
-	// @param InData is copied
-	// @param InDataSize >0
-	void PushBlendableDataPtr(float InWeight, FName InBlendableType, const uint8* InData, uint32 InDataSize)
-	{
-		check(InWeight > 0.0f && InWeight <= 1.0f);
-		check(InData);
-		check(InDataSize);
-
-		uint32 OldSize = Scratch.AddUninitialized(sizeof(FBlendableEntry) + InDataSize);
-
-		FBlendableEntry* Dst = (FBlendableEntry*)&Scratch[OldSize];
-
-		Dst->Weight = InWeight;
-		Dst->BlendableType = InBlendableType;
-		Dst->DataSize = InDataSize;
-		memcpy(Dst->GetDataPtr(), InData, InDataSize);
-	}
-};
-
-
-/** All blended postprocessing in one place, non lerpable data is stored in non merged form */
-class FFinalPostProcessSettings : public FPostProcessSettings
-{
-public:
-	FFinalPostProcessSettings()
-	: HighResScreenshotMaterial(NULL)
-	, HighResScreenshotMaskMaterial(NULL)
-	, HighResScreenshotCaptureRegionMaterial(NULL)
-	, bBufferVisualizationDumpRequired(false)
-	{
-		// to avoid reallocations we reserve some elements
-
-		ContributingCubemaps.Empty(8);
-		ContributingLUTs.Empty(8);
-
-		SetLUT(0);
-	}
-	
-	struct FCubemapEntry
-	{
-		// 0..
-		FLinearColor AmbientCubemapTintMulScaleValue;
-		// can be 0
-		class UTexture* AmbientCubemap;
-
-		FCubemapEntry() :
-			AmbientCubemapTintMulScaleValue(FLinearColor(0, 0, 0, 0)),
-			AmbientCubemap(NULL)
-		{}
-	};
-
-	struct FLUTBlenderEntry
-	{
-		// 0..1
-		float Weight;
-		// can be 0
-		class UTexture* LUTTexture;
-	};
-
-	// Updated existing one or creates a new one
-	// this allows to combined volumes for fading between them but also to darken/disable volumes in some areas
-	void UpdateEntry(const FCubemapEntry &Entry, float Weight)
-	{
-		bool Existing = false;
-		for(int32 i = 0; i < ContributingCubemaps.Num(); ++i)
-		{
-			FCubemapEntry& Local = ContributingCubemaps[i];
-			
-			if(Local.AmbientCubemap == Entry.AmbientCubemap)
-			{
-				Local.AmbientCubemapTintMulScaleValue = FMath::Lerp(Local.AmbientCubemapTintMulScaleValue, Entry.AmbientCubemapTintMulScaleValue, Weight);
-				Existing = true;
-			}
-			else
-			{
-				Local.AmbientCubemapTintMulScaleValue *= 1.0f - Weight;
-			}
-
-			if(Local.AmbientCubemapTintMulScaleValue.IsAlmostBlack())
-			{
-				ContributingCubemaps.RemoveAt(i);
-			}
-		}
-
-		if( !Existing )
-		{
-			// We didn't find the entry previously, so we need to Lerp up from zero.
-			FCubemapEntry WeightedEntry = Entry;
-			WeightedEntry.AmbientCubemapTintMulScaleValue *= Weight;
-			if(!WeightedEntry.AmbientCubemapTintMulScaleValue.IsAlmostBlack()) // Only bother to add it if the lerped value isn't near-black.
-			{
-				ContributingCubemaps.Push(WeightedEntry);
-			}
-		}
-	}
-
-	// @param InTexture - must not be 0
-	// @param Weight - 0..1
-	void LerpTo(UTexture* InTexture, float Weight)
-	{
-		check(InTexture);
-		check(Weight >= 0.0f && Weight <= 1.0f);
-
-		if(Weight > 254.0f/255.0f || !ContributingLUTs.Num())
-		{
-			SetLUT(InTexture);
-			return;
-		}
-
-		for(uint32 i = 0; i < (uint32)ContributingLUTs.Num(); ++i)
-		{
-			ContributingLUTs[i].Weight *= 1.0f - Weight;
-		}
-
-		uint32 ExistingIndex = FindIndex(InTexture);
-		if(ExistingIndex == 0xffffffff)
-		{
-			PushLUT(InTexture, Weight);
-		}
-		else
-		{
-			ContributingLUTs[ExistingIndex].Weight += Weight;
-		}
-	}
-	
-	/**
-	 * add a LUT(look up table) to the ones that are blended together
-	 * @param Texture - can be 0
-	 * @param Weight - 0..1
-	 */
-	void PushLUT(UTexture* Texture, float Weight)
-	{
-		check(Weight >= 0.0f && Weight <= 1.0f);
-
-		FLUTBlenderEntry Entry;
-
-		Entry.LUTTexture = Texture;
-		Entry.Weight = Weight;
-
-		ContributingLUTs.Add(Entry);
-	}
-
-	/** @return 0xffffffff if not found */
-	uint32 FindIndex(UTexture* Tex) const
-	{
-		for(uint32 i = 0; i < (uint32)ContributingLUTs.Num(); ++i)
-		{
-			if(ContributingLUTs[i].LUTTexture == Tex)
-			{
-				return i;
-			}
-		}
-
-		return 0xffffffff;
-	}
-
-	void SetLUT(UTexture* Texture)
-	{
-		// intentionally no deallocations
-		ContributingLUTs.Empty();
-
-		PushLUT(Texture, 1.0f);
-	}
-
-	// was not merged during blending unlike e.g. BloomThreshold 
-	TArray<FCubemapEntry> ContributingCubemaps;
-	TArray<FLUTBlenderEntry> ContributingLUTs;
-
-	// List of materials to use in the buffer visualization overview
-	TArray<UMaterialInterface*> BufferVisualizationOverviewMaterials;
-
-	// Material to use for rendering high res screenshot with mask. Post process expects this material to be set all the time.
-	UMaterialInterface* HighResScreenshotMaterial;
-
-	// Material to use for rendering just the high res screenshot mask. Post process expects this material to be set all the time.
-	UMaterialInterface* HighResScreenshotMaskMaterial;
-
-	// Material to use for rendering the high resolution screenshot capture region. Post processing only draws the region if this material is set.
-	UMaterialInstanceDynamic* HighResScreenshotCaptureRegionMaterial;
-
-	// Current buffer visualization dumping settings
-	bool bBufferVisualizationDumpRequired;
-	FString BufferVisualizationDumpBaseFilename;
-
-	// Maintains a container with IBlendableInterface objects and their data
-	FBlendableManager BlendableManager;
-};
-
-// Projection data for a FSceneView
-struct FSceneViewProjectionData
-{
-	FMatrix ViewMatrix;
-
-	/** UE4 projection matrix projects such that clip space Z=1 is the near plane, and Z=0 is the infinite far plane. */
-	FMatrix ProjectionMatrix;
-
-protected:
-	//The unconstrained (no aspect ratio bars applied) view rectangle (also unscaled)
-	FIntRect ViewRect;
-
-	// The constrained view rectangle (identical to UnconstrainedUnscaledViewRect if aspect ratio is not constrained)
-	FIntRect ConstrainedViewRect;
-
-public:
-	void SetViewRectangle(const FIntRect& InViewRect)
-	{
-		ViewRect = InViewRect;
-		ConstrainedViewRect = InViewRect;
-	}
-
-	void SetConstrainedViewRectangle(const FIntRect& InViewRect)
-	{
-		ConstrainedViewRect = InViewRect;
-	}
-
-	bool IsValidViewRectangle() const
-	{
-		return (ConstrainedViewRect.Min.X >= 0) &&
-			(ConstrainedViewRect.Min.Y >= 0) &&
-			(ConstrainedViewRect.Width() > 0) &&
-			(ConstrainedViewRect.Height() > 0);
-	}
-
-	const FIntRect& GetViewRect() const { return ViewRect; }
-	const FIntRect& GetConstrainedViewRect() const { return ConstrainedViewRect; }
-};
-
-// Construction parameters for a FSceneView
-struct FSceneViewInitOptions : public FSceneViewProjectionData
-{
-	const FSceneViewFamily* ViewFamily;
-	FSceneViewStateInterface* SceneViewStateInterface;
-	const AActor* ViewActor;
-	FViewElementDrawer* ViewElementDrawer;
-
-	FLinearColor BackgroundColor;
-	FLinearColor OverlayColor;
-	FLinearColor ColorScale;
-
-	/** For stereoscopic rendering, whether or not this is a full pass, or a left / right eye pass */
-	EStereoscopicPass StereoPass;
-
-	/** Conversion from world units (uu) to meters, so we can scale motion to the world appropriately */
-	float WorldToMetersScale;
-
-	TSet<FPrimitiveComponentId> HiddenPrimitives;
-
-	// -1,-1 if not setup
-	FIntPoint CursorPos;
-
-	float LODDistanceFactor;
-
-	/** If > 0, overrides the view's far clipping plane with a plane at the specified distance. */
-	float OverrideFarClippingPlaneDistance;
-
-	// Was there a camera cut this frame?
-	bool bInCameraCut;
-
-#if WITH_EDITOR
-	// default to 0'th view index, which is a bitfield of 1
-	uint64 EditorViewBitflag;
-
-	// this can be specified for ortho views so that it's min draw distance/LOD parenting etc, is controlled by a perspective viewport
-	FVector OverrideLODViewOrigin;
-
-	// In case of ortho, generate a fake view position that has a non-zero W component. The view position will be derived based on the view matrix.
-	bool bUseFauxOrthoViewPos;
-#endif
-
-	FSceneViewInitOptions()
-		: ViewFamily(NULL)
-		, SceneViewStateInterface(NULL)
-		, ViewActor(NULL)
-		, ViewElementDrawer(NULL)
-		, BackgroundColor(FLinearColor::Transparent)
-		, OverlayColor(FLinearColor::Transparent)
-		, ColorScale(FLinearColor::White)
-		, StereoPass(eSSP_FULL)
-		, WorldToMetersScale(100.f)
-		, CursorPos(-1, -1)
-		, LODDistanceFactor(1.0f)
-		, OverrideFarClippingPlaneDistance(-1.0f)
-		, bInCameraCut(false)
-#if WITH_EDITOR
-		, EditorViewBitflag(1)
-		, OverrideLODViewOrigin(ForceInitToZero)
-		, bUseFauxOrthoViewPos(false)
-		//@TODO: , const TBitArray<>& InSpriteCategoryVisibility=TBitArray<>()
-#endif
-	{
-	}
-};
-
-/**
- * A projection from scene space into a 2D screen region.
- */
-class ENGINE_API FSceneView
-{
-public:
-	const FSceneViewFamily* Family;
-	/** can be 0 (thumbnail rendering) */
-	FSceneViewStateInterface* State;
-
-	/** The uniform buffer for the view's parameters.  This is only initialized in the rendering thread's copies of the FSceneView. */
-	TUniformBufferRef<FViewUniformShaderParameters> UniformBuffer;
-
-	/** The actor which is being viewed from. */
-	const AActor* ViewActor;
-
-	/** An interaction which draws the view's interaction elements. */
-	FViewElementDrawer* Drawer;
-
-	/* Final position of the view in the final render target (in pixels), potentially scaled down by ScreenPercentage */
-	FIntRect ViewRect;
-
-	/* Final position of the view in the final render target (in pixels), potentially constrained by an aspect ratio requirement (black bars) */
-	const FIntRect UnscaledViewRect;
-
-	/* Raw view size (in pixels), used for screen space calculations */
-	const FIntRect UnconstrainedViewRect;
-
-	/** 
-	 * Copy from GFrameNumber
-	 */
-	uint32 FrameNumber;
-
-	/** Maximum number of shadow cascades to render with. */
-	int32 MaxShadowCascades;
-
-    FViewMatrices ViewMatrices;
-
-    /** Variables used to determine the view matrix */
-    FVector		ViewLocation;
-    FRotator	ViewRotation;
-	FQuat		BaseHmdOrientation;
-	FVector		BaseHmdLocation;
-	float		WorldToMetersScale;
-
-	// normally the same as ViewMatrices unless "r.Shadow.FreezeCamera" is activated
-	FViewMatrices ShadowViewMatrices;
-
-	FMatrix ProjectionMatrixUnadjustedForRHI;
-
-	FLinearColor BackgroundColor;
-	FLinearColor OverlayColor;
-
-	/** Color scale multiplier used during post processing */
-	FLinearColor ColorScale;
-
-	/** For stereoscopic rendering, whether or not this is a full pass, or a left / right eye pass */
-	EStereoscopicPass StereoPass;
-
-	/** Current buffer visualization mode */
-	FName CurrentBufferVisualizationMode;
-
-	/**
-	* These can be used to override material parameters across the scene without recompiling shaders.
-	* The last component is how much to include of the material's value for that parameter, so 0 will completely remove the material's value.
-	*/
-	FVector4 DiffuseOverrideParameter;
-	FVector4 SpecularOverrideParameter;
-	FVector4 NormalOverrideParameter;
-	FVector2D RoughnessOverrideParameter;
-
-	/** The primitives which are hidden for this view. */
-	TSet<FPrimitiveComponentId> HiddenPrimitives;
-
-	// Derived members.
-
-	/** redundant, ViewMatrices.GetViewProjMatrix() */
-	/* UE4 projection matrix projects such that clip space Z=1 is the near plane, and Z=0 is the infinite far plane. */
-	FMatrix ViewProjectionMatrix;
-	/** redundant, ViewMatrices.GetInvViewMatrix() */
-	FMatrix InvViewMatrix;				
-	/** redundant, ViewMatrices.GetInvViewProjMatrix() */
-	FMatrix InvViewProjectionMatrix;	
-
-	float TemporalJitterPixelsX;
-	float TemporalJitterPixelsY;
-
-	FConvexVolume ViewFrustum;
-
-	bool bHasNearClippingPlane;
-
-	FPlane NearClippingPlane;
-
-	float NearClippingDistance;
-
-	/** true if ViewMatrix.Determinant() is negative. */
-	bool bReverseCulling;
-
-	/* Vector used by shaders to convert depth buffer samples into z coordinates in world space */
-	FVector4 InvDeviceZToWorldZTransform;
-
-	/** FOV based multiplier for cull distance on objects */
-	float LODDistanceFactor;
-	/** Square of the FOV based multiplier for cull distance on objects */
-	float LODDistanceFactorSquared;
-
-	/** Whether we did a camera cut for this view this frame. */
-	bool bCameraCut;
-
-	// -1,-1 if not setup
-	FIntPoint CursorPos;
-
-	/** True if this scene was created from a game world. */
-	bool bIsGameView;
-
-	/** Whether this view should override the view family's EngineShowFlags.Materials setting and force it on. */
-	bool bForceShowMaterials;
-
-	/** For sanity checking casts that are assumed to be safe. */
-	bool bIsViewInfo;
-
-	/** Whether this view is being used to render a scene capture. */
-	bool bIsSceneCapture;
-
-	/** Whether this view is being used to render a reflection capture. */
-	bool bIsReflectionCapture;
-	
-	/** Whether this view was created from a locked viewpoint. */
-	bool bIsLocked;
-
-	/** 
-	 * Whether to only render static lights and objects.  
-	 * This is used when capturing the scene for reflection captures, which aren't updated at runtime. 
-	 */
-	bool bStaticSceneOnly;
-
-	/** 0 if valid (we are rendering a screen postprocess pass )*/
-	struct FRenderingCompositePassContext* RenderingCompositePassContext; 
-
-	/** Aspect ratio constrained view rect. In the editor, when attached to a camera actor and the camera black bar showflag is enabled, the normal viewrect 
-	  * remains as the full viewport, and the black bars are just simulated by drawing black bars. This member stores the effective constrained area within the
-	  * bars.
-	 **/
-	FIntRect CameraConstrainedViewRect;
-
-#if WITH_EDITOR
-	/** The set of (the first 64) groups' visibility info for this view */
-	uint64 EditorViewBitflag;
-
-	/** For ortho views, this can control how to determine LOD parenting (ortho has no "distance-to-camera") */
-	FVector OverrideLODViewOrigin;
-
-	/** True if we should draw translucent objects when rendering hit proxies */
-	bool bAllowTranslucentPrimitivesInHitProxy;
-
-	/** BitArray representing the visibility state of the various sprite categories in the editor for this view */
-	TBitArray<> SpriteCategoryVisibility;
-	/** Selection color for the editor (used by post processing) */
-	FLinearColor SelectionOutlineColor;
-#endif
-
-	/**
-	 * The final settings for the current viewer position (blended together from many volumes).
-	 * Setup by the main thread, passed to the render thread and never touched again by the main thread.
-	 */
-	FFinalPostProcessSettings FinalPostProcessSettings;
-
-	/** Parameters for atmospheric fog. */
-	FTextureRHIRef AtmosphereTransmittanceTexture;
-	FTextureRHIRef AtmosphereIrradianceTexture;
-	FTextureRHIRef AtmosphereInscatterTexture;
-
-	/** Areas of the screen where the depth of field should be forced to the max. */
-	TArray< FIntRect > UIBlurOverrideRectangles;
-
-	/** Initialization constructor. */
-	FSceneView(const FSceneViewInitOptions& InitOptions);
-
-	/** used by ScreenPercentage */
-	void SetScaledViewRect(FIntRect InScaledViewRect);
-
-	/** Transforms a point from world-space to the view's screen-space. */
-	FVector4 WorldToScreen(const FVector& WorldPoint) const;
-
-	/** Transforms a point from the view's screen-space to world-space. */
-	FVector ScreenToWorld(const FVector4& ScreenPoint) const;
-
-	/** Transforms a point from the view's screen-space into pixel coordinates relative to the view's X,Y. */
-	bool ScreenToPixel(const FVector4& ScreenPoint,FVector2D& OutPixelLocation) const;
-
-	/** Transforms a point from pixel coordinates relative to the view's X,Y (left, top) into the view's screen-space. */
-	FVector4 PixelToScreen(float X,float Y,float Z) const;
-
-	/** Transforms a point from the view's world-space into pixel coordinates relative to the view's X,Y (left, top). */
-	bool WorldToPixel(const FVector& WorldPoint,FVector2D& OutPixelLocation) const;
-
-	/** Transforms a point from pixel coordinates relative to the view's X,Y (left, top) into the view's world-space. */
-	FVector4 PixelToWorld(float X,float Y,float Z) const;
-
-	/** 
-	 * Transforms a point from the view's world-space into the view's screen-space. 
-	 * Divides the resulting X, Y, Z by W before returning. 
-	 */
-	FPlane Project(const FVector& WorldPoint) const;
-
-	/** 
-	 * Transforms a point from the view's screen-space into world coordinates
-	 * multiplies X, Y, Z by W before transforming. 
-	 */
-	FVector Deproject(const FPlane& ScreenPoint) const;
-
-	/** transforms 2D screen coordinates into a 3D world-space origin and direction 
-	 * @param ScreenPos - screen coordinates in pixels
-	 * @param out_WorldOrigin (out) - world-space origin vector
-	 * @param out_WorldDirection (out) - world-space direction vector
-	 */
-	void DeprojectFVector2D(const FVector2D& ScreenPos, FVector& out_WorldOrigin, FVector& out_WorldDirection) const;
-
-	/** transforms 2D screen coordinates into a 3D world-space origin and direction 
-	 * @param ScreenPos - screen coordinates in pixels
-	 * @param ViewRect - view rectangle
-	 * @param InvViewMatrix - inverse view matrix
-	 * @param InvProjMatrix - inverse projection matrix
-	 * @param out_WorldOrigin (out) - world-space origin vector
-	 * @param out_WorldDirection (out) - world-space direction vector
-	 */
-	static void DeprojectScreenToWorld(const FVector2D& ScreenPos, const FIntRect& ViewRect, const FMatrix& InvViewMatrix, const FMatrix& InvProjMatrix, FVector& out_WorldOrigin, FVector& out_WorldDirection);
-
-	inline FVector GetViewRight() const { return ViewMatrices.ViewMatrix.GetColumn(0); }
-	inline FVector GetViewUp() const { return ViewMatrices.ViewMatrix.GetColumn(1); }
-	inline FVector GetViewDirection() const { return ViewMatrices.ViewMatrix.GetColumn(2); }
-
-	/** @return true:perspective, false:orthographic */
-	inline bool IsPerspectiveProjection() const { return ViewMatrices.IsPerspectiveProjection(); }
-
-	/** Allow things like HMD displays to update the view matrix at the last minute, to minimize perceived latency */
-	void UpdateViewMatrix();
-
-	/** Setup defaults and depending on view position (postprocess volumes) */
-	void StartFinalPostprocessSettings(FVector ViewLocation);
-
-	/**
-	 * custom layers can be combined with the existing settings
-	 * @param Weight usually 0..1 but outside range is clamped
-	 */
-	void OverridePostProcessSettings(const FPostProcessSettings& Src, float Weight);
-
-	/** applied global restrictions from show flags */
-	void EndFinalPostprocessSettings();
-
-	/** Configure post process settings for the buffer visualization system */
-	void ConfigureBufferVisualizationSettings();
-
-	/** Get the feature level for this view **/
-	ERHIFeatureLevel::Type GetFeatureLevel() const;
-};
-
-
-/**
- * A set of views into a scene which only have different view transforms and owner actors.
- */
-class FSceneViewFamily
-{
-public:
-	/**
-	* Helper struct for creating FSceneViewFamily instances
-	* If created with specifying a time it will retrieve them from the world in the given scene.
-	* 
-	* @param InRenderTarget		The render target which the views are being rendered to.
-	* @param InScene			The scene being viewed.
-	* @param InShowFlags		The show flags for the views.
-	*
-	*/
-	struct ConstructionValues
-	{
-		ConstructionValues(
-			const FRenderTarget* InRenderTarget,
-			FSceneInterface* InScene,
-			const FEngineShowFlags& InEngineShowFlags
-			)
-		:	RenderTarget(InRenderTarget)
-		,	Scene(InScene)
-		,	EngineShowFlags(InEngineShowFlags)
-		,	CurrentWorldTime(0.0f)
-		,	DeltaWorldTime(0.0f)
-		,	CurrentRealTime(0.0f)
-		,	GammaCorrection(1.0f)
-		,	bRealtimeUpdate(false)
-		,	bDeferClear(false)
-		,	bResolveScene(true)			
-		,	bTimesSet(false)
-		{
-			if( InScene != NULL )			
-			{
-				UWorld* World = InScene->GetWorld();
-				// Ensure the world is valid and that we are being called from a game thread (GetRealTimeSeconds requires this)
-				if( World && IsInGameThread() )
-				{					
-					CurrentWorldTime = World->GetTimeSeconds();
-					DeltaWorldTime = World->GetDeltaSeconds();
-					CurrentRealTime = World->GetRealTimeSeconds();
-					bTimesSet = true;
-				}
-			}
-		}
-		/** The views which make up the family. */
-		const FRenderTarget* RenderTarget;
-
-		/** The render target which the views are being rendered to. */
-		FSceneInterface* Scene;
-
-		/** The engine show flags for the views. */
-		FEngineShowFlags EngineShowFlags;
-
-		/** The current world time. */
-		float CurrentWorldTime;
-
-		/** The difference between the last world time and CurrentWorldTime. */
-		float DeltaWorldTime;
-		
-		/** The current real time. */
-		float CurrentRealTime;
-
-		/** Gamma correction used when rendering this family. Default is 1.0 */
-		float GammaCorrection;
-		
-		/** Indicates whether the view family is updated in real-time. */
-		uint32 bRealtimeUpdate:1;
-		
-		/** Used to defer the back buffer clearing to just before the back buffer is drawn to */
-		uint32 bDeferClear:1;
-		
-		/** If true then results of scene rendering are copied/resolved to the RenderTarget. */
-		uint32 bResolveScene:1;		
-		
-		/** Safety check to ensure valid times are set either from a valid world/scene pointer or via the SetWorldTimes function */
-		uint32 bTimesSet:1;
-
-		/** Set the world time ,difference between the last world time and CurrentWorldTime and current real time. */
-		ConstructionValues& SetWorldTimes(const float InCurrentWorldTime,const float InDeltaWorldTime,const float InCurrentRealTime) { CurrentWorldTime = InCurrentWorldTime; DeltaWorldTime = InDeltaWorldTime; CurrentRealTime = InCurrentRealTime;bTimesSet = true;return *this; }
-		
-		/** Set  whether the view family is updated in real-time. */
-		ConstructionValues& SetRealtimeUpdate(const bool Value) { bRealtimeUpdate = Value; return *this; }
-		
-		/** Set whether to defer the back buffer clearing to just before the back buffer is drawn to */
-		ConstructionValues& SetDeferClear(const bool Value) { bDeferClear = Value; return *this; }
-		
-		/** Setting to if true then results of scene rendering are copied/resolved to the RenderTarget. */
-		ConstructionValues& SetResolveScene(const bool Value) { bResolveScene = Value; return *this; }
-		
-		/** Set Gamma correction used when rendering this family. */
-		ConstructionValues& SetGammaCorrection(const float Value) { GammaCorrection = Value; return *this; }		
-	};
-	
-	/** The views which make up the family. */
-	TArray<const FSceneView*> Views;
-
-	/** The width in screen pixels of the view family being rendered. */
-	uint32 FamilySizeX;
-
-	/** The height in screen pixels of the view family being rendered. */
-	uint32 FamilySizeY;
-
-	/** The render target which the views are being rendered to. */
-	const FRenderTarget* RenderTarget;
-
-	/** Indicates that a separate render target is in use (not a backbuffer RT) */
-	bool bUseSeparateRenderTarget;
-
-	/** The scene being viewed. */
-	FSceneInterface* Scene;
-
-	/** The new show flags for the views (meant to replace the old system). */
-	FEngineShowFlags EngineShowFlags;
-
-	/** The current world time. */
-	float CurrentWorldTime;
-
-	/** The difference between the last world time and CurrentWorldTime. */
-	float DeltaWorldTime;
-
-	/** The current real time. */
-	float CurrentRealTime;
-
-	/** Indicates whether the view family is updated in realtime. */
-	bool bRealtimeUpdate;
-
-	/** Used to defer the back buffer clearing to just before the back buffer is drawn to */
-	bool bDeferClear;
-
-	/** if true then results of scene rendering are copied/resolved to the RenderTarget. */
-	bool bResolveScene;
-
-	/** Gamma correction used when rendering this family. Default is 1.0 */
-	float GammaCorrection;
-	
-	/** Editor setting to allow designers to override the automatic expose. 0:Automatic, following indices: -4 .. +4 */
-	FExposureSettings ExposureSettings;
-
-    /** Extensions that can modify view parameters on the render thread. */
-    TArray<class ISceneViewExtension*> ViewExtensions;
-
-#if WITH_EDITOR
-	// Override the LOD of landscape in this viewport
-	int8 LandscapeLODOverride;
-
-	/** Indicates whether, of not, the base attachment volume should be drawn. */
-	bool bDrawBaseInfo;
-#endif
-
-	/** Initialization constructor. */
-	ENGINE_API FSceneViewFamily( const ConstructionValues& CVS );
-
-	/** Computes FamilySizeX and FamilySizeY from the Views array. */
-	ENGINE_API void ComputeFamilySize();
-};
-
-/**
- * A view family which deletes its views when it goes out of scope.
- */
-class FSceneViewFamilyContext : public FSceneViewFamily
-{
-public:
-	/** Initialization constructor. */
-	FSceneViewFamilyContext( const ConstructionValues& CVS)
-		:	FSceneViewFamily(CVS)
-	{}
-
-	/** Destructor. */
-	ENGINE_API ~FSceneViewFamilyContext();
-};
 
 //
 // Primitive drawing utility functions.
diff --git a/Engine/Source/Runtime/Engine/Public/SceneTypes.h b/Engine/Source/Runtime/Engine/Public/SceneTypes.h
new file mode 100644
--- /dev/null
+++ b/Engine/Source/Runtime/Engine/Public/SceneTypes.h
@@ -0,0 +1,105 @@
+// Copyright 1998-2014 Epic Games, Inc. All Rights Reserved.
+
+#pragma once
+
+class FLightMap;
+class FShadowMap;
+
+/** 
+ * Class used to identify UPrimitiveComponents on the rendering thread without having to pass the pointer around, 
+ * Which would make it easy for people to access game thread variables from the rendering thread.
+ */
+class FPrimitiveComponentId
+{
+public:
+
+	FPrimitiveComponentId() : Value(0)
+	{}
+
+	inline bool IsValid() const
+	{
+		return Value > 0;
+	}
+
+	inline bool operator==(FPrimitiveComponentId OtherId) const
+	{
+		return Value == OtherId.Value;
+	}
+
+	friend uint32 GetTypeHash( FPrimitiveComponentId Id )
+	{
+		return GetTypeHash(Id.Value);
+	}
+
+	uint64 Value;
+};
+
+/**
+ * The types of interactions between a light and a primitive.
+ */
+enum ELightMapInteractionType
+{
+	LMIT_None	= 0,
+	LMIT_Texture= 2,
+
+	LMIT_NumBits= 3
+};
+
+/** A reference to a light-map. */
+typedef TRefCountPtr<FLightMap> FLightMapRef;
+
+enum EShadowMapInteractionType
+{
+	SMIT_None = 0,
+	SMIT_Texture = 2,
+
+	SMIT_NumBits = 3
+};
+
+/** A reference to a shadow-map. */
+typedef TRefCountPtr<FShadowMap> FShadowMapRef;
+
+/** Quality levels that a material can be compiled for. */
+namespace EMaterialQualityLevel
+{
+	enum Type
+	{
+		Low,
+		High,
+		Num
+	};
+}
+
+//
+//	EMaterialProperty
+//
+
+enum EMaterialProperty
+{
+	MP_EmissiveColor = 0,
+	MP_Opacity,
+	MP_OpacityMask,
+	MP_DiffuseColor,
+	MP_SpecularColor,
+	MP_BaseColor,
+	MP_Metallic,
+	MP_Specular,
+	MP_Roughness,
+	MP_Normal,
+	MP_WorldPositionOffset,
+	MP_WorldDisplacement,
+	MP_TessellationMultiplier,
+	MP_SubsurfaceColor,
+	MP_AmbientOcclusion,
+	MP_Refraction,
+	MP_CustomizedUVs0,
+	MP_CustomizedUVs1,
+	MP_CustomizedUVs2,
+	MP_CustomizedUVs3,
+	MP_CustomizedUVs4,
+	MP_CustomizedUVs5,
+	MP_CustomizedUVs6,
+	MP_CustomizedUVs7,
+	MP_MaterialAttributes,
+	MP_MAX,
+};
\ No newline at end of file
diff --git a/Engine/Source/Runtime/Engine/Public/SceneUtils.h b/Engine/Source/Runtime/Engine/Public/SceneUtils.h
--- a/Engine/Source/Runtime/Engine/Public/SceneUtils.h
+++ b/Engine/Source/Runtime/Engine/Public/SceneUtils.h
@@ -10,7 +10,6 @@
 // Each event type will be displayed using the defined color
 #pragma once
 
-#include "DynamicRHI.h"
 
 #define DEC_LIGHT			FColor(255,0,0,255)
 #define DEC_SKEL_MESH		FColor(255,0,255,255)
@@ -38,7 +37,7 @@ inline void appSetCounterValue(const TCHAR* CounterName, float Value) {}
 	 * Class that logs draw events based upon class scope. Draw events can be seen
 	 * in PIX
 	 */
-	struct FDrawEvent
+	struct ENGINE_API FDrawEvent
 	{
 		/** Whether a draw event has been emitted or not. */
 		bool bDrawEventHasBeenEmitted;
@@ -51,31 +50,16 @@ inline void appSetCounterValue(const TCHAR* CounterName, float Value) {}
 		/**
 		 * Terminate the event based upon scope
 		 */
-		~FDrawEvent()
-		{
-			if( bDrawEventHasBeenEmitted )
-			{
-				GDynamicRHI->PopEvent();
-			}
-		}
+		 ~FDrawEvent();
+
 		/**
 		 * Function for logging a PIX event with var args
 		 */
-		void CDECL Start(const FColor& Color,const TCHAR* Fmt, ...)
-		{
-			check( IsInRenderingThread_Internal() );
-			va_list ptr;
-			va_start(ptr, Fmt);
-			TCHAR TempStr[256];
-			// Build the string in the temp buffer
-			FCString::GetVarArgs(TempStr,ARRAY_COUNT(TempStr),ARRAY_COUNT(TempStr)-1,Fmt,ptr);
-			GDynamicRHI->PushEvent(TempStr);
-			bDrawEventHasBeenEmitted = true;		
-		}
+		void CDECL Start(const FColor& Color,const TCHAR* Fmt, ...);
 
 		private:
 		/** Helper function to determine if we are currently int the render thread. This is needed to avoid using IsInRenderThread() in a public header since it comes from RenderCore */
-		ENGINE_API bool IsInRenderingThread_Internal();
+		bool IsInRenderingThread_Internal();
 	};
 
 	// Macros to allow for scoping of draw events
diff --git a/Engine/Source/Runtime/Engine/Public/SceneView.h b/Engine/Source/Runtime/Engine/Public/SceneView.h
new file mode 100644
--- /dev/null
+++ b/Engine/Source/Runtime/Engine/Public/SceneView.h
@@ -0,0 +1,621 @@
+// Copyright 1998-2014 Epic Games, Inc. All Rights Reserved.
+
+#pragma once
+
+#include "Runtime/RenderCore/Public/UniformBuffer.h"
+#include "Engine/GameViewportClient.h"
+#include "ConvexVolume.h"
+#include "FinalPostProcessSettings.h"
+#include "SceneInterface.h"
+
+class FSceneViewStateInterface;
+class FViewUniformShaderParameters;
+class FViewElementDrawer;
+
+// Projection data for a FSceneView
+struct FSceneViewProjectionData
+{
+	FMatrix ViewMatrix;
+
+	/** UE4 projection matrix projects such that clip space Z=1 is the near plane, and Z=0 is the infinite far plane. */
+	FMatrix ProjectionMatrix;
+
+protected:
+	//The unconstrained (no aspect ratio bars applied) view rectangle (also unscaled)
+	FIntRect ViewRect;
+
+	// The constrained view rectangle (identical to UnconstrainedUnscaledViewRect if aspect ratio is not constrained)
+	FIntRect ConstrainedViewRect;
+
+public:
+	void SetViewRectangle(const FIntRect& InViewRect)
+	{
+		ViewRect = InViewRect;
+		ConstrainedViewRect = InViewRect;
+	}
+
+	void SetConstrainedViewRectangle(const FIntRect& InViewRect)
+	{
+		ConstrainedViewRect = InViewRect;
+	}
+
+	bool IsValidViewRectangle() const
+	{
+		return (ConstrainedViewRect.Min.X >= 0) &&
+			(ConstrainedViewRect.Min.Y >= 0) &&
+			(ConstrainedViewRect.Width() > 0) &&
+			(ConstrainedViewRect.Height() > 0);
+	}
+
+	const FIntRect& GetViewRect() const { return ViewRect; }
+	const FIntRect& GetConstrainedViewRect() const { return ConstrainedViewRect; }
+};
+
+// Construction parameters for a FSceneView
+struct FSceneViewInitOptions : public FSceneViewProjectionData
+{
+	const FSceneViewFamily* ViewFamily;
+	FSceneViewStateInterface* SceneViewStateInterface;
+	const AActor* ViewActor;
+	FViewElementDrawer* ViewElementDrawer;
+
+	FLinearColor BackgroundColor;
+	FLinearColor OverlayColor;
+	FLinearColor ColorScale;
+
+	/** For stereoscopic rendering, whether or not this is a full pass, or a left / right eye pass */
+	EStereoscopicPass StereoPass;
+
+	/** Conversion from world units (uu) to meters, so we can scale motion to the world appropriately */
+	float WorldToMetersScale;
+
+	TSet<FPrimitiveComponentId> HiddenPrimitives;
+
+	// -1,-1 if not setup
+	FIntPoint CursorPos;
+
+	float LODDistanceFactor;
+
+	/** If > 0, overrides the view's far clipping plane with a plane at the specified distance. */
+	float OverrideFarClippingPlaneDistance;
+
+	// Was there a camera cut this frame?
+	bool bInCameraCut;
+
+#if WITH_EDITOR
+	// default to 0'th view index, which is a bitfield of 1
+	uint64 EditorViewBitflag;
+
+	// this can be specified for ortho views so that it's min draw distance/LOD parenting etc, is controlled by a perspective viewport
+	FVector OverrideLODViewOrigin;
+
+	// In case of ortho, generate a fake view position that has a non-zero W component. The view position will be derived based on the view matrix.
+	bool bUseFauxOrthoViewPos;
+#endif
+
+	FSceneViewInitOptions()
+		: ViewFamily(NULL)
+		, SceneViewStateInterface(NULL)
+		, ViewActor(NULL)
+		, ViewElementDrawer(NULL)
+		, BackgroundColor(FLinearColor::Transparent)
+		, OverlayColor(FLinearColor::Transparent)
+		, ColorScale(FLinearColor::White)
+		, StereoPass(eSSP_FULL)
+		, WorldToMetersScale(100.f)
+		, CursorPos(-1, -1)
+		, LODDistanceFactor(1.0f)
+		, OverrideFarClippingPlaneDistance(-1.0f)
+		, bInCameraCut(false)
+#if WITH_EDITOR
+		, EditorViewBitflag(1)
+		, OverrideLODViewOrigin(ForceInitToZero)
+		, bUseFauxOrthoViewPos(false)
+		//@TODO: , const TBitArray<>& InSpriteCategoryVisibility=TBitArray<>()
+#endif
+	{
+	}
+};
+
+
+//////////////////////////////////////////////////////////////////////////
+
+struct FViewMatrices
+{
+	FViewMatrices()
+	{
+		ProjMatrix.SetIdentity();
+		ViewMatrix.SetIdentity();
+		TranslatedViewProjectionMatrix.SetIdentity();
+		InvTranslatedViewProjectionMatrix.SetIdentity();
+		PreViewTranslation = FVector::ZeroVector;
+		ViewOrigin = FVector::ZeroVector;
+	}
+
+	/** ViewToClip : UE4 projection matrix projects such that clip space Z=1 is the near plane, and Z=0 is the infinite far plane. */
+	FMatrix		ProjMatrix;
+	// WorldToView..
+	FMatrix		ViewMatrix;
+	/** The view-projection transform, starting from world-space points translated by -ViewOrigin. */
+	FMatrix		TranslatedViewProjectionMatrix;
+	/** The inverse view-projection transform, ending with world-space points translated by -ViewOrigin. */
+	FMatrix		InvTranslatedViewProjectionMatrix;
+	/** The translation to apply to the world before TranslatedViewProjectionMatrix. Usually it is -ViewOrigin but with rereflections this can differ */
+	FVector		PreViewTranslation;
+	/** To support ortho and other modes this is redundant, in world space */
+	FVector		ViewOrigin;
+
+	//
+	// World = TranslatedWorld - PreViewTranslation
+	// TranslatedWorld = World + PreViewTranslation
+	// 
+
+	// ----------------
+
+	/** @return true:perspective, false:orthographic */
+	inline bool IsPerspectiveProjection() const
+	{
+		return ProjMatrix.M[3][3] < 1.0f;
+	}
+
+	FMatrix GetViewProjMatrix() const
+	{
+		return ViewMatrix * ProjMatrix;
+	}
+
+	FMatrix GetViewRotationProjMatrix() const
+	{
+		return ViewMatrix.RemoveTranslation() * ProjMatrix;
+	}
+
+	FMatrix GetInvProjMatrix() const
+	{
+		return ProjMatrix.Inverse();
+	}
+
+	FMatrix GetInvViewMatrix() const
+	{
+		// can be optimized: it's not a perspective matrix so transpose would be enough
+		return ViewMatrix.Inverse();
+	}
+
+	FMatrix GetInvViewProjMatrix() const
+	{
+		return GetInvProjMatrix() * GetInvViewMatrix();
+	}
+};
+
+//////////////////////////////////////////////////////////////////////////
+
+/**
+ * A projection from scene space into a 2D screen region.
+ */
+class ENGINE_API FSceneView
+{
+public:
+	const FSceneViewFamily* Family;
+	/** can be 0 (thumbnail rendering) */
+	FSceneViewStateInterface* State;
+
+	/** The uniform buffer for the view's parameters.  This is only initialized in the rendering thread's copies of the FSceneView. */
+	TUniformBufferRef<FViewUniformShaderParameters> UniformBuffer;
+
+	/** The actor which is being viewed from. */
+	const AActor* ViewActor;
+
+	/** An interaction which draws the view's interaction elements. */
+	FViewElementDrawer* Drawer;
+
+	/* Final position of the view in the final render target (in pixels), potentially scaled down by ScreenPercentage */
+	FIntRect ViewRect;
+
+	/* Final position of the view in the final render target (in pixels), potentially constrained by an aspect ratio requirement (black bars) */
+	const FIntRect UnscaledViewRect;
+
+	/* Raw view size (in pixels), used for screen space calculations */
+	const FIntRect UnconstrainedViewRect;
+
+	/** 
+	 * Copy from GFrameNumber
+	 */
+	uint32 FrameNumber;
+
+	/** Maximum number of shadow cascades to render with. */
+	int32 MaxShadowCascades;
+
+    FViewMatrices ViewMatrices;
+
+    /** Variables used to determine the view matrix */
+    FVector		ViewLocation;
+    FRotator	ViewRotation;
+	FQuat		BaseHmdOrientation;
+	FVector		BaseHmdLocation;
+	float		WorldToMetersScale;
+
+	// normally the same as ViewMatrices unless "r.Shadow.FreezeCamera" is activated
+	FViewMatrices ShadowViewMatrices;
+
+	FMatrix ProjectionMatrixUnadjustedForRHI;
+
+	FLinearColor BackgroundColor;
+	FLinearColor OverlayColor;
+
+	/** Color scale multiplier used during post processing */
+	FLinearColor ColorScale;
+
+	/** For stereoscopic rendering, whether or not this is a full pass, or a left / right eye pass */
+	EStereoscopicPass StereoPass;
+
+	/** Current buffer visualization mode */
+	FName CurrentBufferVisualizationMode;
+
+	/**
+	* These can be used to override material parameters across the scene without recompiling shaders.
+	* The last component is how much to include of the material's value for that parameter, so 0 will completely remove the material's value.
+	*/
+	FVector4 DiffuseOverrideParameter;
+	FVector4 SpecularOverrideParameter;
+	FVector4 NormalOverrideParameter;
+	FVector2D RoughnessOverrideParameter;
+
+	/** The primitives which are hidden for this view. */
+	TSet<FPrimitiveComponentId> HiddenPrimitives;
+
+	// Derived members.
+
+	/** redundant, ViewMatrices.GetViewProjMatrix() */
+	/* UE4 projection matrix projects such that clip space Z=1 is the near plane, and Z=0 is the infinite far plane. */
+	FMatrix ViewProjectionMatrix;
+	/** redundant, ViewMatrices.GetInvViewMatrix() */
+	FMatrix InvViewMatrix;				
+	/** redundant, ViewMatrices.GetInvViewProjMatrix() */
+	FMatrix InvViewProjectionMatrix;	
+
+	float TemporalJitterPixelsX;
+	float TemporalJitterPixelsY;
+
+	FConvexVolume ViewFrustum;
+
+	bool bHasNearClippingPlane;
+
+	FPlane NearClippingPlane;
+
+	float NearClippingDistance;
+
+	/** true if ViewMatrix.Determinant() is negative. */
+	bool bReverseCulling;
+
+	/* Vector used by shaders to convert depth buffer samples into z coordinates in world space */
+	FVector4 InvDeviceZToWorldZTransform;
+
+	/** FOV based multiplier for cull distance on objects */
+	float LODDistanceFactor;
+	/** Square of the FOV based multiplier for cull distance on objects */
+	float LODDistanceFactorSquared;
+
+	/** Whether we did a camera cut for this view this frame. */
+	bool bCameraCut;
+
+	// -1,-1 if not setup
+	FIntPoint CursorPos;
+
+	/** True if this scene was created from a game world. */
+	bool bIsGameView;
+
+	/** Whether this view should override the view family's EngineShowFlags.Materials setting and force it on. */
+	bool bForceShowMaterials;
+
+	/** For sanity checking casts that are assumed to be safe. */
+	bool bIsViewInfo;
+
+	/** Whether this view is being used to render a scene capture. */
+	bool bIsSceneCapture;
+
+	/** Whether this view is being used to render a reflection capture. */
+	bool bIsReflectionCapture;
+	
+	/** Whether this view was created from a locked viewpoint. */
+	bool bIsLocked;
+
+	/** 
+	 * Whether to only render static lights and objects.  
+	 * This is used when capturing the scene for reflection captures, which aren't updated at runtime. 
+	 */
+	bool bStaticSceneOnly;
+
+	/** 0 if valid (we are rendering a screen postprocess pass )*/
+	struct FRenderingCompositePassContext* RenderingCompositePassContext; 
+
+	/** Aspect ratio constrained view rect. In the editor, when attached to a camera actor and the camera black bar showflag is enabled, the normal viewrect 
+	  * remains as the full viewport, and the black bars are just simulated by drawing black bars. This member stores the effective constrained area within the
+	  * bars.
+	 **/
+	FIntRect CameraConstrainedViewRect;
+
+#if WITH_EDITOR
+	/** The set of (the first 64) groups' visibility info for this view */
+	uint64 EditorViewBitflag;
+
+	/** For ortho views, this can control how to determine LOD parenting (ortho has no "distance-to-camera") */
+	FVector OverrideLODViewOrigin;
+
+	/** True if we should draw translucent objects when rendering hit proxies */
+	bool bAllowTranslucentPrimitivesInHitProxy;
+
+	/** BitArray representing the visibility state of the various sprite categories in the editor for this view */
+	TBitArray<> SpriteCategoryVisibility;
+	/** Selection color for the editor (used by post processing) */
+	FLinearColor SelectionOutlineColor;
+#endif
+
+	/**
+	 * The final settings for the current viewer position (blended together from many volumes).
+	 * Setup by the main thread, passed to the render thread and never touched again by the main thread.
+	 */
+	FFinalPostProcessSettings FinalPostProcessSettings;
+
+	/** Parameters for atmospheric fog. */
+	FTextureRHIRef AtmosphereTransmittanceTexture;
+	FTextureRHIRef AtmosphereIrradianceTexture;
+	FTextureRHIRef AtmosphereInscatterTexture;
+
+	/** Areas of the screen where the depth of field should be forced to the max. */
+	TArray< FIntRect > UIBlurOverrideRectangles;
+
+	/** Initialization constructor. */
+	FSceneView(const FSceneViewInitOptions& InitOptions);
+
+	/** used by ScreenPercentage */
+	void SetScaledViewRect(FIntRect InScaledViewRect);
+
+	/** Transforms a point from world-space to the view's screen-space. */
+	FVector4 WorldToScreen(const FVector& WorldPoint) const;
+
+	/** Transforms a point from the view's screen-space to world-space. */
+	FVector ScreenToWorld(const FVector4& ScreenPoint) const;
+
+	/** Transforms a point from the view's screen-space into pixel coordinates relative to the view's X,Y. */
+	bool ScreenToPixel(const FVector4& ScreenPoint,FVector2D& OutPixelLocation) const;
+
+	/** Transforms a point from pixel coordinates relative to the view's X,Y (left, top) into the view's screen-space. */
+	FVector4 PixelToScreen(float X,float Y,float Z) const;
+
+	/** Transforms a point from the view's world-space into pixel coordinates relative to the view's X,Y (left, top). */
+	bool WorldToPixel(const FVector& WorldPoint,FVector2D& OutPixelLocation) const;
+
+	/** Transforms a point from pixel coordinates relative to the view's X,Y (left, top) into the view's world-space. */
+	FVector4 PixelToWorld(float X,float Y,float Z) const;
+
+	/** 
+	 * Transforms a point from the view's world-space into the view's screen-space. 
+	 * Divides the resulting X, Y, Z by W before returning. 
+	 */
+	FPlane Project(const FVector& WorldPoint) const;
+
+	/** 
+	 * Transforms a point from the view's screen-space into world coordinates
+	 * multiplies X, Y, Z by W before transforming. 
+	 */
+	FVector Deproject(const FPlane& ScreenPoint) const;
+
+	/** transforms 2D screen coordinates into a 3D world-space origin and direction 
+	 * @param ScreenPos - screen coordinates in pixels
+	 * @param out_WorldOrigin (out) - world-space origin vector
+	 * @param out_WorldDirection (out) - world-space direction vector
+	 */
+	void DeprojectFVector2D(const FVector2D& ScreenPos, FVector& out_WorldOrigin, FVector& out_WorldDirection) const;
+
+	/** transforms 2D screen coordinates into a 3D world-space origin and direction 
+	 * @param ScreenPos - screen coordinates in pixels
+	 * @param ViewRect - view rectangle
+	 * @param InvViewMatrix - inverse view matrix
+	 * @param InvProjMatrix - inverse projection matrix
+	 * @param out_WorldOrigin (out) - world-space origin vector
+	 * @param out_WorldDirection (out) - world-space direction vector
+	 */
+	static void DeprojectScreenToWorld(const FVector2D& ScreenPos, const FIntRect& ViewRect, const FMatrix& InvViewMatrix, const FMatrix& InvProjMatrix, FVector& out_WorldOrigin, FVector& out_WorldDirection);
+
+	inline FVector GetViewRight() const { return ViewMatrices.ViewMatrix.GetColumn(0); }
+	inline FVector GetViewUp() const { return ViewMatrices.ViewMatrix.GetColumn(1); }
+	inline FVector GetViewDirection() const { return ViewMatrices.ViewMatrix.GetColumn(2); }
+
+	/** @return true:perspective, false:orthographic */
+	inline bool IsPerspectiveProjection() const { return ViewMatrices.IsPerspectiveProjection(); }
+
+	/** Allow things like HMD displays to update the view matrix at the last minute, to minimize perceived latency */
+	void UpdateViewMatrix();
+
+	/** Setup defaults and depending on view position (postprocess volumes) */
+	void StartFinalPostprocessSettings(FVector ViewLocation);
+
+	/**
+	 * custom layers can be combined with the existing settings
+	 * @param Weight usually 0..1 but outside range is clamped
+	 */
+	void OverridePostProcessSettings(const FPostProcessSettings& Src, float Weight);
+
+	/** applied global restrictions from show flags */
+	void EndFinalPostprocessSettings();
+
+	/** Configure post process settings for the buffer visualization system */
+	void ConfigureBufferVisualizationSettings();
+
+	/** Get the feature level for this view **/
+	ERHIFeatureLevel::Type GetFeatureLevel() const;
+};
+
+//////////////////////////////////////////////////////////////////////////
+
+/**
+ * A set of views into a scene which only have different view transforms and owner actors.
+ */
+class FSceneViewFamily
+{
+public:
+	/**
+	* Helper struct for creating FSceneViewFamily instances
+	* If created with specifying a time it will retrieve them from the world in the given scene.
+	* 
+	* @param InRenderTarget		The render target which the views are being rendered to.
+	* @param InScene			The scene being viewed.
+	* @param InShowFlags		The show flags for the views.
+	*
+	*/
+	struct ConstructionValues
+	{
+		ConstructionValues(
+			const FRenderTarget* InRenderTarget,
+			FSceneInterface* InScene,
+			const FEngineShowFlags& InEngineShowFlags
+			)
+		:	RenderTarget(InRenderTarget)
+		,	Scene(InScene)
+		,	EngineShowFlags(InEngineShowFlags)
+		,	CurrentWorldTime(0.0f)
+		,	DeltaWorldTime(0.0f)
+		,	CurrentRealTime(0.0f)
+		,	GammaCorrection(1.0f)
+		,	bRealtimeUpdate(false)
+		,	bDeferClear(false)
+		,	bResolveScene(true)			
+		,	bTimesSet(false)
+		{
+			if( InScene != NULL )			
+			{
+				UWorld* World = InScene->GetWorld();
+				// Ensure the world is valid and that we are being called from a game thread (GetRealTimeSeconds requires this)
+				if( World && IsInGameThread() )
+				{					
+					CurrentWorldTime = World->GetTimeSeconds();
+					DeltaWorldTime = World->GetDeltaSeconds();
+					CurrentRealTime = World->GetRealTimeSeconds();
+					bTimesSet = true;
+				}
+			}
+		}
+		/** The views which make up the family. */
+		const FRenderTarget* RenderTarget;
+
+		/** The render target which the views are being rendered to. */
+		FSceneInterface* Scene;
+
+		/** The engine show flags for the views. */
+		FEngineShowFlags EngineShowFlags;
+
+		/** The current world time. */
+		float CurrentWorldTime;
+
+		/** The difference between the last world time and CurrentWorldTime. */
+		float DeltaWorldTime;
+		
+		/** The current real time. */
+		float CurrentRealTime;
+
+		/** Gamma correction used when rendering this family. Default is 1.0 */
+		float GammaCorrection;
+		
+		/** Indicates whether the view family is updated in real-time. */
+		uint32 bRealtimeUpdate:1;
+		
+		/** Used to defer the back buffer clearing to just before the back buffer is drawn to */
+		uint32 bDeferClear:1;
+		
+		/** If true then results of scene rendering are copied/resolved to the RenderTarget. */
+		uint32 bResolveScene:1;		
+		
+		/** Safety check to ensure valid times are set either from a valid world/scene pointer or via the SetWorldTimes function */
+		uint32 bTimesSet:1;
+
+		/** Set the world time ,difference between the last world time and CurrentWorldTime and current real time. */
+		ConstructionValues& SetWorldTimes(const float InCurrentWorldTime,const float InDeltaWorldTime,const float InCurrentRealTime) { CurrentWorldTime = InCurrentWorldTime; DeltaWorldTime = InDeltaWorldTime; CurrentRealTime = InCurrentRealTime;bTimesSet = true;return *this; }
+		
+		/** Set  whether the view family is updated in real-time. */
+		ConstructionValues& SetRealtimeUpdate(const bool Value) { bRealtimeUpdate = Value; return *this; }
+		
+		/** Set whether to defer the back buffer clearing to just before the back buffer is drawn to */
+		ConstructionValues& SetDeferClear(const bool Value) { bDeferClear = Value; return *this; }
+		
+		/** Setting to if true then results of scene rendering are copied/resolved to the RenderTarget. */
+		ConstructionValues& SetResolveScene(const bool Value) { bResolveScene = Value; return *this; }
+		
+		/** Set Gamma correction used when rendering this family. */
+		ConstructionValues& SetGammaCorrection(const float Value) { GammaCorrection = Value; return *this; }		
+	};
+	
+	/** The views which make up the family. */
+	TArray<const FSceneView*> Views;
+
+	/** The width in screen pixels of the view family being rendered. */
+	uint32 FamilySizeX;
+
+	/** The height in screen pixels of the view family being rendered. */
+	uint32 FamilySizeY;
+
+	/** The render target which the views are being rendered to. */
+	const FRenderTarget* RenderTarget;
+
+	/** Indicates that a separate render target is in use (not a backbuffer RT) */
+	bool bUseSeparateRenderTarget;
+
+	/** The scene being viewed. */
+	FSceneInterface* Scene;
+
+	/** The new show flags for the views (meant to replace the old system). */
+	FEngineShowFlags EngineShowFlags;
+
+	/** The current world time. */
+	float CurrentWorldTime;
+
+	/** The difference between the last world time and CurrentWorldTime. */
+	float DeltaWorldTime;
+
+	/** The current real time. */
+	float CurrentRealTime;
+
+	/** Indicates whether the view family is updated in realtime. */
+	bool bRealtimeUpdate;
+
+	/** Used to defer the back buffer clearing to just before the back buffer is drawn to */
+	bool bDeferClear;
+
+	/** if true then results of scene rendering are copied/resolved to the RenderTarget. */
+	bool bResolveScene;
+
+	/** Gamma correction used when rendering this family. Default is 1.0 */
+	float GammaCorrection;
+	
+	/** Editor setting to allow designers to override the automatic expose. 0:Automatic, following indices: -4 .. +4 */
+	FExposureSettings ExposureSettings;
+
+    /** Extensions that can modify view parameters on the render thread. */
+    TArray<class ISceneViewExtension*> ViewExtensions;
+
+#if WITH_EDITOR
+	// Override the LOD of landscape in this viewport
+	int8 LandscapeLODOverride;
+
+	/** Indicates whether, of not, the base attachment volume should be drawn. */
+	bool bDrawBaseInfo;
+#endif
+
+	/** Initialization constructor. */
+	ENGINE_API FSceneViewFamily( const ConstructionValues& CVS );
+
+	/** Computes FamilySizeX and FamilySizeY from the Views array. */
+	ENGINE_API void ComputeFamilySize();
+};
+
+/**
+ * A view family which deletes its views when it goes out of scope.
+ */
+class FSceneViewFamilyContext : public FSceneViewFamily
+{
+public:
+	/** Initialization constructor. */
+	FSceneViewFamilyContext( const ConstructionValues& CVS)
+		:	FSceneViewFamily(CVS)
+	{}
+
+	/** Destructor. */
+	ENGINE_API ~FSceneViewFamilyContext();
+};
diff --git a/Engine/Source/Runtime/Engine/Public/ShadowMap.h b/Engine/Source/Runtime/Engine/Public/ShadowMap.h
--- a/Engine/Source/Runtime/Engine/Public/ShadowMap.h
+++ b/Engine/Source/Runtime/Engine/Public/ShadowMap.h
@@ -342,8 +342,5 @@ private:
 };
 
 
-/** A reference to a shadow-map. */
-typedef TRefCountPtr<FShadowMap> FShadowMapRef;
-
 /** Shadowmap reference serializer */
 extern FArchive& operator<<(FArchive& Ar,FShadowMap*& R);
diff --git a/Engine/Source/Runtime/Engine/Public/StaticParameterSet.h b/Engine/Source/Runtime/Engine/Public/StaticParameterSet.h
new file mode 100644
--- /dev/null
+++ b/Engine/Source/Runtime/Engine/Public/StaticParameterSet.h
@@ -0,0 +1,173 @@
+// Copyright 1998-2014 Epic Games, Inc. All Rights Reserved.
+
+#pragma once
+
+/**
+* Holds the information for a static switch parameter
+*/
+class FStaticSwitchParameter
+{
+public:
+	FName ParameterName;
+	bool Value;
+	bool bOverride;
+	FGuid ExpressionGUID;
+
+	FStaticSwitchParameter() :
+		ParameterName(TEXT("None")),
+		Value(false),
+		bOverride(false),
+		ExpressionGUID(0, 0, 0, 0)
+	{ }
+
+	FStaticSwitchParameter(FName InName, bool InValue, bool InOverride, FGuid InGuid) :
+		ParameterName(InName),
+		Value(InValue),
+		bOverride(InOverride),
+		ExpressionGUID(InGuid)
+	{ }
+
+	friend FArchive& operator<<(FArchive& Ar, FStaticSwitchParameter& P)
+	{
+		Ar << P.ParameterName << P.Value << P.bOverride;
+		Ar << P.ExpressionGUID;
+		return Ar;
+	}
+};
+
+/**
+* Holds the information for a static component mask parameter
+*/
+class FStaticComponentMaskParameter
+{
+public:
+	FName ParameterName;
+	bool R, G, B, A;
+	bool bOverride;
+	FGuid ExpressionGUID;
+
+	FStaticComponentMaskParameter() :
+		ParameterName(TEXT("None")),
+		R(false),
+		G(false),
+		B(false),
+		A(false),
+		bOverride(false),
+		ExpressionGUID(0, 0, 0, 0)
+	{ }
+
+	FStaticComponentMaskParameter(FName InName, bool InR, bool InG, bool InB, bool InA, bool InOverride, FGuid InGuid) :
+		ParameterName(InName),
+		R(InR),
+		G(InG),
+		B(InB),
+		A(InA),
+		bOverride(InOverride),
+		ExpressionGUID(InGuid)
+	{ }
+
+	friend FArchive& operator<<(FArchive& Ar, FStaticComponentMaskParameter& P)
+	{
+		Ar << P.ParameterName << P.R << P.G << P.B << P.A << P.bOverride;
+		Ar << P.ExpressionGUID;
+		return Ar;
+	}
+};
+
+/**
+* Holds the information for a static switch parameter
+*/
+class FStaticTerrainLayerWeightParameter
+{
+public:
+	FName ParameterName;
+	bool bOverride;
+	FGuid ExpressionGUID;
+
+	int32 WeightmapIndex;
+
+	FStaticTerrainLayerWeightParameter() :
+		ParameterName(TEXT("None")),
+		bOverride(false),
+		ExpressionGUID(0, 0, 0, 0),
+		WeightmapIndex(INDEX_NONE)
+	{ }
+
+	FStaticTerrainLayerWeightParameter(FName InName, int32 InWeightmapIndex, bool InOverride, FGuid InGuid) :
+		ParameterName(InName),
+		bOverride(InOverride),
+		ExpressionGUID(InGuid),
+		WeightmapIndex(InWeightmapIndex)
+	{ }
+
+	friend FArchive& operator<<(FArchive& Ar, FStaticTerrainLayerWeightParameter& P)
+	{
+		Ar << P.ParameterName << P.WeightmapIndex << P.bOverride;
+		Ar << P.ExpressionGUID;
+		return Ar;
+	}
+};
+
+/** Contains all the information needed to identify a single permutation of static parameters. */
+class FStaticParameterSet
+{
+public:
+	/** An array of static switch parameters in this set */
+	TArray<FStaticSwitchParameter> StaticSwitchParameters;
+
+	/** An array of static component mask parameters in this set */
+	TArray<FStaticComponentMaskParameter> StaticComponentMaskParameters;
+
+	/** An array of terrain layer weight parameters in this set */
+	TArray<FStaticTerrainLayerWeightParameter> TerrainLayerWeightParameters;
+
+	FStaticParameterSet() {}
+
+	/** 
+	* Checks if this set contains any parameters
+	* 
+	* @return	true if this set has no parameters
+	*/
+	bool IsEmpty() const
+	{
+		return StaticSwitchParameters.Num() == 0 && StaticComponentMaskParameters.Num() == 0 && TerrainLayerWeightParameters.Num() == 0;
+	}
+
+	void Serialize(FArchive& Ar)
+	{
+		// Note: FStaticParameterSet is saved both in packages (UMaterialInstance) and the DDC (FMaterialShaderMap)
+		// Backwards compatibility only works with FStaticParameterSet's stored in packages.  
+		// You must bump MATERIALSHADERMAP_DERIVEDDATA_VER as well if changing the serialization of FStaticParameterSet.
+
+		Ar << StaticSwitchParameters << StaticComponentMaskParameters;
+		Ar << TerrainLayerWeightParameters;
+	}
+
+	/** Updates the hash state with the static parameter names and values. */
+	void UpdateHash(FSHA1& HashState) const;
+
+	/** 
+	* Indicates whether this set is equal to another, copying override settings.
+	* 
+	* @param ReferenceSet	The set to compare against
+	* @return				true if the sets are not equal
+	*/
+	bool ShouldMarkDirty(const FStaticParameterSet* ReferenceSet);
+
+	/** 
+	* Tests this set against another for equality, disregarding override settings.
+	* 
+	* @param ReferenceSet	The set to compare against
+	* @return				true if the sets are equal
+	*/
+	bool operator==(const FStaticParameterSet& ReferenceSet) const;
+
+	bool operator!=(const FStaticParameterSet& ReferenceSet) const
+	{
+		return !(*this == ReferenceSet);
+	}
+
+	FString GetSummaryString() const;
+
+	void AppendKeyString(FString& KeyString) const;
+};
\ No newline at end of file
diff --git a/Engine/Source/Runtime/Engine/Public/WorldCollision.h b/Engine/Source/Runtime/Engine/Public/WorldCollision.h
--- a/Engine/Source/Runtime/Engine/Public/WorldCollision.h
+++ b/Engine/Source/Runtime/Engine/Public/WorldCollision.h
@@ -7,6 +7,8 @@
 #pragma once 
 
 #include "UniquePtr.h"
+#include "CollisionQueryParams.h"
+#include "Engine/World.h"
 
 /**
  * Async Trace Data Structures 
diff --git a/Engine/Source/Runtime/RenderCore/Public/RenderCommandFence.h b/Engine/Source/Runtime/RenderCore/Public/RenderCommandFence.h
new file mode 100644
--- /dev/null
+++ b/Engine/Source/Runtime/RenderCore/Public/RenderCommandFence.h
@@ -0,0 +1,37 @@
+// Copyright 1998-2014 Epic Games, Inc. All Rights Reserved.
+
+#pragma once
+
+#include "Runtime/Core/Public/Async/TaskGraphInterfaces.h"
+
+////////////////////////////////////
+// Render fences
+////////////////////////////////////
+
+ /**
+ * Used to track pending rendering commands from the game thread.
+ */
+class RENDERCORE_API FRenderCommandFence
+{
+public:
+
+	/**
+	 * Adds a fence command to the rendering command queue.
+	 * Conceptually, the pending fence count is incremented to reflect the pending fence command.
+	 * Once the rendering thread has executed the fence command, it decrements the pending fence count.
+	 */
+	void BeginFence();
+
+	/**
+	 * Waits for pending fence commands to retire.
+	 * @param bProcessGameThreadTasks, if true we are on a short callstack where it is safe to process arbitrary game thread tasks while we wait
+	 */
+	void Wait(bool bProcessGameThreadTasks = false) const;
+
+	// return true if the fence is complete
+	bool IsFenceComplete() const;
+
+private:
+	/** Graph event that represents completion of this fence **/
+	mutable FGraphEventRef CompletionEvent;
+};
\ No newline at end of file
diff --git a/Engine/Source/Runtime/RenderCore/Public/RenderingThread.h b/Engine/Source/Runtime/RenderCore/Public/RenderingThread.h
--- a/Engine/Source/Runtime/RenderCore/Public/RenderingThread.h
+++ b/Engine/Source/Runtime/RenderCore/Public/RenderingThread.h
@@ -9,6 +9,7 @@
 
 #include "RenderCore.h"
 #include "TaskGraphInterfaces.h"
+#include "RenderCommandFence.h"
 
 ////////////////////////////////////
 // Render thread API
@@ -388,37 +389,7 @@ DECLARE_STATS_GROUP(TEXT("Render Thread Commands"), STATGROUP_RenderThreadComman
 
 #endif //UE_SERVER
 
-////////////////////////////////////
-// Render fences
-////////////////////////////////////
-
- /**
- * Used to track pending rendering commands from the game thread.
- */
-class RENDERCORE_API FRenderCommandFence
-{
-public:
 
-	/**
-	 * Adds a fence command to the rendering command queue.
-	 * Conceptually, the pending fence count is incremented to reflect the pending fence command.
-	 * Once the rendering thread has executed the fence command, it decrements the pending fence count.
-	 */
-	void BeginFence();
-
-	/**
-	 * Waits for pending fence commands to retire.
-	 * @param bProcessGameThreadTasks, if true we are on a short callstack where it is safe to process arbitrary game thread tasks while we wait
-	 */
-	void Wait(bool bProcessGameThreadTasks = false) const;
-
-	// return true if the fence is complete
-	bool IsFenceComplete() const;
-
-private:
-	/** Graph event that represents completion of this fence **/
-	mutable FGraphEventRef CompletionEvent;
-};
 
 ////////////////////////////////////
 // Deferred cleanup
