Fix GetAllocatedSize() for FNavLinkRenderingProxy and FLandscapeSplinesSceneProxy

diff --git a/Engine/Source/Runtime/Engine/Private/AI/Navigation/NavLinkRenderingComponent.cpp b/Engine/Source/Runtime/Engine/Private/AI/Navigation/NavLinkRenderingComponent.cpp
--- a/Engine/Source/Runtime/Engine/Private/AI/Navigation/NavLinkRenderingComponent.cpp
+++ b/Engine/Source/Runtime/Engine/Private/AI/Navigation/NavLinkRenderingComponent.cpp
@@ -225,5 +225,5 @@ uint32 FNavLinkRenderingProxy::GetMemoryFootprint( void ) const
 
 uint32 FNavLinkRenderingProxy::GetAllocatedSize( void ) const 
 { 
-	return( FPrimitiveSceneProxy::GetAllocatedSize() + OffMeshPointLinks.Num() + OffMeshSegmentLinks.Num() ); 
+	return(FPrimitiveSceneProxy::GetAllocatedSize() + OffMeshPointLinks.GetAllocatedSize() + OffMeshSegmentLinks.GetAllocatedSize());
 }
diff --git a/Engine/Source/Runtime/Engine/Private/Landscape/LandscapeSplines.cpp b/Engine/Source/Runtime/Engine/Private/Landscape/LandscapeSplines.cpp
--- a/Engine/Source/Runtime/Engine/Private/Landscape/LandscapeSplines.cpp
+++ b/Engine/Source/Runtime/Engine/Private/Landscape/LandscapeSplines.cpp
@@ -264,17 +264,21 @@ public:
 		return Result;
 	}
 
-	virtual uint32 GetMemoryFootprint( void ) const
+	virtual uint32 GetMemoryFootprint() const
 	{
-		return sizeof( *this ) + GetAllocatedSize();
+		return sizeof(*this) + GetAllocatedSize();
 	}
-	uint32 GetAllocatedSize( void ) const
+	uint32 GetAllocatedSize() const
 	{
-		uint32 AllocatedSize = FPrimitiveSceneProxy::GetAllocatedSize() + Segments.GetAllocatedSize();
+		uint32 AllocatedSize = FPrimitiveSceneProxy::GetAllocatedSize() + Segments.GetAllocatedSize() + ControlPoints.GetAllocatedSize();
 		for (const FSegmentProxy& Segment : Segments)
 		{
 			AllocatedSize += Segment.Points.GetAllocatedSize();
 		}
+		for (const FControlPointProxy& ControlPoint : ControlPoints)
+		{
+			AllocatedSize += ControlPoint.Points.GetAllocatedSize();
+		}
 		return AllocatedSize;
 	}
 };
