Put FLight/ShadowMapRef back into Light/ShadowMap.h, as complete definition is required for TRefCountPtr
Move FSceneViewStateReference and ELightComponentType into SceneTypes.h
Add some other includes/forward declares
Remove ParticleHelper.h from AnimNotifyState_Trail.h

diff --git a/Engine/Source/Runtime/Engine/Classes/Animation/AnimNotifies/AnimNotifyState_Trail.h b/Engine/Source/Runtime/Engine/Classes/Animation/AnimNotifies/AnimNotifyState_Trail.h
--- a/Engine/Source/Runtime/Engine/Classes/Animation/AnimNotifies/AnimNotifyState_Trail.h
+++ b/Engine/Source/Runtime/Engine/Classes/Animation/AnimNotifies/AnimNotifyState_Trail.h
@@ -1,7 +1,6 @@
 // Copyright 1998-2014 Epic Games, Inc. All Rights Reserved.
 
 #pragma once
-#include "ParticleHelper.h"
 #include "AnimNotifyState_Trail.generated.h"
 
 DECLARE_LOG_CATEGORY_EXTERN(LogAnimTrails, Log, All);
diff --git a/Engine/Source/Runtime/Engine/Classes/Components/DecalComponent.h b/Engine/Source/Runtime/Engine/Classes/Components/DecalComponent.h
--- a/Engine/Source/Runtime/Engine/Classes/Components/DecalComponent.h
+++ b/Engine/Source/Runtime/Engine/Classes/Components/DecalComponent.h
@@ -4,7 +4,7 @@
 #pragma once
 #include "DecalComponent.generated.h"
 
-	// -> will be exported to EngineDecalClasses.h
+class FDeferredDecalProxy;
 
 /** 
  * A material that is rendered onto the surface of a mesh. A kind of 'bumper sticker' for a model.
diff --git a/Engine/Source/Runtime/Engine/Classes/Components/LightComponent.h b/Engine/Source/Runtime/Engine/Classes/Components/LightComponent.h
--- a/Engine/Source/Runtime/Engine/Classes/Components/LightComponent.h
+++ b/Engine/Source/Runtime/Engine/Classes/Components/LightComponent.h
@@ -2,6 +2,7 @@
 
 
 #pragma once
+#include "SceneTypes.h"
 #include "LightComponent.generated.h"
 
 UCLASS(abstract, HideCategories=(Trigger,Activation,"Components|Activation",Physics), ShowCategories=(Mobility))
diff --git a/Engine/Source/Runtime/Engine/Classes/Components/ReflectionCaptureComponent.h b/Engine/Source/Runtime/Engine/Classes/Components/ReflectionCaptureComponent.h
--- a/Engine/Source/Runtime/Engine/Classes/Components/ReflectionCaptureComponent.h
+++ b/Engine/Source/Runtime/Engine/Classes/Components/ReflectionCaptureComponent.h
@@ -3,6 +3,8 @@
 #pragma once
 #include "ReflectionCaptureComponent.generated.h"
 
+class FReflectionCaptureProxy;
+
 class FReflectionCaptureFullHDRDerivedData
 {
 public:
diff --git a/Engine/Source/Runtime/Engine/Classes/Components/StaticMeshComponent.h b/Engine/Source/Runtime/Engine/Classes/Components/StaticMeshComponent.h
--- a/Engine/Source/Runtime/Engine/Classes/Components/StaticMeshComponent.h
+++ b/Engine/Source/Runtime/Engine/Classes/Components/StaticMeshComponent.h
@@ -2,7 +2,8 @@
 
 
 #pragma once
-#include "SceneTypes.h"
+#include "LightMap.h"
+#include "ShadowMap.h"
 #include "Components/MeshComponent.h"
 #include "Runtime/RenderCore/Public/RenderUtils.h"
 #include "StaticMeshComponent.generated.h"
diff --git a/Engine/Source/Runtime/Engine/Classes/Engine/Canvas.h b/Engine/Source/Runtime/Engine/Classes/Engine/Canvas.h
--- a/Engine/Source/Runtime/Engine/Classes/Engine/Canvas.h
+++ b/Engine/Source/Runtime/Engine/Classes/Engine/Canvas.h
@@ -9,6 +9,7 @@
 #include "Canvas.generated.h"
 
 class FSceneView;
+class FCanvas;
 
 /**
  * Holds texture information with UV coordinates as well.
diff --git a/Engine/Source/Runtime/Engine/Classes/Engine/Level.h b/Engine/Source/Runtime/Engine/Classes/Engine/Level.h
--- a/Engine/Source/Runtime/Engine/Classes/Engine/Level.h
+++ b/Engine/Source/Runtime/Engine/Classes/Engine/Level.h
@@ -7,9 +7,12 @@
 #pragma once
 #include "PhysicsEngine/BodyInstance.h"
 #include "Engine/LevelBase.h"
+#include "Engine/World.h"
+#include "Model.h"
 #include "Level.generated.h"
 
 class ALevelBounds;
+class UTexture2D;
 
 /**
  * Structure containing all information needed for determining the screen space
diff --git a/Engine/Source/Runtime/Engine/Classes/PhysicsEngine/BodySetup.h b/Engine/Source/Runtime/Engine/Classes/PhysicsEngine/BodySetup.h
--- a/Engine/Source/Runtime/Engine/Classes/PhysicsEngine/BodySetup.h
+++ b/Engine/Source/Runtime/Engine/Classes/PhysicsEngine/BodySetup.h
@@ -5,6 +5,8 @@
 #include "BodySetup.generated.h"
 
 struct FDynamicMeshVertex;
+class FMaterialRenderProxy;
+class FPrimitiveDrawInterface;
 
 namespace physx
 {
diff --git a/Engine/Source/Runtime/Engine/Public/EngineUtils.h b/Engine/Source/Runtime/Engine/Public/EngineUtils.h
--- a/Engine/Source/Runtime/Engine/Public/EngineUtils.h
+++ b/Engine/Source/Runtime/Engine/Public/EngineUtils.h
@@ -9,6 +9,7 @@
 
 #include "Engine/Level.h"
 #include "GameFramework/WorldSettings.h"
+#include "HitProxies.h"
 
 /*-----------------------------------------------------------------------------
 	Hit proxies.
@@ -569,7 +570,7 @@ private:
  *	@param ViewLocation	Location of camera
  *	@param ViewRotation	Rotation of camera
  */
-ENGINE_API void DrawStatsHUD( UWorld* InWorld, FViewport* Viewport, FCanvas* Canvas, UCanvas* CanvasObject, TArray<FDebugDisplayProperty>& DebugProperties, const FVector& ViewLocation, const FRotator& ViewRotation );
+ENGINE_API void DrawStatsHUD( UWorld* InWorld, FViewport* Viewport, FCanvas* Canvas, UCanvas* CanvasObject, TArray<struct FDebugDisplayProperty>& DebugProperties, const FVector& ViewLocation, const FRotator& ViewRotation );
 
 /** SubLevel status information */
 struct FSubLevelStatus
diff --git a/Engine/Source/Runtime/Engine/Public/LightMap.h b/Engine/Source/Runtime/Engine/Public/LightMap.h
--- a/Engine/Source/Runtime/Engine/Public/LightMap.h
+++ b/Engine/Source/Runtime/Engine/Public/LightMap.h
@@ -118,7 +118,8 @@ private:
 	virtual void FinishCleanup();
 };
 
-
+/** A reference to a light-map. */
+typedef TRefCountPtr<FLightMap> FLightMapRef;
 
 /** Lightmap reference serializer */
 extern FArchive& operator<<(FArchive& Ar,FLightMap*& R);
diff --git a/Engine/Source/Runtime/Engine/Public/ParticleHelper.h b/Engine/Source/Runtime/Engine/Public/ParticleHelper.h
--- a/Engine/Source/Runtime/Engine/Public/ParticleHelper.h
+++ b/Engine/Source/Runtime/Engine/Public/ParticleHelper.h
@@ -9,6 +9,7 @@
 #include "ParticleVertexFactory.h"
 #include "ParticleBeamTrailVertexFactory.h"
 #include "MeshParticleVertexFactory.h"
+#include "MaterialShared.h"
 
 #include "ParticleHelper.generated.h"
 
diff --git a/Engine/Source/Runtime/Engine/Public/SceneManagement.h b/Engine/Source/Runtime/Engine/Public/SceneManagement.h
--- a/Engine/Source/Runtime/Engine/Public/SceneManagement.h
+++ b/Engine/Source/Runtime/Engine/Public/SceneManagement.h
@@ -130,42 +130,7 @@ private:
 	int32							NumChildren;
 };
 
-/** 
- * Class used to reference an FSceneViewStateInterface that allows destruction and recreation of all FSceneViewStateInterface's when needed. 
- * This is used to support reloading the renderer module on the fly.
- */
-class FSceneViewStateReference
-{
-public:
-	FSceneViewStateReference() :
-		Reference(NULL)
-	{}
-
-	ENGINE_API virtual ~FSceneViewStateReference();
-
-	/** Allocates the Scene view state. */
-	ENGINE_API void Allocate();
-
-	/** Destorys the Scene view state. */
-	ENGINE_API void Destroy();
-
-	/** Destroys all view states, but does not remove them from the linked list. */
-	ENGINE_API static void DestroyAll();
 
-	/** Recreates all view states in the global list. */
-	ENGINE_API static void AllocateAll();
-
-	FSceneViewStateInterface* GetReference()
-	{
-		return Reference;
-	}
-
-private:
-	FSceneViewStateInterface* Reference;
-	TLinkedList<FSceneViewStateReference*> GlobalListLink;
-
-	static TLinkedList<FSceneViewStateReference*>*& GetSceneViewStateList();
-};
 
 /**
  * The types of interactions between a light and a primitive.
@@ -613,15 +578,6 @@ public:
 	FSHVectorRGB3 IrradianceEnvironmentMap;
 };
 
-/** different light component types */
-enum ELightComponentType
-{
-	LightType_Directional = 0,
-	LightType_Point,
-	LightType_Spot,
-	LightType_MAX,
-	LightType_NumBits = 2
-};
 
 /** Encapsulates the data which is used to render a light parallel to the game thread. */
 class ENGINE_API FLightSceneProxy
diff --git a/Engine/Source/Runtime/Engine/Public/SceneTypes.h b/Engine/Source/Runtime/Engine/Public/SceneTypes.h
--- a/Engine/Source/Runtime/Engine/Public/SceneTypes.h
+++ b/Engine/Source/Runtime/Engine/Public/SceneTypes.h
@@ -4,6 +4,7 @@
 
 class FLightMap;
 class FShadowMap;
+class FSceneViewStateInterface;
 
 /** 
  * Class used to identify UPrimitiveComponents on the rendering thread without having to pass the pointer around, 
@@ -34,6 +35,53 @@ public:
 	uint64 Value;
 };
 
+/** 
+ * Class used to reference an FSceneViewStateInterface that allows destruction and recreation of all FSceneViewStateInterface's when needed. 
+ * This is used to support reloading the renderer module on the fly.
+ */
+class FSceneViewStateReference
+{
+public:
+	FSceneViewStateReference() :
+		Reference(NULL)
+	{}
+
+	ENGINE_API virtual ~FSceneViewStateReference();
+
+	/** Allocates the Scene view state. */
+	ENGINE_API void Allocate();
+
+	/** Destorys the Scene view state. */
+	ENGINE_API void Destroy();
+
+	/** Destroys all view states, but does not remove them from the linked list. */
+	ENGINE_API static void DestroyAll();
+
+	/** Recreates all view states in the global list. */
+	ENGINE_API static void AllocateAll();
+
+	FSceneViewStateInterface* GetReference()
+	{
+		return Reference;
+	}
+
+private:
+	FSceneViewStateInterface* Reference;
+	TLinkedList<FSceneViewStateReference*> GlobalListLink;
+
+	static TLinkedList<FSceneViewStateReference*>*& GetSceneViewStateList();
+};
+
+/** different light component types */
+enum ELightComponentType
+{
+	LightType_Directional = 0,
+	LightType_Point,
+	LightType_Spot,
+	LightType_MAX,
+	LightType_NumBits = 2
+};
+
 /**
  * The types of interactions between a light and a primitive.
  */
@@ -45,9 +93,6 @@ enum ELightMapInteractionType
 	LMIT_NumBits= 3
 };
 
-/** A reference to a light-map. */
-typedef TRefCountPtr<FLightMap> FLightMapRef;
-
 enum EShadowMapInteractionType
 {
 	SMIT_None = 0,
@@ -56,9 +101,6 @@ enum EShadowMapInteractionType
 	SMIT_NumBits = 3
 };
 
-/** A reference to a shadow-map. */
-typedef TRefCountPtr<FShadowMap> FShadowMapRef;
-
 /** Quality levels that a material can be compiled for. */
 namespace EMaterialQualityLevel
 {
diff --git a/Engine/Source/Runtime/Engine/Public/SceneView.h b/Engine/Source/Runtime/Engine/Public/SceneView.h
--- a/Engine/Source/Runtime/Engine/Public/SceneView.h
+++ b/Engine/Source/Runtime/Engine/Public/SceneView.h
@@ -4,6 +4,7 @@
 
 #include "Runtime/RenderCore/Public/UniformBuffer.h"
 #include "Engine/GameViewportClient.h"
+#include "Engine/World.h"
 #include "ConvexVolume.h"
 #include "FinalPostProcessSettings.h"
 #include "SceneInterface.h"
@@ -11,6 +12,7 @@
 class FSceneViewStateInterface;
 class FViewUniformShaderParameters;
 class FViewElementDrawer;
+class FSceneViewFamily;
 
 // Projection data for a FSceneView
 struct FSceneViewProjectionData
diff --git a/Engine/Source/Runtime/Engine/Public/ShadowMap.h b/Engine/Source/Runtime/Engine/Public/ShadowMap.h
--- a/Engine/Source/Runtime/Engine/Public/ShadowMap.h
+++ b/Engine/Source/Runtime/Engine/Public/ShadowMap.h
@@ -342,5 +342,8 @@ private:
 };
 
 
+/** A reference to a shadow-map. */
+typedef TRefCountPtr<FShadowMap> FShadowMapRef;
+
 /** Shadowmap reference serializer */
 extern FArchive& operator<<(FArchive& Ar,FShadowMap*& R);
