Move ETrailWidthMode into EngineTypes.h

diff --git a/Engine/Source/Runtime/Engine/Classes/Engine/EngineTypes.h b/Engine/Source/Runtime/Engine/Classes/Engine/EngineTypes.h
--- a/Engine/Source/Runtime/Engine/Classes/Engine/EngineTypes.h
+++ b/Engine/Source/Runtime/Engine/Classes/Engine/EngineTypes.h
@@ -99,6 +99,15 @@ enum ETranslucencyLightingMode
 	TLM_MAX,
 };
 
+/** Controls the way that the width scale property affects anim trails */
+UENUM()
+enum ETrailWidthMode
+{
+	ETrailWidthMode_FromCentre UMETA(DisplayName = "From Centre"),
+	ETrailWidthMode_FromFirst UMETA(DisplayName = "From First Socket"),
+	ETrailWidthMode_FromSecond UMETA(DisplayName = "From Second Socket"),
+};
+
 UENUM()
 enum EMaterialShadingModel
 {
diff --git a/Engine/Source/Runtime/Engine/Public/ParticleHelper.h b/Engine/Source/Runtime/Engine/Public/ParticleHelper.h
--- a/Engine/Source/Runtime/Engine/Public/ParticleHelper.h
+++ b/Engine/Source/Runtime/Engine/Public/ParticleHelper.h
@@ -11,8 +11,6 @@
 #include "MeshParticleVertexFactory.h"
 #include "MaterialShared.h"
 
-#include "ParticleHelper.generated.h"
-
 #define _ENABLE_PARTICLE_LOD_INGAME_
 
 DECLARE_LOG_CATEGORY_EXTERN(LogParticles, Log, All);
@@ -2870,11 +2868,3 @@ public:
 /** The global null color vertex buffer, which is set with a stride of 0 on meshes without a color component. */
 extern TGlobalResource<FNullDynamicParameterVertexBuffer> GNullDynamicParameterVertexBuffer;
 
-/** Controls the way that the width scale property affects anim trails */
-UENUM()
-enum ETrailWidthMode
-{
-	ETrailWidthMode_FromCentre UMETA(DisplayName = "From Centre"),
-	ETrailWidthMode_FromFirst UMETA(DisplayName = "From First Socket"),
-	ETrailWidthMode_FromSecond UMETA(DisplayName = "From Second Socket"),
-};
\ No newline at end of file
