Eliminate Actor.bWantsInitialize

diff --git a/Engine/Config/BaseEngine.ini b/Engine/Config/BaseEngine.ini
--- a/Engine/Config/BaseEngine.ini
+++ b/Engine/Config/BaseEngine.ini
@@ -579,7 +579,6 @@ bUseStreamingPause=false
 
 ; Property name redirects
 +TaggedPropertyRedirects=(ClassName="Actor",OldPropertyName="LifeSpan",NewPropertyName="InitialLifeSpan")
-+TaggedPropertyRedirects=(ClassName="Actor",OldPropertyName="bWantsBeginPlay",NewPropertyName="bWantsInitialize")
 +TaggedPropertyRedirects=(ClassName="ActorComponent",OldPropertyName="bWantsBeginPlayComponent",NewPropertyName="bWantsInitializeComponent")
 +TaggedPropertyRedirects=(ClassName="AudioComponent",OldPropertyName="SoundCue",NewPropertyName="Sound")
 +TaggedPropertyRedirects=(ClassName="AudioComponent",OldPropertyName="PitchMultiplierMin",NewPropertyName="PitchModulationMin")
diff --git a/Engine/Source/Runtime/Engine/Classes/GameFramework/Actor.h b/Engine/Source/Runtime/Engine/Classes/GameFramework/Actor.h
--- a/Engine/Source/Runtime/Engine/Classes/GameFramework/Actor.h
+++ b/Engine/Source/Runtime/Engine/Classes/GameFramework/Actor.h
@@ -138,9 +138,9 @@ public:
 * The functions of interest to initialization order for an Actor is roughly as follows:
 * PostLoad/PostActorCreated - Do any setup of the actor required for construction. PostLoad for serialized actors, PostActorCreated for spawned.  
 * AActor::OnConstruction - The construction of the actor, this is where Blueprint actors have their components created and blueprint variables are initialized
-* AActor::PreInitializeComponents - Called before InitializeComponent is called on the actor's components (if bWantsInitialize is true)
+* AActor::PreInitializeComponents - Called before InitializeComponent is called on the actor's components
 * UActorComponent::InitializeComponent - Each component in the actor's components array gets an initialize call (if bWantsInitializeComponent is true for that component)
-* AActor::PostInitializeComponents - Called after the actor's components have been initialized (if bWantsInitialize is true)
+* AActor::PostInitializeComponents - Called after the actor's components have been initialized
 * AActor::BeginPlay - Called when the level is started
 */
 UCLASS(abstract, dependson=(UEngineBaseTypes, UActorComponent, UEngineTypes), BlueprintType, Blueprintable, config=Engine)
@@ -177,10 +177,6 @@ public:
 	UPROPERTY(EditAnywhere, Category=Rendering, BlueprintReadOnly, replicated, meta=(DisplayName = "Actor Hidden In Game"))
 	uint32 bHidden:1;
 
-	/** Whether to route PreInitializeComponents/PostInitializeComponents to this Actor. */
-	UPROPERTY()
-	uint32 bWantsInitialize:1;
-
 	/** Used for replication of our RootComponent's position and velocity */
 	UPROPERTY(ReplicatedUsing=OnRep_ReplicatedMovement)
 	struct FRepMovement ReplicatedMovement;
diff --git a/Engine/Source/Runtime/Engine/Private/Actor.cpp b/Engine/Source/Runtime/Engine/Private/Actor.cpp
--- a/Engine/Source/Runtime/Engine/Private/Actor.cpp
+++ b/Engine/Source/Runtime/Engine/Private/Actor.cpp
@@ -60,7 +60,6 @@ AActor::AActor(const class FPostConstructInitializeProperties& PCIP)
 	bActorLabelEditable = true;
 	SpriteScale = 1.0f;
 #endif // WITH_EDITORONLY_DATA
-	bWantsInitialize = true;
 	NetCullDistanceSquared = 225000000.0f;
 	NetDriverName = NAME_GameNetDriver;
 	NetDormancy = DORM_Awake;
diff --git a/Engine/Source/Runtime/Engine/Private/BlockingVolume.cpp b/Engine/Source/Runtime/Engine/Private/BlockingVolume.cpp
--- a/Engine/Source/Runtime/Engine/Private/BlockingVolume.cpp
+++ b/Engine/Source/Runtime/Engine/Private/BlockingVolume.cpp
@@ -10,8 +10,6 @@ ABlockingVolume::ABlockingVolume(const class FPostConstructInitializeProperties&
 {
 	BrushComponent->BodyInstance.bEnableCollision_DEPRECATED = true;
 	BrushComponent->SetCollisionProfileName(InvisibleWall_NAME);
-
-	bWantsInitialize = false;
 }
 
 bool ABlockingVolume::UpdateNavigationRelevancy() 
diff --git a/Engine/Source/Runtime/Engine/Private/Brush.cpp b/Engine/Source/Runtime/Engine/Private/Brush.cpp
--- a/Engine/Source/Runtime/Engine/Private/Brush.cpp
+++ b/Engine/Source/Runtime/Engine/Private/Brush.cpp
@@ -29,7 +29,6 @@ ABrush::ABrush(const class FPostConstructInitializeProperties& PCIP)
 	
 	bHidden = true;
 	bNotForClientOrServer = false;
-	bWantsInitialize = false;
 	bCanBeDamaged = false;
 }
 
diff --git a/Engine/Source/Runtime/Engine/Private/Level.cpp b/Engine/Source/Runtime/Engine/Private/Level.cpp
--- a/Engine/Source/Runtime/Engine/Private/Level.cpp
+++ b/Engine/Source/Runtime/Engine/Private/Level.cpp
@@ -1600,7 +1600,7 @@ void ULevel::RouteActorInitialize()
 		AActor* const Actor = Actors[ActorIndex];
 		if( Actor )
 		{
-			if( !Actor->bActorInitialized && Actor->bWantsInitialize )
+			if( !Actor->bActorInitialized )
 			{
 				Actor->PreInitializeComponents();
 			}
@@ -1621,17 +1621,10 @@ void ULevel::RouteActorInitialize()
 				// Call Initialize on Components.
 				Actor->InitializeComponents();
 
-				if(Actor->bWantsInitialize)
-				{
-					Actor->PostInitializeComponents(); // should set Actor->bActorInitialized = true
-					if (!Actor->bActorInitialized && !Actor->IsPendingKill())
-					{
-						UE_LOG(LogActor, Fatal, TEXT("%s failed to route PostInitializeComponents.  Please call Super::PostInitializeComponents() in your <className>::PostInitializeComponents() function. "), *Actor->GetFullName() );
-					}
-				}
-				else
+				Actor->PostInitializeComponents(); // should set Actor->bActorInitialized = true
+				if (!Actor->bActorInitialized && !Actor->IsPendingKill())
 				{
-					Actor->bActorInitialized = true;
+					UE_LOG(LogActor, Fatal, TEXT("%s failed to route PostInitializeComponents.  Please call Super::PostInitializeComponents() in your <className>::PostInitializeComponents() function. "), *Actor->GetFullName() );
 				}
 
 				if (bCallBeginPlay)
diff --git a/Engine/Source/Runtime/Engine/Private/LevelBounds.cpp b/Engine/Source/Runtime/Engine/Private/LevelBounds.cpp
--- a/Engine/Source/Runtime/Engine/Private/LevelBounds.cpp
+++ b/Engine/Source/Runtime/Engine/Private/LevelBounds.cpp
@@ -22,7 +22,6 @@ ALevelBounds::ALevelBounds(const class FPostConstructInitializeProperties& PCIP)
 	BoxComponent->SetCollisionProfileName(UCollisionProfile::NoCollision_ProfileName);
 	BoxComponent->InitBoxExtent(FVector(0.5f, 0.5f, 0.5f));
 	
-	bWantsInitialize = false;
 	bCanBeDamaged = false;
 	
 #if WITH_EDITOR
diff --git a/Engine/Source/Runtime/Engine/Private/Light.cpp b/Engine/Source/Runtime/Engine/Private/Light.cpp
--- a/Engine/Source/Runtime/Engine/Private/Light.cpp
+++ b/Engine/Source/Runtime/Engine/Private/Light.cpp
@@ -11,7 +11,6 @@ ALight::ALight(const class FPostConstructInitializeProperties& PCIP)
 	LightComponent = PCIP.CreateAbstractDefaultSubobject<ULightComponent>(this, TEXT("LightComponent0"));
 
 	bHidden = true;
-	bWantsInitialize = false;
 	bCollideWhenPlacing = true;
 }
 
diff --git a/Engine/Source/Runtime/Engine/Private/NavMeshBoundsVolume.cpp b/Engine/Source/Runtime/Engine/Private/NavMeshBoundsVolume.cpp
--- a/Engine/Source/Runtime/Engine/Private/NavMeshBoundsVolume.cpp
+++ b/Engine/Source/Runtime/Engine/Private/NavMeshBoundsVolume.cpp
@@ -14,8 +14,6 @@ ANavMeshBoundsVolume::ANavMeshBoundsVolume(const class FPostConstructInitializeP
 	BrushColor = FColor(200, 200, 200, 255);
 
 	bColored = true;
-
-	bWantsInitialize = true;
 }
 
 #if WITH_EDITOR
diff --git a/Engine/Source/Runtime/Engine/Private/Note.cpp b/Engine/Source/Runtime/Engine/Private/Note.cpp
--- a/Engine/Source/Runtime/Engine/Private/Note.cpp
+++ b/Engine/Source/Runtime/Engine/Private/Note.cpp
@@ -55,6 +55,5 @@ ANote::ANote(const class FPostConstructInitializeProperties& PCIP)
 #endif // WITH_EDITORONLY_DATA
 
 	bHidden = true;
-	bWantsInitialize = false;
 	bCanBeDamaged = false;
 }
diff --git a/Engine/Source/Runtime/Engine/Private/PainCausingVolume.cpp b/Engine/Source/Runtime/Engine/Private/PainCausingVolume.cpp
--- a/Engine/Source/Runtime/Engine/Private/PainCausingVolume.cpp
+++ b/Engine/Source/Runtime/Engine/Private/PainCausingVolume.cpp
@@ -12,7 +12,6 @@ APainCausingVolume::APainCausingVolume(const class FPostConstructInitializePrope
 	DamagePerSec = 1.0f;
 	bEntryPain = true;
 	PainInterval = 1.0f;
-	bWantsInitialize = true;
 }
 
 void APainCausingVolume::PostInitializeComponents()
diff --git a/Engine/Source/Runtime/Engine/Private/ReflectionCaptureComponent.cpp b/Engine/Source/Runtime/Engine/Private/ReflectionCaptureComponent.cpp
--- a/Engine/Source/Runtime/Engine/Private/ReflectionCaptureComponent.cpp
+++ b/Engine/Source/Runtime/Engine/Private/ReflectionCaptureComponent.cpp
@@ -66,8 +66,6 @@ AReflectionCapture::AReflectionCapture(const class FPostConstructInitializePrope
 		SpriteComponent->bIsScreenSizeScaled = true;
 	}
 #endif // WITH_EDITORONLY_DATA
-	
-	bWantsInitialize = false;
 }
 
 #if WITH_EDITOR
diff --git a/Engine/Source/Runtime/Engine/Private/ReverbVolume.cpp b/Engine/Source/Runtime/Engine/Private/ReverbVolume.cpp
--- a/Engine/Source/Runtime/Engine/Private/ReverbVolume.cpp
+++ b/Engine/Source/Runtime/Engine/Private/ReverbVolume.cpp
@@ -129,7 +129,6 @@ AReverbVolume::AReverbVolume(const class FPostConstructInitializeProperties& PCI
 	bColored = true;
 	BrushColor = FColor(255, 255, 0, 255);
 
-	bWantsInitialize = false;
 	bEnabled = true;
 }
 
diff --git a/Engine/Source/Runtime/Engine/Private/SceneCaptureComponent.cpp b/Engine/Source/Runtime/Engine/Private/SceneCaptureComponent.cpp
--- a/Engine/Source/Runtime/Engine/Private/SceneCaptureComponent.cpp
+++ b/Engine/Source/Runtime/Engine/Private/SceneCaptureComponent.cpp
@@ -19,8 +19,6 @@ ASceneCapture::ASceneCapture(const class FPostConstructInitializeProperties& PCI
 	MeshComp->CastShadow = false;
 	MeshComp->PostPhysicsComponentTick.bCanEverTick = false;
 	RootComponent = MeshComp;
-	
-	bWantsInitialize = false;
 }
 // -----------------------------------------------
 
diff --git a/Engine/Source/Runtime/Engine/Private/StaticMeshActor.cpp b/Engine/Source/Runtime/Engine/Private/StaticMeshActor.cpp
--- a/Engine/Source/Runtime/Engine/Private/StaticMeshActor.cpp
+++ b/Engine/Source/Runtime/Engine/Private/StaticMeshActor.cpp
@@ -15,7 +15,6 @@
 AStaticMeshActor::AStaticMeshActor(const class FPostConstructInitializeProperties& PCIP)
 	: Super(PCIP)
 {
-	bWantsInitialize = false;
 	bCanBeDamaged = false;
 
 	StaticMeshComponent = PCIP.CreateDefaultSubobject<UStaticMeshComponent>(this, TEXT("StaticMeshComponent0"));
