Fix compile issue with forward declared enums.

diff --git a/Engine/Source/Runtime/Engine/Classes/Animation/AnimNotifies/AnimNotifyState_Trail.h b/Engine/Source/Runtime/Engine/Classes/Animation/AnimNotifies/AnimNotifyState_Trail.h
--- a/Engine/Source/Runtime/Engine/Classes/Animation/AnimNotifies/AnimNotifyState_Trail.h
+++ b/Engine/Source/Runtime/Engine/Classes/Animation/AnimNotifies/AnimNotifyState_Trail.h
@@ -1,6 +1,7 @@
 // Copyright 1998-2014 Epic Games, Inc. All Rights Reserved.
 
 #pragma once
+#include "ParticleHelper.h"
 #include "AnimNotifyState_Trail.generated.h"
 
 DECLARE_LOG_CATEGORY_EXTERN(LogAnimTrails, Log, All);
