Settings: Moved UConsole settings into UConsoleSettings class

- also fixed auto complete settings not actually editable
- should also fix console history currently not working

diff --git a/Engine/Config/BaseInput.ini b/Engine/Config/BaseInput.ini
--- a/Engine/Config/BaseInput.ini
+++ b/Engine/Config/BaseInput.ini
@@ -44,8 +44,10 @@ bRequireCtrlToNavigateAutoComplete=False
 +DebugExecBindings=(Key=Quote,Command="cheat EnableGDT",Control=False,Shift=False)
 
 [/Script/Engine.Console]
-MaxScrollbackSize=1024
 HistoryBot=-1
+
+[/Script/EngineSettings.ConsoleSettings]
+MaxScrollbackSize=1024
 +AutoCompleteMapPaths=Content/Maps
 +ManualAutoCompleteList=(Command="Exit",Desc="Exit (Exits the game)")
 +ManualAutoCompleteList=(Command="DebugCreatePlayer 1",Desc=)
@@ -163,4 +165,3 @@ HistoryBot=-1
 +ManualAutoCompleteList=(Command="Log ",Desc="Log <category> <level> (Change verbosity level for a log category)")
 +ManualAutoCompleteList=(Command="Log list",Desc="Log list <search string> (Search for log categories)")
 +ManualAutoCompleteList=(Command="Log reset",Desc="Log reset (Undo any changes to log verbosity)")
-
diff --git a/Engine/Source/Editor/ProjectSettingsViewer/Private/ProjectSettingsViewerModule.cpp b/Engine/Source/Editor/ProjectSettingsViewer/Private/ProjectSettingsViewerModule.cpp
--- a/Engine/Source/Editor/ProjectSettingsViewer/Private/ProjectSettingsViewerModule.cpp
+++ b/Engine/Source/Editor/ProjectSettingsViewer/Private/ProjectSettingsViewerModule.cpp
@@ -90,7 +90,7 @@ protected:
 			GetMutableDefault<UEngine>()
 		);
 
-		// Audio settings
+		// audio settings
 		SettingsModule.RegisterSettings("Project", "Engine", "Audio",
 			LOCTEXT("EngineAudioSettingsName", "Audio"),
 			LOCTEXT("ProjectAudioSettingsDescription", "Audio settings."),
@@ -104,14 +104,14 @@ protected:
 			GetMutableDefault<UCollisionProfile>()
 		);
 
-		// game console settings
+		// command console settings
 		SettingsModule.RegisterSettings("Project", "Engine", "Console",
 			LOCTEXT("ProjectConsoleSettingsName", "Console"),
 			LOCTEXT("ProjectConsoleSettingsDescription", "Configure the in-game input console."),
-			GetMutableDefault<UConsole>()
+			GetMutableDefault<UConsoleSettings>()
 		);
 
-		// Input settings
+		// input settings
 		SettingsModule.RegisterSettings("Project", "Engine", "Input",
 			LOCTEXT("EngineInputSettingsName", "Input"),
 			LOCTEXT("ProjectInputSettingsDescription", "Input settings, including default input action and axis bindings."),
diff --git a/Engine/Source/Runtime/Engine/Classes/Engine/Console.h b/Engine/Source/Runtime/Engine/Classes/Engine/Console.h
--- a/Engine/Source/Runtime/Engine/Classes/Engine/Console.h
+++ b/Engine/Source/Runtime/Engine/Classes/Engine/Console.h
@@ -8,23 +8,6 @@
 #include "Console.generated.h"
 
 
-/**
- * Structure for auto-complete commands and their descriptions.
- */
-USTRUCT()
-struct FAutoCompleteCommand
-{
-	GENERATED_USTRUCT_BODY()
-
-	UPROPERTY()
-	FString Command;
-
-	UPROPERTY()
-	FString Desc;
-
-};
-
-
 /**
  * Node for storing an auto-complete tree based on each char in the command.
  */
@@ -66,7 +49,7 @@ struct FAutoCompleteNode
 };
 
 
-UCLASS(Within=GameViewportClient, config=Input, defaultconfig, transient)
+UCLASS(Within=GameViewportClient, config=Input, transient)
 class ENGINE_API UConsole
 	: public UObject
 	, public FOutputDevice
@@ -83,10 +66,6 @@ class ENGINE_API UConsole
 	UPROPERTY()
 	class UTexture2D* DefaultTexture_White;
 
-	/**  Visible Console stuff */
-	UPROPERTY(globalconfig, EditAnywhere, Category=General)
-	int32 MaxScrollbackSize;
-
 	/** Holds the scrollback buffer */
 	TArray<FString> Scrollback;
 
@@ -128,14 +107,6 @@ class ENGINE_API UConsole
 	/** True while a control key is pressed. */
 	uint32 bCtrl:1;
 
-	/** Manual list of auto-complete commands and info specified in BaseInput.ini */
-	UPROPERTY(config, EditAnywhere, Category=AutoComplete)
-	TArray<struct FAutoCompleteCommand> ManualAutoCompleteList;
-
-	/** List of relative paths (e.g. Content/Maps) to search for map names for auto-complete usage. Specified in BaseInput.ini. */
-	UPROPERTY(config, EditAnywhere, Category=AutoComplete)
-	TArray<FString> AutoCompleteMapPaths;
-
 	/** Full list of auto-complete commands and info */
 	TArray<struct FAutoCompleteCommand> AutoCompleteList;
 
diff --git a/Engine/Source/Runtime/Engine/Private/UserInterface/Console.cpp b/Engine/Source/Runtime/Engine/Private/UserInterface/Console.cpp
--- a/Engine/Source/Runtime/Engine/Private/UserInterface/Console.cpp
+++ b/Engine/Source/Runtime/Engine/Private/UserInterface/Console.cpp
@@ -94,6 +94,7 @@ void UConsole::BuildRuntimeAutoCompleteList(bool bForce)
 		bIsRuntimeAutoCompleteUpToDate = false;
 		return;
 	}
+
 	// clear the existing tree
 	//@todo - probably only need to rebuild the tree + partial command list on level load
 	for (int32 Idx = 0; Idx < AutoCompleteTree.ChildNodes.Num(); Idx++)
@@ -101,13 +102,17 @@ void UConsole::BuildRuntimeAutoCompleteList(bool bForce)
 		FAutoCompleteNode *Node = AutoCompleteTree.ChildNodes[Idx];
 		delete Node;
 	}
+
 	AutoCompleteTree.ChildNodes.Empty();
+
+	const UConsoleSettings* ConsoleSettings = GetDefault<UConsoleSettings>();
+
 	// copy the manual list first
 	AutoCompleteList.Empty();
-	AutoCompleteList.AddZeroed(ManualAutoCompleteList.Num());
-	for (int32 Idx = 0; Idx < ManualAutoCompleteList.Num(); Idx++)
+	AutoCompleteList.AddZeroed(ConsoleSettings->ManualAutoCompleteList.Num());
+	for (int32 Idx = 0; Idx < ConsoleSettings->ManualAutoCompleteList.Num(); Idx++)
 	{
-		AutoCompleteList[Idx] = ManualAutoCompleteList[Idx];
+		AutoCompleteList[Idx] = ConsoleSettings->ManualAutoCompleteList[Idx];
 	}
 
 	// console variables
@@ -162,9 +167,9 @@ void UConsole::BuildRuntimeAutoCompleteList(bool bForce)
 
 	// enumerate maps
 	TArray<FString> Packages;
-	for (int32 PathIdx=0; PathIdx<AutoCompleteMapPaths.Num(); ++PathIdx)
+	for (int32 PathIdx = 0; PathIdx < ConsoleSettings->AutoCompleteMapPaths.Num(); ++PathIdx)
 	{
-		FPackageName::FindPackagesInDirectory(Packages, FString::Printf(TEXT("%s%s"), *FPaths::GameDir(), *AutoCompleteMapPaths[PathIdx]));
+		FPackageName::FindPackagesInDirectory(Packages, FString::Printf(TEXT("%s%s"), *FPaths::GameDir(), *ConsoleSettings->AutoCompleteMapPaths[PathIdx]));
 	}
 	
 	// also include maps in this user's developer dir
@@ -394,14 +399,18 @@ void UConsole::ClearOutput()
 
 void UConsole::OutputTextLine(const FString& Text)
 {
+	const UConsoleSettings* ConsoleSettings = GetDefault<UConsoleSettings>();
+
 	// If we are full, delete the first line
-	if (Scrollback.Num() > MaxScrollbackSize)
+	if (Scrollback.Num() > ConsoleSettings->MaxScrollbackSize)
 	{
-		Scrollback.RemoveAt(0,1);
-		SBHead = MaxScrollbackSize-1;
+		Scrollback.RemoveAt(0, 1);
+		SBHead = ConsoleSettings->MaxScrollbackSize - 1;
 	}
 	else
+	{
 		SBHead++;
+	}
 
 	// Add the line
 	Scrollback.Add(Text);
@@ -1154,7 +1163,7 @@ void UConsole::PostRender_InputLine(UCanvas* Canvas, FIntPoint UserInputLinePos)
 
 			if(OutStr.IsEmpty())
 			{
-				// no Description means we display the Command directly, without that the line would be empty (happens for ConsoleVariables and some ManualAutoCompleteList)
+				// no Description means we display the Command directly, without that the line would be empty (happens for ConsoleVariables and some ConsoleSettings->ManualAutoCompleteList)
 				OutStr = Cmd.Command;
 			}
 
diff --git a/Engine/Source/Runtime/EngineSettings/Classes/ConsoleSettings.h b/Engine/Source/Runtime/EngineSettings/Classes/ConsoleSettings.h
new file mode 100644
--- /dev/null
+++ b/Engine/Source/Runtime/EngineSettings/Classes/ConsoleSettings.h
@@ -0,0 +1,48 @@
+// Copyright 1998-2014 Epic Games, Inc. All Rights Reserved.
+
+/*=============================================================================
+	ConsoleSettings.h: Declares the UConsoleSettings class.
+=============================================================================*/
+
+#pragma once
+
+#include "ConsoleSettings.generated.h"
+
+
+/**
+ * Structure for auto-complete commands and their descriptions.
+ */
+USTRUCT()
+struct FAutoCompleteCommand
+{
+	GENERATED_USTRUCT_BODY()
+
+	UPROPERTY(config, EditAnywhere, Category=Command)
+	FString Command;
+
+	UPROPERTY(config, EditAnywhere, Category=Command)
+	FString Desc;
+};
+
+
+/**
+ * Implements the settings for the UConsole class.
+ */
+UCLASS(config=Input, defaultconfig)
+class ENGINESETTINGS_API UConsoleSettings
+	: public UObject
+{
+	GENERATED_UCLASS_BODY()
+
+	/**  Visible Console stuff */
+	UPROPERTY(globalconfig, EditAnywhere, Category=General)
+	int32 MaxScrollbackSize;
+
+	/** Manual list of auto-complete commands and info specified in BaseInput.ini */
+	UPROPERTY(config, EditAnywhere, Category=AutoComplete)
+	TArray<struct FAutoCompleteCommand> ManualAutoCompleteList;
+
+	/** List of relative paths (e.g. Content/Maps) to search for map names for auto-complete usage. Specified in BaseInput.ini. */
+	UPROPERTY(config, EditAnywhere, Category=AutoComplete)
+	TArray<FString> AutoCompleteMapPaths;
+};
diff --git a/Engine/Source/Runtime/EngineSettings/Private/EngineSettingsModule.cpp b/Engine/Source/Runtime/EngineSettings/Private/EngineSettingsModule.cpp
--- a/Engine/Source/Runtime/EngineSettings/Private/EngineSettingsModule.cpp
+++ b/Engine/Source/Runtime/EngineSettings/Private/EngineSettingsModule.cpp
@@ -33,6 +33,11 @@ public:
 /* Class constructors
  *****************************************************************************/
 
+UConsoleSettings::UConsoleSettings( const class FPostConstructInitializeProperties& PCIP )
+	: Super(PCIP)
+{ }
+
+
 UGameMapsSettings::UGameMapsSettings( const class FPostConstructInitializeProperties& PCIP )
 	: Super(PCIP)
 	, bUseSplitscreen(true)
diff --git a/Engine/Source/Runtime/EngineSettings/Public/EngineSettings.h b/Engine/Source/Runtime/EngineSettings/Public/EngineSettings.h
--- a/Engine/Source/Runtime/EngineSettings/Public/EngineSettings.h
+++ b/Engine/Source/Runtime/EngineSettings/Public/EngineSettings.h
@@ -16,6 +16,7 @@
 /* Messages
  *****************************************************************************/
 
+#include "ConsoleSettings.h"
 #include "GameNetworkManagerSettings.h"
 #include "GameMapsSettings.h"
 #include "GameSessionSettings.h"
