Added a check for name of AIModule to load if it's not None #UE4

diff --git a/Engine/Source/Runtime/Engine/Private/World.cpp b/Engine/Source/Runtime/Engine/Private/World.cpp
--- a/Engine/Source/Runtime/Engine/Private/World.cpp
+++ b/Engine/Source/Runtime/Engine/Private/World.cpp
@@ -712,13 +712,17 @@ UAISystemBase* UWorld::CreateAISystem()
 	// create navigation system for editor and server targets, but remove it from game clients
 	if (AISystem == NULL && GetNetMode() != NM_Client)
 	{
-		auto AISystemModule = FModuleManager::LoadModulePtr<IAISystemModule>(UAISystemBase::GetAISystemModuleName());
-		if (AISystemModule)
+		FName AIModuleName = UAISystemBase::GetAISystemModuleName();
+		if (AIModuleName.IsNone() == false)
 		{
-			AISystem = AISystemModule->CreateAISystemInstance(this);
-			if (AISystem == NULL)
+			auto AISystemModule = FModuleManager::LoadModulePtr<IAISystemModule>(UAISystemBase::GetAISystemModuleName());
+			if (AISystemModule)
 			{
-				UE_LOG(LogWorld, Error, TEXT("Failed to create AISystem instance of class %s!"), *UAISystemBase::GetAISystemClassName().ClassName);
+				AISystem = AISystemModule->CreateAISystemInstance(this);
+				if (AISystem == NULL)
+				{
+					UE_LOG(LogWorld, Error, TEXT("Failed to create AISystem instance of class %s!"), *UAISystemBase::GetAISystemClassName().ClassName);
+				}
 			}
 		}
 	}
