Fixed UInterpTrackInst::GetWorld() asserting for interp tracks in CameraAnimInsts.

diff --git a/Engine/Source/Runtime/Engine/Classes/Matinee/InterpTrackInst.h b/Engine/Source/Runtime/Engine/Classes/Matinee/InterpTrackInst.h
--- a/Engine/Source/Runtime/Engine/Classes/Matinee/InterpTrackInst.h
+++ b/Engine/Source/Runtime/Engine/Classes/Matinee/InterpTrackInst.h
@@ -21,7 +21,7 @@ class UInterpTrackInst : public UObject
 	 *	Return the  AActor  associated with this instance of a Group. 
 	 *	Note that all Groups have at least 1 instance, even if no  AActor  variable is attached, so this may return NULL. 
 	 */
-	ENGINE_API AActor* GetGroupActor();
+	ENGINE_API AActor* GetGroupActor() const;
 
 	/** 
 	 * Save any variables from the actor that will be modified by this instance.
diff --git a/Engine/Source/Runtime/Engine/Private/Interpolation.cpp b/Engine/Source/Runtime/Engine/Private/Interpolation.cpp
--- a/Engine/Source/Runtime/Engine/Private/Interpolation.cpp
+++ b/Engine/Source/Runtime/Engine/Private/Interpolation.cpp
@@ -3010,7 +3010,7 @@ UInterpTrackInst::UInterpTrackInst(const class FPostConstructInitializePropertie
 }
 
 
-AActor* UInterpTrackInst::GetGroupActor()
+AActor* UInterpTrackInst::GetGroupActor() const
 {
 	UInterpGroupInst* GrInst = CastChecked<UInterpGroupInst>( GetOuter() );
 	return GrInst->GetGroupActor();
@@ -3019,9 +3019,15 @@ AActor* UInterpTrackInst::GetGroupActor()
 
 UWorld* UInterpTrackInst::GetWorld() const
 {
-	UInterpGroupInst* GrInst = CastChecked<UInterpGroupInst>( GetOuter() );
-	AMatineeActor* MatineeActor = CastChecked<AMatineeActor>( GrInst->GetOuter() );
-	return MatineeActor->GetWorld();
+	// find an actor
+	AActor const* Actor = GetGroupActor();
+	if (Actor == nullptr)
+	{
+		// search the outer chain for an actor
+		Actor = GetTypedOuter<AActor>();
+	}
+
+	return Actor ? Actor->GetWorld() : nullptr;
 }
 
 /*-----------------------------------------------------------------------------
