TTP 334625:  Mesh Emitters initial orientation
and TTP 334626, make initial orientation and rotation rates take effect when PSA_Velocity is active.
Changed orientation on the mesh emitter type data to be a distribution vector, and a bunch of stuff around it to be able to apply a random value per-particle. Always apply initial orientation and rates in addition to velocity orientation (needed to add a separate initial orientation to the rotation payload for this).
Bumps object version and puts the original flat initial orientation vector into the distribution if needed.

diff --git a/Engine/Source/Runtime/Core/Public/UObject/ObjectVersion.h b/Engine/Source/Runtime/Core/Public/UObject/ObjectVersion.h
--- a/Engine/Source/Runtime/Core/Public/UObject/ObjectVersion.h
+++ b/Engine/Source/Runtime/Core/Public/UObject/ObjectVersion.h
@@ -554,6 +554,8 @@ enum EUnrealEngineObjectUE4Version
 	VER_UE4_FIX_MATERIAL_COMMENTS,
 	// Bone names stored as FName mean that we cant guarrentee the correct case on export, now we store a separate string for export purposes only
 	VER_UE4_STORE_BONE_EXPORT_NAMES,
+	// changed mesh emitter initial orientation to distribution
+	VER_UE4_MESH_EMITTER_INITIAL_ORIENTATION_DISTRIBUTION,
 
 	// -----<new versions can be added before this line>-------------------------------------------------
 	// - this needs to be the last line (see note below)
diff --git a/Engine/Source/Runtime/Engine/Classes/Particles/TypeData/ParticleModuleTypeDataMesh.h b/Engine/Source/Runtime/Engine/Classes/Particles/TypeData/ParticleModuleTypeDataMesh.h
--- a/Engine/Source/Runtime/Engine/Classes/Particles/TypeData/ParticleModuleTypeDataMesh.h
+++ b/Engine/Source/Runtime/Engine/Classes/Particles/TypeData/ParticleModuleTypeDataMesh.h
@@ -89,17 +89,22 @@ class UParticleModuleTypeDataMesh : public UParticleModuleTypeDataBase
 	UPROPERTY(EditAnywhere, Category=Mesh)
 	uint32 bOverrideMaterial:1;
 
-	/** The 'pre' rotation pitch (in degrees) to apply to the static mesh used. */
-	UPROPERTY(EditAnywhere, Category=Orientation)
-	float Pitch;
 
-	/** The 'pre' rotation roll (in degrees) to apply to the static mesh used. */
-	UPROPERTY(EditAnywhere, Category=Orientation)
-	float Roll;
+	/** deprecated properties for initial orientation */
+	UPROPERTY()
+	float Pitch_DEPRECATED;
+	UPROPERTY()
+	float Roll_DEPRECATED;
+	UPROPERTY()
+	float Yaw_DEPRECATED;
 
-	/** The 'pre' rotation yaw (in degrees) to apply to the static mesh used. */
+
+	/** The 'pre' rotation pitch (in degrees) to apply to the static mesh used. */
 	UPROPERTY(EditAnywhere, Category=Orientation)
-	float Yaw;
+	FRawDistributionVector RollPitchYawRange;
+
+	/** Random stream for the initial rotation distribution */
+	FRandomStream RandomStream;
 
 	/**
 	 *	The axis to lock the mesh on. This overrides TypeSpecific mesh alignment as well as the LockAxis module.
@@ -176,8 +181,13 @@ class UParticleModuleTypeDataMesh : public UParticleModuleTypeDataBase
 #if WITH_EDITOR
 	virtual void	PostEditChangeProperty(FPropertyChangedEvent& PropertyChangedEvent) OVERRIDE;
 #endif // WITH_EDITOR
+
+	virtual void	Serialize(FArchive& Ar) override;
+
 	// End UObject Interface
 
+	void CreateDistribution();
+
 	// Begin UParticleModule Interface
 	virtual void	SetToSensibleDefaults(UParticleEmitter* Owner) OVERRIDE;
 	// End UParticleModule Interface
diff --git a/Engine/Source/Runtime/Engine/Private/Particles/ParticleEmitterInstances.cpp b/Engine/Source/Runtime/Engine/Private/Particles/ParticleEmitterInstances.cpp
--- a/Engine/Source/Runtime/Engine/Private/Particles/ParticleEmitterInstances.cpp
+++ b/Engine/Source/Runtime/Engine/Private/Particles/ParticleEmitterInstances.cpp
@@ -2892,9 +2892,9 @@ void FParticleMeshEmitterInstance::Tick(float DeltaTime, bool bSuppressSpawning)
 
 				FQuat Rotation	= FQuat::FindBetween(OldDirection, NewDirection);
 				FVector Euler	= Rotation.Euler();
-				PayloadData->Rotation.X	= Euler.X;
-				PayloadData->Rotation.Y	= Euler.Y;
-				PayloadData->Rotation.Z	= Euler.Z;
+				PayloadData->Rotation.X	= PayloadData->InitialOrientation.X + Euler.X;
+				PayloadData->Rotation.Y = PayloadData->InitialOrientation.Y + Euler.Y;
+				PayloadData->Rotation.Z = PayloadData->InitialOrientation.Z + Euler.Z;
 			}
 	    }
 	}
@@ -2911,7 +2911,7 @@ void FParticleMeshEmitterInstance::Tick(float DeltaTime, bool bSuppressSpawning)
 			if ((Particle.Flags & STATE_Particle_FreezeRotation) == 0)
 			{
 				FMeshRotationPayloadData* PayloadData	 = (FMeshRotationPayloadData*)((uint8*)&Particle + MeshRotationOffset);
-				PayloadData->Rotation					+= DeltaTime * PayloadData->RotationRate;
+				PayloadData->Rotation					+= Particle.RelativeTime * PayloadData->RotationRate;
 			}
 		}
 	}
@@ -3074,6 +3074,9 @@ void FParticleMeshEmitterInstance::PostSpawn(FBaseParticle* Particle, float Inte
 {
 	FParticleEmitterInstance::PostSpawn(Particle, InterpolationPercentage, SpawnTime);
 	UParticleLODLevel* LODLevel = GetCurrentLODLevelChecked();
+
+	FMeshRotationPayloadData* PayloadData = (FMeshRotationPayloadData*)((uint8*)Particle + MeshRotationOffset);
+
 	if (LODLevel->RequiredModule->ScreenAlignment == PSA_Velocity)
 	{
 		// Determine the rotation to the velocity vector and apply it to the mesh
@@ -3084,12 +3087,14 @@ void FParticleMeshEmitterInstance::PostSpawn(FBaseParticle* Particle, float Inte
 		FQuat Rotation	= FQuat::FindBetween(OldDirection, NewDirection);
 		FVector Euler	= Rotation.Euler();
 
-		FMeshRotationPayloadData* PayloadData	= (FMeshRotationPayloadData*)((uint8*)Particle + MeshRotationOffset);
 		PayloadData->Rotation.X	+= Euler.X;
 		PayloadData->Rotation.Y	+= Euler.Y;
 		PayloadData->Rotation.Z	+= Euler.Z;
-		//
 	}
+
+	FVector InitialOrient = MeshTypeData->RollPitchYawRange.GetValue(SpawnTime, 0, 0, &MeshTypeData->RandomStream);
+	PayloadData->InitialOrientation = InitialOrient;
+	PayloadData->Rotation += InitialOrient;
 }
 
 bool FParticleMeshEmitterInstance::IsDynamicDataRequired(UParticleLODLevel* CurrentLODLevel)
diff --git a/Engine/Source/Runtime/Engine/Private/Particles/ParticleModules.cpp b/Engine/Source/Runtime/Engine/Private/Particles/ParticleModules.cpp
--- a/Engine/Source/Runtime/Engine/Private/Particles/ParticleModules.cpp
+++ b/Engine/Source/Runtime/Engine/Private/Particles/ParticleModules.cpp
@@ -3088,6 +3088,8 @@ FParticleEmitterInstance* UParticleModuleTypeDataMesh::CreateInstance(UParticleE
 
 	Instance->InitParameters(InEmitterParent, InComponent);
 
+	CreateDistribution();
+
 	return Instance;
 }
 
@@ -3099,6 +3101,29 @@ void UParticleModuleTypeDataMesh::SetToSensibleDefaults(UParticleEmitter* Owner)
 	}
 }
 
+void UParticleModuleTypeDataMesh::Serialize(FArchive& Ar)
+{
+	Super::Serialize(Ar);
+	if (Ar.IsLoading() && Ar.UE4Ver() < VER_UE4_MESH_EMITTER_INITIAL_ORIENTATION_DISTRIBUTION)
+	{
+		FVector oldOrient(Roll_DEPRECATED, Pitch_DEPRECATED, Yaw_DEPRECATED);
+		UDistributionVectorUniform* RPYDistribution = Cast<UDistributionVectorUniform>(RollPitchYawRange.Distribution);
+		RPYDistribution->Min = oldOrient;
+		RPYDistribution->Max = oldOrient;
+		RPYDistribution->bIsDirty = true;
+	}
+}
+
+void UParticleModuleTypeDataMesh::CreateDistribution()
+{
+	if (!RollPitchYawRange.Distribution)
+	{
+		RollPitchYawRange.Distribution = NewNamedObject<UDistributionVectorUniform>(this, TEXT("DistributionRollPitchYaw"));
+	}
+}
+
+
+
 #if WITH_EDITOR
 void UParticleModuleTypeDataMesh::PostEditChangeProperty(FPropertyChangedEvent& PropertyChangedEvent)
 {
diff --git a/Engine/Source/Runtime/Engine/Private/Particles/ParticleSystemRender.cpp b/Engine/Source/Runtime/Engine/Private/Particles/ParticleSystemRender.cpp
--- a/Engine/Source/Runtime/Engine/Private/Particles/ParticleSystemRender.cpp
+++ b/Engine/Source/Runtime/Engine/Private/Particles/ParticleSystemRender.cpp
@@ -1559,7 +1559,6 @@ FDynamicMeshEmitterData::FDynamicMeshEmitterData(const UParticleModuleRequired*
 	, StaticMesh( NULL )
 	, MeshTypeDataOffset(0xFFFFFFFF)
 	, bApplyPreRotation(false)
-	, RollPitchYaw(0.0f, 0.0f, 0.0f)
 	, bUseMeshLockedAxis(false)
 	, bUseCameraFacing(false)
 	, bApplyParticleRotationAsSpin(false)
@@ -1615,11 +1614,14 @@ void FDynamicMeshEmitterData::Init( bool bInSelected,
 		// offset to the mesh emitter type data
 		MeshTypeDataOffset = InEmitterInstance->TypeDataOffset;
 
-		// Setup pre-rotation values...
-		if ((MeshTD->Pitch != 0.0f) || (MeshTD->Roll != 0.0f) || (MeshTD->Yaw != 0.0f))
+
+		FVector Mins, Maxs;
+		MeshTD->RollPitchYawRange.Distribution->GetRange(Mins, Maxs);
+
+		// Enable/Disable pre-rotation
+		if ( Mins.SizeSquared() || Maxs.SizeSquared() )
 		{
 			bApplyPreRotation = true;
-			RollPitchYaw = FVector(MeshTD->Roll, MeshTD->Pitch, MeshTD->Yaw);
 		}
 		else
 		{
@@ -2214,16 +2216,18 @@ void FDynamicMeshEmitterData::GetParticleTransform(FBaseParticle& InParticle, co
 				kRotator = FRotator(MeshRotation);
 		}
 	}
-	else if (Source.bMeshRotationActive)
-	{
-		FMeshRotationPayloadData* PayloadData = (FMeshRotationPayloadData*)((uint8*)&InParticle + Source.MeshRotationOffset);
-		kRotator = FRotator::MakeFromEuler(PayloadData->Rotation);
-	}
 	else
 	{
 		float fRot = InParticle.Rotation * 180.0f / PI;
 		FVector kRotVec = FVector(fRot, fRot, fRot);
 		kRotator = FRotator::MakeFromEuler(kRotVec);
+
+		// if we're using initial orientation, apply here
+		if (bApplyPreRotation)
+		{
+			FMeshRotationPayloadData* PayloadData = (FMeshRotationPayloadData*)((uint8*)&InParticle + Source.MeshRotationOffset);
+			kRotator += FRotator::MakeFromEuler(PayloadData->Rotation);
+		}
 	}
 
 	FRotationMatrix kRotMat(kRotator);
@@ -2231,11 +2235,11 @@ void FDynamicMeshEmitterData::GetParticleTransform(FBaseParticle& InParticle, co
 	{
 		if ((bUseCameraFacing == true) || (bUseMeshLockedAxis == true))
 		{
-			OutTransformMat = FRotationMatrix(FRotator::MakeFromEuler(RollPitchYaw)) * kScaleMat * FRotationMatrix(kLockedAxisRotator) * kRotMat * kTransMat;
+			OutTransformMat = kScaleMat * FRotationMatrix(kLockedAxisRotator) * kRotMat * kTransMat;
 		}
 		else
 		{
-			OutTransformMat = FRotationMatrix(FRotator::MakeFromEuler(RollPitchYaw)) * kScaleMat * kRotMat * kTransMat;
+			OutTransformMat = kScaleMat * kRotMat * kTransMat;
 		}
 	}
 	else if ((bUseCameraFacing == true) || (bUseMeshLockedAxis == true))
diff --git a/Engine/Source/Runtime/Engine/Public/ParticleHelper.h b/Engine/Source/Runtime/Engine/Public/ParticleHelper.h
--- a/Engine/Source/Runtime/Engine/Public/ParticleHelper.h
+++ b/Engine/Source/Runtime/Engine/Public/ParticleHelper.h
@@ -750,6 +750,7 @@ struct FAnimTrailTypeDataPayload : public FTrailsBaseTypeDataPayload
 /** Mesh rotation data payload										*/
 struct FMeshRotationPayloadData
 {
+	FVector	 InitialOrientation;
 	FVector  Rotation;
 	FVector  RotationRate;
 	FVector  RotationRateBase;
@@ -2018,8 +2019,6 @@ struct FDynamicMeshEmitterData : public FDynamicSpriteEmitterDataBase
 	// These don't need to go into the replay data, as they are constant over the life of the emitter
 	/** If true, apply the 'pre-rotation' values to the mesh. */
 	uint32 bApplyPreRotation:1;
-	/** The pitch/roll/yaw to apply in the pre-rotation step */
-	FVector RollPitchYaw;
 	/** If true, then use the locked axis setting supplied. Trumps locked axis module and/or TypeSpecific mesh settings. */
 	uint32 bUseMeshLockedAxis:1;
 	/** If true, then use the camera facing options supplied. Trumps all other settings. */
