Fixing potential crash.

diff --git a/Engine/Source/Runtime/Engine/Private/Particles/ParticleModules.cpp b/Engine/Source/Runtime/Engine/Private/Particles/ParticleModules.cpp
--- a/Engine/Source/Runtime/Engine/Private/Particles/ParticleModules.cpp
+++ b/Engine/Source/Runtime/Engine/Private/Particles/ParticleModules.cpp
@@ -3107,6 +3107,7 @@ void UParticleModuleTypeDataMesh::Serialize(FArchive& Ar)
 	if (Ar.IsLoading() && Ar.UE4Ver() < VER_UE4_MESH_EMITTER_INITIAL_ORIENTATION_DISTRIBUTION)
 	{
 		FVector oldOrient(Roll_DEPRECATED, Pitch_DEPRECATED, Yaw_DEPRECATED);
+		CreateDistribution();
 		UDistributionVectorUniform* RPYDistribution = Cast<UDistributionVectorUniform>(RollPitchYawRange.Distribution);
 		RPYDistribution->Min = oldOrient;
 		RPYDistribution->Max = oldOrient;
