make constraint UI a bit easier to understand
TTP#318776

diff --git a/Engine/Source/Runtime/Engine/Classes/PhysicsEngine/ConstraintInstance.h b/Engine/Source/Runtime/Engine/Classes/PhysicsEngine/ConstraintInstance.h
--- a/Engine/Source/Runtime/Engine/Classes/PhysicsEngine/ConstraintInstance.h
+++ b/Engine/Source/Runtime/Engine/Classes/PhysicsEngine/ConstraintInstance.h
@@ -275,11 +275,11 @@ struct ENGINE_API FConstraintInstance
 	FVector LinearVelocityTarget;
 
 	/** Spring to apply to the for linear drive. */
-	UPROPERTY(EditAnywhere, BlueprintReadOnly, Category=LinearMotor)
+	UPROPERTY(EditAnywhere, BlueprintReadOnly, Category = LinearMotor, meta = (DisplayName = "Linear Position Strength", editcondition = "bLinearPositionDrive"))
 	float LinearDriveSpring;
 
 	/** Damping to apply to the for linear drive. */
-	UPROPERTY(EditAnywhere, BlueprintReadOnly, Category=LinearMotor)
+	UPROPERTY(EditAnywhere, BlueprintReadOnly, Category = LinearMotor, meta = (DisplayName = "Linear Velocity Strength", editcondition = "bLinearVelocityDrive"))
 	float LinearDriveDamping;
 
 	/** Limit to the force the linear drive can apply. */
@@ -321,11 +321,11 @@ struct ENGINE_API FConstraintInstance
 	FVector AngularVelocityTarget;    // Revolutions per second
 
 	/** Spring value to apply to the for angular drive. */
-	UPROPERTY(EditAnywhere, BlueprintReadOnly, Category=AngularMotor)
+	UPROPERTY(EditAnywhere, BlueprintReadOnly, Category = AngularMotor, meta = (DisplayName = "Angular Position Strength", editcondition = "bAngularOrientationDrive"))
 	float AngularDriveSpring;
 
-	/** Damping value to apply to the for angulardrive. */
-	UPROPERTY(EditAnywhere, BlueprintReadOnly, Category=AngularMotor)
+	/** Damping value to apply to the for angular drive. */
+	UPROPERTY(EditAnywhere, BlueprintReadOnly, Category = AngularMotor, meta = (DisplayName = "Angular Velocity Strength", editcondition = "bAngularVelocityDrive"))
 	float AngularDriveDamping;
 
 	/** Limit to the force the angular drive can apply. */
