#ue4 Making ResetOwnedComponents public, to undo all the friend classing we had to do in places where blueprints needed a refresh

diff --git a/Engine/Source/Runtime/Engine/Classes/GameFramework/Actor.h b/Engine/Source/Runtime/Engine/Classes/GameFramework/Actor.h
--- a/Engine/Source/Runtime/Engine/Classes/GameFramework/Actor.h
+++ b/Engine/Source/Runtime/Engine/Classes/GameFramework/Actor.h
@@ -2039,15 +2039,12 @@ public:
 #if DO_CHECK
 	bool OwnsComponent(UActorComponent* Component) const;
 #endif
+	void ResetOwnedComponents();
 
 private:
 	UPROPERTY(transient, duplicatetransient)
 	TArray<UActorComponent*> OwnedComponents;
-
-	friend class FBlueprintEditorUtils;
-	friend class UEngine;
-	void ResetOwnedComponents();
-
+	
 public:
 	//=============================================================================
 	// Navigation related functions
