#ue4 -  Add native and Blueprint events triggered at the end of a CharacterMovementComponent movement update.

- This is the preferred event to use rather than the Tick event when performing custom updates to CharacterMovement properties based on the current state. This is mainly due to the nature of network updates, where client corrections in position from the server can cause multiple iterations of a movement update, which allows this event to update as well, while a Tick event would not.

diff --git a/Engine/Source/Runtime/Engine/Classes/GameFramework/Character.h b/Engine/Source/Runtime/Engine/Classes/GameFramework/Character.h
--- a/Engine/Source/Runtime/Engine/Classes/GameFramework/Character.h
+++ b/Engine/Source/Runtime/Engine/Classes/GameFramework/Character.h
@@ -5,6 +5,7 @@
 #include "Character.generated.h"
 
 DECLARE_DYNAMIC_MULTICAST_DELEGATE_ThreeParams(FMovementModeChangedSignature, class ACharacter*, Character, EMovementMode, PrevMovementMode, uint8, PreviousCustomMode);
+DECLARE_DYNAMIC_MULTICAST_DELEGATE_ThreeParams(FCharacterMovementUpdatedSignature, float, DeltaSeconds, FVector, OldLocation, FVector, OldVelocity);
 //
 // Forward declarations
 //
@@ -467,6 +468,19 @@ public:
 	UFUNCTION(BlueprintImplementableEvent, meta=(FriendlyName= "UpdateCustomMovement"))
 	virtual void K2_UpdateCustomMovement(float DeltaTime);
 
+	/**
+	 * Event triggered at the end of a CharacterMovementComponent movement update.
+	 * This is the preferred event to use rather than the Tick event when performing custom updates to CharacterMovement properties based on the current state.
+	 * This is mainly due to the nature of network updates, where client corrections in position from the server can cause multiple iterations of a movement update,
+	 * which allows this event to update as well, while a Tick event would not.
+	 *
+	 * @param	DeltaSeconds		Delta time in seconds for this update
+	 * @param	InitialLocation		Location at the start of the update. May be different than the current location if movement occurred.
+	 * @param	InitialVelocity		Velocity at the start of the update. May be different than the current velocity.
+	 */
+	UPROPERTY(BlueprintAssignable, Category="Pawn|Character")
+	FCharacterMovementUpdatedSignature OnCharacterMovementUpdated;
+
 	/**
 	 * Called when our pawn has landed from a fall.
 	 * @param	PrevMovementMode	Movement mode before the change
diff --git a/Engine/Source/Runtime/Engine/Classes/GameFramework/CharacterMovementComponent.h b/Engine/Source/Runtime/Engine/Classes/GameFramework/CharacterMovementComponent.h
--- a/Engine/Source/Runtime/Engine/Classes/GameFramework/CharacterMovementComponent.h
+++ b/Engine/Source/Runtime/Engine/Classes/GameFramework/CharacterMovementComponent.h
@@ -392,7 +392,12 @@ public:
 	UPROPERTY()
 	uint32 bMovementInProgress:1;
 
-	/** If true, high-level movement updates will be wrapped in a movement scope that accumulates updates and defers a bulk of the work until the end. */
+	/**
+	 * If true, high-level movement updates will be wrapped in a movement scope that accumulates updates and defers a bulk of the work until the end.
+	 * When enabled, touch and hit events will not be triggered until the end of multiple moves within an update, which can improve performance.
+	 *
+	 * @see FScopedMovementUpdate
+	 */
 	UPROPERTY(Category="Character Movement", EditAnywhere, AdvancedDisplay)
 	uint32 bEnableScopedMovementUpdates:1;
 
@@ -1337,6 +1342,15 @@ protected:
 	/** Scale input acceleration, based on movement acceleration rate. */
 	virtual FVector ScaleInputAcceleration(const FVector& InputAcceleration) const;
 
+	/**
+	 * Event triggered at the end of a movement update. If scoped movement updates are enabled (bEnableScopedMovementUpdates), this is within such a scope.
+	 * If that is not desired, bind to the CharacterOwner's OnMovementUpdated event instead, as that is triggered after the scoped movement update.
+	 */
+	virtual void OnMovementUpdated(float DeltaSeconds, const FVector& OldLocation, const FVector& OldVelocity);
+
+	/** Internal function to call OnMovementUpdated delegate on CharacterOwner. */
+	virtual void CallMovementUpdateDelegate(float DeltaSeconds, const FVector& OldLocation, const FVector& OldVelocity);
+
 public:
 	// Movement functions broken out based on owner's network Role.
 	// TickComponent calls the correct version based on the Role.
diff --git a/Engine/Source/Runtime/Engine/Private/CharacterMovementComponent.cpp b/Engine/Source/Runtime/Engine/Private/CharacterMovementComponent.cpp
--- a/Engine/Source/Runtime/Engine/Private/CharacterMovementComponent.cpp
+++ b/Engine/Source/Runtime/Engine/Private/CharacterMovementComponent.cpp
@@ -975,96 +975,111 @@ void UCharacterMovementComponent::SimulateMovement(float DeltaSeconds)
 		return;
 	}
 
-	FScopedMovementUpdate ScopedMovementUpdate(UpdatedComponent, bEnableScopedMovementUpdates ? EScopedUpdate::DeferredUpdates : EScopedUpdate::ImmediateUpdates);
+	FVector OldVelocity;
+	FVector OldLocation;
 
-	if (bIsSimulatedProxy)
+	// Scoped updates can improve performance of multiple MoveComponent calls.
 	{
-		// Handle network changes
-		if (bNetworkUpdateReceived)
+		FScopedMovementUpdate ScopedMovementUpdate(UpdatedComponent, bEnableScopedMovementUpdates ? EScopedUpdate::DeferredUpdates : EScopedUpdate::ImmediateUpdates);
+
+		if (bIsSimulatedProxy)
 		{
-			bNetworkUpdateReceived = false;
-			if (bNetworkMovementModeChanged)
-			{
-				bNetworkMovementModeChanged = false;
-				ApplyNetworkMovementMode(CharacterOwner->GetReplicatedMovementMode());
-			}
-			else if (bJustTeleported)
+			// Handle network changes
+			if (bNetworkUpdateReceived)
 			{
-				// Make sure floor is current. We will continue using the replicated base, if there was one.
-				bJustTeleported = false;
-				UpdateFloorFromAdjustment();
+				bNetworkUpdateReceived = false;
+				if (bNetworkMovementModeChanged)
+				{
+					bNetworkMovementModeChanged = false;
+					ApplyNetworkMovementMode(CharacterOwner->GetReplicatedMovementMode());
+				}
+				else if (bJustTeleported)
+				{
+					// Make sure floor is current. We will continue using the replicated base, if there was one.
+					bJustTeleported = false;
+					UpdateFloorFromAdjustment();
+				}
 			}
+
+			HandlePendingLaunch();
 		}
 
-		HandlePendingLaunch();
-	}
+		if (MovementMode == MOVE_None)
+		{
+			return;
+		}
 
-	if (MovementMode == MOVE_None)
-	{
-		return;
-	}
+		Acceleration = Velocity.SafeNormal();	// Not currently used for simulated movement
+		AnalogInputModifier = 1.0f;				// Not currently used for simulated movement
 
-	Acceleration = Velocity.SafeNormal();	// Not currently used for simulated movement
-	AnalogInputModifier = 1.0f;				// Not currently used for simulated movement
+		MaybeUpdateBasedMovement(DeltaSeconds);
 
-	// simulated pawns predict location
-	MaybeUpdateBasedMovement(DeltaSeconds);
-	FStepDownResult StepDownResult;
-	MoveSmooth(Velocity, DeltaSeconds, &StepDownResult);
+		// simulated pawns predict location
+		OldVelocity = Velocity;
+		OldLocation = UpdatedComponent->GetComponentLocation();
+		FStepDownResult StepDownResult;
+		MoveSmooth(Velocity, DeltaSeconds, &StepDownResult);
 
-	// if simulated gravity, find floor and check if falling
-	const bool bEnableFloorCheck = (!CharacterOwner->bSimGravityDisabled || !bIsSimulatedProxy);
-	if (bEnableFloorCheck && (MovementMode == MOVE_Walking || MovementMode == MOVE_Falling))
-	{
-		const FVector CollisionCenter = UpdatedComponent->GetComponentLocation();
-		if (StepDownResult.bComputedFloor)
-		{
-			CurrentFloor = StepDownResult.FloorResult;
-		}
-		else if (Velocity.Z <= 0.f)
+		// if simulated gravity, find floor and check if falling
+		const bool bEnableFloorCheck = (!CharacterOwner->bSimGravityDisabled || !bIsSimulatedProxy);
+		if (bEnableFloorCheck && (MovementMode == MOVE_Walking || MovementMode == MOVE_Falling))
 		{
-			FindFloor(CollisionCenter, CurrentFloor, Velocity.IsZero(), NULL);
-		}
-		else
-		{
-			CurrentFloor.Clear();
-		}
+			const FVector CollisionCenter = UpdatedComponent->GetComponentLocation();
+			if (StepDownResult.bComputedFloor)
+			{
+				CurrentFloor = StepDownResult.FloorResult;
+			}
+			else if (Velocity.Z <= 0.f)
+			{
+				FindFloor(CollisionCenter, CurrentFloor, Velocity.IsZero(), NULL);
+			}
+			else
+			{
+				CurrentFloor.Clear();
+			}
 
-		if (!CurrentFloor.IsWalkableFloor())
-		{
-			// No floor, must fall.
-			Velocity = NewFallVelocity(Velocity, FVector(0.f,0.f,GetGravityZ()), DeltaSeconds);
-			SetMovementMode(MOVE_Falling);
-		}
-		else
-		{
-			// Walkable floor
-			if (MovementMode == MOVE_Walking)
+			if (!CurrentFloor.IsWalkableFloor())
 			{
-				AdjustFloorHeight();
-				SetBase(CurrentFloor.HitResult.Component.Get());
+				// No floor, must fall.
+				Velocity = NewFallVelocity(Velocity, FVector(0.f,0.f,GetGravityZ()), DeltaSeconds);
+				SetMovementMode(MOVE_Falling);
 			}
-			else if (MovementMode == MOVE_Falling)
+			else
 			{
-				if (CurrentFloor.FloorDist <= MIN_FLOOR_DIST)
+				// Walkable floor
+				if (MovementMode == MOVE_Walking)
 				{
-					// Landed
-					SetMovementMode(MOVE_Walking);
+					AdjustFloorHeight();
+					SetBase(CurrentFloor.HitResult.Component.Get());
 				}
-				else
+				else if (MovementMode == MOVE_Falling)
 				{
-					// Continue falling.
-					Velocity = NewFallVelocity(Velocity, FVector(0.f,0.f,GetGravityZ()), DeltaSeconds);
-					CurrentFloor.Clear();
+					if (CurrentFloor.FloorDist <= MIN_FLOOR_DIST)
+					{
+						// Landed
+						SetMovementMode(MOVE_Walking);
+					}
+					else
+					{
+						// Continue falling.
+						Velocity = NewFallVelocity(Velocity, FVector(0.f,0.f,GetGravityZ()), DeltaSeconds);
+						CurrentFloor.Clear();
+					}
 				}
 			}
 		}
-	}
+
+		OnMovementUpdated(DeltaSeconds, OldLocation, OldVelocity);
+	} // End scoped movement update
+
+	// Call custom post-movement events. These happen after the scoped movement completes in case the events want to use the current state of overlaps etc.
+	CallMovementUpdateDelegate(DeltaSeconds, OldLocation, OldVelocity);
 
 	MaybeSaveBaseLocation();
-	
 	UpdateComponentVelocity();
 	bJustTeleported = false;
+
+	LastUpdateLocation = UpdatedComponent ? UpdatedComponent->GetComponentLocation() : FVector::ZeroVector;
 }
 
 void UCharacterMovementComponent::SetBase( UPrimitiveComponent* NewBase, bool bNotifyActor )
@@ -1245,140 +1260,166 @@ void UCharacterMovementComponent::PerformMovement(float DeltaSeconds)
 	// Force floor update if we've moved outside of CharacterMovement since last update.
 	bForceNextFloorCheck |= (IsMovingOnGround() && UpdatedComponent->GetComponentLocation() != LastUpdateLocation);
 
-	// Scoped updates can improve performance of multiple MoveComponent calls.
-	FScopedMovementUpdate ScopedMovementUpdate(UpdatedComponent, bEnableScopedMovementUpdates ? EScopedUpdate::DeferredUpdates : EScopedUpdate::ImmediateUpdates);
-
-	MaybeUpdateBasedMovement(DeltaSeconds);
-	ApplyAccumulatedForces(DeltaSeconds);
+	FVector OldVelocity;
+	FVector OldLocation;
 
-	FVector OldVelocity = Velocity;
-	FVector OldLocation = CharacterOwner->GetActorLocation();
-
-	// Check for a change in crouch state. Players toggle crouch by changing bWantsToCrouch.
-	const bool bAllowedToCrouch = CanCrouchInCurrentState();
-	if ((!bAllowedToCrouch || !bWantsToCrouch) && IsCrouching())
-	{
-		UnCrouch(false);
-	}
-	else if (bWantsToCrouch && bAllowedToCrouch && !IsCrouching()) 
+	// Scoped updates can improve performance of multiple MoveComponent calls.
 	{
-		Crouch(false);
-	}
+		FScopedMovementUpdate ScopedMovementUpdate(UpdatedComponent, bEnableScopedMovementUpdates ? EScopedUpdate::DeferredUpdates : EScopedUpdate::ImmediateUpdates);
 
-	// Character::LaunchCharacter() has been deferred until now.
-	HandlePendingLaunch();
+		MaybeUpdateBasedMovement(DeltaSeconds);
+		ApplyAccumulatedForces(DeltaSeconds);
 
-	// If using RootMotion, tick animations before running physics.
-	if( !CharacterOwner->bClientUpdating && CharacterOwner->IsPlayingRootMotion() && CharacterOwner->Mesh )
-	{
-		CharacterOwner->Mesh->TickPose(DeltaSeconds);
+		OldVelocity = Velocity;
+		OldLocation = CharacterOwner->GetActorLocation();
 
-		// Make sure animation didn't trigger an event that destroyed us
-		if (!HasValidData())
+		// Check for a change in crouch state. Players toggle crouch by changing bWantsToCrouch.
+		const bool bAllowedToCrouch = CanCrouchInCurrentState();
+		if ((!bAllowedToCrouch || !bWantsToCrouch) && IsCrouching())
 		{
-			return;
+			UnCrouch(false);
 		}
-
-		// For local human clients, save off root motion data so it can be used by movement networking code.
-		if( CharacterOwner->IsLocallyControlled() && (CharacterOwner->Role == ROLE_AutonomousProxy) )
+		else if (bWantsToCrouch && bAllowedToCrouch && !IsCrouching()) 
 		{
-			CharacterOwner->ClientRootMotionParams = RootMotionParams;
+			Crouch(false);
 		}
-	}
 
-	// if we're about to use root motion, convert it to world space first.
-	if( HasRootMotion() )
-	{
-		USkeletalMeshComponent * SkelMeshComp = CharacterOwner->Mesh;
-		if( SkelMeshComp )
+		// Character::LaunchCharacter() has been deferred until now.
+		HandlePendingLaunch();
+
+		// If using RootMotion, tick animations before running physics.
+		if( !CharacterOwner->bClientUpdating && CharacterOwner->IsPlayingRootMotion() && CharacterOwner->Mesh )
 		{
-			// Convert Local Space Root Motion to world space. Do it right before used by physics to make sure we use up to date transforms, as translation is relative to rotation.
-			RootMotionParams.Set( SkelMeshComp->ConvertLocalRootMotionToWorld(RootMotionParams.RootMotionTransform) );
-			UE_LOG(LogRootMotion, Log,  TEXT("PerformMovement WorldSpaceRootMotion Translation: %s, Rotation: %s, Actor Facing: %s"),
-				*RootMotionParams.RootMotionTransform.GetTranslation().ToCompactString(), *RootMotionParams.RootMotionTransform.GetRotation().Rotator().ToCompactString(), *CharacterOwner->GetActorRotation().Vector().ToCompactString());
+			CharacterOwner->Mesh->TickPose(DeltaSeconds);
+
+			// Make sure animation didn't trigger an event that destroyed us
+			if (!HasValidData())
+			{
+				return;
+			}
+
+			// For local human clients, save off root motion data so it can be used by movement networking code.
+			if( CharacterOwner->IsLocallyControlled() && (CharacterOwner->Role == ROLE_AutonomousProxy) )
+			{
+				CharacterOwner->ClientRootMotionParams = RootMotionParams;
+			}
 		}
 
-		// Then turn root motion to velocity to be used by various physics modes.
-		if( DeltaSeconds > 0.f )
+		// if we're about to use root motion, convert it to world space first.
+		if( HasRootMotion() )
 		{
-			const FVector RootMotionVelocity = RootMotionParams.RootMotionTransform.GetTranslation() / DeltaSeconds;
-			// Do not override Velocity.Z if in falling physics, we want to keep the effect of gravity.
-			Velocity = FVector(RootMotionVelocity.X, RootMotionVelocity.Y, (MovementMode == MOVE_Falling ? Velocity.Z : RootMotionVelocity.Z));
-		}
-	}
+			USkeletalMeshComponent * SkelMeshComp = CharacterOwner->Mesh;
+			if( SkelMeshComp )
+			{
+				// Convert Local Space Root Motion to world space. Do it right before used by physics to make sure we use up to date transforms, as translation is relative to rotation.
+				RootMotionParams.Set( SkelMeshComp->ConvertLocalRootMotionToWorld(RootMotionParams.RootMotionTransform) );
+				UE_LOG(LogRootMotion, Log,  TEXT("PerformMovement WorldSpaceRootMotion Translation: %s, Rotation: %s, Actor Facing: %s"),
+					*RootMotionParams.RootMotionTransform.GetTranslation().ToCompactString(), *RootMotionParams.RootMotionTransform.GetRotation().Rotator().ToCompactString(), *CharacterOwner->GetActorRotation().Vector().ToCompactString());
+			}
 
-	// NaN tracking
-	checkf(!Velocity.ContainsNaN(), TEXT("UCharacterMovementComponent::PerformMovement: Velocity contains NaN (%s: %s)\n%s"), *GetPathNameSafe(this), *GetPathNameSafe(GetOuter()), *Velocity.ToString());
+			// Then turn root motion to velocity to be used by various physics modes.
+			if( DeltaSeconds > 0.f )
+			{
+				const FVector RootMotionVelocity = RootMotionParams.RootMotionTransform.GetTranslation() / DeltaSeconds;
+				// Do not override Velocity.Z if in falling physics, we want to keep the effect of gravity.
+				Velocity = FVector(RootMotionVelocity.X, RootMotionVelocity.Y, (MovementMode == MOVE_Falling ? Velocity.Z : RootMotionVelocity.Z));
+			}
+		}
 
-	// change position
-	StartNewPhysics(DeltaSeconds, 0);
+		// NaN tracking
+		checkf(!Velocity.ContainsNaN(), TEXT("UCharacterMovementComponent::PerformMovement: Velocity contains NaN (%s: %s)\n%s"), *GetPathNameSafe(this), *GetPathNameSafe(GetOuter()), *Velocity.ToString());
 
-	if (!HasValidData())
-	{
-		return;
-	}
+		// change position
+		StartNewPhysics(DeltaSeconds, 0);
 
-	// uncrouch if no longer allowed to be crouched
-	if (IsCrouching() && !CanCrouchInCurrentState())
-	{
-		UnCrouch(false);
-	}
+		if (!HasValidData())
+		{
+			return;
+		}
 
-	if( CharacterOwner->Controller || bRunPhysicsWithNoController )
-	{
-		if (!HasRootMotion() && !CharacterOwner->IsMatineeControlled())
+		// uncrouch if no longer allowed to be crouched
+		if (IsCrouching() && !CanCrouchInCurrentState())
 		{
-			PhysicsRotation(DeltaSeconds);
+			UnCrouch(false);
 		}
-	}
 
-	// Apply Root Motion rotation after movement is complete.
-	if( HasRootMotion() )
-	{
-		const FRotator OldActorRotation = CharacterOwner->GetActorRotation();
-		const FRotator RootMotionRotation = RootMotionParams.RootMotionTransform.GetRotation().Rotator();
-		if( !RootMotionRotation.IsNearlyZero() )
+		if( CharacterOwner->Controller || bRunPhysicsWithNoController )
 		{
-			const FRotator NewActorRotation = (OldActorRotation + RootMotionRotation).GetNormalized();
-			MoveUpdatedComponent(FVector::ZeroVector, NewActorRotation, true);
+			if (!HasRootMotion() && !CharacterOwner->IsMatineeControlled())
+			{
+				PhysicsRotation(DeltaSeconds);
+			}
 		}
 
-		// debug
-		if( false )
+		// Apply Root Motion rotation after movement is complete.
+		if( HasRootMotion() )
 		{
-			const FVector ResultingLocation = CharacterOwner->GetActorLocation();
-			const FRotator ResultingRotation = CharacterOwner->GetActorRotation();
+			const FRotator OldActorRotation = CharacterOwner->GetActorRotation();
+			const FRotator RootMotionRotation = RootMotionParams.RootMotionTransform.GetRotation().Rotator();
+			if( !RootMotionRotation.IsNearlyZero() )
+			{
+				const FRotator NewActorRotation = (OldActorRotation + RootMotionRotation).GetNormalized();
+				MoveUpdatedComponent(FVector::ZeroVector, NewActorRotation, true);
+			}
+
+			// debug
+			if( false )
+			{
+				const FVector ResultingLocation = CharacterOwner->GetActorLocation();
+				const FRotator ResultingRotation = CharacterOwner->GetActorRotation();
+
+				// Show current position
+				DrawDebugCoordinateSystem(GetWorld(), CharacterOwner->Mesh->GetComponentLocation() + FVector(0,0,1), ResultingRotation, 50.f, false);
 
-			// Show current position
-			DrawDebugCoordinateSystem(GetWorld(), CharacterOwner->Mesh->GetComponentLocation() + FVector(0,0,1), ResultingRotation, 50.f, false);
+				// Show resulting delta move.
+				DrawDebugLine(GetWorld(), OldLocation, ResultingLocation, FColor::Red, true, 10.f);
 
-			// Show resulting delta move.
-			DrawDebugLine(GetWorld(), OldLocation, ResultingLocation, FColor::Red, true, 10.f);
+				// Log details.
+				UE_LOG(LogRootMotion, Warning,  TEXT("PerformMovement Resulting DeltaMove Translation: %s, Rotation: %s, MovementBase: %s"),
+					*(ResultingLocation - OldLocation).ToCompactString(), *(ResultingRotation - OldActorRotation).GetNormalized().ToCompactString(), *GetNameSafe(CharacterOwner->GetMovementBase()) );
 
-			// Log details.
-			UE_LOG(LogRootMotion, Warning,  TEXT("PerformMovement Resulting DeltaMove Translation: %s, Rotation: %s, MovementBase: %s"),
-				*(ResultingLocation - OldLocation).ToCompactString(), *(ResultingRotation - OldActorRotation).GetNormalized().ToCompactString(), *GetNameSafe(CharacterOwner->GetMovementBase()) );
+				const FVector RMTranslation = RootMotionParams.RootMotionTransform.GetTranslation();
+				const FRotator RMRotation = RootMotionParams.RootMotionTransform.GetRotation().Rotator();
+				UE_LOG(LogRootMotion, Warning,  TEXT("PerformMovement Resulting DeltaError Translation: %s, Rotation: %s"),
+					*(ResultingLocation - OldLocation - RMTranslation).ToCompactString(), *(ResultingRotation - OldActorRotation - RMRotation).GetNormalized().ToCompactString() );
+			}
 
-			const FVector RMTranslation = RootMotionParams.RootMotionTransform.GetTranslation();
-			const FRotator RMRotation = RootMotionParams.RootMotionTransform.GetRotation().Rotator();
-			UE_LOG(LogRootMotion, Warning,  TEXT("PerformMovement Resulting DeltaError Translation: %s, Rotation: %s"),
-				*(ResultingLocation - OldLocation - RMTranslation).ToCompactString(), *(ResultingRotation - OldActorRotation - RMRotation).GetNormalized().ToCompactString() );
+			// Root Motion has been used, clear
+			RootMotionParams.Clear();
 		}
 
-		// Root Motion has been used, clear
-		RootMotionParams.Clear();
-	}
+		OnMovementUpdated(DeltaSeconds, OldLocation, OldVelocity);
+	} // End scoped movement update
 
-	// update relative location/rotation
-	MaybeSaveBaseLocation();
+	// Call external post-movement events. These happen after the scoped movement completes in case the events want to use the current state of overlaps etc.
+	CallMovementUpdateDelegate(DeltaSeconds, OldLocation, OldVelocity);
 
-	// update component velocity
+	MaybeSaveBaseLocation();
 	UpdateComponentVelocity();
 
 	LastUpdateLocation = UpdatedComponent ? UpdatedComponent->GetComponentLocation() : FVector::ZeroVector;
 }
 
 
+void UCharacterMovementComponent::CallMovementUpdateDelegate(float DeltaTime, const FVector& OldLocation, const FVector& OldVelocity)
+{
+	// Update component velocity in case events want to read it
+	UpdateComponentVelocity();
+
+	// Delegate (for blueprints)
+	if (CharacterOwner)
+	{
+		CharacterOwner->OnCharacterMovementUpdated.Broadcast(DeltaTime, OldLocation, OldVelocity);
+	}
+}
+
+
+void UCharacterMovementComponent::OnMovementUpdated(float DeltaTime, const FVector& OldLocation, const FVector& OldVelocity)
+{
+	// empty base implementation, intended for derived classes to override.
+}
+
+
 void UCharacterMovementComponent::SaveBaseLocation()
 {
 	if (!HasValidData())
