Physics: Cleanup various collision routines, unifying common paths where possible and reducing/consolidating PhysX dependencies where possible
- Replaced many FCollsionShape switches with a FPhysXShapeAdaptor helper
- Various internal sweep/overlap/query helper methods (Geom*) now take Unreal types instead of being PhysX-specific
- Unused method UWorld::ComponentSweepSingle has been removed, use UWorld::ComponentSweepMulti instead if necessary
- Unused method UWorld::ComponentOverlapTest has been removed, use UWorld::ComponentOverlapMulti instead if necessary

diff --git a/Engine/Source/Runtime/Engine/Classes/Engine/World.h b/Engine/Source/Runtime/Engine/Classes/Engine/World.h
--- a/Engine/Source/Runtime/Engine/Classes/Engine/World.h
+++ b/Engine/Source/Runtime/Engine/Classes/Engine/World.h
@@ -1061,7 +1061,7 @@ public:
 	 * 	@param 	ResponseParam	ResponseContainer to be used for this trace	 	 
 	 *  @return TRUE if a blocking hit is found
 	 */
-	bool SweepTest(const FVector& Start, const FVector& End, const FQuat& Rot, ECollisionChannel TraceChannel, const struct FCollisionShape & CollisionShape, const struct FCollisionQueryParams& Params, const struct FCollisionResponseParams& ResponseParam = FCollisionResponseParams::DefaultResponseParam) const;
+	bool SweepTest(const FVector& Start, const FVector& End, const FQuat& Rot, ECollisionChannel TraceChannel, const struct FCollisionShape& CollisionShape, const struct FCollisionQueryParams& Params, const struct FCollisionResponseParams& ResponseParam = FCollisionResponseParams::DefaultResponseParam) const;
 
 	/**
 	 *  Sweep a sphere against the world and return if a blocking hit is found.
@@ -1072,7 +1072,7 @@ public:
 	 *	@param	ObjectQueryParams	List of object types it's looking for
 	 *  @return TRUE if any hit is found
 	 */
-	bool SweepTest(const FVector& Start, const FVector& End, const FQuat& Rot, const struct FCollisionShape & CollisionShape, const struct FCollisionQueryParams& Params, const struct FCollisionObjectQueryParams& ObjectQueryParams) const;
+	bool SweepTest(const FVector& Start, const FVector& End, const FQuat& Rot, const struct FCollisionShape& CollisionShape, const struct FCollisionQueryParams& Params, const struct FCollisionObjectQueryParams& ObjectQueryParams) const;
 
 	/**
 	 *  Sweep a sphere against the world and return the first blocking hit
@@ -1085,7 +1085,7 @@ public:
 	 * 	@param 	ResponseParam	ResponseContainer to be used for this trace	 
 	 *  @return TRUE if OutHits contains any blocking hit entries
 	 */
-	bool SweepSingle(struct FHitResult& OutHit,const FVector& Start,const FVector& End, const FQuat& Rot, ECollisionChannel TraceChannel,const struct FCollisionShape & CollisionShape, const struct FCollisionQueryParams& Params, const struct FCollisionResponseParams& ResponseParam = FCollisionResponseParams::DefaultResponseParam) const;
+	bool SweepSingle(struct FHitResult& OutHit, const FVector& Start, const FVector& End, const FQuat& Rot, ECollisionChannel TraceChannel, const struct FCollisionShape& CollisionShape, const struct FCollisionQueryParams& Params, const struct FCollisionResponseParams& ResponseParam = FCollisionResponseParams::DefaultResponseParam) const;
 
 	/**
 	 *  Sweep a sphere against the world and return the first blocking hit
@@ -1097,7 +1097,7 @@ public:
 	 *	@param	ObjectQueryParams	List of object types it's looking for
 	 *  @return TRUE if any hit is found
 	 */
-	bool SweepSingle(struct FHitResult& OutHit,const FVector& Start,const FVector& End, const FQuat& Rot, const struct FCollisionShape & CollisionShape, const struct FCollisionQueryParams& Params, const struct FCollisionObjectQueryParams& ObjectQueryParams) const;
+	bool SweepSingle(struct FHitResult& OutHit, const FVector& Start, const FVector& End, const FQuat& Rot, const struct FCollisionShape& CollisionShape, const struct FCollisionQueryParams& Params, const struct FCollisionObjectQueryParams& ObjectQueryParams) const;
 
 	/**
 	 *  Sweep a sphere against the world and return all initial overlaps (including blocking) if requested, then overlapping hits and then first blocking hit
@@ -1112,7 +1112,7 @@ public:
 	 * 	@param 	ResponseParam	ResponseContainer to be used for this trace	 
 	 *  @return TRUE if OutHits contains any blocking hit entries
 	 */
-	bool SweepMulti(TArray<struct FHitResult>& OutHits,const FVector& Start,const FVector& End, const FQuat& Rot, ECollisionChannel TraceChannel,const struct FCollisionShape & CollisionShape, const struct FCollisionQueryParams& Params, const struct FCollisionResponseParams& ResponseParam = FCollisionResponseParams::DefaultResponseParam) const;
+	bool SweepMulti(TArray<struct FHitResult>& OutHits, const FVector& Start, const FVector& End, const FQuat& Rot, ECollisionChannel TraceChannel, const struct FCollisionShape& CollisionShape, const struct FCollisionQueryParams& Params, const struct FCollisionResponseParams& ResponseParam = FCollisionResponseParams::DefaultResponseParam) const;
 
 	/**
 	 *  Sweep a sphere against the world and return all initial overlaps (including blocking) if requested, then overlapping hits and then first blocking hit
@@ -1126,7 +1126,7 @@ public:
 	 *	@param	ObjectQueryParams	List of object types it's looking for
 	 *  @return TRUE if any hit is found
 	 */
-	bool SweepMulti(TArray<struct FHitResult>& OutHits,const FVector& Start,const FVector& End, const FQuat& Rot, const struct FCollisionShape & CollisionShape, const struct FCollisionQueryParams& Params, const struct FCollisionObjectQueryParams& ObjectQueryParams) const;
+	bool SweepMulti(TArray<struct FHitResult>& OutHits, const FVector& Start, const FVector& End, const FQuat& Rot, const struct FCollisionShape& CollisionShape, const struct FCollisionQueryParams& Params, const struct FCollisionObjectQueryParams& ObjectQueryParams) const;
 
 	/**
 	 *  Test the collision of an AABB at the supplied location, and return if any blocking overlap is found
@@ -1137,7 +1137,7 @@ public:
 	 * 	@param 	ResponseParam	ResponseContainer to be used for this trace	 
 	 *  @return TRUE if any blocking results are found
 	 */
-	bool OverlapTest(const FVector& Pos, const FQuat& Rot, ECollisionChannel TraceChannel, const struct FCollisionShape & CollisionShape, const struct FCollisionQueryParams& Params, const struct FCollisionResponseParams& ResponseParam = FCollisionResponseParams::DefaultResponseParam) const;
+	bool OverlapTest(const FVector& Pos, const FQuat& Rot, ECollisionChannel TraceChannel, const struct FCollisionShape& CollisionShape, const struct FCollisionQueryParams& Params, const struct FCollisionResponseParams& ResponseParam = FCollisionResponseParams::DefaultResponseParam) const;
 
 	/**
 	 *  Test the collision of an AABB at the supplied location, and return if any blocking overlap is found
@@ -1147,7 +1147,7 @@ public:
 	 *	@param	ObjectQueryParams	List of object types it's looking for
 	 *  @return TRUE if any hit is found
 	 */
-	bool OverlapTest(const FVector& Pos, const FQuat& Rot, const struct FCollisionShape & CollisionShape, const struct FCollisionQueryParams& Params, const struct FCollisionObjectQueryParams& ObjectQueryParams) const;
+	bool OverlapTest(const FVector& Pos, const FQuat& Rot, const struct FCollisionShape& CollisionShape, const struct FCollisionQueryParams& Params, const struct FCollisionObjectQueryParams& ObjectQueryParams) const;
 
 	/**
 	 *  Test the collision of a sphere at the supplied location, and determine the set of components that it overlaps
@@ -1159,7 +1159,7 @@ public:
 	 * 	@param 	ResponseParam	ResponseContainer to be used for this trace
 	 *  @return TRUE if OutOverlaps contains any blocking results
 	 */
-	bool OverlapSingle(struct FOverlapResult& OutOverlap,const FVector& Pos, const FQuat& Rot,ECollisionChannel TraceChannel,const struct FCollisionShape & CollisionShape, const struct FCollisionQueryParams& Params, const struct FCollisionResponseParams& ResponseParam = FCollisionResponseParams::DefaultResponseParam) const;
+	bool OverlapSingle(struct FOverlapResult& OutOverlap, const FVector& Pos, const FQuat& Rot, ECollisionChannel TraceChannel, const struct FCollisionShape& CollisionShape, const struct FCollisionQueryParams& Params, const struct FCollisionResponseParams& ResponseParam = FCollisionResponseParams::DefaultResponseParam) const;
 
 	/**
 	 *  Test the collision of a sphere at the supplied location, and determine the set of components that it overlaps
@@ -1170,7 +1170,7 @@ public:
 	 *	@param	ObjectQueryParams	List of object types it's looking for
 	 *  @return TRUE if any hit is found
 	 */
-	bool OverlapSingle(struct FOverlapResult& OutOverlap,const FVector& Pos, const FQuat& Rot, const struct FCollisionShape & CollisionShape, const struct FCollisionQueryParams& Params, const struct FCollisionObjectQueryParams& ObjectQueryParams) const;
+	bool OverlapSingle(struct FOverlapResult& OutOverlap, const FVector& Pos, const FQuat& Rot, const struct FCollisionShape& CollisionShape, const struct FCollisionQueryParams& Params, const struct FCollisionObjectQueryParams& ObjectQueryParams) const;
 
 	/**
 	 *  Test the collision of a sphere at the supplied location, and determine the set of components that it overlaps
@@ -1182,7 +1182,7 @@ public:
 	 * 	@param 	ResponseParam	ResponseContainer to be used for this trace
 	 *  @return TRUE if OutOverlaps contains any blocking results
 	 */
-	bool OverlapMulti(TArray<struct FOverlapResult>& OutOverlaps,const FVector& Pos, const FQuat& Rot,ECollisionChannel TraceChannel,const struct FCollisionShape & CollisionShape, const struct FCollisionQueryParams& Params, const struct FCollisionResponseParams& ResponseParam = FCollisionResponseParams::DefaultResponseParam) const;
+	bool OverlapMulti(TArray<struct FOverlapResult>& OutOverlaps, const FVector& Pos, const FQuat& Rot, ECollisionChannel TraceChannel, const struct FCollisionShape& CollisionShape, const struct FCollisionQueryParams& Params, const struct FCollisionResponseParams& ResponseParam = FCollisionResponseParams::DefaultResponseParam) const;
 
 	/**
 	 *  Test the collision of a sphere at the supplied location, and determine the set of components that it overlaps
@@ -1193,20 +1193,9 @@ public:
 	 *	@param	ObjectQueryParams	List of object types it's looking for
 	 *  @return TRUE if any hit is found
 	 */
-	bool OverlapMulti(TArray<struct FOverlapResult>& OutOverlaps,const FVector& Pos, const FQuat& Rot, const struct FCollisionShape & CollisionShape, const struct FCollisionQueryParams& Params, const struct FCollisionObjectQueryParams& ObjectQueryParams) const;
+	bool OverlapMulti(TArray<struct FOverlapResult>& OutOverlaps, const FVector& Pos, const FQuat& Rot, const struct FCollisionShape& CollisionShape, const struct FCollisionQueryParams& Params, const struct FCollisionObjectQueryParams& ObjectQueryParams) const;
 	// COMPONENT SWEEP
 
-	/**
-	 *  Sweep this component against the world and return the first blocking hit
-	 *  @param  OutHit          First blocking hit found
-	 *	@param	PrimComp		Component to use geometry from to test against the world. Transform of this component is ignored
-	 *  @param  Start           Start location of the sphere
-	 *  @param  End             End location of the sphere
-	 *  @param  Params          Additional parameters used for the trace
-	 *  @return TRUE if OutHits contains any blocking hit entries
-	 */
-	bool ComponentSweepSingle(struct FHitResult& OutHit,class UPrimitiveComponent* PrimComp, const FVector& Start,const FVector& End, const FRotator& Rot, const struct FComponentQueryParams& Params) const;
-
 	/**
 	 *  Sweep the geometry of the supplied component, and determine the set of components that it hits
 	 *  @param  OutHits         Array of hits found between ray and the world
@@ -1217,17 +1206,9 @@ public:
 	 *  @param  Params          Additional parameters used for the trace
 	 *  @return TRUE if OutHits contains any blocking hit entries
 	 */
-	bool ComponentSweepMulti(TArray<struct FHitResult>& OutHits,class UPrimitiveComponent* PrimComp, const FVector& Start,const FVector& End, const FRotator& Rot, const struct FComponentQueryParams& Params)  const;
+	bool ComponentSweepMulti(TArray<struct FHitResult>& OutHits, class UPrimitiveComponent* PrimComp, const FVector& Start,const FVector& End, const FRotator& Rot, const struct FComponentQueryParams& Params)  const;
 
 	// COMPONENT OVERLAP
-	/**
-	 *  Test the collision of the supplied component at the supplied location/rotation, and determine the set of components that it overlaps
-	 *  @param  PrimComp        Component to use geometry from to test against the world. Transform of this component is ignored
-	 *  @param  Pos             Location to place PrimComp geometry at to test against the world
-	 *  @param  Rot             Rotation to place PrimComp geometry at to test against the world
-	 *  @return TRUE if OutOverlaps contains any blocking results
-	 */
-	bool ComponentOverlapTest(class UPrimitiveComponent* PrimComp, const FVector& Pos, const FRotator& Rot, const struct FComponentQueryParams& Params) const;
 
 	/**
 	 *  Test the collision of the supplied component at the supplied location/rotation, and determine the set of components that it overlaps
@@ -1316,7 +1297,7 @@ public:
 	 *	@param	UserData		UserData
 	 *	@param bMultiTrace		true if you'd like to get continuous result from the trace, false if you want single
 	 */ 
-	FTraceHandle	AsyncSweep(const FVector& Start,const FVector& End,ECollisionChannel TraceChannel, const struct FCollisionShape & CollisionShape, const struct FCollisionQueryParams& Params, const struct FCollisionResponseParams& ResponseParam = FCollisionResponseParams::DefaultResponseParam, FTraceDelegate * InDelegate=NULL, uint32 UserData = 0, bool bMultiTrace=false );
+	FTraceHandle	AsyncSweep(const FVector& Start, const FVector& End, ECollisionChannel TraceChannel, const struct FCollisionShape& CollisionShape, const struct FCollisionQueryParams& Params, const struct FCollisionResponseParams& ResponseParam = FCollisionResponseParams::DefaultResponseParam, FTraceDelegate * InDelegate = NULL, uint32 UserData = 0, bool bMultiTrace = false);
 
 	/**
 	 * Interface for Async trace
@@ -1339,7 +1320,7 @@ public:
 	 *	@param	UserData		UserData
 	 *	@param bMultiTrace		true if you'd like to get continuous result from the trace, false if you want single
 	 */ 
-	FTraceHandle	AsyncSweep(const FVector& Start,const FVector& End,const struct FCollisionShape & CollisionShape, const struct FCollisionQueryParams& Params, const struct FCollisionObjectQueryParams& ObjectQueryParams, FTraceDelegate * InDelegate=NULL, uint32 UserData = 0, bool bMultiTrace=false );
+	FTraceHandle	AsyncSweep(const FVector& Start, const FVector& End, const struct FCollisionShape& CollisionShape, const struct FCollisionQueryParams& Params, const struct FCollisionObjectQueryParams& ObjectQueryParams, FTraceDelegate * InDelegate = NULL, uint32 UserData = 0, bool bMultiTrace = false);
 	// overlap functions
 
 	/**
@@ -1364,7 +1345,7 @@ public:
 	 *	@param UserData			UserData
 	 *	@param bMultiTrace		true if you'd like to get continuous result from the trace, false if you want single
 	 */ 
-	FTraceHandle	AsyncOverlap(const FVector & Pos, const FQuat & Rot, ECollisionChannel TraceChannel, const struct FCollisionShape & CollisionShape, const struct FCollisionQueryParams& Params, const struct FCollisionResponseParams& ResponseParam = FCollisionResponseParams::DefaultResponseParam, FOverlapDelegate * InDelegate=NULL, uint32 UserData = 0, bool bMultiTrace=false);
+	FTraceHandle	AsyncOverlap(const FVector & Pos, const FQuat & Rot, ECollisionChannel TraceChannel, const struct FCollisionShape& CollisionShape, const struct FCollisionQueryParams& Params, const struct FCollisionResponseParams& ResponseParam = FCollisionResponseParams::DefaultResponseParam, FOverlapDelegate * InDelegate = NULL, uint32 UserData = 0, bool bMultiTrace = false);
 
 	/**
 	 * Interface for Async trace
@@ -1387,7 +1368,7 @@ public:
 	 *	@param UserData			UserData
 	 *	@param bMultiTrace		true if you'd like to get continuous result from the trace, false if you want single
 	 */ 
-	FTraceHandle	AsyncOverlap(const FVector & Pos, const FQuat & Rot, const struct FCollisionShape & CollisionShape, const struct FCollisionQueryParams& Params, const struct FCollisionObjectQueryParams& ObjectQueryParams, FOverlapDelegate * InDelegate=NULL, uint32 UserData = 0, bool bMultiTrace=false);
+	FTraceHandle	AsyncOverlap(const FVector & Pos, const FQuat & Rot, const struct FCollisionShape& CollisionShape, const struct FCollisionQueryParams& Params, const struct FCollisionObjectQueryParams& ObjectQueryParams, FOverlapDelegate * InDelegate = NULL, uint32 UserData = 0, bool bMultiTrace = false);
 
 	/**
 	 * Query function 
@@ -1395,7 +1376,7 @@ public:
 	 * return false if either expired or not yet evaluated or invalid
 	 * Use IsTraceHandleValid to find out if valid and to be evaluated
 	 */
-	bool QueryTraceData(const FTraceHandle & Handle, FTraceDatum & OutData);
+	bool QueryTraceData(const FTraceHandle& Handle, FTraceDatum& OutData);
 
 	/**
 	 * Query function 
@@ -1403,7 +1384,7 @@ public:
 	 * return false if either expired or not yet evaluated or invalid
 	 * Use IsTraceHandleValid to find out if valid and to be evaluated
 	 */
-	bool QueryOverlapData(const FTraceHandle & Handle, FOverlapDatum & OutData);
+	bool QueryOverlapData(const FTraceHandle& Handle, FOverlapDatum& OutData);
 	/** 
 	 * See if TraceHandle is still valid or not
 	 *
diff --git a/Engine/Source/Runtime/Engine/Classes/PhysicsEngine/BodyInstance.h b/Engine/Source/Runtime/Engine/Classes/PhysicsEngine/BodyInstance.h
--- a/Engine/Source/Runtime/Engine/Classes/PhysicsEngine/BodyInstance.h
+++ b/Engine/Source/Runtime/Engine/Classes/PhysicsEngine/BodyInstance.h
@@ -460,7 +460,7 @@ public:
 	 *  @param bReturnPhysicalMaterial	Fill in the PhysMaterial field of OutHit
 	 *  @return true if a hit is found
 	 */
-	bool LineTrace(struct FHitResult& OutHit, const FVector& Start, const FVector& End, bool bTraceComplex, bool bReturnPhysicalMaterial = false);
+	bool LineTrace(struct FHitResult& OutHit, const FVector& Start, const FVector& End, bool bTraceComplex, bool bReturnPhysicalMaterial = false) const;
 
 	/** 
 	 *  Trace a box against just this bodyinstance
@@ -471,7 +471,7 @@ public:
 	 *	@param	bTraceComplex	Should we trace against complex or simple collision of this body
 	 *  @return true if a hit is found
 	 */
-	bool Sweep(struct FHitResult& OutHit, const FVector& Start, const FVector& End, const FCollisionShape& Shape, bool bTraceComplex);
+	bool Sweep(struct FHitResult& OutHit, const FVector& Start, const FVector& End, const FCollisionShape& Shape, bool bTraceComplex) const;
 
 #if WITH_PHYSX
 	/**
@@ -481,9 +481,34 @@ public:
 	 *  @param  ShapePose       Transform information in world. Use U2PTransform to convert from FTransform
 	 *  @return true if PrimComp overlaps this component at the specified location/rotation
 	 */
-	bool Overlap(const physx::PxGeometry& Geom, const physx::PxTransform&  ShapePose);
+	bool OverlapPhysX(const physx::PxGeometry& Geom, const physx::PxTransform&  ShapePose) const;
 #endif	//WITH_PHYSX
 
+	/**
+	 *  Test if the bodyinstance overlaps with the specified shape at the specified position/rotation
+	 *
+	 *  @param  Position		Position to place the shape at before testing
+	 *  @param  Rotation		Rotation to apply to the shape before testing
+	 *	@param	CollisionShape	Shape to test against
+	 *  @return true if the geometry associated with this body instance overlaps the query shape at the specified location/rotation
+	 */
+	bool OverlapTest(const FVector& Position, const FQuat& Rotation, const struct FCollisionShape& CollisionShape) const;
+
+	/**
+	 *  Determines the set of components that this body instance would overlap with at the supplied location/rotation
+	 *  @param  InOutOverlaps   Array of overlaps found between this component in specified pose and the world (new overlaps will be appended, the array is not cleared!)
+	 *  @param  World			World to use for overlap test
+	 *  @param  pWorldToComponent Used to convert the body instance world space position into local space before teleporting it to (Pos, Rot) [optional, use when the body instance isn't centered on the component instance]
+	 *  @param  Pos             Location to place the component's geometry at to test against the world
+	 *  @param  Rot             Rotation to place components' geometry at to test against the world
+	 *  @param  TestChannel		The 'channel' that this ray is in, used to determine which components to hit
+	 *  @param  Params
+	 *  @param  ResponseParams
+	 *	@param	ObjectQueryParams	List of object types it's looking for. When this enters, we do object query with component shape
+	 *  @return TRUE if OutOverlaps contains any blocking results
+	 */
+	bool OverlapMulti(TArray<struct FOverlapResult>& InOutOverlaps, const class UWorld* World, const FTransform* pWorldToComponent, const FVector& Pos, const FRotator& Rot, ECollisionChannel TestChannel, const struct FComponentQueryParams& Params, const struct FCollisionResponseParams& ResponseParams, const struct FCollisionObjectQueryParams& ObjectQueryParams = FCollisionObjectQueryParams::DefaultObjectQueryParam) const;
+
 	/**
 	 * Add an impulse to this bodyinstance, radiating out from the specified position.
 	 *
@@ -541,7 +566,7 @@ private:
 	 *  @param	PGeom			Geometry that will sweep
 	 *  @return true if a hit is found
 	 */
-	bool InternalSweep(struct FHitResult& OutHit, const FVector& Start, const FVector& End, const FCollisionShape& Shape, bool bTraceComplex, const physx::PxRigidActor* RigidBody, const physx::PxGeometry* Geometry);
+	bool InternalSweepPhysX(struct FHitResult& OutHit, const FVector& Start, const FVector& End, const FCollisionShape& Shape, bool bTraceComplex, const physx::PxRigidActor* RigidBody, const physx::PxGeometry* Geometry) const;
 #endif 
 	/**
 	 * Invalidate Collision Profile Name
diff --git a/Engine/Source/Runtime/Engine/Private/Collision/PhysXCollision.cpp b/Engine/Source/Runtime/Engine/Private/Collision/PhysXCollision.cpp
--- a/Engine/Source/Runtime/Engine/Private/Collision/PhysXCollision.cpp
+++ b/Engine/Source/Runtime/Engine/Private/Collision/PhysXCollision.cpp
@@ -723,7 +723,7 @@ void CaptureGeomSweep(const UWorld* World, const FVector& Start, const FVector&
 	// Do a touch all query to find things we _didn't_ hit
 	bSkipCapture = true;
 	TArray<FHitResult> TouchAllResults;
-	GeomSweepMulti(World, PGeom, PRot, TouchAllResults, Start, End, DefaultCollisionChannel, Params, ResponseParams, FCollisionObjectQueryParams(FCollisionObjectQueryParams::InitType::AllObjects));
+	GeomSweepMulti_PhysX(World, PGeom, PRot, TouchAllResults, Start, End, DefaultCollisionChannel, Params, ResponseParams, FCollisionObjectQueryParams(FCollisionObjectQueryParams::InitType::AllObjects));
 	bSkipCapture = false;
 
 	// Now tell analyzer
@@ -1086,16 +1086,12 @@ bool RaycastMulti(const UWorld* World, TArray<struct FHitResult>& OutHits, const
 	return bHaveBlockingHit;
 }
 
-#endif //UE_WITH_PHYSICS
-
 //////////////////////////////////////////////////////////////////////////
 // GEOM SWEEP
 
-#if WITH_PHYSX
-
-bool GeomSweepTest(const UWorld* World, const PxGeometry& PGeom, const PxQuat& PGeomRot, FVector Start, FVector End, ECollisionChannel TraceChannel, const struct FCollisionQueryParams& Params, const struct FCollisionResponseParams& ResponseParams, const struct FCollisionObjectQueryParams& ObjectParams)
+bool GeomSweepTest(const UWorld* World, const struct FCollisionShape& CollisionShape, const FQuat& Rot, FVector Start, FVector End, ECollisionChannel TraceChannel, const struct FCollisionQueryParams& Params, const struct FCollisionResponseParams& ResponseParams, const struct FCollisionObjectQueryParams& ObjectParams)
 {
-	if(World == NULL || World->GetPhysicsScene() == NULL)
+	if ((World == NULL) || (World->GetPhysicsScene() == NULL))
 	{
 		return false;
 	}
@@ -1104,9 +1100,14 @@ bool GeomSweepTest(const UWorld* World, const PxGeometry& PGeom, const PxQuat& P
 
 	bool bHaveBlockingHit = false; // Track if we get any 'blocking' hits
 
-	FVector Delta = End - Start;
-	float DeltaMag = Delta.Size();
-	if(DeltaMag > KINDA_SMALL_NUMBER)
+#if WITH_PHYSX
+	FPhysXShapeAdaptor ShapeAdaptor(Rot, CollisionShape);
+	const PxGeometry& PGeom = ShapeAdaptor.GetGeometry();
+	const PxQuat& PGeomRot = ShapeAdaptor.GetGeomOrientation();
+
+	const FVector Delta = End - Start;
+	const float DeltaMag = Delta.Size();
+	if (DeltaMag > KINDA_SMALL_NUMBER)
 	{
 		// Create filter data used to filter collisions
 		PxFilterData PFilter = CreateQueryFilterData(TraceChannel, Params.bTraceComplex, ResponseParams.CollisionResponse, ObjectParams, false);
@@ -1148,11 +1149,14 @@ bool GeomSweepTest(const UWorld* World, const PxGeometry& PGeom, const PxQuat& P
 #endif //!(UE_BUILD_SHIPPING || UE_BUILD_TEST)
 
 	CAPTUREGEOMSWEEP(World, Start, End, PGeomRot, ECAQueryType::Test, PGeom, TraceChannel, Params, ResponseParams, ObjectParams, Hits);
+#endif // WITH_PHYSX
+
+	//@TODO: BOX2D: Implement GeomSweepTest
 
 	return bHaveBlockingHit;
 }
 
-bool GeomSweepSingle(const UWorld * World, const PxGeometry& PGeom, const PxQuat& PGeomRot, FHitResult& OutHit, FVector Start, FVector End, ECollisionChannel TraceChannel, const struct FCollisionQueryParams& Params, const struct FCollisionResponseParams& ResponseParams, const struct FCollisionObjectQueryParams& ObjectParams)
+bool GeomSweepSingle(const UWorld* World, const struct FCollisionShape& CollisionShape, const FQuat& Rot, FHitResult& OutHit, FVector Start, FVector End, ECollisionChannel TraceChannel, const struct FCollisionQueryParams& Params, const struct FCollisionResponseParams& ResponseParams, const struct FCollisionObjectQueryParams& ObjectParams)
 {
 	SCOPE_CYCLE_COUNTER(STAT_Collision_GeomSweepSingle);
 	STARTQUERYTIMER();
@@ -1160,7 +1164,7 @@ bool GeomSweepSingle(const UWorld * World, const PxGeometry& PGeom, const PxQuat
 	OutHit.TraceStart = Start;
 	OutHit.TraceEnd = End;
 
-	if(World == NULL || World->GetPhysicsScene() == NULL)
+	if ((World == NULL) || (World->GetPhysicsScene() == NULL))
 	{
 		return false;
 	}
@@ -1168,9 +1172,14 @@ bool GeomSweepSingle(const UWorld * World, const PxGeometry& PGeom, const PxQuat
 	// Track if we get any 'blocking' hits
 	bool bHaveBlockingHit = false;
 
+#if WITH_PHYSX
+	FPhysXShapeAdaptor ShapeAdaptor(Rot, CollisionShape);
+	const PxGeometry& PGeom = ShapeAdaptor.GetGeometry();
+	const PxQuat& PGeomRot = ShapeAdaptor.GetGeomOrientation();
+
 	const FVector Delta = End - Start;
 	const float DeltaMag = Delta.Size();
-	if(DeltaMag > KINDA_SMALL_NUMBER)
+	if (DeltaMag > KINDA_SMALL_NUMBER)
 	{
 		// Create filter data used to filter collisions
 		PxFilterData PFilter = CreateQueryFilterData(TraceChannel, Params.bTraceComplex, ResponseParams.CollisionResponse, ObjectParams, false);
@@ -1255,35 +1264,28 @@ bool GeomSweepSingle(const UWorld * World, const PxGeometry& PGeom, const PxQuat
 
 	CAPTUREGEOMSWEEP(World, Start, End, PGeomRot, ECAQueryType::Single, PGeom, TraceChannel, Params, ResponseParams, ObjectParams, Hits);
 #endif
+#endif // WITH_PHYSX
+
+	//@TODO: BOX2D: Implement GeomSweepSingle
 
 	return bHaveBlockingHit;
 }
 
-
-bool GeomSweepMulti(const UWorld* World, const PxGeometry& PGeom, const PxQuat& PGeomRot, TArray<FHitResult>& OutHits, FVector Start, FVector End, ECollisionChannel TraceChannel, const struct FCollisionQueryParams& Params, const struct FCollisionResponseParams& ResponseParams, const struct FCollisionObjectQueryParams& ObjectParams)
+bool GeomSweepMulti_PhysX(const UWorld* World, const PxGeometry& PGeom, const PxQuat& PGeomRot, TArray<FHitResult>& OutHits, FVector Start, FVector End, ECollisionChannel TraceChannel, const struct FCollisionQueryParams& Params, const struct FCollisionResponseParams& ResponseParams, const struct FCollisionObjectQueryParams& ObjectParams)
 {
-	if(World == NULL || World->GetPhysicsScene() == NULL)
-	{
-		return false;
-	}
-	SCOPE_CYCLE_COUNTER(STAT_Collision_GeomSweepMultiple);
 	STARTQUERYTIMER();
-
-	OutHits.Reset();
-
-	// Track if we get any 'blocking' hits
 	bool bBlockingHit = false;
 
 	// Create filter data used to filter collisions
 	PxFilterData PFilter = CreateQueryFilterData(TraceChannel, Params.bTraceComplex, ResponseParams.CollisionResponse, ObjectParams, true);
-	PxSceneQueryFilterData PQueryFilterData(PFilter, PxSceneQueryFilterFlag::eSTATIC | PxSceneQueryFilterFlag::eDYNAMIC | PxSceneQueryFilterFlag::ePREFILTER| PxSceneQueryFilterFlag::ePOSTFILTER);
+	PxSceneQueryFilterData PQueryFilterData(PFilter, PxSceneQueryFilterFlag::eSTATIC | PxSceneQueryFilterFlag::eDYNAMIC | PxSceneQueryFilterFlag::ePREFILTER | PxSceneQueryFilterFlag::ePOSTFILTER);
 	PxSceneQueryFlags POutputFlags = PxSceneQueryFlag::ePOSITION | PxSceneQueryFlag::eNORMAL | PxSceneQueryFlag::eDISTANCE;
 	FPxQueryFilterCallbackSweep PQueryCallbackSweep(Params.IgnoreActors);
 	PQueryCallbackSweep.DiscardInitialOverlaps = !Params.bFindInitialOverlaps;
 
 	const FVector Delta = End - Start;
 	const float DeltaMag = Delta.Size();
-	if(DeltaMag > KINDA_SMALL_NUMBER)
+	if (DeltaMag > KINDA_SMALL_NUMBER)
 	{
 		FPhysScene* PhysScene = World->GetPhysicsScene();
 		PxScene* SyncScene = PhysScene->GetPhysXScene(PST_Sync);
@@ -1291,7 +1293,7 @@ bool GeomSweepMulti(const UWorld* World, const PxGeometry& PGeom, const PxQuat&
 		SCOPED_SCENE_READ_LOCK(SyncScene);
 
 		const PxTransform PStartTM(U2PVector(Start), PGeomRot);
-		const PxVec3 PDir = U2PVector(Delta/DeltaMag);
+		const PxVec3 PDir = U2PVector(Delta / DeltaMag);
 
 		// Keep track of closest blocking hit distance.
 		float MinBlockDistance = DeltaMag;
@@ -1302,7 +1304,7 @@ bool GeomSweepMulti(const UWorld* World, const PxGeometry& PGeom, const PxQuat&
 
 		bool bBlockingHitSync;
 		PxI32 NumHits = SyncScene->sweepMultiple(PGeom, PStartTM, PDir, DeltaMag, POutputFlags, PHits, HIT_BUFFER_MAX_SYNC_QUERIES, bBlockingHitSync, PQueryFilterData, &PQueryCallbackSweep);
-		if(NumHits == -1)
+		if (NumHits == -1)
 		{
 			UE_LOG(LogCollision, Warning, TEXT("GeomSweepMulti : Buffer overflow from synchronous scene query"));
 			UE_LOG(LogCollision, Warning, TEXT("--------TraceChannel : %d"), (int32)TraceChannel);
@@ -1311,29 +1313,29 @@ bool GeomSweepMulti(const UWorld* World, const PxGeometry& PGeom, const PxQuat&
 			NumHits = HIT_BUFFER_MAX_SYNC_QUERIES;
 		}
 
-		if(bBlockingHitSync)
+		if (bBlockingHitSync)
 		{
 			/** If we got a blocking hit, it will be the last in the results - get the blocking distance */
-			MinBlockDistance = PHits[NumHits-1].distance;
+			MinBlockDistance = PHits[NumHits - 1].distance;
 			bBlockingHit = true;
 		}
 
 #if VERIFY_GEOMSWEEPMULTI
 		// If we hit something, test we can sweep back...
 		{
-			if(OutHits.Num() == 0)
+			if (OutHits.Num() == 0)
 			{
 				FVector HitLocation = End;
 				float TestLength = DeltaMag;
-				if(bBlockingHitSync)
+				if (bBlockingHitSync)
 				{
-					TestLength = PHits[NumHits-1].distance;
+					TestLength = PHits[NumHits - 1].distance;
 					FHitResult TestHit;
-					ConvertQueryImpactHit(PHits[NumHits-1], TestHit, DeltaMag, PFilter, Start, End, &PGeom, PStartTM, false, Params.bReturnPhysicalMaterial);
+					ConvertQueryImpactHit(PHits[NumHits - 1], TestHit, DeltaMag, PFilter, Start, End, &PGeom, PStartTM, false, Params.bReturnPhysicalMaterial);
 					HitLocation = TestHit.Location;
 				}
 
-				UE_LOG(LogCollision, Warning, TEXT("%d Start: %s End: %s"), bBlockingHitSync, *Start.ToString(), *HitLocation.ToString() );
+				UE_LOG(LogCollision, Warning, TEXT("%d Start: %s End: %s"), bBlockingHitSync, *Start.ToString(), *HitLocation.ToString());
 
 				PxTransform PEndTM(U2PVector(HitLocation), PGeomRot);
 
@@ -1341,7 +1343,7 @@ bool GeomSweepMulti(const UWorld* World, const PxGeometry& PGeom, const PxQuat&
 				bool bTestBlockingHit;
 				PxI32 NumTestHits = SyncScene->sweepMultiple(PGeom, PEndTM, -PDir, TestLength, POutputFlags, PTestHits, HIT_BUFFER_MAX_SYNC_QUERIES, bTestBlockingHit, PQueryFilterData, &PQueryCallbackSweep);
 
-				if(bTestBlockingHit)
+				if (bTestBlockingHit)
 				{
 					UE_LOG(LogCollision, Warning, TEXT("GeomSweepMulti : Cannot sweep back from end location."));
 
@@ -1358,7 +1360,7 @@ bool GeomSweepMulti(const UWorld* World, const PxGeometry& PGeom, const PxQuat&
 #endif
 
 		// Test async scene if async tests are requested and there was no overflow
-		if( Params.bTraceAsyncScene && MinBlockDistance > SMALL_NUMBER && PhysScene->HasAsyncScene())
+		if (Params.bTraceAsyncScene && MinBlockDistance > SMALL_NUMBER && PhysScene->HasAsyncScene())
 		{
 			PxScene* AsyncScene = PhysScene->GetPhysXScene(PST_Async);
 			SCOPED_SCENE_READ_LOCK(AsyncScene);
@@ -1369,7 +1371,7 @@ bool GeomSweepMulti(const UWorld* World, const PxGeometry& PGeom, const PxQuat&
 			check(AsyncBufferSize > 0);
 			bool bBlockingHitAsync = false;
 			PxI32 NumAsyncHits = AsyncScene->sweepMultiple(PGeom, PStartTM, PDir, MinBlockDistance, POutputFlags, PHitsAsync, AsyncBufferSize, bBlockingHitAsync, PQueryFilterData, &PQueryCallbackSweep);
-			if(NumAsyncHits == -1)
+			if (NumAsyncHits == -1)
 			{
 				UE_LOG(LogCollision, Log, TEXT("GeomSweepMulti : Buffer overflow from asynchronous scene query"));
 				UE_LOG(LogCollision, Warning, TEXT("--------TraceChannel : %d"), (int32)TraceChannel);
@@ -1380,22 +1382,22 @@ bool GeomSweepMulti(const UWorld* World, const PxGeometry& PGeom, const PxQuat&
 
 			NumHits += NumAsyncHits;
 
-			if(bBlockingHitAsync)
+			if (bBlockingHitAsync)
 			{
-				MinBlockDistance = FMath::Min<float>(PHits[NumHits-1].distance, MinBlockDistance);
+				MinBlockDistance = FMath::Min<float>(PHits[NumHits - 1].distance, MinBlockDistance);
 				bBlockingHit = true;
 			}
 		}
 
 		// Convert all hits to unreal structs. This will remove any hits further than MinBlockDistance, and sort results.
-		if(NumHits > 0)
+		if (NumHits > 0)
 		{
 			bBlockingHit |= AddSweepResults(NumHits, PHits, DeltaMag, PFilter, OutHits, Start, End, PGeom, PStartTM, MinBlockDistance, Params.bReturnPhysicalMaterial);
 		}
 	}
 
 #if !(UE_BUILD_SHIPPING || UE_BUILD_TEST)
-	if((World->DebugDrawTraceTag != NAME_None) && (World->DebugDrawTraceTag == Params.TraceTag))
+	if ((World->DebugDrawTraceTag != NAME_None) && (World->DebugDrawTraceTag == Params.TraceTag))
 	{
 		DrawGeomSweeps(World, Start, End, PGeom, PGeomRot, OutHits, DebugLineLifetime);
 	}
@@ -1406,12 +1408,39 @@ bool GeomSweepMulti(const UWorld* World, const PxGeometry& PGeom, const PxQuat&
 	return bBlockingHit;
 }
 
+
+bool GeomSweepMulti(const UWorld* World, const struct FCollisionShape& CollisionShape, const FQuat& Rot, TArray<FHitResult>& OutHits, FVector Start, FVector End, ECollisionChannel TraceChannel, const struct FCollisionQueryParams& Params, const struct FCollisionResponseParams& ResponseParams, const struct FCollisionObjectQueryParams& ObjectParams)
+{
+	if ((World == NULL) || (World->GetPhysicsScene() == NULL))
+	{
+		return false;
+	}
+	SCOPE_CYCLE_COUNTER(STAT_Collision_GeomSweepMultiple);
+
+	OutHits.Reset();
+
+	// Track if we get any 'blocking' hits
+	bool bBlockingHit = false;
+
+#if WITH_PHYSX
+	FPhysXShapeAdaptor ShapeAdaptor(Rot, CollisionShape);
+	const PxGeometry& PGeom = ShapeAdaptor.GetGeometry();
+	const PxQuat& PGeomRot = ShapeAdaptor.GetGeomOrientation();
+
+	bBlockingHit = GeomSweepMulti_PhysX(World, PGeom, PGeomRot, OutHits, Start, End, TraceChannel, Params, ResponseParams, ObjectParams);
+#endif // WITH_PHYSX
+
+	//@TODO: BOX2D: Implement GeomSweepMulti
+
+	return bBlockingHit;
+}
+
 //////////////////////////////////////////////////////////////////////////
 // GEOM OVERLAP
 
-bool GeomOverlapTest(const UWorld* World, const PxGeometry& PGeom, const PxTransform& PGeomPose, ECollisionChannel TraceChannel, const struct FCollisionQueryParams& Params, const struct FCollisionResponseParams& ResponseParams, const struct FCollisionObjectQueryParams& ObjectParams)
+bool GeomOverlapTest(const UWorld* World, const struct FCollisionShape& CollisionShape, const FVector& Pos, const FQuat& Rot, ECollisionChannel TraceChannel, const struct FCollisionQueryParams& Params, const struct FCollisionResponseParams& ResponseParams, const struct FCollisionObjectQueryParams& ObjectParams)
 {
-	if(World == NULL || World->GetPhysicsScene() == NULL)
+	if ((World == NULL) || (World->GetPhysicsScene() == NULL))
 	{
 		return false;
 	}
@@ -1420,6 +1449,11 @@ bool GeomOverlapTest(const UWorld* World, const PxGeometry& PGeom, const PxTrans
 	// Track if we get any 'blocking' hits
 	bool bHaveBlockingHit = false;
 
+#if WITH_PHYSX
+	FPhysXShapeAdaptor ShapeAdaptor(Rot, CollisionShape);
+	const PxGeometry& PGeom = ShapeAdaptor.GetGeometry();
+	const PxTransform& PGeomPose = ShapeAdaptor.GetGeomPose(Pos);
+
 	// Create filter data used to filter collisions
 	PxFilterData PFilter = CreateQueryFilterData(TraceChannel, Params.bTraceComplex, ResponseParams.CollisionResponse, ObjectParams, false);
 	PxSceneQueryFilterData PQueryFilterData(PFilter, PxSceneQueryFilterFlag::eSTATIC | PxSceneQueryFilterFlag::eDYNAMIC | PxSceneQueryFilterFlag::ePREFILTER);
@@ -1471,20 +1505,29 @@ bool GeomOverlapTest(const UWorld* World, const PxGeometry& PGeom, const PxTrans
 		DrawGeomOverlaps(World, PGeom, PGeomPose, Overlaps, DebugLineLifetime);
 	}
 #endif //!(UE_BUILD_SHIPPING || UE_BUILD_TEST)
+#endif // WITH_PHYSX
+
+	//@TODO: BOX2D: Implement GeomOverlapTest
 
 	return bHaveBlockingHit;
 }
 
-bool GeomOverlapSingle(const UWorld* World, const PxGeometry& PGeom, const PxTransform& PGeomPose, FOverlapResult& OutOverlap, ECollisionChannel TraceChannel, const struct FCollisionQueryParams& Params, const struct FCollisionResponseParams& ResponseParams, const struct FCollisionObjectQueryParams& ObjectParams)
+bool GeomOverlapSingle(const UWorld* World, const struct FCollisionShape& CollisionShape, const FVector& Pos, const FQuat& Rot, FOverlapResult& OutOverlap, ECollisionChannel TraceChannel, const struct FCollisionQueryParams& Params, const struct FCollisionResponseParams& ResponseParams, const struct FCollisionObjectQueryParams& ObjectParams)
 {
-	if(World == NULL || World->GetPhysicsScene() == NULL)
+	if ((World == NULL) || (World->GetPhysicsScene() == NULL))
 	{
 		return false;
 	}
 	SCOPE_CYCLE_COUNTER(STAT_Collision_GeomOverlapSingle);
+
 	// Track if we get any 'blocking' hits
 	bool bHaveBlockingHit = false;
 
+#if WITH_PHYSX
+	FPhysXShapeAdaptor ShapeAdaptor(Rot, CollisionShape);
+	const PxGeometry& PGeom = ShapeAdaptor.GetGeometry();
+	const PxTransform& PGeomPose = ShapeAdaptor.GetGeomPose(Pos);
+
 	// Create filter data used to filter collisions
 	PxFilterData PFilter = CreateQueryFilterData(TraceChannel, Params.bTraceComplex, ResponseParams.CollisionResponse, ObjectParams, false);
 	PxSceneQueryFilterData PQueryFilterData(PFilter, PxSceneQueryFilterFlag::eSTATIC | PxSceneQueryFilterFlag::eDYNAMIC | PxSceneQueryFilterFlag::ePREFILTER);
@@ -1501,14 +1544,14 @@ bool GeomOverlapSingle(const UWorld* World, const PxGeometry& PGeom, const PxTra
 
 	// if we got a hit, need to see if it was a touch or a block
 	// @todo UE4 james remove this once overlapAny only returns blocking hits
-	if(bHit)
+	if (bHit)
 	{
 		ConvertQueryOverlap( POverlapResult.shape, POverlapResult.actor, OutOverlap, PFilter);
 		bHaveBlockingHit = OutOverlap.bBlockingHit;
 	}
 
 	// Test async scene if async tests are requested and there was no blocking hit was found in the sync scene (since no hit info other than a boolean yes/no is recorded)
-	if( !bHaveBlockingHit && Params.bTraceAsyncScene && PhysScene->HasAsyncScene() )
+	if (!bHaveBlockingHit && Params.bTraceAsyncScene && PhysScene->HasAsyncScene())
 	{
 		PxScene* AsyncScene = PhysScene->GetPhysXScene(PST_Async);
 		SCOPED_SCENE_READ_LOCK(AsyncScene);
@@ -1516,7 +1559,7 @@ bool GeomOverlapSingle(const UWorld* World, const PxGeometry& PGeom, const PxTra
 		PxOverlapHit PAnotherOverlapResult;
 		bHit = AsyncScene->overlapAny(PGeom, PGeomPose, PAnotherOverlapResult, PQueryFilterData, &PQueryCallback);
 
-		if(bHit)
+		if (bHit)
 		{
 			FOverlapResult NewOverlap;
 			ConvertQueryOverlap(PAnotherOverlapResult.shape, PAnotherOverlapResult.actor, NewOverlap, PFilter);
@@ -1525,7 +1568,7 @@ bool GeomOverlapSingle(const UWorld* World, const PxGeometry& PGeom, const PxTra
 	}
 
 #if !(UE_BUILD_SHIPPING || UE_BUILD_TEST)
-	if((World->DebugDrawTraceTag != NAME_None) && (World->DebugDrawTraceTag == Params.TraceTag))
+	if ((World->DebugDrawTraceTag != NAME_None) && (World->DebugDrawTraceTag == Params.TraceTag))
 	{
 		TArray<FOverlapResult> Overlaps;
 		if (bHit)
@@ -1536,19 +1579,16 @@ bool GeomOverlapSingle(const UWorld* World, const PxGeometry& PGeom, const PxTra
 		DrawGeomOverlaps(World, PGeom, PGeomPose, Overlaps, DebugLineLifetime);
 	}
 #endif //!(UE_BUILD_SHIPPING || UE_BUILD_TEST)
+#endif // WITH_PHYSX
+
+	//@TODO: BOX2D: Implement GeomOverlapSingle
 
 	return bHaveBlockingHit;
 }
 
-bool GeomOverlapMulti(const UWorld* World, const PxGeometry& PGeom, const PxTransform& PGeomPose, TArray<FOverlapResult>& OutOverlaps, ECollisionChannel TraceChannel, const struct FCollisionQueryParams& Params, const struct FCollisionResponseParams& ResponseParams, const struct FCollisionObjectQueryParams& ObjectParams)
+#if WITH_PHYSX
+bool GeomOverlapMulti_PhysX(const UWorld* World, const PxGeometry& PGeom, const PxTransform& PGeomPose, TArray<FOverlapResult>& OutOverlaps, ECollisionChannel TraceChannel, const struct FCollisionQueryParams& Params, const struct FCollisionResponseParams& ResponseParams, const struct FCollisionObjectQueryParams& ObjectParams)
 {
-	if(World == NULL || World->GetPhysicsScene() == NULL)
-	{
-		return false;
-	}
-	SCOPE_CYCLE_COUNTER(STAT_Collision_GeomOverlapMultiple);
-
-	// Track if we get any 'blocking' hits
 	bool bHaveBlockingHit = false;
 
 	// overlapMultiple only supports sphere/capsule/box 
@@ -1571,7 +1611,7 @@ bool GeomOverlapMulti(const UWorld* World, const PxGeometry& PGeom, const PxTran
 
 		PxI32 NumHits = SyncScene->overlapMultiple(PGeom, PGeomPose, POverlapArray,  OVERLAP_BUFFER_SIZE_MAX_SYNC_QUERIES, PQueryFilterData, &PQueryCallback);
 		
-		if(NumHits == -1)
+		if (NumHits == -1)
 		{
 			UE_LOG(LogCollision, Warning, TEXT("GeomOverlapMulti : Buffer overflow from synchronous scene query"));
 			UE_LOG(LogCollision, Warning, TEXT("--------TraceChannel : %d"), (int32)TraceChannel);
@@ -1582,7 +1622,7 @@ bool GeomOverlapMulti(const UWorld* World, const PxGeometry& PGeom, const PxTran
 		}
 
 		// Test async scene if async tests are requested and there was no overflow
-		if( Params.bTraceAsyncScene && PhysScene->HasAsyncScene())
+		if (Params.bTraceAsyncScene && PhysScene->HasAsyncScene())
 		{		
 			PxScene* AsyncScene = PhysScene->GetPhysXScene(PST_Async);
 			SCOPED_SCENE_READ_LOCK(AsyncScene);
@@ -1593,7 +1633,7 @@ bool GeomOverlapMulti(const UWorld* World, const PxGeometry& PGeom, const PxTran
 			check(AsyncBufferSize > 0);
 			PxI32 NumAsyncHits = AsyncScene->overlapMultiple(PGeom, PGeomPose, PAsyncOverlapArray, AsyncBufferSize, PQueryFilterData, &PQueryCallback);
 
-			if(NumAsyncHits == -1)
+			if (NumAsyncHits == -1)
 			{
 				UE_LOG(LogCollision, Log, TEXT("GeomOverlapMulti : Buffer overflow from asynchronous scene query"));
 				UE_LOG(LogCollision, Warning, TEXT("--------TraceChannel : %d"), (int32)TraceChannel);
@@ -1606,13 +1646,13 @@ bool GeomOverlapMulti(const UWorld* World, const PxGeometry& PGeom, const PxTran
 			NumHits += NumAsyncHits;
 		}
 
-		if(NumHits > 0)
+		if (NumHits > 0)
 		{
 			bHaveBlockingHit = ConvertOverlapResults(NumHits, POverlapArray, PFilter, OutOverlaps);
 		}
 
 #if !(UE_BUILD_SHIPPING || UE_BUILD_TEST)
-		if((World->DebugDrawTraceTag != NAME_None) && (World->DebugDrawTraceTag == Params.TraceTag))
+		if ((World->DebugDrawTraceTag != NAME_None) && (World->DebugDrawTraceTag == Params.TraceTag))
 		{
 			DrawGeomOverlaps(World, PGeom, PGeomPose, OutOverlaps, DebugLineLifetime);
 		}
@@ -1625,9 +1665,38 @@ bool GeomOverlapMulti(const UWorld* World, const PxGeometry& PGeom, const PxTran
 
 	return bHaveBlockingHit;
 }
+#endif
+
+
+bool GeomOverlapMulti(const UWorld* World, const struct FCollisionShape& CollisionShape, const FVector& Pos, const FQuat& Rot, TArray<FOverlapResult>& OutOverlaps, ECollisionChannel TraceChannel, const struct FCollisionQueryParams& Params, const struct FCollisionResponseParams& ResponseParams, const struct FCollisionObjectQueryParams& ObjectParams)
+{
+	if ((World == NULL) || (World->GetPhysicsScene() == NULL))
+	{
+		return false;
+	}
+	SCOPE_CYCLE_COUNTER(STAT_Collision_GeomOverlapMultiple);
+
+	// Track if we get any 'blocking' hits
+	bool bHaveBlockingHit = false;
+
+#if WITH_PHYSX
+	FPhysXShapeAdaptor ShapeAdaptor(Rot, CollisionShape);
+	const PxGeometry& PGeom = ShapeAdaptor.GetGeometry();
+	const PxTransform& PGeomPose = ShapeAdaptor.GetGeomPose(Pos);
+	bHaveBlockingHit = GeomOverlapMulti_PhysX(World, PGeom, PGeomPose, OutOverlaps, TraceChannel, Params, ResponseParams, ObjectParams);
+
+#endif // WITH_PHYSX
+
+	//@TODO: BOX2D: Implement GeomOverlapMulti
+
+	return bHaveBlockingHit;
+}
+#endif //UE_WITH_PHYSICS
 
 //////////////////////////////////////////////////////////////////////////
 
+#if WITH_PHYSX
+
 static PxQuat CapsuleRotator(0.f, 0.707106781f, 0.f, 0.707106781f );
 
 PxQuat ConvertToPhysXCapsuleRot(const FQuat& GeomRot)
diff --git a/Engine/Source/Runtime/Engine/Private/Collision/PhysXCollision.h b/Engine/Source/Runtime/Engine/Private/Collision/PhysXCollision.h
--- a/Engine/Source/Runtime/Engine/Private/Collision/PhysXCollision.h
+++ b/Engine/Source/Runtime/Engine/Private/Collision/PhysXCollision.h
@@ -8,6 +8,7 @@
 
 #if WITH_PHYSX
 
+#include "Union.h"
 #include "../PhysicsEngine/PhysXSupport.h"
 
 /** Temporary result buffer size */
@@ -133,32 +134,6 @@ public:
 	virtual PxSceneQueryHitType::Enum postFilter(const PxFilterData& filterData, const PxSceneQueryHit& hit) OVERRIDE;
 };
 
-// GEOM OVERLAP
-
-/** Function for testing overlap between a supplied PxGeometry and the world. */
-bool GeomOverlapTest(const UWorld * World, const PxGeometry& PGeom, const PxTransform& PGeomPose, ECollisionChannel TraceChannel, const struct FCollisionQueryParams& Params, const struct FCollisionResponseParams& ResponseParams, const struct FCollisionObjectQueryParams& ObjectParams = FCollisionObjectQueryParams::DefaultObjectQueryParam);
-
-/** Function for testing overlap between a supplied PxGeometry and the world. */
-bool GeomOverlapSingle(const UWorld * World, const PxGeometry& PGeom, const PxTransform& PGeomPose, FOverlapResult& OutOverlaps, ECollisionChannel TraceChannel, const struct FCollisionQueryParams& Params, const struct FCollisionResponseParams& ResponseParams, const struct FCollisionObjectQueryParams& ObjectParams = FCollisionObjectQueryParams::DefaultObjectQueryParam);
-
-/** 
- * Function for overlapping a supplied PxGeometry against the world 
- * Note that this does not clear OutOverlaps, but adds unique results to it
- */
-bool GeomOverlapMulti(const UWorld * World, const PxGeometry& PGeom, const PxTransform& PGeomPose, TArray<FOverlapResult>& OutOverlaps, ECollisionChannel TraceChannel, const struct FCollisionQueryParams& Params, const struct FCollisionResponseParams& ResponseParams, const struct FCollisionObjectQueryParams& ObjectParams = FCollisionObjectQueryParams::DefaultObjectQueryParam);
-
-
-// GEOM SWEEP
-
-/** Function used for sweeping a supplied PxGeometry against the world as a test */
-bool GeomSweepTest(const UWorld * World, const PxGeometry& PGeom, const PxQuat& PGeomRot, FVector Start, FVector End, ECollisionChannel TraceChannel, const struct FCollisionQueryParams& Params, const struct FCollisionResponseParams& ResponseParams, const struct FCollisionObjectQueryParams& ObjectParams = FCollisionObjectQueryParams::DefaultObjectQueryParam);
-
-/** Function for sweeping a supplied PxGeometry against the world */
-bool GeomSweepSingle(const UWorld * World, const PxGeometry& PGeom, const PxQuat& PGeomRot, FHitResult& OutHit, FVector Start, FVector End, ECollisionChannel TraceChannel, const struct FCollisionQueryParams& Params, const struct FCollisionResponseParams& ResponseParams, const struct FCollisionObjectQueryParams& ObjectParams = FCollisionObjectQueryParams::DefaultObjectQueryParam);
-
-/** Function for sweeping a supplied PxGeometry against the world */
-bool GeomSweepMulti(const UWorld * World, const PxGeometry& PGeom, const PxQuat& PGeomRot, TArray<FHitResult>& OutHits, FVector Start, FVector End, ECollisionChannel TraceChannel, const struct FCollisionQueryParams& Params, const struct FCollisionResponseParams& ResponseParams, const struct FCollisionObjectQueryParams& ObjectParams = FCollisionObjectQueryParams::DefaultObjectQueryParam);
-
 // MISC
 
 /** Convert from unreal to physx capsule rotation */
@@ -193,4 +168,104 @@ bool RaycastSingle(const UWorld* World, struct FHitResult& OutHit, const FVector
  */
 bool RaycastMulti(const UWorld* World, TArray<struct FHitResult>& OutHits, const FVector& Start, const FVector& End, ECollisionChannel TraceChannel, const struct FCollisionQueryParams& Params, const struct FCollisionResponseParams& ResponseParams, const struct FCollisionObjectQueryParams& ObjectParams = FCollisionObjectQueryParams::DefaultObjectQueryParam);
 
+// GEOM OVERLAP
+
+/** Function for testing overlap between a supplied PxGeometry and the world. */
+bool GeomOverlapTest(const UWorld* World, const struct FCollisionShape& CollisionShape, const FVector& Pos, const FQuat& Rot, ECollisionChannel TraceChannel, const struct FCollisionQueryParams& Params, const struct FCollisionResponseParams& ResponseParams, const struct FCollisionObjectQueryParams& ObjectParams = FCollisionObjectQueryParams::DefaultObjectQueryParam);
+
+/** Function for testing overlap between a supplied PxGeometry and the world. */
+bool GeomOverlapSingle(const UWorld* World, const struct FCollisionShape& CollisionShape, const FVector& Pos, const FQuat& Rot, FOverlapResult& OutOverlaps, ECollisionChannel TraceChannel, const struct FCollisionQueryParams& Params, const struct FCollisionResponseParams& ResponseParams, const struct FCollisionObjectQueryParams& ObjectParams = FCollisionObjectQueryParams::DefaultObjectQueryParam);
+
+/**
+* Function for overlapping a supplied PxGeometry against the world
+* Note that this does not clear OutOverlaps, but adds unique results to it
+*/
+bool GeomOverlapMulti(const UWorld* World, const struct FCollisionShape& CollisionShape, const FVector& Pos, const FQuat& Rot, TArray<FOverlapResult>& OutOverlaps, ECollisionChannel TraceChannel, const struct FCollisionQueryParams& Params, const struct FCollisionResponseParams& ResponseParams, const struct FCollisionObjectQueryParams& ObjectParams = FCollisionObjectQueryParams::DefaultObjectQueryParam);
+
+// GEOM SWEEP
+
+/** Function used for sweeping a supplied PxGeometry against the world as a test */
+bool GeomSweepTest(const UWorld* World, const struct FCollisionShape& CollisionShape, const FQuat& Rot, FVector Start, FVector End, ECollisionChannel TraceChannel, const struct FCollisionQueryParams& Params, const struct FCollisionResponseParams& ResponseParams, const struct FCollisionObjectQueryParams& ObjectParams = FCollisionObjectQueryParams::DefaultObjectQueryParam);
+
+/** Function for sweeping a supplied PxGeometry against the world */
+bool GeomSweepSingle(const UWorld* World, const struct FCollisionShape& CollisionShape, const FQuat& Rot, FHitResult& OutHit, FVector Start, FVector End, ECollisionChannel TraceChannel, const struct FCollisionQueryParams& Params, const struct FCollisionResponseParams& ResponseParams, const struct FCollisionObjectQueryParams& ObjectParams = FCollisionObjectQueryParams::DefaultObjectQueryParam);
+
+/** Function for sweeping a supplied PxGeometry against the world */
+bool GeomSweepMulti(const UWorld* World, const struct FCollisionShape& CollisionShape, const FQuat& Rot, TArray<FHitResult>& OutHits, FVector Start, FVector End, ECollisionChannel TraceChannel, const struct FCollisionQueryParams& Params, const struct FCollisionResponseParams& ResponseParams, const struct FCollisionObjectQueryParams& ObjectParams = FCollisionObjectQueryParams::DefaultObjectQueryParam);
+
+#endif
+
+// Note: Do not use these methods for new code, they are being phased out!
+#if WITH_PHYSX
+bool GeomOverlapMulti_PhysX(const UWorld* World, const PxGeometry& PGeom, const PxTransform& PGeomPose, TArray<FOverlapResult>& OutOverlaps, ECollisionChannel TraceChannel, const struct FCollisionQueryParams& Params, const struct FCollisionResponseParams& ResponseParams, const struct FCollisionObjectQueryParams& ObjectParams);
+bool GeomSweepMulti_PhysX(const UWorld* World, const PxGeometry& PGeom, const PxQuat& PGeomRot, TArray<FHitResult>& OutHits, FVector Start, FVector End, ECollisionChannel TraceChannel, const struct FCollisionQueryParams& Params, const struct FCollisionResponseParams& ResponseParams, const struct FCollisionObjectQueryParams& ObjectParams = FCollisionObjectQueryParams::DefaultObjectQueryParam);
+#endif
+
+
+
+#if WITH_PHYSX
+
+// Adapts a FCollisionShape to a PxGeometry type, used for various queries
+struct FPhysXShapeAdaptor
+{
+public:
+	FPhysXShapeAdaptor(const FQuat& Rot, const FCollisionShape& CollisionShape)
+		: Rotation(PxIdentity)
+	{
+		// Perform other kinds of zero-extent queries as zero-extent sphere queries
+		if ((CollisionShape.ShapeType != ECollisionShape::Sphere) && CollisionShape.IsNearlyZero())
+		{
+			PtrToUnionData = UnionData.SetSubtype<PxSphereGeometry>(PxSphereGeometry(0.0f));
+		}
+		else
+		{
+			switch (CollisionShape.ShapeType)
+			{
+			case ECollisionShape::Box:
+				PtrToUnionData = UnionData.SetSubtype<PxBoxGeometry>(PxBoxGeometry(U2PVector(CollisionShape.GetBox())));
+				Rotation = U2PQuat(Rot);
+				break;
+			case ECollisionShape::Sphere:
+				PtrToUnionData = UnionData.SetSubtype<PxSphereGeometry>(PxSphereGeometry(CollisionShape.GetSphereRadius()));
+				break;
+			case ECollisionShape::Capsule:
+				PtrToUnionData = UnionData.SetSubtype<PxCapsuleGeometry>(PxCapsuleGeometry(CollisionShape.GetCapsuleRadius(), CollisionShape.GetCapsuleAxisHalfLength()));
+				Rotation = ConvertToPhysXCapsuleRot(Rot);
+				break;
+			default:
+				// invalid point
+				ensure(false);
+			}
+		}
+	}
+
+	PxGeometry& GetGeometry() const
+	{
+		return *PtrToUnionData;
+	}
+
+public:
+	PxTransform GetGeomPose(const FVector& Pos) const
+	{
+		return PxTransform(U2PVector(Pos), Rotation);
+	}
+
+	PxQuat GetGeomOrientation() const
+	{
+		return Rotation;
+	}
+
+private:
+	TUnion<PxSphereGeometry, PxBoxGeometry, PxCapsuleGeometry> UnionData;
+	
+	PxGeometry* PtrToUnionData;
+	PxQuat Rotation;
+};
+
+#endif
+
+
+#if WITH_BOX2D
+
+
 #endif
\ No newline at end of file
diff --git a/Engine/Source/Runtime/Engine/Private/Collision/WorldCollision.cpp b/Engine/Source/Runtime/Engine/Private/Collision/WorldCollision.cpp
--- a/Engine/Source/Runtime/Engine/Private/Collision/WorldCollision.cpp
+++ b/Engine/Source/Runtime/Engine/Private/Collision/WorldCollision.cpp
@@ -12,6 +12,12 @@
 	#include "../PhysicsEngine/PhysXSupport.h"
 #endif
 
+#if WITH_BOX2D
+	#include "../PhysicsEngine2D/Box2DIntegration.h"
+	#include "PhysicsEngine/BodySetup2D.h"
+	#include "PhysicsEngine/AggregateGeometry2D.h"
+#endif
+
 #include "PhysXCollision.h"
 #include "CollisionConversions.h"
 #include "CollisionDebugDrawing.h"
@@ -47,11 +53,11 @@ bool UWorld::LineTraceTest(const FVector& Start, const FVector& End, ECollisionC
 {
 #if UE_WITH_PHYSICS
 	return RaycastTest(this, Start, End, TraceChannel, Params, ResponseParam, FCollisionObjectQueryParams::DefaultObjectQueryParam);
-#endif
+#else
 	return false;
+#endif
 }
 
-
 bool UWorld::LineTraceSingle(struct FHitResult& OutHit, const FVector& Start, const FVector& End, ECollisionChannel TraceChannel, const struct FCollisionQueryParams& Params, const struct FCollisionResponseParams& ResponseParam) const
 {
 #if UE_WITH_PHYSICS
@@ -67,249 +73,96 @@ bool UWorld::LineTraceMulti(TArray<struct FHitResult>& OutHits, const FVector& S
 {
 #if UE_WITH_PHYSICS
 	return RaycastMulti(this, OutHits, Start, End, TraceChannel, Params, ResponseParam, FCollisionObjectQueryParams::DefaultObjectQueryParam);
-#endif // WITH_PHYSX
+#else
 	return false;
+#endif
 }
 
-bool UWorld::SweepTest(const FVector& Start, const FVector& End, const FQuat& Rot, ECollisionChannel TraceChannel, const struct FCollisionShape & CollisionShape, const struct FCollisionQueryParams& Params, const struct FCollisionResponseParams& ResponseParam) const
+bool UWorld::SweepTest(const FVector& Start, const FVector& End, const FQuat& Rot, ECollisionChannel TraceChannel, const struct FCollisionShape& CollisionShape, const struct FCollisionQueryParams& Params, const struct FCollisionResponseParams& ResponseParam) const
 {
-	if(CollisionShape.IsNearlyZero())
+	if (CollisionShape.IsNearlyZero())
 	{
 		// if extent is 0, we'll just do linetrace instead
 		return LineTraceTest(Start, End, TraceChannel, Params);
 	}
-
-#if WITH_PHYSX
-	switch (CollisionShape.ShapeType)
+	else
 	{
-	case ECollisionShape::Box:
-		{
-			const PxBoxGeometry PBoxGeom( U2PVector(CollisionShape.GetBox()) );
-			const PxQuat PGeomRot = U2PQuat(Rot);
-
-			return GeomSweepTest(this, PBoxGeom, PGeomRot, Start, End, TraceChannel, Params, ResponseParam, FCollisionObjectQueryParams::DefaultObjectQueryParam);
-		}
-	case ECollisionShape::Sphere:
-		{
-			PxSphereGeometry PSphereGeom( CollisionShape.GetSphereRadius());
-			PxQuat PGeomRot = PxQuat::createIdentity();
-
-			return GeomSweepTest(this, PSphereGeom, PGeomRot, Start, End, TraceChannel, Params, ResponseParam, FCollisionObjectQueryParams::DefaultObjectQueryParam);
-		}
-
-	case ECollisionShape::Capsule:
-		{
-			PxCapsuleGeometry PCapsuleGeom( CollisionShape.GetCapsuleRadius(), CollisionShape.GetCapsuleAxisHalfLength() );
-			const PxQuat PGeomRot = ConvertToPhysXCapsuleRot(Rot);
-			return GeomSweepTest(this, PCapsuleGeom, PGeomRot, Start, End, TraceChannel, Params, ResponseParam, FCollisionObjectQueryParams::DefaultObjectQueryParam);
-		}
-	default:
-		// invalid point
-		ensure(false);
+#if UE_WITH_PHYSICS
+		return GeomSweepTest(this, CollisionShape, Rot, Start, End, TraceChannel, Params, ResponseParam, FCollisionObjectQueryParams::DefaultObjectQueryParam);
+#else
+		return false;
+#endif
 	}
-
-#endif //WITH_PHYSX
-
-	//@TODO: BOX2D: Implement UWorld::SweepTest
-
-	return false;
 }
 
-bool UWorld::SweepSingle(struct FHitResult& OutHit, const FVector& Start, const FVector& End, const FQuat& Rot, ECollisionChannel TraceChannel, const struct FCollisionShape & CollisionShape, const struct FCollisionQueryParams& Params, const struct FCollisionResponseParams& ResponseParam) const
+bool UWorld::SweepSingle(struct FHitResult& OutHit, const FVector& Start, const FVector& End, const FQuat& Rot, ECollisionChannel TraceChannel, const struct FCollisionShape& CollisionShape, const struct FCollisionQueryParams& Params, const struct FCollisionResponseParams& ResponseParam) const
 {
-	if(CollisionShape.IsNearlyZero())
+	if (CollisionShape.IsNearlyZero())
 	{
 		return LineTraceSingle(OutHit, Start, End, TraceChannel, Params);
 	}
-#if WITH_PHYSX
-	switch (CollisionShape.ShapeType)
+	else
 	{
-	case ECollisionShape::Box:
-		{
-			const PxBoxGeometry PBoxGeom( U2PVector(CollisionShape.GetBox()) );
-			const PxQuat PGeomRot = U2PQuat(Rot);
-
-			return GeomSweepSingle(this, PBoxGeom, PGeomRot, OutHit, Start, End, TraceChannel, Params, ResponseParam, FCollisionObjectQueryParams::DefaultObjectQueryParam);
-		}
-	case ECollisionShape::Sphere:
-		{
-			PxSphereGeometry PSphereGeom( CollisionShape.GetSphereRadius());
-			PxQuat PGeomRot = PxQuat::createIdentity();
-
-			return GeomSweepSingle(this, PSphereGeom, PGeomRot, OutHit, Start, End, TraceChannel, Params, ResponseParam, FCollisionObjectQueryParams::DefaultObjectQueryParam);
-		}
-
-	case ECollisionShape::Capsule:
-		{
-			PxCapsuleGeometry PCapsuleGeom( CollisionShape.GetCapsuleRadius(), CollisionShape.GetCapsuleAxisHalfLength() );
-			const PxQuat PGeomRot = ConvertToPhysXCapsuleRot(Rot);
-			return GeomSweepSingle(this, PCapsuleGeom, PGeomRot, OutHit, Start, End, TraceChannel, Params, ResponseParam, FCollisionObjectQueryParams::DefaultObjectQueryParam);
-		}
-	default:
-		// invalid point
-		ensure(false);
+#if UE_WITH_PHYSICS
+		return GeomSweepSingle(this, CollisionShape, Rot, OutHit, Start, End, TraceChannel, Params, ResponseParam, FCollisionObjectQueryParams::DefaultObjectQueryParam);
+#else
+		OutHit.TraceStart = Start;
+		OutHit.TraceEnd = End;
+		return false;
+#endif
 	}
-
-#endif //WITH_PHYSX
-
-	//@TODO: BOX2D: Implement UWorld::SweepSingle
-
-	OutHit.TraceStart = Start;
-	OutHit.TraceEnd = End;
-	return false;
 }
 
-bool UWorld::SweepMulti( TArray<FHitResult>& OutHits, const FVector& Start, const FVector& End, const FQuat& Rot,ECollisionChannel TraceChannel, const struct FCollisionShape & CollisionShape, const struct FCollisionQueryParams& Params, const struct FCollisionResponseParams& ResponseParam ) const
+bool UWorld::SweepMulti(TArray<FHitResult>& OutHits, const FVector& Start, const FVector& End, const FQuat& Rot, ECollisionChannel TraceChannel, const struct FCollisionShape& CollisionShape, const struct FCollisionQueryParams& Params, const struct FCollisionResponseParams& ResponseParam) const
 {
-	if(CollisionShape.IsNearlyZero())
+	if (CollisionShape.IsNearlyZero())
 	{
 		return LineTraceMulti(OutHits, Start, End, TraceChannel, Params);
 	}
-#if WITH_PHYSX
-	switch (CollisionShape.ShapeType)
+	else
 	{
-	case ECollisionShape::Box:
-		{
-			const PxBoxGeometry PBoxGeom( U2PVector(CollisionShape.GetBox()) );
-			const PxQuat PGeomRot = U2PQuat(Rot);
-
-			return GeomSweepMulti(this, PBoxGeom, PGeomRot, OutHits, Start, End, TraceChannel, Params, ResponseParam, FCollisionObjectQueryParams::DefaultObjectQueryParam);
-		}
-	case ECollisionShape::Sphere:
-		{
-			PxSphereGeometry PSphereGeom( CollisionShape.GetSphereRadius());
-			PxQuat PGeomRot = PxQuat::createIdentity();
-
-			return GeomSweepMulti(this, PSphereGeom, PGeomRot, OutHits, Start, End, TraceChannel, Params, ResponseParam, FCollisionObjectQueryParams::DefaultObjectQueryParam);
-		}
-
-	case ECollisionShape::Capsule:
-		{
-			PxCapsuleGeometry PCapsuleGeom( CollisionShape.GetCapsuleRadius(), CollisionShape.GetCapsuleAxisHalfLength() );
-			const PxQuat PGeomRot = ConvertToPhysXCapsuleRot(Rot);
-			return GeomSweepMulti(this, PCapsuleGeom, PGeomRot, OutHits, Start, End, TraceChannel, Params, ResponseParam, FCollisionObjectQueryParams::DefaultObjectQueryParam);
-		}
-	default:
-		// invalid point
-		ensure(false);
+#if UE_WITH_PHYSICS
+		return GeomSweepMulti(this, CollisionShape, Rot, OutHits, Start, End, TraceChannel, Params, ResponseParam, FCollisionObjectQueryParams::DefaultObjectQueryParam);
+#else
+		return false;
+#endif
 	}
-#endif //WITH_PHYSX
-
-	//@TODO: BOX2D: Implement UWorld::SweepMulti
+}
 
+bool UWorld::OverlapTest(const FVector& Pos, const FQuat& Rot, ECollisionChannel TraceChannel, const struct FCollisionShape& CollisionShape, const struct FCollisionQueryParams& Params, const struct FCollisionResponseParams& ResponseParam) const
+{
+#if UE_WITH_PHYSICS
+	return GeomOverlapTest(this, CollisionShape, Pos, Rot, TraceChannel, Params, ResponseParam, FCollisionObjectQueryParams::DefaultObjectQueryParam);
+#else
 	return false;
+#endif
 }
 
-bool UWorld::OverlapTest(const FVector& Pos, const FQuat& Rot,ECollisionChannel TraceChannel, const struct FCollisionShape & CollisionShape, const struct FCollisionQueryParams& Params, const struct FCollisionResponseParams& ResponseParam ) const
+bool UWorld::OverlapSingle(struct FOverlapResult& OutOverlap, const FVector& Pos, const FQuat& Rot, ECollisionChannel TraceChannel, const struct FCollisionShape& CollisionShape, const struct FCollisionQueryParams& Params, const struct FCollisionResponseParams& ResponseParam) const
 {
-#if WITH_PHYSX
-	switch (CollisionShape.ShapeType)
-	{
-	case ECollisionShape::Box:
-		{
-			PxBoxGeometry PBoxGeom( U2PVector(CollisionShape.GetBox()) );
-			PxTransform PGeomPose = U2PTransform(FTransform(Rot, Pos));
-
-			return GeomOverlapTest(this, PBoxGeom, PGeomPose, TraceChannel, Params, ResponseParam, FCollisionObjectQueryParams::DefaultObjectQueryParam);
-		}
-	case ECollisionShape::Sphere:
-		{
-			PxSphereGeometry PSphereGeom( CollisionShape.GetSphereRadius() );
-			PxTransform PGeomPose(U2PVector(Pos), PxQuat::createIdentity());
-			return GeomOverlapTest(this, PSphereGeom, PGeomPose, TraceChannel, Params, ResponseParam, FCollisionObjectQueryParams::DefaultObjectQueryParam);
-		}
-
-	case ECollisionShape::Capsule:
-		{
-			PxCapsuleGeometry PCapsuleGeom( CollisionShape.GetCapsuleRadius(), CollisionShape.GetCapsuleAxisHalfLength() );
-			PxTransform PGeomPose = ConvertToPhysXCapsulePose( FTransform(Rot,Pos) );
-
-			return GeomOverlapTest(this, PCapsuleGeom, PGeomPose, TraceChannel, Params, ResponseParam, FCollisionObjectQueryParams::DefaultObjectQueryParam);
-		}
-	default:
-		// invalid point
-		ensure(false);
-	}
-
-#endif //WITH_PHYSX
-
-	//@TODO: BOX2D: Implement UWorld::OverlapTest
-
+#if UE_WITH_PHYSICS
+	return GeomOverlapSingle(this, CollisionShape, Pos, Rot, OutOverlap, TraceChannel, Params, ResponseParam, FCollisionObjectQueryParams::DefaultObjectQueryParam);
+#else
 	return false;
+#endif
 }
 
-bool UWorld::OverlapSingle(struct FOverlapResult& OutOverlap, const FVector& Pos, const FQuat& Rot, ECollisionChannel TraceChannel, const struct FCollisionShape & CollisionShape, const struct FCollisionQueryParams& Params, const struct FCollisionResponseParams& ResponseParam ) const
+bool UWorld::OverlapMulti(TArray<struct FOverlapResult>& OutOverlaps, const FVector& Pos, const FQuat& Rot, ECollisionChannel TraceChannel, const struct FCollisionShape& CollisionShape, const struct FCollisionQueryParams& Params, const struct FCollisionResponseParams& ResponseParam) const
 {
-#if WITH_PHYSX
-	switch (CollisionShape.ShapeType)
-	{
-	case ECollisionShape::Box:
-		{
-			PxBoxGeometry PBoxGeom( U2PVector(CollisionShape.GetBox()) );
-			PxTransform PGeomPose = U2PTransform(FTransform(Rot, Pos));
-
-			return GeomOverlapSingle(this, PBoxGeom, PGeomPose, OutOverlap, TraceChannel, Params, ResponseParam, FCollisionObjectQueryParams::DefaultObjectQueryParam);
-		}
-	case ECollisionShape::Sphere:
-		{
-			PxSphereGeometry PSphereGeom( CollisionShape.GetSphereRadius() );
-			PxTransform PGeomPose(U2PVector(Pos), PxQuat::createIdentity());
-			return GeomOverlapSingle(this, PSphereGeom, PGeomPose, OutOverlap, TraceChannel, Params, ResponseParam, FCollisionObjectQueryParams::DefaultObjectQueryParam);
-		}
-
-	case ECollisionShape::Capsule:
-		{
-			PxCapsuleGeometry PCapsuleGeom( CollisionShape.GetCapsuleRadius(), CollisionShape.GetCapsuleAxisHalfLength() );
-			PxTransform PGeomPose = ConvertToPhysXCapsulePose( FTransform(Rot,Pos) );
-
-			return GeomOverlapSingle(this, PCapsuleGeom, PGeomPose, OutOverlap, TraceChannel, Params, ResponseParam, FCollisionObjectQueryParams::DefaultObjectQueryParam);
-		}
-	default:
-		// invalid point
-		ensure(false);
-	}
-#endif //WITH_PHYSX
-
-	//@TODO: BOX2D: Implement UWorld::OverlapSingle
-
+#if UE_WITH_PHYSICS
+	return GeomOverlapMulti(this, CollisionShape, Pos, Rot, OutOverlaps, TraceChannel, Params, ResponseParam, FCollisionObjectQueryParams::DefaultObjectQueryParam);
+#else
 	return false;
+#endif
 }
 
-bool UWorld::OverlapMulti(TArray<struct FOverlapResult>& OutOverlaps, const FVector& Pos, const FQuat& Rot,ECollisionChannel TraceChannel, const struct FCollisionShape & CollisionShape, const struct FCollisionQueryParams& Params, const struct FCollisionResponseParams& ResponseParam ) const
+bool UWorld::OverlapMulti(TArray<struct FOverlapResult>& OutOverlaps, const FVector& Pos, const FQuat& Rot, const struct FCollisionShape& CollisionShape, const struct FCollisionQueryParams& Params, const struct FCollisionObjectQueryParams& ObjectQueryParams) const
 {
-#if WITH_PHYSX
-	switch (CollisionShape.ShapeType)
-	{
-	case ECollisionShape::Box:
-		{
-			PxBoxGeometry PBoxGeom( U2PVector(CollisionShape.GetBox()) );
-			PxTransform PGeomPose = U2PTransform(FTransform(Rot, Pos));
-
-			return GeomOverlapMulti(this, PBoxGeom, PGeomPose, OutOverlaps, TraceChannel, Params, ResponseParam, FCollisionObjectQueryParams::DefaultObjectQueryParam);
-		}
-	case ECollisionShape::Sphere:
-		{
-			PxSphereGeometry PSphereGeom( CollisionShape.GetSphereRadius() );
-			PxTransform PGeomPose(U2PVector(Pos), PxQuat::createIdentity());
-			return GeomOverlapMulti(this, PSphereGeom, PGeomPose, OutOverlaps, TraceChannel, Params, ResponseParam, FCollisionObjectQueryParams::DefaultObjectQueryParam);
-		}
-
-	case ECollisionShape::Capsule:
-		{
-			PxCapsuleGeometry PCapsuleGeom( CollisionShape.GetCapsuleRadius(), CollisionShape.GetCapsuleAxisHalfLength() );
-			PxTransform PGeomPose = ConvertToPhysXCapsulePose( FTransform(Rot,Pos) );
-
-			return GeomOverlapMulti(this, PCapsuleGeom, PGeomPose, OutOverlaps, TraceChannel, Params, ResponseParam, FCollisionObjectQueryParams::DefaultObjectQueryParam);
-		}
-	default:
-		// invalid point
-		ensure(false);
-	}
-#endif //WITH_PHYSX
-
-	//@TODO: BOX2D: Implement UWorld::OverlapMulti
-
+#if UE_WITH_PHYSICS
+	return GeomOverlapMulti(this, CollisionShape, Pos, Rot, OutOverlaps, DefaultCollisionChannel, Params, FCollisionResponseParams::DefaultResponseParam, ObjectQueryParams);
+#else
 	return false;
+#endif
 }
 
 // object query interfaces
@@ -318,8 +171,9 @@ bool UWorld::LineTraceTest(const FVector& Start,const FVector& End,const struct
 {
 #if UE_WITH_PHYSICS
 	return RaycastTest(this, Start, End, DefaultCollisionChannel, Params, FCollisionResponseParams::DefaultResponseParam, ObjectQueryParams);
-#endif
+#else
 	return false;
+#endif
 }
 
 bool UWorld::LineTraceSingle(struct FHitResult& OutHit,const FVector& Start,const FVector& End,const struct FCollisionQueryParams& Params, const struct FCollisionObjectQueryParams& ObjectQueryParams) const
@@ -335,321 +189,86 @@ bool UWorld::LineTraceSingle(struct FHitResult& OutHit,const FVector& Start,cons
 
 bool UWorld::LineTraceMulti(TArray<struct FHitResult>& OutHits, const FVector& Start, const FVector& End, const struct FCollisionQueryParams& Params, const struct FCollisionObjectQueryParams& ObjectQueryParams) const
 {
-	// object query returns true if any hit is found, not only blocking hit
 #if UE_WITH_PHYSICS
 	RaycastMulti(this, OutHits, Start, End, DefaultCollisionChannel, Params, FCollisionResponseParams::DefaultResponseParam, ObjectQueryParams);
-#endif
+
+	// object query returns true if any hit is found, not only blocking hit
 	return (OutHits.Num() > 0);
+#else
+	return false;
+#endif
 }
 
-bool UWorld::SweepTest(const FVector& Start, const FVector& End, const FQuat& Rot, const struct FCollisionShape & CollisionShape, const struct FCollisionQueryParams& Params, const struct FCollisionObjectQueryParams& ObjectQueryParams) const
+bool UWorld::SweepTest(const FVector& Start, const FVector& End, const FQuat& Rot, const struct FCollisionShape& CollisionShape, const struct FCollisionQueryParams& Params, const struct FCollisionObjectQueryParams& ObjectQueryParams) const
 {
-	if(CollisionShape.IsNearlyZero())
+	if (CollisionShape.IsNearlyZero())
 	{
 		// if extent is 0, we'll just do linetrace instead
 		return LineTraceTest(Start, End, Params, ObjectQueryParams);
 	}
-
-#if WITH_PHYSX
-	switch (CollisionShape.ShapeType)
+	else
 	{
-	case ECollisionShape::Box:
-		{
-			const PxBoxGeometry PBoxGeom( U2PVector(CollisionShape.GetBox()) );
-			const PxQuat PGeomRot = U2PQuat(Rot);
-
-			return GeomSweepTest(this, PBoxGeom, PGeomRot, Start, End, DefaultCollisionChannel, Params, FCollisionResponseParams::DefaultResponseParam, ObjectQueryParams);
-		}
-	case ECollisionShape::Sphere:
-		{
-			PxSphereGeometry PSphereGeom( CollisionShape.GetSphereRadius());
-			PxQuat PGeomRot = PxQuat::createIdentity();
-
-			return GeomSweepTest(this, PSphereGeom, PGeomRot, Start, End, DefaultCollisionChannel, Params, FCollisionResponseParams::DefaultResponseParam, ObjectQueryParams);
-		}
-
-	case ECollisionShape::Capsule:
-		{
-			PxCapsuleGeometry PCapsuleGeom( CollisionShape.GetCapsuleRadius(), CollisionShape.GetCapsuleAxisHalfLength() );
-			const PxQuat PGeomRot = ConvertToPhysXCapsuleRot(Rot);
-			return GeomSweepTest(this, PCapsuleGeom, PGeomRot, Start, End, DefaultCollisionChannel, Params, FCollisionResponseParams::DefaultResponseParam, ObjectQueryParams);
-		}
-	default:
-		// invalid point
-		ensure(false);
+#if UE_WITH_PHYSICS
+		return GeomSweepTest(this, CollisionShape, Rot, Start, End, DefaultCollisionChannel, Params, FCollisionResponseParams::DefaultResponseParam, ObjectQueryParams);
+#else
+		return false;
+#endif
 	}
-
-#endif //WITH_PHYSX
-
-	//@TODO: BOX2D: Implement UWorld::SweepTest
-
-	return false;
 }
 
-bool UWorld::SweepSingle(struct FHitResult& OutHit,const FVector& Start,const FVector& End, const FQuat& Rot, const struct FCollisionShape & CollisionShape, const struct FCollisionQueryParams& Params, const struct FCollisionObjectQueryParams& ObjectQueryParams) const
+bool UWorld::SweepSingle(struct FHitResult& OutHit, const FVector& Start, const FVector& End, const FQuat& Rot, const struct FCollisionShape& CollisionShape, const struct FCollisionQueryParams& Params, const struct FCollisionObjectQueryParams& ObjectQueryParams) const
 {
-	if(CollisionShape.IsNearlyZero())
+	if (CollisionShape.IsNearlyZero())
 	{
 		return LineTraceSingle(OutHit, Start, End, Params, ObjectQueryParams);
 	}
-#if WITH_PHYSX
-	switch (CollisionShape.ShapeType)
+	else
 	{
-	case ECollisionShape::Box:
-		{
-			const PxBoxGeometry PBoxGeom( U2PVector(CollisionShape.GetBox()) );
-			const PxQuat PGeomRot = U2PQuat(Rot);
-
-			return GeomSweepSingle(this, PBoxGeom, PGeomRot, OutHit, Start, End, DefaultCollisionChannel, Params, FCollisionResponseParams::DefaultResponseParam, ObjectQueryParams);
-		}
-	case ECollisionShape::Sphere:
-		{
-			PxSphereGeometry PSphereGeom( CollisionShape.GetSphereRadius());
-			PxQuat PGeomRot = PxQuat::createIdentity();
-
-			return GeomSweepSingle(this, PSphereGeom, PGeomRot, OutHit, Start, End, DefaultCollisionChannel, Params, FCollisionResponseParams::DefaultResponseParam, ObjectQueryParams);
-		}
-
-	case ECollisionShape::Capsule:
-		{
-			PxCapsuleGeometry PCapsuleGeom( CollisionShape.GetCapsuleRadius(), CollisionShape.GetCapsuleAxisHalfLength() );
-			const PxQuat PGeomRot = ConvertToPhysXCapsuleRot(Rot);
-			return GeomSweepSingle(this, PCapsuleGeom, PGeomRot, OutHit, Start, End, DefaultCollisionChannel, Params, FCollisionResponseParams::DefaultResponseParam, ObjectQueryParams);
-		}
-	default:
-		// invalid point
-		ensure(false);
+#if UE_WITH_PHYSICS
+		return GeomSweepSingle(this, CollisionShape, Rot, OutHit, Start, End, DefaultCollisionChannel, Params, FCollisionResponseParams::DefaultResponseParam, ObjectQueryParams);
+#else
+		OutHit.TraceStart = Start;
+		OutHit.TraceEnd = End;
+		return false;
+#endif
 	}
-
-#endif //WITH_PHYSX
-
-	//@TODO: BOX2D: Implement UWorld::SweepSingle
-
-	OutHit.TraceStart = Start;
-	OutHit.TraceEnd = End;
-	return false;
 }
 
-bool UWorld::SweepMulti(TArray<struct FHitResult>& OutHits,const FVector& Start,const FVector& End, const FQuat& Rot, const struct FCollisionShape & CollisionShape, const struct FCollisionQueryParams& Params, const struct FCollisionObjectQueryParams& ObjectQueryParams) const
+bool UWorld::SweepMulti(TArray<struct FHitResult>& OutHits, const FVector& Start, const FVector& End, const FQuat& Rot, const struct FCollisionShape& CollisionShape, const struct FCollisionQueryParams& Params, const struct FCollisionObjectQueryParams& ObjectQueryParams) const
 {
-	// object query returns true if any hit is found, not only blocking hit
-	if(CollisionShape.IsNearlyZero())
+	if (CollisionShape.IsNearlyZero())
 	{
-		LineTraceMulti(OutHits, Start, End, Params, ObjectQueryParams);
+		return LineTraceMulti(OutHits, Start, End, Params, ObjectQueryParams);
 	}
-#if WITH_PHYSX
 	else
 	{
-		switch (CollisionShape.ShapeType)
-		{
-		case ECollisionShape::Box:
-			{
-				const PxBoxGeometry PBoxGeom( U2PVector(CollisionShape.GetBox()) );
-				const PxQuat PGeomRot = U2PQuat(Rot);
-
-				GeomSweepMulti(this, PBoxGeom, PGeomRot, OutHits, Start, End, DefaultCollisionChannel, Params, FCollisionResponseParams::DefaultResponseParam, ObjectQueryParams);
-			}
-			break;
-		case ECollisionShape::Sphere:
-			{
-				PxSphereGeometry PSphereGeom( CollisionShape.GetSphereRadius());
-				PxQuat PGeomRot = PxQuat::createIdentity();
-
-				GeomSweepMulti(this, PSphereGeom, PGeomRot, OutHits, Start, End, DefaultCollisionChannel, Params, FCollisionResponseParams::DefaultResponseParam, ObjectQueryParams);
-			}
-			break;
-		case ECollisionShape::Capsule:
-			{
-				PxCapsuleGeometry PCapsuleGeom( CollisionShape.GetCapsuleRadius(), CollisionShape.GetCapsuleAxisHalfLength() );
-				const PxQuat PGeomRot = ConvertToPhysXCapsuleRot(Rot);
-				GeomSweepMulti(this, PCapsuleGeom, PGeomRot, OutHits, Start, End, DefaultCollisionChannel, Params, FCollisionResponseParams::DefaultResponseParam, ObjectQueryParams);
-			}
-			break;
-		default:
-			// invalid point
-			ensure(false);
-		}
-	}
-#endif //WITH_PHYSX
-
-	//@TODO: BOX2D: Implement UWorld::SweepMulti
-
-	return (OutHits.Num() > 0);
-}
-
-bool UWorld::OverlapTest(const FVector& Pos, const FQuat& Rot,const struct FCollisionShape & CollisionShape, const struct FCollisionQueryParams& Params, const struct FCollisionObjectQueryParams& ObjectQueryParams) const
-{
-#if WITH_PHYSX
-	switch (CollisionShape.ShapeType)
-	{
-	case ECollisionShape::Box:
-		{
-			PxBoxGeometry PBoxGeom( U2PVector(CollisionShape.GetBox()) );
-			PxTransform PGeomPose = U2PTransform(FTransform(Rot, Pos));
-
-			return GeomOverlapTest(this, PBoxGeom, PGeomPose, DefaultCollisionChannel, Params, FCollisionResponseParams::DefaultResponseParam, ObjectQueryParams);
-		}
-	case ECollisionShape::Sphere:
-		{
-			PxSphereGeometry PSphereGeom( CollisionShape.GetSphereRadius() );
-			PxTransform PGeomPose(U2PVector(Pos), PxQuat::createIdentity());
-			return GeomOverlapTest(this, PSphereGeom, PGeomPose, DefaultCollisionChannel, Params, FCollisionResponseParams::DefaultResponseParam, ObjectQueryParams);
-		}
-
-	case ECollisionShape::Capsule:
-		{
-			PxCapsuleGeometry PCapsuleGeom( CollisionShape.GetCapsuleRadius(), CollisionShape.GetCapsuleAxisHalfLength() );
-			PxTransform PGeomPose = ConvertToPhysXCapsulePose( FTransform(Rot,Pos) );
+#if UE_WITH_PHYSICS
+		GeomSweepMulti(this, CollisionShape, Rot, OutHits, Start, End, DefaultCollisionChannel, Params, FCollisionResponseParams::DefaultResponseParam, ObjectQueryParams);
 
-			return GeomOverlapTest(this, PCapsuleGeom, PGeomPose, DefaultCollisionChannel, Params, FCollisionResponseParams::DefaultResponseParam, ObjectQueryParams);
-		}
-	default:
-		// invalid point
-		ensure(false);
+		// object query returns true if any hit is found, not only blocking hit
+		return (OutHits.Num() > 0);
+#else
+		return false;
+#endif
 	}
-
-#endif //WITH_PHYSX
-
-	//@TODO: BOX2D: Implement UWorld::OverlapTest
-
-	return false;
 }
 
-bool UWorld::OverlapSingle(struct FOverlapResult& OutOverlap,const FVector& Pos, const FQuat& Rot, const struct FCollisionShape & CollisionShape, const struct FCollisionQueryParams& Params, const struct FCollisionObjectQueryParams& ObjectQueryParams) const
+bool UWorld::OverlapTest(const FVector& Pos, const FQuat& Rot, const struct FCollisionShape& CollisionShape, const struct FCollisionQueryParams& Params, const struct FCollisionObjectQueryParams& ObjectQueryParams) const
 {
-#if WITH_PHYSX
-	switch (CollisionShape.ShapeType)
-	{
-	case ECollisionShape::Box:
-		{
-			PxBoxGeometry PBoxGeom( U2PVector(CollisionShape.GetBox()) );
-			PxTransform PGeomPose = U2PTransform(FTransform(Rot, Pos));
-
-			return GeomOverlapSingle(this, PBoxGeom, PGeomPose, OutOverlap, DefaultCollisionChannel, Params, FCollisionResponseParams::DefaultResponseParam, ObjectQueryParams);
-		}
-	case ECollisionShape::Sphere:
-		{
-			PxSphereGeometry PSphereGeom( CollisionShape.GetSphereRadius() );
-			PxTransform PGeomPose(U2PVector(Pos), PxQuat::createIdentity());
-			return GeomOverlapSingle(this, PSphereGeom, PGeomPose, OutOverlap, DefaultCollisionChannel, Params, FCollisionResponseParams::DefaultResponseParam, ObjectQueryParams);
-		}
-
-	case ECollisionShape::Capsule:
-		{
-			PxCapsuleGeometry PCapsuleGeom( CollisionShape.GetCapsuleRadius(), CollisionShape.GetCapsuleAxisHalfLength() );
-			PxTransform PGeomPose = ConvertToPhysXCapsulePose( FTransform(Rot,Pos) );
-
-			return GeomOverlapSingle(this, PCapsuleGeom, PGeomPose, OutOverlap, DefaultCollisionChannel, Params, FCollisionResponseParams::DefaultResponseParam, ObjectQueryParams);
-		}
-	default:
-		// invalid point
-		ensure(false);
-	}
-#endif //WITH_PHYSX
-
-	//@TODO: BOX2D: Implement UWorld::OverlapSingle
-
+#if UE_WITH_PHYSICS
+	return GeomOverlapTest(this, CollisionShape, Pos, Rot, DefaultCollisionChannel, Params, FCollisionResponseParams::DefaultResponseParam, ObjectQueryParams);
+#else
 	return false;
+#endif
 }
 
-bool UWorld::OverlapMulti(TArray<struct FOverlapResult>& OutOverlaps,const FVector& Pos, const FQuat& Rot, const struct FCollisionShape & CollisionShape, const struct FCollisionQueryParams& Params, const struct FCollisionObjectQueryParams& ObjectQueryParams) const
-{
-	// object query returns true if any hit is found, not only blocking hit
-#if WITH_PHYSX
-	switch (CollisionShape.ShapeType)
-	{
-	case ECollisionShape::Box:
-		{
-			PxBoxGeometry PBoxGeom( U2PVector(CollisionShape.GetBox()) );
-			PxTransform PGeomPose = U2PTransform(FTransform(Rot, Pos));
-
-			GeomOverlapMulti(this, PBoxGeom, PGeomPose, OutOverlaps, DefaultCollisionChannel, Params, FCollisionResponseParams::DefaultResponseParam, ObjectQueryParams);
-		}
-		break;
-	case ECollisionShape::Sphere:
-		{
-			PxSphereGeometry PSphereGeom( CollisionShape.GetSphereRadius() );
-			PxTransform PGeomPose(U2PVector(Pos), PxQuat::createIdentity());
-			GeomOverlapMulti(this, PSphereGeom, PGeomPose, OutOverlaps, DefaultCollisionChannel, Params, FCollisionResponseParams::DefaultResponseParam, ObjectQueryParams);
-		}
-		break;
-	case ECollisionShape::Capsule:
-		{
-			PxCapsuleGeometry PCapsuleGeom( CollisionShape.GetCapsuleRadius(), CollisionShape.GetCapsuleAxisHalfLength() );
-			PxTransform PGeomPose = ConvertToPhysXCapsulePose( FTransform(Rot,Pos) );
-
-			GeomOverlapMulti(this, PCapsuleGeom, PGeomPose, OutOverlaps, DefaultCollisionChannel, Params, FCollisionResponseParams::DefaultResponseParam, ObjectQueryParams);
-		}
-		break;
-	default:
-		// invalid point
-		ensure(false);
-	}
-#endif //WITH_PHYSX
-
-	//@TODO: BOX2D: Implement UWorld::OverlapMulti
-
-	return (OutOverlaps.Num() > 0);
-}
-
-bool UWorld::ComponentOverlapTest(class UPrimitiveComponent* PrimComp, const FVector& Pos, const FRotator& Rot, const struct FComponentQueryParams& Params) const
+bool UWorld::OverlapSingle(struct FOverlapResult& OutOverlap, const FVector& Pos, const FQuat& Rot, const struct FCollisionShape& CollisionShape, const struct FCollisionQueryParams& Params, const struct FCollisionObjectQueryParams& ObjectQueryParams) const
 {
-	if(GetPhysicsScene() == NULL)
-	{
-		return false;
-	}
-
-	if(PrimComp == NULL)
-	{
-		UE_LOG(LogCollision, Log, TEXT("ComponentOverlapMulti : No PrimComp"));
-		return false;
-	}
-
-#if WITH_PHYSX
-	ECollisionChannel TraceChannel = PrimComp->GetCollisionObjectType();
-
-	PxRigidActor* PRigidActor = PrimComp->BodyInstance.GetPxRigidActor();
-	if(PRigidActor == NULL)
-	{
-		UE_LOG(LogCollision, Log, TEXT("ComponentOverlapMulti : (%s) No physics data"), *PrimComp->GetPathName());
-		return false;
-	}
-
-	// calculate the test global pose of the actor
-	PxTransform PTestGlobalPose = U2PTransform(FTransform(Rot, Pos));
-
-	// Get all the shapes from the actor
-	TArray<PxShape*, TInlineAllocator<8>> PShapes;
-	PShapes.AddZeroed(PRigidActor->getNbShapes());
-	int32 NumShapes = PRigidActor->getShapes(PShapes.GetData(), PShapes.Num());
-
-	// Iterate over each shape
-	for(int32 ShapeIdx=0; ShapeIdx<PShapes.Num(); ShapeIdx++)
-	{
-		PxShape* PShape = PShapes[ShapeIdx];
-		check(PShape);
-		// Calc shape global pose
-		PxTransform PLocalPose = PShape->getLocalPose();
-		PxTransform PShapeGlobalPose = PTestGlobalPose.transform(PLocalPose);
-
-		GET_GEOMETRY_FROM_SHAPE(PGeom, PShape);
-
-		if(PGeom != NULL)
-		{
-			if( GeomOverlapTest(this, *PGeom, PShapeGlobalPose, TraceChannel, Params, FCollisionResponseParams(PrimComp->GetCollisionResponseToChannels())))
-			{
-				// in this test, it only matters true or false. 
-				// if we found first true, we don't care next test anymore. 
-				return true;
-			}
-		}
-	}
-
-#endif //WITH_PHYSX
-
-	//@TODO: BOX2D: Implement UWorld::ComponentOverlapTest
-
+#if UE_WITH_PHYSICS
+	return GeomOverlapSingle(this, CollisionShape, Pos, Rot, OutOverlap, DefaultCollisionChannel, Params, FCollisionResponseParams::DefaultResponseParam, ObjectQueryParams);
+#else
 	return false;
+#endif
 }
 
 bool UWorld::ComponentOverlapMulti(TArray<struct FOverlapResult>& OutOverlaps, const class UPrimitiveComponent* PrimComp, const FVector& Pos, const FRotator& Rot, const struct FComponentQueryParams& Params, const struct FCollisionObjectQueryParams& ObjectQueryParams) const
@@ -658,113 +277,34 @@ bool UWorld::ComponentOverlapMulti(TArray<struct FOverlapResult>& OutOverlaps, c
 	{
 		return ComponentOverlapMulti(OutOverlaps, PrimComp, Pos, Rot, PrimComp->GetCollisionObjectType(), Params, ObjectQueryParams);
 	}
-
-	UE_LOG(LogCollision, Log, TEXT("ComponentOverlapMulti : No PrimComp"));
-	return false;
-}
-
-bool UWorld::ComponentOverlapMulti(TArray<struct FOverlapResult>& OutOverlaps, const class UPrimitiveComponent* PrimComp, const FVector& Pos, const FRotator& Rot, ECollisionChannel TestChannel, const struct FComponentQueryParams& Params, const struct FCollisionObjectQueryParams& ObjectQueryParams) const
-{
-	if(GetPhysicsScene() == NULL)
-	{
-		return false;
-	}
-
-	if(PrimComp == NULL)
+	else
 	{
 		UE_LOG(LogCollision, Log, TEXT("ComponentOverlapMulti : No PrimComp"));
 		return false;
 	}
-
-	return PrimComp->ComponentOverlapMulti(OutOverlaps, this, Pos, Rot, TestChannel, Params, ObjectQueryParams);
 }
 
-bool UWorld::ComponentSweepSingle(struct FHitResult& OutHit,class UPrimitiveComponent* PrimComp, const FVector& Start, const FVector& End, const FRotator& Rot, const struct FComponentQueryParams& Params) const
+bool UWorld::ComponentOverlapMulti(TArray<struct FOverlapResult>& OutOverlaps, const class UPrimitiveComponent* PrimComp, const FVector& Pos, const FRotator& Rot, ECollisionChannel TestChannel, const struct FComponentQueryParams& Params, const struct FCollisionObjectQueryParams& ObjectQueryParams) const
 {
-	OutHit.TraceStart = Start;
-	OutHit.TraceEnd = End;
-
-	if(GetPhysicsScene() == NULL)
-	{
-		return false;
-	}
-
-	if(PrimComp == NULL)
-	{
-		UE_LOG(LogCollision, Log, TEXT("ComponentSweepSingle : No PrimComp"));
-		return false;
-	}
-	
-	ECollisionChannel TraceChannel = PrimComp->GetCollisionObjectType();
-
-	// if extent is 0, do line trace
-	if (PrimComp->IsZeroExtent())
+	if (PrimComp)
 	{
-		return RaycastSingle(this, OutHit, Start, End, TraceChannel, Params, FCollisionResponseParams(PrimComp->GetCollisionResponseToChannels()));
+		return PrimComp->ComponentOverlapMulti(OutOverlaps, this, Pos, Rot, TestChannel, Params, ObjectQueryParams);
 	}
-
-	//@TODO: BOX2D: Implement UWorld::ComponentSweepSingle
-#if WITH_PHYSX
-	PxRigidActor* PRigidActor = PrimComp->BodyInstance.GetPxRigidActor();
-	if(PRigidActor == NULL)
+	else
 	{
-		UE_LOG(LogCollision, Log, TEXT("ComponentSweepMulti : (%s) No physics data"), *PrimComp->GetPathName());
+		UE_LOG(LogCollision, Log, TEXT("ComponentOverlapMulti : No PrimComp"));
 		return false;
 	}
-
-	// Get all the shapes from the actor
-	TArray<PxShape*, TInlineAllocator<8>> PShapes;
-	PShapes.AddZeroed(PRigidActor->getNbShapes());
-	int32 NumShapes = PRigidActor->getShapes(PShapes.GetData(), PShapes.Num());
-
-	// calculate the test global pose of the actor
-	PxTransform PGlobalStartPose = U2PTransform(FTransform(Start));
-	PxTransform PGlobalEndPose = U2PTransform(FTransform(End));
-
-	bool bHaveBlockingHit = false;
-	PxQuat PGeomRot = U2PQuat(Rot.Quaternion());
-
-	// Iterate over each shape
-	for(int32 ShapeIdx=0; ShapeIdx<PShapes.Num(); ShapeIdx++)
-	{
-		PxShape* PShape = PShapes[ShapeIdx];
-		check(PShape);
-
-		// Calc shape global pose
-		PxTransform PLocalShape = PShape->getLocalPose();
-		PxTransform PShapeGlobalStartPose = PGlobalStartPose.transform(PLocalShape);
-		PxTransform PShapeGlobalEndPose = PGlobalEndPose.transform(PLocalShape);
-		// consider localshape rotation for shape rotation
-		PxQuat PShapeRot = PGeomRot * PLocalShape.q;
-
-		GET_GEOMETRY_FROM_SHAPE(PGeom, PShape);
-
-		if(PGeom != NULL)
-		{
-			// @todo UE4, this might not be the best behavior. If we're looking for the most closest, this have to change to save the result, and find the closest one or 
-			// any other options, right now if anything hits first, it will return
-			if (GeomSweepSingle(this, *PGeom, PShapeRot, OutHit, P2UVector(PShapeGlobalStartPose.p), P2UVector(PShapeGlobalEndPose.p), TraceChannel, Params, FCollisionResponseParams(PrimComp->GetCollisionResponseToChannels())))
-			{
-				bHaveBlockingHit = true;
-				break;
-			}
-		}
-	}
-
-	return bHaveBlockingHit;
-#endif //WITH_PHYSX
-
-	return false;
 }
 
 bool UWorld::ComponentSweepMulti(TArray<struct FHitResult>& OutHits, class UPrimitiveComponent* PrimComp, const FVector& Start, const FVector& End, const FRotator& Rot, const struct FComponentQueryParams& Params) const
 {
-	if(GetPhysicsScene() == NULL)
+	if (GetPhysicsScene() == NULL)
 	{
 		return false;
 	}
 
-	if(PrimComp == NULL)
+	if (PrimComp == NULL)
 	{
 		UE_LOG(LogCollision, Log, TEXT("ComponentSweepMulti : No PrimComp"));
 		return false;
@@ -778,70 +318,80 @@ bool UWorld::ComponentSweepMulti(TArray<struct FHitResult>& OutHits, class UPrim
 		return RaycastMulti(this, OutHits, Start, End, TraceChannel, Params, FCollisionResponseParams(PrimComp->GetCollisionResponseToChannels()));
 	}
 
-#if WITH_PHYSX
-	PxRigidActor* PRigidActor = PrimComp->BodyInstance.GetPxRigidActor();
-	if(PRigidActor == NULL)
-	{
-		UE_LOG(LogCollision, Log, TEXT("ComponentSweepMulti : (%s) No physics data"), *PrimComp->GetPathName());
-		return false;
-	}
-	PxScene * const PScene = PRigidActor->getScene();
-
 	OutHits.Empty();
 
-	// Get all the shapes from the actor
-	TArray<PxShape*, TInlineAllocator<8>> PShapes;
+#if UE_WITH_PHYSICS
+	if (!PrimComp->BodyInstance.IsValidBodyInstance())
 	{
-		SCOPED_SCENE_READ_LOCK(PScene);
-		PShapes.AddZeroed(PRigidActor->getNbShapes());
-		PRigidActor->getShapes(PShapes.GetData(), PShapes.Num());
+		UE_LOG(LogCollision, Log, TEXT("ComponentSweepMulti : (%s) No physics data"), *PrimComp->GetReadableName());
+		return false;
 	}
-
-	// calculate the test global pose of the actor
-	PxTransform PGlobalStartPose = U2PTransform(FTransform(Start));
-	PxTransform PGlobalEndPose = U2PTransform(FTransform(End));
+#endif
 
 	bool bHaveBlockingHit = false;
-	PxQuat PGeomRot = U2PQuat(Rot.Quaternion());
 
-	// Iterate over each shape
-	SCENE_LOCK_READ(PScene);
-	for(int32 ShapeIdx=0; ShapeIdx<PShapes.Num(); ShapeIdx++)
+#if WITH_PHYSX
+	if (PxRigidActor* PRigidActor = PrimComp->BodyInstance.GetPxRigidActor())
 	{
-		PxShape* PShape = PShapes[ShapeIdx];
-		check(PShape);
+		PxScene* const PScene = PRigidActor->getScene();
 
-		TArray<struct FHitResult> Hits;
+		// Get all the shapes from the actor
+		TArray<PxShape*, TInlineAllocator<8>> PShapes;
+		{
+			SCOPED_SCENE_READ_LOCK(PScene);
+			PShapes.AddZeroed(PRigidActor->getNbShapes());
+			PRigidActor->getShapes(PShapes.GetData(), PShapes.Num());
+		}
 
-		// Calc shape global pose
-		PxTransform PLocalShape = PShape->getLocalPose();
-		PxTransform PShapeGlobalStartPose = PGlobalStartPose.transform(PLocalShape);
-		PxTransform PShapeGlobalEndPose = PGlobalEndPose.transform(PLocalShape);
-		// consider localshape rotation for shape rotation
-		PxQuat PShapeRot = PGeomRot * PLocalShape.q;
+		// calculate the test global pose of the actor
+		PxTransform PGlobalStartPose = U2PTransform(FTransform(Start));
+		PxTransform PGlobalEndPose = U2PTransform(FTransform(End));
 
-		GET_GEOMETRY_FROM_SHAPE(PGeom, PShape);
+		PxQuat PGeomRot = U2PQuat(Rot.Quaternion());
 
-		if(PGeom != NULL)
+		// Iterate over each shape
+		SCENE_LOCK_READ(PScene);
+		for(int32 ShapeIdx=0; ShapeIdx<PShapes.Num(); ShapeIdx++)
 		{
-			SCENE_UNLOCK_READ(PScene);
-			if (GeomSweepMulti(this, *PGeom, PShapeRot, Hits, P2UVector(PShapeGlobalStartPose.p), P2UVector(PShapeGlobalEndPose.p), TraceChannel, Params, FCollisionResponseParams(PrimComp->GetCollisionResponseToChannels())))
+			PxShape* PShape = PShapes[ShapeIdx];
+			check(PShape);
+
+			TArray<struct FHitResult> Hits;
+
+			// Calc shape global pose
+			PxTransform PLocalShape = PShape->getLocalPose();
+			PxTransform PShapeGlobalStartPose = PGlobalStartPose.transform(PLocalShape);
+			PxTransform PShapeGlobalEndPose = PGlobalEndPose.transform(PLocalShape);
+			// consider localshape rotation for shape rotation
+			PxQuat PShapeRot = PGeomRot * PLocalShape.q;
+
+			GET_GEOMETRY_FROM_SHAPE(PGeom, PShape);
+
+			if(PGeom != NULL)
 			{
-				bHaveBlockingHit = true;
+				SCENE_UNLOCK_READ(PScene);
+				if (GeomSweepMulti_PhysX(this, *PGeom, PShapeRot, Hits, P2UVector(PShapeGlobalStartPose.p), P2UVector(PShapeGlobalEndPose.p), TraceChannel, Params, FCollisionResponseParams(PrimComp->GetCollisionResponseToChannels())))
+				{
+					bHaveBlockingHit = true;
+				}
+
+				OutHits.Append(Hits);
+				SCENE_LOCK_READ(PScene);
 			}
-
-			OutHits.Append(Hits);
-			SCENE_LOCK_READ(PScene);
 		}
+		SCENE_UNLOCK_READ(PScene);
 	}
-	SCENE_UNLOCK_READ(PScene);
-
-	return bHaveBlockingHit;
 #endif //WITH_PHYSX
 
 	//@TODO: BOX2D: Implement UWorld::ComponentSweepMulti
+#if WITH_BOX2D
+// 	if (b2Body* BodyInstance = PrimComp->BodyInstance.BodyInstancePtr)
+// 	{
+// 		
+// 	}
+#endif
 
-	return false;
+	return bHaveBlockingHit;
 }
 
 
diff --git a/Engine/Source/Runtime/Engine/Private/Collision/WorldCollisionAsync.cpp b/Engine/Source/Runtime/Engine/Private/Collision/WorldCollisionAsync.cpp
--- a/Engine/Source/Runtime/Engine/Private/Collision/WorldCollisionAsync.cpp
+++ b/Engine/Source/Runtime/Engine/Private/Collision/WorldCollisionAsync.cpp
@@ -87,8 +87,8 @@ namespace
 
 	void RunTraceTask(FTraceDatum* TraceDataBuffer, int32 TotalCount)
 	{
-	#if WITH_PHYSX
-		check (TraceDataBuffer);
+#if UE_WITH_PHYSICS
+		check(TraceDataBuffer);
 
 		for (; TotalCount; --TotalCount)
 		{
@@ -97,109 +97,47 @@ namespace
 
 			if (TraceData.PhysWorld.IsValid())
 			{
-				switch (TraceData.CollisionParams.CollisionShape.ShapeType)
+				if ((TraceData.CollisionParams.CollisionShape.ShapeType == ECollisionShape::Line) || TraceData.CollisionParams.CollisionShape.IsNearlyZero())
 				{
-				case ECollisionShape::Line:
 					if (TraceData.bIsMultiTrace)
 					{
-						RaycastMulti(TraceData.PhysWorld.Get(), TraceData.OutHits, TraceData.Start, TraceData.End, TraceData.TraceChannel, 
-							TraceData.CollisionParams.CollisionQueryParam, TraceData.CollisionParams.ResponseParam, TraceData.CollisionParams.ObjectQueryParam );
+						RaycastMulti(TraceData.PhysWorld.Get(), TraceData.OutHits, TraceData.Start, TraceData.End, TraceData.TraceChannel,
+							TraceData.CollisionParams.CollisionQueryParam, TraceData.CollisionParams.ResponseParam, TraceData.CollisionParams.ObjectQueryParam);
 					}
 					else
 					{
 						TraceData.OutHits.AddZeroed(1);
-						RaycastSingle(TraceData.PhysWorld.Get(), TraceData.OutHits[0], TraceData.Start, TraceData.End, TraceData.TraceChannel, 
-							TraceData.CollisionParams.CollisionQueryParam, TraceData.CollisionParams.ResponseParam, TraceData.CollisionParams.ObjectQueryParam );
+						RaycastSingle(TraceData.PhysWorld.Get(), TraceData.OutHits[0], TraceData.Start, TraceData.End, TraceData.TraceChannel,
+							TraceData.CollisionParams.CollisionQueryParam, TraceData.CollisionParams.ResponseParam, TraceData.CollisionParams.ObjectQueryParam);
 					}
-					break;
-				case ECollisionShape::Box:
+				}
+				else
+				{
+					if (TraceData.bIsMultiTrace)
 					{
-						PxBoxGeometry PBoxGeom( U2PVector(TraceData.CollisionParams.CollisionShape.GetBox()) );
-						PxQuat PGeomRot = PxQuat::createIdentity();
-
-						// when extent is nearly zero, just do linetrace, 
-						// thought of changing param, when entered, but if we do, it might confuse users when they get the result. 
-						if (TraceData.CollisionParams.CollisionShape.IsNearlyZero())
-						{
-							if (TraceData.bIsMultiTrace)
-							{
-								RaycastMulti(TraceData.PhysWorld.Get(), TraceData.OutHits, TraceData.Start, TraceData.End, TraceData.TraceChannel, 
-									TraceData.CollisionParams.CollisionQueryParam, TraceData.CollisionParams.ResponseParam, TraceData.CollisionParams.ObjectQueryParam );
-							}
-							else
-							{
-								TraceData.OutHits.AddZeroed(1);
-								RaycastSingle(TraceData.PhysWorld.Get(), TraceData.OutHits[0], TraceData.Start, TraceData.End,TraceData.TraceChannel, 
-									TraceData.CollisionParams.CollisionQueryParam, TraceData.CollisionParams.ResponseParam, TraceData.CollisionParams.ObjectQueryParam );
-							}
-						}
-						else
-						{
-							if (TraceData.bIsMultiTrace)
-							{
-								GeomSweepMulti(TraceData.PhysWorld.Get(), PBoxGeom, PGeomRot, TraceData.OutHits, TraceData.Start, TraceData.End, TraceData.TraceChannel, 
-									TraceData.CollisionParams.CollisionQueryParam, TraceData.CollisionParams.ResponseParam, TraceData.CollisionParams.ObjectQueryParam );
-							}
-							else
-							{
-								TraceData.OutHits.AddZeroed(1);
-								GeomSweepSingle(TraceData.PhysWorld.Get(), PBoxGeom, PGeomRot, TraceData.OutHits[0], TraceData.Start, TraceData.End, TraceData.TraceChannel, 
-									TraceData.CollisionParams.CollisionQueryParam, TraceData.CollisionParams.ResponseParam, TraceData.CollisionParams.ObjectQueryParam );
-							}
-						}
+						GeomSweepMulti(TraceData.PhysWorld.Get(), TraceData.CollisionParams.CollisionShape, FQuat::Identity, TraceData.OutHits, TraceData.Start, TraceData.End, TraceData.TraceChannel,
+							TraceData.CollisionParams.CollisionQueryParam, TraceData.CollisionParams.ResponseParam, TraceData.CollisionParams.ObjectQueryParam);
 					}
-					break;
-				case ECollisionShape::Capsule:
+					else
 					{
-						PxCapsuleGeometry PCapsuleGeom( TraceData.CollisionParams.CollisionShape.Capsule.Radius, TraceData.CollisionParams.CollisionShape.GetCapsuleAxisHalfLength() );
-						FQuat PGeomRot = FQuat::Identity; 
-
-						// when extent is nearly zero, just do linetrace, 
-						// thought of changing param, when entered, but if we do, it might confuse users when they get the result. 
-						if(TraceData.CollisionParams.CollisionShape.IsNearlyZero())
-						{
-							if (TraceData.bIsMultiTrace)
-							{
-								RaycastMulti(TraceData.PhysWorld.Get(), TraceData.OutHits, TraceData.Start, TraceData.End, TraceData.TraceChannel, 
-									TraceData.CollisionParams.CollisionQueryParam, TraceData.CollisionParams.ResponseParam, TraceData.CollisionParams.ObjectQueryParam );
-							}
-							else
-							{
-								TraceData.OutHits.AddZeroed(1);
-								RaycastSingle(TraceData.PhysWorld.Get(), TraceData.OutHits[0], TraceData.Start, TraceData.End, TraceData.TraceChannel, 
-									TraceData.CollisionParams.CollisionQueryParam, TraceData.CollisionParams.ResponseParam, TraceData.CollisionParams.ObjectQueryParam );
-							}
-						}
-						else
-						{
-							if (TraceData.bIsMultiTrace)
-							{
-								GeomSweepMulti(TraceData.PhysWorld.Get(), PCapsuleGeom, ConvertToPhysXCapsuleRot(PGeomRot), TraceData.OutHits, TraceData.Start, TraceData.End, TraceData.TraceChannel, 
-									TraceData.CollisionParams.CollisionQueryParam, TraceData.CollisionParams.ResponseParam, TraceData.CollisionParams.ObjectQueryParam );
-							}
-							else
-							{
-								TraceData.OutHits.AddZeroed(1);
-								GeomSweepSingle(TraceData.PhysWorld.Get(), PCapsuleGeom, ConvertToPhysXCapsuleRot(PGeomRot), TraceData.OutHits[0], TraceData.Start, TraceData.End, TraceData.TraceChannel, 
-									TraceData.CollisionParams.CollisionQueryParam, TraceData.CollisionParams.ResponseParam, TraceData.CollisionParams.ObjectQueryParam );
-							}
-						}
-						break;
+						TraceData.OutHits.AddZeroed(1);
+						GeomSweepSingle(TraceData.PhysWorld.Get(), TraceData.CollisionParams.CollisionShape, FQuat::Identity, TraceData.OutHits[0], TraceData.Start, TraceData.End, TraceData.TraceChannel,
+							TraceData.CollisionParams.CollisionQueryParam, TraceData.CollisionParams.ResponseParam, TraceData.CollisionParams.ObjectQueryParam);
 					}
 				}
 			}
 		}
-	#endif //WITH_PHYSX
+	#endif //UE_WITH_PHYSICS
 	}
 
 	void RunTraceTask(FOverlapDatum* OverlapDataBuffer, int32 TotalCount)
 	{
-	#if WITH_PHYSX
-		check (OverlapDataBuffer);
+	#if UE_WITH_PHYSICS
+		check(OverlapDataBuffer);
 
 		for (; TotalCount; --TotalCount)
 		{
-			FOverlapDatum & OverlapData = *OverlapDataBuffer++;
+			FOverlapDatum& OverlapData = *OverlapDataBuffer++;
 			OverlapData.OutOverlaps.Empty();
 
 			if (!OverlapData.PhysWorld.IsValid())
@@ -207,101 +145,36 @@ namespace
 				continue;
 			}
 
-			switch (OverlapData.CollisionParams.CollisionShape.ShapeType)
+			if (OverlapData.bIsMultiTrace)
 			{
-			case ECollisionShape::Box:
-				{
-					PxTransform PGeomPose(U2PTransform(FTransform(OverlapData.Rot, OverlapData.Pos)));
-					if (OverlapData.CollisionParams.CollisionShape.IsNearlyZero())
-					{
-						PxSphereGeometry PSphereGeom( 0.f );
-						// put right sphere
-						if (OverlapData.bIsMultiTrace)
-						{
-							GeomOverlapMulti(OverlapData.PhysWorld.Get(), PSphereGeom, PGeomPose, OverlapData.OutOverlaps, OverlapData.TraceChannel, 
-								OverlapData.CollisionParams.CollisionQueryParam, OverlapData.CollisionParams.ResponseParam, OverlapData.CollisionParams.ObjectQueryParam );
-						}
-						else
-						{
-							OverlapData.OutOverlaps.AddZeroed(1);
-							GeomOverlapSingle(OverlapData.PhysWorld.Get(),PSphereGeom, PGeomPose, OverlapData.OutOverlaps[0], OverlapData.TraceChannel, 
-								OverlapData.CollisionParams.CollisionQueryParam, OverlapData.CollisionParams.ResponseParam, OverlapData.CollisionParams.ObjectQueryParam );
-						}					
-					}
-					else
-					{
-						PxBoxGeometry PBoxGeom( U2PVector(OverlapData.CollisionParams.CollisionShape.GetBox()) );
-						if (OverlapData.bIsMultiTrace)
-						{
-							GeomOverlapMulti(OverlapData.PhysWorld.Get(),PBoxGeom, PGeomPose, OverlapData.OutOverlaps, OverlapData.TraceChannel, 
-								OverlapData.CollisionParams.CollisionQueryParam, OverlapData.CollisionParams.ResponseParam, OverlapData.CollisionParams.ObjectQueryParam );
-						}
-						else
-						{
-							OverlapData.OutOverlaps.AddZeroed(1);
-							GeomOverlapSingle(OverlapData.PhysWorld.Get(),PBoxGeom, PGeomPose, OverlapData.OutOverlaps[0], OverlapData.TraceChannel, 
-								OverlapData.CollisionParams.CollisionQueryParam, OverlapData.CollisionParams.ResponseParam, OverlapData.CollisionParams.ObjectQueryParam );
-						}
-					}
-				}
-				break;
-			case ECollisionShape::Capsule:
-				{
-					FTransform CapsulePose(OverlapData.Rot, OverlapData.Pos);
-
-					if(OverlapData.CollisionParams.CollisionShape.IsNearlyZero())
-					{
-						PxSphereGeometry PSphereGeom( 0.f );
-						// put right sphere
-						if (OverlapData.bIsMultiTrace)
-						{
-							GeomOverlapMulti(OverlapData.PhysWorld.Get(),PSphereGeom, U2PTransform(CapsulePose), OverlapData.OutOverlaps, OverlapData.TraceChannel, 
-								OverlapData.CollisionParams.CollisionQueryParam, OverlapData.CollisionParams.ResponseParam, OverlapData.CollisionParams.ObjectQueryParam);
-						}
-						else
-						{
-							OverlapData.OutOverlaps.AddZeroed(1);
-							GeomOverlapSingle(OverlapData.PhysWorld.Get(),PSphereGeom,  U2PTransform(CapsulePose), OverlapData.OutOverlaps[0], OverlapData.TraceChannel, 
-								OverlapData.CollisionParams.CollisionQueryParam, OverlapData.CollisionParams.ResponseParam, OverlapData.CollisionParams.ObjectQueryParam );
-						}					
-					}
-					else
-					{
-						PxCapsuleGeometry PCapsuleGeom( OverlapData.CollisionParams.CollisionShape.Capsule.Radius, OverlapData.CollisionParams.CollisionShape.GetCapsuleAxisHalfLength() );
-						if (OverlapData.bIsMultiTrace)
-						{
-							GeomOverlapMulti(OverlapData.PhysWorld.Get(),PCapsuleGeom, ConvertToPhysXCapsulePose(CapsulePose), OverlapData.OutOverlaps, OverlapData.TraceChannel, 
-								OverlapData.CollisionParams.CollisionQueryParam, OverlapData.CollisionParams.ResponseParam, OverlapData.CollisionParams.ObjectQueryParam );
-						}
-						else
-						{
-							OverlapData.OutOverlaps.AddZeroed(1);
-							GeomOverlapSingle(OverlapData.PhysWorld.Get(),PCapsuleGeom, ConvertToPhysXCapsulePose(CapsulePose), OverlapData.OutOverlaps[0], OverlapData.TraceChannel, 
-								OverlapData.CollisionParams.CollisionQueryParam, OverlapData.CollisionParams.ResponseParam, OverlapData.CollisionParams.ObjectQueryParam );
-						}
-					}
-				}
-				break;
-			case ECollisionShape::Sphere:
-				{
-					PxSphereGeometry PSphereGeom( OverlapData.CollisionParams.CollisionShape.Sphere.Radius );
-					PxTransform PGeomPose(U2PVector(OverlapData.Pos), PxQuat::createIdentity());
-					if (OverlapData.bIsMultiTrace)
-					{
-						GeomOverlapMulti(OverlapData.PhysWorld.Get(),PSphereGeom, PGeomPose, OverlapData.OutOverlaps, OverlapData.TraceChannel, 
-							OverlapData.CollisionParams.CollisionQueryParam, OverlapData.CollisionParams.ResponseParam, OverlapData.CollisionParams.ObjectQueryParam );
-					}
-					else
-					{
-						OverlapData.OutOverlaps.AddZeroed(1);
-						GeomOverlapSingle(OverlapData.PhysWorld.Get(),PSphereGeom, PGeomPose, OverlapData.OutOverlaps[0], OverlapData.TraceChannel, 
-							OverlapData.CollisionParams.CollisionQueryParam, OverlapData.CollisionParams.ResponseParam, OverlapData.CollisionParams.ObjectQueryParam );
-					}
-				}
-				break;
+				GeomOverlapMulti(
+					OverlapData.PhysWorld.Get(),
+					OverlapData.CollisionParams.CollisionShape,
+					OverlapData.Pos,
+					OverlapData.Rot,
+					OverlapData.OutOverlaps,
+					OverlapData.TraceChannel,
+					OverlapData.CollisionParams.CollisionQueryParam,
+					OverlapData.CollisionParams.ResponseParam,
+					OverlapData.CollisionParams.ObjectQueryParam);
+			}
+			else
+			{
+				OverlapData.OutOverlaps.AddZeroed(1);
+
+				GeomOverlapSingle(
+					OverlapData.PhysWorld.Get(),
+					OverlapData.CollisionParams.CollisionShape,
+					OverlapData.Pos,
+					OverlapData.Rot,
+					OverlapData.OutOverlaps[0],
+					OverlapData.TraceChannel,
+					OverlapData.CollisionParams.CollisionQueryParam,
+					OverlapData.CollisionParams.ResponseParam,
+					OverlapData.CollisionParams.ObjectQueryParam);
 			}
 		}
-	#endif //WITH_PHYSX
+	#endif //UE_WITH_PHYSICS
 	}
 
 	#if RUN_ASYNC_TRACE
diff --git a/Engine/Source/Runtime/Engine/Private/PhysicsEngine/BodyInstance.cpp b/Engine/Source/Runtime/Engine/Private/PhysicsEngine/BodyInstance.cpp
--- a/Engine/Source/Runtime/Engine/Private/PhysicsEngine/BodyInstance.cpp
+++ b/Engine/Source/Runtime/Engine/Private/PhysicsEngine/BodyInstance.cpp
@@ -2852,141 +2852,130 @@ FString FBodyInstance::GetBodyDebugName() const
 /////////////////////////////////////////////////////////////////////////////////////////////////////////////////////////
 // COLLISION
 
-bool FBodyInstance::LineTrace(struct FHitResult& OutHit, const FVector& Start, const FVector& End, bool bTraceComplex, bool bReturnPhysicalMaterial)
+bool FBodyInstance::LineTrace(struct FHitResult& OutHit, const FVector& Start, const FVector& End, bool bTraceComplex, bool bReturnPhysicalMaterial) const
 {
 	SCOPE_CYCLE_COUNTER(STAT_Collision_RaycastAny);
 
 	OutHit.TraceStart = Start;
 	OutHit.TraceEnd = End;
 
-#if WITH_PHYSX
-	const PxRigidActor* RigidBody = GetPxRigidActor();
-	if (RigidBody == NULL || RigidBody->getNbShapes()==0)
-	{
-		return false;
-	}
+	bool bHitSomething = false;
 
 	const FVector Delta = End - Start;
 	const float DeltaMag = Delta.Size();
-	if( DeltaMag > KINDA_SMALL_NUMBER )
+	if (DeltaMag > KINDA_SMALL_NUMBER)
 	{
-		// Create filter data used to filter collisions, should always return eTOUCH for LineTraceComponent		
-		PxSceneQueryFlags POutputFlags = PxSceneQueryFlag::eIMPACT | PxSceneQueryFlag::eNORMAL | PxSceneQueryFlag::eDISTANCE;
-
-		PxRaycastHit BestHit;
-		float BestHitDistance = BIG_NUMBER;
-
-		// Get all the shapes from the actor
-		TArray<PxShape*> PShapes;
-		PShapes.AddZeroed(RigidBody->getNbShapes());
-		int32 NumShapes = RigidBody->getShapes(PShapes.GetData(), PShapes.Num());
-
-		// Iterate over each shape
-		for(int32 ShapeIdx=0; ShapeIdx<PShapes.Num(); ShapeIdx++)
+#if WITH_PHYSX
+		const PxRigidActor* RigidBody = GetPxRigidActor();
+		if ((RigidBody != NULL) && (RigidBody->getNbShapes() != 0))
 		{
-			PxShape* PShape = PShapes[ShapeIdx];
-			check(PShape);
+			// Create filter data used to filter collisions, should always return eTOUCH for LineTraceComponent		
+			PxSceneQueryFlags POutputFlags = PxSceneQueryFlag::eIMPACT | PxSceneQueryFlag::eNORMAL | PxSceneQueryFlag::eDISTANCE;
 
-			const PxU32 HitBufferSize = 1; 
-			PxRaycastHit PHits[HitBufferSize];
+			PxRaycastHit BestHit;
+			float BestHitDistance = BIG_NUMBER;
 
-			// Filter so we trace against the right kind of collision
-			PxFilterData ShapeFilter = PShape->getQueryFilterData();
-			const bool bShapeIsComplex = (ShapeFilter.word3 & EPDF_ComplexCollision) != 0;
-			const bool bShapeIsSimple = (ShapeFilter.word3 & EPDF_SimpleCollision) != 0;
-			if ((bTraceComplex && bShapeIsComplex) || (!bTraceComplex && bShapeIsSimple))
+			// Get all the shapes from the actor
+			TArray<PxShape*> PShapes;
+			PShapes.AddZeroed(RigidBody->getNbShapes());
+			int32 NumShapes = RigidBody->getShapes(PShapes.GetData(), PShapes.Num());
+
+			// Iterate over each shape
+			for (int32 ShapeIdx = 0; ShapeIdx < PShapes.Num(); ShapeIdx++)
 			{
-				const int32 ArraySize = ARRAY_COUNT(PHits);
-				// only care about one hit - not closest hit			
-				const PxI32 NumHits = PxGeometryQuery::raycast(U2PVector(Start), U2PVector(Delta/DeltaMag), PShape->getGeometry().any(), PxShapeExt::getGlobalPose(*PShape, *RigidBody), DeltaMag, POutputFlags, ArraySize, PHits, true);
+				PxShape* PShape = PShapes[ShapeIdx];
+				check(PShape);
 
+				const PxU32 HitBufferSize = 1;
+				PxRaycastHit PHits[HitBufferSize];
 
-				if ( ensure (NumHits <= ArraySize) )
+				// Filter so we trace against the right kind of collision
+				PxFilterData ShapeFilter = PShape->getQueryFilterData();
+				const bool bShapeIsComplex = (ShapeFilter.word3 & EPDF_ComplexCollision) != 0;
+				const bool bShapeIsSimple = (ShapeFilter.word3 & EPDF_SimpleCollision) != 0;
+				if ((bTraceComplex && bShapeIsComplex) || (!bTraceComplex && bShapeIsSimple))
 				{
-					for(int HitIndex = 0; HitIndex < NumHits; HitIndex++)
+					const int32 ArraySize = ARRAY_COUNT(PHits);
+					// only care about one hit - not closest hit			
+					const PxI32 NumHits = PxGeometryQuery::raycast(U2PVector(Start), U2PVector(Delta / DeltaMag), PShape->getGeometry().any(), PxShapeExt::getGlobalPose(*PShape, *RigidBody), DeltaMag, POutputFlags, ArraySize, PHits, true);
+
+
+					if (ensure(NumHits <= ArraySize))
 					{
-						if(PHits[HitIndex].distance < BestHitDistance)
+						for (int HitIndex = 0; HitIndex < NumHits; HitIndex++)
 						{
-							BestHitDistance = PHits[HitIndex].distance;
-							BestHit = PHits[HitIndex];
-
-							// we don't get Shape information when we access via PShape, so I filled it up
-							BestHit.shape = PShape;
-							BestHit.actor = PShape->getActor();
+							if (PHits[HitIndex].distance < BestHitDistance)
+							{
+								BestHitDistance = PHits[HitIndex].distance;
+								BestHit = PHits[HitIndex];
+
+								// we don't get Shape information when we access via PShape, so I filled it up
+								BestHit.shape = PShape;
+								BestHit.actor = PShape->getActor();
+							}
 						}
 					}
 				}
 			}
-		}
 
-		if (BestHitDistance < BIG_NUMBER)
-		{
-			// we just like to make sure if the hit is made, set to test touch
-			PxFilterData QueryFilter;
-			QueryFilter.word2 = 0xFFFFF;
+			if (BestHitDistance < BIG_NUMBER)
+			{
+				// we just like to make sure if the hit is made, set to test touch
+				PxFilterData QueryFilter;
+				QueryFilter.word2 = 0xFFFFF;
 
-			PxTransform PStartTM(U2PVector(Start));
-			ConvertQueryImpactHit(BestHit, OutHit, DeltaMag, QueryFilter, Start, End, NULL, PStartTM, false, bReturnPhysicalMaterial);
-			return true;
+				PxTransform PStartTM(U2PVector(Start));
+				ConvertQueryImpactHit(BestHit, OutHit, DeltaMag, QueryFilter, Start, End, NULL, PStartTM, false, bReturnPhysicalMaterial);
+				bHitSomething = true;
+			}
 		}
-	}
 #endif //WITH_PHYSX
 
-	//@TODO: BOX2D: Implement LineTrace
+#if WITH_BOX2D
+		if (BodyInstancePtr != nullptr)
+		{
+			//@TODO: BOX2D: Implement FBodyInstance::LineTrace
+		}
+#endif
+	}
 
-	return false;
+	return bHitSomething;
 }
 
-bool FBodyInstance::Sweep(struct FHitResult& OutHit, const FVector& Start, const FVector& End, const FCollisionShape& CollisionShape, bool bTraceComplex)
+bool FBodyInstance::Sweep(struct FHitResult& OutHit, const FVector& Start, const FVector& End, const FCollisionShape& CollisionShape, bool bTraceComplex) const
 {
-	OutHit.TraceStart = Start;
-	OutHit.TraceEnd = End;
-
-#if WITH_PHYSX
-	const PxRigidActor* RigidBody = GetPxRigidActor();
-	if(RigidBody == NULL || RigidBody->getNbShapes()==0 || OwnerComponent==NULL)
-	{
-		return false;
-	}
-
-	if(CollisionShape.IsNearlyZero())
+	if (CollisionShape.IsNearlyZero())
 	{
 		return LineTrace(OutHit, Start, End, bTraceComplex);
 	}
-
-	switch(CollisionShape.ShapeType)
+	else
 	{
-	case ECollisionShape::Box:
-		{
-			const PxBoxGeometry PBoxGeom(U2PVector(CollisionShape.GetBox()));
-			const PxQuat PGeomRot = PxQuat::createIdentity();
+		OutHit.TraceStart = Start;
+		OutHit.TraceEnd = End;
 
-			return InternalSweep(OutHit, Start, End, CollisionShape, bTraceComplex, RigidBody, &PBoxGeom);
-		}
-	case ECollisionShape::Sphere:
+		bool bSweepHit = false;
+#if WITH_PHYSX
+		const PxRigidActor* RigidBody = GetPxRigidActor();
+		if ((RigidBody != NULL) && (RigidBody->getNbShapes() != 0) && (OwnerComponent != NULL))
 		{
-			PxSphereGeometry PSphereGeom(CollisionShape.GetSphereRadius());
-			PxQuat PGeomRot = PxQuat::createIdentity();
-
-			return InternalSweep(OutHit, Start, End, CollisionShape, bTraceComplex, RigidBody, &PSphereGeom);
+			FPhysXShapeAdaptor ShapeAdaptor(FQuat::Identity, CollisionShape);
+			bSweepHit = InternalSweepPhysX(OutHit, Start, End, CollisionShape, bTraceComplex, RigidBody, &ShapeAdaptor.GetGeometry());
 		}
+#endif //WITH_PHYSX
 
-	case ECollisionShape::Capsule:
+#if WITH_BOX2D
+		if (BodyInstancePtr != nullptr)
 		{
-			PxCapsuleGeometry PCapsuleGeom(CollisionShape.GetCapsuleRadius(), CollisionShape.GetCapsuleAxisHalfLength());
-			const PxQuat PGeomRot = PxQuat::createIdentity();
-			return InternalSweep(OutHit, Start, End, CollisionShape, bTraceComplex, RigidBody, &PCapsuleGeom);
+			//@TODO: BOX2D: Implement FBodyInstance::Sweep
 		}
-	}
-#endif //WITH_PHYSX
-
-	//@TODO: BOX2D: Implement Sweep
+#endif
 
-	return false;
+		return bSweepHit;
+	}
 }
 
 #if WITH_PHYSX
-bool FBodyInstance::InternalSweep(struct FHitResult& OutHit, const FVector& Start, const FVector& End, const FCollisionShape& Shape, bool bTraceComplex, const PxRigidActor* RigidBody, const PxGeometry* Geometry)
+bool FBodyInstance::InternalSweepPhysX(struct FHitResult& OutHit, const FVector& Start, const FVector& End, const FCollisionShape& Shape, bool bTraceComplex, const PxRigidActor* RigidBody, const PxGeometry* Geometry) const
 {
 	const FVector Delta = End - Start;
 	const float DeltaMag = Delta.Size();
@@ -3113,8 +3102,102 @@ float FBodyInstance::GetDistanceToBody(const FVector& Point, FVector& OutPointOn
 	//@TODO: BOX2D: Implement DistanceToBody
 }
 
+bool FBodyInstance::OverlapTest(const FVector& Position, const FQuat& Rotation, const struct FCollisionShape& CollisionShape) const
+{
+	bool bHasOverlap = false;
+
+#if WITH_PHYSX
+	FPhysXShapeAdaptor ShapeAdaptor(Rotation, CollisionShape);
+	bHasOverlap = OverlapPhysX(ShapeAdaptor.GetGeometry(), ShapeAdaptor.GetGeomPose(Position));
+#endif
+
+#if WITH_BOX2D
+	if (!bHasOverlap && (BodyInstancePtr != nullptr))
+	{
+		//@TODO: BOX2D: Implement FBodyInstance::OverlapTest
+	}
+#endif
+
+	return bHasOverlap;
+}
+
+bool FBodyInstance::OverlapMulti(TArray<struct FOverlapResult>& InOutOverlaps, const class UWorld* World, const FTransform* pWorldToComponent, const FVector& Pos, const FRotator& Rot, ECollisionChannel TestChannel, const struct FComponentQueryParams& Params, const struct FCollisionResponseParams& ResponseParams, const struct FCollisionObjectQueryParams& ObjectQueryParams) const
+{
+	if (!IsValidBodyInstance())
+	{
+		UE_LOG(LogCollision, Log, TEXT("FBodyInstance::OverlapMulti : (%s) No physics data"), *GetBodyDebugName());
+		return false;
+	}
+
+	bool bHaveBlockingHit = false;
+
+	// Determine how to convert the local space of this body instance to the test space
+	const FTransform ComponentSpaceToTestSpace(Rot, Pos);
+
+	FTransform BodyInstanceSpaceToTestSpace;
+	if (pWorldToComponent)
+	{
+		const FTransform ShapeSpaceToComponentSpace = GetUnrealWorldTransform() * (*pWorldToComponent);
+		BodyInstanceSpaceToTestSpace = ShapeSpaceToComponentSpace * ComponentSpaceToTestSpace;
+	}
+	else
+	{
+		BodyInstanceSpaceToTestSpace = ComponentSpaceToTestSpace;
+	}
+
+#if WITH_PHYSX
+	if (PxRigidActor* PRigidActor = GetPxRigidActor())
+	{
+		const PxTransform PBodyInstanceSpaceToTestSpace = U2PTransform(BodyInstanceSpaceToTestSpace);
+
+		// Get all the shapes from the actor
+		{
+			SCOPED_SCENE_READ_LOCK(PRigidActor->getScene());
+
+			TArray<PxShape*, TInlineAllocator<8>> PShapes;
+			PShapes.AddZeroed(PRigidActor->getNbShapes());
+			int32 NumShapes = PRigidActor->getShapes(PShapes.GetData(), PShapes.Num());
+
+			// Iterate over each shape
+			TArray<struct FOverlapResult> TempOverlaps;
+			for (int32 ShapeIdx = 0; ShapeIdx < PShapes.Num(); ShapeIdx++)
+			{
+				PxShape* PShape = PShapes[ShapeIdx];
+				check(PShape);
+
+				// Calc shape global pose
+				const PxTransform PLocalPose = PShape->getLocalPose();
+				const PxTransform PShapeGlobalPose = PBodyInstanceSpaceToTestSpace.transform(PLocalPose);
+
+				GET_GEOMETRY_FROM_SHAPE(PGeom, PShape);
+
+				if (PGeom != NULL)
+				{
+					TempOverlaps.Reset();
+					if (GeomOverlapMulti_PhysX(World, *PGeom, PShapeGlobalPose, TempOverlaps, TestChannel, Params, ResponseParams, ObjectQueryParams))
+					{
+						bHaveBlockingHit = true;
+					}
+					InOutOverlaps.Append(TempOverlaps);
+				}
+			}
+		}
+	}
+#endif
+
+#if WITH_BOX2D
+	if (BodyInstancePtr != nullptr)
+	{
+		//@TODO: BOX2D: Implement FBodyInstance::OverlapMulti
+	}
+#endif
+
+	return bHaveBlockingHit;
+}
+
+
 #if WITH_PHYSX
-bool FBodyInstance::Overlap(const PxGeometry& PGeom, const PxTransform&  ShapePose)
+bool FBodyInstance::OverlapPhysX(const PxGeometry& PGeom, const PxTransform& ShapePose) const
 {
 	const PxRigidActor* RigidBody = GetPxRigidActor();
 
diff --git a/Engine/Source/Runtime/Engine/Private/PrimitiveComponent.cpp b/Engine/Source/Runtime/Engine/Private/PrimitiveComponent.cpp
--- a/Engine/Source/Runtime/Engine/Private/PrimitiveComponent.cpp
+++ b/Engine/Source/Runtime/Engine/Private/PrimitiveComponent.cpp
@@ -1616,6 +1616,8 @@ bool UPrimitiveComponent::SweepComponent(struct FHitResult& OutHit, const FVecto
 
 bool UPrimitiveComponent::ComponentOverlapComponent(class UPrimitiveComponent* PrimComp, const FVector Pos, const FRotator Rot, const struct FCollisionQueryParams& Params)
 {
+	//@TODO: BOX2D: Implement UPrimitiveComponent::ComponentOverlapComponent
+
 	// if target is skeletalmeshcomponent and do not support singlebody physics
 	USkeletalMeshComponent * OtherComp = Cast<USkeletalMeshComponent>(PrimComp);
 	if (OtherComp)
@@ -1652,7 +1654,7 @@ bool UPrimitiveComponent::ComponentOverlapComponent(class UPrimitiveComponent* P
 
 		if(PGeom != NULL)
 		{
-			bHaveOverlap = BodyInstance.Overlap(*PGeom, PShapeGlobalPose);
+			bHaveOverlap = BodyInstance.OverlapPhysX(*PGeom, PShapeGlobalPose);
 
 #if !(UE_BUILD_SHIPPING || UE_BUILD_TEST)
 			if((GetWorld()->DebugDrawTraceTag != NAME_None) && (GetWorld()->DebugDrawTraceTag == Params.TraceTag))
@@ -1686,49 +1688,7 @@ bool UPrimitiveComponent::ComponentOverlapComponent(class UPrimitiveComponent* P
 
 bool UPrimitiveComponent::OverlapComponent(const FVector& Pos, const FQuat& Rot, const struct FCollisionShape& CollisionShape)
 {
-#if WITH_PHYSX
-	PxSphereGeometry PSphereGeom;
-	PxBoxGeometry PBoxGeom;
-	PxCapsuleGeometry PCapsuleGeom;
-	PxGeometry * PGeom = NULL;
-	PxTransform PShapePose;
-
-	switch (CollisionShape.ShapeType)
-	{
-	case ECollisionShape::Sphere:
-		{
-			PSphereGeom = PxSphereGeometry(CollisionShape.GetSphereRadius()) ;
-			// calculate the test global pose of the actor
-			PShapePose = U2PTransform(FTransform(Rot, Pos));
-			PGeom = &PSphereGeom;
-			break;
-		}
-	case ECollisionShape::Box:
-		{
-			PBoxGeom = PxBoxGeometry(U2PVector(CollisionShape.GetBox()));
-			PShapePose = U2PTransform(FTransform(Rot, Pos));
-			PGeom = &PBoxGeom;
-			break;
-		}
-	case ECollisionShape::Capsule:
-		{
-			PCapsuleGeom = PxCapsuleGeometry( CollisionShape.GetCapsuleRadius(), CollisionShape.GetCapsuleAxisHalfLength() );
-			PShapePose = ConvertToPhysXCapsulePose( FTransform(Rot,Pos) );
-			PGeom = &PCapsuleGeom;
-			break;
-		}
-	default:
-		// invalid point
-		ensure(false);
-	}
-
-	if (PGeom)
-	{
-		return BodyInstance.Overlap(*PGeom, PShapePose);
-	}
-
-#endif //WITH_PHYSX
-	return false;
+	return BodyInstance.OverlapTest(Pos, Rot, CollisionShape);
 }
 
 bool UPrimitiveComponent::ComputePenetration(FMTDResult & OutMTD, const FCollisionShape & CollisionShape, const FVector & Pos, const FQuat & Rot)
@@ -2202,58 +2162,10 @@ void UPrimitiveComponent::UpdateOverlaps(TArray<FOverlapInfo> const* PendingOver
 	}
 }
 
-bool UPrimitiveComponent::ComponentOverlapMulti(TArray<struct FOverlapResult>& OutOverlaps, const UWorld * World, const FVector& Pos, const FRotator& Rot, ECollisionChannel TestChannel, const struct FComponentQueryParams& Params, const struct FCollisionObjectQueryParams& ObjectQueryParams) const
+bool UPrimitiveComponent::ComponentOverlapMulti(TArray<struct FOverlapResult>& OutOverlaps, const UWorld* World, const FVector& Pos, const FRotator& Rot, ECollisionChannel TestChannel, const struct FComponentQueryParams& Params, const struct FCollisionObjectQueryParams& ObjectQueryParams) const
 {
-#if WITH_PHYSX
-	PxRigidActor* PRigidActor = BodyInstance.GetPxRigidActor();
-	if(PRigidActor == NULL)
-	{
-		UE_LOG(LogCollision, Log, TEXT("ComponentOverlapMulti : (%s) No physics data"), *GetPathName());
-		return false;
-	}
-
 	OutOverlaps.Reset();
-
-	// calculate the test global pose of the actor
-	PxTransform PTestGlobalPose = U2PTransform(FTransform(Rot, Pos));
-
-	// Get all the shapes from the actor
-	bool bHaveBlockingHit = false;
-	{
-		SCOPED_SCENE_READ_LOCK(PRigidActor->getScene());
-
-		TArray<PxShape*, TInlineAllocator<8>> PShapes;
-		PShapes.AddZeroed(PRigidActor->getNbShapes());
-		int32 NumShapes = PRigidActor->getShapes(PShapes.GetData(), PShapes.Num());
-
-		// Iterate over each shape
-		for(int32 ShapeIdx=0; ShapeIdx<PShapes.Num(); ShapeIdx++)
-		{
-			PxShape* PShape = PShapes[ShapeIdx];
-			check(PShape);
-			TArray<struct FOverlapResult> Overlaps;
-			// Calc shape global pose
-			PxTransform PLocalPose = PShape->getLocalPose();
-			PxTransform PShapeGlobalPose = PTestGlobalPose.transform(PLocalPose);
-
-			GET_GEOMETRY_FROM_SHAPE(PGeom, PShape);
-
-			if(PGeom != NULL)
-			{
-				// if object query param isn't valid, it will use trace channel
-				if (GeomOverlapMulti(World, *PGeom, PShapeGlobalPose, Overlaps, TestChannel, Params, FCollisionResponseParams(GetCollisionResponseToChannels()), ObjectQueryParams))
-				{
-					bHaveBlockingHit = true;
-				}
-
-				OutOverlaps.Append(Overlaps);
-			}
-		}
-	}
-
-	return bHaveBlockingHit;
-#endif //WITH_PHYSX
-	return false;
+	return BodyInstance.OverlapMulti(OutOverlaps, World, /*pWorldToComponent=*/ nullptr, Pos, Rot, TestChannel, Params, FCollisionResponseParams(GetCollisionResponseToChannels()), ObjectQueryParams);
 }
 
 void UPrimitiveComponent::UpdatePhysicsVolume( bool bTriggerNotifiers )
diff --git a/Engine/Source/Runtime/Engine/Private/SkeletalMeshComponentPhysics.cpp b/Engine/Source/Runtime/Engine/Private/SkeletalMeshComponentPhysics.cpp
--- a/Engine/Source/Runtime/Engine/Private/SkeletalMeshComponentPhysics.cpp
+++ b/Engine/Source/Runtime/Engine/Private/SkeletalMeshComponentPhysics.cpp
@@ -1682,7 +1682,7 @@ extern float DebugLineLifetime;
 
 bool USkeletalMeshComponent::LineTraceComponent(struct FHitResult& OutHit, const FVector Start, const FVector End, const struct FCollisionQueryParams& Params)
 {
-	UPhysicsAsset * const PhysicsAsset = GetPhysicsAsset();
+	UPhysicsAsset* const PhysicsAsset = GetPhysicsAsset();
 	bool bHaveHit = false;
 
 	float MinTime = MAX_FLT;
@@ -1732,6 +1732,7 @@ bool USkeletalMeshComponent::SweepComponent( FHitResult& OutHit, const FVector S
 
 bool USkeletalMeshComponent::ComponentOverlapComponent(class UPrimitiveComponent* PrimComp,const FVector Pos,const FRotator Rot,const struct FCollisionQueryParams& Params)
 {
+	//@TODO: BOX2D: USkeletalMeshComponent::ComponentOverlapComponent is not supported.  Try to rephrase this in terms of agnostic operations on a set of FBodyInstances, reducing the amount of PhysX-specific code
 #if WITH_PHYSX
 	// will have to do per component - default single body or physicsinstance 
 	const PxRigidActor* TargetRigidBody = (PrimComp)? PrimComp->BodyInstance.GetPxRigidActor():NULL;
@@ -1772,7 +1773,7 @@ bool USkeletalMeshComponent::ComponentOverlapComponent(class UPrimitiveComponent
 		{
 			for (int32 BodyIdx=0; BodyIdx < Bodies.Num(); ++BodyIdx)
 			{
-				bHaveOverlap = Bodies[BodyIdx]->Overlap(*PGeom, PShapeGlobalPose);
+				bHaveOverlap = Bodies[BodyIdx]->OverlapPhysX(*PGeom, PShapeGlobalPose);
 #if !(UE_BUILD_SHIPPING || UE_BUILD_TEST)
 				if((GetWorld()->DebugDrawTraceTag != NAME_None) && (GetWorld()->DebugDrawTraceTag == Params.TraceTag))
 				{
@@ -1803,127 +1804,40 @@ bool USkeletalMeshComponent::ComponentOverlapComponent(class UPrimitiveComponent
 
 bool USkeletalMeshComponent::OverlapComponent(const FVector& Pos, const FQuat& Rot, const FCollisionShape& CollisionShape)
 {
-
-#if WITH_PHYSX
-	PxSphereGeometry PSphereGeom;
-	PxBoxGeometry PBoxGeom;
-	PxCapsuleGeometry PCapsuleGeom;
-	PxGeometry * PGeom = NULL;
-	PxTransform PShapePose;
-
-	switch (CollisionShape.ShapeType)
+	for (FBodyInstance* Body : Bodies)
 	{
-	case ECollisionShape::Sphere:
-		{
-			PSphereGeom = PxSphereGeometry(CollisionShape.GetSphereRadius()) ;
-			// calculate the test global pose of the actor
-			PShapePose = U2PTransform(FTransform(Rot, Pos));
-			PGeom = &PSphereGeom;
-			break;
-		}
-	case ECollisionShape::Box:
-		{
-			PBoxGeom = PxBoxGeometry(U2PVector(CollisionShape.GetBox()));
-			PShapePose = U2PTransform(FTransform(Rot, Pos));
-			PGeom = &PBoxGeom;
-			break;
-		}
-	case ECollisionShape::Capsule:
+		if (Body->OverlapTest(Pos, Rot, CollisionShape))
 		{
-			PCapsuleGeom = PxCapsuleGeometry( CollisionShape.GetCapsuleRadius(), CollisionShape.GetCapsuleAxisHalfLength() );
-			PShapePose = ConvertToPhysXCapsulePose( FTransform(Rot,Pos) );
-			PGeom = &PCapsuleGeom;
-			break;
-		}
-	default:
-		// invalid point
-		ensure(false);
-	}
-
-	if (PGeom)
-	{
-		for (int32 BodyIdx=0; BodyIdx < Bodies.Num(); ++BodyIdx)
-		{
-			if (Bodies[BodyIdx]->Overlap(*PGeom, PShapePose))
-			{
-				return true;
-			}
+			return true;
 		}
 	}
-#endif //WITH_PHYSX
 
 	return false;
 }
 
 bool USkeletalMeshComponent::ComponentOverlapMulti(TArray<struct FOverlapResult>& OutOverlaps, const UWorld* World, const FVector& Pos, const FRotator& Rot, ECollisionChannel TestChannel, const struct FComponentQueryParams& Params, const struct FCollisionObjectQueryParams& ObjectQueryParams) const
 {
-#if WITH_PHYSX
-
 	OutOverlaps.Reset();
 
-	if(!Bodies.IsValidIndex(RootBodyIndex))
+	if (!Bodies.IsValidIndex(RootBodyIndex))
 	{
 		return false;
 	}
 
-	// calculate the test global pose of the actor
-	// because we're using skeletal mesh we need to use relative transform to root body
-	PxTransform PTestGlobalPose = U2PTransform(FTransform(Rot, Pos));
-	check(Bodies[RootBodyIndex]);
-	PxTransform PRootGlobalPoseInv = U2PTransform(Bodies[RootBodyIndex]->GetUnrealWorldTransform()).getInverse();
-	PxTransform PGlobalRootSpacePose = PTestGlobalPose.transform(PRootGlobalPoseInv);
+	const FTransform WorldToComponent(ComponentToWorld.InverseSafe());
+	const FCollisionResponseParams ResponseParams(GetCollisionResponseToChannels());
 
 	bool bHaveBlockingHit = false;
-
-	for(int32 BodyIdx = 0; BodyIdx < Bodies.Num(); ++BodyIdx)
+	for (const FBodyInstance* Body : Bodies)
 	{
-		const FBodyInstance * BodyInstance = Bodies[BodyIdx];
-		PxRigidActor* PRigidActor = BodyInstance->GetPxRigidActor();
-
-		if(PRigidActor == NULL)
+		checkSlow(Body);
+		if (Body->OverlapMulti(/*inout*/ OutOverlaps, World, &WorldToComponent, Pos, Rot, TestChannel, Params, ResponseParams, ObjectQueryParams))
 		{
-			UE_LOG(LogCollision, Log, TEXT("ComponentOverlapMulti : (%s) No physics data"), *GetPathName());
-			return false;
-		}
-
-		// Get all the shapes from the actor
-		{
-			SCOPED_SCENE_READ_LOCK(PRigidActor->getScene());
-
-			TArray<PxShape*, TInlineAllocator<8>> PShapes;
-			PShapes.AddZeroed(PRigidActor->getNbShapes());
-			int32 NumShapes = PRigidActor->getShapes(PShapes.GetData(), PShapes.Num());
-
-			// Iterate over each shape
-			for(int32 ShapeIdx=0; ShapeIdx<PShapes.Num(); ShapeIdx++)
-			{
-				PxShape* PShape = PShapes[ShapeIdx];
-				check(PShape);
-				TArray<struct FOverlapResult> Overlaps;
-				// Calc shape global pose
-				PxTransform PLocalPose = PShape->getLocalPose();
-				PxTransform PBodyPoseGlobal = U2PTransform(BodyInstance->GetUnrealWorldTransform());
-				PxTransform PShapeGlobalPose = PGlobalRootSpacePose.transform(PBodyPoseGlobal.transform(PLocalPose));
-
-				GET_GEOMETRY_FROM_SHAPE(PGeom, PShape);
-
-				if(PGeom != NULL)
-				{
-					// if object query param isn't valid, it will use trace channel
-					if (GeomOverlapMulti(World, *PGeom, PShapeGlobalPose, Overlaps, TestChannel, Params, FCollisionResponseParams(GetCollisionResponseToChannels()), ObjectQueryParams))
-					{
-						bHaveBlockingHit = true;
-					}
-
-					OutOverlaps.Append(Overlaps);
-				}
-			}
+			bHaveBlockingHit = true;
 		}
 	}
 
 	return bHaveBlockingHit;
-#endif //WITH_PHYSX
-	return false;
 }
 
 
