#ue4 - Expose CalcVelocity(), NewFallVelocity(), and GetMaxJumpHeight() to blueprints.

diff --git a/Engine/Source/Runtime/Engine/Classes/GameFramework/CharacterMovementComponent.h b/Engine/Source/Runtime/Engine/Classes/GameFramework/CharacterMovementComponent.h
--- a/Engine/Source/Runtime/Engine/Classes/GameFramework/CharacterMovementComponent.h
+++ b/Engine/Source/Runtime/Engine/Classes/GameFramework/CharacterMovementComponent.h
@@ -939,8 +939,9 @@ public:
 	/** Called if bNotifyApex is true and character has just passed the apex of its jump. */
 	virtual void NotifyJumpApex();
 
-	/** @return new falling velocity from old velocity and acceleration. */
-	virtual FVector NewFallVelocity(FVector OldVelocity, FVector OldAcceleration, float timeTick);
+	/** Compute new falling velocity from given velocity and gravity. */
+	UFUNCTION(BlueprintCallable, Category="Pawn|Components|CharacterMovement")
+	virtual FVector NewFallVelocity(FVector InitialVelocity, FVector Gravity, float DeltaTime);
 
 	/* Determine how deep in water the character is immersed.
 	 * @return float in range 0.0 = not in water, 1.0 = fully immersed
@@ -948,16 +949,19 @@ public:
 	virtual float ImmersionDepth();
 
 	/** 
-	 * Calculates new velocity and acceleration for pawn for this tick, bounds acceleration and velocity, and adding effects of friction.
+	 * Updates Velocity and Acceleration based on the current state, applying the effects of friction and acceleration or deceleration. Does *not* apply gravity.
+	 * This is used internally during movement updates. Normally you don't need to call this from outside code, but you might want to use it for custom movement.
 	 *
 	 * @param	DeltaTime						time elapsed since last frame.
-	 * @param	Friction						coefficient of friction when not accelerating, or in the direction opposite of acceleration.
+	 * @param	Friction						coefficient of friction when not accelerating, or in the direction opposite acceleration.
 	 * @param	bFluid							true if moving through a fluid, causing Friction to always be applied regardless of acceleration.
 	 * @param	BrakingDeceleration				deceleration applied when not accelerating, or when exceeding max velocity.
 	 */
+	UFUNCTION(BlueprintCallable, Category="Pawn|Components|CharacterMovement")
 	virtual void CalcVelocity(float DeltaTime, float Friction, bool bFluid, float BrakingDeceleration);
 	
 	/** Compute the max jump height based on the JumpZVelocity velocity and gravity. */
+	UFUNCTION(BlueprintCallable, Category="Pawn|Components|CharacterMovement")
 	virtual float GetMaxJumpHeight() const;
 	
 	/** @return Maximum acceleration for the current state, based on MaxAcceleration and any additional modifiers. */
diff --git a/Engine/Source/Runtime/Engine/Private/CharacterMovementComponent.cpp b/Engine/Source/Runtime/Engine/Private/CharacterMovementComponent.cpp
--- a/Engine/Source/Runtime/Engine/Private/CharacterMovementComponent.cpp
+++ b/Engine/Source/Runtime/Engine/Private/CharacterMovementComponent.cpp
@@ -1860,9 +1860,9 @@ FVector UCharacterMovementComponent::AdjustUpperHemisphereImpact(const FVector&
 	return FVector(Delta.X, Delta.Y, Delta.Z * ZScale);
 }
 
-FVector UCharacterMovementComponent::NewFallVelocity(FVector OldVelocity, FVector OldAcceleration, float timeTick)
+FVector UCharacterMovementComponent::NewFallVelocity(FVector InitialVelocity, FVector Gravity, float DeltaTime)
 {
-	return OldVelocity + (OldAcceleration * timeTick);
+	return InitialVelocity + (Gravity * DeltaTime);
 }
 
 
@@ -1974,7 +1974,6 @@ void UCharacterMovementComponent::CalcVelocity(float DeltaTime, float Friction,
 		Acceleration = Acceleration.ClampMaxSize(MaxAccel);
 	}
 
-	// Update Velocity
 	if (bForceMaxAccel)
 	{
 		// Force acceleration at full speed.
