Fix for potential load order issue in TextRenderComponent. Font may not have post-loaded to cache transient data.

diff --git a/Engine/Source/Runtime/Engine/Private/TextRenderComponent.cpp b/Engine/Source/Runtime/Engine/Private/TextRenderComponent.cpp
--- a/Engine/Source/Runtime/Engine/Private/TextRenderComponent.cpp
+++ b/Engine/Source/Runtime/Engine/Private/TextRenderComponent.cpp
@@ -705,6 +705,7 @@ UTextRenderComponent::UTextRenderComponent(const class FPostConstructInitializeP
 
 	if( Font )
 	{
+		Font->ConditionalPostLoad();
 		WorldSize = Font->GetMaxCharHeight();
 		InvDefaultSize = 1.0f / WorldSize;
 	}
