Fixing Cooker crash when creating unnamed subobjects

diff --git a/Engine/Source/Runtime/Engine/Private/SimpleConstructionScript.cpp b/Engine/Source/Runtime/Engine/Private/SimpleConstructionScript.cpp
--- a/Engine/Source/Runtime/Engine/Private/SimpleConstructionScript.cpp
+++ b/Engine/Source/Runtime/Engine/Private/SimpleConstructionScript.cpp
@@ -615,8 +615,9 @@ USCS_Node* USimpleConstructionScript::CreateNode(UActorComponent* NewComponentTe
 {
 	if(NewComponentTemplate)
 	{
-		// Create a node for the script, and save a pointer to the template.
-		USCS_Node* NewNode = NewObject<USCS_Node>(this);
+		// Create a node for the script, and save a pointer to the template
+		// NewNamedObject to work around the fact we shouldn't use NewObject for default subobjects
+		USCS_Node* NewNode = NewNamedObject<USCS_Node>(this, NAME_None);
 		NewNode->SetFlags(RF_Transactional);
 		NewNode->ComponentTemplate = NewComponentTemplate;
 
