Stats - Fixed a comment about GHitchThresholdCVar

diff --git a/Engine/Source/Runtime/Engine/Private/UnrealEngine.cpp b/Engine/Source/Runtime/Engine/Private/UnrealEngine.cpp
--- a/Engine/Source/Runtime/Engine/Private/UnrealEngine.cpp
+++ b/Engine/Source/Runtime/Engine/Private/UnrealEngine.cpp
@@ -115,7 +115,7 @@ FSystemResolution GSystemResolution;
 static TAutoConsoleVariable<float> GHitchThresholdCVar(
 	TEXT("t.HitchThreshold"),
 	GHitchThreshold,
-	TEXT("Time in seconds that is considered a hitch by \"stat hitches\"")
+	TEXT("Time in seconds that is considered a hitch by \"stat dumphitches\"")
 	);
 
 static TAutoConsoleVariable<int32> CVarAllowOneFrameThreadLag(
