Editor: Add an adjustable knob to USplineComponent for how many points are used per spline segment when reparameterizing

diff --git a/Engine/Source/Runtime/Engine/Classes/Components/SplineComponent.h b/Engine/Source/Runtime/Engine/Classes/Components/SplineComponent.h
--- a/Engine/Source/Runtime/Engine/Classes/Components/SplineComponent.h
+++ b/Engine/Source/Runtime/Engine/Classes/Components/SplineComponent.h
@@ -19,6 +19,10 @@ class ENGINE_API USplineComponent : public USceneComponent
 	UPROPERTY()
 	FInterpCurveFloat SplineReparamTable;
 
+	/** Number of steps per spline segment to place in the reparameterization table */
+	UPROPERTY(EditAnywhere, AdvancedDisplay, Category = Spline, meta=(ClampMin=4, UIMin=4))
+	int32 ReparamStepsPerSegment;
+
 	/** Specifies the duration of the spline in seconds */
 	UPROPERTY(EditAnywhere, BlueprintReadWrite, Category=Spline)
 	float Duration;
@@ -92,8 +96,11 @@ class ENGINE_API USplineComponent : public USceneComponent
 	UFUNCTION(BlueprintCallable, Category=Spline)
 	FRotator GetWorldRotationAtTime(float Time, bool bUseConstantVelocity = false) const;
 
-
-
+#if WITH_EDITORONLY_DATA
+	// UObject interface
+	virtual void PostEditChangeProperty(FPropertyChangedEvent& PropertyChangedEvent) override;
+	// End of UObject interface
+#endif
 
 	/** Walk through keys and set time for each one */
 	void RefreshSplineInputs();
diff --git a/Engine/Source/Runtime/Engine/Private/SplineComponent.cpp b/Engine/Source/Runtime/Engine/Private/SplineComponent.cpp
--- a/Engine/Source/Runtime/Engine/Private/SplineComponent.cpp
+++ b/Engine/Source/Runtime/Engine/Private/SplineComponent.cpp
@@ -11,9 +11,9 @@
 
 USplineComponent::USplineComponent(const class FPostConstructInitializeProperties& PCIP)
 	: Super(PCIP)
+	, ReparamStepsPerSegment(4)
+	, Duration(1.0f)
 {
-	Duration = 1.0f;
-
 	// Add 2 keys by default
 	int32 PointIndex = SplineInfo.AddPoint(0.f, FVector(0,0,0));
 	SplineInfo.Points[PointIndex].InterpMode = CIM_CurveAuto;
@@ -37,7 +37,7 @@ void USplineComponent::UpdateSplineReparamTable()
 	}
 
 	const int32 NumSegments = SplineInfo.Points.Num() - 1;
-	const int32 NumSteps = 4 * NumSegments;
+	const int32 NumSteps = FMath::Max<int32>(ReparamStepsPerSegment, 4) * NumSegments;
 
 	// Build array of segment lengths
 	TArray<float> SegmentLengths;
@@ -54,9 +54,10 @@ void USplineComponent::UpdateSplineReparamTable()
 	// Now step through the spine at a constant distance interval, adding points to the reparam table
 	int32 CurrentSegment = 0;
 	float AccumulatedLength = 0.0f;
+	const float LengthPerStep = TotalLength / (float)NumSteps;
 	for (int32 Step = 0; Step < NumSteps; ++Step)
 	{
-		const float TargetLength = TotalLength * Step / NumSteps;
+		const float TargetLength = Step * LengthPerStep;
 
 		while (TargetLength - AccumulatedLength > SegmentLengths[CurrentSegment])
 		{
@@ -390,4 +391,20 @@ void USplineComponent::ApplyComponentInstanceData(TSharedPtr<FComponentInstanceD
 {
 	SplineInfo = StaticCastSharedPtr<FSplineInstanceData>(ComponentInstanceData)->SplineInfo;
 	UpdateSplineReparamTable();
-}
\ No newline at end of file
+}
+
+#if WITH_EDITORONLY_DATA
+void USplineComponent::PostEditChangeProperty(FPropertyChangedEvent& PropertyChangedEvent)
+{
+	if (PropertyChangedEvent.Property != nullptr)
+	{
+		const FName PropertyName(PropertyChangedEvent.Property->GetFName());
+		if (PropertyName == GET_MEMBER_NAME_CHECKED(USplineComponent, ReparamStepsPerSegment))
+		{
+			UpdateSplineReparamTable();
+		}
+	}
+
+	Super::PostEditChangeProperty(PropertyChangedEvent);
+}
+#endif
\ No newline at end of file
