First optimization pass for per poly CPU skinning

diff --git a/Engine/Source/Runtime/Engine/Classes/Components/SkinnedMeshComponent.h b/Engine/Source/Runtime/Engine/Classes/Components/SkinnedMeshComponent.h
--- a/Engine/Source/Runtime/Engine/Classes/Components/SkinnedMeshComponent.h
+++ b/Engine/Source/Runtime/Engine/Classes/Components/SkinnedMeshComponent.h
@@ -467,6 +467,13 @@ public:
 	 */
 	virtual FVector GetSkinnedVertexPosition(int32 VertexIndex) const;
 
+	/**
+	* CPU evaluation of the positions of all vertices (returned in component space)
+	*
+	* @param OutPositions buffer to place positions into
+	*/
+	virtual void USkinnedMeshComponent::ComputeSkinnedPositions(TArray<FVector> & OutPositions) const;
+
 	/**
 	* Returns color of the vertex.
 	*
@@ -832,6 +839,6 @@ private:
 	/** 
 	 * Simple, CPU evaluation of a vertex's skinned position helper function
 	 */
-	template <bool bExtraBoneInfluencesT>
-	FVector GetTypedSkinnedVertexPosition(const FSkelMeshChunk& Chunk, const FSkeletalMeshVertexBuffer& VertexBufferGPUSkin, int32 VertIndex, bool bSoftVertex) const;
+	template <bool bExtraBoneInfluencesT, bool bCachedMatrices = false>
+	FVector GetTypedSkinnedVertexPosition(const FSkelMeshChunk& Chunk, const FSkeletalMeshVertexBuffer& VertexBufferGPUSkin, int32 VertIndex, bool bSoftVertex, const TArray<FMatrix> & RefToLocals = TArray<FMatrix>()) const;
 };
diff --git a/Engine/Source/Runtime/Engine/Private/PhysicsEngine/PhysAnim.cpp b/Engine/Source/Runtime/Engine/Private/PhysicsEngine/PhysAnim.cpp
--- a/Engine/Source/Runtime/Engine/Private/PhysicsEngine/PhysAnim.cpp
+++ b/Engine/Source/Runtime/Engine/Private/PhysicsEngine/PhysAnim.cpp
@@ -241,6 +241,7 @@ void USkeletalMeshComponent::BlendInPhysics()
 }
 
 
+
 void USkeletalMeshComponent::UpdateKinematicBonesToPhysics(bool bTeleport, bool bNeedsSkinning)
 {
 	SCOPE_CYCLE_COUNTER(STAT_UpdateRBBones);
@@ -370,12 +371,20 @@ void USkeletalMeshComponent::UpdateKinematicBonesToPhysics(bool bTeleport, bool
 			{
 				const FStaticLODModel& Model = MeshObject->GetSkeletalMeshResource().LODModels[0];
 				TArray<FVector> NewPositions;
-				NewPositions.AddUninitialized(Model.NumVertices);
+				if (true)
 				{
 					SCOPE_CYCLE_COUNTER(STAT_SkinPerPolyVertices);
-					for (uint32 VertIndex = 0; VertIndex < Model.NumVertices; ++VertIndex)
+					ComputeSkinnedPositions(NewPositions);
+				}
+				else	//keep old way around for now - useful for comparing performance
+				{
+					NewPositions.AddUninitialized(Model.NumVertices);
 					{
-						NewPositions[VertIndex] = GetSkinnedVertexPosition(VertIndex);
+						SCOPE_CYCLE_COUNTER(STAT_SkinPerPolyVertices);
+						for (uint32 VertIndex = 0; VertIndex < Model.NumVertices; ++VertIndex)
+						{
+							NewPositions[VertIndex] = GetSkinnedVertexPosition(VertIndex);
+						}
 					}
 				}
 
diff --git a/Engine/Source/Runtime/Engine/Private/SkinnedMeshComponent.cpp b/Engine/Source/Runtime/Engine/Private/SkinnedMeshComponent.cpp
--- a/Engine/Source/Runtime/Engine/Private/SkinnedMeshComponent.cpp
+++ b/Engine/Source/Runtime/Engine/Private/SkinnedMeshComponent.cpp
@@ -1346,8 +1346,8 @@ void USkinnedMeshComponent::GetUsedMaterials( TArray<UMaterialInterface*>& OutMa
 	}
 }
 
-template <bool bExtraBoneInfluencesT>
-FORCEINLINE FVector USkinnedMeshComponent::GetTypedSkinnedVertexPosition(const FSkelMeshChunk& Chunk, const FSkeletalMeshVertexBuffer& VertexBufferGPUSkin, int32 VertIndex, bool bSoftVertex) const
+template <bool bExtraBoneInfluencesT, bool bCachedMatrices>
+FORCEINLINE FVector USkinnedMeshComponent::GetTypedSkinnedVertexPosition(const FSkelMeshChunk& Chunk, const FSkeletalMeshVertexBuffer& VertexBufferGPUSkin, int32 VertIndex, bool bSoftVertex, const TArray<FMatrix> & RefToLocals) const
 {
 	FVector SkinnedPos(0,0,0);
 
@@ -1372,10 +1372,19 @@ FORCEINLINE FVector USkinnedMeshComponent::GetTypedSkinnedVertexPosition(const F
 				BoneIndex = MasterBoneMap[BoneIndex];
 			}
 
-			const float		Weight = (float)SrcSoftVertex->InfluenceWeights[InfluenceIndex] / 255.0f;
-			const FMatrix	RefToLocal = SkeletalMesh->RefBasesInvMatrix[BoneIndex] * BaseComponent->SpaceBases[BoneIndex].ToMatrixWithScale();
-
-			SkinnedPos += RefToLocal.TransformPosition(VertexBufferGPUSkin.GetVertexPositionFast(SrcSoftVertex)) * Weight;
+			const float	Weight = (float)SrcSoftVertex->InfluenceWeights[InfluenceIndex] / 255.0f;
+			{
+				if (bCachedMatrices)
+				{
+					const FMatrix & RefToLocal = RefToLocals[BoneIndex];
+					SkinnedPos += RefToLocal.TransformPosition(VertexBufferGPUSkin.GetVertexPositionFast(SrcSoftVertex)) * Weight;
+				}
+				else
+				{
+					const FMatrix RefToLocal = SkeletalMesh->RefBasesInvMatrix[BoneIndex] * BaseComponent->SpaceBases[BoneIndex].ToMatrixWithScale();
+					SkinnedPos += RefToLocal.TransformPosition(VertexBufferGPUSkin.GetVertexPositionFast(SrcSoftVertex)) * Weight;
+				}
+			}
 		}
 	}
 	// Do rigid (one-influence) skinning for this vertex.
@@ -1390,9 +1399,16 @@ FORCEINLINE FVector USkinnedMeshComponent::GetTypedSkinnedVertexPosition(const F
 			BoneIndex = MasterBoneMap[BoneIndex];
 		}
 
-		const FMatrix	RefToLocal = SkeletalMesh->RefBasesInvMatrix[BoneIndex] * BaseComponent->SpaceBases[BoneIndex].ToMatrixWithScale();
-
-		SkinnedPos = RefToLocal.TransformPosition(VertexBufferGPUSkin.GetVertexPositionFast(SrcRigidVertex));
+		if (bCachedMatrices)
+		{
+			const FMatrix & RefToLocal = RefToLocals[BoneIndex];
+			SkinnedPos = RefToLocal.TransformPosition(VertexBufferGPUSkin.GetVertexPositionFast(SrcRigidVertex));
+		}
+		else
+		{
+			const FMatrix RefToLocal = SkeletalMesh->RefBasesInvMatrix[BoneIndex] * BaseComponent->SpaceBases[BoneIndex].ToMatrixWithScale();
+			SkinnedPos = RefToLocal.TransformPosition(VertexBufferGPUSkin.GetVertexPositionFast(SrcRigidVertex));
+		}
 	}
 
 	return SkinnedPos;
@@ -1410,19 +1426,79 @@ FVector USkinnedMeshComponent::GetSkinnedVertexPosition(int32 VertexIndex) const
 
 	FStaticLODModel& Model = MeshObject->GetSkeletalMeshResource().LODModels[0];
 
-	// Find the chunk and vertex within that chunk, and skinning type, for this vertex.
+	//cache RefToLocal matrices
 	int32 ChunkIndex;
 	int32 VertIndex;
 	bool bSoftVertex;
 	bool bHasExtraBoneInfluences;
 	Model.GetChunkAndSkinType(VertexIndex, ChunkIndex, VertIndex, bSoftVertex, bHasExtraBoneInfluences);
 
+	//update positions
 	check(ChunkIndex < Model.Chunks.Num());
 	const FSkelMeshChunk& Chunk = Model.Chunks[ChunkIndex];
+			//rigid
 
 	return bHasExtraBoneInfluences
 		? GetTypedSkinnedVertexPosition<true>(Chunk, Model.VertexBufferGPUSkin, VertIndex, bSoftVertex)
+			//soft
 		: GetTypedSkinnedVertexPosition<false>(Chunk, Model.VertexBufferGPUSkin, VertIndex, bSoftVertex);
+
+
+}
+
+void USkinnedMeshComponent::ComputeSkinnedPositions(TArray<FVector> & OutPositions) const
+{
+	OutPositions.Empty();
+
+	// Fail if no mesh
+	if (!SkeletalMesh || !MeshObject)
+	{
+		return;
+	}
+
+	FStaticLODModel& Model = MeshObject->GetSkeletalMeshResource().LODModels[0];
+	OutPositions.AddUninitialized(Model.NumVertices);
+
+	//cache RefToLocal matrices
+	const USkinnedMeshComponent* BaseComponent = MasterPoseComponent.IsValid() ? MasterPoseComponent.Get() : this;
+	TArray<FMatrix> RefToLocals;
+	RefToLocals.AddUninitialized(SkeletalMesh->RefBasesInvMatrix.Num());
+	{
+		for (int32 MatrixIdx = 0; MatrixIdx < RefToLocals.Num(); ++MatrixIdx)
+		{
+			RefToLocals[MatrixIdx] = SkeletalMesh->RefBasesInvMatrix[MatrixIdx] * BaseComponent->SpaceBases[MatrixIdx].ToMatrixWithScale();
+		}
+	}
+
+	//update positions
+	for (int32 ChunkIdx = 0; ChunkIdx < Model.Chunks.Num(); ++ChunkIdx)
+	{
+		const FSkelMeshChunk & Chunk = Model.Chunks[ChunkIdx];
+		const uint32 NumChunkVerts = Chunk.GetNumVertices();
+		bool bHasExtraBoneInfluences = Chunk.HasExtraBoneInfluences();
+		{
+			//rigid
+			const uint32 RigidOffset = Chunk.GetRigidVertexBufferIndex();
+			const uint32 NumRigidVerts = Chunk.GetNumRigidVertices();
+			for (uint32 RigidIdx = 0; RigidIdx < NumRigidVerts; ++RigidIdx)
+			{
+				FVector SkinnedPosition = bHasExtraBoneInfluences ? GetTypedSkinnedVertexPosition<true, true>(Chunk, Model.VertexBufferGPUSkin, RigidIdx, false, RefToLocals) :
+																	GetTypedSkinnedVertexPosition<false, true>(Chunk, Model.VertexBufferGPUSkin, RigidIdx, false, RefToLocals);
+				OutPositions[RigidOffset + RigidIdx] = SkinnedPosition;
+			}
+		}
+		{
+			//soft
+			const uint32 SoftOffset = Chunk.GetSoftVertexBufferIndex();
+			const uint32 NumSoftVerts = Chunk.GetNumSoftVertices();
+			for (uint32 SoftIdx = 0; SoftIdx < NumSoftVerts; ++SoftIdx)
+			{
+				FVector SkinnedPosition = bHasExtraBoneInfluences ? GetTypedSkinnedVertexPosition<true, true>(Chunk, Model.VertexBufferGPUSkin, SoftIdx, true, RefToLocals) :
+																	GetTypedSkinnedVertexPosition<false,true>(Chunk, Model.VertexBufferGPUSkin, SoftIdx, true, RefToLocals);
+				OutPositions[SoftOffset + SoftIdx] = SkinnedPosition;
+			}
+		}
+	}
 }
 
 FColor USkinnedMeshComponent::GetVertexColor(int32 VertexIndex) const
