check bodysetup is valid before trying to update mesh data

diff --git a/Engine/Source/Runtime/Engine/Private/PhysicsEngine/BodyInstance.cpp b/Engine/Source/Runtime/Engine/Private/PhysicsEngine/BodyInstance.cpp
--- a/Engine/Source/Runtime/Engine/Private/PhysicsEngine/BodyInstance.cpp
+++ b/Engine/Source/Runtime/Engine/Private/PhysicsEngine/BodyInstance.cpp
@@ -312,27 +312,28 @@ TArray<PxShape*> FBodyInstance::GetAllShapes(int32& OutNumSyncShapes) const
 void FBodyInstance::UpdateTriMeshVertices(const TArray<FVector> & NewPositions, const TArray<uint32> & Indices)
 {
 #if WITH_PHYSX
-	
-	SCOPED_SCENE_WRITE_LOCK((RigidActorSync ? RigidActorSync->getScene() : NULL));
-	SCOPED_SCENE_WRITE_LOCK((RigidActorAsync ? RigidActorAsync->getScene() : NULL));
+	if (BodySetup.IsValid())
+	{
+		SCOPED_SCENE_WRITE_LOCK((RigidActorSync ? RigidActorSync->getScene() : NULL));
+		SCOPED_SCENE_WRITE_LOCK((RigidActorAsync ? RigidActorAsync->getScene() : NULL));
 
-	BodySetup->UpdateTriMeshVertices(NewPositions, Indices);
+		BodySetup->UpdateTriMeshVertices(NewPositions, Indices);
 
-	//after updating the vertices we must call setGeometry again to update any shapes referencing the mesh
+		//after updating the vertices we must call setGeometry again to update any shapes referencing the mesh
 
-	int32 SyncShapeCount = 0;
-	TArray<PxShape *> PShapes = GetAllShapes(SyncShapeCount);
-	PxTriangleMeshGeometry PTriangleMeshGeometry;
-	for (int32 ShapeIdx = 0; ShapeIdx < PShapes.Num(); ShapeIdx++)
-	{
-		PxShape* PShape = PShapes[ShapeIdx];
-		if (PShape->getGeometryType() == PxGeometryType::eTRIANGLEMESH)
+		int32 SyncShapeCount = 0;
+		TArray<PxShape *> PShapes = GetAllShapes(SyncShapeCount);
+		PxTriangleMeshGeometry PTriangleMeshGeometry;
+		for (int32 ShapeIdx = 0; ShapeIdx < PShapes.Num(); ShapeIdx++)
 		{
-			PShape->getTriangleMeshGeometry(PTriangleMeshGeometry);
-			PShape->setGeometry(PTriangleMeshGeometry);
+			PxShape* PShape = PShapes[ShapeIdx];
+			if (PShape->getGeometryType() == PxGeometryType::eTRIANGLEMESH)
+			{
+				PShape->getTriangleMeshGeometry(PTriangleMeshGeometry);
+				PShape->setGeometry(PTriangleMeshGeometry);
+			}
 		}
 	}
-	
 #endif
 }
 
