#ue4 - PhysicsVolume terminal velocity is only applied in the gravity direction (assumed to be down) in CharacterMovement falling mode.

- Previously the terminal velocity was limiting the overall magnitude of the velocity which was incorrect. Now it limits the contribution of acceleration. Initial velocities over terminal velocity are allowed, but no further acceleration is applied.
- NewFallVelocity() is now const, and no longer exposed to blueprints (only was exposed yesterday).

TTP 338262

diff --git a/Engine/Source/Runtime/Engine/Classes/GameFramework/CharacterMovementComponent.h b/Engine/Source/Runtime/Engine/Classes/GameFramework/CharacterMovementComponent.h
--- a/Engine/Source/Runtime/Engine/Classes/GameFramework/CharacterMovementComponent.h
+++ b/Engine/Source/Runtime/Engine/Classes/GameFramework/CharacterMovementComponent.h
@@ -939,9 +939,11 @@ public:
 	/** Called if bNotifyApex is true and character has just passed the apex of its jump. */
 	virtual void NotifyJumpApex();
 
-	/** Compute new falling velocity from given velocity and gravity. */
-	UFUNCTION(BlueprintCallable, Category="Pawn|Components|CharacterMovement")
-	virtual FVector NewFallVelocity(FVector InitialVelocity, FVector Gravity, float DeltaTime);
+	/**
+	 * Compute new falling velocity from given velocity and gravity. Applies the limits of the current Physics Volume's TerminalVelocity.
+	 * Note: currently assumes Gravity points downward when enforcing terminal velocity.
+	 */
+	virtual FVector NewFallVelocity(FVector InitialVelocity, FVector Gravity, float DeltaTime) const;
 
 	/* Determine how deep in water the character is immersed.
 	 * @return float in range 0.0 = not in water, 1.0 = fully immersed
diff --git a/Engine/Source/Runtime/Engine/Private/CharacterMovementComponent.cpp b/Engine/Source/Runtime/Engine/Private/CharacterMovementComponent.cpp
--- a/Engine/Source/Runtime/Engine/Private/CharacterMovementComponent.cpp
+++ b/Engine/Source/Runtime/Engine/Private/CharacterMovementComponent.cpp
@@ -1860,9 +1860,20 @@ FVector UCharacterMovementComponent::AdjustUpperHemisphereImpact(const FVector&
 	return FVector(Delta.X, Delta.Y, Delta.Z * ZScale);
 }
 
-FVector UCharacterMovementComponent::NewFallVelocity(FVector InitialVelocity, FVector Gravity, float DeltaTime)
+
+FVector UCharacterMovementComponent::NewFallVelocity(FVector InitialVelocity, FVector Gravity, float DeltaTime) const
 {
-	return InitialVelocity + (Gravity * DeltaTime);
+	FVector Result = InitialVelocity;
+	const float TerminalLimit = FMath::Abs(GetPhysicsVolume()->TerminalVelocity);
+
+	// Only apply gravity if below terminal velocity, and don't exceed it if so.
+	if (InitialVelocity.Z > -TerminalLimit)
+	{
+		Result = InitialVelocity + (Gravity * DeltaTime);
+		Result.Z = FMath::Max(Result.Z, -TerminalLimit);
+	}
+	
+	return Result;
 }
 
 
@@ -2770,8 +2781,6 @@ void UCharacterMovementComponent::PhysFalling(float deltaTime, int32 Iterations)
 				Velocity.X = 0.f;
 				Velocity.Y = 0.f;
 			}
-
-			Velocity = Velocity.ClampMaxSize(GetPhysicsVolume()->TerminalVelocity);
 		}
 	}
 }
