fix build

diff --git a/Engine/Source/Runtime/Engine/Classes/Components/SkinnedMeshComponent.h b/Engine/Source/Runtime/Engine/Classes/Components/SkinnedMeshComponent.h
--- a/Engine/Source/Runtime/Engine/Classes/Components/SkinnedMeshComponent.h
+++ b/Engine/Source/Runtime/Engine/Classes/Components/SkinnedMeshComponent.h
@@ -472,7 +472,7 @@ public:
 	*
 	* @param OutPositions buffer to place positions into
 	*/
-	virtual void USkinnedMeshComponent::ComputeSkinnedPositions(TArray<FVector> & OutPositions) const;
+	virtual void ComputeSkinnedPositions(TArray<FVector> & OutPositions) const;
 
 	/**
 	* Returns color of the vertex.
