More cleanup (hopefully the bulk is done now)

* Combine pending list with cache list
* Clarify which function are server vs client
* Separate functions into more specific roles (server vs client), and call correct functions accordingly
  * Moved code out of function if only used in certain scenarios (client or server)
  * This simplifies the scope of each function, rather than having to solve everything

diff --git a/Engine/Source/Runtime/CoreUObject/Public/UObject/CoreNet.h b/Engine/Source/Runtime/CoreUObject/Public/UObject/CoreNet.h
--- a/Engine/Source/Runtime/CoreUObject/Public/UObject/CoreNet.h
+++ b/Engine/Source/Runtime/CoreUObject/Public/UObject/CoreNet.h
@@ -124,8 +124,7 @@ class COREUOBJECT_API UPackageMap : public UObject
 
 	virtual void		SetLocked( bool L ) { }
 
-	virtual UObject *	ResolvePathAndAssignNetGUID( FNetworkGUID & InOutNetGUID, const FString & PathName, UObject * ObjOuter, const bool bNoLoad = false ) { return NULL; }
-	virtual void		UpdatePendingResolveGUIDs() { }
+	virtual UObject *	ResolvePathAndAssignNetGUID( const FNetworkGUID & NetGUID, const FString & PathName ) { return NULL; }
 
 	virtual bool		SerializeNewActor(FArchive & Ar, class UActorChannel * Channel, class AActor *& Actor) { return false; }
 
diff --git a/Engine/Source/Runtime/Engine/Classes/Engine/PackageMapClient.h b/Engine/Source/Runtime/Engine/Classes/Engine/PackageMapClient.h
--- a/Engine/Source/Runtime/Engine/Classes/Engine/PackageMapClient.h
+++ b/Engine/Source/Runtime/Engine/Classes/Engine/PackageMapClient.h
@@ -22,26 +22,6 @@
 #include "Net/DataBunch.h"
 #include "PackageMapClient.generated.h"
 
-/**
- * A guid that needs to eventually be resolved
- * GUID's that need to be resolved, but can't yet due to:
- * 1. They are streaming levels, or levels in general
- * 	1a. We always expect clients to be in charge of actually loading levels, so we wait for that to happen before we resolve any guids within the level
- * 2. We are asynchronously loading the object we will assign the net guid to, so we need to wait for the load to finish
- * 	2a. This is work in progress, we don't have any official async loading code yet
- * 3. We don't currently do this, but if we were depending on a dynamic outer, but have a stable name (stably named sub-object) we would need to defer the resolve
- */
-class FPendingResolveGUID
-{
-public:
-	FPendingResolveGUID() : bNoLoad( false ), StartTimeInSeconds( 0 ) {}
-
-	FString			PathName;				// Path name of object to load for this guid
-	FNetworkGUID	OuterGUID;				// OuterGUID that this object depends on
-	bool			bNoLoad;				// Whether or not we should load this object, or just wait for it (we wait for levels for example, never load them in package map)
-	double			StartTimeInSeconds;		// When this guid was first deferred (in seconds)
-};
-
 class FMarkSerializeNewActorScoped
 {
 public:
@@ -62,8 +42,17 @@ public:
 class FNetGuidCacheObject
 {
 public:
-	TWeakObjectPtr< UObject >		Object;
-	FString							FullPath;
+	FNetGuidCacheObject() : bNoLoad( 0 ), bIgnoreWhenMissing( 0 ), bIsPending( 0 ) {}
+
+	TWeakObjectPtr< UObject >	Object;
+
+	// These fields are set when this guid is static
+	FNetworkGUID				OuterGUID;
+	FName						PathName;
+
+	uint8						bNoLoad				: 1;	// Don't load this, only do a find
+	uint8						bIgnoreWhenMissing	: 1;	// Don't warn when this asset can't be found or loaded
+	uint8						bIsPending			: 1;	// This object is waiting to be fully loaded
 };
 
 class ENGINE_API FNetGUIDCache
@@ -77,9 +66,14 @@ public:
 	bool			IsDynamicObject( const UObject * Object );
 	bool			IsNetGUIDAuthority();
 	FNetworkGUID	GetOrAssignNetGUID( const UObject * Object );
-	FNetworkGUID	AssignNewNetGUID( const UObject * Object );
-	void			AssignNetGUID( const UObject * Object, const FNetworkGUID & NewNetworkGUID );
+	FNetworkGUID	AssignNewNetGUID_Server( const UObject * Object );
+	void			RegisterNetGUID_Internal( const FNetworkGUID & NetGUID, const FNetGuidCacheObject & CacheObject );
+	void			RegisterNetGUID_Server( const FNetworkGUID & NetGUID, const UObject * Object );
+	void			RegisterNetGUID_Client( const FNetworkGUID & NetGUID, const UObject * Object );
+	void			RegisterNetGUIDFromPath_Client( const FNetworkGUID & NetGUID, const FString & PathName, const FNetworkGUID & OuterGUID, const bool bNoLoad, const bool bIgnoreWhenMissing );
 	UObject *		GetObjectFromNetGUID( const FNetworkGUID & NetGUID );
+	bool			ShouldIgnoreWhenMissing( const FNetworkGUID & NetGUID );
+	UObject *		ResolvePath( const FString & PathName, UObject * ObjOuter, const bool bNoLoad );
 
 	TMap< FNetworkGUID, FNetGuidCacheObject >		ObjectLookup;
 	TMap< TWeakObjectPtr< UObject >, FNetworkGUID >	NetGUIDLookup;
@@ -162,12 +156,7 @@ protected:
 	void			InternalWriteObject( FArchive& Ar, FNetworkGUID NetGUID, const UObject * Object, FString ObjectPathName, UObject * ObjectOuter );	
 	FNetworkGUID	InternalLoadObject( FArchive & Ar, UObject *& Object, const bool bIsInnerLevel, int InternalLoadObjectRecursionCount );
 
-	void		ResolvePathAndAssignNetGUID_Deferred( const FNetworkGUID & NetGUID, const FString & PathName, const FNetworkGUID & OuterGUID, const bool bNoLoad = false );
-	UObject *	ResolvePathAndAssignNetGUID( FNetworkGUID & InOutNetGUID, const FString & PathName, const FNetworkGUID & OuterGUID, const bool bNoLoad = false );
-
-	virtual UObject * ResolvePathAndAssignNetGUID( FNetworkGUID & InOutNetGUID, const FString & PathName, UObject * ObjOuter, const bool bNoLoad = false ) OVERRIDE;
-
-	virtual void UpdatePendingResolveGUIDs() override;
+	virtual UObject * ResolvePathAndAssignNetGUID( const FNetworkGUID & NetGUID, const FString & PathName ) OVERRIDE;
 
 	bool	ShouldSendFullPath(const UObject* Object, const FNetworkGUID &NetGUID);
 	
@@ -180,21 +169,18 @@ protected:
 
 	bool ObjectLevelHasFinishedLoading(UObject* Obj);
 
-	TSet< FNetworkGUID >						CurrentExportNetGUIDs;				// Current list of NetGUIDs being written to the Export Bunch.
+	TSet< FNetworkGUID >				CurrentExportNetGUIDs;				// Current list of NetGUIDs being written to the Export Bunch.
 
-	TMap< FNetworkGUID, int32 >					NetGUIDAckStatus;
-	TArray< FNetworkGUID >						PendingAckGUIDs;					// Quick access to all GUID's that haven't been acked
+	TMap< FNetworkGUID, int32 >			NetGUIDAckStatus;
+	TArray< FNetworkGUID >				PendingAckGUIDs;					// Quick access to all GUID's that haven't been acked
 
 	// Bunches of NetGUID/path tables to send with the current content bunch
-	TArray<FOutBunch* >							ExportBunches;
-	FOutBunch *									CurrentExportBunch;
-
-	int32										ExportNetGUIDCount;
+	TArray<FOutBunch* >					ExportBunches;
+	FOutBunch *							CurrentExportBunch;
 
-	// A guid that needs to eventually be resolved
-	TMap< FNetworkGUID, FPendingResolveGUID >	PendingResolveGUIDs;
+	int32								ExportNetGUIDCount;
 
-	TSharedPtr< FNetGUIDCache >					GuidCache;
+	TSharedPtr< FNetGUIDCache >			GuidCache;
 
-	bool										IsSerializingNewActor;
+	bool								IsSerializingNewActor;
 };
diff --git a/Engine/Source/Runtime/Engine/Private/DataChannel.cpp b/Engine/Source/Runtime/Engine/Private/DataChannel.cpp
--- a/Engine/Source/Runtime/Engine/Private/DataChannel.cpp
+++ b/Engine/Source/Runtime/Engine/Private/DataChannel.cpp
@@ -1997,7 +1997,7 @@ UObject * UActorChannel::ReadContentBlockHeader( FInBunch & Bunch )
 		SubObj = ConstructObject< UObject >( SubObjClass, Actor );
 		check( SubObj != NULL );
 		Actor->OnSubobjectCreatedFromReplication( SubObj );
-		Connection->Driver->GuidCache->AssignNetGUID( SubObj, NetGUID );
+		Connection->Driver->GuidCache->RegisterNetGUID_Client( NetGUID, SubObj );
 	}
 
 	check( Cast< AActor >( SubObj ) == NULL );
@@ -2074,7 +2074,7 @@ bool UActorChannel::ReplicateSubobject(UObject *Obj, FOutBunch &Bunch, const FRe
 	// Once we can lazily handle unmapped references on the client side, this can be simplified.
 	if ( !Connection->Driver->GuidCache->SupportsObject( Obj ) )
 	{
-		FNetworkGUID NetGUID = Connection->Driver->GuidCache->AssignNewNetGUID( Obj );	//Make sure he gets a NetGUID so that he is now 'supported'
+		FNetworkGUID NetGUID = Connection->Driver->GuidCache->AssignNewNetGUID_Server( Obj );	//Make sure he gets a NetGUID so that he is now 'supported'
 	}
 
 	bool NewSubobject = false;
diff --git a/Engine/Source/Runtime/Engine/Private/NetworkDriver.cpp b/Engine/Source/Runtime/Engine/Private/NetworkDriver.cpp
--- a/Engine/Source/Runtime/Engine/Private/NetworkDriver.cpp
+++ b/Engine/Source/Runtime/Engine/Private/NetworkDriver.cpp
@@ -303,14 +303,6 @@ void UNetDriver::TickFlush(float DeltaSeconds)
 		DrawNetDriverDebug();
 	}
 
-	// Update any pending resolves on the server package map
-	// We shouldn't need to do this as the server, since clients should only be sending us names that resolve immediately
-	// Do this before we try to resolve unmapped properties, since this may reduce the number of unmapped properties
-	if ( ServerConnection != NULL && ServerConnection->PackageMap != NULL )
-	{
-		ServerConnection->PackageMap->UpdatePendingResolveGUIDs();
-	}
-
 	// Update properties that are unmapped, try to hook up the object pointers if they exist now
 	for ( auto It = UnmappedReplicators.CreateIterator(); It; ++It )
 	{
diff --git a/Engine/Source/Runtime/Engine/Private/PackageMapClient.cpp b/Engine/Source/Runtime/Engine/Private/PackageMapClient.cpp
--- a/Engine/Source/Runtime/Engine/Private/PackageMapClient.cpp
+++ b/Engine/Source/Runtime/Engine/Private/PackageMapClient.cpp
@@ -184,11 +184,6 @@ bool UPackageMapClient::SerializeNewActor(FArchive& Ar, class UActorChannel *Cha
 	UObject *NewObj = Actor;
 	SerializeObject(Ar, AActor::StaticClass(), NewObj, &NetGUID);
 
-	if ( PendingResolveGUIDs.Contains( NetGUID ) )
-	{
-		UE_LOG( LogNetPackageMap, Error, TEXT( "UPackageMapClient::SerializeNewActor. Unresolved Actor GUID: Path: %s, NetGUID: %s." ), *PendingResolveGUIDs.FindChecked( NetGUID ).PathName, *NetGUID.ToString() );
-	}
-
 	if ( Ar.IsError() )
 	{
 		UE_LOG( LogNetPackageMap, Error, TEXT( "UPackageMapClient::SerializeNewActor: Ar.IsError after SerializeObject 1" ) );
@@ -227,6 +222,7 @@ bool UPackageMapClient::SerializeNewActor(FArchive& Ar, class UActorChannel *Cha
 		{
 			Archetype = Actor->GetArchetype();
 
+			check( Archetype != NULL );
 			check( Actor->NeedsLoadForClient() );			// We have no business sending this unless the client can load
 			check( Archetype->NeedsLoadForClient() );		// We have no business sending this unless the client can load
 
@@ -236,11 +232,20 @@ bool UPackageMapClient::SerializeNewActor(FArchive& Ar, class UActorChannel *Cha
 		}
 
 		FNetworkGUID ArchetypeNetGUID;
-		SerializeObject(Ar, UObject::StaticClass(), Archetype, &ArchetypeNetGUID );
+		SerializeObject( Ar, UObject::StaticClass(), Archetype, &ArchetypeNetGUID );
 
-		if ( PendingResolveGUIDs.Contains( ArchetypeNetGUID ) )
+		if ( ArchetypeNetGUID.IsValid() && Archetype == NULL )
 		{
-			UE_LOG( LogNetPackageMap, Error, TEXT( "UPackageMapClient::SerializeNewActor. Unresolved Archetype GUID: Path: %s, NetGUID: %s." ), *PendingResolveGUIDs.FindChecked( ArchetypeNetGUID ).PathName, *ArchetypeNetGUID.ToString() );
+			const FNetGuidCacheObject * ExistingCacheObjectPtr = GuidCache->ObjectLookup.Find( ArchetypeNetGUID );
+
+			if ( ExistingCacheObjectPtr != NULL )
+			{
+				UE_LOG( LogNetPackageMap, Error, TEXT( "UPackageMapClient::SerializeNewActor. Unresolved Archetype GUID. Path: %s, NetGUID: %s." ), *ExistingCacheObjectPtr->PathName.ToString(), *ArchetypeNetGUID.ToString() );
+			}
+			else
+			{
+				UE_LOG( LogNetPackageMap, Error, TEXT( "UPackageMapClient::SerializeNewActor. Unresolved Archetype GUID. Guid not registered! NetGUID: %s." ), *ArchetypeNetGUID.ToString() );
+			}
 		}
 
 		// SerializeCompressedInitial
@@ -263,12 +268,12 @@ bool UPackageMapClient::SerializeNewActor(FArchive& Ar, class UActorChannel *Cha
 			}
 		}
 
-		if (Ar.IsLoading())
+		if ( Ar.IsLoading() )
 		{
 			// Spawn actor if necessary (we may have already found it if it was dormant)
-			if (Actor == NULL)
+			if ( Actor == NULL )
 			{
-				if (Archetype)
+				if ( Archetype )
 				{
 					FActorSpawnParameters SpawnInfo;
 					SpawnInfo.Template = Cast<AActor>(Archetype);
@@ -278,7 +283,7 @@ bool UPackageMapClient::SerializeNewActor(FArchive& Ar, class UActorChannel *Cha
 					Actor = Connection->Driver->GetWorld()->SpawnActor(Archetype->GetClass(), &Location, &Rotation, SpawnInfo );
 					Actor->PostNetReceiveVelocity(Velocity);
 
-					GuidCache->AssignNetGUID(Actor, NetGUID);
+					GuidCache->RegisterNetGUID_Client( NetGUID, Actor );
 					SpawnedActor = true;
 				}
 				else
@@ -366,7 +371,7 @@ bool UPackageMapClient::SerializeNewActor(FArchive& Ar, class UActorChannel *Cha
 						}
 					}
 
-					GuidCache->AssignNetGUID( Subobjs[i], SubObjGuids[i] );
+					GuidCache->RegisterNetGUID_Client( SubObjGuids[i], Subobjs[i] );
 
 					UE_LOG( LogNetPackageMap, Log, TEXT( "SerializeNewActor: [Loading] Assigned NetGUID <%s> to subobject %s" ), *SubObjGuids[i].ToString(), *Subobjs[i]->GetPathName() );
 				}
@@ -540,14 +545,6 @@ FNetworkGUID UPackageMapClient::InternalLoadObject( FArchive & Ar, UObject *& Ob
 		return NetGUID;
 	}
 
-	if ( PendingResolveGUIDs.Contains( NetGUID ) )
-	{
-		// If this was previously on the deferred resolve list, go ahead and take this opportunity to 
-		// resolve any pending guids now (which may or may not resolve this guid)
-		// If this guid is resolved, it will also try and resolve any guids that have this guid as an outer as well
-		UpdatePendingResolveGUIDs();
-	}
-
 	// ----------------	
 	// Try to resolve NetGUID
 	// ----------------	
@@ -594,19 +591,9 @@ FNetworkGUID UPackageMapClient::InternalLoadObject( FArchive & Ar, UObject *& Ob
 			return NetGUID;
 		}
 
+		// Remap name for PIE
 		GEngine->NetworkRemapPath( Connection->Driver->GetWorld(), PathName, true );
 
-		// Make note if the outer is still deferred (we then must defer the load of this guid)
-		const bool OuterIsDeferred = PendingResolveGUIDs.Contains( OuterGUID );
-
-#if STRESS_PENDING_GUIDS == 0
-		if ( Object == NULL && !OuterIsDeferred )
-		{
-			// Try to resolve the name (we may have already loaded it), but DON'T load yet if we haven't already
-			Object = ResolvePathAndAssignNetGUID( NetGUID, PathName, OuterGUID, true );
-		}
-#endif
-
 		if ( Object != NULL )
 		{
 			// If we already have the object, just do some sanity checking and return
@@ -614,210 +601,84 @@ FNetworkGUID UPackageMapClient::InternalLoadObject( FArchive & Ar, UObject *& Ob
 			return NetGUID;
 		}
 
-		const bool bIsForLevel = ExportFlags.bIsLevel || bIsInnerLevel;
-
-#if STRESS_PENDING_GUIDS == 1
-		// For testing, defer everything
-		const bool bShouldDefer = !NetGUID.IsDefault();
-#else
-		// If this is a level (or any of our inners are) or our outer is being deferred, then defer this guid
-		const bool bShouldDefer = !NetGUID.IsDefault() && ( bIsForLevel || OuterIsDeferred );
-#endif
-
-		if ( bShouldDefer )
-		{
-			UE_LOG( LogNetPackageMap, Log, TEXT( "InternalLoadObject: Deferring resolve. Path: %s, NetGUID: %s, Outer: %s" ), *PathName, *NetGUID.ToString(), ObjOuter != NULL ? *ObjOuter->GetFullName() : TEXT( "NULL" ) );
-			ResolvePathAndAssignNetGUID_Deferred( NetGUID, PathName, OuterGUID, bIsForLevel );
-		}
-		else
+		if ( NetGUID.IsDefault() )
 		{
-			Object = ResolvePathAndAssignNetGUID( NetGUID, PathName, OuterGUID );
-		}
-	}
-
-	if ( Object == NULL && NetGUID.IsStatic() && !PendingResolveGUIDs.Contains( NetGUID ) )
-	{
-		UE_LOG( LogNetPackageMap, Warning, TEXT( "InternalLoadObject: Unable to resolve object. NetGUID: %s, HasPath: %s" ), *NetGUID.ToString(), ExportFlags.bHasPath ? TEXT( "TRUE" ) : TEXT( "FALSE" ) );
-	}
-
-	return NetGUID;
-}
-
-void UPackageMapClient::ResolvePathAndAssignNetGUID_Deferred( const FNetworkGUID & NetGUID, const FString & PathName, const FNetworkGUID & OuterGUID, const bool bNoLoad )
-{
-	check( !NetGUID.IsDefault() );		// We can't defer default guids, doesn't make sense
-
-	if ( !PendingResolveGUIDs.Contains( NetGUID ) )
-	{
-		FPendingResolveGUID PendingResolveGUID;
-
-		PendingResolveGUID.OuterGUID			= OuterGUID;
-		PendingResolveGUID.PathName				= PathName;
-		PendingResolveGUID.bNoLoad				= bNoLoad;
-		PendingResolveGUID.StartTimeInSeconds	= FPlatformTime::Seconds();
-
-		PendingResolveGUIDs.Add( NetGUID, PendingResolveGUID );
-	}
-	else
-	{
-		// Sanity check
-		FPendingResolveGUID & PendingResolveGUID = PendingResolveGUIDs.FindChecked( NetGUID );
-
-		if ( PendingResolveGUID.OuterGUID != OuterGUID )
-		{
-			UE_LOG( LogNetPackageMap, Warning, TEXT( "ResolvePathAndAssignNetGUID_Deferred: Deferred Outer mismatch. NetGUID: %s, Path: %s, Current: %s, Expected: %s" ), *NetGUID.ToString(), *PathName, *PendingResolveGUID.OuterGUID.ToString(), *OuterGUID.ToString() );
-		}
-	}
-}
-
-void UPackageMapClient::UpdatePendingResolveGUIDs()
-{
-#if STRESS_PENDING_GUIDS == 1
-	const double CurrentTimeInSeconds = FPlatformTime::Seconds();
-#endif
-
-	bool NeedToKeepResolving = true;
-
-	while ( NeedToKeepResolving )
-	{
-		NeedToKeepResolving = false;
-
-		for ( auto It = PendingResolveGUIDs.CreateIterator(); It; ++It )
-		{
-			const FPendingResolveGUID & PendingResolveGUID = It.Value();
-
-#if STRESS_PENDING_GUIDS == 1
-			if ( !NeedToKeepResolving && !IsSerializingNewActor && ( CurrentTimeInSeconds - PendingResolveGUID.StartTimeInSeconds < 15.0f ) )
-			{
-				continue;		// Wait 15 seconds for testing (unless we are in SerializeNewActor, which can't handle deferred guids yet
-			}
-#endif
-
-			if ( PendingResolveGUIDs.Contains( PendingResolveGUID.OuterGUID ) )
-			{
-				// We can't load yet if our outer hasn't been resolved yet
-				continue;
-			}
+			// This should be from the client
+			// If we get here, we want to go ahead and assign a network guid, 
+			// then export that to the client at the next available opportunity
+			check( IsNetGUIDAuthority() );
 
-			// First, see if we have an object for this GUID already
-			UObject * Object = GetObjectFromNetGUID( It.Key() );
+			Object = GuidCache->ResolvePath( PathName, ObjOuter, false );
 
 			if ( Object == NULL )
 			{
-				// If not, try to find the object by path
-				FNetworkGUID NetGUID = It.Key();
-				Object = ResolvePathAndAssignNetGUID( NetGUID, PendingResolveGUID.PathName, PendingResolveGUID.OuterGUID, PendingResolveGUID.bNoLoad );
+				UE_LOG( LogNetPackageMap, Warning, TEXT( "UPackageMapClient::InternalLoadObject: Unable to resolve default guid from client: PathName: %s, ObjOuter: %s " ), *PathName, ObjOuter != NULL ? *ObjOuter->GetPathName() : TEXT( "NULL" ) );
+				return NetGUID;
 			}
 
-			if ( Object != NULL )
-			{
-				UE_LOG( LogNetPackageMap, Log, TEXT( "UPackageMapClient::UpdatePendingResolveGUIDs. Resolved deferred NetGUID. Path: %s, NetGUID: %s. NumLeft: %i, IsSerializingNewActor: %s" ), *PendingResolveGUID.PathName, *It.Key().ToString(), PendingResolveGUIDs.Num() - 1, IsSerializingNewActor ? TEXT("YES") : TEXT("NO") );
-
-				// If we have an object, we can make another pass in case there were dependencies on this guid loading
-				NeedToKeepResolving = true;
-				It.RemoveCurrent();		// We don't need to check this one anymore
-			}
-		}
-	}
-}
+			// Assign the guid to the object
+			NetGUID = GuidCache->GetOrAssignNetGUID( Object );
 
-UObject * UPackageMapClient::ResolvePathAndAssignNetGUID( FNetworkGUID & InOutNetGUID, const FString & PathName, UObject * ObjOuter, const bool bNoLoad )
-{
-	if ( InOutNetGUID.IsDynamic() )
-	{
-		UE_LOG( LogNetPackageMap, Warning, TEXT( "UPackageMapClient::ResolvePathAndAssignNetGUID. Trying to load dynamic NetGUID by path: <%s,%s>. Outer: %s" ), *InOutNetGUID.ToString(), *PathName, ObjOuter ? *ObjOuter->GetName() : TEXT( "NULL" ) );
-	}
+			// Let this client know what guid we assigned
+			HandleUnAssignedObject( Object );
 
-	// We don't have the object mapped yet
-	UObject * Object = StaticFindObject( UObject::StaticClass(), ObjOuter, *PathName, false );
+			return NetGUID;
+		}
 
-	if ( !Object && !bNoLoad )
-	{
+		// If we are the server, we should have found the object by now
 		if ( IsNetGUIDAuthority() )
 		{
-			// The authority shouldn't be loading resources on demand, at least for now.
-			// This could be garbage or a security risk
-			// Another possibility is in dynamic property replication if the server reads the previously serialized state
-			// that has a now destroyed actor in it.
-			UE_LOG( LogNetPackageMap, Warning, TEXT( "UPackageMapClient::ResolvePathAndAssignNetGUID: Could not find Object for: NetGUID <%s, %s> (and IsNetGUIDAuthority() == false)" ), *InOutNetGUID.ToString(), *PathName );
-			return NULL;
-		}
-		else
-		{
-			UE_LOG( LogNetPackageMap, Log, TEXT( "UPackageMapClient::ResolvePathAndAssignNetGUID: Could not find Object for: NetGUID <%s, %s>" ), *InOutNetGUID.ToString(), *PathName );
+			UE_LOG( LogNetPackageMap, Warning, TEXT( "UPackageMapClient::InternalLoadObject: Server could not resolve non default guid from client. PathName: %s, ObjOuter: %s " ), *PathName, ObjOuter != NULL ? *ObjOuter->GetPathName() : TEXT( "NULL" ) );
+			return NetGUID;
 		}
 
-		Object = StaticLoadObject( UObject::StaticClass(), ObjOuter, *PathName, NULL, LOAD_NoWarn );
+		// At this point, only the client gets this far
+		const bool bIsForLevel			= ExportFlags.bIsLevel || bIsInnerLevel;			// Make note if this is for a level (or any of our inners are)
+		const bool bNoLoad				= bIsForLevel;
+		const bool bIgnoreWhenMissing	= bIsForLevel || GuidCache->ShouldIgnoreWhenMissing( OuterGUID );
 
-		if ( Object )
-		{
-			UE_LOG( LogNetPackageMap, Log, TEXT( "UPackageMapClient::ResolvePathAndAssignNetGUID: StaticLoadObject. Found: %s" ), Object ? *Object->GetName() : TEXT( "NULL" ) );
-		}
-		else
-		{
-			//
-			// If we failed to load it as an object, try to load it as a package
-			//
+		// Register this path and outer guid combo with the net guid
+		GuidCache->RegisterNetGUIDFromPath_Client( NetGUID, PathName, OuterGUID, bNoLoad, bIgnoreWhenMissing );
 
-			UPackage * PackageOuter = Cast< UPackage >( ObjOuter );
+		// Try again now that we've registered the path
+		Object = GuidCache->GetObjectFromNetGUID( NetGUID );
 
-			// If we have an outer at this point, it only makes sense for it to be a package (or NULL)
-			if ( PackageOuter != NULL || ObjOuter == NULL )
-			{
-				Object = LoadPackage( PackageOuter, *PathName, LOAD_None );
-				UE_LOG( LogNetPackageMap, Log, TEXT( "UPackageMapClient::ResolvePathAndAssignNetGUID: LoadPackage: %s. Found: %s" ), *PathName, Object ? *Object->GetName() : TEXT( "NULL" ) );
-			} 
-			else
-			{
-				UE_LOG( LogNetPackageMap, Warning, TEXT( "UPackageMapClient::ResolvePathAndAssignNetGUID: Outer %s is not a package for object %s" ), ObjOuter ? *ObjOuter->GetName() : TEXT( "NULL" ), *PathName );
-			}
+		if ( Object == NULL && !GuidCache->ShouldIgnoreWhenMissing( NetGUID ) )
+		{
+			UE_LOG( LogNetPackageMap, Warning, TEXT( "InternalLoadObject: Unable to resolve object from path. Path: %s, Outer: %s, NetGUID: %s" ), *PathName, ObjOuter ? *ObjOuter->GetPathName() : TEXT( "NULL" ), *NetGUID.ToString() );
 		}
 	}
-
-	if ( Object )
+	else if ( Object == NULL && NetGUID.IsStatic() && !GuidCache->ShouldIgnoreWhenMissing( NetGUID ) )
 	{
-		if ( InOutNetGUID.IsDefault() )
-		{
-			// This should be from the client
-			// If we get here, we want to go ahead and assign a network guid, then export that to the client 
-			// on the next available opportunity
-			check( IsNetGUIDAuthority() );
-			InOutNetGUID = GuidCache->GetOrAssignNetGUID( Object );
-			HandleUnAssignedObject( Object );
-		}
-		else
-		{
-			GuidCache->AssignNetGUID( Object, InOutNetGUID );
-		}
+		UE_LOG( LogNetPackageMap, Warning, TEXT( "InternalLoadObject: Unable to resolve object. NetGUID: %s, HasPath: %s" ), *NetGUID.ToString(), ExportFlags.bHasPath ? TEXT( "TRUE" ) : TEXT( "FALSE" ) );
 	}
 
-	return Object;
+	return NetGUID;
 }
 
-UObject * UPackageMapClient::ResolvePathAndAssignNetGUID( FNetworkGUID & InOutNetGUID, const FString & PathName, const FNetworkGUID & OuterGUID, const bool bNoLoad )
+UObject * UPackageMapClient::ResolvePathAndAssignNetGUID( const FNetworkGUID & NetGUID, const FString & PathName )
 {
-	UObject * ObjOuter = NULL;
-
-	// Resolve the outer if the guid is valid
-	if ( OuterGUID.IsValid() )
+	if ( NetGUID.IsDynamic() )
 	{
-		ObjOuter = GetObjectFromNetGUID( OuterGUID );
+		UE_LOG( LogNetPackageMap, Warning, TEXT( "UPackageMapClient::ResolvePathAndAssignNetGUID. Trying to load dynamic NetGUID by path. Path: %s, NetGUID: %s" ), *PathName, *NetGUID.ToString() );
+	}
 
-		// If we have a valid outer guid, we expect it to resolve
-		if ( ObjOuter == NULL )
-		{
-			UE_LOG( LogNetPackageMap, Warning, TEXT( "ResolvePathAndAssignNetGUID: Unable resolve outer. Path: %s, NetGUID: %s, OuterGUID: %s " ), *PathName, *InOutNetGUID.ToString(), *OuterGUID.ToString() );
-			return NULL;
-		}
+	if ( NetGUID.IsDefault() )
+	{
+		UE_LOG( LogNetPackageMap, Warning, TEXT( "UPackageMapClient::ResolvePathAndAssignNetGUID. Trying to load default NetGUID by path. Path: %s, NetGUID: %s" ), *PathName, *NetGUID.ToString() );
+		return NULL;
 	}
 
-	UObject * Object = ResolvePathAndAssignNetGUID( InOutNetGUID, PathName, ObjOuter, bNoLoad );
+	UObject * Object = GuidCache->ResolvePath( PathName, NULL, false );
 
-	if ( Object == NULL && !bNoLoad )
+	if ( Object == NULL )
 	{
-		UE_LOG( LogNetPackageMap, Warning, TEXT( "ResolvePathAndAssignNetGUID: Unable to resolve object. Path: %s, NetGUID: %s, Outer: %s" ), *PathName, *InOutNetGUID.ToString(), ObjOuter != NULL ? *ObjOuter->GetPathName() : TEXT( "NULL" ) );
+		UE_LOG( LogNetPackageMap, Warning, TEXT( "UPackageMapClient::ResolvePathAndAssignNetGUID. Could not resolve path. Path: %s, NetGUID: %s" ), *PathName, *NetGUID.ToString() );
+		return NULL;
 	}
 
+	GuidCache->RegisterNetGUID_Client( NetGUID, Object );
+
 	return Object;
 }
 
@@ -1424,14 +1285,11 @@ void FNetGUIDCache::CleanReferences()
 	{
 		if ( It.Value().Object.IsValid() )
 		{
-			// Make sure the path matches (which can change during seamless travel)
-			It.Value().FullPath = It.Value().Object->GetPathName();
 			continue;
 		}
 
 		if ( It.Key().IsDynamic() )
 		{
-			UE_LOG( LogNetPackageMap, Log, TEXT( "Cleaning reference to NetGUID: <%s> (%s) (ObjectLookup)" ), *It.Key().ToString(), *It.Value().FullPath );
 			It.RemoveCurrent();
 		}
 	}
@@ -1440,7 +1298,6 @@ void FNetGUIDCache::CleanReferences()
 	{
 		if ( !It.Key().IsValid() )
 		{
-			UE_LOG( LogNetPackageMap, Log, TEXT( "Cleaning reference to NetGUID: <%s> (NetGUIDLookup)" ), *It.Value().ToString() );
 			It.RemoveCurrent();
 		}
 	}
@@ -1550,7 +1407,15 @@ FNetworkGUID FNetGUIDCache::GetOrAssignNetGUID( const UObject * Object )
 		return NetGUID;
 	}
 
-	return AssignNewNetGUID( Object );
+	if ( !IsNetGUIDAuthority() )
+	{
+		// We cannot make or assign new NetGUIDs
+		// Generate a default GUID, which signifies we write the full path
+		// The server should detect this, and assign a full-time guid, and send that back to us
+		return FNetworkGUID::GetDefault();
+	}
+
+	return AssignNewNetGUID_Server( Object );
 }
 
 static const AActor * GetOuterActor( const UObject * Object )
@@ -1575,15 +1440,9 @@ static const AActor * GetOuterActor( const UObject * Object )
 /**
  *	Generate a new NetGUID for this object and assign it.
  */
-FNetworkGUID FNetGUIDCache::AssignNewNetGUID( const UObject * Object )
+FNetworkGUID FNetGUIDCache::AssignNewNetGUID_Server( const UObject * Object )
 {
-	if ( !IsNetGUIDAuthority() )
-	{
-		// We cannot make or assign new NetGUIDs
-		// Generate a default GUID, which signifies we write the full path
-		// The server should detect this, and assign a full-time guid, and send that back to us
-		return FNetworkGUID::GetDefault();
-	}
+	check( IsNetGUIDAuthority() );
 
 	const AActor * Actor = Cast< const AActor >( Object );
 
@@ -1606,7 +1465,7 @@ FNetworkGUID FNetGUIDCache::AssignNewNetGUID( const UObject * Object )
 				check( !NetGUIDLookup.Contains( OuterActor ) );
 
 				// Assign our owning actor a guid first (which should also assign our net guid if things are working correctly)
-				AssignNewNetGUID( OuterActor );
+				AssignNewNetGUID_Server( OuterActor );
 
 				return NetGUIDLookup.FindChecked( Object );
 			}
@@ -1620,7 +1479,7 @@ FNetworkGUID FNetGUIDCache::AssignNewNetGUID( const UObject * Object )
 
 	const FNetworkGUID NewNetGuid( ALLOC_NEW_NET_GUID( IsStatic ) );
 
-	AssignNetGUID( Object, NewNetGuid );
+	RegisterNetGUID_Server( NewNetGuid, Object );
 
 	// If this is a dynamic actor, assign all of our stably named sub-objects their net guids now as well (they come as a group)
 	if ( Actor != NULL && !IsStatic )
@@ -1637,89 +1496,180 @@ FNetworkGUID FNetGUIDCache::AssignNewNetGUID( const UObject * Object )
 
 			if ( NetGUIDLookup.Contains( SubObj ) )		// We should NOT have a guid yet
 			{
-				UE_LOG( LogNetPackageMap, Error, TEXT( "AssignNewNetGUID: Sub object already registered: Actor: %s, SubObj: %s, NetGuid: %s" ), *Actor->GetPathName(), *SubObj->GetPathName(), *NetGUIDLookup.FindChecked( SubObj ).ToString() );
+				UE_LOG( LogNetPackageMap, Error, TEXT( "AssignNewNetGUID_Server: Sub object already registered: Actor: %s, SubObj: %s, NetGuid: %s" ), *Actor->GetPathName(), *SubObj->GetPathName(), *NetGUIDLookup.FindChecked( SubObj ).ToString() );
 			}
 			const int32 SubIsStatic = IsDynamicObject( SubObj ) ? 0 : 1;
 			check( SubIsStatic == IsStatic );
 			const FNetworkGUID SubobjNetGUID( ALLOC_NEW_NET_GUID( SubIsStatic ) );
 			check( SubobjNetGUID.Value == COMPOSE_RELATIVE_NET_GUID( NewNetGuid, i + 1 ) );
-			AssignNetGUID( SubObj, SubobjNetGUID );
+			RegisterNetGUID_Server( SubobjNetGUID, SubObj );
 		}
 	}
 
 	return NewNetGuid;
 }
 
+void FNetGUIDCache::RegisterNetGUID_Internal( const FNetworkGUID & NetGUID, const FNetGuidCacheObject & CacheObject )
+{
+	// We're pretty strict in this function, we expect everything to have been handled before we get here
+	check( !ObjectLookup.Contains( NetGUID ) );
+
+	ObjectLookup.Add( NetGUID, CacheObject );
+
+	if ( CacheObject.Object != NULL )
+	{
+		check( !NetGUIDLookup.Contains( CacheObject.Object ) );
+
+		// If we have an object, associate it with this guid now
+		NetGUIDLookup.Add( CacheObject.Object, NetGUID );
+
+#if !(UE_BUILD_SHIPPING || UE_BUILD_TEST)
+		History.Add( NetGUID, CacheObject.Object->GetPathName() );
+#endif
+	}
+	else
+	{
+#if !(UE_BUILD_SHIPPING || UE_BUILD_TEST)
+		History.Add( NetGUID, CacheObject.PathName.ToString() );
+#endif
+	}
+}
+
 /**
- *	Assigns the given UObject the give NetNetworkGUID.
+ *	Associates a net guid directly with an object
+ *  This function is only called on server
  */
-void FNetGUIDCache::AssignNetGUID( const UObject * Object, const FNetworkGUID & NewNetworkGUID )
+void FNetGUIDCache::RegisterNetGUID_Server( const FNetworkGUID & NetGUID, const UObject * Object )
 {
+	check( IsNetGUIDAuthority() );				// Only the server should call this
 	check( !Object->IsPendingKill() );
-	check( !NewNetworkGUID.IsDefault() );
+	check( !NetGUID.IsDefault() );
+	check( !ObjectLookup.Contains( NetGUID ) );	// Server should never add twice
 
-	UE_LOG( LogNetPackageMap, Log, TEXT("Assigning %s NetGUID <%s>"), Object ? *Object->GetName() : TEXT("NULL"), *NewNetworkGUID.ToString() );
-	
-	// Actually assign
 	FNetGuidCacheObject CacheObject;
 
 	CacheObject.Object = Object;
-	CacheObject.FullPath = Object->GetPathName();
 
-	// Remove any existing entries
+	RegisterNetGUID_Internal( NetGUID, CacheObject );
+}
+
+/**
+ *	Associates a net guid directly with an object
+ *  This function is only called on clients with dynamic guids
+ */
+void FNetGUIDCache::RegisterNetGUID_Client( const FNetworkGUID & NetGUID, const UObject * Object )
+{
+	check( !IsNetGUIDAuthority() );			// Only clients should be here
+	check( !Object->IsPendingKill() );
+	check( !NetGUID.IsDefault() );
+	check( NetGUID.IsDynamic() );	// Clients should only assign dynamic guids through here (static guids go through RegisterNetGUIDFromPath_Client)
+
+	UE_LOG( LogNetPackageMap, Log, TEXT( "RegisterNetGUID_Client: Assigning %s NetGUID <%s>" ), Object ? *Object->GetName() : TEXT( "NULL" ), *NetGUID.ToString() );
+	
+	//
+	// If we have an existing entry, make sure things match up properly
+	// We also completely disassociate anything so that RegisterNetGUID_Internal can be fairly strict
+	//
+
+	const FNetGuidCacheObject * ExistingCacheObjectPtr = ObjectLookup.Find( NetGUID );
+
+	if ( ExistingCacheObjectPtr )
 	{
-		const FNetGuidCacheObject * ExistingCacheObjectPtr = ObjectLookup.Find( NewNetworkGUID );
+		if ( ExistingCacheObjectPtr->PathName != NAME_None )
+		{
+			UE_LOG( LogNetPackageMap, Warning, TEXT( "RegisterNetGUID_Client: Guid with pathname. Path: %s, NetGUID: %s" ), *ExistingCacheObjectPtr->PathName.ToString(), *NetGUID.ToString() );
+		}
+
+		// If this net guid was found but the old object is NULL, this can happen due to:
+		//	1. Actor channel was closed locally (but we don't remove the net guid entry, since we can't know for sure if it will be referenced again)
+		//		a. Then when we re-create a channel, and assign this actor, we will find the old guid entry here
+		//	2. Dynamic object was locally GC'd, but then exported again from the server
+		//
+		// If this net guid was found and the objects match, we don't care. This can happen due to:
+		//	1. Same thing above can happen, but if we for some reason didn't destroy the actor/object we will see this case
+		//
+		// If the object pointers are different, this can be a problem, 
+		//	since this should only be possible if something gets out of sync during the net guid exchange code
+
+		const UObject * OldObject = ExistingCacheObjectPtr->Object.Get();
 
-		if ( ExistingCacheObjectPtr )
+		if ( OldObject != NULL && OldObject != Object )
 		{
-			const UObject * OldObject = ExistingCacheObjectPtr->Object.Get();
+			UE_LOG( LogNetPackageMap, Warning, TEXT( "RegisterNetGUID_Client: Reassigning NetGUID <%s> to %s (was assigned to object %s)" ), *NetGUID.ToString(), Object ? *Object->GetPathName() : TEXT( "NULL" ), OldObject ? *OldObject->GetPathName() : TEXT( "NULL" ) );
+		}
+		else
+		{
+			UE_LOG( LogNetPackageMap, Verbose, TEXT( "RegisterNetGUID_Client: Reassigning NetGUID <%s> to %s (was assigned to object %s)" ), *NetGUID.ToString(), Object ? *Object->GetPathName() : TEXT( "NULL" ), OldObject ? *OldObject->GetPathName() : TEXT( "NULL" ) );
+		}
 
-			// If we find static guid in the cache, this means we didn't find it in GetObjectFromNetGUID
-			// This is a problem, since we should never get this far if we do find it there
-			if ( !NewNetworkGUID.IsDynamic() )
-			{
-				UE_LOG( LogNetPackageMap, Warning, TEXT( "AssignNetGUID: NewNetworkGUID was in ObjectLookup cache but is static: <%s> %s" ), *NewNetworkGUID.ToString(), Object ? *Object->GetPathName() : TEXT( "NULL" ) );
-			}
+		NetGUIDLookup.Remove( ExistingCacheObjectPtr->Object );
+		ObjectLookup.Remove( NetGUID );
+	}
 
-			// If this net guid was found but the old object is NULL, this can happen due to:
-			//	1. Actor channel was closed locally (but we don't remove the net guid entry, since we can't know for sure if it will be referenced again)
-			//		a. Then when we re-create a channel, and assign this actor, we will find the old guid entry here
-			//	2. Dynamic object was locally GC'd, but then exported again from the server
-			//
-			// If this net guid was found and the objects match, we don't care. This can happen due to:
-			//	1. Same thing above can happen, but if we for some reason didn't destroy the actor/object we will see this case
-			//
-			// If the object pointers are different, this can be a problem, 
-			//	since this should only be possible if something gets out of sync during the net guid exchange code
-			if ( OldObject != NULL && OldObject != Object )
-			{
-				UE_LOG( LogNetPackageMap, Warning, TEXT( "Reassigning NetGUID <%s> to %s (was assigned to object %s)" ), *NewNetworkGUID.ToString(), Object ? *Object->GetPathName() : TEXT( "NULL" ), OldObject ? *OldObject->GetPathName() : TEXT( "NULL" ) );
-			}
-			else
-			{
-				UE_LOG( LogNetPackageMap, Verbose, TEXT( "Reassigning NetGUID <%s> to %s (was assigned to object %s)" ), *NewNetworkGUID.ToString(), Object ? *Object->GetPathName() : TEXT( "NULL" ), OldObject ? *OldObject->GetPathName() : TEXT( "NULL" ) );
-			}
+	const FNetworkGUID * ExistingNetworkGUIDPtr = NetGUIDLookup.Find( Object );
+
+	if ( ExistingNetworkGUIDPtr )
+	{
+		// This shouldn't happen on dynamic guids
+		UE_LOG( LogNetPackageMap, Warning, TEXT( "Changing NetGUID on object %s from <%s:%s> to <%s:%s>" ), Object ? *Object->GetPathName() : TEXT( "NULL" ), *ExistingNetworkGUIDPtr->ToString(), ExistingNetworkGUIDPtr->IsDynamic() ? TEXT("TRUE") : TEXT("FALSE"), *NetGUID.ToString(), NetGUID.IsDynamic() ? TEXT("TRUE") : TEXT("FALSE") );
+		ObjectLookup.Remove( *ExistingNetworkGUIDPtr );
+		NetGUIDLookup.Remove( Object );
+	}
+
+	FNetGuidCacheObject CacheObject;
+
+	CacheObject.Object = Object;
 
-			NetGUIDLookup.Remove( ExistingCacheObjectPtr->Object );
+	RegisterNetGUID_Internal( NetGUID, CacheObject );
+}
+
+/**
+ *	Associates a net guid with a path, that can be loaded or found later
+ *  This function is only called on the client
+ */
+void FNetGUIDCache::RegisterNetGUIDFromPath_Client( const FNetworkGUID & NetGUID, const FString & PathName, const FNetworkGUID & OuterGUID, const bool bNoLoad, const bool bIgnoreWhenMissing )
+{
+	check( !IsNetGUIDAuthority() );		// Server never calls this locally
+	check( !NetGUID.IsDefault() );
+	check( !NetGUID.IsDynamic() );		// It only makes sense to do this for static guids
+
+	const FNetGuidCacheObject * ExistingCacheObjectPtr = ObjectLookup.Find( NetGUID );
+
+	// If we find this guid, make sure nothing changes
+	if ( ExistingCacheObjectPtr != NULL )
+	{
+		if ( ExistingCacheObjectPtr->PathName.ToString() != PathName )
+		{
+			UE_LOG( LogNetPackageMap, Warning, TEXT( "FNetGUIDCache::RegisterNetGUIDFromPath_Client: Path mismatch. Path: %s, Expected: %s, NetGUID: %s" ), *PathName, *ExistingCacheObjectPtr->PathName.ToString(), *NetGUID.ToString() );
 		}
 
-		const FNetworkGUID * ExistingNetworkGUIDPtr = NetGUIDLookup.Find( Object );
+		if ( ExistingCacheObjectPtr->OuterGUID != OuterGUID )
+		{
+			UE_LOG( LogNetPackageMap, Warning, TEXT( "FNetGUIDCache::RegisterNetGUIDFromPath_Client: Outer mismatch. Path: %s, Outer: %s, Expected: %s, NetGUID: %s" ), *PathName, *OuterGUID.ToString(), *ExistingCacheObjectPtr->OuterGUID.ToString(), *NetGUID.ToString() );
+		}
 
-		if ( ExistingNetworkGUIDPtr )
+		if ( ExistingCacheObjectPtr->Object != NULL )
 		{
-			// This can happen when the server destroys a package/object and the client doesn't.
-			// The server will simply assign a new guid, and the client will assign to the old object that never got deleted locally
-			UE_LOG( LogNetPackageMap, Log, TEXT( "Changing NetGUID on object %s from <%s:%s> to <%s:%s>" ), Object ? *Object->GetPathName() : TEXT( "NULL" ), *ExistingNetworkGUIDPtr->ToString(), ExistingNetworkGUIDPtr->IsDynamic() ? TEXT("TRUE") : TEXT("FALSE"), *NewNetworkGUID.ToString(), NewNetworkGUID.IsDynamic() ? TEXT("TRUE") : TEXT("FALSE") );
-			ObjectLookup.Remove( *ExistingNetworkGUIDPtr );
+			FNetworkGUID CurrentNetGUID = NetGUIDLookup.FindRef( ExistingCacheObjectPtr->Object );
+
+			if ( CurrentNetGUID != NetGUID )
+			{
+				UE_LOG( LogNetPackageMap, Warning, TEXT( "FNetGUIDCache::RegisterNetGUIDFromPath_Client: Netguid mismatch. Path: %s, NetGUID: %s, Expected: %s" ), *PathName, *NetGUID.ToString(), *CurrentNetGUID.ToString() );
+			}
 		}
+
+		ObjectLookup.Remove( NetGUID );
 	}
 
-	ObjectLookup.Add( NewNetworkGUID, CacheObject );
-	NetGUIDLookup.Add( Object, NewNetworkGUID );
+	// Register a new guid with this path
+	FNetGuidCacheObject CacheObject;
 
-#if !(UE_BUILD_SHIPPING || UE_BUILD_TEST)
-	History.Add(NewNetworkGUID, Object ? Object->GetPathName() : TEXT("NULL"));
-#endif
+	CacheObject.PathName			= FName( *PathName );
+	CacheObject.OuterGUID			= OuterGUID;
+	CacheObject.bNoLoad				= bNoLoad;
+	CacheObject.bIgnoreWhenMissing	= bIgnoreWhenMissing;
+
+	RegisterNetGUID_Internal( NetGUID, CacheObject );
 }
 
 UObject * FNetGUIDCache::GetObjectFromNetGUID( const FNetworkGUID & NetGUID )
@@ -1738,7 +1688,7 @@ UObject * FNetGUIDCache::GetObjectFromNetGUID( const FNetworkGUID & NetGUID )
 
 	if ( CacheObjectPtr == NULL )
 	{
-		// We don't have the object mapped yet
+		// This net guid has never been registered
 		return NULL;
 	}
 
@@ -1746,76 +1696,141 @@ UObject * FNetGUIDCache::GetObjectFromNetGUID( const FNetworkGUID & NetGUID )
 
 	if ( Object != NULL )
 	{
-		check( Object->GetPathName() == CacheObjectPtr->FullPath );
+		// Either the name should match, or this is dynamic, or we're on the server
+		check( Object->GetFName() == CacheObjectPtr->PathName || NetGUID.IsDynamic() || IsNetGUIDAuthority() );
 		return Object;
 	}
 
-	if ( CacheObjectPtr->FullPath.IsEmpty() )
+	if ( IsNetGUIDAuthority() )
 	{
-		// This probably isn't possible, but making it a warning
-		UE_LOG( LogNetPackageMap, Warning, TEXT( "GetObjectFromNetGUID: No full path for %s" ), *NetGUID.ToString() );
+		UE_LOG( LogNetPackageMap, Warning, TEXT( "GetObjectFromNetGUID: Guid with no object on server. Path: %s, NetGUID: %s" ), *CacheObjectPtr->PathName.ToString(), *NetGUID.ToString() );
 		return NULL;
 	}
 
 	if ( NetGUID.IsDynamic() )
 	{
-		// Dynamic objects don't have stable names, so we can't possibly reload the same object
-		UE_LOG( LogNetPackageMap, VeryVerbose, TEXT( "GetObjectFromNetGUID: Cannot re-create dynamic object after GC <%s, %s>" ), *NetGUID.ToString(), *CacheObjectPtr->FullPath );
+		// Dynamic objects can never load from their names (name is not stable)
 		return NULL;
 	}
 
-	//
-	// The object was previously mapped, but we GC'd it
-	// We need to reload it now
-	//
-
-	Object = StaticFindObject( UObject::StaticClass(), NULL, *CacheObjectPtr->FullPath, false );
+	if ( CacheObjectPtr->PathName == NAME_None )
+	{
+		// Static guids should have a path
+		UE_LOG( LogNetPackageMap, Warning, TEXT( "GetObjectFromNetGUID: Static guid on client with no path. NetGUID: %s" ), *NetGUID.ToString() );
+		return NULL;
+	}
 
-	if ( Object == NULL )
+	// First, resolve the outer
+	UObject * ObjOuter = NULL;
+	
+	if ( CacheObjectPtr->OuterGUID.IsValid() )
 	{
-		if ( IsNetGUIDAuthority() )
+		// We should have an outer, resolve it
+		ObjOuter = GetObjectFromNetGUID( CacheObjectPtr->OuterGUID );
+
+		// If we can't resolve the outer, we need to stop
+		if ( ObjOuter == NULL )
 		{
-			// The authority shouldn't be loading resources on demand, at least for now.
-			// This could be garbage or a security risk
-			// Another possibility is in dynamic property replication if the server reads the previously serialized state
-			// that has a now destroyed actor in it.
-			UE_LOG( LogNetPackageMap, Warning, TEXT( "GetObjectFromNetGUID: Could not find Object for: NetGUID <%s, %s> (and IsNetGUIDAuthority())" ), *NetGUID.ToString(), *CacheObjectPtr->FullPath );
+			// If the outer is not pending, we need to warn
+			if ( !ShouldIgnoreWhenMissing( CacheObjectPtr->OuterGUID ) )
+			{
+				UE_LOG( LogNetPackageMap, Warning, TEXT( "GetObjectFromNetGUID: Failed to find outer. Path: %s, NetGUID: %s" ), *CacheObjectPtr->PathName.ToString(), *NetGUID.ToString() );
+			}
 			return NULL;
 		}
-		else
-		{
-			UE_LOG( LogNetPackageMap, Log, TEXT( "GetObjectFromNetGUID: Could not find Object for: NetGUID <%s, %s>"), *NetGUID.ToString(), *CacheObjectPtr->FullPath );
-		}
+	}
 
-		Object = StaticLoadObject( UObject::StaticClass(), NULL, *CacheObjectPtr->FullPath, NULL, LOAD_NoWarn );
+	// Resolve the path of this object relative to outer
+	Object = ResolvePath( CacheObjectPtr->PathName.ToString(), ObjOuter, CacheObjectPtr->bNoLoad || CacheObjectPtr->bIsPending );
 
-		if ( Object )
-		{
-			UE_LOG( LogNetPackageMap, Log, TEXT( "GetObjectFromNetGUID: StaticLoadObject. Found: %s" ), Object ? *Object->GetName() : TEXT( "NULL" ) );
-		}
-		else
+	if ( Object == NULL )
+	{
+		// If we can't find an object, and it's not pending, warn
+		if ( !CacheObjectPtr->bIgnoreWhenMissing )
 		{
-			Object = LoadPackage( NULL, *CacheObjectPtr->FullPath, LOAD_None );
-			UE_LOG( LogNetPackageMap, Log, TEXT( "GetObjectFromNetGUID: LoadPackage. Found: %s" ), Object ? *Object->GetName() : TEXT( "NULL" ) );
+			UE_LOG( LogNetPackageMap, Warning, TEXT( "GetObjectFromNetGUID: Failed to resolve path. Path: %s, NetGUID: %s" ), *CacheObjectPtr->PathName.ToString(), *NetGUID.ToString() );
 		}
+
+		return NULL;
 	}
-	else
+
+	// Assign the resolved object to this guid
+	CacheObjectPtr->Object = Object;		
+
+	if ( NetGUIDLookup.Contains( Object ) )
 	{
-		// If we actually found the object, we probably shouldn't have GC'd it, so that's odd
-		//UE_LOG( LogNetPackageMap, Warning, TEXT( "GetObjectFromNetGUID: Object should not be found after GC: <%s, %s>" ), *NetGUID.ToString(), *CacheObjectPtr->FullPath );
+		// This can happen to static guids on the client. The server deleted this object, but for some reason the client did not
+		// When this happens, the server will give the object a different guid, and we will just conform
+		ObjectLookup.Remove( NetGUIDLookup.FindRef( Object ) );
 	}
 
-	if ( Object == NULL )
+	// Make sure the object is in the GUIDToObjectLookup.
+	NetGUIDLookup.Add( Object, NetGUID );
+
+	return Object;
+}
+
+bool FNetGUIDCache::ShouldIgnoreWhenMissing( const FNetworkGUID & NetGUID )
+{
+	if ( IsNetGUIDAuthority() )
+	{
+		return false;		// Server never ignores when missing, always warns
+	}
+
+	const FNetGuidCacheObject * CacheObject = ObjectLookup.Find( NetGUID );
+
+	if ( CacheObject != NULL )
+	{
+		// Ignore warnings when we explicitly are told to, or if we're pending load (which could also be asynchronously loading)
+		return CacheObject->bIgnoreWhenMissing || CacheObject->bIsPending;
+	}
+
+	return false;
+}
+
+UObject * FNetGUIDCache::ResolvePath( const FString & PathName, UObject * ObjOuter, const bool bNoLoad )
+{
+	UObject * Object = StaticFindObject( UObject::StaticClass(), ObjOuter, *PathName, false );
+
+	if ( bNoLoad || Object != NULL )
+	{
+		return Object;
+	}
+
+	if ( IsNetGUIDAuthority() )
+	{
+		// The authority shouldn't be loading resources on demand, at least for now.
+		// This could be garbage or a security risk
+		// Another possibility is in dynamic property replication if the server reads the previously serialized state
+		// that has a now destroyed actor in it.
+		UE_LOG( LogNetPackageMap, Warning, TEXT( "FNetGUIDCache::ResolvePath: Server can't load objects from client. Path: %s" ), *PathName );
+		return NULL;
+	}
+
+	Object = StaticLoadObject( UObject::StaticClass(), ObjOuter, *PathName, NULL, LOAD_NoWarn );
+
+	if ( Object )
 	{
-		// If we get here, that means we have an invalid path, which shouldn't be possible, but making it a warning
-		UE_LOG( LogNetPackageMap, Warning, TEXT( "GetObjectFromNetGUID: FAILED to re-create object after GC: <%s, %s>" ), *NetGUID.ToString(), *CacheObjectPtr->FullPath );
+		UE_LOG( LogNetPackageMap, Log, TEXT( "FNetGUIDCache::ResolvePath: StaticLoadObject. Found: %s" ), Object ? *Object->GetName() : TEXT( "NULL" ) );
 	}
 	else
 	{
-		// Reassign the object pointer for quick access next time
-		CacheObjectPtr->Object = Object;		
-		// Make sure the object is in the GUIDToObjectLookup.
-		NetGUIDLookup.Add( Object, NetGUID );
+		//
+		// If we failed to load it as an object, try to load it as a package
+		//
+
+		UPackage * PackageOuter = Cast< UPackage >( ObjOuter );
+
+		// If we have an outer at this point, it only makes sense for it to be a package (or NULL)
+		if ( PackageOuter != NULL || ObjOuter == NULL )
+		{
+			Object = LoadPackage( PackageOuter, *PathName, LOAD_None );
+			UE_LOG( LogNetPackageMap, Log, TEXT( "FNetGUIDCache::ResolvePath: LoadPackage: %s. Found: %s" ), *PathName, Object ? *Object->GetName() : TEXT( "NULL" ) );
+		} 
+		else
+		{
+			UE_LOG( LogNetPackageMap, Warning, TEXT( "FNetGUIDCache::ResolvePath: Outer %s is not a package for object %s" ), ObjOuter ? *ObjOuter->GetName() : TEXT( "NULL" ), *PathName );
+		}
 	}
 
 	return Object;
diff --git a/Engine/Source/Runtime/Engine/Private/PendingNetGame.cpp b/Engine/Source/Runtime/Engine/Private/PendingNetGame.cpp
--- a/Engine/Source/Runtime/Engine/Private/PendingNetGame.cpp
+++ b/Engine/Source/Runtime/Engine/Private/PendingNetGame.cpp
@@ -325,7 +325,7 @@ void UPendingNetGame::NotifyControlMessage(UNetConnection* Connection, uint8 Mes
 			FNetworkGUID NetGUID;
 			FString Path;
 			FNetControlMessage<NMT_NetGUIDAssign>::Receive(Bunch, NetGUID, Path);
-			NetDriver->ServerConnection->PackageMap->ResolvePathAndAssignNetGUID(NetGUID, Path, NULL);
+			NetDriver->ServerConnection->PackageMap->ResolvePathAndAssignNetGUID(NetGUID, Path);
 			break;
 		}
 		default:
diff --git a/Engine/Source/Runtime/Engine/Private/World.cpp b/Engine/Source/Runtime/Engine/Private/World.cpp
--- a/Engine/Source/Runtime/Engine/Private/World.cpp
+++ b/Engine/Source/Runtime/Engine/Private/World.cpp
@@ -3420,7 +3420,7 @@ void UWorld::NotifyControlMessage(UNetConnection* Connection, uint8 MessageType,
 				FNetControlMessage<NMT_NetGUIDAssign>::Receive(Bunch, NetGUID, Path);
 
 				UE_LOG(LogNet, Verbose, TEXT("NMT_NetGUIDAssign  NetGUID %s. Path: %s. "), *NetGUID.ToString(), *Path );
-				Connection->PackageMap->ResolvePathAndAssignNetGUID( NetGUID, Path, NULL );
+				Connection->PackageMap->ResolvePathAndAssignNetGUID( NetGUID, Path );
 				break;
 			}
 		}
diff --git a/Engine/Source/Runtime/Online/OnlineSubsystemUtils/Private/OnlineBeaconClient.cpp b/Engine/Source/Runtime/Online/OnlineSubsystemUtils/Private/OnlineBeaconClient.cpp
--- a/Engine/Source/Runtime/Online/OnlineSubsystemUtils/Private/OnlineBeaconClient.cpp
+++ b/Engine/Source/Runtime/Online/OnlineSubsystemUtils/Private/OnlineBeaconClient.cpp
@@ -133,7 +133,7 @@ void AOnlineBeaconClient::NotifyControlMessage(UNetConnection* Connection, uint8
 				FNetControlMessage<NMT_BeaconAssignGUID>::Receive(Bunch, NetGUID);
 				if (NetGUID.IsValid())
 				{
-					Connection->Driver->GuidCache->AssignNetGUID( this, NetGUID );
+					Connection->Driver->GuidCache->RegisterNetGUID_Client( NetGUID, this );
 
 					FString BeaconType = GetBeaconType();
 					FNetControlMessage<NMT_BeaconNetGUIDAck>::Send(Connection, BeaconType);
diff --git a/Engine/Source/Runtime/Online/OnlineSubsystemUtils/Private/OnlineBeaconHost.cpp b/Engine/Source/Runtime/Online/OnlineSubsystemUtils/Private/OnlineBeaconHost.cpp
--- a/Engine/Source/Runtime/Online/OnlineSubsystemUtils/Private/OnlineBeaconHost.cpp
+++ b/Engine/Source/Runtime/Online/OnlineSubsystemUtils/Private/OnlineBeaconHost.cpp
@@ -140,7 +140,7 @@ void AOnlineBeaconHost::NotifyControlMessage(UNetConnection* Connection, uint8 M
 
 						if (NewClientActor && BeaconType == NewClientActor->GetBeaconType())
 						{
-							FNetworkGUID NetGUID = Connection->Driver->GuidCache->AssignNewNetGUID( NewClientActor );
+							FNetworkGUID NetGUID = Connection->Driver->GuidCache->AssignNewNetGUID_Server( NewClientActor );
 							NewClientActor->SetNetConnection(Connection);
 							Connection->OwningActor = NewClientActor;
 							NewClientActor->Role = ROLE_None;
