Fix no default function template parameter

diff --git a/Engine/Source/Runtime/Engine/Classes/Components/SkinnedMeshComponent.h b/Engine/Source/Runtime/Engine/Classes/Components/SkinnedMeshComponent.h
--- a/Engine/Source/Runtime/Engine/Classes/Components/SkinnedMeshComponent.h
+++ b/Engine/Source/Runtime/Engine/Classes/Components/SkinnedMeshComponent.h
@@ -839,6 +839,6 @@ private:
 	/** 
 	 * Simple, CPU evaluation of a vertex's skinned position helper function
 	 */
-	template <bool bExtraBoneInfluencesT, bool bCachedMatrices = false>
+	template <bool bExtraBoneInfluencesT, bool bCachedMatrices>
 	FVector GetTypedSkinnedVertexPosition(const FSkelMeshChunk& Chunk, const FSkeletalMeshVertexBuffer& VertexBufferGPUSkin, int32 VertIndex, bool bSoftVertex, const TArray<FMatrix> & RefToLocals = TArray<FMatrix>()) const;
 };
diff --git a/Engine/Source/Runtime/Engine/Private/SkinnedMeshComponent.cpp b/Engine/Source/Runtime/Engine/Private/SkinnedMeshComponent.cpp
--- a/Engine/Source/Runtime/Engine/Private/SkinnedMeshComponent.cpp
+++ b/Engine/Source/Runtime/Engine/Private/SkinnedMeshComponent.cpp
@@ -1439,9 +1439,9 @@ FVector USkinnedMeshComponent::GetSkinnedVertexPosition(int32 VertexIndex) const
 			//rigid
 
 	return bHasExtraBoneInfluences
-		? GetTypedSkinnedVertexPosition<true>(Chunk, Model.VertexBufferGPUSkin, VertIndex, bSoftVertex)
+		? GetTypedSkinnedVertexPosition<true, false>(Chunk, Model.VertexBufferGPUSkin, VertIndex, bSoftVertex)
 			//soft
-		: GetTypedSkinnedVertexPosition<false>(Chunk, Model.VertexBufferGPUSkin, VertIndex, bSoftVertex);
+		: GetTypedSkinnedVertexPosition<false, false>(Chunk, Model.VertexBufferGPUSkin, VertIndex, bSoftVertex);
 
 
 }
