UE4 Refactoring. Changed OVERRIDE and FINAL macros to keywords override and final.

diff --git a/Engine/Source/Runtime/Engine/Classes/Animation/BoneControllers/AnimNode_CopyBone.h b/Engine/Source/Runtime/Engine/Classes/Animation/BoneControllers/AnimNode_CopyBone.h
--- a/Engine/Source/Runtime/Engine/Classes/Animation/BoneControllers/AnimNode_CopyBone.h
+++ b/Engine/Source/Runtime/Engine/Classes/Animation/BoneControllers/AnimNode_CopyBone.h
@@ -37,16 +37,16 @@ struct ENGINE_API FAnimNode_CopyBone : public FAnimNode_SkeletalControlBase
 	FAnimNode_CopyBone();
 
 	// FAnimNode_Base interface
-	virtual void GatherDebugData(FNodeDebugData& DebugData) OVERRIDE;
+	virtual void GatherDebugData(FNodeDebugData& DebugData) override;
 	// End of FAnimNode_Base interface
 
 	// FAnimNode_SkeletalControlBase interface
-	virtual void EvaluateBoneTransforms(USkeletalMeshComponent* SkelComp, const FBoneContainer & RequiredBones, FA2CSPose& MeshBases, TArray<FBoneTransform>& OutBoneTransforms) OVERRIDE;
-	virtual bool IsValidToEvaluate(const USkeleton * Skeleton, const FBoneContainer & RequiredBones) OVERRIDE;
+	virtual void EvaluateBoneTransforms(USkeletalMeshComponent* SkelComp, const FBoneContainer & RequiredBones, FA2CSPose& MeshBases, TArray<FBoneTransform>& OutBoneTransforms) override;
+	virtual bool IsValidToEvaluate(const USkeleton * Skeleton, const FBoneContainer & RequiredBones) override;
 	// End of FAnimNode_SkeletalControlBase interface
 
 private:
 	// FAnimNode_SkeletalControlBase interface
-	virtual void InitializeBoneReferences(const FBoneContainer & RequiredBones) OVERRIDE;
+	virtual void InitializeBoneReferences(const FBoneContainer & RequiredBones) override;
 	// End of FAnimNode_SkeletalControlBase interface
 };
diff --git a/Engine/Source/Runtime/Engine/Classes/Animation/BoneControllers/AnimNode_Fabrik.h b/Engine/Source/Runtime/Engine/Classes/Animation/BoneControllers/AnimNode_Fabrik.h
--- a/Engine/Source/Runtime/Engine/Classes/Animation/BoneControllers/AnimNode_Fabrik.h
+++ b/Engine/Source/Runtime/Engine/Classes/Animation/BoneControllers/AnimNode_Fabrik.h
@@ -94,14 +94,14 @@ public:
 	FAnimNode_Fabrik();
 
 	// FAnimNode_SkeletalControlBase interface
-	virtual void EvaluateBoneTransforms(USkeletalMeshComponent* SkelComp, const FBoneContainer & RequiredBones, FA2CSPose& MeshBases, TArray<FBoneTransform>& OutBoneTransforms) OVERRIDE;
-	virtual bool IsValidToEvaluate(const USkeleton * Skeleton, const FBoneContainer & RequiredBones) OVERRIDE;
+	virtual void EvaluateBoneTransforms(USkeletalMeshComponent* SkelComp, const FBoneContainer & RequiredBones, FA2CSPose& MeshBases, TArray<FBoneTransform>& OutBoneTransforms) override;
+	virtual bool IsValidToEvaluate(const USkeleton * Skeleton, const FBoneContainer & RequiredBones) override;
 	// End of FAnimNode_SkeletalControlBase interface
 
 
 private:
 	// FAnimNode_SkeletalControlBase interface
-	virtual void InitializeBoneReferences(const FBoneContainer & RequiredBones) OVERRIDE;
+	virtual void InitializeBoneReferences(const FBoneContainer & RequiredBones) override;
 	// End of FAnimNode_SkeletalControlBase interface
 
 	// Convenience function to get current (pre-translation iteration) component space location of bone by bone index
diff --git a/Engine/Source/Runtime/Engine/Classes/Animation/BoneControllers/AnimNode_HandIKRetargeting.h b/Engine/Source/Runtime/Engine/Classes/Animation/BoneControllers/AnimNode_HandIKRetargeting.h
--- a/Engine/Source/Runtime/Engine/Classes/Animation/BoneControllers/AnimNode_HandIKRetargeting.h
+++ b/Engine/Source/Runtime/Engine/Classes/Animation/BoneControllers/AnimNode_HandIKRetargeting.h
@@ -44,16 +44,16 @@ struct ENGINE_API FAnimNode_HandIKRetargeting : public FAnimNode_SkeletalControl
 	FAnimNode_HandIKRetargeting();
 
 	// FAnimNode_Base interface
-	virtual void GatherDebugData(FNodeDebugData& DebugData) OVERRIDE;
+	virtual void GatherDebugData(FNodeDebugData& DebugData) override;
 	// End of FAnimNode_Base interface
 
 	// FAnimNode_SkeletalControlBase interface
-	virtual void EvaluateBoneTransforms(USkeletalMeshComponent* SkelComp, const FBoneContainer & RequiredBones, FA2CSPose& MeshBases, TArray<FBoneTransform>& OutBoneTransforms) OVERRIDE;
-	virtual bool IsValidToEvaluate(const USkeleton * Skeleton, const FBoneContainer & RequiredBones) OVERRIDE;
+	virtual void EvaluateBoneTransforms(USkeletalMeshComponent* SkelComp, const FBoneContainer & RequiredBones, FA2CSPose& MeshBases, TArray<FBoneTransform>& OutBoneTransforms) override;
+	virtual bool IsValidToEvaluate(const USkeleton * Skeleton, const FBoneContainer & RequiredBones) override;
 	// End of FAnimNode_SkeletalControlBase interface
 
 private:
 	// FAnimNode_SkeletalControlBase interface
-	virtual void InitializeBoneReferences(const FBoneContainer & RequiredBones) OVERRIDE;
+	virtual void InitializeBoneReferences(const FBoneContainer & RequiredBones) override;
 	// End of FAnimNode_SkeletalControlBase interface
 };
diff --git a/Engine/Source/Runtime/Engine/Classes/Animation/BoneControllers/AnimNode_SkeletalControlBase.h b/Engine/Source/Runtime/Engine/Classes/Animation/BoneControllers/AnimNode_SkeletalControlBase.h
--- a/Engine/Source/Runtime/Engine/Classes/Animation/BoneControllers/AnimNode_SkeletalControlBase.h
+++ b/Engine/Source/Runtime/Engine/Classes/Animation/BoneControllers/AnimNode_SkeletalControlBase.h
@@ -37,10 +37,10 @@ public:
 
 public:
 	// FAnimNode_Base interface
-	virtual void Initialize(const FAnimationInitializeContext& Context) OVERRIDE;
-	virtual void CacheBones(const FAnimationCacheBonesContext & Context)  OVERRIDE;
-	virtual void Update(const FAnimationUpdateContext& Context) OVERRIDE;
-	virtual void EvaluateComponentSpace(FComponentSpacePoseContext& Output) OVERRIDE;
+	virtual void Initialize(const FAnimationInitializeContext& Context) override;
+	virtual void CacheBones(const FAnimationCacheBonesContext & Context)  override;
+	virtual void Update(const FAnimationUpdateContext& Context) override;
+	virtual void EvaluateComponentSpace(FComponentSpacePoseContext& Output) override;
 	// End of FAnimNode_Base interface
 	
 protected:
diff --git a/Engine/Source/Runtime/Engine/Classes/Animation/BoneControllers/AnimNode_SpringBone.h b/Engine/Source/Runtime/Engine/Classes/Animation/BoneControllers/AnimNode_SpringBone.h
--- a/Engine/Source/Runtime/Engine/Classes/Animation/BoneControllers/AnimNode_SpringBone.h
+++ b/Engine/Source/Runtime/Engine/Classes/Animation/BoneControllers/AnimNode_SpringBone.h
@@ -58,19 +58,19 @@ public:
 	FAnimNode_SpringBone();
 
 	// FAnimNode_Base interface
-	virtual void Initialize(const FAnimationInitializeContext& Context) OVERRIDE;
-	virtual void CacheBones(const FAnimationCacheBonesContext & Context) OVERRIDE;
-	virtual void Update(const FAnimationUpdateContext& Context) OVERRIDE;
-	virtual void GatherDebugData(FNodeDebugData& DebugData) OVERRIDE;
+	virtual void Initialize(const FAnimationInitializeContext& Context) override;
+	virtual void CacheBones(const FAnimationCacheBonesContext & Context) override;
+	virtual void Update(const FAnimationUpdateContext& Context) override;
+	virtual void GatherDebugData(FNodeDebugData& DebugData) override;
 	// End of FAnimNode_Base interface
 
 	// FAnimNode_SkeletalControlBase interface
-	virtual void EvaluateBoneTransforms(USkeletalMeshComponent* SkelComp, const FBoneContainer & RequiredBones, FA2CSPose& MeshBases, TArray<FBoneTransform>& OutBoneTransforms) OVERRIDE;
-	virtual bool IsValidToEvaluate(const USkeleton * Skeleton, const FBoneContainer & RequiredBones) OVERRIDE;
+	virtual void EvaluateBoneTransforms(USkeletalMeshComponent* SkelComp, const FBoneContainer & RequiredBones, FA2CSPose& MeshBases, TArray<FBoneTransform>& OutBoneTransforms) override;
+	virtual bool IsValidToEvaluate(const USkeleton * Skeleton, const FBoneContainer & RequiredBones) override;
 	// End of FAnimNode_SkeletalControlBase interface
 
 private:
 	// FAnimNode_SkeletalControlBase interface
-	virtual void InitializeBoneReferences(const FBoneContainer & RequiredBones) OVERRIDE;
+	virtual void InitializeBoneReferences(const FBoneContainer & RequiredBones) override;
 	// End of FAnimNode_SkeletalControlBase interface
 };
diff --git a/Engine/Source/Runtime/Engine/Classes/Animation/BoneControllers/AnimNode_WheelHandler.h b/Engine/Source/Runtime/Engine/Classes/Animation/BoneControllers/AnimNode_WheelHandler.h
--- a/Engine/Source/Runtime/Engine/Classes/Animation/BoneControllers/AnimNode_WheelHandler.h
+++ b/Engine/Source/Runtime/Engine/Classes/Animation/BoneControllers/AnimNode_WheelHandler.h
@@ -30,18 +30,18 @@ struct ENGINE_API FAnimNode_WheelHandler : public FAnimNode_SkeletalControlBase
 	FAnimNode_WheelHandler();
 
 	// FAnimNode_Base interface
-	virtual void GatherDebugData(FNodeDebugData& DebugData) OVERRIDE;
+	virtual void GatherDebugData(FNodeDebugData& DebugData) override;
 	// End of FAnimNode_Base interface
 
 	// FAnimNode_SkeletalControlBase interface
-	virtual void EvaluateBoneTransforms(USkeletalMeshComponent* SkelComp, const FBoneContainer & RequiredBones, FA2CSPose& MeshBases, TArray<FBoneTransform>& OutBoneTransforms) OVERRIDE;
-	virtual bool IsValidToEvaluate(const USkeleton * Skeleton, const FBoneContainer & RequiredBones) OVERRIDE;
-	virtual void Update(const FAnimationUpdateContext& Context) OVERRIDE;
-	virtual void Initialize(const FAnimationInitializeContext& Context) OVERRIDE;
+	virtual void EvaluateBoneTransforms(USkeletalMeshComponent* SkelComp, const FBoneContainer & RequiredBones, FA2CSPose& MeshBases, TArray<FBoneTransform>& OutBoneTransforms) override;
+	virtual bool IsValidToEvaluate(const USkeleton * Skeleton, const FBoneContainer & RequiredBones) override;
+	virtual void Update(const FAnimationUpdateContext& Context) override;
+	virtual void Initialize(const FAnimationInitializeContext& Context) override;
 	// End of FAnimNode_SkeletalControlBase interface
 
 private:
 	// FAnimNode_SkeletalControlBase interface
-	virtual void InitializeBoneReferences(const FBoneContainer & RequiredBones) OVERRIDE;
+	virtual void InitializeBoneReferences(const FBoneContainer & RequiredBones) override;
 	// End of FAnimNode_SkeletalControlBase interface
 };
diff --git a/Engine/Source/Runtime/Engine/Classes/Animation/SkeletalMeshActor.h b/Engine/Source/Runtime/Engine/Classes/Animation/SkeletalMeshActor.h
--- a/Engine/Source/Runtime/Engine/Classes/Animation/SkeletalMeshActor.h
+++ b/Engine/Source/Runtime/Engine/Classes/Animation/SkeletalMeshActor.h
@@ -56,24 +56,24 @@ public:
 
 	// Begin AActor interface
 #if WITH_EDITOR
-	virtual void CheckForErrors() OVERRIDE;
-	virtual bool GetReferencedContentObjects( TArray<UObject*>& Objects ) const OVERRIDE;
-	virtual void EditorReplacedActor(AActor* OldActor) OVERRIDE;
-	virtual void LoadedFromAnotherClass(const FName& OldClassName) OVERRIDE;
+	virtual void CheckForErrors() override;
+	virtual bool GetReferencedContentObjects( TArray<UObject*>& Objects ) const override;
+	virtual void EditorReplacedActor(AActor* OldActor) override;
+	virtual void LoadedFromAnotherClass(const FName& OldClassName) override;
 #endif
-	virtual void PostInitializeComponents() OVERRIDE;
+	virtual void PostInitializeComponents() override;
 	// End AActor interface
 
 	// Begin IMatineeAnimInterface Interface
-	virtual void PreviewBeginAnimControl(class UInterpGroup* InInterpGroup) OVERRIDE;
-	virtual void PreviewSetAnimPosition(FName SlotName, int32 ChannelIndex, UAnimSequence* InAnimSequence, float InPosition, bool bLooping, bool bFireNotifies, float AdvanceTime) OVERRIDE;
-	virtual void PreviewSetAnimWeights(TArray<FAnimSlotInfo>& SlotInfos) OVERRIDE;
-	virtual void PreviewFinishAnimControl(class UInterpGroup* InInterpGroup) OVERRIDE;
-	virtual void GetAnimControlSlotDesc(TArray<struct FAnimSlotDesc>& OutSlotDescs) OVERRIDE {};
-	virtual void SetAnimWeights( const TArray<struct FAnimSlotInfo>& SlotInfos ) OVERRIDE;
-	virtual void BeginAnimControl(class UInterpGroup* InInterpGroup) OVERRIDE;
-	virtual void SetAnimPosition(FName SlotName, int32 ChannelIndex, class UAnimSequence* InAnimSequence, float InPosition, bool bFireNotifies, bool bLooping) OVERRIDE;
-	virtual void FinishAnimControl(class UInterpGroup* InInterpGroup) OVERRIDE;
+	virtual void PreviewBeginAnimControl(class UInterpGroup* InInterpGroup) override;
+	virtual void PreviewSetAnimPosition(FName SlotName, int32 ChannelIndex, UAnimSequence* InAnimSequence, float InPosition, bool bLooping, bool bFireNotifies, float AdvanceTime) override;
+	virtual void PreviewSetAnimWeights(TArray<FAnimSlotInfo>& SlotInfos) override;
+	virtual void PreviewFinishAnimControl(class UInterpGroup* InInterpGroup) override;
+	virtual void GetAnimControlSlotDesc(TArray<struct FAnimSlotDesc>& OutSlotDescs) override {};
+	virtual void SetAnimWeights( const TArray<struct FAnimSlotInfo>& SlotInfos ) override;
+	virtual void BeginAnimControl(class UInterpGroup* InInterpGroup) override;
+	virtual void SetAnimPosition(FName SlotName, int32 ChannelIndex, class UAnimSequence* InAnimSequence, float InPosition, bool bFireNotifies, bool bLooping) override;
+	virtual void FinishAnimControl(class UInterpGroup* InInterpGroup) override;
 	// End IMatineeAnimInterface Interface
 
 private:
diff --git a/Engine/Source/Runtime/Engine/Classes/Animation/Skeleton.h b/Engine/Source/Runtime/Engine/Classes/Animation/Skeleton.h
--- a/Engine/Source/Runtime/Engine/Classes/Animation/Skeleton.h
+++ b/Engine/Source/Runtime/Engine/Classes/Animation/Skeleton.h
@@ -470,10 +470,10 @@ public:
 
 	ENGINE_API void SetBoneTranslationRetargetingMode(const int32 & BoneIndex, EBoneTranslationRetargetingMode::Type NewRetargetingMode, bool bChildrenToo=false);
 
-	virtual void PostLoad() OVERRIDE;
-	virtual void PostDuplicate(bool bDuplicateForPIE) OVERRIDE;
-	virtual void PostInitProperties() OVERRIDE;
-	virtual void Serialize(FArchive& Ar) OVERRIDE;
+	virtual void PostLoad() override;
+	virtual void PostDuplicate(bool bDuplicateForPIE) override;
+	virtual void PostInitProperties() override;
+	virtual void Serialize(FArchive& Ar) override;
 
 	/** 
 	 * Create RefLocalPoses from InSkelMesh
diff --git a/Engine/Source/Runtime/Engine/Classes/Animation/VertexAnim/MorphTarget.h b/Engine/Source/Runtime/Engine/Classes/Animation/VertexAnim/MorphTarget.h
--- a/Engine/Source/Runtime/Engine/Classes/Animation/VertexAnim/MorphTarget.h
+++ b/Engine/Source/Runtime/Engine/Classes/Animation/VertexAnim/MorphTarget.h
@@ -74,8 +74,8 @@ public:
 	TArray<FMorphTargetLODModel>	MorphLODModels;
 
 	// Begin UObject interface.
-	virtual void Serialize( FArchive& Ar ) OVERRIDE;
-	virtual SIZE_T GetResourceSize(EResourceSizeMode::Type Mode) OVERRIDE;
+	virtual void Serialize( FArchive& Ar ) override;
+	virtual SIZE_T GetResourceSize(EResourceSizeMode::Type Mode) override;
 	// Begin UObject interface.
 
 	/** Post process after importing **/
@@ -85,8 +85,8 @@ public:
 	void RemapVertexIndices( USkeletalMesh* InBaseMesh, const TArray< TArray<uint32> > & BasedWedgePointIndices );
 
 	// Begin UVertexAnimBase interface
-	virtual FVertexAnimDelta* GetDeltasAtTime(float Time, int32 LODIndex, FVertexAnimEvalStateBase* State, int32& OutNumDeltas) OVERRIDE;
-	virtual bool HasDataForLOD(int32 LODIndex) OVERRIDE;
+	virtual FVertexAnimDelta* GetDeltasAtTime(float Time, int32 LODIndex, FVertexAnimEvalStateBase* State, int32& OutNumDeltas) override;
+	virtual bool HasDataForLOD(int32 LODIndex) override;
 	// End UVertexAnimBase interface
 
 private:
diff --git a/Engine/Source/Runtime/Engine/Classes/Animation/VertexAnim/VertexAnimation.h b/Engine/Source/Runtime/Engine/Classes/Animation/VertexAnim/VertexAnimation.h
--- a/Engine/Source/Runtime/Engine/Classes/Animation/VertexAnim/VertexAnimation.h
+++ b/Engine/Source/Runtime/Engine/Classes/Animation/VertexAnim/VertexAnimation.h
@@ -37,15 +37,15 @@ class UVertexAnimation : public UVertexAnimBase
 	TArray<FVertexAnimFrame>	VertexAnimData;
 
 	// Begin UObject interface.
-	virtual void Serialize( FArchive& Ar ) OVERRIDE;
-	virtual SIZE_T GetResourceSize(EResourceSizeMode::Type Mode) OVERRIDE;
+	virtual void Serialize( FArchive& Ar ) override;
+	virtual SIZE_T GetResourceSize(EResourceSizeMode::Type Mode) override;
 	// Begin UObject interface.
 
 	// Begin UVertexAnimBase interface
-	virtual FVertexAnimEvalStateBase* InitEval() OVERRIDE;
-	virtual void TermEval(FVertexAnimEvalStateBase* State) OVERRIDE;
-	virtual FVertexAnimDelta* GetDeltasAtTime(float Time, int32 LODIndex, FVertexAnimEvalStateBase* State, int32& OutNumDeltas) OVERRIDE;
-	virtual bool HasDataForLOD(int32 LODIndex) OVERRIDE;
+	virtual FVertexAnimEvalStateBase* InitEval() override;
+	virtual void TermEval(FVertexAnimEvalStateBase* State) override;
+	virtual FVertexAnimDelta* GetDeltasAtTime(float Time, int32 LODIndex, FVertexAnimEvalStateBase* State, int32& OutNumDeltas) override;
+	virtual bool HasDataForLOD(int32 LODIndex) override;
 	// End UVertexAnimBase interface
 
 	/** Get the number of frames in the animation */
diff --git a/Engine/Source/Runtime/Engine/Classes/Atmosphere/AtmosphericFog.h b/Engine/Source/Runtime/Engine/Classes/Atmosphere/AtmosphericFog.h
--- a/Engine/Source/Runtime/Engine/Classes/Atmosphere/AtmosphericFog.h
+++ b/Engine/Source/Runtime/Engine/Classes/Atmosphere/AtmosphericFog.h
@@ -18,7 +18,7 @@ class AAtmosphericFog : public AInfo
 #endif
 
 #if WITH_EDITOR
-	virtual void PostActorCreated() OVERRIDE;
+	virtual void PostActorCreated() override;
 #endif
 };
 
diff --git a/Engine/Source/Runtime/Engine/Classes/Atmosphere/AtmosphericFogComponent.h b/Engine/Source/Runtime/Engine/Classes/Atmosphere/AtmosphericFogComponent.h
--- a/Engine/Source/Runtime/Engine/Classes/Atmosphere/AtmosphericFogComponent.h
+++ b/Engine/Source/Runtime/Engine/Classes/Atmosphere/AtmosphericFogComponent.h
@@ -224,9 +224,9 @@ public:
 
 protected:
 	// Begin UActorComponent interface.
-	virtual void CreateRenderState_Concurrent() OVERRIDE;
-	virtual void SendRenderTransform_Concurrent() OVERRIDE;
-	virtual void DestroyRenderState_Concurrent() OVERRIDE;
+	virtual void CreateRenderState_Concurrent() override;
+	virtual void SendRenderTransform_Concurrent() override;
+	virtual void DestroyRenderState_Concurrent() override;
 	// End UActorComponent interface.
 
 	 void AddFogIfNeeded();
@@ -243,24 +243,24 @@ public:
 	mutable FByteBulkData InscatterData;
 	
 	// Begin UObject interface.
-	virtual void PostLoad() OVERRIDE;
-	virtual void BeginDestroy() OVERRIDE;
+	virtual void PostLoad() override;
+	virtual void BeginDestroy() override;
 
 #if WITH_EDITOR
-	virtual void PostEditChangeProperty(FPropertyChangedEvent& PropertyChangedEvent) OVERRIDE;
+	virtual void PostEditChangeProperty(FPropertyChangedEvent& PropertyChangedEvent) override;
 	void UpdatePrecomputedData();
 #endif // WITH_EDITOR
-	virtual void PostInterpChange(UProperty* PropertyThatChanged) OVERRIDE;
-	virtual void Serialize(FArchive& Ar) OVERRIDE;
+	virtual void PostInterpChange(UProperty* PropertyThatChanged) override;
+	virtual void Serialize(FArchive& Ar) override;
 	// End UObject Interface
 
 	ENGINE_API void InitResource();
 	ENGINE_API void ReleaseResource();
 
 	// Begin UActorComponent interface.
-	virtual TSharedPtr<FComponentInstanceDataBase> GetComponentInstanceData() const OVERRIDE;
-	virtual FName GetComponentInstanceDataType() const OVERRIDE;
-	virtual void ApplyComponentInstanceData(TSharedPtr<FComponentInstanceDataBase> ComponentInstanceData) OVERRIDE;
+	virtual TSharedPtr<FComponentInstanceDataBase> GetComponentInstanceData() const override;
+	virtual FName GetComponentInstanceDataType() const override;
+	virtual void ApplyComponentInstanceData(TSharedPtr<FComponentInstanceDataBase> ComponentInstanceData) override;
 	// End UActorComponent interface.
 
 	const FAtmospherePrecomputeParameters& GetPrecomputeParameters() const { return PrecomputeParams;  }
diff --git a/Engine/Source/Runtime/Engine/Classes/Camera/CameraActor.h b/Engine/Source/Runtime/Engine/Classes/Camera/CameraActor.h
--- a/Engine/Source/Runtime/Engine/Classes/Camera/CameraActor.h
+++ b/Engine/Source/Runtime/Engine/Classes/Camera/CameraActor.h
@@ -38,11 +38,11 @@ private:
 
 public:
 	// Begin UObject interface
-	virtual void Serialize(FArchive& Ar) OVERRIDE;
-	ENGINE_API virtual void PostLoadSubobjects(FObjectInstancingGraph* OuterInstanceGraph) OVERRIDE;
+	virtual void Serialize(FArchive& Ar) override;
+	ENGINE_API virtual void PostLoadSubobjects(FObjectInstancingGraph* OuterInstanceGraph) override;
 
 #if WITH_EDITOR
-	ENGINE_API virtual void PostEditChangeProperty(FPropertyChangedEvent& PropertyChangedEvent) OVERRIDE;
+	ENGINE_API virtual void PostEditChangeProperty(FPropertyChangedEvent& PropertyChangedEvent) override;
 #endif
 	// End UObject interface
 	
diff --git a/Engine/Source/Runtime/Engine/Classes/Camera/CameraAnim.h b/Engine/Source/Runtime/Engine/Classes/Camera/CameraAnim.h
--- a/Engine/Source/Runtime/Engine/Classes/Camera/CameraAnim.h
+++ b/Engine/Source/Runtime/Engine/Classes/Camera/CameraAnim.h
@@ -48,9 +48,9 @@ protected:
 
 public:
 	// Begin UObject Interface
-	virtual void PreSave() OVERRIDE;
-	virtual void PostLoad() OVERRIDE;
-	virtual SIZE_T GetResourceSize(EResourceSizeMode::Type Mode) OVERRIDE;
+	virtual void PreSave() override;
+	virtual void PostLoad() override;
+	virtual SIZE_T GetResourceSize(EResourceSizeMode::Type Mode) override;
 	// End UObject Interface
 
 	/** 
diff --git a/Engine/Source/Runtime/Engine/Classes/Camera/CameraComponent.h b/Engine/Source/Runtime/Engine/Classes/Camera/CameraComponent.h
--- a/Engine/Source/Runtime/Engine/Classes/Camera/CameraComponent.h
+++ b/Engine/Source/Runtime/Engine/Classes/Camera/CameraComponent.h
@@ -45,14 +45,14 @@ class UCameraComponent : public USceneComponent
 	struct FPostProcessSettings PostProcessSettings;
 
 	// UActorComponent interface
-	ENGINE_API virtual void OnRegister() OVERRIDE;
-	ENGINE_API virtual void OnUnregister() OVERRIDE;
+	ENGINE_API virtual void OnRegister() override;
+	ENGINE_API virtual void OnUnregister() override;
 #if WITH_EDITOR
-	ENGINE_API virtual void CheckForErrors() OVERRIDE;
+	ENGINE_API virtual void CheckForErrors() override;
 	// End of UActorComponent interface
 
 	// UObject interface
-	ENGINE_API virtual void PostEditChangeProperty(FPropertyChangedEvent& PropertyChangedEvent) OVERRIDE;
+	ENGINE_API virtual void PostEditChangeProperty(FPropertyChangedEvent& PropertyChangedEvent) override;
 #endif
 	// End of UObject interface
 
diff --git a/Engine/Source/Runtime/Engine/Classes/Camera/CameraModifier_CameraShake.h b/Engine/Source/Runtime/Engine/Classes/Camera/CameraModifier_CameraShake.h
--- a/Engine/Source/Runtime/Engine/Classes/Camera/CameraModifier_CameraShake.h
+++ b/Engine/Source/Runtime/Engine/Classes/Camera/CameraModifier_CameraShake.h
@@ -117,7 +117,7 @@ public:
 	virtual void UpdateCameraShake(float DeltaTime, FCameraShakeInstance& Shake, struct FMinimalViewInfo& InOutPOV);
 	
 	// Begin UCameraModifer Interface
-	virtual bool ModifyCamera(APlayerCameraManager* Camera, float DeltaTime, struct FMinimalViewInfo& InOutPOV) OVERRIDE;
+	virtual bool ModifyCamera(APlayerCameraManager* Camera, float DeltaTime, struct FMinimalViewInfo& InOutPOV) override;
 	// End UCameraModifer Interface
 };
 
diff --git a/Engine/Source/Runtime/Engine/Classes/Camera/PlayerCameraManager.h b/Engine/Source/Runtime/Engine/Classes/Camera/PlayerCameraManager.h
--- a/Engine/Source/Runtime/Engine/Classes/Camera/PlayerCameraManager.h
+++ b/Engine/Source/Runtime/Engine/Classes/Camera/PlayerCameraManager.h
@@ -399,11 +399,11 @@ public:
 	APawn* GetViewTargetPawn() const;
 
 	// Begin AActor Interface
-	virtual bool ShouldTickIfViewportsOnly() const OVERRIDE;
-	virtual void PostInitializeComponents() OVERRIDE;
-	virtual void Destroyed() OVERRIDE;
-	virtual void DisplayDebug(class UCanvas* Canvas, const FDebugDisplayInfo& DebugDisplay, float& YL, float& YPos) OVERRIDE;
-	virtual void ApplyWorldOffset(const FVector& InOffset, bool bWorldShift) OVERRIDE;
+	virtual bool ShouldTickIfViewportsOnly() const override;
+	virtual void PostInitializeComponents() override;
+	virtual void Destroyed() override;
+	virtual void DisplayDebug(class UCanvas* Canvas, const FDebugDisplayInfo& DebugDisplay, float& YL, float& YPos) override;
+	virtual void ApplyWorldOffset(const FVector& InOffset, bool bWorldShift) override;
 	// End AActor Interface
 
 	/** Static.  Plays an in-world camera shake that affects all nearby players, with distance-based attenuation.
diff --git a/Engine/Source/Runtime/Engine/Classes/Components/ActorComponent.h b/Engine/Source/Runtime/Engine/Classes/Components/ActorComponent.h
--- a/Engine/Source/Runtime/Engine/Classes/Components/ActorComponent.h
+++ b/Engine/Source/Runtime/Engine/Classes/Components/ActorComponent.h
@@ -87,7 +87,7 @@ public:
 	class AActor* GetOwner() const;
 
 	UFUNCTION(BlueprintCallable, Category="Components", meta=(DeprecatedFunction))
-	virtual class UWorld* GetWorld() const OVERRIDE;
+	virtual class UWorld* GetWorld() const override;
 
 	/** See if this component contains the supplied tag */
 	UFUNCTION(BlueprintCallable, Category="Components")
@@ -134,10 +134,10 @@ public:
 	void SetNetAddressable();
 
 	/** IsNameStableForNetworking means an object can be referred to its path name (relative to outer) over the network */
-	virtual bool IsNameStableForNetworking() const OVERRIDE;
+	virtual bool IsNameStableForNetworking() const override;
 
 	/** IsSupportedForNetworking means an object can be referenced over the network */
-	virtual bool IsSupportedForNetworking() const OVERRIDE;
+	virtual bool IsSupportedForNetworking() const override;
 
 	/** Enable or disable replication. This is the equivelent of RemoteRole for actors (only a bool is required for components) */
 	UFUNCTION(BlueprintCallable, Category="Components")
@@ -429,24 +429,24 @@ public:
 	virtual void ApplyComponentInstanceData(TSharedPtr<class FComponentInstanceDataBase> ComponentInstanceData ) {}
 
 	// Begin UObject interface.
-	virtual void BeginDestroy() OVERRIDE;
-	virtual bool NeedsLoadForClient() const OVERRIDE;
-	virtual bool NeedsLoadForServer() const OVERRIDE;
-	virtual int32 GetFunctionCallspace( UFunction* Function, void* Parameters, FFrame* Stack ) OVERRIDE;
-	virtual bool CallRemoteFunction( UFunction* Function, void* Parameters, FOutParmRec* OutParms, FFrame* Stack ) OVERRIDE;
-	virtual void PostInitProperties() OVERRIDE;
-	virtual void PostRename(UObject* OldOuter, const FName OldName) OVERRIDE;
+	virtual void BeginDestroy() override;
+	virtual bool NeedsLoadForClient() const override;
+	virtual bool NeedsLoadForServer() const override;
+	virtual int32 GetFunctionCallspace( UFunction* Function, void* Parameters, FFrame* Stack ) override;
+	virtual bool CallRemoteFunction( UFunction* Function, void* Parameters, FOutParmRec* OutParms, FFrame* Stack ) override;
+	virtual void PostInitProperties() override;
+	virtual void PostRename(UObject* OldOuter, const FName OldName) override;
 #if WITH_EDITOR
-	virtual void PreEditChange(UProperty* PropertyThatWillChange) OVERRIDE;
-	virtual void PostEditChangeProperty(FPropertyChangedEvent& PropertyChangedEvent) OVERRIDE;
-	virtual void PostEditUndo() OVERRIDE;
+	virtual void PreEditChange(UProperty* PropertyThatWillChange) override;
+	virtual void PostEditChangeProperty(FPropertyChangedEvent& PropertyChangedEvent) override;
+	virtual void PostEditUndo() override;
 #endif // WITH_EDITOR
 	// End UObject interface.
 
 	// Begin IInterface_AssetUserData Interface
-	virtual void AddAssetUserData(UAssetUserData* InUserData) OVERRIDE;
-	virtual void RemoveUserDataOfClass(TSubclassOf<UAssetUserData> InUserDataClass) OVERRIDE;
-	virtual UAssetUserData* GetAssetUserDataOfClass(TSubclassOf<UAssetUserData> InUserDataClass) OVERRIDE;
+	virtual void AddAssetUserData(UAssetUserData* InUserData) override;
+	virtual void RemoveUserDataOfClass(TSubclassOf<UAssetUserData> InUserDataClass) override;
+	virtual UAssetUserData* GetAssetUserDataOfClass(TSubclassOf<UAssetUserData> InUserDataClass) override;
 	// End IInterface_AssetUserData Interface
 
 	/** See if this component is currently registered */
@@ -494,7 +494,7 @@ public:
 
 private:
 	// this is the old name of the tick function. We just want to avoid mistakes with an attempt to override this
-	virtual void Tick( float DeltaTime ) FINAL { check(0); }
+	virtual void Tick( float DeltaTime ) final { check(0); }
 
 #endif
 };
diff --git a/Engine/Source/Runtime/Engine/Classes/Components/ApplicationLifecycleComponent.h b/Engine/Source/Runtime/Engine/Classes/Components/ApplicationLifecycleComponent.h
--- a/Engine/Source/Runtime/Engine/Classes/Components/ApplicationLifecycleComponent.h
+++ b/Engine/Source/Runtime/Engine/Classes/Components/ApplicationLifecycleComponent.h
@@ -39,8 +39,8 @@ class ENGINE_API UApplicationLifecycleComponent : public UActorComponent
 	FApplicationLifetimeDelegate ApplicationWillTerminateDelegate;
 
 public:
-	void OnRegister() OVERRIDE;
-	void OnUnregister() OVERRIDE;
+	void OnRegister() override;
+	void OnUnregister() override;
 
 private:
 	/** Native handlers that get registered with the actual FCoreDelegates, and then proceed to broadcast to the delegates above */
diff --git a/Engine/Source/Runtime/Engine/Classes/Components/ArrowComponent.h b/Engine/Source/Runtime/Engine/Classes/Components/ArrowComponent.h
--- a/Engine/Source/Runtime/Engine/Classes/Components/ArrowComponent.h
+++ b/Engine/Source/Runtime/Engine/Classes/Components/ArrowComponent.h
@@ -57,11 +57,11 @@ class UArrowComponent : public UPrimitiveComponent
 	virtual void SetArrowColor_New(FLinearColor NewColor);
 
 	// Begin UPrimitiveComponent interface.
-	virtual FPrimitiveSceneProxy* CreateSceneProxy() OVERRIDE;
+	virtual FPrimitiveSceneProxy* CreateSceneProxy() override;
 	// End UPrimitiveComponent interface.
 
 	// Begin USceneComponent interface.
-	virtual FBoxSphereBounds CalcBounds(const FTransform & LocalToWorld) const OVERRIDE;
+	virtual FBoxSphereBounds CalcBounds(const FTransform & LocalToWorld) const override;
 	// Begin USceneComponent interface.
 
 #if WITH_EDITORONLY_DATA
diff --git a/Engine/Source/Runtime/Engine/Classes/Components/AudioComponent.h b/Engine/Source/Runtime/Engine/Classes/Components/AudioComponent.h
--- a/Engine/Source/Runtime/Engine/Classes/Components/AudioComponent.h
+++ b/Engine/Source/Runtime/Engine/Classes/Components/AudioComponent.h
@@ -281,16 +281,16 @@ public:
 
 	// Begin UObject interface.
 #if WITH_EDITOR
-	virtual void PostEditChangeProperty(FPropertyChangedEvent& PropertyChangedEvent) OVERRIDE;
+	virtual void PostEditChangeProperty(FPropertyChangedEvent& PropertyChangedEvent) override;
 #endif // WITH_EDITOR
-	virtual void PostLoad() OVERRIDE;
-	virtual FString GetDetailedInfoInternal() const OVERRIDE;
+	virtual void PostLoad() override;
+	virtual FString GetDetailedInfoInternal() const override;
 	// End UObject interface.
 
 	// Begin USceneComponent Interface
-	virtual void Activate(bool bReset=false) OVERRIDE;
-	virtual void Deactivate() OVERRIDE;
-	virtual void OnUpdateTransform(bool bSkipPhysicsMove) OVERRIDE;
+	virtual void Activate(bool bReset=false) override;
+	virtual void Deactivate() override;
+	virtual void OnUpdateTransform(bool bSkipPhysicsMove) override;
 	// End USceneComponent Interface
 
 	/** @return true if this component is currently playing a SoundCue. */
@@ -300,9 +300,9 @@ public:
 
 	// Begin ActorComponent interface.
 #if WITH_EDITORONLY_DATA
-	virtual void OnRegister() OVERRIDE;
+	virtual void OnRegister() override;
 #endif
-	virtual void OnUnregister() OVERRIDE;
+	virtual void OnUnregister() override;
 	// End ActorComponent interface.
 
 	const FAttenuationSettings* GetAttenuationSettingsToApply() const;
diff --git a/Engine/Source/Runtime/Engine/Classes/Components/BillboardComponent.h b/Engine/Source/Runtime/Engine/Classes/Components/BillboardComponent.h
--- a/Engine/Source/Runtime/Engine/Classes/Components/BillboardComponent.h
+++ b/Engine/Source/Runtime/Engine/Classes/Components/BillboardComponent.h
@@ -60,8 +60,8 @@ class UBillboardComponent : public UPrimitiveComponent
 
 
 	// Begin UPrimitiveComponent Interface
-	virtual FPrimitiveSceneProxy* CreateSceneProxy() OVERRIDE;
-	virtual FBoxSphereBounds CalcBounds(const FTransform & LocalToWorld) const OVERRIDE;
+	virtual FPrimitiveSceneProxy* CreateSceneProxy() override;
+	virtual FBoxSphereBounds CalcBounds(const FTransform & LocalToWorld) const override;
 	// End UPrimitiveComponent Interface
 
 #if WITH_EDITORONLY_DATA
diff --git a/Engine/Source/Runtime/Engine/Classes/Components/BoxComponent.h b/Engine/Source/Runtime/Engine/Classes/Components/BoxComponent.h
--- a/Engine/Source/Runtime/Engine/Classes/Components/BoxComponent.h
+++ b/Engine/Source/Runtime/Engine/Classes/Components/BoxComponent.h
@@ -35,17 +35,17 @@ public:
 	FVector GetUnscaledBoxExtent() const;
 
 	// Begin UPrimitiveComponent interface.
-	virtual bool IsZeroExtent() const OVERRIDE;
-	virtual FPrimitiveSceneProxy* CreateSceneProxy() OVERRIDE;
-	virtual struct FCollisionShape GetCollisionShape(float Inflation = 0.0f) const OVERRIDE;
+	virtual bool IsZeroExtent() const override;
+	virtual FPrimitiveSceneProxy* CreateSceneProxy() override;
+	virtual struct FCollisionShape GetCollisionShape(float Inflation = 0.0f) const override;
 	// End UPrimitiveComponent interface.
 
 	// Begin USceneComponent interface
-	virtual FBoxSphereBounds CalcBounds(const FTransform & LocalToWorld) const OVERRIDE;
+	virtual FBoxSphereBounds CalcBounds(const FTransform & LocalToWorld) const override;
 	// End USceneComponent interface
 
 	// Begin UShapeComponent interface
-	virtual void UpdateBodySetup() OVERRIDE;
+	virtual void UpdateBodySetup() override;
 	// End UShapeComponent interface
 
 	// Sets the box extents without triggering a render or physics update.
diff --git a/Engine/Source/Runtime/Engine/Classes/Components/BoxReflectionCaptureComponent.h b/Engine/Source/Runtime/Engine/Classes/Components/BoxReflectionCaptureComponent.h
--- a/Engine/Source/Runtime/Engine/Classes/Components/BoxReflectionCaptureComponent.h
+++ b/Engine/Source/Runtime/Engine/Classes/Components/BoxReflectionCaptureComponent.h
@@ -26,7 +26,7 @@ public:
 
 	// Begin UObject Interface
 #if WITH_EDITOR
-	virtual void PostEditChangeProperty(FPropertyChangedEvent& PropertyChangedEvent) OVERRIDE;
+	virtual void PostEditChangeProperty(FPropertyChangedEvent& PropertyChangedEvent) override;
 #endif // WITH_EDITOR
 	// End UObject Interface
 };
diff --git a/Engine/Source/Runtime/Engine/Classes/Components/BrushComponent.h b/Engine/Source/Runtime/Engine/Classes/Components/BrushComponent.h
--- a/Engine/Source/Runtime/Engine/Classes/Components/BrushComponent.h
+++ b/Engine/Source/Runtime/Engine/Classes/Components/BrushComponent.h
@@ -21,28 +21,28 @@ class UBrushComponent : public UPrimitiveComponent
 	FVector PrePivot;
 
 	// Begin UObject interface
-	virtual void PostLoad() OVERRIDE;
-	virtual SIZE_T GetResourceSize(EResourceSizeMode::Type Mode) OVERRIDE;
+	virtual void PostLoad() override;
+	virtual SIZE_T GetResourceSize(EResourceSizeMode::Type Mode) override;
 	// End UObject interface
 
 	// Begin USceneComponent interface
-	virtual FBoxSphereBounds CalcBounds(const FTransform & LocalToWorld) const OVERRIDE;
-	virtual FVector GetCustomLocation() const OVERRIDE;
+	virtual FBoxSphereBounds CalcBounds(const FTransform & LocalToWorld) const override;
+	virtual FVector GetCustomLocation() const override;
 protected:
 	
 	/** Calculate the new ComponentToWorld transform for this component.
 	Parent is optional and can be used for computing ComponentToWorld based on arbitrary USceneComponent.
 	If Parent is not passed in we use the component's AttachParent*/
-	virtual FTransform CalcNewComponentToWorld(const FTransform& NewRelativeTransform, const USceneComponent * Parent = NULL) const OVERRIDE;
+	virtual FTransform CalcNewComponentToWorld(const FTransform& NewRelativeTransform, const USceneComponent * Parent = NULL) const override;
 	// End USceneComponent interface
 
 public:
 
 	// Begin UPrimitiveComponent interface.
-	virtual FPrimitiveSceneProxy* CreateSceneProxy() OVERRIDE;
-	virtual class UBodySetup* GetBodySetup() OVERRIDE { return BrushBodySetup; };
-	virtual void GetUsedMaterials( TArray<UMaterialInterface*>& OutMaterials ) const OVERRIDE;
-	virtual uint8 GetStaticDepthPriorityGroup() const OVERRIDE;
+	virtual FPrimitiveSceneProxy* CreateSceneProxy() override;
+	virtual class UBodySetup* GetBodySetup() override { return BrushBodySetup; };
+	virtual void GetUsedMaterials( TArray<UMaterialInterface*>& OutMaterials ) const override;
+	virtual uint8 GetStaticDepthPriorityGroup() const override;
 	// End UPrimitiveComponent interface.
 
 	/** Create the AggGeom collection-of-convex-primitives from the Brush UModel data. */
diff --git a/Engine/Source/Runtime/Engine/Classes/Components/CapsuleComponent.h b/Engine/Source/Runtime/Engine/Classes/Components/CapsuleComponent.h
--- a/Engine/Source/Runtime/Engine/Classes/Components/CapsuleComponent.h
+++ b/Engine/Source/Runtime/Engine/Classes/Components/CapsuleComponent.h
@@ -54,26 +54,26 @@ public:
 	void SetCapsuleHalfHeight(float HalfHeight, bool bUpdateOverlaps=true);
 
 	// Begin UObject interface
-	virtual void Serialize(FArchive& Ar) OVERRIDE;
+	virtual void Serialize(FArchive& Ar) override;
 #if WITH_EDITOR
-	virtual void PostEditChangeProperty(FPropertyChangedEvent& PropertyChangedEvent) OVERRIDE;
+	virtual void PostEditChangeProperty(FPropertyChangedEvent& PropertyChangedEvent) override;
 #endif // WITH_EDITOR
 	// End UObject interface
 
 	// Begin USceneComponent interface
-	virtual FBoxSphereBounds CalcBounds(const FTransform & LocalToWorld) const OVERRIDE;
-	virtual void CalcBoundingCylinder(float& CylinderRadius, float& CylinderHalfHeight) const OVERRIDE;
+	virtual FBoxSphereBounds CalcBounds(const FTransform & LocalToWorld) const override;
+	virtual void CalcBoundingCylinder(float& CylinderRadius, float& CylinderHalfHeight) const override;
 	// End USceneComponent interface
 
 	// Begin UPrimitiveComponent interface.
-	virtual FPrimitiveSceneProxy* CreateSceneProxy() OVERRIDE;
-	virtual bool IsZeroExtent() const OVERRIDE;
-	virtual struct FCollisionShape GetCollisionShape(float Inflation = 0.0f) const OVERRIDE;
-	virtual bool AreSymmetricRotations(const FQuat& A, const FQuat& B, const FVector& Scale3D) const OVERRIDE;
+	virtual FPrimitiveSceneProxy* CreateSceneProxy() override;
+	virtual bool IsZeroExtent() const override;
+	virtual struct FCollisionShape GetCollisionShape(float Inflation = 0.0f) const override;
+	virtual bool AreSymmetricRotations(const FQuat& A, const FQuat& B, const FVector& Scale3D) const override;
 	// End UPrimitiveComponent interface.
 
 	// Begin UShapeComponent interface
-	virtual void UpdateBodySetup() OVERRIDE;
+	virtual void UpdateBodySetup() override;
 	// End UShapeComponent interface
 
 	// @return the capsule radius scaled by the component scale.
diff --git a/Engine/Source/Runtime/Engine/Classes/Components/ChildActorComponent.h b/Engine/Source/Runtime/Engine/Classes/Components/ChildActorComponent.h
--- a/Engine/Source/Runtime/Engine/Classes/Components/ChildActorComponent.h
+++ b/Engine/Source/Runtime/Engine/Classes/Components/ChildActorComponent.h
@@ -27,11 +27,11 @@ class UChildActorComponent : public USceneComponent
 	TSharedPtr<FChildActorComponentInstanceData> CachedInstanceData;
 
 	// Begin ActorComponent interface.
-	virtual void OnComponentCreated() OVERRIDE;
-	virtual void OnComponentDestroyed() OVERRIDE;
-	virtual TSharedPtr<FComponentInstanceDataBase> GetComponentInstanceData() const OVERRIDE;
-	virtual FName GetComponentInstanceDataType() const OVERRIDE;
-	virtual void ApplyComponentInstanceData(TSharedPtr<FComponentInstanceDataBase> ComponentInstanceData) OVERRIDE;
+	virtual void OnComponentCreated() override;
+	virtual void OnComponentDestroyed() override;
+	virtual TSharedPtr<FComponentInstanceDataBase> GetComponentInstanceData() const override;
+	virtual FName GetComponentInstanceDataType() const override;
+	virtual void ApplyComponentInstanceData(TSharedPtr<FComponentInstanceDataBase> ComponentInstanceData) override;
 
 	// End ActorComponent interface.
 
diff --git a/Engine/Source/Runtime/Engine/Classes/Components/DecalComponent.h b/Engine/Source/Runtime/Engine/Classes/Components/DecalComponent.h
--- a/Engine/Source/Runtime/Engine/Classes/Components/DecalComponent.h
+++ b/Engine/Source/Runtime/Engine/Classes/Components/DecalComponent.h
@@ -76,13 +76,13 @@ public:
 	}
 	
 	// Begin UActorComponent Interface
-	virtual void CreateRenderState_Concurrent() OVERRIDE;
-	virtual void DestroyRenderState_Concurrent() OVERRIDE;
-	virtual void SendRenderTransform_Concurrent() OVERRIDE;
+	virtual void CreateRenderState_Concurrent() override;
+	virtual void DestroyRenderState_Concurrent() override;
+	virtual void SendRenderTransform_Concurrent() override;
 	// End UActorComponent Interface
 
 	// Begin USceneComponent Interface
-	virtual FBoxSphereBounds CalcBounds(const FTransform & LocalToWorld) const OVERRIDE;
+	virtual FBoxSphereBounds CalcBounds(const FTransform & LocalToWorld) const override;
 	// End USceneComponent Interface
 };
 
diff --git a/Engine/Source/Runtime/Engine/Classes/Components/DestructibleComponent.h b/Engine/Source/Runtime/Engine/Classes/Components/DestructibleComponent.h
--- a/Engine/Source/Runtime/Engine/Classes/Components/DestructibleComponent.h
+++ b/Engine/Source/Runtime/Engine/Classes/Components/DestructibleComponent.h
@@ -76,9 +76,9 @@ class ENGINE_API UDestructibleComponent : public USkinnedMeshComponent
 
 #if WITH_EDITOR
 	// Begin UObject interface.
-	virtual void PostEditChangeProperty( struct FPropertyChangedEvent& PropertyChangedEvent) OVERRIDE;
+	virtual void PostEditChangeProperty( struct FPropertyChangedEvent& PropertyChangedEvent) override;
 #endif // WITH_EDITOR
-	virtual void Serialize(FArchive& Ar) OVERRIDE;
+	virtual void Serialize(FArchive& Ar) override;
 	// End UObject interface.
 
 	// Take damage
@@ -101,39 +101,39 @@ public:
 #endif	//WITH_APEX
 
 	// Begin USceneComponent interface.
-	virtual FBoxSphereBounds CalcBounds(const FTransform & LocalToWorld) const OVERRIDE;
-	virtual void OnUpdateTransform(bool bSkipPhysicsMove) OVERRIDE;
-	virtual void Activate(bool bReset=false) OVERRIDE;
-	virtual void Deactivate() OVERRIDE;
+	virtual FBoxSphereBounds CalcBounds(const FTransform & LocalToWorld) const override;
+	virtual void OnUpdateTransform(bool bSkipPhysicsMove) override;
+	virtual void Activate(bool bReset=false) override;
+	virtual void Deactivate() override;
 	// End USceneComponent interface.
 
 	// Begin UActorComponent interface.
-	virtual void CreatePhysicsState() OVERRIDE;
-	virtual void DestroyPhysicsState() OVERRIDE;
-	virtual class UBodySetup* GetBodySetup() OVERRIDE;
+	virtual void CreatePhysicsState() override;
+	virtual void DestroyPhysicsState() override;
+	virtual class UBodySetup* GetBodySetup() override;
 	// End UActorComponent interface.
 
 	// Begin UPrimitiveComponent interface.
-	virtual FBodyInstance* GetBodyInstance(FName BoneName = NAME_None) const OVERRIDE;
-	virtual bool IsAnySimulatingPhysics() const OVERRIDE;
-
-	virtual void AddImpulse(FVector Impulse, FName BoneName = NAME_None, bool bVelChange = false) OVERRIDE;
-	virtual void AddImpulseAtLocation(FVector Impulse, FVector Position, FName BoneName = NAME_None) OVERRIDE;
-	virtual void AddForce(FVector Force, FName BoneName = NAME_None) OVERRIDE;
-	virtual void AddForceAtLocation(FVector Force, FVector Location, FName BoneName = NAME_None) OVERRIDE;
-	virtual void AddRadialImpulse(FVector Origin, float Radius, float Strength, ERadialImpulseFalloff Falloff, bool bVelChange=false) OVERRIDE;
-	virtual void AddRadialForce(FVector Origin, float Radius, float Strength, ERadialImpulseFalloff Falloff) OVERRIDE;
-	virtual void ReceiveComponentDamage(float DamageAmount, FDamageEvent const& DamageEvent, AController* EventInstigator, AActor* DamageCauser) OVERRIDE;
-
-	virtual bool LineTraceComponent( FHitResult& OutHit, const FVector Start, const FVector End, const FCollisionQueryParams& Params ) OVERRIDE;
-	virtual bool SweepComponent( FHitResult& OutHit, const FVector Start, const FVector End, const FCollisionShape& CollisionShape, bool bTraceComplex=false) OVERRIDE;
+	virtual FBodyInstance* GetBodyInstance(FName BoneName = NAME_None) const override;
+	virtual bool IsAnySimulatingPhysics() const override;
+
+	virtual void AddImpulse(FVector Impulse, FName BoneName = NAME_None, bool bVelChange = false) override;
+	virtual void AddImpulseAtLocation(FVector Impulse, FVector Position, FName BoneName = NAME_None) override;
+	virtual void AddForce(FVector Force, FName BoneName = NAME_None) override;
+	virtual void AddForceAtLocation(FVector Force, FVector Location, FName BoneName = NAME_None) override;
+	virtual void AddRadialImpulse(FVector Origin, float Radius, float Strength, ERadialImpulseFalloff Falloff, bool bVelChange=false) override;
+	virtual void AddRadialForce(FVector Origin, float Radius, float Strength, ERadialImpulseFalloff Falloff) override;
+	virtual void ReceiveComponentDamage(float DamageAmount, FDamageEvent const& DamageEvent, AController* EventInstigator, AActor* DamageCauser) override;
+
+	virtual bool LineTraceComponent( FHitResult& OutHit, const FVector Start, const FVector End, const FCollisionQueryParams& Params ) override;
+	virtual bool SweepComponent( FHitResult& OutHit, const FVector Start, const FVector End, const FCollisionShape& CollisionShape, bool bTraceComplex=false) override;
 	// End UPrimitiveComponent interface.
 
 	// Begin SkinnedMeshComponent interface.
-	virtual bool ShouldUpdateTransform(bool bLODHasChanged) const OVERRIDE;
-	virtual void RefreshBoneTransforms(FActorComponentTickFunction* TickFunction = NULL) OVERRIDE;
-	virtual void SetSkeletalMesh(USkeletalMesh* InSkelMesh) OVERRIDE;
-	virtual FTransform GetSocketTransform(FName InSocketName, ERelativeTransformSpace TransformSpace = RTS_World) const OVERRIDE;
+	virtual bool ShouldUpdateTransform(bool bLODHasChanged) const override;
+	virtual void RefreshBoneTransforms(FActorComponentTickFunction* TickFunction = NULL) override;
+	virtual void SetSkeletalMesh(USkeletalMesh* InSkelMesh) override;
+	virtual FTransform GetSocketTransform(FName InSocketName, ERelativeTransformSpace TransformSpace = RTS_World) const override;
 	// End SkinnedMeshComponent interface.
 
 
@@ -181,7 +181,7 @@ public:
 
 	// End DestructibleComponent interface.
 
-	virtual bool DoCustomNavigableGeometryExport(struct FNavigableGeometryExport* GeomExport) const OVERRIDE;
+	virtual bool DoCustomNavigableGeometryExport(struct FNavigableGeometryExport* GeomExport) const override;
 
 	FORCEINLINE static int32 ChunkIdxToBoneIdx(int32 ChunkIdx) { return ChunkIdx + 1; }
 	FORCEINLINE static int32 BoneIdxToChunkIdx(int32 BoneIdx) { return FMath::Max(BoneIdx - 1, 0); }
diff --git a/Engine/Source/Runtime/Engine/Classes/Components/DirectionalLightComponent.h b/Engine/Source/Runtime/Engine/Classes/Components/DirectionalLightComponent.h
--- a/Engine/Source/Runtime/Engine/Classes/Components/DirectionalLightComponent.h
+++ b/Engine/Source/Runtime/Engine/Classes/Components/DirectionalLightComponent.h
@@ -129,21 +129,21 @@ class UDirectionalLightComponent : public ULightComponent
 	// ULightComponent interface.
 	virtual FVector4 GetLightPosition() const;
 	virtual ELightComponentType GetLightType() const;
-	virtual FLightSceneProxy* CreateSceneProxy() const OVERRIDE;
-	virtual bool IsUsedAsAtmosphereSunLight() const OVERRIDE
+	virtual FLightSceneProxy* CreateSceneProxy() const override;
+	virtual bool IsUsedAsAtmosphereSunLight() const override
 	{
 		return bUsedAsAtmosphereSunLight;
 	}
 
 	// Begin UObject Interface
 #if WITH_EDITOR
-	virtual void PostEditChangeProperty(FPropertyChangedEvent& PropertyChangedEvent) OVERRIDE;
-	virtual bool CanEditChange(const UProperty* InProperty) const OVERRIDE;
+	virtual void PostEditChangeProperty(FPropertyChangedEvent& PropertyChangedEvent) override;
+	virtual bool CanEditChange(const UProperty* InProperty) const override;
 #endif // WITH_EDITOR
-	virtual void Serialize(FArchive& Ar) OVERRIDE;
+	virtual void Serialize(FArchive& Ar) override;
 	// Begin UObject Interface
 
-	virtual void InvalidateLightingCacheDetailed(bool bInvalidateBuildEnqueuedLighting, bool bTranslationOnly) OVERRIDE;
+	virtual void InvalidateLightingCacheDetailed(bool bInvalidateBuildEnqueuedLighting, bool bTranslationOnly) override;
 };
 
 
diff --git a/Engine/Source/Runtime/Engine/Classes/Components/DrawFrustumComponent.h b/Engine/Source/Runtime/Engine/Classes/Components/DrawFrustumComponent.h
--- a/Engine/Source/Runtime/Engine/Classes/Components/DrawFrustumComponent.h
+++ b/Engine/Source/Runtime/Engine/Classes/Components/DrawFrustumComponent.h
@@ -40,8 +40,8 @@ class UDrawFrustumComponent : public UPrimitiveComponent
 
 
 	// Begin UPrimitiveComponent interface.
-	virtual FPrimitiveSceneProxy* CreateSceneProxy() OVERRIDE;
-	virtual FBoxSphereBounds CalcBounds(const FTransform & LocalToWorld) const OVERRIDE;
+	virtual FPrimitiveSceneProxy* CreateSceneProxy() override;
+	virtual FBoxSphereBounds CalcBounds(const FTransform & LocalToWorld) const override;
 	// End UPrimitiveComponent interface.
 };
 
diff --git a/Engine/Source/Runtime/Engine/Classes/Components/ExponentialHeightFogComponent.h b/Engine/Source/Runtime/Engine/Classes/Components/ExponentialHeightFogComponent.h
--- a/Engine/Source/Runtime/Engine/Classes/Components/ExponentialHeightFogComponent.h
+++ b/Engine/Source/Runtime/Engine/Classes/Components/ExponentialHeightFogComponent.h
@@ -82,20 +82,20 @@ public:
 
 protected:
 	// Begin UActorComponent interface.
-	virtual void CreateRenderState_Concurrent() OVERRIDE;
-	virtual void SendRenderTransform_Concurrent() OVERRIDE;
-	virtual void DestroyRenderState_Concurrent() OVERRIDE;
+	virtual void CreateRenderState_Concurrent() override;
+	virtual void SendRenderTransform_Concurrent() override;
+	virtual void DestroyRenderState_Concurrent() override;
 	// End UActorComponent interface.
 
 	void AddFogIfNeeded();
 
 public:
 	// Begin UObject Interface
-	virtual void PostLoad() OVERRIDE; 
+	virtual void PostLoad() override; 
 #if WITH_EDITOR
-	virtual void PostEditChangeProperty(FPropertyChangedEvent& PropertyChangedEvent) OVERRIDE;
+	virtual void PostEditChangeProperty(FPropertyChangedEvent& PropertyChangedEvent) override;
 #endif // WITH_EDITOR
-	virtual void PostInterpChange(UProperty* PropertyThatChanged) OVERRIDE;
+	virtual void PostInterpChange(UProperty* PropertyThatChanged) override;
 	// End UObject Interface
 };
 
diff --git a/Engine/Source/Runtime/Engine/Classes/Components/InstancedStaticMeshComponent.h b/Engine/Source/Runtime/Engine/Classes/Components/InstancedStaticMeshComponent.h
--- a/Engine/Source/Runtime/Engine/Classes/Components/InstancedStaticMeshComponent.h
+++ b/Engine/Source/Runtime/Engine/Classes/Components/InstancedStaticMeshComponent.h
@@ -125,31 +125,31 @@ public:
 	TArray<FBodyInstance*> InstanceBodies;
 
 	// Begin UActorComponent interface 
-	virtual TSharedPtr<FComponentInstanceDataBase> GetComponentInstanceData() const OVERRIDE;
-	virtual FName GetComponentInstanceDataType() const OVERRIDE;
-	virtual void ApplyComponentInstanceData(TSharedPtr<FComponentInstanceDataBase> ComponentInstanceData) OVERRIDE;
+	virtual TSharedPtr<FComponentInstanceDataBase> GetComponentInstanceData() const override;
+	virtual FName GetComponentInstanceDataType() const override;
+	virtual void ApplyComponentInstanceData(TSharedPtr<FComponentInstanceDataBase> ComponentInstanceData) override;
 	// End UActorComponent interface 
 
 	// Begin UPrimitiveComponent Interface
-	virtual FPrimitiveSceneProxy* CreateSceneProxy() OVERRIDE;
-	virtual void CreatePhysicsState() OVERRIDE;
-	virtual void DestroyPhysicsState() OVERRIDE;
+	virtual FPrimitiveSceneProxy* CreateSceneProxy() override;
+	virtual void CreatePhysicsState() override;
+	virtual void DestroyPhysicsState() override;
 
-	virtual FBoxSphereBounds CalcBounds(const FTransform& BoundTransform) const OVERRIDE;
+	virtual FBoxSphereBounds CalcBounds(const FTransform& BoundTransform) const override;
 #if WITH_EDITOR
-	virtual void GetStaticLightingInfo(FStaticLightingPrimitiveInfo& OutPrimitiveInfo,const TArray<ULightComponent*>& InRelevantLights,const FLightingBuildOptions& Options) OVERRIDE;
+	virtual void GetStaticLightingInfo(FStaticLightingPrimitiveInfo& OutPrimitiveInfo,const TArray<ULightComponent*>& InRelevantLights,const FLightingBuildOptions& Options) override;
 #endif
-	virtual void GetLightAndShadowMapMemoryUsage( int32& LightMapMemoryUsage, int32& ShadowMapMemoryUsage ) const OVERRIDE;
-	virtual void ApplyWorldOffset(const FVector& InOffset, bool bWorldShift) OVERRIDE;
+	virtual void GetLightAndShadowMapMemoryUsage( int32& LightMapMemoryUsage, int32& ShadowMapMemoryUsage ) const override;
+	virtual void ApplyWorldOffset(const FVector& InOffset, bool bWorldShift) override;
 
-	virtual bool DoCustomNavigableGeometryExport(struct FNavigableGeometryExport* GeomExport) const OVERRIDE;
+	virtual bool DoCustomNavigableGeometryExport(struct FNavigableGeometryExport* GeomExport) const override;
 	// End UPrimitiveComponent Interface
 
 	//Begin UObject Interface
-	virtual void Serialize(FArchive& Ar) OVERRIDE;
-	virtual SIZE_T GetResourceSize(EResourceSizeMode::Type Mode) OVERRIDE;
+	virtual void Serialize(FArchive& Ar) override;
+	virtual SIZE_T GetResourceSize(EResourceSizeMode::Type Mode) override;
 #if WITH_EDITOR
-	virtual void PostEditChangeChainProperty(FPropertyChangedChainEvent& PropertyChangedEvent) OVERRIDE;
+	virtual void PostEditChangeChainProperty(FPropertyChangedChainEvent& PropertyChangedEvent) override;
 #endif
 	//End UObject Interface
 
diff --git a/Engine/Source/Runtime/Engine/Classes/Components/InteractiveFoliageComponent.h b/Engine/Source/Runtime/Engine/Classes/Components/InteractiveFoliageComponent.h
--- a/Engine/Source/Runtime/Engine/Classes/Components/InteractiveFoliageComponent.h
+++ b/Engine/Source/Runtime/Engine/Classes/Components/InteractiveFoliageComponent.h
@@ -14,11 +14,11 @@ public:
 	class FInteractiveFoliageSceneProxy* FoliageSceneProxy;
 
 	// Begin UPrimitiveComponent Interface
-	virtual FPrimitiveSceneProxy* CreateSceneProxy() OVERRIDE;
+	virtual FPrimitiveSceneProxy* CreateSceneProxy() override;
 	// End UPrimitiveComponent Interface
 
 	// Begin UActorComponent Interface
-	virtual void DestroyRenderState_Concurrent() OVERRIDE;
+	virtual void DestroyRenderState_Concurrent() override;
 	// End UActorComponent Interface
 
 	friend class AInteractiveFoliageActor;
diff --git a/Engine/Source/Runtime/Engine/Classes/Components/LightComponent.h b/Engine/Source/Runtime/Engine/Classes/Components/LightComponent.h
--- a/Engine/Source/Runtime/Engine/Classes/Components/LightComponent.h
+++ b/Engine/Source/Runtime/Engine/Classes/Components/LightComponent.h
@@ -201,7 +201,7 @@ public:
 	/**
 	 * Update/reset light GUIDs.
 	 */
-	virtual void UpdateLightGUIDs() OVERRIDE;
+	virtual void UpdateLightGUIDs() override;
 
 	/**
 	 * Check whether a given primitive will cast shadows from this light.
@@ -220,19 +220,19 @@ public:
 	float ComputeLightBrightness() const;
 
 	// Begin UObject interface.
-	virtual void PostLoad() OVERRIDE;
+	virtual void PostLoad() override;
 #if WITH_EDITOR
-	virtual void PreEditUndo() OVERRIDE;
-	virtual void PostEditUndo() OVERRIDE;
-	virtual bool CanEditChange(const UProperty* InProperty) const OVERRIDE;
-	virtual void PostEditChangeProperty(FPropertyChangedEvent& PropertyChangedEvent) OVERRIDE;
-	virtual void UpdateLightSpriteTexture() OVERRIDE;
+	virtual void PreEditUndo() override;
+	virtual void PostEditUndo() override;
+	virtual bool CanEditChange(const UProperty* InProperty) const override;
+	virtual void PostEditChangeProperty(FPropertyChangedEvent& PropertyChangedEvent) override;
+	virtual void UpdateLightSpriteTexture() override;
 #endif // WITH_EDITOR
 	// End UObject interface.
 
-	virtual TSharedPtr<FComponentInstanceDataBase> GetComponentInstanceData() const OVERRIDE;
-	virtual FName GetComponentInstanceDataType() const OVERRIDE;
-	virtual void ApplyComponentInstanceData(TSharedPtr<FComponentInstanceDataBase> ComponentInstanceData) OVERRIDE;
+	virtual TSharedPtr<FComponentInstanceDataBase> GetComponentInstanceData() const override;
+	virtual FName GetComponentInstanceDataType() const override;
+	virtual void ApplyComponentInstanceData(TSharedPtr<FComponentInstanceDataBase> ComponentInstanceData) override;
 
 	/** @return number of material elements in this primitive */
 	virtual int32 GetNumMaterials() const;
@@ -251,14 +251,14 @@ public:
 
 protected:
 	// Begin UActorComponent Interface
-	virtual void OnRegister() OVERRIDE;
-	virtual void CreateRenderState_Concurrent() OVERRIDE;
-	virtual void SendRenderTransform_Concurrent() OVERRIDE;
-	virtual void DestroyRenderState_Concurrent() OVERRIDE;
+	virtual void OnRegister() override;
+	virtual void CreateRenderState_Concurrent() override;
+	virtual void SendRenderTransform_Concurrent() override;
+	virtual void DestroyRenderState_Concurrent() override;
 	// Begin UActorComponent Interface
 
 public:
-	virtual void InvalidateLightingCacheDetailed(bool bInvalidateBuildEnqueuedLighting, bool bTranslationOnly) OVERRIDE;
+	virtual void InvalidateLightingCacheDetailed(bool bInvalidateBuildEnqueuedLighting, bool bTranslationOnly) override;
 
 	/** Invalidates the light's cached lighting with the option to recreate the light Guids. */
 	void InvalidateLightingCacheInner(bool bRecreateLightGuids);
diff --git a/Engine/Source/Runtime/Engine/Classes/Components/LightComponentBase.h b/Engine/Source/Runtime/Engine/Classes/Components/LightComponentBase.h
--- a/Engine/Source/Runtime/Engine/Classes/Components/LightComponentBase.h
+++ b/Engine/Source/Runtime/Engine/Classes/Components/LightComponentBase.h
@@ -94,7 +94,7 @@ class ENGINE_API ULightComponentBase : public USceneComponent
 	float DynamicEditorTextureScale;
 #endif
 
-	virtual void Serialize(FArchive& Ar) OVERRIDE;
+	virtual void Serialize(FArchive& Ar) override;
 
 	/**
 	 * Called after duplication & serialization and before PostLoad. Used to e.g. make sure GUIDs remains globally unique.
@@ -103,8 +103,8 @@ class ENGINE_API ULightComponentBase : public USceneComponent
 
 #if WITH_EDITOR
 	/** UObject interface */
-	virtual void PostEditImport() OVERRIDE;
-	virtual void PostEditChangeProperty(FPropertyChangedEvent& PropertyChangedEvent) OVERRIDE;
+	virtual void PostEditImport() override;
+	virtual void PostEditChangeProperty(FPropertyChangedEvent& PropertyChangedEvent) override;
 
 	/**
 	* @return Path to the editor sprite for the light component class
@@ -150,13 +150,13 @@ class ENGINE_API ULightComponentBase : public USceneComponent
 	bool HasStaticShadowing() const;
 
 	/** UActorComponent Interface */
-	virtual void OnRegister() OVERRIDE;
+	virtual void OnRegister() override;
 
 	/** We return a small bounds to allow us to non-interpenetrates when placing lights in the level. */
-	virtual bool ShouldCollideWhenPlacing() const OVERRIDE;
+	virtual bool ShouldCollideWhenPlacing() const override;
 
 	/** Get the extent used when placing this component in the editor, used for 'pulling back' hit. */
-	virtual FBoxSphereBounds GetPlacementExtent() const OVERRIDE;
+	virtual FBoxSphereBounds GetPlacementExtent() const override;
 
 protected:
 #if WITH_EDITORONLY_DATA
diff --git a/Engine/Source/Runtime/Engine/Classes/Components/LineBatchComponent.h b/Engine/Source/Runtime/Engine/Classes/Components/LineBatchComponent.h
--- a/Engine/Source/Runtime/Engine/Classes/Components/LineBatchComponent.h
+++ b/Engine/Source/Runtime/Engine/Classes/Components/LineBatchComponent.h
@@ -150,13 +150,13 @@ class ULineBatchComponent : public UPrimitiveComponent
 	void DrawMesh(TArray<FVector> const& Verts, TArray<int32> const& Indices, FColor const& Color, uint8 DepthPriority, float LifeTime);
 
 	// Begin UPrimitiveComponent interface.
-	virtual FPrimitiveSceneProxy* CreateSceneProxy() OVERRIDE;
-	virtual FBoxSphereBounds CalcBounds(const FTransform & LocalToWorld) const OVERRIDE;
+	virtual FPrimitiveSceneProxy* CreateSceneProxy() override;
+	virtual FBoxSphereBounds CalcBounds(const FTransform & LocalToWorld) const override;
 	// End UPrimitiveComponent interface.
 	
 	
 	// Begin UActorComponent interface.
-	virtual void TickComponent(float DeltaTime, enum ELevelTick TickType, FActorComponentTickFunction *ThisTickFunction) OVERRIDE;
+	virtual void TickComponent(float DeltaTime, enum ELevelTick TickType, FActorComponentTickFunction *ThisTickFunction) override;
 	// End UActorComponent interface.
 
 	ENGINE_API void Flush();
diff --git a/Engine/Source/Runtime/Engine/Classes/Components/MaterialBillboardComponent.h b/Engine/Source/Runtime/Engine/Classes/Components/MaterialBillboardComponent.h
--- a/Engine/Source/Runtime/Engine/Classes/Components/MaterialBillboardComponent.h
+++ b/Engine/Source/Runtime/Engine/Classes/Components/MaterialBillboardComponent.h
@@ -70,7 +70,7 @@ class UMaterialBillboardComponent : public UPrimitiveComponent
 		);
 
 	// Begin UPrimitiveComponent Interface
-	virtual FPrimitiveSceneProxy* CreateSceneProxy() OVERRIDE;
-	virtual FBoxSphereBounds CalcBounds(const FTransform& LocalToWorld) const OVERRIDE;
+	virtual FPrimitiveSceneProxy* CreateSceneProxy() override;
+	virtual FBoxSphereBounds CalcBounds(const FTransform& LocalToWorld) const override;
 	// End UPrimitiveComponent Interface
 };
diff --git a/Engine/Source/Runtime/Engine/Classes/Components/MeshComponent.h b/Engine/Source/Runtime/Engine/Classes/Components/MeshComponent.h
--- a/Engine/Source/Runtime/Engine/Classes/Components/MeshComponent.h
+++ b/Engine/Source/Runtime/Engine/Classes/Components/MeshComponent.h
@@ -15,13 +15,13 @@ class ENGINE_API UMeshComponent : public UPrimitiveComponent
 
 
 	// Begin UObject Interface
-	virtual void BeginDestroy() OVERRIDE;
+	virtual void BeginDestroy() override;
 	// End UObject Interface
 
 	// Begin UPrimitiveComponent Interface
-	virtual int32 GetNumMaterials() const OVERRIDE;
-	virtual UMaterialInterface* GetMaterial(int32 ElementIndex) const OVERRIDE;
-	virtual void SetMaterial(int32 ElementIndex, UMaterialInterface* Material) OVERRIDE;
+	virtual int32 GetNumMaterials() const override;
+	virtual UMaterialInterface* GetMaterial(int32 ElementIndex) const override;
+	virtual void SetMaterial(int32 ElementIndex, UMaterialInterface* Material) override;
 	// End UPrimitiveComponent Interface
 
 	/** Accesses the scene relevance information for the materials applied to the mesh. */
diff --git a/Engine/Source/Runtime/Engine/Classes/Components/ModelComponent.h b/Engine/Source/Runtime/Engine/Classes/Components/ModelComponent.h
--- a/Engine/Source/Runtime/Engine/Classes/Components/ModelComponent.h
+++ b/Engine/Source/Runtime/Engine/Classes/Components/ModelComponent.h
@@ -67,42 +67,42 @@ public:
 	ENGINE_API void CopyElementsFrom(UModelComponent* OtherModelComponent);
 
 	// Begin UPrimitiveComponent interface.
-	virtual bool GetLightMapResolution( int32& Width, int32& Height ) const OVERRIDE;
-	virtual int32 GetStaticLightMapResolution() const OVERRIDE;
-	virtual void GetLightAndShadowMapMemoryUsage( int32& LightMapMemoryUsage, int32& ShadowMapMemoryUsage ) const OVERRIDE;
-	virtual FBoxSphereBounds CalcBounds(const FTransform & LocalToWorld) const OVERRIDE;
-	virtual FPrimitiveSceneProxy* CreateSceneProxy() OVERRIDE;
-	virtual bool ShouldRecreateProxyOnUpdateTransform() const OVERRIDE;
+	virtual bool GetLightMapResolution( int32& Width, int32& Height ) const override;
+	virtual int32 GetStaticLightMapResolution() const override;
+	virtual void GetLightAndShadowMapMemoryUsage( int32& LightMapMemoryUsage, int32& ShadowMapMemoryUsage ) const override;
+	virtual FBoxSphereBounds CalcBounds(const FTransform & LocalToWorld) const override;
+	virtual FPrimitiveSceneProxy* CreateSceneProxy() override;
+	virtual bool ShouldRecreateProxyOnUpdateTransform() const override;
 #if WITH_EDITOR
-	virtual void GetStaticLightingInfo(FStaticLightingPrimitiveInfo& OutPrimitiveInfo,const TArray<ULightComponent*>& InRelevantLights,const FLightingBuildOptions& Options) OVERRIDE;
-	virtual bool SetupLightmapResolutionViewInfo(FPrimitiveSceneProxy& Proxy) const OVERRIDE;
+	virtual void GetStaticLightingInfo(FStaticLightingPrimitiveInfo& OutPrimitiveInfo,const TArray<ULightComponent*>& InRelevantLights,const FLightingBuildOptions& Options) override;
+	virtual bool SetupLightmapResolutionViewInfo(FPrimitiveSceneProxy& Proxy) const override;
 #endif
-	virtual ELightMapInteractionType GetStaticLightingType() const OVERRIDE	{ return LMIT_Texture;	}
-	virtual void GetStreamingTextureInfo(TArray<FStreamingTexturePrimitiveInfo>& OutStreamingTextures) const OVERRIDE;
-	virtual void GetUsedMaterials(TArray<UMaterialInterface*>& OutMaterials) const OVERRIDE;
-	virtual class UBodySetup* GetBodySetup() OVERRIDE { return ModelBodySetup; };
-	virtual int32 GetNumMaterials() const OVERRIDE;
-	virtual UMaterialInterface* GetMaterial(int32 MaterialIndex) const OVERRIDE;
+	virtual ELightMapInteractionType GetStaticLightingType() const override	{ return LMIT_Texture;	}
+	virtual void GetStreamingTextureInfo(TArray<FStreamingTexturePrimitiveInfo>& OutStreamingTextures) const override;
+	virtual void GetUsedMaterials(TArray<UMaterialInterface*>& OutMaterials) const override;
+	virtual class UBodySetup* GetBodySetup() override { return ModelBodySetup; };
+	virtual int32 GetNumMaterials() const override;
+	virtual UMaterialInterface* GetMaterial(int32 MaterialIndex) const override;
 	// End UPrimitiveComponent interface.
 
 	// Begin UActorComponent interface.
-	virtual void InvalidateLightingCacheDetailed(bool bInvalidateBuildEnqueuedLighting, bool bTranslationOnly) OVERRIDE;
+	virtual void InvalidateLightingCacheDetailed(bool bInvalidateBuildEnqueuedLighting, bool bTranslationOnly) override;
 	// End UActorComponent interface.
 
 	// Begin UObject interface.
-	virtual void Serialize(FArchive& Ar) OVERRIDE;
-	virtual void PostLoad() OVERRIDE;
+	virtual void Serialize(FArchive& Ar) override;
+	virtual void PostLoad() override;
 #if WITH_EDITOR
-	virtual void PostEditUndo() OVERRIDE;
+	virtual void PostEditUndo() override;
 #endif // WITH_EDITOR
-	virtual SIZE_T GetResourceSize(EResourceSizeMode::Type Mode) OVERRIDE;
+	virtual SIZE_T GetResourceSize(EResourceSizeMode::Type Mode) override;
 	static void AddReferencedObjects(UObject* InThis, FReferenceCollector& Collector);
 	// End UObject interface.
 
 	// Begin Interface_CollisionDataProvider Interface
-	virtual bool GetPhysicsTriMeshData(struct FTriMeshCollisionData* CollisionData, bool InUseAllTriData) OVERRIDE;
-	virtual bool ContainsPhysicsTriMeshData(bool InUseAllTriData) const OVERRIDE;
-	virtual bool WantsNegXTriMesh() OVERRIDE { return false; }
+	virtual bool GetPhysicsTriMeshData(struct FTriMeshCollisionData* CollisionData, bool InUseAllTriData) override;
+	virtual bool ContainsPhysicsTriMeshData(bool InUseAllTriData) const override;
+	virtual bool WantsNegXTriMesh() override { return false; }
 	// End Interface_CollisionDataProvider Interface
 
 
diff --git a/Engine/Source/Runtime/Engine/Classes/Components/NiagaraComponent.h b/Engine/Source/Runtime/Engine/Classes/Components/NiagaraComponent.h
--- a/Engine/Source/Runtime/Engine/Classes/Components/NiagaraComponent.h
+++ b/Engine/Source/Runtime/Engine/Classes/Components/NiagaraComponent.h
@@ -25,25 +25,25 @@ class ENGINE_API UNiagaraComponent : public UPrimitiveComponent
 	float SpawnRate;
 
 	// Begin UActorComponent interface.
-	virtual void TickComponent(float DeltaTime, enum ELevelTick TickType, FActorComponentTickFunction *ThisTickFunction) OVERRIDE;
+	virtual void TickComponent(float DeltaTime, enum ELevelTick TickType, FActorComponentTickFunction *ThisTickFunction) override;
 protected:
-	virtual void OnRegister() OVERRIDE;
-	virtual void OnUnregister()  OVERRIDE;
-	virtual void SendRenderDynamicData_Concurrent() OVERRIDE;
+	virtual void OnRegister() override;
+	virtual void OnUnregister()  override;
+	virtual void SendRenderDynamicData_Concurrent() override;
 public:
 	// End UActorComponent interface.
 
 	// Begin UPrimitiveComponent Interface
-	virtual int32 GetNumMaterials() const OVERRIDE;
-	virtual class UMaterialInterface* GetMaterial(int32 ElementIndex) const OVERRIDE;
-	virtual void SetMaterial(int32 ElementIndex, class UMaterialInterface* Material) OVERRIDE;
-	virtual FBoxSphereBounds CalcBounds(const FTransform & LocalToWorld) const OVERRIDE;
-	virtual FPrimitiveSceneProxy* CreateSceneProxy() OVERRIDE;
+	virtual int32 GetNumMaterials() const override;
+	virtual class UMaterialInterface* GetMaterial(int32 ElementIndex) const override;
+	virtual void SetMaterial(int32 ElementIndex, class UMaterialInterface* Material) override;
+	virtual FBoxSphereBounds CalcBounds(const FTransform & LocalToWorld) const override;
+	virtual FPrimitiveSceneProxy* CreateSceneProxy() override;
 	// End UPrimitiveComponent Interface
 
 	// Begin UObject interface.
 #if WITH_EDITOR
-	virtual void PostEditChangeProperty(FPropertyChangedEvent& PropertyChangedEvent) OVERRIDE;
+	virtual void PostEditChangeProperty(FPropertyChangedEvent& PropertyChangedEvent) override;
 #endif // WITH_EDITOR
 	// End UObject interface.
 };
diff --git a/Engine/Source/Runtime/Engine/Classes/Components/PointLightComponent.h b/Engine/Source/Runtime/Engine/Classes/Components/PointLightComponent.h
--- a/Engine/Source/Runtime/Engine/Classes/Components/PointLightComponent.h
+++ b/Engine/Source/Runtime/Engine/Classes/Components/PointLightComponent.h
@@ -68,7 +68,7 @@ class ENGINE_API UPointLightComponent : public ULightComponent
 
 protected:
 	// Begin UActorComponent Interface
-	virtual void SendRenderTransform_Concurrent() OVERRIDE;
+	virtual void SendRenderTransform_Concurrent() override;
 	// End UActorComponent Interface
 
 public:
@@ -79,14 +79,14 @@ public:
 	virtual FBox GetBoundingBox() const;
 	virtual FSphere GetBoundingSphere() const;
 	virtual ELightComponentType GetLightType() const;
-	virtual FLightSceneProxy* CreateSceneProxy() const OVERRIDE;
+	virtual FLightSceneProxy* CreateSceneProxy() const override;
 
 	// Begin UObject Interface
-	virtual void Serialize(FArchive& Ar) OVERRIDE;
-	virtual void PostLoad() OVERRIDE;
+	virtual void Serialize(FArchive& Ar) override;
+	virtual void PostLoad() override;
 #if WITH_EDITOR
-	virtual bool CanEditChange(const UProperty* InProperty) const OVERRIDE;
-	virtual void PostEditChangeProperty(FPropertyChangedEvent& PropertyChangedEvent) OVERRIDE;
+	virtual bool CanEditChange(const UProperty* InProperty) const override;
+	virtual void PostEditChangeProperty(FPropertyChangedEvent& PropertyChangedEvent) override;
 #endif // WITH_EDITOR
 	// End UObject Interface
 
diff --git a/Engine/Source/Runtime/Engine/Classes/Components/PoseableMeshComponent.h b/Engine/Source/Runtime/Engine/Classes/Components/PoseableMeshComponent.h
--- a/Engine/Source/Runtime/Engine/Classes/Components/PoseableMeshComponent.h
+++ b/Engine/Source/Runtime/Engine/Classes/Components/PoseableMeshComponent.h
@@ -59,8 +59,8 @@ class ENGINE_API UPoseableMeshComponent : public USkinnedMeshComponent
 	void ResetBoneTransformByName(FName BoneName);
 
 	// Begin USkinnedMeshComponent Interface
-	virtual void RefreshBoneTransforms(FActorComponentTickFunction* TickFunction = NULL) OVERRIDE;
-	virtual bool AllocateTransformData() OVERRIDE;
+	virtual void RefreshBoneTransforms(FActorComponentTickFunction* TickFunction = NULL) override;
+	virtual bool AllocateTransformData() override;
 	// End USkinnedMeshComponent Interface
 
 	/**
diff --git a/Engine/Source/Runtime/Engine/Classes/Components/PostProcessComponent.h b/Engine/Source/Runtime/Engine/Classes/Components/PostProcessComponent.h
--- a/Engine/Source/Runtime/Engine/Classes/Components/PostProcessComponent.h
+++ b/Engine/Source/Runtime/Engine/Classes/Components/PostProcessComponent.h
@@ -41,8 +41,8 @@ class UPostProcessComponent : public USceneComponent, public IInterface_PostProc
 	uint32 bUnbound: 1;
  
 	// Begin IInterface_PostProcessVolume Interface
-	ENGINE_API virtual bool EncompassesPoint(FVector Point, float SphereRadius/*=0.f*/, float* OutDistanceToPoint) OVERRIDE;
-	ENGINE_API virtual FPostProcessVolumeProperties GetProperties() const OVERRIDE
+	ENGINE_API virtual bool EncompassesPoint(FVector Point, float SphereRadius/*=0.f*/, float* OutDistanceToPoint) override;
+	ENGINE_API virtual FPostProcessVolumeProperties GetProperties() const override
 	{
 		FPostProcessVolumeProperties Ret;
 		Ret.bIsEnabled = bEnabled != 0;
diff --git a/Engine/Source/Runtime/Engine/Classes/Components/PrimitiveComponent.h b/Engine/Source/Runtime/Engine/Classes/Components/PrimitiveComponent.h
--- a/Engine/Source/Runtime/Engine/Classes/Components/PrimitiveComponent.h
+++ b/Engine/Source/Runtime/Engine/Classes/Components/PrimitiveComponent.h
@@ -590,14 +590,14 @@ public:
 	 * @param OverlapsAtEndLocation		If non-null, the given list of overlaps will be used as the overlaps for this component at the current location, rather than checking for them separately.
 	 *									Generally this should only be used if this component is the RootComponent of the owning actor and overlaps with other descendant components have been verified.
 	 */
-	virtual void UpdateOverlaps(TArray<FOverlapInfo> const* PendingOverlaps=NULL, bool bDoNotifies=true, const TArray<FOverlapInfo>* OverlapsAtEndLocation=NULL) OVERRIDE;
+	virtual void UpdateOverlaps(TArray<FOverlapInfo> const* PendingOverlaps=NULL, bool bDoNotifies=true, const TArray<FOverlapInfo>* OverlapsAtEndLocation=NULL) override;
 
 	/** Tells this component to ignore collision with the specified actor when being moved. */
 	UFUNCTION(BlueprintCallable, Category = "Collision")
 	void IgnoreActorWhenMoving(AActor* Actor, bool bShouldIgnore);
 
 	/** Overridden to use the overlaps to find the physics volume. */
-	virtual void UpdatePhysicsVolume( bool bTriggerNotifiers ) OVERRIDE;
+	virtual void UpdatePhysicsVolume( bool bTriggerNotifiers ) override;
 
 	/**
 	 *  Test the collision of the supplied component at the supplied location/rotation, and determine the set of components that it overlaps
@@ -911,10 +911,10 @@ public:
 #endif
 
 	// Begin UActorComponent Interface
-	virtual void InvalidateLightingCacheDetailed(bool bInvalidateBuildEnqueuedLighting, bool bTranslationOnly) OVERRIDE;
-	virtual bool IsEditorOnly() const OVERRIDE;
-	virtual bool ShouldCreatePhysicsState() const OVERRIDE;
-	virtual bool HasValidPhysicsState() const OVERRIDE;
+	virtual void InvalidateLightingCacheDetailed(bool bInvalidateBuildEnqueuedLighting, bool bTranslationOnly) override;
+	virtual bool IsEditorOnly() const override;
+	virtual bool ShouldCreatePhysicsState() const override;
+	virtual bool HasValidPhysicsState() const override;
 	// End UActorComponent Interface
 
 	/** @return true if the owner is selected and this component is selectable */
@@ -1085,37 +1085,37 @@ protected:
 	friend class FStaticMeshComponentRecreateRenderStateContext;
 
 	// Begin USceneComponent Interface
-	virtual void OnUpdateTransform(bool bSkipPhysicsMove) OVERRIDE;
+	virtual void OnUpdateTransform(bool bSkipPhysicsMove) override;
 
 	/** Called when AttachParent changes, to allow the scene to update its attachment state. */
-	virtual void OnAttachmentChanged() OVERRIDE;
+	virtual void OnAttachmentChanged() override;
 
 	/**
 	* Called after a child is attached to this component.
 	* Note: Do not change the attachment state of the child during this call.
 	*/
-	virtual void OnChildAttached(USceneComponent* ChildComponent) OVERRIDE;
+	virtual void OnChildAttached(USceneComponent* ChildComponent) override;
 
 public:
-	virtual bool IsSimulatingPhysics(FName BoneName = NAME_None) const OVERRIDE;
+	virtual bool IsSimulatingPhysics(FName BoneName = NAME_None) const override;
 
 	// End USceneComponentInterface
 
 
 	// Begin UActorComponent Interface
 protected:
-	virtual void CreateRenderState_Concurrent() OVERRIDE;
-	virtual void SendRenderTransform_Concurrent() OVERRIDE;
-	virtual void OnRegister()  OVERRIDE;
-	virtual void OnUnregister()  OVERRIDE;
-	virtual void DestroyRenderState_Concurrent() OVERRIDE;
-	virtual void CreatePhysicsState() OVERRIDE;
-	virtual void DestroyPhysicsState() OVERRIDE;
-	virtual void OnActorEnableCollisionChanged() OVERRIDE;
+	virtual void CreateRenderState_Concurrent() override;
+	virtual void SendRenderTransform_Concurrent() override;
+	virtual void OnRegister()  override;
+	virtual void OnUnregister()  override;
+	virtual void DestroyRenderState_Concurrent() override;
+	virtual void CreatePhysicsState() override;
+	virtual void DestroyPhysicsState() override;
+	virtual void OnActorEnableCollisionChanged() override;
 public:
-	virtual void RegisterComponentTickFunctions(bool bRegister) OVERRIDE;
+	virtual void RegisterComponentTickFunctions(bool bRegister) override;
 #if WITH_EDITOR
-	virtual void CheckForErrors() OVERRIDE;
+	virtual void CheckForErrors() override;
 #endif // WITH_EDITOR	
 	// End UActorComponent Interface
 
@@ -1130,16 +1130,16 @@ protected:
 public:
 
 	// Begin UObject interface.
-	virtual void Serialize(FArchive& Ar) OVERRIDE;
+	virtual void Serialize(FArchive& Ar) override;
 #if WITH_EDITOR
-	virtual void PostEditChangeProperty(FPropertyChangedEvent& PropertyChangedEvent) OVERRIDE;
-	virtual void PostEditChangeChainProperty(FPropertyChangedChainEvent& PropertyChangedEvent) OVERRIDE;
-	virtual bool CanEditChange(const UProperty* InProperty) const OVERRIDE;
+	virtual void PostEditChangeProperty(FPropertyChangedEvent& PropertyChangedEvent) override;
+	virtual void PostEditChangeChainProperty(FPropertyChangedChainEvent& PropertyChangedEvent) override;
+	virtual bool CanEditChange(const UProperty* InProperty) const override;
 	virtual void UpdateCollisionProfile();
 #endif // WITH_EDITOR
-	virtual void PostLoad() OVERRIDE;
-	virtual void PostDuplicate(bool bDuplicateForPIE) OVERRIDE;
-	virtual SIZE_T GetResourceSize(EResourceSizeMode::Type Mode) OVERRIDE;
+	virtual void PostLoad() override;
+	virtual void PostDuplicate(bool bDuplicateForPIE) override;
+	virtual SIZE_T GetResourceSize(EResourceSizeMode::Type Mode) override;
 
 #if WITH_EDITOR
 	/**
@@ -1150,22 +1150,22 @@ public:
 	virtual void PostEditImport();
 #endif
 
-	virtual void BeginDestroy() OVERRIDE;
-	virtual void FinishDestroy() OVERRIDE;
-	virtual bool IsReadyForFinishDestroy() OVERRIDE;
-	virtual bool NeedsLoadForClient() const OVERRIDE;
-	virtual bool NeedsLoadForServer() const OVERRIDE;
+	virtual void BeginDestroy() override;
+	virtual void FinishDestroy() override;
+	virtual bool IsReadyForFinishDestroy() override;
+	virtual bool NeedsLoadForClient() const override;
+	virtual bool NeedsLoadForServer() const override;
 	// End UObject interface.
 
 	//Begin USceneComponent Interface
-	virtual void SetRelativeScale3D(FVector NewScale3D) OVERRIDE FINAL;
-	virtual bool MoveComponent(const FVector& Delta, const FRotator& NewRotation, bool bSweep, FHitResult* OutHit = NULL, EMoveComponentFlags MoveFlags = MOVECOMP_NoFlags) OVERRIDE;
-	virtual bool IsWorldGeometry() const OVERRIDE;
-	virtual ECollisionEnabled::Type GetCollisionEnabled() const OVERRIDE;
-	virtual ECollisionResponse GetCollisionResponseToChannel(ECollisionChannel Channel) const OVERRIDE;
-	virtual ECollisionChannel GetCollisionObjectType() const OVERRIDE;
-	virtual const FCollisionResponseContainer & GetCollisionResponseToChannels() const OVERRIDE;
-	virtual FVector GetComponentVelocity() const OVERRIDE;
+	virtual void SetRelativeScale3D(FVector NewScale3D) override final;
+	virtual bool MoveComponent(const FVector& Delta, const FRotator& NewRotation, bool bSweep, FHitResult* OutHit = NULL, EMoveComponentFlags MoveFlags = MOVECOMP_NoFlags) override;
+	virtual bool IsWorldGeometry() const override;
+	virtual ECollisionEnabled::Type GetCollisionEnabled() const override;
+	virtual ECollisionResponse GetCollisionResponseToChannel(ECollisionChannel Channel) const override;
+	virtual ECollisionChannel GetCollisionObjectType() const override;
+	virtual const FCollisionResponseContainer & GetCollisionResponseToChannels() const override;
+	virtual FVector GetComponentVelocity() const override;
 	//End USceneComponent Interface
 
 	/**
diff --git a/Engine/Source/Runtime/Engine/Classes/Components/ReflectionCaptureComponent.h b/Engine/Source/Runtime/Engine/Classes/Components/ReflectionCaptureComponent.h
--- a/Engine/Source/Runtime/Engine/Classes/Components/ReflectionCaptureComponent.h
+++ b/Engine/Source/Runtime/Engine/Classes/Components/ReflectionCaptureComponent.h
@@ -110,24 +110,24 @@ class UReflectionCaptureComponent : public USceneComponent
 	}
 
 	// Begin UActorComponent Interface
-	virtual void CreateRenderState_Concurrent() OVERRIDE;
-	virtual void DestroyRenderState_Concurrent() OVERRIDE;
-	virtual void SendRenderTransform_Concurrent() OVERRIDE;
+	virtual void CreateRenderState_Concurrent() override;
+	virtual void DestroyRenderState_Concurrent() override;
+	virtual void SendRenderTransform_Concurrent() override;
 	// End UActorComponent Interface
 
 	// Begin UObject Interface
-	virtual void PostInitProperties() OVERRIDE;	
-	virtual void Serialize(FArchive& Ar) OVERRIDE;
-	virtual void PostLoad() OVERRIDE;
-	virtual void PreSave() OVERRIDE;
-	virtual void PostDuplicate(bool bDuplicateForPIE) OVERRIDE;
+	virtual void PostInitProperties() override;	
+	virtual void Serialize(FArchive& Ar) override;
+	virtual void PostLoad() override;
+	virtual void PreSave() override;
+	virtual void PostDuplicate(bool bDuplicateForPIE) override;
 #if WITH_EDITOR
-	virtual void PostEditChangeProperty(FPropertyChangedEvent& PropertyChangedEvent) OVERRIDE;
+	virtual void PostEditChangeProperty(FPropertyChangedEvent& PropertyChangedEvent) override;
 	virtual void PostEditImport();
 #endif // WITH_EDITOR
-	virtual void BeginDestroy() OVERRIDE;
-	virtual bool IsReadyForFinishDestroy() OVERRIDE;
-	virtual void FinishDestroy() OVERRIDE;
+	virtual void BeginDestroy() override;
+	virtual bool IsReadyForFinishDestroy() override;
+	virtual void FinishDestroy() override;
 	// End UObject Interface
 
 private:
diff --git a/Engine/Source/Runtime/Engine/Classes/Components/SceneCaptureComponent2D.h b/Engine/Source/Runtime/Engine/Classes/Components/SceneCaptureComponent2D.h
--- a/Engine/Source/Runtime/Engine/Classes/Components/SceneCaptureComponent2D.h
+++ b/Engine/Source/Runtime/Engine/Classes/Components/SceneCaptureComponent2D.h
@@ -44,18 +44,18 @@ private:
 
 public:
 	// Begin UActorComponent Interface
-	virtual void SendRenderTransform_Concurrent() OVERRIDE;
+	virtual void SendRenderTransform_Concurrent() override;
 	virtual bool RequiresGameThreadEndOfFrameUpdates() const
 	{
 		// this method could probably be removed allowing them to run on any thread, but it isn't worth the trouble
 		return true;
 	}
-	virtual void TickComponent(float DeltaTime, enum ELevelTick TickType, FActorComponentTickFunction *ThisTickFunction) OVERRIDE;
+	virtual void TickComponent(float DeltaTime, enum ELevelTick TickType, FActorComponentTickFunction *ThisTickFunction) override;
 	// End UActorComponent Interface
 
 	// Begin UObject Interface
 #if WITH_EDITOR
-	virtual void PostEditChangeProperty(FPropertyChangedEvent& PropertyChangedEvent) OVERRIDE;
+	virtual void PostEditChangeProperty(FPropertyChangedEvent& PropertyChangedEvent) override;
 #endif // WITH_EDITOR
 	// End UObject Interface
 
diff --git a/Engine/Source/Runtime/Engine/Classes/Components/SceneCaptureComponentCube.h b/Engine/Source/Runtime/Engine/Classes/Components/SceneCaptureComponentCube.h
--- a/Engine/Source/Runtime/Engine/Classes/Components/SceneCaptureComponentCube.h
+++ b/Engine/Source/Runtime/Engine/Classes/Components/SceneCaptureComponentCube.h
@@ -17,13 +17,13 @@ class USceneCaptureComponentCube : public USceneCaptureComponent
 
 public:
 	// Begin UActorComponent Interface
-	virtual void SendRenderTransform_Concurrent() OVERRIDE;
-	virtual void TickComponent(float DeltaTime, enum ELevelTick TickType, FActorComponentTickFunction* ThisTickFunction) OVERRIDE;
+	virtual void SendRenderTransform_Concurrent() override;
+	virtual void TickComponent(float DeltaTime, enum ELevelTick TickType, FActorComponentTickFunction* ThisTickFunction) override;
 	// End UActorComponent Interface
 
 	// Begin UObject Interface
 #if WITH_EDITOR
-	virtual void PostEditChangeProperty(FPropertyChangedEvent& PropertyChangedEvent) OVERRIDE;
+	virtual void PostEditChangeProperty(FPropertyChangedEvent& PropertyChangedEvent) override;
 #endif // WITH_EDITOR
 	// End UObject Interface
 
diff --git a/Engine/Source/Runtime/Engine/Classes/Components/SceneComponent.h b/Engine/Source/Runtime/Engine/Classes/Components/SceneComponent.h
--- a/Engine/Source/Runtime/Engine/Classes/Components/SceneComponent.h
+++ b/Engine/Source/Runtime/Engine/Classes/Components/SceneComponent.h
@@ -187,8 +187,8 @@ private:
 	UFUNCTION()
 	void OnRep_Visibility(bool OldValue);
 
-	virtual void PreNetReceive() OVERRIDE;
-	virtual void PostNetReceive() OVERRIDE;
+	virtual void PreNetReceive() override;
+	virtual void PostNetReceive() override;
 
 public:
 
@@ -441,15 +441,15 @@ public:
 	FPhysicsVolumeChanged PhysicsVolumeChangedDelegate;
 
 	// Begin ActorComponent interface
-	virtual void UpdateComponentToWorld(bool bSkipPhysicsMove = false) OVERRIDE FINAL;
-	virtual void DestroyComponent() OVERRIDE;
-	virtual void ApplyWorldOffset(const FVector& InOffset, bool bWorldShift) OVERRIDE;
+	virtual void UpdateComponentToWorld(bool bSkipPhysicsMove = false) override final;
+	virtual void DestroyComponent() override;
+	virtual void ApplyWorldOffset(const FVector& InOffset, bool bWorldShift) override;
 	// End ActorComponent interface
 
 	// Begin UObject Interface
-	virtual void Serialize(FArchive& Ar) OVERRIDE;
-	virtual void PostInterpChange(UProperty* PropertyThatChanged) OVERRIDE;
-	virtual void BeginDestroy() OVERRIDE;
+	virtual void Serialize(FArchive& Ar) override;
+	virtual void PostInterpChange(UProperty* PropertyThatChanged) override;
+	virtual void BeginDestroy() override;
 	// End UObject Interface
 
 protected:
@@ -626,7 +626,7 @@ public:
 #if WITH_EDITOR
 	/** Called when this component is moved in the editor */
 	virtual void PostEditComponentMove(bool bFinished) {}
-	virtual bool CanEditChange( const UProperty* Property ) const OVERRIDE;
+	virtual bool CanEditChange( const UProperty* Property ) const override;
 
 	virtual const int32 GetNumUncachedStaticLightingInteractions() const;
 #endif // WITH_EDITOR
diff --git a/Engine/Source/Runtime/Engine/Classes/Components/ShapeComponent.h b/Engine/Source/Runtime/Engine/Classes/Components/ShapeComponent.h
--- a/Engine/Source/Runtime/Engine/Classes/Components/ShapeComponent.h
+++ b/Engine/Source/Runtime/Engine/Classes/Components/ShapeComponent.h
@@ -29,14 +29,14 @@ class ENGINE_API UShapeComponent : public UPrimitiveComponent
 	uint32 bShouldCollideWhenPlacing:1;
 
 	// Begin UPrimitiveComponent interface.
-	virtual FPrimitiveSceneProxy* CreateSceneProxy() OVERRIDE;
-	virtual void GetUsedMaterials( TArray<UMaterialInterface*>& OutMaterials ) const OVERRIDE; 
-	virtual class UBodySetup* GetBodySetup() OVERRIDE;
+	virtual FPrimitiveSceneProxy* CreateSceneProxy() override;
+	virtual void GetUsedMaterials( TArray<UMaterialInterface*>& OutMaterials ) const override; 
+	virtual class UBodySetup* GetBodySetup() override;
 	// End UPrimitiveComponent interface.
 
 	// Begin USceneComponent interface
-	virtual FBoxSphereBounds CalcBounds(const FTransform & LocalToWorld) const OVERRIDE;
-	virtual bool ShouldCollideWhenPlacing() const OVERRIDE
+	virtual FBoxSphereBounds CalcBounds(const FTransform & LocalToWorld) const override;
+	virtual bool ShouldCollideWhenPlacing() const override
 	{
 		return bShouldCollideWhenPlacing || IsCollisionEnabled();
 	}
@@ -44,7 +44,7 @@ class ENGINE_API UShapeComponent : public UPrimitiveComponent
 
 	// Begin UObject interface.
 #if WITH_EDITOR
-	virtual void PostEditChangeProperty(FPropertyChangedEvent& PropertyChangedEvent) OVERRIDE;
+	virtual void PostEditChangeProperty(FPropertyChangedEvent& PropertyChangedEvent) override;
 #endif // WITH_EDITOR
 	// End UObject interface.
 
