UE4 Refactoring. Changed OVERRIDE and FINAL macros to keywords override and final.

diff --git a/Engine/Source/Runtime/Engine/Classes/Components/SkeletalMeshComponent.h b/Engine/Source/Runtime/Engine/Classes/Components/SkeletalMeshComponent.h
--- a/Engine/Source/Runtime/Engine/Classes/Components/SkeletalMeshComponent.h
+++ b/Engine/Source/Runtime/Engine/Classes/Components/SkeletalMeshComponent.h
@@ -200,7 +200,7 @@ struct FSkeletalMeshComponentPreClothTickFunction : public FTickFunction
 	* @param CurrentThread - thread we are executing on, useful to pass along as new tasks are created
 	* @param MyCompletionGraphEvent - completion event for this task. Useful for holding the completetion of this task until certain child tasks are complete.
 	**/
-	virtual void ExecuteTick(float DeltaTime, enum ELevelTick TickType, ENamedThreads::Type CurrentThread, const FGraphEventRef& MyCompletionGraphEvent) OVERRIDE;
+	virtual void ExecuteTick(float DeltaTime, enum ELevelTick TickType, ENamedThreads::Type CurrentThread, const FGraphEventRef& MyCompletionGraphEvent) override;
 	/** Abstract function to describe this tick. Used to print messages about illegal cycles in the dependency graph **/
 	virtual FString DiagnosticMessage();
 };
@@ -399,9 +399,9 @@ public:
 	void CreateBodySetup();
 
 	// Begin Interface_CollisionDataProvider Interface
-	virtual bool GetPhysicsTriMeshData(struct FTriMeshCollisionData* CollisionData, bool InUseAllTriData) OVERRIDE;
-	virtual bool ContainsPhysicsTriMeshData(bool InUseAllTriData) const OVERRIDE;
-	virtual bool WantsNegXTriMesh() OVERRIDE
+	virtual bool GetPhysicsTriMeshData(struct FTriMeshCollisionData* CollisionData, bool InUseAllTriData) override;
+	virtual bool ContainsPhysicsTriMeshData(bool InUseAllTriData) const override;
+	virtual bool WantsNegXTriMesh() override
 	{
 		return true;
 	}
@@ -692,7 +692,7 @@ public:
 
 public:
 	// Begin UObject interface.
-	virtual void Serialize(FArchive& Ar) OVERRIDE;
+	virtual void Serialize(FArchive& Ar) override;
 #if WITH_EDITOR
 	DECLARE_MULTICAST_DELEGATE(FOnSkeletalMeshPropertyChangedMulticaster)
 	FOnSkeletalMeshPropertyChangedMulticaster OnSkeletalMeshPropertyChanged;
@@ -702,34 +702,34 @@ public:
 	void RegisterOnSkeletalMeshPropertyChanged(const FOnSkeletalMeshPropertyChanged& Delegate);
 	void UnregisterOnSkeletalMeshPropertyChanged(const FOnSkeletalMeshPropertyChanged& Delegate);
 
-	virtual void PostEditChangeProperty(FPropertyChangedEvent& PropertyChangedEvent) OVERRIDE;
+	virtual void PostEditChangeProperty(FPropertyChangedEvent& PropertyChangedEvent) override;
 
 	/** Validates the animation asset or blueprint, making sure it is compatible with the current skeleton */
 	void ValidateAnimation();
 
-	virtual void LoadedFromAnotherClass(const FName& OldClassName) OVERRIDE;
-	virtual void UpdateCollisionProfile() OVERRIDE;
+	virtual void LoadedFromAnotherClass(const FName& OldClassName) override;
+	virtual void UpdateCollisionProfile() override;
 #endif // WITH_EDITOR
-	virtual SIZE_T GetResourceSize(EResourceSizeMode::Type Mode) OVERRIDE;
+	virtual SIZE_T GetResourceSize(EResourceSizeMode::Type Mode) override;
 	// End UObject interface.
 
 	// Begin UActorComponent interface.
-	virtual void OnRegister() OVERRIDE;
-	virtual void OnUnregister() OVERRIDE;
-	virtual void CreateRenderState_Concurrent() OVERRIDE;
-	virtual void CreatePhysicsState() OVERRIDE;
-	virtual void DestroyPhysicsState() OVERRIDE;
-	virtual void InitializeComponent() OVERRIDE;
-	virtual void TickComponent(float DeltaTime, enum ELevelTick TickType, FActorComponentTickFunction *ThisTickFunction) OVERRIDE;
-	virtual void RegisterComponentTickFunctions(bool bRegister) OVERRIDE;
+	virtual void OnRegister() override;
+	virtual void OnUnregister() override;
+	virtual void CreateRenderState_Concurrent() override;
+	virtual void CreatePhysicsState() override;
+	virtual void DestroyPhysicsState() override;
+	virtual void InitializeComponent() override;
+	virtual void TickComponent(float DeltaTime, enum ELevelTick TickType, FActorComponentTickFunction *ThisTickFunction) override;
+	virtual void RegisterComponentTickFunctions(bool bRegister) override;
 	// End UActorComponent interface.
 
 	// Begin USceneComponent interface.
 	virtual void UpdateBounds();
-	virtual FBoxSphereBounds CalcBounds(const FTransform & LocalToWorld) const OVERRIDE;
-	virtual bool IsAnySimulatingPhysics() const OVERRIDE;
-	virtual void OnUpdateTransform(bool bSkipPhysicsMove) OVERRIDE;
-	virtual void UpdateOverlaps(TArray<FOverlapInfo> const* PendingOverlaps=NULL, bool bDoNotifies=true, const TArray<FOverlapInfo>* OverlapsAtEndLocation=NULL) OVERRIDE;
+	virtual FBoxSphereBounds CalcBounds(const FTransform & LocalToWorld) const override;
+	virtual bool IsAnySimulatingPhysics() const override;
+	virtual void OnUpdateTransform(bool bSkipPhysicsMove) override;
+	virtual void UpdateOverlaps(TArray<FOverlapInfo> const* PendingOverlaps=NULL, bool bDoNotifies=true, const TArray<FOverlapInfo>* OverlapsAtEndLocation=NULL) override;
 	
 	/**
 	 *  Test the collision of the supplied component at the supplied location/rotation, and determine the set of components that it overlaps
@@ -741,48 +741,48 @@ public:
 	 *	@param	ObjectQueryParams	List of object types it's looking for. When this enters, we do object query with component shape
 	 *  @return TRUE if OutOverlaps contains any blocking results
 	 */
-	virtual bool ComponentOverlapMulti(TArray<struct FOverlapResult>& OutOverlaps, const class UWorld* World, const FVector& Pos, const FRotator& Rot, ECollisionChannel TestChannel, const struct FComponentQueryParams& Params, const struct FCollisionObjectQueryParams& ObjectQueryParams = FCollisionObjectQueryParams::DefaultObjectQueryParam) const OVERRIDE;
+	virtual bool ComponentOverlapMulti(TArray<struct FOverlapResult>& OutOverlaps, const class UWorld* World, const FVector& Pos, const FRotator& Rot, ECollisionChannel TestChannel, const struct FComponentQueryParams& Params, const struct FCollisionObjectQueryParams& ObjectQueryParams = FCollisionObjectQueryParams::DefaultObjectQueryParam) const override;
 	// End USceneComponent interface.
 
 	// Begin UPrimitiveComponent interface.
-	virtual void PostPhysicsTick(FPrimitiveComponentPostPhysicsTickFunction &ThisTickFunction) OVERRIDE;
-	virtual class UBodySetup* GetBodySetup() OVERRIDE;
-	virtual FBodyInstance* GetBodyInstance(FName BoneName = NAME_None) const OVERRIDE;
-	virtual void UpdatePhysicsToRBChannels() OVERRIDE;
-	virtual void SetAllPhysicsAngularVelocity(FVector const& NewVel, bool bAddToCurrent = false) OVERRIDE;
-	virtual void SetAllPhysicsPosition(FVector NewPos) OVERRIDE;
-	virtual void SetAllPhysicsRotation(FRotator NewRot) OVERRIDE;
-	virtual void WakeAllRigidBodies() OVERRIDE;
-	virtual void PutAllRigidBodiesToSleep() OVERRIDE;
-	virtual bool IsAnyRigidBodyAwake() OVERRIDE;
-	virtual void OnComponentCollisionSettingsChanged() OVERRIDE;
-	virtual void SetPhysMaterialOverride(UPhysicalMaterial* NewPhysMaterial) OVERRIDE;
-	virtual bool LineTraceComponent( FHitResult& OutHit, const FVector Start, const FVector End, const FCollisionQueryParams& Params ) OVERRIDE;
-	virtual bool SweepComponent( FHitResult& OutHit, const FVector Start, const FVector End, const FCollisionShape& CollisionShape, bool bTraceComplex=false) OVERRIDE;
-	virtual bool ComponentOverlapComponent(class UPrimitiveComponent* PrimComp, const FVector Pos, const FRotator FRotator, const FCollisionQueryParams& Params) OVERRIDE;
-	virtual bool OverlapComponent(const FVector& Pos, const FQuat& Rot, const FCollisionShape& CollisionShape) OVERRIDE;
-	virtual void SetSimulatePhysics(bool bEnabled) OVERRIDE;
-	virtual void AddRadialImpulse(FVector Origin, float Radius, float Strength, ERadialImpulseFalloff Falloff, bool bVelChange=false) OVERRIDE;
-	virtual void AddRadialForce(FVector Origin, float Radius, float Strength, ERadialImpulseFalloff Falloff) OVERRIDE;
-	virtual void SetAllPhysicsLinearVelocity(FVector NewVel,bool bAddToCurrent = false) OVERRIDE;
-	virtual float GetMass() const OVERRIDE;
-	virtual float CalculateMass(FName BoneName = NAME_None) OVERRIDE;
+	virtual void PostPhysicsTick(FPrimitiveComponentPostPhysicsTickFunction &ThisTickFunction) override;
+	virtual class UBodySetup* GetBodySetup() override;
+	virtual FBodyInstance* GetBodyInstance(FName BoneName = NAME_None) const override;
+	virtual void UpdatePhysicsToRBChannels() override;
+	virtual void SetAllPhysicsAngularVelocity(FVector const& NewVel, bool bAddToCurrent = false) override;
+	virtual void SetAllPhysicsPosition(FVector NewPos) override;
+	virtual void SetAllPhysicsRotation(FRotator NewRot) override;
+	virtual void WakeAllRigidBodies() override;
+	virtual void PutAllRigidBodiesToSleep() override;
+	virtual bool IsAnyRigidBodyAwake() override;
+	virtual void OnComponentCollisionSettingsChanged() override;
+	virtual void SetPhysMaterialOverride(UPhysicalMaterial* NewPhysMaterial) override;
+	virtual bool LineTraceComponent( FHitResult& OutHit, const FVector Start, const FVector End, const FCollisionQueryParams& Params ) override;
+	virtual bool SweepComponent( FHitResult& OutHit, const FVector Start, const FVector End, const FCollisionShape& CollisionShape, bool bTraceComplex=false) override;
+	virtual bool ComponentOverlapComponent(class UPrimitiveComponent* PrimComp, const FVector Pos, const FRotator FRotator, const FCollisionQueryParams& Params) override;
+	virtual bool OverlapComponent(const FVector& Pos, const FQuat& Rot, const FCollisionShape& CollisionShape) override;
+	virtual void SetSimulatePhysics(bool bEnabled) override;
+	virtual void AddRadialImpulse(FVector Origin, float Radius, float Strength, ERadialImpulseFalloff Falloff, bool bVelChange=false) override;
+	virtual void AddRadialForce(FVector Origin, float Radius, float Strength, ERadialImpulseFalloff Falloff) override;
+	virtual void SetAllPhysicsLinearVelocity(FVector NewVel,bool bAddToCurrent = false) override;
+	virtual float GetMass() const override;
+	virtual float CalculateMass(FName BoneName = NAME_None) override;
 	// End UPrimitiveComponent interface.
 
 	// Begin USkinnedMeshComponent interface
-	virtual bool UpdateLODStatus() OVERRIDE;
-	virtual void RefreshBoneTransforms( FActorComponentTickFunction* TickFunction = NULL ) OVERRIDE;
-	virtual void TickPose( float DeltaTime ) OVERRIDE;
-	virtual void UpdateSlaveComponent() OVERRIDE;
-	virtual bool ShouldUpdateTransform(bool bLODHasChanged) const OVERRIDE;
-	virtual bool ShouldTickPose() const OVERRIDE;
-	virtual bool AllocateTransformData() OVERRIDE;
-	virtual void DeallocateTransformData() OVERRIDE;
-	virtual void HideBone( int32 BoneIndex, EPhysBodyOp PhysBodyOption ) OVERRIDE;
-	virtual void UnHideBone( int32 BoneIndex ) OVERRIDE;
-	virtual void SetPhysicsAsset(class UPhysicsAsset* NewPhysicsAsset,bool bForceReInit = false) OVERRIDE;
-	virtual void SetSkeletalMesh(class USkeletalMesh* NewMesh) OVERRIDE;
-	virtual FVector GetSkinnedVertexPosition(int32 VertexIndex) const OVERRIDE;
+	virtual bool UpdateLODStatus() override;
+	virtual void RefreshBoneTransforms( FActorComponentTickFunction* TickFunction = NULL ) override;
+	virtual void TickPose( float DeltaTime ) override;
+	virtual void UpdateSlaveComponent() override;
+	virtual bool ShouldUpdateTransform(bool bLODHasChanged) const override;
+	virtual bool ShouldTickPose() const override;
+	virtual bool AllocateTransformData() override;
+	virtual void DeallocateTransformData() override;
+	virtual void HideBone( int32 BoneIndex, EPhysBodyOp PhysBodyOption ) override;
+	virtual void UnHideBone( int32 BoneIndex ) override;
+	virtual void SetPhysicsAsset(class UPhysicsAsset* NewPhysicsAsset,bool bForceReInit = false) override;
+	virtual void SetSkeletalMesh(class USkeletalMesh* NewMesh) override;
+	virtual FVector GetSkinnedVertexPosition(int32 VertexIndex) const override;
 	// End USkinnedMeshComponent interface
 	/** 
 	 *	Iterate over each joint in the physics for this mesh, setting its AngularPositionTarget based on the animation information.
diff --git a/Engine/Source/Runtime/Engine/Classes/Components/SkinnedMeshComponent.h b/Engine/Source/Runtime/Engine/Classes/Components/SkinnedMeshComponent.h
--- a/Engine/Source/Runtime/Engine/Classes/Components/SkinnedMeshComponent.h
+++ b/Engine/Source/Runtime/Engine/Classes/Components/SkinnedMeshComponent.h
@@ -393,21 +393,21 @@ public:
 	FSkeletalMeshResource* GetSkeletalMeshResource() const;
 
 	// Begin UObject interface
-	virtual void Serialize(FArchive& Ar) OVERRIDE;
-	virtual SIZE_T GetResourceSize(EResourceSizeMode::Type Mode) OVERRIDE;
-	virtual FString GetDetailedInfoInternal() const OVERRIDE;
+	virtual void Serialize(FArchive& Ar) override;
+	virtual SIZE_T GetResourceSize(EResourceSizeMode::Type Mode) override;
+	virtual FString GetDetailedInfoInternal() const override;
 #if WITH_EDITOR
-	virtual void PostEditChangeProperty(FPropertyChangedEvent& PropertyChangedEvent) OVERRIDE;
+	virtual void PostEditChangeProperty(FPropertyChangedEvent& PropertyChangedEvent) override;
 #endif // WITH_EDITOR
 	// End UObject interface
 
 protected:
 	// Begin UActorComponent interface
-	virtual void OnRegister() OVERRIDE;
-	virtual void OnUnregister() OVERRIDE;
-	virtual void CreateRenderState_Concurrent() OVERRIDE;
-	virtual void SendRenderDynamicData_Concurrent() OVERRIDE;
-	virtual void DestroyRenderState_Concurrent() OVERRIDE;
+	virtual void OnRegister() override;
+	virtual void OnUnregister() override;
+	virtual void CreateRenderState_Concurrent() override;
+	virtual void SendRenderDynamicData_Concurrent() override;
+	virtual void DestroyRenderState_Concurrent() override;
 #if 0
 	/** return true if this component requires end of frame updates to happen from the game thread. */
 	virtual bool RequiresGameThreadEndOfFrameUpdates()
@@ -417,24 +417,24 @@ protected:
 		return true;
 	}
 #endif
-	virtual void TickComponent(float DeltaTime, enum ELevelTick TickType, FActorComponentTickFunction *ThisTickFunction) OVERRIDE;
-	virtual UObject const* AdditionalStatObject() const OVERRIDE;
+	virtual void TickComponent(float DeltaTime, enum ELevelTick TickType, FActorComponentTickFunction *ThisTickFunction) override;
+	virtual UObject const* AdditionalStatObject() const override;
 	// End UActorComponent interface
 
 public:
 	// Begin USceneComponent interface
-	virtual FBoxSphereBounds CalcBounds(const FTransform & LocalToWorld) const OVERRIDE;
-	virtual bool HasAnySockets() const OVERRIDE;
-	virtual void QuerySupportedSockets(TArray<FComponentSocketDescription>& OutSockets) const OVERRIDE;
-	virtual void UpdateOverlaps(TArray<FOverlapInfo> const* PendingOverlaps=NULL, bool bDoNotifies=true, const TArray<FOverlapInfo>* OverlapsAtEndLocation=NULL) OVERRIDE;
+	virtual FBoxSphereBounds CalcBounds(const FTransform & LocalToWorld) const override;
+	virtual bool HasAnySockets() const override;
+	virtual void QuerySupportedSockets(TArray<FComponentSocketDescription>& OutSockets) const override;
+	virtual void UpdateOverlaps(TArray<FOverlapInfo> const* PendingOverlaps=NULL, bool bDoNotifies=true, const TArray<FOverlapInfo>* OverlapsAtEndLocation=NULL) override;
 	// End USceneComponent interface
 
 	// Begin UPrimitiveComponent interface
-	virtual UMaterialInterface* GetMaterial(int32 MaterialIndex) const OVERRIDE;
-	virtual FPrimitiveSceneProxy* CreateSceneProxy() OVERRIDE;
-	virtual void GetUsedMaterials(TArray<UMaterialInterface*>& OutMaterials) const OVERRIDE;
-	virtual void GetStreamingTextureInfo(TArray<FStreamingTexturePrimitiveInfo>& OutStreamingTextures) const OVERRIDE;
-	virtual int32 GetNumMaterials() const OVERRIDE;
+	virtual UMaterialInterface* GetMaterial(int32 MaterialIndex) const override;
+	virtual FPrimitiveSceneProxy* CreateSceneProxy() override;
+	virtual void GetUsedMaterials(TArray<UMaterialInterface*>& OutMaterials) const override;
+	virtual void GetStreamingTextureInfo(TArray<FStreamingTexturePrimitiveInfo>& OutStreamingTextures) const override;
+	virtual int32 GetNumMaterials() const override;
 	// End UPrimitiveComponent interface
 
 	/**
@@ -634,20 +634,20 @@ public:
 	// Get all socket names.
 	//
 
-	virtual TArray<FName> GetAllSocketNames() const OVERRIDE;
+	virtual TArray<FName> GetAllSocketNames() const override;
 
 	//
 	// Bone Transform.
 	//
 
-	virtual FTransform GetSocketTransform(FName InSocketName, ERelativeTransformSpace TransformSpace = RTS_World) const OVERRIDE;
+	virtual FTransform GetSocketTransform(FName InSocketName, ERelativeTransformSpace TransformSpace = RTS_World) const override;
 
 	/**
 	 * @return SkeletalMeshSocket of named socket on the skeletal mesh component, or NULL if not found.
 	 */
 	class USkeletalMeshSocket const* GetSocketByName( FName InSocketName ) const;
 
-	virtual bool DoesSocketExist(FName InSocketName) const OVERRIDE;
+	virtual bool DoesSocketExist(FName InSocketName) const override;
 
 	/** 
 	 * Get Bone Matrix from index
diff --git a/Engine/Source/Runtime/Engine/Classes/Components/SkyLightComponent.h b/Engine/Source/Runtime/Engine/Classes/Components/SkyLightComponent.h
--- a/Engine/Source/Runtime/Engine/Classes/Components/SkyLightComponent.h
+++ b/Engine/Source/Runtime/Engine/Classes/Components/SkyLightComponent.h
@@ -112,19 +112,19 @@ class ENGINE_API USkyLightComponent : public ULightComponentBase
 
 	// Begin UObject Interface
 	virtual void PostInitProperties();	
-	virtual void PostLoad() OVERRIDE;
+	virtual void PostLoad() override;
 #if WITH_EDITOR
-	virtual void PostEditChangeProperty(FPropertyChangedEvent& PropertyChangedEvent) OVERRIDE;
-	virtual bool CanEditChange(const UProperty* InProperty) const OVERRIDE;
-	virtual void CheckForErrors() OVERRIDE;
+	virtual void PostEditChangeProperty(FPropertyChangedEvent& PropertyChangedEvent) override;
+	virtual bool CanEditChange(const UProperty* InProperty) const override;
+	virtual void CheckForErrors() override;
 #endif // WITH_EDITOR
-	virtual void BeginDestroy() OVERRIDE;
-	virtual bool IsReadyForFinishDestroy() OVERRIDE;
+	virtual void BeginDestroy() override;
+	virtual bool IsReadyForFinishDestroy() override;
 	// End UObject Interface
 
-	virtual TSharedPtr<FComponentInstanceDataBase> GetComponentInstanceData() const OVERRIDE;
-	virtual FName GetComponentInstanceDataType() const OVERRIDE;
-	virtual void ApplyComponentInstanceData(TSharedPtr<FComponentInstanceDataBase> ComponentInstanceData) OVERRIDE;
+	virtual TSharedPtr<FComponentInstanceDataBase> GetComponentInstanceData() const override;
+	virtual FName GetComponentInstanceDataType() const override;
+	virtual void ApplyComponentInstanceData(TSharedPtr<FComponentInstanceDataBase> ComponentInstanceData) override;
 
 	/** Called each tick to recapture and queued sky captures. */
 	static void UpdateSkyCaptureContents(UWorld* WorldToUpdate);
@@ -180,8 +180,8 @@ protected:
 	static TArray<USkyLightComponent*> SkyCapturesToUpdate;
 
 	// Begin UActorComponent Interface
-	virtual void CreateRenderState_Concurrent() OVERRIDE;
-	virtual void DestroyRenderState_Concurrent() OVERRIDE;
+	virtual void CreateRenderState_Concurrent() override;
+	virtual void DestroyRenderState_Concurrent() override;
 	// Begin UActorComponent Interface
 
 	friend class FSkyLightSceneProxy;
diff --git a/Engine/Source/Runtime/Engine/Classes/Components/SphereComponent.h b/Engine/Source/Runtime/Engine/Classes/Components/SphereComponent.h
--- a/Engine/Source/Runtime/Engine/Classes/Components/SphereComponent.h
+++ b/Engine/Source/Runtime/Engine/Classes/Components/SphereComponent.h
@@ -35,19 +35,19 @@ public:
 	float GetUnscaledSphereRadius() const;
 
 	// Begin UPrimitiveComponent interface.
-	virtual FPrimitiveSceneProxy* CreateSceneProxy() OVERRIDE;
-	virtual bool IsZeroExtent() const OVERRIDE;
-	virtual struct FCollisionShape GetCollisionShape(float Inflation = 0.0f) const OVERRIDE;
-	virtual bool AreSymmetricRotations(const FQuat& A, const FQuat& B, const FVector& Scale3D) const OVERRIDE;
+	virtual FPrimitiveSceneProxy* CreateSceneProxy() override;
+	virtual bool IsZeroExtent() const override;
+	virtual struct FCollisionShape GetCollisionShape(float Inflation = 0.0f) const override;
+	virtual bool AreSymmetricRotations(const FQuat& A, const FQuat& B, const FVector& Scale3D) const override;
 	// End UPrimitiveComponent interface.
 
 	// Begin USceneComponent interface
-	virtual FBoxSphereBounds CalcBounds(const FTransform & LocalToWorld) const OVERRIDE;
-	virtual void CalcBoundingCylinder(float& CylinderRadius, float& CylinderHalfHeight) const OVERRIDE;
+	virtual FBoxSphereBounds CalcBounds(const FTransform & LocalToWorld) const override;
+	virtual void CalcBoundingCylinder(float& CylinderRadius, float& CylinderHalfHeight) const override;
 	// End USceneComponent interface
 
 	// Begin UShapeComponent interface
-	virtual void UpdateBodySetup() OVERRIDE;
+	virtual void UpdateBodySetup() override;
 	// End UShapeComponent interface
 
 	// Get the scale used by this shape. This is a uniform scale that is the minimum of any non-uniform scaling.
diff --git a/Engine/Source/Runtime/Engine/Classes/Components/SplineComponent.h b/Engine/Source/Runtime/Engine/Classes/Components/SplineComponent.h
--- a/Engine/Source/Runtime/Engine/Classes/Components/SplineComponent.h
+++ b/Engine/Source/Runtime/Engine/Classes/Components/SplineComponent.h
@@ -28,9 +28,9 @@ class ENGINE_API USplineComponent : public USceneComponent
 	float Duration;
 
 	// Begin UActorComponent interface.
-	virtual TSharedPtr<FComponentInstanceDataBase> GetComponentInstanceData() const OVERRIDE;
-	virtual FName GetComponentInstanceDataType() const OVERRIDE;
-	virtual void ApplyComponentInstanceData(TSharedPtr<FComponentInstanceDataBase> ComponentInstanceData) OVERRIDE;
+	virtual TSharedPtr<FComponentInstanceDataBase> GetComponentInstanceData() const override;
+	virtual FName GetComponentInstanceDataType() const override;
+	virtual void ApplyComponentInstanceData(TSharedPtr<FComponentInstanceDataBase> ComponentInstanceData) override;
 	// End UActorComponent interface.
 
 	/** Update the SplineReparamTable */
diff --git a/Engine/Source/Runtime/Engine/Classes/Components/SplineMeshComponent.h b/Engine/Source/Runtime/Engine/Classes/Components/SplineMeshComponent.h
--- a/Engine/Source/Runtime/Engine/Classes/Components/SplineMeshComponent.h
+++ b/Engine/Source/Runtime/Engine/Classes/Components/SplineMeshComponent.h
@@ -116,25 +116,25 @@ class ENGINE_API USplineMeshComponent : public UStaticMeshComponent, public IInt
 #endif
 
 	//Begin UObject Interface
-	virtual void Serialize(FArchive& Ar) OVERRIDE;
-	virtual bool Modify(bool bAlwaysMarkDirty = true) OVERRIDE;
+	virtual void Serialize(FArchive& Ar) override;
+	virtual bool Modify(bool bAlwaysMarkDirty = true) override;
 	//End UObject Interface
 
 	//Begin USceneComponent Interface
-	virtual FPrimitiveSceneProxy* CreateSceneProxy() OVERRIDE;
-	virtual FBoxSphereBounds CalcBounds(const FTransform & LocalToWorld) const OVERRIDE;
+	virtual FPrimitiveSceneProxy* CreateSceneProxy() override;
+	virtual FBoxSphereBounds CalcBounds(const FTransform & LocalToWorld) const override;
 	//End USceneComponent Interface
 
 	//Begin UPrimitiveComponent Interface
-	virtual void CreatePhysicsState() OVERRIDE;
-	virtual class UBodySetup* GetBodySetup() OVERRIDE;
+	virtual void CreatePhysicsState() override;
+	virtual class UBodySetup* GetBodySetup() override;
 #if WITH_EDITOR
-	virtual bool ShouldRenderSelected() const OVERRIDE
+	virtual bool ShouldRenderSelected() const override
 	{
 		return Super::ShouldRenderSelected() || bSelected;
 	}
 #endif
-	virtual bool DoCustomNavigableGeometryExport(struct FNavigableGeometryExport* GeomExport) const OVERRIDE;
+	virtual bool DoCustomNavigableGeometryExport(struct FNavigableGeometryExport* GeomExport) const override;
 	//End UPrimitiveComponent Interface
 
 	//Begin UStaticMeshComponent Interface
@@ -142,9 +142,9 @@ class ENGINE_API USplineMeshComponent : public UStaticMeshComponent, public IInt
 	//End UStaticMeshComponent Interface
 
 	// Begin Interface_CollisionDataProvider Interface
-	virtual bool GetPhysicsTriMeshData(struct FTriMeshCollisionData* CollisionData, bool InUseAllTriData) OVERRIDE;
-	virtual bool ContainsPhysicsTriMeshData(bool InUseAllTriData) const OVERRIDE;
-	virtual bool WantsNegXTriMesh() OVERRIDE { return false; }
+	virtual bool GetPhysicsTriMeshData(struct FTriMeshCollisionData* CollisionData, bool InUseAllTriData) override;
+	virtual bool ContainsPhysicsTriMeshData(bool InUseAllTriData) const override;
+	virtual bool WantsNegXTriMesh() override { return false; }
 	// End Interface_CollisionDataProvider Interface
 
 	/** Called when spline params are changed, to notify render thread and possibly collision */
diff --git a/Engine/Source/Runtime/Engine/Classes/Components/SpotLightComponent.h b/Engine/Source/Runtime/Engine/Classes/Components/SpotLightComponent.h
--- a/Engine/Source/Runtime/Engine/Classes/Components/SpotLightComponent.h
+++ b/Engine/Source/Runtime/Engine/Classes/Components/SpotLightComponent.h
@@ -38,7 +38,7 @@ class ENGINE_API USpotLightComponent : public UPointLightComponent
 	virtual FSphere GetBoundingSphere() const;
 	virtual bool AffectsBounds(const FBoxSphereBounds& Bounds) const;
 	virtual ELightComponentType GetLightType() const;
-	virtual FLightSceneProxy* CreateSceneProxy() const OVERRIDE;
+	virtual FLightSceneProxy* CreateSceneProxy() const override;
 };
 
 
diff --git a/Engine/Source/Runtime/Engine/Classes/Components/StaticMeshComponent.h b/Engine/Source/Runtime/Engine/Classes/Components/StaticMeshComponent.h
--- a/Engine/Source/Runtime/Engine/Classes/Components/StaticMeshComponent.h
+++ b/Engine/Source/Runtime/Engine/Classes/Components/StaticMeshComponent.h
@@ -192,79 +192,79 @@ class ENGINE_API UStaticMeshComponent : public UMeshComponent
 public:
 
 	// Begin UObject interface.
-	virtual void BeginDestroy() OVERRIDE;
-	virtual void ExportCustomProperties(FOutputDevice& Out, uint32 Indent) OVERRIDE;
-	virtual void ImportCustomProperties(const TCHAR* SourceText, FFeedbackContext* Warn) OVERRIDE;	
-	virtual void Serialize(FArchive& Ar) OVERRIDE;
+	virtual void BeginDestroy() override;
+	virtual void ExportCustomProperties(FOutputDevice& Out, uint32 Indent) override;
+	virtual void ImportCustomProperties(const TCHAR* SourceText, FFeedbackContext* Warn) override;	
+	virtual void Serialize(FArchive& Ar) override;
 #if WITH_EDITOR
-	virtual void PostEditUndo() OVERRIDE;
-	virtual void PreEditUndo() OVERRIDE;
-	virtual void PostEditChangeProperty(FPropertyChangedEvent& PropertyChangedEvent) OVERRIDE;
+	virtual void PostEditUndo() override;
+	virtual void PreEditUndo() override;
+	virtual void PostEditChangeProperty(FPropertyChangedEvent& PropertyChangedEvent) override;
 #endif // WITH_EDITOR
-	virtual void PreSave() OVERRIDE;
-	virtual void PostLoad() OVERRIDE;
-	virtual bool AreNativePropertiesIdenticalTo( UObject* Other ) const OVERRIDE;
-	virtual FString GetDetailedInfoInternal() const OVERRIDE;
+	virtual void PreSave() override;
+	virtual void PostLoad() override;
+	virtual bool AreNativePropertiesIdenticalTo( UObject* Other ) const override;
+	virtual FString GetDetailedInfoInternal() const override;
 	static void AddReferencedObjects(UObject* InThis, FReferenceCollector& Collector);
 	// End UObject interface.
 
 	// Begin USceneComponent Interface
-	virtual FBoxSphereBounds CalcBounds(const FTransform & LocalToWorld) const OVERRIDE;
-	virtual bool HasAnySockets() const OVERRIDE;
-	virtual void QuerySupportedSockets(TArray<FComponentSocketDescription>& OutSockets) const OVERRIDE;
-	virtual FTransform GetSocketTransform(FName InSocketName, ERelativeTransformSpace TransformSpace = RTS_World) const OVERRIDE;
-	virtual bool DoesSocketExist(FName InSocketName) const OVERRIDE;
-	virtual bool ShouldCollideWhenPlacing() const OVERRIDE
+	virtual FBoxSphereBounds CalcBounds(const FTransform & LocalToWorld) const override;
+	virtual bool HasAnySockets() const override;
+	virtual void QuerySupportedSockets(TArray<FComponentSocketDescription>& OutSockets) const override;
+	virtual FTransform GetSocketTransform(FName InSocketName, ERelativeTransformSpace TransformSpace = RTS_World) const override;
+	virtual bool DoesSocketExist(FName InSocketName) const override;
+	virtual bool ShouldCollideWhenPlacing() const override
 	{
 		// Current Method of collision does not work with non-capsule shapes, enable when it works with static meshes
 		// return IsCollisionEnabled() && (StaticMesh != NULL);
 		return false;
 	}
-	virtual TArray<FName> GetAllSocketNames() const OVERRIDE;
+	virtual TArray<FName> GetAllSocketNames() const override;
 	// End USceneComponent Interface
 
 	// Begin UActorComponent interface.
 protected: 
-	virtual void OnRegister() OVERRIDE;
-	virtual void OnUnregister() OVERRIDE;
+	virtual void OnRegister() override;
+	virtual void OnUnregister() override;
 public:
-	virtual void InvalidateLightingCacheDetailed(bool bInvalidateBuildEnqueuedLighting, bool bTranslationOnly) OVERRIDE;
-	virtual UObject const* AdditionalStatObject() const OVERRIDE;
+	virtual void InvalidateLightingCacheDetailed(bool bInvalidateBuildEnqueuedLighting, bool bTranslationOnly) override;
+	virtual UObject const* AdditionalStatObject() const override;
 #if WITH_EDITOR
-	virtual void CheckForErrors() OVERRIDE;
+	virtual void CheckForErrors() override;
 #endif
-	virtual TSharedPtr<FComponentInstanceDataBase> GetComponentInstanceData() const OVERRIDE;
-	virtual FName GetComponentInstanceDataType() const OVERRIDE;
-	virtual void ApplyComponentInstanceData(TSharedPtr<FComponentInstanceDataBase> ComponentInstanceData) OVERRIDE;
+	virtual TSharedPtr<FComponentInstanceDataBase> GetComponentInstanceData() const override;
+	virtual FName GetComponentInstanceDataType() const override;
+	virtual void ApplyComponentInstanceData(TSharedPtr<FComponentInstanceDataBase> ComponentInstanceData) override;
 	// End UActorComponent interface.
 
 
 
 	// Begin UPrimitiveComponent interface.
-	virtual int32 GetNumMaterials() const OVERRIDE;
+	virtual int32 GetNumMaterials() const override;
 #if WITH_EDITOR
-	virtual void GetStaticLightingInfo(FStaticLightingPrimitiveInfo& OutPrimitiveInfo,const TArray<ULightComponent*>& InRelevantLights,const FLightingBuildOptions& Options) OVERRIDE;
+	virtual void GetStaticLightingInfo(FStaticLightingPrimitiveInfo& OutPrimitiveInfo,const TArray<ULightComponent*>& InRelevantLights,const FLightingBuildOptions& Options) override;
 #endif
-	virtual float GetEmissiveBoost(int32 ElementIndex) const OVERRIDE;
-	virtual float GetDiffuseBoost(int32 ElementIndex) const OVERRIDE;
-	virtual bool GetShadowIndirectOnly() const OVERRIDE
+	virtual float GetEmissiveBoost(int32 ElementIndex) const override;
+	virtual float GetDiffuseBoost(int32 ElementIndex) const override;
+	virtual bool GetShadowIndirectOnly() const override
 	{
 		return LightmassSettings.bShadowIndirectOnly;
 	}
-	virtual ELightMapInteractionType GetStaticLightingType() const OVERRIDE;
-	virtual void GetStreamingTextureInfo(TArray<FStreamingTexturePrimitiveInfo>& OutStreamingTextures) const OVERRIDE;
-	virtual class UBodySetup* GetBodySetup() OVERRIDE;
-	virtual FPrimitiveSceneProxy* CreateSceneProxy() OVERRIDE;
-	virtual bool ShouldRecreateProxyOnUpdateTransform() const OVERRIDE;
-	virtual bool UsesOnlyUnlitMaterials() const OVERRIDE;
-	virtual bool GetLightMapResolution( int32& Width, int32& Height ) const OVERRIDE;
-	virtual int32 GetStaticLightMapResolution() const OVERRIDE;
+	virtual ELightMapInteractionType GetStaticLightingType() const override;
+	virtual void GetStreamingTextureInfo(TArray<FStreamingTexturePrimitiveInfo>& OutStreamingTextures) const override;
+	virtual class UBodySetup* GetBodySetup() override;
+	virtual FPrimitiveSceneProxy* CreateSceneProxy() override;
+	virtual bool ShouldRecreateProxyOnUpdateTransform() const override;
+	virtual bool UsesOnlyUnlitMaterials() const override;
+	virtual bool GetLightMapResolution( int32& Width, int32& Height ) const override;
+	virtual int32 GetStaticLightMapResolution() const override;
 	/** Returns true if the component is static AND has the right static mesh setup to support lightmaps. */
-	virtual bool HasValidSettingsForStaticLighting() const OVERRIDE;
+	virtual bool HasValidSettingsForStaticLighting() const override;
 
-	virtual void GetLightAndShadowMapMemoryUsage( int32& LightMapMemoryUsage, int32& ShadowMapMemoryUsage ) const OVERRIDE;
-	virtual void GetUsedMaterials(TArray<UMaterialInterface*>& OutMaterials) const OVERRIDE;
-	virtual UMaterialInterface* GetMaterial(int32 MaterialIndex) const OVERRIDE;
+	virtual void GetLightAndShadowMapMemoryUsage( int32& LightMapMemoryUsage, int32& ShadowMapMemoryUsage ) const override;
+	virtual void GetUsedMaterials(TArray<UMaterialInterface*>& OutMaterials) const override;
+	virtual UMaterialInterface* GetMaterial(int32 MaterialIndex) const override;
 
 	virtual bool DoCustomNavigableGeometryExport(struct FNavigableGeometryExport* GeomExport) const;
 	// End UPrimitiveComponent interface.
diff --git a/Engine/Source/Runtime/Engine/Classes/Components/TextRenderComponent.h b/Engine/Source/Runtime/Engine/Classes/Components/TextRenderComponent.h
--- a/Engine/Source/Runtime/Engine/Classes/Components/TextRenderComponent.h
+++ b/Engine/Source/Runtime/Engine/Classes/Components/TextRenderComponent.h
@@ -129,21 +129,21 @@ class ENGINE_API UTextRenderComponent : public UPrimitiveComponent
 	// -----------------------------
 
 	// Begin UPrimitiveComponent interface.
-	virtual FPrimitiveSceneProxy* CreateSceneProxy() OVERRIDE;
+	virtual FPrimitiveSceneProxy* CreateSceneProxy() override;
 	virtual void GetUsedMaterials( TArray<UMaterialInterface*>& OutMaterials ) const;
 	virtual int32 GetNumMaterials() const;
 	virtual UMaterialInterface* GetMaterial(int32 ElementIndex) const;
-	virtual bool ShouldRecreateProxyOnUpdateTransform() const OVERRIDE;
+	virtual bool ShouldRecreateProxyOnUpdateTransform() const override;
 	virtual void SetMaterial(int32 ElementIndex, UMaterialInterface* InMaterial);
 	virtual FMatrix GetRenderMatrix() const;
 	// End UPrimitiveComponent interface.
 
 	// Begin USceneComponent interface.
-	virtual FBoxSphereBounds CalcBounds(const FTransform & LocalToWorld) const OVERRIDE;
+	virtual FBoxSphereBounds CalcBounds(const FTransform & LocalToWorld) const override;
 	// Begin USceneComponent interface.
 
 	// Begin UObject interface.
-	virtual void PostLoad() OVERRIDE;
+	virtual void PostLoad() override;
 	// End UObject interface.
 };
 
diff --git a/Engine/Source/Runtime/Engine/Classes/Components/TimelineComponent.h b/Engine/Source/Runtime/Engine/Classes/Components/TimelineComponent.h
--- a/Engine/Source/Runtime/Engine/Classes/Components/TimelineComponent.h
+++ b/Engine/Source/Runtime/Engine/Classes/Components/TimelineComponent.h
@@ -443,7 +443,7 @@ public:
 	void OnRep_Timeline();
 
 	// Begin ActorComponent interface.
-	virtual void TickComponent(float DeltaTime, enum ELevelTick TickType, FActorComponentTickFunction *ThisTickFunction) OVERRIDE;
+	virtual void TickComponent(float DeltaTime, enum ELevelTick TickType, FActorComponentTickFunction *ThisTickFunction) override;
 	// End ActorComponent interface.
 
 	/** Get the signature function for Timeline event functions */
diff --git a/Engine/Source/Runtime/Engine/Classes/Components/VectorFieldComponent.h b/Engine/Source/Runtime/Engine/Classes/Components/VectorFieldComponent.h
--- a/Engine/Source/Runtime/Engine/Classes/Components/VectorFieldComponent.h
+++ b/Engine/Source/Runtime/Engine/Classes/Components/VectorFieldComponent.h
@@ -46,20 +46,20 @@ public:
 
 
 	// Begin UPrimitiveComponent interface.
-	virtual FPrimitiveSceneProxy* CreateSceneProxy() OVERRIDE;
-	virtual FBoxSphereBounds CalcBounds(const FTransform & LocalToWorld) const OVERRIDE;
+	virtual FPrimitiveSceneProxy* CreateSceneProxy() override;
+	virtual FBoxSphereBounds CalcBounds(const FTransform & LocalToWorld) const override;
 	// End UPrimitiveComponent interface.
 
 	// Begin UActorComponent interface.
-	virtual void OnRegister() OVERRIDE;
-	virtual void OnUnregister() OVERRIDE;
-	virtual void SendRenderTransform_Concurrent() OVERRIDE;
+	virtual void OnRegister() override;
+	virtual void OnUnregister() override;
+	virtual void SendRenderTransform_Concurrent() override;
 	// End UActorComponent interface.
 
 	// Begin UObject interface.
-	virtual void PostInterpChange(UProperty* PropertyThatChanged) OVERRIDE;
+	virtual void PostInterpChange(UProperty* PropertyThatChanged) override;
 #if WITH_EDITOR
-	virtual void PostEditChangeProperty(FPropertyChangedEvent& PropertyChangedEvent) OVERRIDE;
+	virtual void PostEditChangeProperty(FPropertyChangedEvent& PropertyChangedEvent) override;
 #endif // WITH_EDITOR
 	// End UObject interface.
 };
diff --git a/Engine/Source/Runtime/Engine/Classes/Components/WindDirectionalSourceComponent.h b/Engine/Source/Runtime/Engine/Classes/Components/WindDirectionalSourceComponent.h
--- a/Engine/Source/Runtime/Engine/Classes/Components/WindDirectionalSourceComponent.h
+++ b/Engine/Source/Runtime/Engine/Classes/Components/WindDirectionalSourceComponent.h
@@ -21,9 +21,9 @@ public:
 
 protected:
 	// Begin UActorComponent interface.
-	virtual void CreateRenderState_Concurrent() OVERRIDE;
-	virtual void SendRenderTransform_Concurrent() OVERRIDE;
-	virtual void DestroyRenderState_Concurrent() OVERRIDE;
+	virtual void CreateRenderState_Concurrent() override;
+	virtual void SendRenderTransform_Concurrent() override;
+	virtual void DestroyRenderState_Concurrent() override;
 	// End UActorComponent interface.
 
 public:
diff --git a/Engine/Source/Runtime/Engine/Classes/Curves/CurveBase.h b/Engine/Source/Runtime/Engine/Classes/Curves/CurveBase.h
--- a/Engine/Source/Runtime/Engine/Classes/Curves/CurveBase.h
+++ b/Engine/Source/Runtime/Engine/Classes/Curves/CurveBase.h
@@ -290,10 +290,10 @@ public:
 	FRichCurveKey GetLastKey() const;
 
 	/** Get number of key in curve. */
-	virtual int32 GetNumKeys() const OVERRIDE;
+	virtual int32 GetNumKeys() const override;
 	
 	/** Checks to see if the key handle is valid for this curve */
-	virtual bool IsKeyHandleValid(FKeyHandle KeyHandle) const OVERRIDE;
+	virtual bool IsKeyHandleValid(FKeyHandle KeyHandle) const override;
 
 	/**
 	  * Add a new key to the curve with the supplied Time and Value. Returns the handle of the new key.
@@ -446,25 +446,25 @@ class UCurveBase : public UObject, public FCurveOwnerInterface
 
 public:
 	// Begin FCurveOwnerInterface
-	virtual TArray<FRichCurveEditInfoConst> GetCurves() const OVERRIDE
+	virtual TArray<FRichCurveEditInfoConst> GetCurves() const override
 	{
 		TArray<FRichCurveEditInfoConst> Curves;
 		return Curves;
 	}
 
-	virtual TArray<FRichCurveEditInfo> GetCurves() OVERRIDE
+	virtual TArray<FRichCurveEditInfo> GetCurves() override
 	{
 		TArray<FRichCurveEditInfo> Curves;
 		return Curves;
 	}
 
-	virtual UObject* GetOwner() OVERRIDE
+	virtual UObject* GetOwner() override
 	{
 		return this;
 	}
 
-	virtual void ModifyOwner() OVERRIDE;
-	virtual void MakeTransactional() OVERRIDE;
+	virtual void ModifyOwner() override;
+	virtual void MakeTransactional() override;
 	// End FCurveOwnerInterface
 
 	// Begin UCurveBase interface
@@ -510,10 +510,10 @@ struct ENGINE_API FIntegralCurve : public FIndexedCurve
 public:
 	
 	/** Get number of keys in curve. */
-	virtual int32 GetNumKeys() const OVERRIDE;
+	virtual int32 GetNumKeys() const override;
 	
 	/** Checks to see if the key handle is valid for this curve */
-	virtual bool IsKeyHandleValid(FKeyHandle KeyHandle) const OVERRIDE;
+	virtual bool IsKeyHandleValid(FKeyHandle KeyHandle) const override;
 
 	/** Evaluates the value of an array of keys at a time */
 	int32 Evaluate(float Time) const;
diff --git a/Engine/Source/Runtime/Engine/Classes/Curves/CurveFloat.h b/Engine/Source/Runtime/Engine/Classes/Curves/CurveFloat.h
--- a/Engine/Source/Runtime/Engine/Classes/Curves/CurveFloat.h
+++ b/Engine/Source/Runtime/Engine/Classes/Curves/CurveFloat.h
@@ -45,11 +45,11 @@ class UCurveFloat : public UCurveBase
 	ENGINE_API float GetFloatValue(float InTime) const;
 
 	// UObject interface
-	virtual void PostLoad() OVERRIDE;
+	virtual void PostLoad() override;
 
 	// Begin FCurveOwnerInterface
-	virtual TArray<FRichCurveEditInfoConst> GetCurves() const OVERRIDE;
-	virtual TArray<FRichCurveEditInfo> GetCurves() OVERRIDE;
+	virtual TArray<FRichCurveEditInfoConst> GetCurves() const override;
+	virtual TArray<FRichCurveEditInfo> GetCurves() override;
 
 	/** Determine if Curve is the same */
 	ENGINE_API bool operator == (const UCurveFloat& Curve) const;
diff --git a/Engine/Source/Runtime/Engine/Classes/Curves/CurveLinearColor.h b/Engine/Source/Runtime/Engine/Classes/Curves/CurveLinearColor.h
--- a/Engine/Source/Runtime/Engine/Classes/Curves/CurveLinearColor.h
+++ b/Engine/Source/Runtime/Engine/Classes/Curves/CurveLinearColor.h
@@ -18,8 +18,8 @@ class UCurveLinearColor : public UCurveBase
 	ENGINE_API FLinearColor GetLinearColorValue(float InTime) const;
 
 	// Begin FCurveOwnerInterface
-	virtual TArray<FRichCurveEditInfoConst> GetCurves() const OVERRIDE;
-	virtual TArray<FRichCurveEditInfo> GetCurves() OVERRIDE;
+	virtual TArray<FRichCurveEditInfoConst> GetCurves() const override;
+	virtual TArray<FRichCurveEditInfo> GetCurves() override;
 
 	/** @return True if the curve as any alpha keys */
 	bool HasAnyAlphaKeys() const { return FloatCurves[3].GetNumKeys() > 0; }
diff --git a/Engine/Source/Runtime/Engine/Classes/Curves/CurveVector.h b/Engine/Source/Runtime/Engine/Classes/Curves/CurveVector.h
--- a/Engine/Source/Runtime/Engine/Classes/Curves/CurveVector.h
+++ b/Engine/Source/Runtime/Engine/Classes/Curves/CurveVector.h
@@ -22,11 +22,11 @@ class UCurveVector : public UCurveBase
 	ENGINE_API FVector GetVectorValue(float InTime) const;
 
 	// UObject interface
-	virtual void PostLoad() OVERRIDE;
+	virtual void PostLoad() override;
 
 	// Begin FCurveOwnerInterface
-	virtual TArray<FRichCurveEditInfoConst> GetCurves() const OVERRIDE;
-	virtual TArray<FRichCurveEditInfo> GetCurves() OVERRIDE;
+	virtual TArray<FRichCurveEditInfoConst> GetCurves() const override;
+	virtual TArray<FRichCurveEditInfo> GetCurves() override;
 
 	/** Determine if Curve is the same */
 	ENGINE_API bool operator == (const UCurveVector& Curve) const;
diff --git a/Engine/Source/Runtime/Engine/Classes/Debug/LogVisualizerCameraController.h b/Engine/Source/Runtime/Engine/Classes/Debug/LogVisualizerCameraController.h
--- a/Engine/Source/Runtime/Engine/Classes/Debug/LogVisualizerCameraController.h
+++ b/Engine/Source/Runtime/Engine/Classes/Debug/LogVisualizerCameraController.h
@@ -16,13 +16,13 @@ public:
 	UPROPERTY()
 	class AActor* PickedActor;
 
-	virtual void PostInitializeComponents() OVERRIDE;
+	virtual void PostInitializeComponents() override;
 
 	ENGINE_API static ALogVisualizerCameraController* EnableCamera(UWorld* InWorld);
 	ENGINE_API static void DisableCamera(UWorld* InWorld);
 	ENGINE_API static bool IsEnabled(UWorld* InWorld);
 
-	virtual void Select( FHitResult const& Hit ) OVERRIDE;
+	virtual void Select( FHitResult const& Hit ) override;
 
 	void ShowNextEntry();
 	void ShowPrevEntry();
@@ -33,5 +33,5 @@ public:
 	FLogEntryIterationDelegate OnIterateLogEntries;
 
 protected:
-	virtual void SetupInputComponent() OVERRIDE;
+	virtual void SetupInputComponent() override;
 };
diff --git a/Engine/Source/Runtime/Engine/Classes/Debug/LogVisualizerHUD.h b/Engine/Source/Runtime/Engine/Classes/Debug/LogVisualizerHUD.h
--- a/Engine/Source/Runtime/Engine/Classes/Debug/LogVisualizerHUD.h
+++ b/Engine/Source/Runtime/Engine/Classes/Debug/LogVisualizerHUD.h
@@ -11,10 +11,10 @@ class ALogVisualizerHUD : public ADebugCameraHUD
 
 	FFontRenderInfo TextRenderInfo;
 
-	virtual bool DisplayMaterials( float X, float& Y, float DY, UMeshComponent* MeshComp ) OVERRIDE;
+	virtual bool DisplayMaterials( float X, float& Y, float DY, UMeshComponent* MeshComp ) override;
 
 	// Begin AActor Interface
-	virtual void PostRender() OVERRIDE;
+	virtual void PostRender() override;
 	// End AActor Interface
 };
 
diff --git a/Engine/Source/Runtime/Engine/Classes/Debug/ReporterGraph.h b/Engine/Source/Runtime/Engine/Classes/Debug/ReporterGraph.h
--- a/Engine/Source/Runtime/Engine/Classes/Debug/ReporterGraph.h
+++ b/Engine/Source/Runtime/Engine/Classes/Debug/ReporterGraph.h
@@ -122,7 +122,7 @@ public:
 	/** Main draw call for the Graph.
 	 * @param Canvas The canvas to draw this graph to
 	 */
-	virtual void Draw(UCanvas* Canvas) OVERRIDE;
+	virtual void Draw(UCanvas* Canvas) override;
 
 	/** Set the axis min and max data for both axes
 	 * @param MinX The normalized minimum X data
@@ -217,7 +217,7 @@ protected:
 	 * @param Canvas The canvas we wish to use during the conversion
 	 * @return The screen space coordinates based on the size of the canvas
 	 */
-	virtual FVector2D ToScreenSpace(const FVector2D& InVector, UCanvas* Canvas) OVERRIDE;
+	virtual FVector2D ToScreenSpace(const FVector2D& InVector, UCanvas* Canvas) override;
 
 	/** Helper to convert data from raw data into screen space
 	 * @param InVector 2D plot point in data space
diff --git a/Engine/Source/Runtime/Engine/Classes/DeviceProfiles/DeviceProfile.h b/Engine/Source/Runtime/Engine/Classes/DeviceProfiles/DeviceProfile.h
--- a/Engine/Source/Runtime/Engine/Classes/DeviceProfiles/DeviceProfile.h
+++ b/Engine/Source/Runtime/Engine/Classes/DeviceProfiles/DeviceProfile.h
@@ -64,7 +64,7 @@ public:
 
 #if WITH_EDITOR
 	// Begin UObject interface
-	virtual void PostEditChangeProperty( FPropertyChangedEvent& PropertyChangedEvent ) OVERRIDE;
+	virtual void PostEditChangeProperty( FPropertyChangedEvent& PropertyChangedEvent ) override;
 	// End UObject interface
 #endif
 };
diff --git a/Engine/Source/Runtime/Engine/Classes/Distributions/DistributionFloat.h b/Engine/Source/Runtime/Engine/Classes/Distributions/DistributionFloat.h
--- a/Engine/Source/Runtime/Engine/Classes/Distributions/DistributionFloat.h
+++ b/Engine/Source/Runtime/Engine/Classes/Distributions/DistributionFloat.h
@@ -117,8 +117,8 @@ class UDistributionFloat : public UDistribution
 	virtual float GetValue( float F = 0.f, UObject* Data = NULL, class FRandomStream* InRandomStream = NULL ) const;
 
 	// Begin FCurveEdInterface Interface
-	virtual void GetInRange(float& MinIn, float& MaxIn) const OVERRIDE;
-	virtual void GetOutRange(float& MinOut, float& MaxOut) const OVERRIDE;
+	virtual void GetInRange(float& MinIn, float& MaxIn) const override;
+	virtual void GetOutRange(float& MinOut, float& MaxOut) const override;
 	// End FCurveEdInterface Interface
 	
 	/** @return true of this distribution can be baked into a FRawDistribution lookup table, otherwise false */
@@ -137,10 +137,10 @@ class UDistributionFloat : public UDistribution
 
 	/** Begin UObject interface */
 #if	WITH_EDITOR
-	virtual void PostEditChangeProperty(FPropertyChangedEvent& PropertyChangedEvent) OVERRIDE;
+	virtual void PostEditChangeProperty(FPropertyChangedEvent& PropertyChangedEvent) override;
 #endif	// WITH_EDITOR
-	virtual bool NeedsLoadForClient() const OVERRIDE;
-	virtual bool NeedsLoadForServer() const OVERRIDE;
+	virtual bool NeedsLoadForClient() const override;
+	virtual bool NeedsLoadForServer() const override;
 	/** End UObject interface */
 };
 
diff --git a/Engine/Source/Runtime/Engine/Classes/Distributions/DistributionFloatConstant.h b/Engine/Source/Runtime/Engine/Classes/Distributions/DistributionFloatConstant.h
--- a/Engine/Source/Runtime/Engine/Classes/Distributions/DistributionFloatConstant.h
+++ b/Engine/Source/Runtime/Engine/Classes/Distributions/DistributionFloatConstant.h
@@ -14,32 +14,32 @@ class UDistributionFloatConstant : public UDistributionFloat
 	float Constant;
 
 	// Begin UObject Interface
-	virtual void PostInitProperties() OVERRIDE;
-	virtual void PostLoad() OVERRIDE;
+	virtual void PostInitProperties() override;
+	virtual void PostLoad() override;
 	// End UObject Interface
 
 	// Begin UDistributionFloat Interface
-	virtual float GetValue( float F = 0.f, UObject* Data = NULL, class FRandomStream* InRandomStream = NULL ) const OVERRIDE;
+	virtual float GetValue( float F = 0.f, UObject* Data = NULL, class FRandomStream* InRandomStream = NULL ) const override;
 	// End UDistributionFloat Interface
 
 
 	// Begin FCurveEdInterface interface
-	virtual int32		GetNumKeys() const OVERRIDE;
-	virtual int32		GetNumSubCurves() const OVERRIDE;
-	virtual float	GetKeyIn(int32 KeyIndex) OVERRIDE;
-	virtual float	GetKeyOut(int32 SubIndex, int32 KeyIndex) OVERRIDE;
-	virtual FColor	GetKeyColor(int32 SubIndex, int32 KeyIndex, const FColor& CurveColor) OVERRIDE;
-	virtual void	GetInRange(float& MinIn, float& MaxIn) const OVERRIDE;
-	virtual void	GetOutRange(float& MinOut, float& MaxOut) const OVERRIDE;
-	virtual EInterpCurveMode	GetKeyInterpMode(int32 KeyIndex) const OVERRIDE;
-	virtual void	GetTangents(int32 SubIndex, int32 KeyIndex, float& ArriveTangent, float& LeaveTangent) const OVERRIDE;
-	virtual float	EvalSub(int32 SubIndex, float InVal) OVERRIDE;
-	virtual int32		CreateNewKey(float KeyIn) OVERRIDE;
-	virtual void	DeleteKey(int32 KeyIndex) OVERRIDE;
-	virtual int32		SetKeyIn(int32 KeyIndex, float NewInVal) OVERRIDE;
-	virtual void	SetKeyOut(int32 SubIndex, int32 KeyIndex, float NewOutVal) OVERRIDE;
-	virtual void	SetKeyInterpMode(int32 KeyIndex, EInterpCurveMode NewMode) OVERRIDE;
-	virtual void	SetTangents(int32 SubIndex, int32 KeyIndex, float ArriveTangent, float LeaveTangent) OVERRIDE;
+	virtual int32		GetNumKeys() const override;
+	virtual int32		GetNumSubCurves() const override;
+	virtual float	GetKeyIn(int32 KeyIndex) override;
+	virtual float	GetKeyOut(int32 SubIndex, int32 KeyIndex) override;
+	virtual FColor	GetKeyColor(int32 SubIndex, int32 KeyIndex, const FColor& CurveColor) override;
+	virtual void	GetInRange(float& MinIn, float& MaxIn) const override;
+	virtual void	GetOutRange(float& MinOut, float& MaxOut) const override;
+	virtual EInterpCurveMode	GetKeyInterpMode(int32 KeyIndex) const override;
+	virtual void	GetTangents(int32 SubIndex, int32 KeyIndex, float& ArriveTangent, float& LeaveTangent) const override;
+	virtual float	EvalSub(int32 SubIndex, float InVal) override;
+	virtual int32		CreateNewKey(float KeyIn) override;
+	virtual void	DeleteKey(int32 KeyIndex) override;
+	virtual int32		SetKeyIn(int32 KeyIndex, float NewInVal) override;
+	virtual void	SetKeyOut(int32 SubIndex, int32 KeyIndex, float NewOutVal) override;
+	virtual void	SetKeyInterpMode(int32 KeyIndex, EInterpCurveMode NewMode) override;
+	virtual void	SetTangents(int32 SubIndex, int32 KeyIndex, float ArriveTangent, float LeaveTangent) override;
 	// End FCurveEdInterface interface
 };
 
diff --git a/Engine/Source/Runtime/Engine/Classes/Distributions/DistributionFloatConstantCurve.h b/Engine/Source/Runtime/Engine/Classes/Distributions/DistributionFloatConstantCurve.h
--- a/Engine/Source/Runtime/Engine/Classes/Distributions/DistributionFloatConstantCurve.h
+++ b/Engine/Source/Runtime/Engine/Classes/Distributions/DistributionFloatConstantCurve.h
@@ -15,26 +15,26 @@ class UDistributionFloatConstantCurve : public UDistributionFloat
 
 
 	// Begin UDistributionFloat Interface
-	virtual float GetValue( float F = 0.f, UObject* Data = NULL, class FRandomStream* InRandomStream = NULL ) const OVERRIDE;
+	virtual float GetValue( float F = 0.f, UObject* Data = NULL, class FRandomStream* InRandomStream = NULL ) const override;
 	// End UDistributionFloat Interface
 
 	// Begin FCurveEdInterface interface
-	virtual int32		GetNumKeys() const OVERRIDE;
-	virtual int32		GetNumSubCurves() const OVERRIDE;
-	virtual float	GetKeyIn(int32 KeyIndex) OVERRIDE;
-	virtual float	GetKeyOut(int32 SubIndex, int32 KeyIndex) OVERRIDE;
-	virtual FColor	GetKeyColor(int32 SubIndex, int32 KeyIndex, const FColor& CurveColor) OVERRIDE;
-	virtual void	GetInRange(float& MinIn, float& MaxIn) const OVERRIDE;
-	virtual void	GetOutRange(float& MinOut, float& MaxOut) const OVERRIDE;
-	virtual EInterpCurveMode	GetKeyInterpMode(int32 KeyIndex) const OVERRIDE;
-	virtual void	GetTangents(int32 SubIndex, int32 KeyIndex, float& ArriveTangent, float& LeaveTangent) const OVERRIDE;
-	virtual float	EvalSub(int32 SubIndex, float InVal) OVERRIDE;
-	virtual int32		CreateNewKey(float KeyIn) OVERRIDE;
-	virtual void	DeleteKey(int32 KeyIndex) OVERRIDE;
-	virtual int32		SetKeyIn(int32 KeyIndex, float NewInVal) OVERRIDE;
-	virtual void	SetKeyOut(int32 SubIndex, int32 KeyIndex, float NewOutVal) OVERRIDE;
-	virtual void	SetKeyInterpMode(int32 KeyIndex, EInterpCurveMode NewMode) OVERRIDE;
-	virtual void	SetTangents(int32 SubIndex, int32 KeyIndex, float ArriveTangent, float LeaveTangent) OVERRIDE;
+	virtual int32		GetNumKeys() const override;
+	virtual int32		GetNumSubCurves() const override;
+	virtual float	GetKeyIn(int32 KeyIndex) override;
+	virtual float	GetKeyOut(int32 SubIndex, int32 KeyIndex) override;
+	virtual FColor	GetKeyColor(int32 SubIndex, int32 KeyIndex, const FColor& CurveColor) override;
+	virtual void	GetInRange(float& MinIn, float& MaxIn) const override;
+	virtual void	GetOutRange(float& MinOut, float& MaxOut) const override;
+	virtual EInterpCurveMode	GetKeyInterpMode(int32 KeyIndex) const override;
+	virtual void	GetTangents(int32 SubIndex, int32 KeyIndex, float& ArriveTangent, float& LeaveTangent) const override;
+	virtual float	EvalSub(int32 SubIndex, float InVal) override;
+	virtual int32		CreateNewKey(float KeyIn) override;
+	virtual void	DeleteKey(int32 KeyIndex) override;
+	virtual int32		SetKeyIn(int32 KeyIndex, float NewInVal) override;
+	virtual void	SetKeyOut(int32 SubIndex, int32 KeyIndex, float NewOutVal) override;
+	virtual void	SetKeyInterpMode(int32 KeyIndex, EInterpCurveMode NewMode) override;
+	virtual void	SetTangents(int32 SubIndex, int32 KeyIndex, float ArriveTangent, float LeaveTangent) override;
 	// End FCurveEdInterface interface
 };
 
diff --git a/Engine/Source/Runtime/Engine/Classes/Distributions/DistributionFloatParameterBase.h b/Engine/Source/Runtime/Engine/Classes/Distributions/DistributionFloatParameterBase.h
--- a/Engine/Source/Runtime/Engine/Classes/Distributions/DistributionFloatParameterBase.h
+++ b/Engine/Source/Runtime/Engine/Classes/Distributions/DistributionFloatParameterBase.h
@@ -44,7 +44,7 @@ class UDistributionFloatParameterBase : public UDistributionFloatConstant
 
 
 	// Begin UDistributionFloat Interface
-	virtual float GetValue( float F = 0.f, UObject* Data = NULL, class FRandomStream* InRandomStream = NULL ) const OVERRIDE;
+	virtual float GetValue( float F = 0.f, UObject* Data = NULL, class FRandomStream* InRandomStream = NULL ) const override;
 	// End UDistributionFloat Interface
 	
 	
@@ -53,7 +53,7 @@ class UDistributionFloatParameterBase : public UDistributionFloatConstant
 
 	
 	// Begin Distribution Float
-	virtual bool CanBeBaked() const OVERRIDE { return false; }
+	virtual bool CanBeBaked() const override { return false; }
 };
 
 
diff --git a/Engine/Source/Runtime/Engine/Classes/Distributions/DistributionFloatParticleParameter.h b/Engine/Source/Runtime/Engine/Classes/Distributions/DistributionFloatParticleParameter.h
--- a/Engine/Source/Runtime/Engine/Classes/Distributions/DistributionFloatParticleParameter.h
+++ b/Engine/Source/Runtime/Engine/Classes/Distributions/DistributionFloatParticleParameter.h
@@ -11,7 +11,7 @@ class UDistributionFloatParticleParameter : public UDistributionFloatParameterBa
 
 
 	// Begin UDistributionFloatParameterBase Interface
-	virtual bool GetParamValue(UObject* Data, FName ParamName, float& OutFloat) const OVERRIDE;
+	virtual bool GetParamValue(UObject* Data, FName ParamName, float& OutFloat) const override;
 	// End UDistributionFloatParameterBase Interface
 };
 
diff --git a/Engine/Source/Runtime/Engine/Classes/Distributions/DistributionFloatUniform.h b/Engine/Source/Runtime/Engine/Classes/Distributions/DistributionFloatUniform.h
--- a/Engine/Source/Runtime/Engine/Classes/Distributions/DistributionFloatUniform.h
+++ b/Engine/Source/Runtime/Engine/Classes/Distributions/DistributionFloatUniform.h
@@ -19,36 +19,36 @@ class UDistributionFloatUniform : public UDistributionFloat
 	float Max;
 
 	// Begin UObject Interface
-	virtual void PostInitProperties() OVERRIDE;
-	virtual void PostLoad() OVERRIDE;
+	virtual void PostInitProperties() override;
+	virtual void PostLoad() override;
 	// End UObject Interface
 
 	// Begin UDistributionFloat Interface
-	virtual float GetValue( float F = 0.f, UObject* Data = NULL, class FRandomStream* InRandomStream = NULL ) const OVERRIDE;
+	virtual float GetValue( float F = 0.f, UObject* Data = NULL, class FRandomStream* InRandomStream = NULL ) const override;
 
 	[REDACTED_EMAIL]: Currently, consoles need this? At least until we have some sort of cooking/packaging step!
-	virtual ERawDistributionOperation GetOperation() const OVERRIDE;
-	virtual uint32 InitializeRawEntry(float Time, float* Values) const OVERRIDE;
+	virtual ERawDistributionOperation GetOperation() const override;
+	virtual uint32 InitializeRawEntry(float Time, float* Values) const override;
 	// End UDistributionFloat Interface
 
 	// Begin FCurveEdInterface interface
-	virtual int32		GetNumKeys() const OVERRIDE;
-	virtual int32		GetNumSubCurves() const OVERRIDE;
-	virtual FColor	GetSubCurveButtonColor(int32 SubCurveIndex, bool bIsSubCurveHidden) const OVERRIDE;
-	virtual float	GetKeyIn(int32 KeyIndex) OVERRIDE;
-	virtual float	GetKeyOut(int32 SubIndex, int32 KeyIndex) OVERRIDE;
-	virtual FColor	GetKeyColor(int32 SubIndex, int32 KeyIndex, const FColor& CurveColor) OVERRIDE;
-	virtual void	GetInRange(float& MinIn, float& MaxIn) const OVERRIDE;
-	virtual void	GetOutRange(float& MinOut, float& MaxOut) const OVERRIDE;
-	virtual EInterpCurveMode	GetKeyInterpMode(int32 KeyIndex) const OVERRIDE;
-	virtual void	GetTangents(int32 SubIndex, int32 KeyIndex, float& ArriveTangent, float& LeaveTangent) const OVERRIDE;
-	virtual float	EvalSub(int32 SubIndex, float InVal) OVERRIDE;
-	virtual int32		CreateNewKey(float KeyIn) OVERRIDE;
-	virtual void	DeleteKey(int32 KeyIndex) OVERRIDE;
-	virtual int32		SetKeyIn(int32 KeyIndex, float NewInVal) OVERRIDE;
-	virtual void	SetKeyOut(int32 SubIndex, int32 KeyIndex, float NewOutVal) OVERRIDE;
-	virtual void	SetKeyInterpMode(int32 KeyIndex, EInterpCurveMode NewMode) OVERRIDE;
-	virtual void	SetTangents(int32 SubIndex, int32 KeyIndex, float ArriveTangent, float LeaveTangent) OVERRIDE;
+	virtual int32		GetNumKeys() const override;
+	virtual int32		GetNumSubCurves() const override;
+	virtual FColor	GetSubCurveButtonColor(int32 SubCurveIndex, bool bIsSubCurveHidden) const override;
+	virtual float	GetKeyIn(int32 KeyIndex) override;
+	virtual float	GetKeyOut(int32 SubIndex, int32 KeyIndex) override;
+	virtual FColor	GetKeyColor(int32 SubIndex, int32 KeyIndex, const FColor& CurveColor) override;
+	virtual void	GetInRange(float& MinIn, float& MaxIn) const override;
+	virtual void	GetOutRange(float& MinOut, float& MaxOut) const override;
+	virtual EInterpCurveMode	GetKeyInterpMode(int32 KeyIndex) const override;
+	virtual void	GetTangents(int32 SubIndex, int32 KeyIndex, float& ArriveTangent, float& LeaveTangent) const override;
+	virtual float	EvalSub(int32 SubIndex, float InVal) override;
+	virtual int32		CreateNewKey(float KeyIn) override;
+	virtual void	DeleteKey(int32 KeyIndex) override;
+	virtual int32		SetKeyIn(int32 KeyIndex, float NewInVal) override;
+	virtual void	SetKeyOut(int32 SubIndex, int32 KeyIndex, float NewOutVal) override;
+	virtual void	SetKeyInterpMode(int32 KeyIndex, EInterpCurveMode NewMode) override;
+	virtual void	SetTangents(int32 SubIndex, int32 KeyIndex, float ArriveTangent, float LeaveTangent) override;
 	// End FCurveEdInterface interface
 };
 
diff --git a/Engine/Source/Runtime/Engine/Classes/Distributions/DistributionFloatUniformCurve.h b/Engine/Source/Runtime/Engine/Classes/Distributions/DistributionFloatUniformCurve.h
--- a/Engine/Source/Runtime/Engine/Classes/Distributions/DistributionFloatUniformCurve.h
+++ b/Engine/Source/Runtime/Engine/Classes/Distributions/DistributionFloatUniformCurve.h
@@ -16,10 +16,10 @@ class UDistributionFloatUniformCurve : public UDistributionFloat
 
 
 	// Begin UDistributionFloat Interface
-	virtual float GetValue( float F = 0.f, UObject* Data = NULL, class FRandomStream* InRandomStream = NULL ) const OVERRIDE;
+	virtual float GetValue( float F = 0.f, UObject* Data = NULL, class FRandomStream* InRandomStream = NULL ) const override;
 	[REDACTED_EMAIL]: Currently, consoles need this? At least until we have some sort of cooking/packaging step!
-	virtual ERawDistributionOperation GetOperation() const OVERRIDE;
-	virtual uint32 InitializeRawEntry(float Time, float* Values) const OVERRIDE;
+	virtual ERawDistributionOperation GetOperation() const override;
+	virtual uint32 InitializeRawEntry(float Time, float* Values) const override;
 	// End UDistributionFloat Interface
 
 	/**
@@ -28,23 +28,23 @@ class UDistributionFloatUniformCurve : public UDistributionFloat
 	virtual FVector2D GetMinMaxValue(float F = 0.f, UObject* Data = NULL) const;
 
 	// Begin FCurveEdInterface interface
-	virtual int32		GetNumKeys() const OVERRIDE;
-	virtual int32		GetNumSubCurves() const OVERRIDE;
-	virtual FColor	GetSubCurveButtonColor(int32 SubCurveIndex, bool bIsSubCurveHidden) const OVERRIDE;
-	virtual float	GetKeyIn(int32 KeyIndex) OVERRIDE;
-	virtual float	GetKeyOut(int32 SubIndex, int32 KeyIndex) OVERRIDE;
-	virtual FColor	GetKeyColor(int32 SubIndex, int32 KeyIndex, const FColor& CurveColor) OVERRIDE;
-	virtual void	GetInRange(float& MinIn, float& MaxIn) const OVERRIDE;
-	virtual void	GetOutRange(float& MinOut, float& MaxOut) const OVERRIDE;
-	virtual EInterpCurveMode	GetKeyInterpMode(int32 KeyIndex) const OVERRIDE;
-	virtual void	GetTangents(int32 SubIndex, int32 KeyIndex, float& ArriveTangent, float& LeaveTangent) const OVERRIDE;
-	virtual float	EvalSub(int32 SubIndex, float InVal) OVERRIDE;
-	virtual int32		CreateNewKey(float KeyIn) OVERRIDE;
-	virtual void	DeleteKey(int32 KeyIndex) OVERRIDE;
-	virtual int32		SetKeyIn(int32 KeyIndex, float NewInVal) OVERRIDE;
-	virtual void	SetKeyOut(int32 SubIndex, int32 KeyIndex, float NewOutVal) OVERRIDE;
-	virtual void	SetKeyInterpMode(int32 KeyIndex, EInterpCurveMode NewMode) OVERRIDE;
-	virtual void	SetTangents(int32 SubIndex, int32 KeyIndex, float ArriveTangent, float LeaveTangent) OVERRIDE;
+	virtual int32		GetNumKeys() const override;
+	virtual int32		GetNumSubCurves() const override;
+	virtual FColor	GetSubCurveButtonColor(int32 SubCurveIndex, bool bIsSubCurveHidden) const override;
+	virtual float	GetKeyIn(int32 KeyIndex) override;
+	virtual float	GetKeyOut(int32 SubIndex, int32 KeyIndex) override;
+	virtual FColor	GetKeyColor(int32 SubIndex, int32 KeyIndex, const FColor& CurveColor) override;
+	virtual void	GetInRange(float& MinIn, float& MaxIn) const override;
+	virtual void	GetOutRange(float& MinOut, float& MaxOut) const override;
+	virtual EInterpCurveMode	GetKeyInterpMode(int32 KeyIndex) const override;
+	virtual void	GetTangents(int32 SubIndex, int32 KeyIndex, float& ArriveTangent, float& LeaveTangent) const override;
+	virtual float	EvalSub(int32 SubIndex, float InVal) override;
+	virtual int32		CreateNewKey(float KeyIn) override;
+	virtual void	DeleteKey(int32 KeyIndex) override;
+	virtual int32		SetKeyIn(int32 KeyIndex, float NewInVal) override;
+	virtual void	SetKeyOut(int32 SubIndex, int32 KeyIndex, float NewOutVal) override;
+	virtual void	SetKeyInterpMode(int32 KeyIndex, EInterpCurveMode NewMode) override;
+	virtual void	SetTangents(int32 SubIndex, int32 KeyIndex, float ArriveTangent, float LeaveTangent) override;
 	// End FCurveEdInterface interface
 };
 
diff --git a/Engine/Source/Runtime/Engine/Classes/Distributions/DistributionVector.h b/Engine/Source/Runtime/Engine/Classes/Distributions/DistributionVector.h
--- a/Engine/Source/Runtime/Engine/Classes/Distributions/DistributionVector.h
+++ b/Engine/Source/Runtime/Engine/Classes/Distributions/DistributionVector.h
@@ -169,10 +169,10 @@ class ENGINE_API UDistributionVector : public UDistribution
 
 	/** Begin UObject interface */
 #if	WITH_EDITOR
-	virtual void PostEditChangeProperty(FPropertyChangedEvent& PropertyChangedEvent) OVERRIDE;
+	virtual void PostEditChangeProperty(FPropertyChangedEvent& PropertyChangedEvent) override;
 #endif // WITH_EDITOR
-	virtual bool NeedsLoadForClient() const OVERRIDE;
-	virtual bool NeedsLoadForServer() const OVERRIDE;
+	virtual bool NeedsLoadForClient() const override;
+	virtual bool NeedsLoadForServer() const override;
 	/** End UObject interface */
 
 };
diff --git a/Engine/Source/Runtime/Engine/Classes/Distributions/DistributionVectorConstant.h b/Engine/Source/Runtime/Engine/Classes/Distributions/DistributionVectorConstant.h
--- a/Engine/Source/Runtime/Engine/Classes/Distributions/DistributionVectorConstant.h
+++ b/Engine/Source/Runtime/Engine/Classes/Distributions/DistributionVectorConstant.h
@@ -22,33 +22,33 @@ class ENGINE_API UDistributionVectorConstant : public UDistributionVector
 
 
 	//Begin UDistributionVector Interface
-	virtual FVector	GetValue( float F = 0.f, UObject* Data = NULL, int32 LastExtreme = 0, class FRandomStream* InRandomStream = NULL ) const OVERRIDE;
-	virtual	void	GetRange(FVector& OutMin, FVector& OutMax) const OVERRIDE;
+	virtual FVector	GetValue( float F = 0.f, UObject* Data = NULL, int32 LastExtreme = 0, class FRandomStream* InRandomStream = NULL ) const override;
+	virtual	void	GetRange(FVector& OutMin, FVector& OutMax) const override;
 	//End UDistributionVector Interface
 
 	// Begin UObject interface
-	virtual void PostInitProperties() OVERRIDE;
-	virtual void PostLoad() OVERRIDE;
+	virtual void PostInitProperties() override;
+	virtual void PostLoad() override;
 	// End UObject interface
 	
 	// Begin FCurveEdInterface interface
-	virtual int32		GetNumKeys() const OVERRIDE;
-	virtual int32		GetNumSubCurves() const OVERRIDE;
-	virtual FColor	GetSubCurveButtonColor(int32 SubCurveIndex, bool bIsSubCurveHidden) const OVERRIDE;
-	virtual float	GetKeyIn(int32 KeyIndex) OVERRIDE;
-	virtual float	GetKeyOut(int32 SubIndex, int32 KeyIndex) OVERRIDE;
-	virtual FColor	GetKeyColor(int32 SubIndex, int32 KeyIndex, const FColor& CurveColor) OVERRIDE;
-	virtual void	GetInRange(float& MinIn, float& MaxIn) const OVERRIDE;
-	virtual void	GetOutRange(float& MinOut, float& MaxOut) const OVERRIDE;
-	virtual EInterpCurveMode	GetKeyInterpMode(int32 KeyIndex) const OVERRIDE;
-	virtual void	GetTangents(int32 SubIndex, int32 KeyIndex, float& ArriveTangent, float& LeaveTangent) const OVERRIDE;
-	virtual float	EvalSub(int32 SubIndex, float InVal) OVERRIDE;
-	virtual int32		CreateNewKey(float KeyIn) OVERRIDE;
-	virtual void	DeleteKey(int32 KeyIndex) OVERRIDE;
-	virtual int32		SetKeyIn(int32 KeyIndex, float NewInVal) OVERRIDE;
-	virtual void	SetKeyOut(int32 SubIndex, int32 KeyIndex, float NewOutVal) OVERRIDE;
-	virtual void	SetKeyInterpMode(int32 KeyIndex, EInterpCurveMode NewMode) OVERRIDE;
-	virtual void	SetTangents(int32 SubIndex, int32 KeyIndex, float ArriveTangent, float LeaveTangent) OVERRIDE;
+	virtual int32		GetNumKeys() const override;
+	virtual int32		GetNumSubCurves() const override;
+	virtual FColor	GetSubCurveButtonColor(int32 SubCurveIndex, bool bIsSubCurveHidden) const override;
+	virtual float	GetKeyIn(int32 KeyIndex) override;
+	virtual float	GetKeyOut(int32 SubIndex, int32 KeyIndex) override;
+	virtual FColor	GetKeyColor(int32 SubIndex, int32 KeyIndex, const FColor& CurveColor) override;
+	virtual void	GetInRange(float& MinIn, float& MaxIn) const override;
+	virtual void	GetOutRange(float& MinOut, float& MaxOut) const override;
+	virtual EInterpCurveMode	GetKeyInterpMode(int32 KeyIndex) const override;
+	virtual void	GetTangents(int32 SubIndex, int32 KeyIndex, float& ArriveTangent, float& LeaveTangent) const override;
+	virtual float	EvalSub(int32 SubIndex, float InVal) override;
+	virtual int32		CreateNewKey(float KeyIn) override;
+	virtual void	DeleteKey(int32 KeyIndex) override;
+	virtual int32		SetKeyIn(int32 KeyIndex, float NewInVal) override;
+	virtual void	SetKeyOut(int32 SubIndex, int32 KeyIndex, float NewOutVal) override;
+	virtual void	SetKeyInterpMode(int32 KeyIndex, EInterpCurveMode NewMode) override;
+	virtual void	SetTangents(int32 SubIndex, int32 KeyIndex, float ArriveTangent, float LeaveTangent) override;
 	// End FCurveEdInterface interface
 
 };
diff --git a/Engine/Source/Runtime/Engine/Classes/Distributions/DistributionVectorConstantCurve.h b/Engine/Source/Runtime/Engine/Classes/Distributions/DistributionVectorConstantCurve.h
--- a/Engine/Source/Runtime/Engine/Classes/Distributions/DistributionVectorConstantCurve.h
+++ b/Engine/Source/Runtime/Engine/Classes/Distributions/DistributionVectorConstantCurve.h
@@ -22,28 +22,28 @@ class UDistributionVectorConstantCurve : public UDistributionVector
 
 
 	//Begin UDistributionVector Interface
-	virtual FVector	GetValue( float F = 0.f, UObject* Data = NULL, int32 LastExtreme = 0, class FRandomStream* InRandomStream = NULL ) const OVERRIDE;
-	virtual	void	GetRange(FVector& OutMin, FVector& OutMax) const OVERRIDE;
+	virtual FVector	GetValue( float F = 0.f, UObject* Data = NULL, int32 LastExtreme = 0, class FRandomStream* InRandomStream = NULL ) const override;
+	virtual	void	GetRange(FVector& OutMin, FVector& OutMax) const override;
 	//End UDistributionVector Interface
 
 	// Begin FCurveEdInterface interface
-	virtual int32		GetNumKeys() const OVERRIDE;
-	virtual int32		GetNumSubCurves() const OVERRIDE;
-	virtual FColor	GetSubCurveButtonColor(int32 SubCurveIndex, bool bIsSubCurveHidden) const OVERRIDE;
-	virtual float	GetKeyIn(int32 KeyIndex) OVERRIDE;
-	virtual float	GetKeyOut(int32 SubIndex, int32 KeyIndex) OVERRIDE;
-	virtual void	GetInRange(float& MinIn, float& MaxIn) const OVERRIDE;
-	virtual void	GetOutRange(float& MinOut, float& MaxOut) const OVERRIDE;
-	virtual FColor	GetKeyColor(int32 SubIndex, int32 KeyIndex, const FColor& CurveColor) OVERRIDE;
-	virtual EInterpCurveMode	GetKeyInterpMode(int32 KeyIndex) const OVERRIDE;
-	virtual void	GetTangents(int32 SubIndex, int32 KeyIndex, float& ArriveTangent, float& LeaveTangent) const OVERRIDE;
-	virtual float	EvalSub(int32 SubIndex, float InVal) OVERRIDE;
-	virtual int32		CreateNewKey(float KeyIn) OVERRIDE;
-	virtual void	DeleteKey(int32 KeyIndex) OVERRIDE;
-	virtual int32		SetKeyIn(int32 KeyIndex, float NewInVal) OVERRIDE;
-	virtual void	SetKeyOut(int32 SubIndex, int32 KeyIndex, float NewOutVal) OVERRIDE;
-	virtual void	SetKeyInterpMode(int32 KeyIndex, EInterpCurveMode NewMode) OVERRIDE;
-	virtual void	SetTangents(int32 SubIndex, int32 KeyIndex, float ArriveTangent, float LeaveTangent) OVERRIDE;
+	virtual int32		GetNumKeys() const override;
+	virtual int32		GetNumSubCurves() const override;
+	virtual FColor	GetSubCurveButtonColor(int32 SubCurveIndex, bool bIsSubCurveHidden) const override;
+	virtual float	GetKeyIn(int32 KeyIndex) override;
+	virtual float	GetKeyOut(int32 SubIndex, int32 KeyIndex) override;
+	virtual void	GetInRange(float& MinIn, float& MaxIn) const override;
+	virtual void	GetOutRange(float& MinOut, float& MaxOut) const override;
+	virtual FColor	GetKeyColor(int32 SubIndex, int32 KeyIndex, const FColor& CurveColor) override;
+	virtual EInterpCurveMode	GetKeyInterpMode(int32 KeyIndex) const override;
+	virtual void	GetTangents(int32 SubIndex, int32 KeyIndex, float& ArriveTangent, float& LeaveTangent) const override;
+	virtual float	EvalSub(int32 SubIndex, float InVal) override;
+	virtual int32		CreateNewKey(float KeyIn) override;
+	virtual void	DeleteKey(int32 KeyIndex) override;
+	virtual int32		SetKeyIn(int32 KeyIndex, float NewInVal) override;
+	virtual void	SetKeyOut(int32 SubIndex, int32 KeyIndex, float NewOutVal) override;
+	virtual void	SetKeyInterpMode(int32 KeyIndex, EInterpCurveMode NewMode) override;
+	virtual void	SetTangents(int32 SubIndex, int32 KeyIndex, float ArriveTangent, float LeaveTangent) override;
 	// End FCurveEdInterface interface
 
 };
diff --git a/Engine/Source/Runtime/Engine/Classes/Distributions/DistributionVectorParameterBase.h b/Engine/Source/Runtime/Engine/Classes/Distributions/DistributionVectorParameterBase.h
--- a/Engine/Source/Runtime/Engine/Classes/Distributions/DistributionVectorParameterBase.h
+++ b/Engine/Source/Runtime/Engine/Classes/Distributions/DistributionVectorParameterBase.h
@@ -29,8 +29,8 @@ class UDistributionVectorParameterBase : public UDistributionVectorConstant
 
 
 	//Begin UDistributionVector Interface
-	virtual FVector GetValue(float F = 0.f, UObject* Data = NULL, int32 Extreme = 0, class FRandomStream* InRandomStream = NULL) const OVERRIDE;
-	virtual bool CanBeBaked() const OVERRIDE { return false; }
+	virtual FVector GetValue(float F = 0.f, UObject* Data = NULL, int32 Extreme = 0, class FRandomStream* InRandomStream = NULL) const override;
+	virtual bool CanBeBaked() const override { return false; }
 	//End UDistributionVector Interface
 	
 	/** todo document */
diff --git a/Engine/Source/Runtime/Engine/Classes/Distributions/DistributionVectorParticleParameter.h b/Engine/Source/Runtime/Engine/Classes/Distributions/DistributionVectorParticleParameter.h
--- a/Engine/Source/Runtime/Engine/Classes/Distributions/DistributionVectorParticleParameter.h
+++ b/Engine/Source/Runtime/Engine/Classes/Distributions/DistributionVectorParticleParameter.h
@@ -11,7 +11,7 @@ class UDistributionVectorParticleParameter : public UDistributionVectorParameter
 
 
 	// Begin UDistributionVectorParameterBase Interface
-	virtual bool GetParamValue(UObject* Data, FName ParamName, FVector& OutVector) const OVERRIDE;
+	virtual bool GetParamValue(UObject* Data, FName ParamName, FVector& OutVector) const override;
 	// End UDistributionVectorParameterBase Interface
 };
 
diff --git a/Engine/Source/Runtime/Engine/Classes/Distributions/DistributionVectorUniform.h b/Engine/Source/Runtime/Engine/Classes/Distributions/DistributionVectorUniform.h
--- a/Engine/Source/Runtime/Engine/Classes/Distributions/DistributionVectorUniform.h
+++ b/Engine/Source/Runtime/Engine/Classes/Distributions/DistributionVectorUniform.h
@@ -33,19 +33,19 @@ class UDistributionVectorUniform : public UDistributionVector
 
 
 	// Begin UObject Interface
-	virtual void PostInitProperties() OVERRIDE;
-	virtual void PostLoad() OVERRIDE;
+	virtual void PostInitProperties() override;
+	virtual void PostLoad() override;
 	// End UObject Interface
 
-	virtual FVector	GetValue( float F = 0.f, UObject* Data = NULL, int32 LastExtreme = 0, class FRandomStream* InRandomStream = NULL ) const OVERRIDE;
+	virtual FVector	GetValue( float F = 0.f, UObject* Data = NULL, int32 LastExtreme = 0, class FRandomStream* InRandomStream = NULL ) const override;
 
 	//Begin UDistributionVector Interface
 
 	[REDACTED_EMAIL]: Currently, consoles need this? At least until we have some sort of cooking/packaging step!
-	virtual ERawDistributionOperation GetOperation() const OVERRIDE;
-	virtual uint8 GetLockFlag() const OVERRIDE;
-	virtual uint32 InitializeRawEntry(float Time, float* Values) const OVERRIDE;
-	virtual	void	GetRange(FVector& OutMin, FVector& OutMax) const OVERRIDE;
+	virtual ERawDistributionOperation GetOperation() const override;
+	virtual uint8 GetLockFlag() const override;
+	virtual uint32 InitializeRawEntry(float Time, float* Values) const override;
+	virtual	void	GetRange(FVector& OutMin, FVector& OutMax) const override;
 	//End UDistributionVector Interface
 
 	/** These two functions will retrieve the Min/Max values respecting the Locked and Mirror flags. */
@@ -58,23 +58,23 @@ class UDistributionVectorUniform : public UDistributionVector
 	// 4,5 = min/max z
 
 	// Begin FCurveEdInterface interface
-	virtual int32		GetNumKeys() const OVERRIDE;
-	virtual int32		GetNumSubCurves() const OVERRIDE;
-	virtual FColor	GetSubCurveButtonColor(int32 SubCurveIndex, bool bIsSubCurveHidden) const OVERRIDE;
-	virtual float	GetKeyIn(int32 KeyIndex) OVERRIDE;
-	virtual float	GetKeyOut(int32 SubIndex, int32 KeyIndex) OVERRIDE;
-	virtual FColor	GetKeyColor(int32 SubIndex, int32 KeyIndex, const FColor& CurveColor) OVERRIDE;
-	virtual void	GetInRange(float& MinIn, float& MaxIn) const OVERRIDE;
-	virtual void	GetOutRange(float& MinOut, float& MaxOut) const OVERRIDE;
-	virtual EInterpCurveMode	GetKeyInterpMode(int32 KeyIndex) const OVERRIDE;
-	virtual void	GetTangents(int32 SubIndex, int32 KeyIndex, float& ArriveTangent, float& LeaveTangent) const OVERRIDE;
-	virtual float	EvalSub(int32 SubIndex, float InVal) OVERRIDE;
-	virtual int32		CreateNewKey(float KeyIn) OVERRIDE;
-	virtual void	DeleteKey(int32 KeyIndex) OVERRIDE;
-	virtual int32		SetKeyIn(int32 KeyIndex, float NewInVal) OVERRIDE;
-	virtual void	SetKeyOut(int32 SubIndex, int32 KeyIndex, float NewOutVal) OVERRIDE;
-	virtual void	SetKeyInterpMode(int32 KeyIndex, EInterpCurveMode NewMode) OVERRIDE;
-	virtual void	SetTangents(int32 SubIndex, int32 KeyIndex, float ArriveTangent, float LeaveTangent) OVERRIDE;
+	virtual int32		GetNumKeys() const override;
+	virtual int32		GetNumSubCurves() const override;
+	virtual FColor	GetSubCurveButtonColor(int32 SubCurveIndex, bool bIsSubCurveHidden) const override;
+	virtual float	GetKeyIn(int32 KeyIndex) override;
+	virtual float	GetKeyOut(int32 SubIndex, int32 KeyIndex) override;
+	virtual FColor	GetKeyColor(int32 SubIndex, int32 KeyIndex, const FColor& CurveColor) override;
+	virtual void	GetInRange(float& MinIn, float& MaxIn) const override;
+	virtual void	GetOutRange(float& MinOut, float& MaxOut) const override;
+	virtual EInterpCurveMode	GetKeyInterpMode(int32 KeyIndex) const override;
+	virtual void	GetTangents(int32 SubIndex, int32 KeyIndex, float& ArriveTangent, float& LeaveTangent) const override;
+	virtual float	EvalSub(int32 SubIndex, float InVal) override;
+	virtual int32		CreateNewKey(float KeyIn) override;
+	virtual void	DeleteKey(int32 KeyIndex) override;
+	virtual int32		SetKeyIn(int32 KeyIndex, float NewInVal) override;
+	virtual void	SetKeyOut(int32 SubIndex, int32 KeyIndex, float NewOutVal) override;
+	virtual void	SetKeyInterpMode(int32 KeyIndex, EInterpCurveMode NewMode) override;
+	virtual void	SetTangents(int32 SubIndex, int32 KeyIndex, float ArriveTangent, float LeaveTangent) override;
 	// Begin FCurveEdInterface interface
 
 };
diff --git a/Engine/Source/Runtime/Engine/Classes/Distributions/DistributionVectorUniformCurve.h b/Engine/Source/Runtime/Engine/Classes/Distributions/DistributionVectorUniformCurve.h
--- a/Engine/Source/Runtime/Engine/Classes/Distributions/DistributionVectorUniformCurve.h
+++ b/Engine/Source/Runtime/Engine/Classes/Distributions/DistributionVectorUniformCurve.h
@@ -30,13 +30,13 @@ class UDistributionVectorUniformCurve : public UDistributionVector
 	UPROPERTY(EditAnywhere, Category=DistributionVectorUniformCurve)
 	uint32 bUseExtremes:1;
 
-	virtual FVector	GetValue( float F = 0.f, UObject* Data = NULL, int32 LastExtreme = 0, class FRandomStream* InRandomStream = NULL ) const OVERRIDE;
+	virtual FVector	GetValue( float F = 0.f, UObject* Data = NULL, int32 LastExtreme = 0, class FRandomStream* InRandomStream = NULL ) const override;
 
 	//Begin UDistributionVector Interface
 	[REDACTED_EMAIL]: Currently, consoles need this? At least until we have some sort of cooking/packaging step!
-	virtual ERawDistributionOperation GetOperation() const OVERRIDE;
-	virtual uint32 InitializeRawEntry(float Time, float* Values) const OVERRIDE;
-	virtual	void	GetRange(FVector& OutMin, FVector& OutMax) const OVERRIDE;
+	virtual ERawDistributionOperation GetOperation() const override;
+	virtual uint32 InitializeRawEntry(float Time, float* Values) const override;
+	virtual	void	GetRange(FVector& OutMin, FVector& OutMax) const override;
 	//End UDistributionVector Interface
 
 	/** These two functions will retrieve the Min/Max values respecting the Locked and Mirror flags. */
@@ -49,23 +49,23 @@ class UDistributionVectorUniformCurve : public UDistributionVector
 	virtual FTwoVectors GetMinMaxValue(float F = 0.f, UObject* Data = NULL) const;
 
 	// Begin FCurveEdInterface interface
-	virtual int32		GetNumKeys() const OVERRIDE;
-	virtual int32		GetNumSubCurves() const OVERRIDE;
-	virtual FColor	GetSubCurveButtonColor(int32 SubCurveIndex, bool bIsSubCurveHidden) const OVERRIDE;
-	virtual float	GetKeyIn(int32 KeyIndex) OVERRIDE;
-	virtual float	GetKeyOut(int32 SubIndex, int32 KeyIndex) OVERRIDE;
-	virtual FColor	GetKeyColor(int32 SubIndex, int32 KeyIndex, const FColor& CurveColor) OVERRIDE;
-	virtual void	GetInRange(float& MinIn, float& MaxIn) const OVERRIDE;
-	virtual void	GetOutRange(float& MinOut, float& MaxOut) const OVERRIDE;
-	virtual EInterpCurveMode	GetKeyInterpMode(int32 KeyIndex) const OVERRIDE;
-	virtual void	GetTangents(int32 SubIndex, int32 KeyIndex, float& ArriveTangent, float& LeaveTangent) const OVERRIDE;
-	virtual float	EvalSub(int32 SubIndex, float InVal) OVERRIDE;
-	virtual int32		CreateNewKey(float KeyIn) OVERRIDE;
-	virtual void	DeleteKey(int32 KeyIndex) OVERRIDE;
-	virtual int32		SetKeyIn(int32 KeyIndex, float NewInVal) OVERRIDE;
-	virtual void	SetKeyOut(int32 SubIndex, int32 KeyIndex, float NewOutVal) OVERRIDE;
-	virtual void	SetKeyInterpMode(int32 KeyIndex, EInterpCurveMode NewMode) OVERRIDE;
-	virtual void	SetTangents(int32 SubIndex, int32 KeyIndex, float ArriveTangent, float LeaveTangent) OVERRIDE;
+	virtual int32		GetNumKeys() const override;
+	virtual int32		GetNumSubCurves() const override;
+	virtual FColor	GetSubCurveButtonColor(int32 SubCurveIndex, bool bIsSubCurveHidden) const override;
+	virtual float	GetKeyIn(int32 KeyIndex) override;
+	virtual float	GetKeyOut(int32 SubIndex, int32 KeyIndex) override;
+	virtual FColor	GetKeyColor(int32 SubIndex, int32 KeyIndex, const FColor& CurveColor) override;
+	virtual void	GetInRange(float& MinIn, float& MaxIn) const override;
+	virtual void	GetOutRange(float& MinOut, float& MaxOut) const override;
+	virtual EInterpCurveMode	GetKeyInterpMode(int32 KeyIndex) const override;
+	virtual void	GetTangents(int32 SubIndex, int32 KeyIndex, float& ArriveTangent, float& LeaveTangent) const override;
+	virtual float	EvalSub(int32 SubIndex, float InVal) override;
+	virtual int32		CreateNewKey(float KeyIn) override;
+	virtual void	DeleteKey(int32 KeyIndex) override;
+	virtual int32		SetKeyIn(int32 KeyIndex, float NewInVal) override;
+	virtual void	SetKeyOut(int32 SubIndex, int32 KeyIndex, float NewOutVal) override;
+	virtual void	SetKeyInterpMode(int32 KeyIndex, EInterpCurveMode NewMode) override;
+	virtual void	SetTangents(int32 SubIndex, int32 KeyIndex, float ArriveTangent, float LeaveTangent) override;
 	virtual void	LockAndMirror(FTwoVectors& Val) const;
 	// Begin FCurveEdInterface interface
 
diff --git a/Engine/Source/Runtime/Engine/Classes/EdGraph/EdGraphNode_Comment.h b/Engine/Source/Runtime/Engine/Classes/EdGraph/EdGraphNode_Comment.h
--- a/Engine/Source/Runtime/Engine/Classes/EdGraph/EdGraphNode_Comment.h
+++ b/Engine/Source/Runtime/Engine/Classes/EdGraph/EdGraphNode_Comment.h
@@ -44,19 +44,19 @@ public:
 
 #if WITH_EDITOR
 	// Begin UEdGraphNode interface
-	virtual void AllocateDefaultPins() OVERRIDE {}
-	ENGINE_API virtual FString GetTooltip() const OVERRIDE;
-	ENGINE_API virtual FLinearColor GetNodeCommentColor() const OVERRIDE;
-	ENGINE_API virtual FText GetNodeTitle(ENodeTitleType::Type TitleType) const OVERRIDE;
-	virtual bool ShouldOverridePinNames() const OVERRIDE { return true; }
-	ENGINE_API virtual FString GetPinNameOverride(const UEdGraphPin& Pin) const OVERRIDE;
-	ENGINE_API virtual void ResizeNode(const FVector2D& NewSize) OVERRIDE;
-	ENGINE_API virtual void PostPlacedNewNode() OVERRIDE;
-	ENGINE_API virtual void OnRenameNode(const FString& NewName) OVERRIDE;
-	ENGINE_API virtual TSharedPtr<class INameValidatorInterface> MakeNameValidator() const OVERRIDE;
-	virtual bool ShouldDrawNodeAsComment() const OVERRIDE { return true; }
-	ENGINE_API virtual FString GetDocumentationLink() const OVERRIDE;
-	ENGINE_API virtual FString GetDocumentationExcerptName() const OVERRIDE;
+	virtual void AllocateDefaultPins() override {}
+	ENGINE_API virtual FString GetTooltip() const override;
+	ENGINE_API virtual FLinearColor GetNodeCommentColor() const override;
+	ENGINE_API virtual FText GetNodeTitle(ENodeTitleType::Type TitleType) const override;
+	virtual bool ShouldOverridePinNames() const override { return true; }
+	ENGINE_API virtual FString GetPinNameOverride(const UEdGraphPin& Pin) const override;
+	ENGINE_API virtual void ResizeNode(const FVector2D& NewSize) override;
+	ENGINE_API virtual void PostPlacedNewNode() override;
+	ENGINE_API virtual void OnRenameNode(const FString& NewName) override;
+	ENGINE_API virtual TSharedPtr<class INameValidatorInterface> MakeNameValidator() const override;
+	virtual bool ShouldDrawNodeAsComment() const override { return true; }
+	ENGINE_API virtual FString GetDocumentationLink() const override;
+	ENGINE_API virtual FString GetDocumentationExcerptName() const override;
 	// End UEdGraphNode interface
 
 	/** Add a node that will be dragged when this comment is dragged */
diff --git a/Engine/Source/Runtime/Engine/Classes/EdGraph/EdGraphSchema.h b/Engine/Source/Runtime/Engine/Classes/EdGraph/EdGraphSchema.h
--- a/Engine/Source/Runtime/Engine/Classes/EdGraph/EdGraphSchema.h
+++ b/Engine/Source/Runtime/Engine/Classes/EdGraph/EdGraphSchema.h
@@ -147,7 +147,7 @@ struct ENGINE_API FEdGraphSchemaAction_NewNode : public FEdGraphSchemaAction
 
 	// Simple type info
 	static FString StaticGetTypeId() {static FString Type = TEXT("FEdGraphSchemaAction_NewNode"); return Type;}
-	virtual FString GetTypeId() const OVERRIDE { return StaticGetTypeId(); } 
+	virtual FString GetTypeId() const override { return StaticGetTypeId(); } 
 
 	/** Template of node we want to create */
 	UPROPERTY()
@@ -165,9 +165,9 @@ struct ENGINE_API FEdGraphSchemaAction_NewNode : public FEdGraphSchemaAction
 	{}
 
 	// FEdGraphSchemaAction interface
-	virtual UEdGraphNode* PerformAction(class UEdGraph* ParentGraph, UEdGraphPin* FromPin, const FVector2D Location, bool bSelectNewNode = true) OVERRIDE;
-	virtual UEdGraphNode* PerformAction(class UEdGraph* ParentGraph, TArray<UEdGraphPin*>& FromPins, const FVector2D Location, bool bSelectNewNode = true) OVERRIDE;
-	virtual void AddReferencedObjects(FReferenceCollector& Collector) OVERRIDE;
+	virtual UEdGraphNode* PerformAction(class UEdGraph* ParentGraph, UEdGraphPin* FromPin, const FVector2D Location, bool bSelectNewNode = true) override;
+	virtual UEdGraphNode* PerformAction(class UEdGraph* ParentGraph, TArray<UEdGraphPin*>& FromPins, const FVector2D Location, bool bSelectNewNode = true) override;
+	virtual void AddReferencedObjects(FReferenceCollector& Collector) override;
 	// End of FEdGraphSchemaAction interface
 
 	template <typename NodeType>
@@ -186,7 +186,7 @@ struct ENGINE_API FEdGraphSchemaAction_NewNode : public FEdGraphSchemaAction
 struct FEdGraphSchemaAction_Dummy : public FEdGraphSchemaAction
 {
 	static FString StaticGetTypeId() { static FString Type = TEXT("FEdGraphSchemaAction_Dummy"); return Type; }
-	virtual FString GetTypeId() const OVERRIDE{ return StaticGetTypeId(); }
+	virtual FString GetTypeId() const override{ return StaticGetTypeId(); }
 
 	FEdGraphSchemaAction_Dummy()
 	: FEdGraphSchemaAction()
@@ -334,8 +334,8 @@ public:
 	ENGINE_API FCategorizedGraphActionListBuilder(FString const& Category = TEXT(""));
 
 	// FGraphActionListBuilderBase Interface
-	ENGINE_API virtual void AddAction(const TSharedPtr<FEdGraphSchemaAction>& NewAction, FString const& Category = TEXT("") ) OVERRIDE;
-	ENGINE_API virtual void AddActionList(const TArray<TSharedPtr<FEdGraphSchemaAction> >& NewActions, FString const& Category = TEXT("")) OVERRIDE;
+	ENGINE_API virtual void AddAction(const TSharedPtr<FEdGraphSchemaAction>& NewAction, FString const& Category = TEXT("") ) override;
+	ENGINE_API virtual void AddActionList(const TArray<TSharedPtr<FEdGraphSchemaAction> >& NewActions, FString const& Category = TEXT("")) override;
 	// End of FGraphActionListBuilderBase Interface
 
 private:
diff --git a/Engine/Source/Runtime/Engine/Classes/Engine/ActorChannel.h b/Engine/Source/Runtime/Engine/Classes/Engine/ActorChannel.h
--- a/Engine/Source/Runtime/Engine/Classes/Engine/ActorChannel.h
+++ b/Engine/Source/Runtime/Engine/Classes/Engine/ActorChannel.h
@@ -81,17 +81,17 @@ class ENGINE_API UActorChannel : public UChannel
 
 	// Begin UChannel interface.
 
-	virtual void Init( UNetConnection* InConnection, int32 InChIndex, bool InOpenedLocally ) OVERRIDE;
+	virtual void Init( UNetConnection* InConnection, int32 InChIndex, bool InOpenedLocally ) override;
 
-	virtual void SetClosingFlag() OVERRIDE;
+	virtual void SetClosingFlag() override;
 
-	virtual void ReceivedBunch( FInBunch& Bunch ) OVERRIDE;
+	virtual void ReceivedBunch( FInBunch& Bunch ) override;
 	
-	virtual void ReceivedNak( int32 NakPacketId ) OVERRIDE;
+	virtual void ReceivedNak( int32 NakPacketId ) override;
 	
-	virtual void Close() OVERRIDE;
+	virtual void Close() override;
 
-	virtual FString Describe() OVERRIDE;
+	virtual FString Describe() override;
 
 	// End UChannel interface.
 
@@ -109,7 +109,7 @@ class ENGINE_API UActorChannel : public UChannel
 
 	void SetChannelActorForDestroy( struct FActorDestructionInfo *DestructInfo );
 
-	virtual void Serialize(FArchive& Ar) OVERRIDE;
+	virtual void Serialize(FArchive& Ar) override;
 
 	static void AddReferencedObjects(UObject* InThis, FReferenceCollector& Collector);
 
@@ -120,7 +120,7 @@ class ENGINE_API UActorChannel : public UChannel
 	virtual bool ReadyForDormancy(bool debug=false);
 	
 	/** Puts the channel in a state to start becoming dormant. It will not become dormant until ReadyForDormancy returns true in Tick */
-	virtual void StartBecomingDormant() OVERRIDE;
+	virtual void StartBecomingDormant() override;
 
 	/** Cleans up replicators and clears references to the actor class this channel was associated with.*/
 	void CleanupReplicators( const bool bKeepReplicators = false );
@@ -234,7 +234,7 @@ protected:
 	TSharedRef< FObjectReplicator > & FindOrCreateReplicator(UObject *Obj);
 	bool ObjectHasReplicator(UObject *Obj);	// returns whether we have already created a replicator for this object or not
 
-	virtual void CleanUp() OVERRIDE;
+	virtual void CleanUp() override;
 
 	/** Closes the actor channel but with a 'dormant' flag set so it can be reopened */
 	virtual void BecomeDormant();
diff --git a/Engine/Source/Runtime/Engine/Classes/Engine/BlockingVolume.h b/Engine/Source/Runtime/Engine/Classes/Engine/BlockingVolume.h
--- a/Engine/Source/Runtime/Engine/Classes/Engine/BlockingVolume.h
+++ b/Engine/Source/Runtime/Engine/Classes/Engine/BlockingVolume.h
@@ -14,12 +14,12 @@ class ABlockingVolume : public AVolume
 {
 	GENERATED_UCLASS_BODY()
 
-	virtual bool UpdateNavigationRelevancy() OVERRIDE;
+	virtual bool UpdateNavigationRelevancy() override;
 
 	// Begin UObject interface.
 #if WITH_EDITOR
-	virtual void LoadedFromAnotherClass(const FName& OldClassName) OVERRIDE;
-	virtual void PostEditChangeChainProperty(FPropertyChangedChainEvent& PropertyChangedEvent) OVERRIDE;
+	virtual void LoadedFromAnotherClass(const FName& OldClassName) override;
+	virtual void PostEditChangeChainProperty(FPropertyChangedChainEvent& PropertyChangedEvent) override;
 #endif // WITH_EDITOR	
 	// End UObject interface.
 };
diff --git a/Engine/Source/Runtime/Engine/Classes/Engine/Blueprint.h b/Engine/Source/Runtime/Engine/Classes/Engine/Blueprint.h
--- a/Engine/Source/Runtime/Engine/Classes/Engine/Blueprint.h
+++ b/Engine/Source/Runtime/Engine/Classes/Engine/Blueprint.h
@@ -459,7 +459,7 @@ public:
 
 protected:
 	/** Gets asset registry tags */
-	virtual void GetAssetRegistryTags(TArray<FAssetRegistryTag>& OutTags) const OVERRIDE;
+	virtual void GetAssetRegistryTags(TArray<FAssetRegistryTag>& OutTags) const override;
 
 private:
 	/** Sets the current object being debugged */
@@ -474,11 +474,11 @@ public:
 
 
 	// Begin UObject interface (WITH_EDITOR)
-	virtual void PostDuplicate(bool bDuplicateForPIE) OVERRIDE;
-	virtual bool Rename(const TCHAR* NewName = NULL, UObject* NewOuter = NULL, ERenameFlags Flags = REN_None) OVERRIDE;
-	virtual UClass* RegenerateClass(UClass* ClassToRegenerate, UObject* PreviousCDO, TArray<UObject*>& ObjLoaded) OVERRIDE;
-	virtual void PostLoad() OVERRIDE;
-	virtual void PostLoadSubobjects( FObjectInstancingGraph* OuterInstanceGraph ) OVERRIDE;
+	virtual void PostDuplicate(bool bDuplicateForPIE) override;
+	virtual bool Rename(const TCHAR* NewName = NULL, UObject* NewOuter = NULL, ERenameFlags Flags = REN_None) override;
+	virtual UClass* RegenerateClass(UClass* ClassToRegenerate, UObject* PreviousCDO, TArray<UObject*>& ObjLoaded) override;
+	virtual void PostLoad() override;
+	virtual void PostLoadSubobjects( FObjectInstancingGraph* OuterInstanceGraph ) override;
 	// End of UObject interface
 
 	/** Consigns the GeneratedClass and the SkeletonGeneratedClass to oblivion, and nulls their references */
@@ -492,12 +492,12 @@ public:
 #endif	//#if WITH_EDITOR
 
 	// Begin UObject interface
-	virtual void Serialize(FArchive& Ar) OVERRIDE;
-	virtual FString GetDesc(void) OVERRIDE;
-	virtual void TagSubobjects(EObjectFlags NewFlags) OVERRIDE;
-	virtual bool NeedsLoadForClient() const OVERRIDE;
-	virtual bool NeedsLoadForServer() const OVERRIDE;
-	virtual bool NeedsLoadForEditorGame() const OVERRIDE;
+	virtual void Serialize(FArchive& Ar) override;
+	virtual FString GetDesc(void) override;
+	virtual void TagSubobjects(EObjectFlags NewFlags) override;
+	virtual bool NeedsLoadForClient() const override;
+	virtual bool NeedsLoadForServer() const override;
+	virtual bool NeedsLoadForEditorGame() const override;
 	// End of UObject interface
 
 	/** Get the Blueprint object that generated the supplied class */
diff --git a/Engine/Source/Runtime/Engine/Classes/Engine/BlueprintCore.h b/Engine/Source/Runtime/Engine/Classes/Engine/BlueprintCore.h
--- a/Engine/Source/Runtime/Engine/Classes/Engine/BlueprintCore.h
+++ b/Engine/Source/Runtime/Engine/Classes/Engine/BlueprintCore.h
@@ -42,10 +42,10 @@ public:
 		return bLegacyGeneratedClassIsAuthoritative;
 	}
 
-	virtual void Serialize( FArchive& Ar ) OVERRIDE;
+	virtual void Serialize( FArchive& Ar ) override;
 
 	/** Gets asset registry tags */
-	virtual void GetAssetRegistryTags(TArray<FAssetRegistryTag>& OutTags) const OVERRIDE;
+	virtual void GetAssetRegistryTags(TArray<FAssetRegistryTag>& OutTags) const override;
 
 	/** Generates a new blueprint Guid, used when creating new blueprints */
 	void GenerateNewGuid()
diff --git a/Engine/Source/Runtime/Engine/Classes/Engine/BlueprintGeneratedClass.h b/Engine/Source/Runtime/Engine/Classes/Engine/BlueprintGeneratedClass.h
--- a/Engine/Source/Runtime/Engine/Classes/Engine/BlueprintGeneratedClass.h
+++ b/Engine/Source/Runtime/Engine/Classes/Engine/BlueprintGeneratedClass.h
@@ -421,16 +421,16 @@ public:
 	virtual void CreateComponentsForActor(AActor* Actor) const;
 
 	// UObject interface
-	virtual void PostLoad() OVERRIDE;
-	virtual void PostInitProperties() OVERRIDE;
+	virtual void PostLoad() override;
+	virtual void PostInitProperties() override;
 	// End UObject interface
 	
 	// UClass interface
 #if WITH_EDITOR
-	virtual UClass* GetAuthoritativeClass() OVERRIDE;
-	virtual void ConditionalRecompileClass(TArray<UObject*>* ObjLoaded) OVERRIDE;
+	virtual UClass* GetAuthoritativeClass() override;
+	virtual void ConditionalRecompileClass(TArray<UObject*>* ObjLoaded) override;
 #endif //WITH_EDITOR
-	virtual bool IsFunctionImplementedInBlueprint(FName InFunctionName) const OVERRIDE;
+	virtual bool IsFunctionImplementedInBlueprint(FName InFunctionName) const override;
 	// End UClass interface
 
 #if WITH_EDITORONLY_DATA
diff --git a/Engine/Source/Runtime/Engine/Classes/Engine/Brush.h b/Engine/Source/Runtime/Engine/Classes/Engine/Brush.h
--- a/Engine/Source/Runtime/Engine/Classes/Engine/Brush.h
+++ b/Engine/Source/Runtime/Engine/Classes/Engine/Brush.h
@@ -139,18 +139,18 @@ class ABrush : public AActor
 
 	
 	// Begin UObject interface.
-	ENGINE_API virtual void PostLoad() OVERRIDE;
+	ENGINE_API virtual void PostLoad() override;
 #if WITH_EDITOR
-	ENGINE_API virtual void PostEditMove(bool bFinished) OVERRIDE;
-	ENGINE_API virtual void PostEditChangeProperty(FPropertyChangedEvent& PropertyChangedEvent) OVERRIDE;
+	ENGINE_API virtual void PostEditMove(bool bFinished) override;
+	ENGINE_API virtual void PostEditChangeProperty(FPropertyChangedEvent& PropertyChangedEvent) override;
 #endif // WITH_EDITOR
-	ENGINE_API virtual bool Modify(bool bAlwaysMarkDirty = false) OVERRIDE;
+	ENGINE_API virtual bool Modify(bool bAlwaysMarkDirty = false) override;
 
-	virtual bool NeedsLoadForClient() const OVERRIDE
+	virtual bool NeedsLoadForClient() const override
 	{ 
 		return !IsNotForClientOrServer(); 
 	}
-	virtual bool NeedsLoadForServer() const OVERRIDE
+	virtual bool NeedsLoadForServer() const override
 	{ 
 		return !IsNotForClientOrServer(); 
 	}
@@ -158,16 +158,16 @@ class ABrush : public AActor
 
 	
 	// Begin AActor interface
-	ENGINE_API virtual void Destroyed() OVERRIDE;
-	ENGINE_API virtual void PostRegisterAllComponents() OVERRIDE;
+	ENGINE_API virtual void Destroyed() override;
+	ENGINE_API virtual void PostRegisterAllComponents() override;
 	
-	virtual bool UpdateNavigationRelevancy() OVERRIDE { SetNavigationRelevancy(false); return false; }
-	virtual bool IsLevelBoundsRelevant() const OVERRIDE;
+	virtual bool UpdateNavigationRelevancy() override { SetNavigationRelevancy(false); return false; }
+	virtual bool IsLevelBoundsRelevant() const override;
 
 #if WITH_EDITOR
-	virtual void CheckForErrors() OVERRIDE;
+	virtual void CheckForErrors() override;
 
-	ENGINE_API virtual void SetIsTemporarilyHiddenInEditor( bool bIsHidden ) OVERRIDE;
+	ENGINE_API virtual void SetIsTemporarilyHiddenInEditor( bool bIsHidden ) override;
 
 	// End AActor interface
 
diff --git a/Engine/Source/Runtime/Engine/Classes/Engine/BrushShape.h b/Engine/Source/Runtime/Engine/Classes/Engine/BrushShape.h
--- a/Engine/Source/Runtime/Engine/Classes/Engine/BrushShape.h
+++ b/Engine/Source/Runtime/Engine/Classes/Engine/BrushShape.h
@@ -13,9 +13,9 @@ class ABrushShape : public ABrush
 	GENERATED_UCLASS_BODY()
 
 
-	virtual bool IsStaticBrush() const OVERRIDE {return false; }
+	virtual bool IsStaticBrush() const override {return false; }
 
-	virtual bool IsBrushShape() const OVERRIDE {return true; }
+	virtual bool IsBrushShape() const override {return true; }
 };
 
 
