UE4 Refactoring. Changed OVERRIDE and FINAL macros to keywords override and final.

diff --git a/Engine/Source/Runtime/Engine/Classes/Engine/ChildConnection.h b/Engine/Source/Runtime/Engine/Classes/Engine/ChildConnection.h
--- a/Engine/Source/Runtime/Engine/Classes/Engine/ChildConnection.h
+++ b/Engine/Source/Runtime/Engine/Classes/Engine/ChildConnection.h
@@ -14,57 +14,57 @@ class UChildConnection : public UNetConnection
 	class UNetConnection* Parent;
 
 	// Begin UNetConnection interface.
-	virtual UChildConnection* GetUChildConnection() OVERRIDE
+	virtual UChildConnection* GetUChildConnection() override
 	{
 		return this;
 	}
 
-	virtual FString LowLevelGetRemoteAddress(bool bAppendPort=false) OVERRIDE
+	virtual FString LowLevelGetRemoteAddress(bool bAppendPort=false) override
 	{
 		return Parent->LowLevelGetRemoteAddress(bAppendPort);
 	}
 
-	virtual FString LowLevelDescribe() OVERRIDE
+	virtual FString LowLevelDescribe() override
 	{
 		return Parent->LowLevelDescribe();
 	}
 
-		virtual void LowLevelSend( void* Data, int32 Count ) OVERRIDE
+		virtual void LowLevelSend( void* Data, int32 Count ) override
 	{
 	}
 
-	virtual void InitSendBuffer() OVERRIDE
+	virtual void InitSendBuffer() override
 	{
 		Parent->InitSendBuffer();
 	}
 
-	virtual void AssertValid() OVERRIDE
+	virtual void AssertValid() override
 	{
 		Parent->AssertValid();
 	}
 
-	virtual void SendAck( int32 PacketId, bool FirstTime=1, bool bHavePingAckData=0, uint32 PingAckData=0 ) OVERRIDE
+	virtual void SendAck( int32 PacketId, bool FirstTime=1, bool bHavePingAckData=0, uint32 PingAckData=0 ) override
 	{
 	}
 
-	virtual void FlushNet(bool bIgnoreSimulation = false) OVERRIDE
+	virtual void FlushNet(bool bIgnoreSimulation = false) override
 	{
 		Parent->FlushNet(bIgnoreSimulation);
 	}
 
-	virtual int32 IsNetReady(bool Saturate) OVERRIDE
+	virtual int32 IsNetReady(bool Saturate) override
 	{
 		return Parent->IsNetReady(Saturate);
 	}
 
-	virtual void Tick() OVERRIDE
+	virtual void Tick() override
 	{
 		State = Parent->State;
 	}
 
-	virtual void HandleClientPlayer(class APlayerController* PC, class UNetConnection* NetConnection) OVERRIDE;
+	virtual void HandleClientPlayer(class APlayerController* PC, class UNetConnection* NetConnection) override;
 	
-	virtual void CleanUp() OVERRIDE;
+	virtual void CleanUp() override;
 
 	// End UNetConnection interface.
 
diff --git a/Engine/Source/Runtime/Engine/Classes/Engine/CodecMovieFallback.h b/Engine/Source/Runtime/Engine/Classes/Engine/CodecMovieFallback.h
--- a/Engine/Source/Runtime/Engine/Classes/Engine/CodecMovieFallback.h
+++ b/Engine/Source/Runtime/Engine/Classes/Engine/CodecMovieFallback.h
@@ -21,23 +21,23 @@ class UCodecMovieFallback : public UCodecMovie
 	* Not all codec implementations are available
 	* @return true if the current codec is supported
 	*/
-	virtual bool IsSupported() OVERRIDE;
+	virtual bool IsSupported() override;
 
-	virtual uint32 GetSizeX() OVERRIDE;
+	virtual uint32 GetSizeX() override;
 
-	virtual uint32 GetSizeY() OVERRIDE;
+	virtual uint32 GetSizeY() override;
 
-	virtual EPixelFormat GetFormat() OVERRIDE;
+	virtual EPixelFormat GetFormat() override;
 
-	virtual float GetFrameRate() OVERRIDE;	
+	virtual float GetFrameRate() override;	
 
-	virtual bool Open( const FString& Filename, uint32 Offset, uint32 Size ) OVERRIDE;
+	virtual bool Open( const FString& Filename, uint32 Offset, uint32 Size ) override;
 
-	virtual bool Open( void* Source, uint32 Size ) OVERRIDE;
+	virtual bool Open( void* Source, uint32 Size ) override;
 
-	virtual void ResetStream() OVERRIDE;
+	virtual void ResetStream() override;
 
-	virtual void GetFrame( class FTextureMovieResource* InTextureMovieResource ) OVERRIDE;
+	virtual void GetFrame( class FTextureMovieResource* InTextureMovieResource ) override;
 
 	// End CodecMovie Interface
 };
diff --git a/Engine/Source/Runtime/Engine/Classes/Engine/ComponentDelegateBinding.h b/Engine/Source/Runtime/Engine/Classes/Engine/ComponentDelegateBinding.h
--- a/Engine/Source/Runtime/Engine/Classes/Engine/ComponentDelegateBinding.h
+++ b/Engine/Source/Runtime/Engine/Classes/Engine/ComponentDelegateBinding.h
@@ -38,6 +38,6 @@ class ENGINE_API UComponentDelegateBinding : public UDynamicBlueprintBinding
 	TArray<FBlueprintComponentDelegateBinding> ComponentDelegateBindings;
 
 	// Begin DynamicBlueprintBinding interface
-	virtual void BindDynamicDelegates(AActor* InInstance) const OVERRIDE;
+	virtual void BindDynamicDelegates(AActor* InInstance) const override;
 	// End DynamicBlueprintBinding interface
 };
diff --git a/Engine/Source/Runtime/Engine/Classes/Engine/ControlChannel.h b/Engine/Source/Runtime/Engine/Classes/Engine/ControlChannel.h
--- a/Engine/Source/Runtime/Engine/Classes/Engine/ControlChannel.h
+++ b/Engine/Source/Runtime/Engine/Classes/Engine/ControlChannel.h
@@ -57,9 +57,9 @@ class UControlChannel : public UChannel
 	virtual void Init( UNetConnection* InConnection, int32 InChIndex, bool InOpenedLocally );
 
 	// Begin UChannel interface.
-	virtual FPacketIdRange SendBunch(FOutBunch* Bunch, bool Merge) OVERRIDE;
+	virtual FPacketIdRange SendBunch(FOutBunch* Bunch, bool Merge) override;
 
-	virtual void Tick() OVERRIDE;
+	virtual void Tick() override;
 	// End UChannel interface.
 
 
diff --git a/Engine/Source/Runtime/Engine/Classes/Engine/CullDistanceVolume.h b/Engine/Source/Runtime/Engine/Classes/Engine/CullDistanceVolume.h
--- a/Engine/Source/Runtime/Engine/Classes/Engine/CullDistanceVolume.h
+++ b/Engine/Source/Runtime/Engine/Classes/Engine/CullDistanceVolume.h
@@ -59,8 +59,8 @@ class ACullDistanceVolume : public AVolume
 
 	// Begin UObject Interface
 #if WITH_EDITOR
-	virtual void PostEditChangeProperty(FPropertyChangedEvent& PropertyChangedEvent) OVERRIDE;
-	virtual void PostEditMove(bool bFinished) OVERRIDE;
+	virtual void PostEditChangeProperty(FPropertyChangedEvent& PropertyChangedEvent) override;
+	virtual void PostEditMove(bool bFinished) override;
 #endif // WITH_EDITOR
 	// End UObject Interface
 
diff --git a/Engine/Source/Runtime/Engine/Classes/Engine/CurveTable.h b/Engine/Source/Runtime/Engine/Classes/Engine/CurveTable.h
--- a/Engine/Source/Runtime/Engine/Classes/Engine/CurveTable.h
+++ b/Engine/Source/Runtime/Engine/Classes/Engine/CurveTable.h
@@ -26,8 +26,8 @@ class UCurveTable : public UObject
 	TMap<FName, FRichCurve*>	RowMap;
 
 	// Begin UObject interface.
-	virtual void FinishDestroy() OVERRIDE;
-	virtual void Serialize( FArchive& Ar ) OVERRIDE;
+	virtual void FinishDestroy() override;
+	virtual void Serialize( FArchive& Ar ) override;
 	// End  UObject interface
 
 	// Begin UCurveTable interface
diff --git a/Engine/Source/Runtime/Engine/Classes/Engine/DataTable.h b/Engine/Source/Runtime/Engine/Classes/Engine/DataTable.h
--- a/Engine/Source/Runtime/Engine/Classes/Engine/DataTable.h
+++ b/Engine/Source/Runtime/Engine/Classes/Engine/DataTable.h
@@ -41,8 +41,8 @@ class UDataTable : public UObject
 	TMap<FName, uint8*>		RowMap;
 
 	// Begin UObject interface.
-	virtual void FinishDestroy() OVERRIDE;
-	virtual void Serialize( FArchive& Ar ) OVERRIDE;
+	virtual void FinishDestroy() override;
+	virtual void Serialize( FArchive& Ar ) override;
 	static void AddReferencedObjects(UObject* InThis, FReferenceCollector& Collector);
 	// End  UObject interface
 
diff --git a/Engine/Source/Runtime/Engine/Classes/Engine/DebugCameraController.h b/Engine/Source/Runtime/Engine/Classes/Engine/DebugCameraController.h
--- a/Engine/Source/Runtime/Engine/Classes/Engine/DebugCameraController.h
+++ b/Engine/Source/Runtime/Engine/Classes/Engine/DebugCameraController.h
@@ -91,7 +91,7 @@ protected:
 
 	// Adjusts movement speed limits based on SpeedScale.
 	virtual void ApplySpeedScale();
-	virtual void SetupInputComponent() OVERRIDE;
+	virtual void SetupInputComponent() override;
 
 public:
 	/** Function called on activation debug camera controller */
@@ -111,10 +111,10 @@ public:
 
 
 	// Begin APlayerController Interface
-	virtual void PostInitializeComponents() OVERRIDE;
-	virtual FString ConsoleCommand(const FString& Command, bool bWriteToLog = true) OVERRIDE;
-	virtual void AddCheats(bool bForce) OVERRIDE;
-	virtual void EndSpectatingState() OVERRIDE;
+	virtual void PostInitializeComponents() override;
+	virtual FString ConsoleCommand(const FString& Command, bool bWriteToLog = true) override;
+	virtual void AddCheats(bool bForce) override;
+	virtual void EndSpectatingState() override;
 	// End APlayerController Interface
 
 protected:
@@ -124,7 +124,7 @@ protected:
 	 */
 	virtual void Select( FHitResult const& Hit );
 
-	virtual void SetSpectatorPawn(class ASpectatorPawn* NewSpectatorPawn) OVERRIDE;
+	virtual void SetSpectatorPawn(class ASpectatorPawn* NewSpectatorPawn) override;
 };
 
 
diff --git a/Engine/Source/Runtime/Engine/Classes/Engine/DebugCameraHUD.h b/Engine/Source/Runtime/Engine/Classes/Engine/DebugCameraHUD.h
--- a/Engine/Source/Runtime/Engine/Classes/Engine/DebugCameraHUD.h
+++ b/Engine/Source/Runtime/Engine/Classes/Engine/DebugCameraHUD.h
@@ -13,7 +13,7 @@ class ENGINE_API ADebugCameraHUD : public AHUD
 	virtual bool DisplayMaterials( float X, float& Y, float DY, class UMeshComponent* MeshComp );
 	
 	// Begin AActor Interface
-	virtual void PostRender() OVERRIDE;
+	virtual void PostRender() override;
 	// End AActor Interface
 
 };
diff --git a/Engine/Source/Runtime/Engine/Classes/Engine/DecalActor.h b/Engine/Source/Runtime/Engine/Classes/Engine/DecalActor.h
--- a/Engine/Source/Runtime/Engine/Classes/Engine/DecalActor.h
+++ b/Engine/Source/Runtime/Engine/Classes/Engine/DecalActor.h
@@ -48,12 +48,12 @@ class ADecalActor : public AActor
 	
 #if WITH_EDITOR
 	// Begin UObject Interface
-	virtual void PostEditMove(bool bFinished) OVERRIDE;
-	virtual void PostEditChangeProperty(FPropertyChangedEvent& PropertyChangedEvent) OVERRIDE;
+	virtual void PostEditMove(bool bFinished) override;
+	virtual void PostEditChangeProperty(FPropertyChangedEvent& PropertyChangedEvent) override;
 	// End UObject Interface
 
 	// Begin AActor interface.
-	virtual void EditorApplyScale(const FVector& DeltaScale, const FVector* PivotLocation, bool bAltDown, bool bShiftDown, bool bCtrlDown) OVERRIDE;
+	virtual void EditorApplyScale(const FVector& DeltaScale, const FVector* PivotLocation, bool bAltDown, bool bShiftDown, bool bCtrlDown) override;
 	// End AActor interface.
 #endif // WITH_EDITOR
 };
diff --git a/Engine/Source/Runtime/Engine/Classes/Engine/DestructibleFractureSettings.h b/Engine/Source/Runtime/Engine/Classes/Engine/DestructibleFractureSettings.h
--- a/Engine/Source/Runtime/Engine/Classes/Engine/DestructibleFractureSettings.h
+++ b/Engine/Source/Runtime/Engine/Classes/Engine/DestructibleFractureSettings.h
@@ -192,9 +192,9 @@ class UDestructibleFractureSettings : public UObject
 #endif // WITH_APEX
 
 	// Begin UObject interface.
-	virtual void PostInitProperties() OVERRIDE;
-	virtual void BeginDestroy() OVERRIDE;
-	virtual void Serialize( FArchive& Ar ) OVERRIDE;
+	virtual void PostInitProperties() override;
+	virtual void BeginDestroy() override;
+	virtual void Serialize( FArchive& Ar ) override;
 	// End  UObject interface
 
 	// Begin UDestructibleFractureSettings interface
diff --git a/Engine/Source/Runtime/Engine/Classes/Engine/DestructibleMesh.h b/Engine/Source/Runtime/Engine/Classes/Engine/DestructibleMesh.h
--- a/Engine/Source/Runtime/Engine/Classes/Engine/DestructibleMesh.h
+++ b/Engine/Source/Runtime/Engine/Classes/Engine/DestructibleMesh.h
@@ -451,11 +451,11 @@ public:
 #endif // WITH_APEX
 
 	// Begin UObject interface.
-	virtual void				PostLoad() OVERRIDE;
-	virtual void				Serialize(FArchive& Ar) OVERRIDE;
-	virtual void				FinishDestroy() OVERRIDE;
+	virtual void				PostLoad() override;
+	virtual void				Serialize(FArchive& Ar) override;
+	virtual void				FinishDestroy() override;
 #if WITH_EDITOR
-	virtual void				PreEditChange(UProperty* PropertyAboutToChange) OVERRIDE;
+	virtual void				PreEditChange(UProperty* PropertyAboutToChange) override;
 #endif
 	// End UObject interface.
 
diff --git a/Engine/Source/Runtime/Engine/Classes/Engine/DirectionalLight.h b/Engine/Source/Runtime/Engine/Classes/Engine/DirectionalLight.h
--- a/Engine/Source/Runtime/Engine/Classes/Engine/DirectionalLight.h
+++ b/Engine/Source/Runtime/Engine/Classes/Engine/DirectionalLight.h
@@ -18,10 +18,10 @@ class ADirectionalLight : public ALight
 public:
 
 	// Begin UObject Interface
-	virtual void PostLoad() OVERRIDE;
+	virtual void PostLoad() override;
 #if WITH_EDITOR
-	virtual void LoadedFromAnotherClass(const FName& OldClassName) OVERRIDE;
-	virtual void PostEditChangeProperty(FPropertyChangedEvent& PropertyChangedEvent) OVERRIDE;
+	virtual void LoadedFromAnotherClass(const FName& OldClassName) override;
+	virtual void PostEditChangeProperty(FPropertyChangedEvent& PropertyChangedEvent) override;
 #endif
 	// End UObject Interface
 };
diff --git a/Engine/Source/Runtime/Engine/Classes/Engine/Engine.h b/Engine/Source/Runtime/Engine/Classes/Engine/Engine.h
--- a/Engine/Source/Runtime/Engine/Classes/Engine/Engine.h
+++ b/Engine/Source/Runtime/Engine/Classes/Engine/Engine.h
@@ -1579,8 +1579,8 @@ public:
 	}
 
 	// Begin UObject interface.
-	virtual void FinishDestroy() OVERRIDE;
-	virtual void Serialize(FArchive& Ar) OVERRIDE;
+	virtual void FinishDestroy() override;
+	virtual void Serialize(FArchive& Ar) override;
 	static void AddReferencedObjects(UObject* InThis, FReferenceCollector& Collector);
 	// End UObject interface.
 
@@ -1595,7 +1595,7 @@ public:
 	void ParseCommandline();
 
 	// Begin FExec Interface
-	virtual bool Exec( UWorld* InWorld, const TCHAR* Cmd, FOutputDevice& Out=*GLog ) OVERRIDE;
+	virtual bool Exec( UWorld* InWorld, const TCHAR* Cmd, FOutputDevice& Out=*GLog ) override;
 	// End FExec Interface
 
 	/** 
diff --git a/Engine/Source/Runtime/Engine/Classes/Engine/EngineBaseTypes.h b/Engine/Source/Runtime/Engine/Classes/Engine/EngineBaseTypes.h
--- a/Engine/Source/Runtime/Engine/Classes/Engine/EngineBaseTypes.h
+++ b/Engine/Source/Runtime/Engine/Classes/Engine/EngineBaseTypes.h
@@ -338,7 +338,7 @@ struct FActorTickFunction : public FTickFunction
 		* @param CurrentThread - thread we are executing on, useful to pass along as new tasks are created
 		* @param MyCompletionGraphEvent - completion event for this task. Useful for holding the completetion of this task until certain child tasks are complete.
 	**/
-	ENGINE_API virtual void ExecuteTick(float DeltaTime, ELevelTick TickType, ENamedThreads::Type CurrentThread, const FGraphEventRef& MyCompletionGraphEvent) OVERRIDE;
+	ENGINE_API virtual void ExecuteTick(float DeltaTime, ELevelTick TickType, ENamedThreads::Type CurrentThread, const FGraphEventRef& MyCompletionGraphEvent) override;
 	/** Abstract function to describe this tick. Used to print messages about illegal cycles in the dependency graph **/
 	ENGINE_API virtual FString DiagnosticMessage();
 };
@@ -361,7 +361,7 @@ struct FActorComponentTickFunction : public FTickFunction
 		* @param CurrentThread - thread we are executing on, useful to pass along as new tasks are created
 		* @param MyCompletionGraphEvent - completion event for this task. Useful for holding the completetion of this task until certain child tasks are complete.
 	**/
-	ENGINE_API virtual void ExecuteTick(float DeltaTime, ELevelTick TickType, ENamedThreads::Type CurrentThread, const FGraphEventRef& MyCompletionGraphEvent) OVERRIDE;
+	ENGINE_API virtual void ExecuteTick(float DeltaTime, ELevelTick TickType, ENamedThreads::Type CurrentThread, const FGraphEventRef& MyCompletionGraphEvent) override;
 	/** Abstract function to describe this tick. Used to print messages about illegal cycles in the dependency graph **/
 	ENGINE_API virtual FString DiagnosticMessage();
 };
@@ -384,7 +384,7 @@ struct FPrimitiveComponentPostPhysicsTickFunction : public FTickFunction
 		* @param CurrentThread - thread we are executing on, useful to pass along as new tasks are created
 		* @param MyCompletionGraphEvent - completion event for this task. Useful for holding the completetion of this task until certain child tasks are complete.
 	**/
-	virtual void ExecuteTick(float DeltaTime, ELevelTick TickType, ENamedThreads::Type CurrentThread, const FGraphEventRef& MyCompletionGraphEvent) OVERRIDE;
+	virtual void ExecuteTick(float DeltaTime, ELevelTick TickType, ENamedThreads::Type CurrentThread, const FGraphEventRef& MyCompletionGraphEvent) override;
 	/** Abstract function to describe this tick. Used to print messages about illegal cycles in the dependency graph **/
 	virtual FString DiagnosticMessage();
 };
diff --git a/Engine/Source/Runtime/Engine/Classes/Engine/EngineTypes.h b/Engine/Source/Runtime/Engine/Classes/Engine/EngineTypes.h
--- a/Engine/Source/Runtime/Engine/Classes/Engine/EngineTypes.h
+++ b/Engine/Source/Runtime/Engine/Classes/Engine/EngineTypes.h
@@ -1843,7 +1843,7 @@ struct ENGINE_API FPointDamageEvent : public FDamageEvent
 	virtual bool IsOfType(int32 InID) const { return (FPointDamageEvent::ClassID == InID) || FDamageEvent::IsOfType(InID); };
 
 	/** Simple API for common cases where we are happy to assume a single hit is expected, even though damage event may have multiple hits. */
-	virtual void GetBestHitInfo(AActor const* HitActor, AActor const* HitInstigator, struct FHitResult& OutHitInfo, FVector& OutImpulseDir) const OVERRIDE;
+	virtual void GetBestHitInfo(AActor const* HitActor, AActor const* HitInstigator, struct FHitResult& OutHitInfo, FVector& OutImpulseDir) const override;
 };
 
 
@@ -1912,7 +1912,7 @@ struct ENGINE_API FRadialDamageEvent : public FDamageEvent
 	virtual bool IsOfType(int32 InID) const { return (FRadialDamageEvent::ClassID == InID) || FDamageEvent::IsOfType(InID); };
 
 	/** Simple API for common cases where we are happy to assume a single hit is expected, even though damage event may have multiple hits. */
-	virtual void GetBestHitInfo(AActor const* HitActor, AActor const* HitInstigator, struct FHitResult& OutHitInfo, FVector& OutImpulseDir) const OVERRIDE;
+	virtual void GetBestHitInfo(AActor const* HitActor, AActor const* HitInstigator, struct FHitResult& OutHitInfo, FVector& OutImpulseDir) const override;
 };
 
 
diff --git a/Engine/Source/Runtime/Engine/Classes/Engine/ExponentialHeightFog.h b/Engine/Source/Runtime/Engine/Classes/Engine/ExponentialHeightFog.h
--- a/Engine/Source/Runtime/Engine/Classes/Engine/ExponentialHeightFog.h
+++ b/Engine/Source/Runtime/Engine/Classes/Engine/ExponentialHeightFog.h
@@ -23,7 +23,7 @@ class AExponentialHeightFog : public AInfo
 
 
 	//Begin AActor Interface
-	virtual void PostInitializeComponents() OVERRIDE;
+	virtual void PostInitializeComponents() override;
 	//End AActor Interface
 };
 
diff --git a/Engine/Source/Runtime/Engine/Classes/Engine/Font.h b/Engine/Source/Runtime/Engine/Classes/Engine/Font.h
--- a/Engine/Source/Runtime/Engine/Classes/Engine/Font.h
+++ b/Engine/Source/Runtime/Engine/Classes/Engine/Font.h
@@ -146,7 +146,7 @@ public:
 	 *
 	 * @return		Size of resource as to be displayed to artists/ LDs in the Editor.
 	 */
-	virtual SIZE_T GetResourceSize(EResourceSizeMode::Type Mode) OVERRIDE;
+	virtual SIZE_T GetResourceSize(EResourceSizeMode::Type Mode) override;
 
 	// UFont interface
 	FORCEINLINE TCHAR RemapChar(TCHAR CharCode) const
@@ -248,9 +248,9 @@ public:
 	}
 
 	// Begin UObject interface
-	virtual void Serialize( FArchive& Ar ) OVERRIDE;
-	virtual void PostLoad() OVERRIDE;
-	virtual bool IsLocalizedResource() OVERRIDE;
+	virtual void Serialize( FArchive& Ar ) override;
+	virtual void PostLoad() override;
+	virtual bool IsLocalizedResource() override;
 	// End UObject interface
 
 	/**
diff --git a/Engine/Source/Runtime/Engine/Classes/Engine/GameEngine.h b/Engine/Source/Runtime/Engine/Classes/Engine/GameEngine.h
--- a/Engine/Source/Runtime/Engine/Classes/Engine/GameEngine.h
+++ b/Engine/Source/Runtime/Engine/Classes/Engine/GameEngine.h
@@ -90,20 +90,20 @@ public:
 	virtual void RedrawViewports( bool bShouldPresent = true );
 
 	// Begin UObject interface.
-	virtual void FinishDestroy() OVERRIDE;
+	virtual void FinishDestroy() override;
 	// End UObject interface.
 
 	// Begin UEngine interface.
-	virtual void Init(class IEngineLoop* InEngineLoop) OVERRIDE;
-	virtual void PreExit() OVERRIDE;
-	virtual void Tick( float DeltaSeconds, bool bIdleMode ) OVERRIDE;
-	virtual float GetMaxTickRate( float DeltaTime, bool bAllowFrameRateSmoothing = true ) OVERRIDE;
-	virtual void ProcessToggleFreezeCommand( UWorld* InWorld ) OVERRIDE;
-	virtual void ProcessToggleFreezeStreamingCommand( UWorld* InWorld ) OVERRIDE;
+	virtual void Init(class IEngineLoop* InEngineLoop) override;
+	virtual void PreExit() override;
+	virtual void Tick( float DeltaSeconds, bool bIdleMode ) override;
+	virtual float GetMaxTickRate( float DeltaTime, bool bAllowFrameRateSmoothing = true ) override;
+	virtual void ProcessToggleFreezeCommand( UWorld* InWorld ) override;
+	virtual void ProcessToggleFreezeStreamingCommand( UWorld* InWorld ) override;
 	// End UEngine interface.
 
 	// Begin FExec Interface
-	virtual bool Exec( UWorld* InWorld, const TCHAR* Cmd, FOutputDevice& Ar=*GLog ) OVERRIDE;
+	virtual bool Exec( UWorld* InWorld, const TCHAR* Cmd, FOutputDevice& Ar=*GLog ) override;
 	// End FExec Interface
 
 	/** 
diff --git a/Engine/Source/Runtime/Engine/Classes/Engine/GameViewportClient.h b/Engine/Source/Runtime/Engine/Classes/Engine/GameViewportClient.h
--- a/Engine/Source/Runtime/Engine/Classes/Engine/GameViewportClient.h
+++ b/Engine/Source/Runtime/Engine/Classes/Engine/GameViewportClient.h
@@ -247,7 +247,7 @@ public:
 	virtual void SetConsoleTarget(int32 PlayerIndex);
 
 	/** Returns a relative world context for this viewport.	 */
-	virtual UWorld* GetWorld() const OVERRIDE;
+	virtual UWorld* GetWorld() const override;
 
 	virtual void SetReferenceToWorldContext(struct FWorldContext& WorldContext);
 
@@ -255,34 +255,34 @@ public:
 
 public:
 	// Begin UObject Interface
-	virtual void PostInitProperties() OVERRIDE;
-	virtual void BeginDestroy() OVERRIDE;
+	virtual void PostInitProperties() override;
+	virtual void BeginDestroy() override;
 	// End UObject Interface
 
 	// FViewportClient interface.
-	virtual void RedrawRequested(FViewport* InViewport) OVERRIDE {}
-	virtual bool InputKey(FViewport* Viewport, int32 ControllerId, FKey Key, EInputEvent EventType, float AmountDepressed=1.f, bool bGamepad=false) OVERRIDE;
-	virtual bool InputAxis(FViewport* Viewport, int32 ControllerId, FKey Key, float Delta, float DeltaTime, int32 NumSamples=1, bool bGamepad=false) OVERRIDE;
-	virtual bool InputChar(FViewport* Viewport,int32 ControllerId, TCHAR Character) OVERRIDE;
-	virtual bool InputTouch(FViewport* Viewport, int32 ControllerId, uint32 Handle, ETouchType::Type Type, const FVector2D& TouchLocation, FDateTime DeviceTimestamp, uint32 TouchpadIndex) OVERRIDE;
-	virtual bool InputMotion(FViewport* Viewport, int32 ControllerId, const FVector& Tilt, const FVector& RotationRate, const FVector& Gravity, const FVector& Acceleration) OVERRIDE;
-	virtual EMouseCursor::Type GetCursor(FViewport* Viewport, int32 X, int32 Y ) OVERRIDE;
-	virtual void Precache() OVERRIDE;
-	virtual void Draw(FViewport* Viewport,FCanvas* SceneCanvas) OVERRIDE;
-	virtual void ProcessScreenShots(FViewport* Viewport) OVERRIDE;
-	virtual void LostFocus(FViewport* Viewport) OVERRIDE;
-	virtual void ReceivedFocus(FViewport* Viewport) OVERRIDE;
-	virtual bool IsFocused(FViewport* Viewport) OVERRIDE;
-	virtual void CloseRequested(FViewport* Viewport) OVERRIDE;
-	virtual bool RequiresHitProxyStorage() OVERRIDE { return 0; }
-	virtual bool IsOrtho() const OVERRIDE;
-	virtual void MouseEnter(FViewport* Viewport, int32 x, int32 y) OVERRIDE;
-	virtual void MouseLeave(FViewport* Viewport) OVERRIDE;
-	virtual void SetIsSimulateInEditorViewport(bool bInIsSimulateInEditorViewport) OVERRIDE;
+	virtual void RedrawRequested(FViewport* InViewport) override {}
+	virtual bool InputKey(FViewport* Viewport, int32 ControllerId, FKey Key, EInputEvent EventType, float AmountDepressed=1.f, bool bGamepad=false) override;
+	virtual bool InputAxis(FViewport* Viewport, int32 ControllerId, FKey Key, float Delta, float DeltaTime, int32 NumSamples=1, bool bGamepad=false) override;
+	virtual bool InputChar(FViewport* Viewport,int32 ControllerId, TCHAR Character) override;
+	virtual bool InputTouch(FViewport* Viewport, int32 ControllerId, uint32 Handle, ETouchType::Type Type, const FVector2D& TouchLocation, FDateTime DeviceTimestamp, uint32 TouchpadIndex) override;
+	virtual bool InputMotion(FViewport* Viewport, int32 ControllerId, const FVector& Tilt, const FVector& RotationRate, const FVector& Gravity, const FVector& Acceleration) override;
+	virtual EMouseCursor::Type GetCursor(FViewport* Viewport, int32 X, int32 Y ) override;
+	virtual void Precache() override;
+	virtual void Draw(FViewport* Viewport,FCanvas* SceneCanvas) override;
+	virtual void ProcessScreenShots(FViewport* Viewport) override;
+	virtual void LostFocus(FViewport* Viewport) override;
+	virtual void ReceivedFocus(FViewport* Viewport) override;
+	virtual bool IsFocused(FViewport* Viewport) override;
+	virtual void CloseRequested(FViewport* Viewport) override;
+	virtual bool RequiresHitProxyStorage() override { return 0; }
+	virtual bool IsOrtho() const override;
+	virtual void MouseEnter(FViewport* Viewport, int32 x, int32 y) override;
+	virtual void MouseLeave(FViewport* Viewport) override;
+	virtual void SetIsSimulateInEditorViewport(bool bInIsSimulateInEditorViewport) override;
 	// End of FViewportClient interface.
 
 	// Begin FExec interface.
-	virtual bool Exec( UWorld* InWorld, const TCHAR* Cmd,FOutputDevice& Ar) OVERRIDE;
+	virtual bool Exec( UWorld* InWorld, const TCHAR* Cmd,FOutputDevice& Ar) override;
 	// End of FExec interface.
 
 	/**
@@ -621,7 +621,7 @@ public:
 	/**
 	 * Get a ptr to the stat unit data for this viewport
 	 */
-	virtual FStatUnitData* GetStatUnitData() const OVERRIDE
+	virtual FStatUnitData* GetStatUnitData() const override
 	{
 		return StatUnitData;
 	}
@@ -629,7 +629,7 @@ public:
 	/**
 	 * Get a ptr to the stat unit data for this viewport
 	 */
-	virtual FStatHitchesData* GetStatHitchesData() const OVERRIDE
+	virtual FStatHitchesData* GetStatHitchesData() const override
 	{
 		return StatHitchesData;
 	}
@@ -637,7 +637,7 @@ public:
 	/**
 	 * Get a ptr to the enabled stats list
 	 */
-	virtual const TArray<FString>* GetEnabledStats() const OVERRIDE
+	virtual const TArray<FString>* GetEnabledStats() const override
 	{
 		return &EnabledStats;
 	}
@@ -645,7 +645,7 @@ public:
 	/**
 	 * Sets all the stats that should be enabled for the viewport
 	 */
-	virtual void SetEnabledStats(const TArray<FString>& InEnabledStats) OVERRIDE
+	virtual void SetEnabledStats(const TArray<FString>& InEnabledStats) override
 	{
 		EnabledStats = InEnabledStats;
 	}
@@ -653,7 +653,7 @@ public:
 	/**
 	 * Check whether a specific stat is enabled for this viewport
 	 */
-	virtual bool IsStatEnabled(const TCHAR* InName) const OVERRIDE
+	virtual bool IsStatEnabled(const TCHAR* InName) const override
 	{
 		return EnabledStats.Contains(InName);
 	}
@@ -661,7 +661,7 @@ public:
 	/**
 	 * Get the sound stat flags enabled for this viewport
 	 */
-	virtual ESoundShowFlags::Type GetSoundShowFlags() const OVERRIDE
+	virtual ESoundShowFlags::Type GetSoundShowFlags() const override
 	{ 
 		return SoundShowFlags;
 	}
@@ -669,7 +669,7 @@ public:
 	/**
 	 * Set the sound stat flags enabled for this viewport
 	 */
-	virtual void SetSoundShowFlags(const ESoundShowFlags::Type InSoundShowFlags) OVERRIDE
+	virtual void SetSoundShowFlags(const ESoundShowFlags::Type InSoundShowFlags) override
 	{
 		SoundShowFlags = InSoundShowFlags;
 	}
diff --git a/Engine/Source/Runtime/Engine/Classes/Engine/InputActionDelegateBinding.h b/Engine/Source/Runtime/Engine/Classes/Engine/InputActionDelegateBinding.h
--- a/Engine/Source/Runtime/Engine/Classes/Engine/InputActionDelegateBinding.h
+++ b/Engine/Source/Runtime/Engine/Classes/Engine/InputActionDelegateBinding.h
@@ -35,6 +35,6 @@ class ENGINE_API UInputActionDelegateBinding : public UInputDelegateBinding
 	TArray<FBlueprintInputActionDelegateBinding> InputActionDelegateBindings;
 
 	// Begin UInputDelegateBinding interface
-	virtual void BindToInputComponent(UInputComponent* InputComponent) const OVERRIDE;
+	virtual void BindToInputComponent(UInputComponent* InputComponent) const override;
 	// End UInputDelegateBinding interface
 };
diff --git a/Engine/Source/Runtime/Engine/Classes/Engine/InputAxisDelegateBinding.h b/Engine/Source/Runtime/Engine/Classes/Engine/InputAxisDelegateBinding.h
--- a/Engine/Source/Runtime/Engine/Classes/Engine/InputAxisDelegateBinding.h
+++ b/Engine/Source/Runtime/Engine/Classes/Engine/InputAxisDelegateBinding.h
@@ -31,6 +31,6 @@ class ENGINE_API UInputAxisDelegateBinding : public UInputDelegateBinding
 	TArray<FBlueprintInputAxisDelegateBinding> InputAxisDelegateBindings;
 
 	// Begin UInputDelegateBinding interface
-	virtual void BindToInputComponent(UInputComponent* InputComponent) const OVERRIDE;
+	virtual void BindToInputComponent(UInputComponent* InputComponent) const override;
 	// End UInputDelegateBinding interface
 };
diff --git a/Engine/Source/Runtime/Engine/Classes/Engine/InteractiveFoliageActor.h b/Engine/Source/Runtime/Engine/Classes/Engine/InteractiveFoliageActor.h
--- a/Engine/Source/Runtime/Engine/Classes/Engine/InteractiveFoliageActor.h
+++ b/Engine/Source/Runtime/Engine/Classes/Engine/InteractiveFoliageActor.h
@@ -87,13 +87,13 @@ protected:
 	void CapsuleTouched(AActor* Other, UPrimitiveComponent* OtherComp, int32 OtherBodyIndex, bool bFromSweep, const FHitResult & OverlapInfo);
 public:
 	// Begin AActor interface
-	virtual void Tick(float DeltaSeconds) OVERRIDE;
-	virtual void PostActorCreated() OVERRIDE;
-	virtual float TakeDamage(float DamageAmount, struct FDamageEvent const& DamageEvent, class AController* EventInstigator, class AActor* DamageCauser) OVERRIDE;
+	virtual void Tick(float DeltaSeconds) override;
+	virtual void PostActorCreated() override;
+	virtual float TakeDamage(float DamageAmount, struct FDamageEvent const& DamageEvent, class AController* EventInstigator, class AActor* DamageCauser) override;
 	// End AActor interface
 
 	// Begin UObject interface
-	virtual void PostLoad() OVERRIDE;
+	virtual void PostLoad() override;
 	// End UObject interface
 
 };
diff --git a/Engine/Source/Runtime/Engine/Classes/Engine/InterpCurveEdSetup.h b/Engine/Source/Runtime/Engine/Classes/Engine/InterpCurveEdSetup.h
--- a/Engine/Source/Runtime/Engine/Classes/Engine/InterpCurveEdSetup.h
+++ b/Engine/Source/Runtime/Engine/Classes/Engine/InterpCurveEdSetup.h
@@ -109,13 +109,13 @@ class UInterpCurveEdSetup : public UObject
 
 
 	// Begin UObject interface
-	virtual void PostLoad() OVERRIDE;
-	virtual void Serialize(FArchive& Ar) OVERRIDE;
-	virtual bool NeedsLoadForClient() const OVERRIDE
+	virtual void PostLoad() override;
+	virtual void Serialize(FArchive& Ar) override;
+	virtual bool NeedsLoadForClient() const override
 	{
 		return false;
 	}
-	virtual bool NeedsLoadForServer() const OVERRIDE
+	virtual bool NeedsLoadForServer() const override
 	{
 		return false;
 	}
diff --git a/Engine/Source/Runtime/Engine/Classes/Engine/Level.h b/Engine/Source/Runtime/Engine/Classes/Engine/Level.h
--- a/Engine/Source/Runtime/Engine/Classes/Engine/Level.h
+++ b/Engine/Source/Runtime/Engine/Classes/Engine/Level.h
@@ -460,18 +460,18 @@ public:
 	~ULevel();
 
 	// Begin UObject interface.
-	virtual void Serialize( FArchive& Ar ) OVERRIDE;
-	virtual void BeginDestroy() OVERRIDE;
-	virtual bool IsReadyForFinishDestroy() OVERRIDE;
-	virtual void FinishDestroy() OVERRIDE;
-	virtual UWorld* GetWorld() const OVERRIDE;
+	virtual void Serialize( FArchive& Ar ) override;
+	virtual void BeginDestroy() override;
+	virtual bool IsReadyForFinishDestroy() override;
+	virtual void FinishDestroy() override;
+	virtual UWorld* GetWorld() const override;
 
 #if	WITH_EDITOR
-	virtual void PreEditUndo() OVERRIDE;
-	virtual void PostEditUndo() OVERRIDE;	
+	virtual void PreEditUndo() override;
+	virtual void PostEditUndo() override;	
 #endif // WITH_EDITOR
-	virtual void PostLoad() OVERRIDE;
-	virtual void PreSave() OVERRIDE;
+	virtual void PostLoad() override;
+	virtual void PreSave() override;
 	static void AddReferencedObjects(UObject* InThis, FReferenceCollector& Collector);
 	// End UObject interface.
 
@@ -548,7 +548,7 @@ public:
 	 */
 	ENGINE_API void SortActorList();
 
-	virtual bool IsNameStableForNetworking() const OVERRIDE { return true; }		// For now, assume all levels have stable net names
+	virtual bool IsNameStableForNetworking() const override { return true; }		// For now, assume all levels have stable net names
 
 	/** Handles network initialization for actors in this level */
 	void InitializeNetworkActors();
@@ -679,9 +679,9 @@ public:
 	
 	
 	// Begin IInterface_AssetUserData Interface
-	virtual void AddAssetUserData(UAssetUserData* InUserData) OVERRIDE;
-	virtual void RemoveUserDataOfClass(TSubclassOf<UAssetUserData> InUserDataClass) OVERRIDE;
-	virtual UAssetUserData* GetAssetUserDataOfClass(TSubclassOf<UAssetUserData> InUserDataClass) OVERRIDE;
+	virtual void AddAssetUserData(UAssetUserData* InUserData) override;
+	virtual void RemoveUserDataOfClass(TSubclassOf<UAssetUserData> InUserDataClass) override;
+	virtual UAssetUserData* GetAssetUserDataOfClass(TSubclassOf<UAssetUserData> InUserDataClass) override;
 	// End IInterface_AssetUserData Interface
 
 #if WITH_EDITOR
diff --git a/Engine/Source/Runtime/Engine/Classes/Engine/LevelBase.h b/Engine/Source/Runtime/Engine/Classes/Engine/LevelBase.h
--- a/Engine/Source/Runtime/Engine/Classes/Engine/LevelBase.h
+++ b/Engine/Source/Runtime/Engine/Classes/Engine/LevelBase.h
@@ -22,7 +22,7 @@ class ULevelBase : public UObject
 	ULevelBase( const class FPostConstructInitializeProperties& PCIP,const FURL& InURL );
 
 	// Begin UObject Interface
-	virtual void Serialize( FArchive& Ar ) OVERRIDE;
+	virtual void Serialize( FArchive& Ar ) override;
 	static void AddReferencedObjects(UObject* InThis, FReferenceCollector& Collector);
 	// End UObject Interface
 
diff --git a/Engine/Source/Runtime/Engine/Classes/Engine/LevelBounds.h b/Engine/Source/Runtime/Engine/Classes/Engine/LevelBounds.h
--- a/Engine/Source/Runtime/Engine/Classes/Engine/LevelBounds.h
+++ b/Engine/Source/Runtime/Engine/Classes/Engine/LevelBounds.h
@@ -18,23 +18,23 @@ class ALevelBounds : public AActor
 	bool bAutoUpdateBounds;
 
 	// Begin UObject Interface
-	virtual void PostLoad() OVERRIDE;
+	virtual void PostLoad() override;
 	// End UObject Interface
 	
 	// Begin AActor interface.
-	virtual FBox GetComponentsBoundingBox(bool bNonColliding = false) const OVERRIDE;
-	virtual bool IsLevelBoundsRelevant() const OVERRIDE { return false; }
+	virtual FBox GetComponentsBoundingBox(bool bNonColliding = false) const override;
+	virtual bool IsLevelBoundsRelevant() const override { return false; }
 	// End AActor interface.
 
 	/** @return Bounding box which includes all relevant actors bounding boxes belonging to specified level */
 	ENGINE_API static FBox CalculateLevelBounds(ULevel* InLevel);
 		
 #if WITH_EDITOR
-	virtual void PostEditUndo() OVERRIDE;
-	virtual void PostEditMove(bool bFinished) OVERRIDE;
-	virtual void PostEditChangeProperty(FPropertyChangedEvent& PropertyChangedEvent) OVERRIDE;
-	virtual void PostRegisterAllComponents() OVERRIDE;
-	virtual void PostUnregisterAllComponents() OVERRIDE;
+	virtual void PostEditUndo() override;
+	virtual void PostEditMove(bool bFinished) override;
+	virtual void PostEditChangeProperty(FPropertyChangedEvent& PropertyChangedEvent) override;
+	virtual void PostRegisterAllComponents() override;
+	virtual void PostUnregisterAllComponents() override;
 	
 	/** Tells actors that level bounds was changed so actor bounds should be recalculated  */
 	void OnLevelBoundsDirtied();
diff --git a/Engine/Source/Runtime/Engine/Classes/Engine/LevelScriptActor.h b/Engine/Source/Runtime/Engine/Classes/Engine/LevelScriptActor.h
--- a/Engine/Source/Runtime/Engine/Classes/Engine/LevelScriptActor.h
+++ b/Engine/Source/Runtime/Engine/Classes/Engine/LevelScriptActor.h
@@ -43,15 +43,15 @@ class ENGINE_API ALevelScriptActor : public AActor
 	
 #if WITH_EDITOR
 	// Begin UObject Interface
-	virtual void PostDuplicate(bool bDuplicateForPIE) OVERRIDE;
-	virtual void BeginDestroy() OVERRIDE;
+	virtual void PostDuplicate(bool bDuplicateForPIE) override;
+	virtual void BeginDestroy() override;
 #endif
-	virtual void PreInitializeComponents() OVERRIDE;
+	virtual void PreInitializeComponents() override;
 	// End UObject Interface
 
 	// Begin AActor Interface
-	virtual void EnableInput(class APlayerController* PlayerController) OVERRIDE;
-	virtual void DisableInput(class APlayerController* PlayerController) OVERRIDE;
+	virtual void EnableInput(class APlayerController* PlayerController) override;
+	virtual void DisableInput(class APlayerController* PlayerController) override;
 	// End AActor Interface
 
 	bool InputEnabled() const { return bInputEnabled; }
diff --git a/Engine/Source/Runtime/Engine/Classes/Engine/LevelScriptBlueprint.h b/Engine/Source/Runtime/Engine/Classes/Engine/LevelScriptBlueprint.h
--- a/Engine/Source/Runtime/Engine/Classes/Engine/LevelScriptBlueprint.h
+++ b/Engine/Source/Runtime/Engine/Classes/Engine/LevelScriptBlueprint.h
@@ -21,9 +21,9 @@ class ULevelScriptBlueprint : public UBlueprint
 	}
 
 	// Begin UBlueprint interface
-	ENGINE_API virtual UObject* GetObjectBeingDebugged() OVERRIDE;
-	ENGINE_API virtual void SetObjectBeingDebugged(UObject* NewObject) OVERRIDE;
-	ENGINE_API virtual FString GetFriendlyName() const OVERRIDE;
+	ENGINE_API virtual UObject* GetObjectBeingDebugged() override;
+	ENGINE_API virtual void SetObjectBeingDebugged(UObject* NewObject) override;
+	ENGINE_API virtual FString GetFriendlyName() const override;
 	// End UBlueprint interface
 
 	/** Generate a name for a level script blueprint from the current level */
diff --git a/Engine/Source/Runtime/Engine/Classes/Engine/LevelStreaming.h b/Engine/Source/Runtime/Engine/Classes/Engine/LevelStreaming.h
--- a/Engine/Source/Runtime/Engine/Classes/Engine/LevelStreaming.h
+++ b/Engine/Source/Runtime/Engine/Classes/Engine/LevelStreaming.h
@@ -47,11 +47,11 @@ public:
 	 */
 	bool UpdateLevel( ULevelStreaming* LevelStreamingObject );
 
-	virtual void UpdateOperation(FLatentResponse& Response) OVERRIDE;
+	virtual void UpdateOperation(FLatentResponse& Response) override;
 
 #if WITH_EDITOR
 	// Returns a human readable description of the latent operation's current state
-	virtual FString GetDescription() const OVERRIDE;
+	virtual FString GetDescription() const override;
 #endif
 };
 
@@ -148,10 +148,10 @@ class ULevelStreaming : public UObject
 	int32 LevelLODIndex;
 
 	// Begin UObject Interface
-	virtual void PostLoad() OVERRIDE;
-	virtual void Serialize( FArchive& Ar ) OVERRIDE;
+	virtual void PostLoad() override;
+	virtual void Serialize( FArchive& Ar ) override;
 #if WITH_EDITOR
-	virtual void PostEditChangeChainProperty(FPropertyChangedChainEvent& PropertyChangedEvent) OVERRIDE;
+	virtual void PostEditChangeChainProperty(FPropertyChangedChainEvent& PropertyChangedEvent) override;
 	
 	/** Remove duplicates in EditorStreamingVolumes list*/
 	void RemoveStreamingVolumeDuplicates();
diff --git a/Engine/Source/Runtime/Engine/Classes/Engine/LevelStreamingAlwaysLoaded.h b/Engine/Source/Runtime/Engine/Classes/Engine/LevelStreamingAlwaysLoaded.h
--- a/Engine/Source/Runtime/Engine/Classes/Engine/LevelStreamingAlwaysLoaded.h
+++ b/Engine/Source/Runtime/Engine/Classes/Engine/LevelStreamingAlwaysLoaded.h
@@ -16,8 +16,8 @@ class ULevelStreamingAlwaysLoaded : public ULevelStreaming
 	GENERATED_UCLASS_BODY()
 
 	// Begin ULevelStreaming Interface
-	virtual bool ShouldBeLoaded( const FVector& ViewLocation ) OVERRIDE;
-	virtual bool ShouldBeAlwaysLoaded() const OVERRIDE { return true; } 
+	virtual bool ShouldBeLoaded( const FVector& ViewLocation ) override;
+	virtual bool ShouldBeAlwaysLoaded() const override { return true; } 
 	// End ULevelStreaming Interface
 };
 
diff --git a/Engine/Source/Runtime/Engine/Classes/Engine/LevelStreamingBounds.h b/Engine/Source/Runtime/Engine/Classes/Engine/LevelStreamingBounds.h
--- a/Engine/Source/Runtime/Engine/Classes/Engine/LevelStreamingBounds.h
+++ b/Engine/Source/Runtime/Engine/Classes/Engine/LevelStreamingBounds.h
@@ -17,8 +17,8 @@ class ULevelStreamingBounds : public ULevelStreaming
 
 
 	// Begin ULevelStreaming Interface
-	virtual bool ShouldBeLoaded( const FVector& ViewLocation ) OVERRIDE;
-	virtual bool ShouldBeVisible( const FVector& ViewLocation ) OVERRIDE;
+	virtual bool ShouldBeLoaded( const FVector& ViewLocation ) override;
+	virtual bool ShouldBeVisible( const FVector& ViewLocation ) override;
 	// End ULevelStreaming Interface
 };
 
diff --git a/Engine/Source/Runtime/Engine/Classes/Engine/LevelStreamingKismet.h b/Engine/Source/Runtime/Engine/Classes/Engine/LevelStreamingKismet.h
--- a/Engine/Source/Runtime/Engine/Classes/Engine/LevelStreamingKismet.h
+++ b/Engine/Source/Runtime/Engine/Classes/Engine/LevelStreamingKismet.h
@@ -24,12 +24,12 @@ class ULevelStreamingKismet : public ULevelStreaming
 	uint32 bInitiallyVisible:1;
 	
 	// Begin UObject Interface
-	virtual void PostLoad() OVERRIDE;
+	virtual void PostLoad() override;
 	// End UObject Interface
 
 	// Begin ULevelStreaming Interface
-	virtual bool ShouldBeLoaded( const FVector& ViewLocation ) OVERRIDE;
-	virtual bool ShouldBeVisible( const FVector& ViewLocation ) OVERRIDE;
+	virtual bool ShouldBeLoaded( const FVector& ViewLocation ) override;
+	virtual bool ShouldBeVisible( const FVector& ViewLocation ) override;
 	// End ULevelStreaming Interface
 };
 
diff --git a/Engine/Source/Runtime/Engine/Classes/Engine/LevelStreamingPersistent.h b/Engine/Source/Runtime/Engine/Classes/Engine/LevelStreamingPersistent.h
--- a/Engine/Source/Runtime/Engine/Classes/Engine/LevelStreamingPersistent.h
+++ b/Engine/Source/Runtime/Engine/Classes/Engine/LevelStreamingPersistent.h
@@ -17,11 +17,11 @@ class ULevelStreamingPersistent : public ULevelStreaming
 
 
 	// Begin ULevelStreaming Interface
-	virtual bool ShouldBeLoaded( const FVector& ViewLocation ) OVERRIDE
+	virtual bool ShouldBeLoaded( const FVector& ViewLocation ) override
 	{
 		return true;
 	}
-	virtual bool ShouldBeVisible( const FVector& ViewLocation ) OVERRIDE
+	virtual bool ShouldBeVisible( const FVector& ViewLocation ) override
 	{
 		return true;
 	}
diff --git a/Engine/Source/Runtime/Engine/Classes/Engine/LevelStreamingVolume.h b/Engine/Source/Runtime/Engine/Classes/Engine/LevelStreamingVolume.h
--- a/Engine/Source/Runtime/Engine/Classes/Engine/LevelStreamingVolume.h
+++ b/Engine/Source/Runtime/Engine/Classes/Engine/LevelStreamingVolume.h
@@ -50,7 +50,7 @@ class ALevelStreamingVolume : public AVolume
 
 #if WITH_EDITOR
 	// Begin AActor interface.
-	virtual void CheckForErrors() OVERRIDE;
+	virtual void CheckForErrors() override;
 	// End AActor interface.
 #endif
 };
diff --git a/Engine/Source/Runtime/Engine/Classes/Engine/Light.h b/Engine/Source/Runtime/Engine/Classes/Engine/Light.h
--- a/Engine/Source/Runtime/Engine/Classes/Engine/Light.h
+++ b/Engine/Source/Runtime/Engine/Classes/Engine/Light.h
@@ -54,7 +54,7 @@ class ENGINE_API ALight : public AActor
 public:
 #if WITH_EDITOR
 	// Begin AActor interface.
-	virtual void CheckForErrors() OVERRIDE;
+	virtual void CheckForErrors() override;
 	// End AActor interface.
 #endif
 
@@ -68,7 +68,7 @@ public:
 
 	// Begin AActor interface.
 	void Destroyed();
-	virtual bool IsLevelBoundsRelevant() const OVERRIDE { return false; }
+	virtual bool IsLevelBoundsRelevant() const override { return false; }
 	// End AActor interface.
 };
 
diff --git a/Engine/Source/Runtime/Engine/Classes/Engine/LightMapTexture2D.h b/Engine/Source/Runtime/Engine/Classes/Engine/LightMapTexture2D.h
--- a/Engine/Source/Runtime/Engine/Classes/Engine/LightMapTexture2D.h
+++ b/Engine/Source/Runtime/Engine/Classes/Engine/LightMapTexture2D.h
@@ -24,8 +24,8 @@ class ULightMapTexture2D : public UTexture2D
 	GENERATED_UCLASS_BODY()
 
 	// Begin UObject interface.
-	virtual void Serialize( FArchive& Ar ) OVERRIDE;
-	virtual FString GetDesc() OVERRIDE;
+	virtual void Serialize( FArchive& Ar ) override;
+	virtual FString GetDesc() override;
 	// End UObject interface. 
 
 	bool IsLowQualityLightmap() const
diff --git a/Engine/Source/Runtime/Engine/Classes/Engine/LocalPlayer.h b/Engine/Source/Runtime/Engine/Classes/Engine/LocalPlayer.h
--- a/Engine/Source/Runtime/Engine/Classes/Engine/LocalPlayer.h
+++ b/Engine/Source/Runtime/Engine/Classes/Engine/LocalPlayer.h
@@ -96,13 +96,13 @@ private:
 
 public:
 	// UObject interface
-	virtual void PostInitProperties() OVERRIDE;
-	virtual void FinishDestroy() OVERRIDE;
+	virtual void PostInitProperties() override;
+	virtual void FinishDestroy() override;
 	static void AddReferencedObjects(UObject* InThis, FReferenceCollector& Collector);
 	// End of UObject interface
 
 	// FExec interface
-	virtual bool Exec(UWorld* InWorld, const TCHAR* Cmd,FOutputDevice& Ar) OVERRIDE;
+	virtual bool Exec(UWorld* InWorld, const TCHAR* Cmd,FOutputDevice& Ar) override;
 	// End of FExec interface
 
 	/** 
@@ -124,7 +124,7 @@ public:
 	bool HandleCancelMatineeCommand( const TCHAR* Cmd, FOutputDevice& Ar );
 	
 	// UPlayer interface
-	virtual void HandleDisconnect(class UWorld *World, class UNetDriver *NetDriver) OVERRIDE;
+	virtual void HandleDisconnect(class UWorld *World, class UNetDriver *NetDriver) override;
 	// End of UPlayer interface
 
 protected:
diff --git a/Engine/Source/Runtime/Engine/Classes/Engine/NavigationObjectBase.h b/Engine/Source/Runtime/Engine/Classes/Engine/NavigationObjectBase.h
--- a/Engine/Source/Runtime/Engine/Classes/Engine/NavigationObjectBase.h
+++ b/Engine/Source/Runtime/Engine/Classes/Engine/NavigationObjectBase.h
@@ -26,11 +26,11 @@ class ENGINE_API ANavigationObjectBase : public AActor, public INavAgentInterfac
 
 
 	// Begin AActor Interface
-	virtual void GetSimpleCollisionCylinder(float& CollisionRadius, float& CollisionHalfHeight) const OVERRIDE;
+	virtual void GetSimpleCollisionCylinder(float& CollisionRadius, float& CollisionHalfHeight) const override;
 
 #if	WITH_EDITOR
-	virtual void PostEditMove(bool bFinished) OVERRIDE;
-	virtual void PostEditUndo() OVERRIDE;
+	virtual void PostEditMove(bool bFinished) override;
+	virtual void PostEditUndo() override;
 #endif	// WITH_EDITOR
 	// End AActor Interface
 
diff --git a/Engine/Source/Runtime/Engine/Classes/Engine/NetConnection.h b/Engine/Source/Runtime/Engine/Classes/Engine/NetConnection.h
--- a/Engine/Source/Runtime/Engine/Classes/Engine/NetConnection.h
+++ b/Engine/Source/Runtime/Engine/Classes/Engine/NetConnection.h
@@ -315,9 +315,9 @@ public:
 
 	// Begin UObject interface.
 
-	ENGINE_API virtual void Serialize( FArchive& Ar ) OVERRIDE;
+	ENGINE_API virtual void Serialize( FArchive& Ar ) override;
 
-	ENGINE_API virtual void FinishDestroy() OVERRIDE;
+	ENGINE_API virtual void FinishDestroy() override;
 
 	ENGINE_API static void AddReferencedObjects(UObject* InThis, FReferenceCollector& Collector);
 
@@ -326,7 +326,7 @@ public:
 
 	// Begin FExec interface.
 
-	ENGINE_API virtual bool Exec( UWorld* InWorld, const TCHAR* Cmd, FOutputDevice& Ar=*GLog ) OVERRIDE;
+	ENGINE_API virtual bool Exec( UWorld* InWorld, const TCHAR* Cmd, FOutputDevice& Ar=*GLog ) override;
 
 	// End FExec interface.
 
diff --git a/Engine/Source/Runtime/Engine/Classes/Engine/NetDriver.h b/Engine/Source/Runtime/Engine/Classes/Engine/NetDriver.h
--- a/Engine/Source/Runtime/Engine/Classes/Engine/NetDriver.h
+++ b/Engine/Source/Runtime/Engine/Classes/Engine/NetDriver.h
@@ -314,9 +314,9 @@ public:
 
 
 	// Begin UObject interface.
-	ENGINE_API virtual void PostInitProperties() OVERRIDE;
-	ENGINE_API virtual void FinishDestroy() OVERRIDE;
-	ENGINE_API virtual void Serialize( FArchive& Ar ) OVERRIDE;
+	ENGINE_API virtual void PostInitProperties() override;
+	ENGINE_API virtual void FinishDestroy() override;
+	ENGINE_API virtual void Serialize( FArchive& Ar ) override;
 	ENGINE_API static void AddReferencedObjects(UObject* InThis, FReferenceCollector& Collector);
 	// End UObject interface.
 
@@ -331,7 +331,7 @@ public:
 	 *
 	 * @return true if the handler consumed the input, false to continue searching handlers
 	 */
-	ENGINE_API virtual bool Exec(UWorld* InWorld, const TCHAR* Cmd, FOutputDevice& Ar=*GLog) OVERRIDE;
+	ENGINE_API virtual bool Exec(UWorld* InWorld, const TCHAR* Cmd, FOutputDevice& Ar=*GLog) override;
 
 	ENGINE_API ENetMode	GetNetMode() const;
 
diff --git a/Engine/Source/Runtime/Engine/Classes/Engine/PackageMapClient.h b/Engine/Source/Runtime/Engine/Classes/Engine/PackageMapClient.h
--- a/Engine/Source/Runtime/Engine/Classes/Engine/PackageMapClient.h
+++ b/Engine/Source/Runtime/Engine/Classes/Engine/PackageMapClient.h
@@ -114,12 +114,12 @@ class UPackageMapClient : public UPackageMap
 
 	
 	// UPackageMap Interface
-	virtual bool SerializeObject( FArchive& Ar, UClass* Class, UObject*& Obj, FNetworkGUID *OutNetGUID = NULL ) OVERRIDE;
-	virtual bool SerializeNewActor( FArchive& Ar, class UActorChannel *Channel, class AActor*& Actor) OVERRIDE;
+	virtual bool SerializeObject( FArchive& Ar, UClass* Class, UObject*& Obj, FNetworkGUID *OutNetGUID = NULL ) override;
+	virtual bool SerializeNewActor( FArchive& Ar, class UActorChannel *Channel, class AActor*& Actor) override;
 	
-	virtual bool WriteObject( FArchive& Ar, UObject* Outer, FNetworkGUID NetGUID, FString ObjName ) OVERRIDE;
+	virtual bool WriteObject( FArchive& Ar, UObject* Outer, FNetworkGUID NetGUID, FString ObjName ) override;
 
-	virtual void ResetPackageMap() OVERRIDE;
+	virtual void ResetPackageMap() override;
 
 	virtual void SetLocked(bool L) { Locked = L; }
 
@@ -127,10 +127,10 @@ class UPackageMapClient : public UPackageMap
 
 	bool NetGUIDHasBeenAckd(FNetworkGUID NetGUID);
 
-	virtual void ReceivedNak( const int32 NakPacketId ) OVERRIDE;
-	virtual void ReceivedAck( const int32 AckPacketId ) OVERRIDE;
-	virtual void NotifyBunchCommit( const int32 OutPacketId, const TArray< FNetworkGUID > & ExportNetGUIDs ) OVERRIDE;
-	virtual void GetNetGUIDStats(int32 &AckCount, int32 &UnAckCount, int32 &PendingCount) OVERRIDE;
+	virtual void ReceivedNak( const int32 NakPacketId ) override;
+	virtual void ReceivedAck( const int32 AckPacketId ) override;
+	virtual void NotifyBunchCommit( const int32 OutPacketId, const TArray< FNetworkGUID > & ExportNetGUIDs ) override;
+	virtual void GetNetGUIDStats(int32 &AckCount, int32 &UnAckCount, int32 &PendingCount) override;
 
 	void ReceiveNetGUIDBunch( FInBunch &InBunch );
 	bool AppendExportBunches(TArray<FOutBunch *>& OutgoingBunches);
@@ -141,12 +141,12 @@ class UPackageMapClient : public UPackageMap
 
 	static void	AddReferencedObjects(UObject* InThis, FReferenceCollector& Collector);
 
-	virtual void NotifyStreamingLevelUnload(UObject* UnloadedLevel) OVERRIDE;
+	virtual void NotifyStreamingLevelUnload(UObject* UnloadedLevel) override;
 
-	virtual bool PrintExportBatch() OVERRIDE;
+	virtual bool PrintExportBatch() override;
 
-	virtual void		LogDebugInfo( FOutputDevice & Ar) OVERRIDE;
-	virtual UObject *	GetObjectFromNetGUID( const FNetworkGUID & NetGUID ) OVERRIDE;
+	virtual void		LogDebugInfo( FOutputDevice & Ar) override;
+	virtual UObject *	GetObjectFromNetGUID( const FNetworkGUID & NetGUID ) override;
 
 protected:
 
@@ -156,7 +156,7 @@ protected:
 	void			InternalWriteObject( FArchive& Ar, FNetworkGUID NetGUID, const UObject * Object, FString ObjectPathName, UObject * ObjectOuter );	
 	FNetworkGUID	InternalLoadObject( FArchive & Ar, UObject *& Object, const bool bIsInnerLevel, int InternalLoadObjectRecursionCount );
 
-	virtual UObject * ResolvePathAndAssignNetGUID( const FNetworkGUID & NetGUID, const FString & PathName ) OVERRIDE;
+	virtual UObject * ResolvePathAndAssignNetGUID( const FNetworkGUID & NetGUID, const FString & PathName ) override;
 
 	bool	ShouldSendFullPath(const UObject* Object, const FNetworkGUID &NetGUID);
 	
diff --git a/Engine/Source/Runtime/Engine/Classes/Engine/PendingNetGame.h b/Engine/Source/Runtime/Engine/Classes/Engine/PendingNetGame.h
--- a/Engine/Source/Runtime/Engine/Classes/Engine/PendingNetGame.h
+++ b/Engine/Source/Runtime/Engine/Classes/Engine/PendingNetGame.h
@@ -97,10 +97,10 @@ public:
 	void	InitNetDriver();
 
 	// Begin FNetworkNotify interface.
-	virtual EAcceptConnection::Type NotifyAcceptingConnection() OVERRIDE;
-	virtual void NotifyAcceptedConnection( class UNetConnection* Connection ) OVERRIDE;
-	virtual bool NotifyAcceptingChannel( class UChannel* Channel ) OVERRIDE;
-	virtual void NotifyControlMessage(UNetConnection* Connection, uint8 MessageType, class FInBunch& Bunch) OVERRIDE;
+	virtual EAcceptConnection::Type NotifyAcceptingConnection() override;
+	virtual void NotifyAcceptedConnection( class UNetConnection* Connection ) override;
+	virtual bool NotifyAcceptingChannel( class UChannel* Channel ) override;
+	virtual void NotifyControlMessage(UNetConnection* Connection, uint8 MessageType, class FInBunch& Bunch) override;
 	// End FNetworkNotify interface.
 
 	/**  Update the pending level's status. */
@@ -113,9 +113,9 @@ public:
 	virtual void SendJoin();
 
 	// Begin UObject interface.
-	virtual void Serialize( FArchive& Ar ) OVERRIDE;
+	virtual void Serialize( FArchive& Ar ) override;
 
-	virtual void FinishDestroy() OVERRIDE
+	virtual void FinishDestroy() override
 	{
 		NetDriver = NULL;
 		
diff --git a/Engine/Source/Runtime/Engine/Classes/Engine/PlaneReflectionCapture.h b/Engine/Source/Runtime/Engine/Classes/Engine/PlaneReflectionCapture.h
--- a/Engine/Source/Runtime/Engine/Classes/Engine/PlaneReflectionCapture.h
+++ b/Engine/Source/Runtime/Engine/Classes/Engine/PlaneReflectionCapture.h
@@ -16,7 +16,7 @@ public:
 
 #if WITH_EDITOR
 	// Begin AActor interface.
-	virtual void EditorApplyScale(const FVector& DeltaScale, const FVector* PivotLocation, bool bAltDown, bool bShiftDown, bool bCtrlDown) OVERRIDE;
+	virtual void EditorApplyScale(const FVector& DeltaScale, const FVector* PivotLocation, bool bAltDown, bool bShiftDown, bool bCtrlDown) override;
 	// End AActor interface.
 #endif
 
diff --git a/Engine/Source/Runtime/Engine/Classes/Engine/Player.h b/Engine/Source/Runtime/Engine/Classes/Engine/Player.h
--- a/Engine/Source/Runtime/Engine/Classes/Engine/Player.h
+++ b/Engine/Source/Runtime/Engine/Classes/Engine/Player.h
@@ -31,7 +31,7 @@ class UPlayer : public UObject, public FExec
 
 public:
 	// Begin FExec interface.
-	ENGINE_API virtual bool Exec( UWorld* InWorld, const TCHAR* Cmd,FOutputDevice& Ar) OVERRIDE;
+	ENGINE_API virtual bool Exec( UWorld* InWorld, const TCHAR* Cmd,FOutputDevice& Ar) override;
 	// End FExec interface.
 
 	/**
diff --git a/Engine/Source/Runtime/Engine/Classes/Engine/PointLight.h b/Engine/Source/Runtime/Engine/Classes/Engine/PointLight.h
--- a/Engine/Source/Runtime/Engine/Classes/Engine/PointLight.h
+++ b/Engine/Source/Runtime/Engine/Classes/Engine/PointLight.h
@@ -21,14 +21,14 @@ class APointLight : public ALight
 
 #if WITH_EDITOR
 	// Begin AActor interface.
-	virtual void EditorApplyScale(const FVector& DeltaScale, const FVector* PivotLocation, bool bAltDown, bool bShiftDown, bool bCtrlDown) OVERRIDE;
+	virtual void EditorApplyScale(const FVector& DeltaScale, const FVector* PivotLocation, bool bAltDown, bool bShiftDown, bool bCtrlDown) override;
 	// End AActor interface.
 #endif
 
 	// Begin UObject interface.
-	virtual void PostLoad() OVERRIDE;
+	virtual void PostLoad() override;
 #if WITH_EDITOR
-	virtual void LoadedFromAnotherClass(const FName& OldClassName) OVERRIDE;
+	virtual void LoadedFromAnotherClass(const FName& OldClassName) override;
 #endif
 	// End UObject interface.
 };
diff --git a/Engine/Source/Runtime/Engine/Classes/Engine/Polys.h b/Engine/Source/Runtime/Engine/Classes/Engine/Polys.h
--- a/Engine/Source/Runtime/Engine/Classes/Engine/Polys.h
+++ b/Engine/Source/Runtime/Engine/Classes/Engine/Polys.h
@@ -582,9 +582,9 @@ class UPolys : public UObject
 	{}
 
 	// Begin UObject Interface
-	ENGINE_API virtual bool Modify(bool bAlwaysMarkDirty = false) OVERRIDE;	
-	ENGINE_API virtual void Serialize( FArchive& Ar ) OVERRIDE;
-	virtual bool IsAsset() const OVERRIDE { return false; }
+	ENGINE_API virtual bool Modify(bool bAlwaysMarkDirty = false) override;	
+	ENGINE_API virtual void Serialize( FArchive& Ar ) override;
+	virtual bool IsAsset() const override { return false; }
 	static void AddReferencedObjects(UObject* InThis, FReferenceCollector& Collector);
 	// End UObject Interface
 };
diff --git a/Engine/Source/Runtime/Engine/Classes/Engine/PostProcessVolume.h b/Engine/Source/Runtime/Engine/Classes/Engine/PostProcessVolume.h
--- a/Engine/Source/Runtime/Engine/Classes/Engine/PostProcessVolume.h
+++ b/Engine/Source/Runtime/Engine/Classes/Engine/PostProcessVolume.h
@@ -43,8 +43,8 @@ class APostProcessVolume : public AVolume, public IInterface_PostProcessVolume
 	uint32 bUnbound:1;
 
 	// Begin IInterface_PostProcessVolume Interface
-	ENGINE_API virtual bool EncompassesPoint(FVector Point, float SphereRadius/*=0.f*/, float* OutDistanceToPoint) OVERRIDE;
-	ENGINE_API virtual FPostProcessVolumeProperties GetProperties() const OVERRIDE
+	ENGINE_API virtual bool EncompassesPoint(FVector Point, float SphereRadius/*=0.f*/, float* OutDistanceToPoint) override;
+	ENGINE_API virtual FPostProcessVolumeProperties GetProperties() const override
 	{
 		FPostProcessVolumeProperties Ret;
 		Ret.bIsEnabled = bEnabled != 0;
@@ -59,16 +59,16 @@ class APostProcessVolume : public AVolume, public IInterface_PostProcessVolume
 
 
 	// Begin AActor Interface
-	virtual void PostUnregisterAllComponents( void ) OVERRIDE;
+	virtual void PostUnregisterAllComponents( void ) override;
 
 protected:
-	virtual void PostRegisterAllComponents() OVERRIDE;
+	virtual void PostRegisterAllComponents() override;
 	// End AActor Interface
 public:
 	
 	// Begin UObject interface
 #if WITH_EDITOR
-	virtual void PostEditChangeProperty(FPropertyChangedEvent& PropertyChangedEvent) OVERRIDE;
+	virtual void PostEditChangeProperty(FPropertyChangedEvent& PropertyChangedEvent) override;
 #endif // WITH_EDITOR
 	// End UObject interface
 };
diff --git a/Engine/Source/Runtime/Engine/Classes/Engine/ReflectionCapture.h b/Engine/Source/Runtime/Engine/Classes/Engine/ReflectionCapture.h
--- a/Engine/Source/Runtime/Engine/Classes/Engine/ReflectionCapture.h
+++ b/Engine/Source/Runtime/Engine/Classes/Engine/ReflectionCapture.h
@@ -17,10 +17,10 @@ class AReflectionCapture : public AActor
 	TSubobjectPtr<UBillboardComponent> SpriteComponent;
 #endif // WITH_EDITORONLY_DATA
 	
-	virtual bool IsLevelBoundsRelevant() const OVERRIDE { return false; }
+	virtual bool IsLevelBoundsRelevant() const override { return false; }
 public:
 #if WITH_EDITOR
-	virtual void PostEditMove(bool bFinished) OVERRIDE;
+	virtual void PostEditMove(bool bFinished) override;
 #endif // WITH_EDITOR
 
 };
diff --git a/Engine/Source/Runtime/Engine/Classes/Engine/RendererSettings.h b/Engine/Source/Runtime/Engine/Classes/Engine/RendererSettings.h
--- a/Engine/Source/Runtime/Engine/Classes/Engine/RendererSettings.h
+++ b/Engine/Source/Runtime/Engine/Classes/Engine/RendererSettings.h
@@ -175,9 +175,9 @@ public:
 
 	// Begin UObject interface
 
-	virtual void PostInitProperties() OVERRIDE;
+	virtual void PostInitProperties() override;
 #if WITH_EDITOR
-	virtual void PostEditChangeProperty(FPropertyChangedEvent& PropertyChangedEvent) OVERRIDE;
+	virtual void PostEditChangeProperty(FPropertyChangedEvent& PropertyChangedEvent) override;
 #endif
 
 	// End UObject interface
diff --git a/Engine/Source/Runtime/Engine/Classes/Engine/SceneCapture2D.h b/Engine/Source/Runtime/Engine/Classes/Engine/SceneCapture2D.h
--- a/Engine/Source/Runtime/Engine/Classes/Engine/SceneCapture2D.h
+++ b/Engine/Source/Runtime/Engine/Classes/Engine/SceneCapture2D.h
@@ -24,7 +24,7 @@ class ASceneCapture2D : public ASceneCapture
 public:
 
 	// Begin AActor interface
-	ENGINE_API virtual void PostActorCreated() OVERRIDE;
+	ENGINE_API virtual void PostActorCreated() override;
 	// End AActor interface.
 
 	/** Used to synchronize the DrawFrustumComponent with the SceneCaptureComponent2D settings. */
diff --git a/Engine/Source/Runtime/Engine/Classes/Engine/SceneCaptureCube.h b/Engine/Source/Runtime/Engine/Classes/Engine/SceneCaptureCube.h
--- a/Engine/Source/Runtime/Engine/Classes/Engine/SceneCaptureCube.h
+++ b/Engine/Source/Runtime/Engine/Classes/Engine/SceneCaptureCube.h
@@ -24,9 +24,9 @@ class ASceneCaptureCube : public ASceneCapture
 public:
 
 	// Begin AActor interface
-	virtual void PostActorCreated() OVERRIDE;
+	virtual void PostActorCreated() override;
 #if WITH_EDITOR
-	ENGINE_API virtual void PostEditMove(bool bFinished) OVERRIDE;
+	ENGINE_API virtual void PostEditMove(bool bFinished) override;
 #endif
 	// End AActor interface.
 
diff --git a/Engine/Source/Runtime/Engine/Classes/Engine/Selection.h b/Engine/Source/Runtime/Engine/Classes/Engine/Selection.h
--- a/Engine/Source/Runtime/Engine/Classes/Engine/Selection.h
+++ b/Engine/Source/Runtime/Engine/Classes/Engine/Selection.h
@@ -290,7 +290,7 @@ public:
 
 
 	// Begin UObject Interface
-	virtual void Serialize(FArchive& Ar) OVERRIDE;
+	virtual void Serialize(FArchive& Ar) override;
 	// End UObject Interface
 
 
diff --git a/Engine/Source/Runtime/Engine/Classes/Engine/SimpleConstructionScript.h b/Engine/Source/Runtime/Engine/Classes/Engine/SimpleConstructionScript.h
--- a/Engine/Source/Runtime/Engine/Classes/Engine/SimpleConstructionScript.h
+++ b/Engine/Source/Runtime/Engine/Classes/Engine/SimpleConstructionScript.h
@@ -10,8 +10,8 @@ class USimpleConstructionScript : public UObject
 	GENERATED_UCLASS_BODY()
 
 	// Begin UObject Interface
-	virtual void Serialize(FArchive& Ar) OVERRIDE;
-	virtual void PostLoad() OVERRIDE;
+	virtual void Serialize(FArchive& Ar) override;
+	virtual void PostLoad() override;
 	// End UObject Interface
 
 	/** Ensures that all root node parent references are still valid and clears the reference if not */
diff --git a/Engine/Source/Runtime/Engine/Classes/Engine/SkeletalMesh.h b/Engine/Source/Runtime/Engine/Classes/Engine/SkeletalMesh.h
--- a/Engine/Source/Runtime/Engine/Classes/Engine/SkeletalMesh.h
+++ b/Engine/Source/Runtime/Engine/Classes/Engine/SkeletalMesh.h
@@ -652,19 +652,19 @@ public:
 
 	// Begin UObject interface.
 #if WITH_EDITOR
-	virtual void PreEditChange(UProperty* PropertyAboutToChange) OVERRIDE;
-	virtual void PostEditChangeProperty(FPropertyChangedEvent& PropertyChangedEvent) OVERRIDE;
-	virtual void PostEditUndo() OVERRIDE;
+	virtual void PreEditChange(UProperty* PropertyAboutToChange) override;
+	virtual void PostEditChangeProperty(FPropertyChangedEvent& PropertyChangedEvent) override;
+	virtual void PostEditUndo() override;
 #endif // WITH_EDITOR
-	virtual void BeginDestroy() OVERRIDE;
-	virtual bool IsReadyForFinishDestroy() OVERRIDE;
-	virtual void PreSave() OVERRIDE;
-	virtual void Serialize(FArchive& Ar) OVERRIDE;
-	virtual void PostLoad() OVERRIDE;	
-	virtual void GetAssetRegistryTags(TArray<FAssetRegistryTag>& OutTags) const OVERRIDE;
-	virtual FString GetDesc() OVERRIDE;
-	virtual FString GetDetailedInfoInternal() const OVERRIDE;
-	virtual SIZE_T GetResourceSize(EResourceSizeMode::Type Mode) OVERRIDE;
+	virtual void BeginDestroy() override;
+	virtual bool IsReadyForFinishDestroy() override;
+	virtual void PreSave() override;
+	virtual void Serialize(FArchive& Ar) override;
+	virtual void PostLoad() override;	
+	virtual void GetAssetRegistryTags(TArray<FAssetRegistryTag>& OutTags) const override;
+	virtual FString GetDesc() override;
+	virtual FString GetDetailedInfoInternal() const override;
+	virtual SIZE_T GetResourceSize(EResourceSizeMode::Type Mode) override;
 	static void AddReferencedObjects(UObject* InThis, FReferenceCollector& Collector);
 	// End UObject interface.
 
diff --git a/Engine/Source/Runtime/Engine/Classes/Engine/SphereReflectionCapture.h b/Engine/Source/Runtime/Engine/Classes/Engine/SphereReflectionCapture.h
--- a/Engine/Source/Runtime/Engine/Classes/Engine/SphereReflectionCapture.h
+++ b/Engine/Source/Runtime/Engine/Classes/Engine/SphereReflectionCapture.h
@@ -25,7 +25,7 @@ public:
 
 #if WITH_EDITOR
 	// Begin AActor interface.
-	virtual void EditorApplyScale(const FVector& DeltaScale, const FVector* PivotLocation, bool bAltDown, bool bShiftDown, bool bCtrlDown) OVERRIDE;
+	virtual void EditorApplyScale(const FVector& DeltaScale, const FVector* PivotLocation, bool bAltDown, bool bShiftDown, bool bCtrlDown) override;
 	// End AActor interface.
 #endif
 
diff --git a/Engine/Source/Runtime/Engine/Classes/Engine/SpotLight.h b/Engine/Source/Runtime/Engine/Classes/Engine/SpotLight.h
--- a/Engine/Source/Runtime/Engine/Classes/Engine/SpotLight.h
+++ b/Engine/Source/Runtime/Engine/Classes/Engine/SpotLight.h
@@ -32,15 +32,15 @@ class ASpotLight : public ALight
 
 #if WITH_EDITOR
 	// Begin AActor Interface.
-	virtual void EditorApplyScale(const FVector& DeltaScale, const FVector* PivotLocation, bool bAltDown, bool bShiftDown, bool bCtrlDown) OVERRIDE;
+	virtual void EditorApplyScale(const FVector& DeltaScale, const FVector* PivotLocation, bool bAltDown, bool bShiftDown, bool bCtrlDown) override;
 	// End AActor Interface.
 #endif
 
 	// Begin UObject Interface
-	virtual void PostLoad() OVERRIDE;
+	virtual void PostLoad() override;
 #if WITH_EDITOR
-	virtual void LoadedFromAnotherClass(const FName& OldClassName) OVERRIDE;
-	virtual void PostEditChangeProperty(FPropertyChangedEvent& PropertyChangedEvent) OVERRIDE;
+	virtual void LoadedFromAnotherClass(const FName& OldClassName) override;
+	virtual void PostEditChangeProperty(FPropertyChangedEvent& PropertyChangedEvent) override;
 #endif
 	// End UObject Interface
 
diff --git a/Engine/Source/Runtime/Engine/Classes/Engine/StaticMesh.h b/Engine/Source/Runtime/Engine/Classes/Engine/StaticMesh.h
--- a/Engine/Source/Runtime/Engine/Classes/Engine/StaticMesh.h
+++ b/Engine/Source/Runtime/Engine/Classes/Engine/StaticMesh.h
@@ -465,16 +465,16 @@ public:
 
 	// Begin UObject interface.
 #if WITH_EDITOR
-	ENGINE_API virtual void PreEditChange(UProperty* PropertyAboutToChange) OVERRIDE;
-	ENGINE_API virtual void PostEditChangeProperty(FPropertyChangedEvent& PropertyChangedEvent) OVERRIDE;
+	ENGINE_API virtual void PreEditChange(UProperty* PropertyAboutToChange) override;
+	ENGINE_API virtual void PostEditChangeProperty(FPropertyChangedEvent& PropertyChangedEvent) override;
 #endif // WITH_EDITOR
-	ENGINE_API virtual void Serialize(FArchive& Ar) OVERRIDE;
-	ENGINE_API virtual void PostLoad() OVERRIDE;
-	ENGINE_API virtual void BeginDestroy() OVERRIDE;
-	ENGINE_API virtual bool IsReadyForFinishDestroy() OVERRIDE;
-	ENGINE_API virtual void GetAssetRegistryTags(TArray<FAssetRegistryTag>& OutTags) const OVERRIDE;
-	ENGINE_API virtual FString GetDesc() OVERRIDE;
-	ENGINE_API virtual SIZE_T GetResourceSize(EResourceSizeMode::Type Mode) OVERRIDE;
+	ENGINE_API virtual void Serialize(FArchive& Ar) override;
+	ENGINE_API virtual void PostLoad() override;
+	ENGINE_API virtual void BeginDestroy() override;
+	ENGINE_API virtual bool IsReadyForFinishDestroy() override;
+	ENGINE_API virtual void GetAssetRegistryTags(TArray<FAssetRegistryTag>& OutTags) const override;
+	ENGINE_API virtual FString GetDesc() override;
+	ENGINE_API virtual SIZE_T GetResourceSize(EResourceSizeMode::Type Mode) override;
 	static void AddReferencedObjects(UObject* InThis, FReferenceCollector& Collector);
 	// End UObject interface.
 
@@ -547,20 +547,20 @@ public:
 	ENGINE_API static void CheckLightMapUVs( UStaticMesh* InStaticMesh, TArray< FString >& InOutAssetsWithMissingUVSets, TArray< FString >& InOutAssetsWithBadUVSets, TArray< FString >& InOutAssetsWithValidUVSets, bool bInVerbose = true );
 
 	// Begin Interface_CollisionDataProvider Interface
-	ENGINE_API virtual bool GetPhysicsTriMeshData(struct FTriMeshCollisionData* CollisionData, bool InUseAllTriData) OVERRIDE;
-	ENGINE_API virtual bool ContainsPhysicsTriMeshData(bool InUseAllTriData) const OVERRIDE;
-	virtual bool WantsNegXTriMesh() OVERRIDE
+	ENGINE_API virtual bool GetPhysicsTriMeshData(struct FTriMeshCollisionData* CollisionData, bool InUseAllTriData) override;
+	ENGINE_API virtual bool ContainsPhysicsTriMeshData(bool InUseAllTriData) const override;
+	virtual bool WantsNegXTriMesh() override
 	{
 		return true;
 	}
-	ENGINE_API virtual void GetMeshId(FString& OutMeshId) OVERRIDE;
+	ENGINE_API virtual void GetMeshId(FString& OutMeshId) override;
 	// End Interface_CollisionDataProvider Interface
 
 	// Begin IInterface_AssetUserData Interface
-	virtual void AddAssetUserData(UAssetUserData* InUserData) OVERRIDE;
-	virtual void RemoveUserDataOfClass(TSubclassOf<UAssetUserData> InUserDataClass) OVERRIDE;
-	virtual UAssetUserData* GetAssetUserDataOfClass(TSubclassOf<UAssetUserData> InUserDataClass) OVERRIDE;
-	virtual const TArray<UAssetUserData*>* GetAssetUserDataArray() const OVERRIDE;
+	virtual void AddAssetUserData(UAssetUserData* InUserData) override;
+	virtual void RemoveUserDataOfClass(TSubclassOf<UAssetUserData> InUserDataClass) override;
+	virtual UAssetUserData* GetAssetUserDataOfClass(TSubclassOf<UAssetUserData> InUserDataClass) override;
+	virtual const TArray<UAssetUserData*>* GetAssetUserDataArray() const override;
 	// End IInterface_AssetUserData Interface
 
 
diff --git a/Engine/Source/Runtime/Engine/Classes/Engine/StaticMeshActor.h b/Engine/Source/Runtime/Engine/Classes/Engine/StaticMeshActor.h
--- a/Engine/Source/Runtime/Engine/Classes/Engine/StaticMeshActor.h
+++ b/Engine/Source/Runtime/Engine/Classes/Engine/StaticMeshActor.h
@@ -16,7 +16,7 @@ class ENGINE_API AStaticMeshActor : public AActor, public INavRelevantActorInter
 	UPROPERTY(Category=StaticMeshActor, VisibleAnywhere, BlueprintReadOnly, meta=(ExposeFunctionCategories="Mesh,Rendering,Physics,Components|StaticMesh"))
 	TSubobjectPtr<class UStaticMeshComponent> StaticMeshComponent;
 	
-	virtual void BeginPlay() OVERRIDE;
+	virtual void BeginPlay() override;
 
 	/** This static mesh should replicate movement. Automatically sets the RemoteRole and bReplicateMovement flags. Meant to be edited on placed actors (those other two proeprties are not) */
 	UPROPERTY(Category=Actor, EditAnywhere, AdvancedDisplay)
@@ -27,17 +27,17 @@ class ENGINE_API AStaticMeshActor : public AActor, public INavRelevantActorInter
 
 #if WITH_EDITOR
 	// Begin AActor Interface
-	virtual void CheckForErrors() OVERRIDE;
-	virtual bool GetReferencedContentObjects( TArray<UObject*>& Objects ) const OVERRIDE;
+	virtual void CheckForErrors() override;
+	virtual bool GetReferencedContentObjects( TArray<UObject*>& Objects ) const override;
 	// End AActor Interface
 #endif // WITH_EDITOR	
 
 protected:
 	// Begin UObject interface.
-	virtual FString GetDetailedInfoInternal() const OVERRIDE;
+	virtual FString GetDetailedInfoInternal() const override;
 #if WITH_EDITOR
-	virtual void LoadedFromAnotherClass(const FName& OldClassName) OVERRIDE;
-	virtual void PostEditChangeChainProperty(FPropertyChangedChainEvent& PropertyChangedEvent) OVERRIDE;
+	virtual void LoadedFromAnotherClass(const FName& OldClassName) override;
+	virtual void PostEditChangeChainProperty(FPropertyChangedChainEvent& PropertyChangedEvent) override;
 #endif // WITH_EDITOR	
 	// End UObject interface.
 };
diff --git a/Engine/Source/Runtime/Engine/Classes/Engine/StaticMeshSocket.h b/Engine/Source/Runtime/Engine/Classes/Engine/StaticMeshSocket.h
--- a/Engine/Source/Runtime/Engine/Classes/Engine/StaticMeshSocket.h
+++ b/Engine/Source/Runtime/Engine/Classes/Engine/StaticMeshSocket.h
@@ -55,7 +55,7 @@ public:
 	FChangedEvent& OnPropertyChanged() { return ChangedEvent; }
 
 	// Begin UObject interface
-	virtual void PostEditChangeProperty( FPropertyChangedEvent& PropertyChangedEvent ) OVERRIDE;
+	virtual void PostEditChangeProperty( FPropertyChangedEvent& PropertyChangedEvent ) override;
 	// End UObject interface
 
 private:
diff --git a/Engine/Source/Runtime/Engine/Classes/Engine/Texture.h b/Engine/Source/Runtime/Engine/Classes/Engine/Texture.h
--- a/Engine/Source/Runtime/Engine/Classes/Engine/Texture.h
+++ b/Engine/Source/Runtime/Engine/Classes/Engine/Texture.h
@@ -586,7 +586,7 @@ public:
 	/**
 	 * Begins caching platform data in the background for the platform requested
 	 */
-	virtual void BeginCacheForCookedPlatformData(  const ITargetPlatform *TargetPlatform ) OVERRIDE;
+	virtual void BeginCacheForCookedPlatformData(  const ITargetPlatform *TargetPlatform ) override;
 
 
 	/**
@@ -625,14 +625,14 @@ public:
 
 	// Begin UObject interface.
 #if WITH_EDITOR
-	virtual void PostEditChangeProperty(FPropertyChangedEvent& PropertyChangedEvent) OVERRIDE;
+	virtual void PostEditChangeProperty(FPropertyChangedEvent& PropertyChangedEvent) override;
 #endif // WITH_EDITOR
-	virtual void Serialize(FArchive& Ar) OVERRIDE;
-	virtual void PostLoad() OVERRIDE;
-	virtual void PreSave() OVERRIDE;
-	virtual void BeginDestroy() OVERRIDE;
-	virtual bool IsReadyForFinishDestroy() OVERRIDE;
-	virtual void FinishDestroy() OVERRIDE;
+	virtual void Serialize(FArchive& Ar) override;
+	virtual void PostLoad() override;
+	virtual void PreSave() override;
+	virtual void BeginDestroy() override;
+	virtual bool IsReadyForFinishDestroy() override;
+	virtual void FinishDestroy() override;
 	// End UObject interface.
 
 	/**
diff --git a/Engine/Source/Runtime/Engine/Classes/Engine/Texture2D.h b/Engine/Source/Runtime/Engine/Classes/Engine/Texture2D.h
--- a/Engine/Source/Runtime/Engine/Classes/Engine/Texture2D.h
+++ b/Engine/Source/Runtime/Engine/Classes/Engine/Texture2D.h
@@ -158,28 +158,28 @@ public:
 public:
 
 	// Begin UObject interface.
-	virtual void Serialize(FArchive& Ar) OVERRIDE;	
+	virtual void Serialize(FArchive& Ar) override;	
 #if WITH_EDITOR
-	virtual void CookerWillNeverCookAgain() OVERRIDE;
-	virtual void PostLinkerChange() OVERRIDE;
-	virtual void PostEditChangeProperty(FPropertyChangedEvent& PropertyChangedEvent) OVERRIDE;
+	virtual void CookerWillNeverCookAgain() override;
+	virtual void PostLinkerChange() override;
+	virtual void PostEditChangeProperty(FPropertyChangedEvent& PropertyChangedEvent) override;
 #endif // WITH_EDITOR
-	virtual void BeginDestroy() OVERRIDE;
-	virtual void PostLoad() OVERRIDE;
-	virtual void PreSave() OVERRIDE;
-	virtual void GetAssetRegistryTags(TArray<FAssetRegistryTag>& OutTags) const OVERRIDE;
-	virtual FString GetDesc() OVERRIDE;
+	virtual void BeginDestroy() override;
+	virtual void PostLoad() override;
+	virtual void PreSave() override;
+	virtual void GetAssetRegistryTags(TArray<FAssetRegistryTag>& OutTags) const override;
+	virtual FString GetDesc() override;
 	// End UObject interface.
 
 	// Begin UTexture interface.
-	virtual float GetSurfaceWidth() const OVERRIDE { return GetSizeX(); }
-	virtual float GetSurfaceHeight() const OVERRIDE { return GetSizeY(); }
-	virtual FTextureResource* CreateResource() OVERRIDE;
-	virtual EMaterialValueType GetMaterialType() OVERRIDE { return MCT_Texture2D; }
-	virtual void UpdateResource() OVERRIDE;
-	virtual float GetAverageBrightness(bool bIgnoreTrueBlack, bool bUseGrayscale) OVERRIDE;
-	virtual FTexturePlatformData** GetRunningPlatformData() OVERRIDE { return &PlatformData; }
-	virtual TMap<FString,FTexturePlatformData*>* GetCookedPlatformData() OVERRIDE { return &CookedPlatformData; }
+	virtual float GetSurfaceWidth() const override { return GetSizeX(); }
+	virtual float GetSurfaceHeight() const override { return GetSizeY(); }
+	virtual FTextureResource* CreateResource() override;
+	virtual EMaterialValueType GetMaterialType() override { return MCT_Texture2D; }
+	virtual void UpdateResource() override;
+	virtual float GetAverageBrightness(bool bIgnoreTrueBlack, bool bUseGrayscale) override;
+	virtual FTexturePlatformData** GetRunningPlatformData() override { return &PlatformData; }
+	virtual TMap<FString,FTexturePlatformData*>* GetCookedPlatformData() override { return &CookedPlatformData; }
 	// End UTexture interface.
 
 	/** Trivial accessors. */
@@ -287,7 +287,7 @@ public:
 	 * @param	Enum	Which mips to calculate size for.
 	 * @return	Total size of all specified mips, in bytes
 	 */
-	virtual uint32 CalcTextureMemorySizeEnum( ETextureMipCount Enum ) const OVERRIDE;
+	virtual uint32 CalcTextureMemorySizeEnum( ETextureMipCount Enum ) const override;
 
 	/**
 	 *	Get the CRC of the source art pixels.
@@ -346,7 +346,7 @@ public:
 	 *
 	 * @return size of resource as to be displayed to artists/ LDs in the Editor.
 	 */
-	virtual SIZE_T GetResourceSize(EResourceSizeMode::Type Mode) OVERRIDE;
+	virtual SIZE_T GetResourceSize(EResourceSizeMode::Type Mode) override;
 
 	/**
 	 * Returns whether miplevels should be forced resident.
diff --git a/Engine/Source/Runtime/Engine/Classes/Engine/Texture2DDynamic.h b/Engine/Source/Runtime/Engine/Classes/Engine/Texture2DDynamic.h
--- a/Engine/Source/Runtime/Engine/Classes/Engine/Texture2DDynamic.h
+++ b/Engine/Source/Runtime/Engine/Classes/Engine/Texture2DDynamic.h
@@ -28,10 +28,10 @@ class UTexture2DDynamic : public UTexture
 
 public:
 	// Begin UTexture Interface.
-	virtual FTextureResource* CreateResource() OVERRIDE;
-	virtual EMaterialValueType GetMaterialType() OVERRIDE { return MCT_Texture2D; }
-	virtual float GetSurfaceWidth() const OVERRIDE;
-	virtual float GetSurfaceHeight() const OVERRIDE;
+	virtual FTextureResource* CreateResource() override;
+	virtual EMaterialValueType GetMaterialType() override { return MCT_Texture2D; }
+	virtual float GetSurfaceWidth() const override;
+	virtual float GetSurfaceHeight() const override;
 	// End UTexture Interface.
 	
 	/**
diff --git a/Engine/Source/Runtime/Engine/Classes/Engine/TextureCube.h b/Engine/Source/Runtime/Engine/Classes/Engine/TextureCube.h
--- a/Engine/Source/Runtime/Engine/Classes/Engine/TextureCube.h
+++ b/Engine/Source/Runtime/Engine/Classes/Engine/TextureCube.h
@@ -15,11 +15,11 @@ public:
 	TMap<FString, FTexturePlatformData*> CookedPlatformData;
 
 	// Begin UObject interface.
-	virtual void Serialize(FArchive& Ar) OVERRIDE;
-	virtual void PostLoad() OVERRIDE;
-	virtual void GetAssetRegistryTags(TArray<FAssetRegistryTag>& OutTags) const OVERRIDE;
-	virtual FString GetDesc() OVERRIDE;
-	virtual SIZE_T GetResourceSize(EResourceSizeMode::Type Mode) OVERRIDE;
+	virtual void Serialize(FArchive& Ar) override;
+	virtual void PostLoad() override;
+	virtual void GetAssetRegistryTags(TArray<FAssetRegistryTag>& OutTags) const override;
+	virtual FString GetDesc() override;
+	virtual SIZE_T GetResourceSize(EResourceSizeMode::Type Mode) override;
 	// End UObject interface.
 
 	/** Trivial accessors. */
@@ -57,13 +57,13 @@ public:
 	}
 
 	// Begin UTexture interface
-	virtual float GetSurfaceWidth() const OVERRIDE { return GetSizeX(); }
-	virtual float GetSurfaceHeight() const OVERRIDE { return GetSizeY(); }
-	virtual FTextureResource* CreateResource() OVERRIDE;
-	virtual void UpdateResource() OVERRIDE;
-	virtual EMaterialValueType GetMaterialType() OVERRIDE { return MCT_TextureCube; }
-	virtual FTexturePlatformData** GetRunningPlatformData() OVERRIDE { return &PlatformData; }
-	virtual TMap<FString, FTexturePlatformData*> *GetCookedPlatformData() OVERRIDE { return &CookedPlatformData; }
+	virtual float GetSurfaceWidth() const override { return GetSizeX(); }
+	virtual float GetSurfaceHeight() const override { return GetSizeY(); }
+	virtual FTextureResource* CreateResource() override;
+	virtual void UpdateResource() override;
+	virtual EMaterialValueType GetMaterialType() override { return MCT_TextureCube; }
+	virtual FTexturePlatformData** GetRunningPlatformData() override { return &PlatformData; }
+	virtual TMap<FString, FTexturePlatformData*> *GetCookedPlatformData() override { return &CookedPlatformData; }
 	// End UTexture interface
 	
 	/**
@@ -80,7 +80,7 @@ public:
 	 * @param	Enum	Which mips to calculate size for.
 	 * @return	Total size of all specified mips, in bytes
 	 */
-	virtual uint32 CalcTextureMemorySizeEnum( ETextureMipCount Enum ) const OVERRIDE;
+	virtual uint32 CalcTextureMemorySizeEnum( ETextureMipCount Enum ) const override;
 };
 
 
diff --git a/Engine/Source/Runtime/Engine/Classes/Engine/TextureMovie.h b/Engine/Source/Runtime/Engine/Classes/Engine/TextureMovie.h
--- a/Engine/Source/Runtime/Engine/Classes/Engine/TextureMovie.h
+++ b/Engine/Source/Runtime/Engine/Classes/Engine/TextureMovie.h
@@ -85,24 +85,24 @@ public:
 	FByteBulkData Data;
 
 	// Begin UTexture interface.
-	virtual float GetSurfaceWidth() const OVERRIDE { return SizeX; }
-	virtual float GetSurfaceHeight() const OVERRIDE { return SizeY; }
-	virtual FTextureResource* CreateResource() OVERRIDE;
-	virtual EMaterialValueType GetMaterialType() OVERRIDE;
+	virtual float GetSurfaceWidth() const override { return SizeX; }
+	virtual float GetSurfaceHeight() const override { return SizeY; }
+	virtual FTextureResource* CreateResource() override;
+	virtual EMaterialValueType GetMaterialType() override;
 	// End UTexture interface.
 	
 	// Begin UObject interface.
-	virtual void Serialize(FArchive& Ar) OVERRIDE;	
-	virtual void PostLoad() OVERRIDE;
-	virtual void BeginDestroy() OVERRIDE;
+	virtual void Serialize(FArchive& Ar) override;	
+	virtual void PostLoad() override;
+	virtual void BeginDestroy() override;
 #if WITH_EDITOR
-	virtual void PreEditChange(UProperty* PropertyAboutToChange) OVERRIDE;
-	virtual void PostEditChangeProperty(FPropertyChangedEvent& PropertyChangedEvent) OVERRIDE;
+	virtual void PreEditChange(UProperty* PropertyAboutToChange) override;
+	virtual void PostEditChangeProperty(FPropertyChangedEvent& PropertyChangedEvent) override;
 #endif // WITH_EDITOR
-	virtual bool IsReadyForFinishDestroy() OVERRIDE;
-	virtual void FinishDestroy() OVERRIDE;
-	virtual FString GetDesc() OVERRIDE;
-	virtual SIZE_T GetResourceSize(EResourceSizeMode::Type Mode) OVERRIDE;
+	virtual bool IsReadyForFinishDestroy() override;
+	virtual void FinishDestroy() override;
+	virtual FString GetDesc() override;
+	virtual SIZE_T GetResourceSize(EResourceSizeMode::Type Mode) override;
 	// End UObject interface.
 
 	/**
diff --git a/Engine/Source/Runtime/Engine/Classes/Engine/TextureRenderTarget2D.h b/Engine/Source/Runtime/Engine/Classes/Engine/TextureRenderTarget2D.h
--- a/Engine/Source/Runtime/Engine/Classes/Engine/TextureRenderTarget2D.h
+++ b/Engine/Source/Runtime/Engine/Classes/Engine/TextureRenderTarget2D.h
@@ -75,19 +75,19 @@ class UTextureRenderTarget2D : public UTextureRenderTarget
 	ENGINE_API void UpdateResourceImmediate();
 
 	// Begin UTexture interface.
-	virtual float GetSurfaceWidth() const OVERRIDE { return SizeX; }
-	virtual float GetSurfaceHeight() const OVERRIDE { return SizeY; }
-	virtual FTextureResource* CreateResource() OVERRIDE;
-	virtual EMaterialValueType GetMaterialType() OVERRIDE;
+	virtual float GetSurfaceWidth() const override { return SizeX; }
+	virtual float GetSurfaceHeight() const override { return SizeY; }
+	virtual FTextureResource* CreateResource() override;
+	virtual EMaterialValueType GetMaterialType() override;
 	// End UTexture interface.
 
 	// Begin UObject interface
 #if WITH_EDITOR
-	virtual void PostEditChangeProperty(FPropertyChangedEvent& PropertyChangedEvent) OVERRIDE;
+	virtual void PostEditChangeProperty(FPropertyChangedEvent& PropertyChangedEvent) override;
 #endif // WITH_EDITOR
-	virtual void PostLoad() OVERRIDE;
-	virtual SIZE_T GetResourceSize(EResourceSizeMode::Type Mode) OVERRIDE;
-	virtual FString GetDesc() OVERRIDE;
+	virtual void PostLoad() override;
+	virtual SIZE_T GetResourceSize(EResourceSizeMode::Type Mode) override;
+	virtual FString GetDesc() override;
 	// End UObject interface
 
 	FORCEINLINE int32 GetNumMips() const
diff --git a/Engine/Source/Runtime/Engine/Classes/Engine/TextureRenderTargetCube.h b/Engine/Source/Runtime/Engine/Classes/Engine/TextureRenderTargetCube.h
--- a/Engine/Source/Runtime/Engine/Classes/Engine/TextureRenderTargetCube.h
+++ b/Engine/Source/Runtime/Engine/Classes/Engine/TextureRenderTargetCube.h
@@ -58,10 +58,10 @@ class UTextureRenderTargetCube : public UTextureRenderTarget
 	ENGINE_API class UTextureCube* ConstructTextureCube(UObject* Outer, const FString& NewTexName, EObjectFlags InFlags);
 
 	// Begin UTexture interface.
-	virtual float GetSurfaceWidth() const  OVERRIDE { return SizeX; }
-	virtual float GetSurfaceHeight()const  OVERRIDE { return SizeX; }	
-	virtual FTextureResource* CreateResource() OVERRIDE;
-	virtual EMaterialValueType GetMaterialType() OVERRIDE;
+	virtual float GetSurfaceWidth() const  override { return SizeX; }
+	virtual float GetSurfaceHeight()const  override { return SizeX; }	
+	virtual FTextureResource* CreateResource() override;
+	virtual EMaterialValueType GetMaterialType() override;
 	// End UTexture interface.
 
 	EPixelFormat GetFormat() const
@@ -82,11 +82,11 @@ class UTextureRenderTargetCube : public UTextureRenderTarget
 	}
 	// Begin UObject interface
 #if WITH_EDITOR
-	virtual void PostEditChangeProperty(FPropertyChangedEvent& PropertyChangedEvent) OVERRIDE;
+	virtual void PostEditChangeProperty(FPropertyChangedEvent& PropertyChangedEvent) override;
 #endif // WITH_EDITOR
-	virtual void PostLoad() OVERRIDE;
-	virtual SIZE_T GetResourceSize(EResourceSizeMode::Type Mode) OVERRIDE;
-	virtual FString GetDesc() OVERRIDE;
+	virtual void PostLoad() override;
+	virtual SIZE_T GetResourceSize(EResourceSizeMode::Type Mode) override;
+	virtual FString GetDesc() override;
 	// Begin UObject interface
 };
 
diff --git a/Engine/Source/Runtime/Engine/Classes/Engine/TimelineTemplate.h b/Engine/Source/Runtime/Engine/Classes/Engine/TimelineTemplate.h
--- a/Engine/Source/Runtime/Engine/Classes/Engine/TimelineTemplate.h
+++ b/Engine/Source/Runtime/Engine/Classes/Engine/TimelineTemplate.h
@@ -206,7 +206,7 @@ class UTimelineTemplate : public UObject
 	ENGINE_API void GetAllCurves(TSet<class UCurveBase*>& InOutCurves) const;
 
 	/* Perform deep copy of curves */
-	virtual void PostDuplicate(bool bDuplicateForPIE) OVERRIDE;
+	virtual void PostDuplicate(bool bDuplicateForPIE) override;
 };
 
 
diff --git a/Engine/Source/Runtime/Engine/Classes/Engine/TriggerVolume.h b/Engine/Source/Runtime/Engine/Classes/Engine/TriggerVolume.h
--- a/Engine/Source/Runtime/Engine/Classes/Engine/TriggerVolume.h
+++ b/Engine/Source/Runtime/Engine/Classes/Engine/TriggerVolume.h
@@ -12,7 +12,7 @@ class ATriggerVolume : public AVolume
 
 	// Begin UObject interface.
 #if WITH_EDITOR
-	virtual void LoadedFromAnotherClass(const FName& OldClassName) OVERRIDE;
+	virtual void LoadedFromAnotherClass(const FName& OldClassName) override;
 #endif // WITH_EDITOR	
 	// End UObject interface.
 };
diff --git a/Engine/Source/Runtime/Engine/Classes/Engine/UserDefinedEnum.h b/Engine/Source/Runtime/Engine/Classes/Engine/UserDefinedEnum.h
--- a/Engine/Source/Runtime/Engine/Classes/Engine/UserDefinedEnum.h
+++ b/Engine/Source/Runtime/Engine/Classes/Engine/UserDefinedEnum.h
@@ -41,7 +41,7 @@ public:
 	 * @param InEnumName Enum name.
 	 * @return Full enum name.
 	 */
-	virtual FString GenerateFullEnumName(const TCHAR* InEnumName) const OVERRIDE;
+	virtual FString GenerateFullEnumName(const TCHAR* InEnumName) const override;
 
 	/*
 	 *	Try to update value in enum variable after an enum's change.
@@ -49,19 +49,19 @@ public:
 	 *	@param EnumeratorIndex	old index
 	 *	@return	new index
 	 */
-	virtual int32 ResolveEnumerator(FArchive& Ar, int32 EnumeratorIndex) const OVERRIDE;
+	virtual int32 ResolveEnumerator(FArchive& Ar, int32 EnumeratorIndex) const override;
 
 	/**
 	 * @return	The enum string at the specified index.
 	 */
-	virtual FText GetEnumText(int32 InIndex) const OVERRIDE;
+	virtual FText GetEnumText(int32 InIndex) const override;
 
 #if WITH_EDITOR
 	// Begin UObject interface
-	virtual bool Rename(const TCHAR* NewName = NULL, UObject* NewOuter = NULL, ERenameFlags Flags = REN_None) OVERRIDE;
-	virtual void PostDuplicate(bool bDuplicateForPIE) OVERRIDE;
-	virtual void PostLoad() OVERRIDE;
-	virtual void PostEditUndo() OVERRIDE;
+	virtual bool Rename(const TCHAR* NewName = NULL, UObject* NewOuter = NULL, ERenameFlags Flags = REN_None) override;
+	virtual void PostDuplicate(bool bDuplicateForPIE) override;
+	virtual void PostLoad() override;
+	virtual void PostEditUndo() override;
 	// End of UObject interface
 
 	FString GenerateNewEnumeratorName();
diff --git a/Engine/Source/Runtime/Engine/Classes/Engine/UserDefinedStruct.h b/Engine/Source/Runtime/Engine/Classes/Engine/UserDefinedStruct.h
--- a/Engine/Source/Runtime/Engine/Classes/Engine/UserDefinedStruct.h
+++ b/Engine/Source/Runtime/Engine/Classes/Engine/UserDefinedStruct.h
@@ -42,15 +42,15 @@ public:
 
 #if WITH_EDITOR
 	// UObject interface.
-	virtual void Serialize(FArchive& Ar) OVERRIDE;
-	virtual void PostDuplicate(bool bDuplicateForPIE) OVERRIDE;
-	virtual void GetAssetRegistryTags(TArray<FAssetRegistryTag>& OutTags) const OVERRIDE;
+	virtual void Serialize(FArchive& Ar) override;
+	virtual void PostDuplicate(bool bDuplicateForPIE) override;
+	virtual void GetAssetRegistryTags(TArray<FAssetRegistryTag>& OutTags) const override;
 	// End of UObject interface.
 #endif	// WITH_EDITOR
 
-	virtual void SerializeTaggedProperties( FArchive& Ar, uint8* Data, UStruct* DefaultsStruct, uint8* Defaults ) const OVERRIDE;
+	virtual void SerializeTaggedProperties( FArchive& Ar, uint8* Data, UStruct* DefaultsStruct, uint8* Defaults ) const override;
 
 	// UScriptStruct interface.
-	virtual void RecursivelyPreload() OVERRIDE;
+	virtual void RecursivelyPreload() override;
 	// End of  UScriptStruct interface.
 };
diff --git a/Engine/Source/Runtime/Engine/Classes/Engine/World.h b/Engine/Source/Runtime/Engine/Classes/Engine/World.h
--- a/Engine/Source/Runtime/Engine/Classes/Engine/World.h
+++ b/Engine/Source/Runtime/Engine/Classes/Engine/World.h
@@ -326,7 +326,7 @@ struct FStartPhysicsTickFunction : public FTickFunction
 		* @param CurrentThread - thread we are executing on, useful to pass along as new tasks are created
 		* @param MyCompletionGraphEvent - completion event for this task. Useful for holding the completetion of this task until certain child tasks are complete.
 	**/
-	virtual void ExecuteTick(float DeltaTime, enum ELevelTick TickType, ENamedThreads::Type CurrentThread, const FGraphEventRef& MyCompletionGraphEvent) OVERRIDE;
+	virtual void ExecuteTick(float DeltaTime, enum ELevelTick TickType, ENamedThreads::Type CurrentThread, const FGraphEventRef& MyCompletionGraphEvent) override;
 	/** Abstract function to describe this tick. Used to print messages about illegal cycles in the dependency graph **/
 	virtual FString DiagnosticMessage();
 };
@@ -349,7 +349,7 @@ struct FEndPhysicsTickFunction : public FTickFunction
 		* @param CurrentThread - thread we are executing on, useful to pass along as new tasks are created
 		* @param MyCompletionGraphEvent - completion event for this task. Useful for holding the completetion of this task until certain child tasks are complete.
 	**/
-	virtual void ExecuteTick(float DeltaTime, enum ELevelTick TickType, ENamedThreads::Type CurrentThread, const FGraphEventRef& MyCompletionGraphEvent) OVERRIDE;
+	virtual void ExecuteTick(float DeltaTime, enum ELevelTick TickType, ENamedThreads::Type CurrentThread, const FGraphEventRef& MyCompletionGraphEvent) override;
 	/** Abstract function to describe this tick. Used to print messages about illegal cycles in the dependency graph **/
 	virtual FString DiagnosticMessage();
 };
@@ -372,7 +372,7 @@ struct FStartClothSimulationFunction : public FTickFunction
 	* @param CurrentThread - thread we are executing on, useful to pass along as new tasks are created
 	* @param MyCompletionGraphEvent - completion event for this task. Useful for holding the completetion of this task until certain child tasks are complete.
 	**/
-	virtual void ExecuteTick(float DeltaTime, enum ELevelTick TickType, ENamedThreads::Type CurrentThread, const FGraphEventRef& MyCompletionGraphEvent) OVERRIDE;
+	virtual void ExecuteTick(float DeltaTime, enum ELevelTick TickType, ENamedThreads::Type CurrentThread, const FGraphEventRef& MyCompletionGraphEvent) override;
 	/** Abstract function to describe this tick. Used to print messages about illegal cycles in the dependency graph **/
 	virtual FString DiagnosticMessage();
 };
@@ -395,7 +395,7 @@ struct FEndClothSimulationFunction : public FTickFunction
 	* @param CurrentThread - thread we are executing on, useful to pass along as new tasks are created
 	* @param MyCompletionGraphEvent - completion event for this task. Useful for holding the completetion of this task until certain child tasks are complete.
 	**/
-	virtual void ExecuteTick(float DeltaTime, enum ELevelTick TickType, ENamedThreads::Type CurrentThread, const FGraphEventRef& MyCompletionGraphEvent) OVERRIDE;
+	virtual void ExecuteTick(float DeltaTime, enum ELevelTick TickType, ENamedThreads::Type CurrentThread, const FGraphEventRef& MyCompletionGraphEvent) override;
 	/** Abstract function to describe this tick. Used to print messages about illegal cycles in the dependency graph **/
 	virtual FString DiagnosticMessage();
 };
@@ -1611,17 +1611,17 @@ public:
 	virtual bool AllowAudioPlayback();
 
 	// Begin UObject Interface
-	virtual void Serialize( FArchive& Ar ) OVERRIDE;
-	virtual void FinishDestroy() OVERRIDE;
-	virtual void PostLoad() OVERRIDE;
-	virtual bool PreSaveRoot(const TCHAR* Filename, TArray<FString>& AdditionalPackagesToCook) OVERRIDE;
-	virtual void PostSaveRoot( bool bCleanupIsRequired ) OVERRIDE;
-	virtual UWorld* GetWorld() const OVERRIDE;
+	virtual void Serialize( FArchive& Ar ) override;
+	virtual void FinishDestroy() override;
+	virtual void PostLoad() override;
+	virtual bool PreSaveRoot(const TCHAR* Filename, TArray<FString>& AdditionalPackagesToCook) override;
+	virtual void PostSaveRoot( bool bCleanupIsRequired ) override;
+	virtual UWorld* GetWorld() const override;
 	static void AddReferencedObjects(UObject* InThis, FReferenceCollector& Collector);
 #if WITH_EDITOR
-	virtual bool Rename(const TCHAR* NewName = NULL, UObject* NewOuter = NULL, ERenameFlags Flags = REN_None) OVERRIDE;
+	virtual bool Rename(const TCHAR* NewName = NULL, UObject* NewOuter = NULL, ERenameFlags Flags = REN_None) override;
 #endif
-	virtual void PostDuplicate(bool bDuplicateForPIE) OVERRIDE;
+	virtual void PostDuplicate(bool bDuplicateForPIE) override;
 
 	// End UObject Interface
 	
@@ -2169,12 +2169,12 @@ public:
 	bool DestroySwappedPC(UNetConnection* Connection);
 
 	// Begin FNetworkNotify interface
-	virtual EAcceptConnection::Type NotifyAcceptingConnection() OVERRIDE;
+	virtual EAcceptConnection::Type NotifyAcceptingConnection() override;
 
 	/** @todo document */
-	virtual void NotifyAcceptedConnection( class UNetConnection* Connection ) OVERRIDE;
-	virtual bool NotifyAcceptingChannel( class UChannel* Channel ) OVERRIDE;
-	virtual void NotifyControlMessage(UNetConnection* Connection, uint8 MessageType, class FInBunch& Bunch) OVERRIDE;
+	virtual void NotifyAcceptedConnection( class UNetConnection* Connection ) override;
+	virtual bool NotifyAcceptingChannel( class UChannel* Channel ) override;
+	virtual void NotifyControlMessage(UNetConnection* Connection, uint8 MessageType, class FInBunch& Bunch) override;
 	// End FNetworkNotify interface
 
 	/** Welcome a new player joining this server. */
diff --git a/Engine/Source/Runtime/Engine/Classes/Engine/WorldComposition.h b/Engine/Source/Runtime/Engine/Classes/Engine/WorldComposition.h
--- a/Engine/Source/Runtime/Engine/Classes/Engine/WorldComposition.h
+++ b/Engine/Source/Runtime/Engine/Classes/Engine/WorldComposition.h
@@ -136,10 +136,10 @@ class ENGINE_API UWorldComposition : public UObject
 
 private:
 	// UObject interface
-	virtual void PostInitProperties() OVERRIDE;
-	virtual void Serialize(FArchive& Ar) OVERRIDE;
-	virtual void PostDuplicate(bool bDuplicateForPIE) OVERRIDE;
-	virtual void PostLoad() OVERRIDE;
+	virtual void PostInitProperties() override;
+	virtual void Serialize(FArchive& Ar) override;
+	virtual void PostDuplicate(bool bDuplicateForPIE) override;
+	virtual void PostLoad() override;
 	// UObject interface
 
 	/** Populate streaming level objects using tiles information */
diff --git a/Engine/Source/Runtime/Engine/Classes/Exporters/Exporter.h b/Engine/Source/Runtime/Engine/Classes/Exporters/Exporter.h
--- a/Engine/Source/Runtime/Engine/Classes/Exporters/Exporter.h
+++ b/Engine/Source/Runtime/Engine/Classes/Exporters/Exporter.h
@@ -59,7 +59,7 @@ class UExporter : public UObject
 	static	const		bool	bEnableDebugBrackets;
 
 	// Begin UObject interface.
-	ENGINE_API virtual void Serialize( FArchive& Ar ) OVERRIDE;
+	ENGINE_API virtual void Serialize( FArchive& Ar ) override;
 	// End UObject interface.
 
 	// Returns whether this exporter supports the specific object
diff --git a/Engine/Source/Runtime/Engine/Classes/Foliage/InstancedFoliageActor.h b/Engine/Source/Runtime/Engine/Classes/Foliage/InstancedFoliageActor.h
--- a/Engine/Source/Runtime/Engine/Classes/Foliage/InstancedFoliageActor.h
+++ b/Engine/Source/Runtime/Engine/Classes/Foliage/InstancedFoliageActor.h
@@ -19,16 +19,16 @@ public:
 	TMap<class UStaticMesh*, struct FFoliageMeshInfo> FoliageMeshes;
 
 	// Begin UObject interface. 	
-	virtual void Serialize(FArchive& Ar) OVERRIDE;
-	virtual void PostLoad() OVERRIDE;
+	virtual void Serialize(FArchive& Ar) override;
+	virtual void PostLoad() override;
 	static void AddReferencedObjects(UObject* InThis, FReferenceCollector& Collector);
 	// Begin UObject interface. 
 
 	// Begin AActor interface.
 	// we don't want to have our components automatically destroyed by the Blueprint code
-	virtual void RerunConstructionScripts() OVERRIDE {}
+	virtual void RerunConstructionScripts() override {}
 	// Add world origin offset
-	virtual void ApplyWorldOffset(const FVector& InOffset, bool bWorldShift) OVERRIDE;
+	virtual void ApplyWorldOffset(const FVector& InOffset, bool bWorldShift) override;
 	// End AActor interface.
 
 #if WITH_EDITOR
