UE4 Refactoring. Changed OVERRIDE and FINAL macros to keywords override and final.

diff --git a/Engine/Source/Runtime/Engine/Classes/GameFramework/Actor.h b/Engine/Source/Runtime/Engine/Classes/GameFramework/Actor.h
--- a/Engine/Source/Runtime/Engine/Classes/GameFramework/Actor.h
+++ b/Engine/Source/Runtime/Engine/Classes/GameFramework/Actor.h
@@ -1036,22 +1036,22 @@ public:
 	FActorEndPlaySignature OnEndPlay;
 	
 	// Begin UObject Interface
-	virtual bool CheckDefaultSubobjects(bool bForceCheck = false) OVERRIDE;
-	virtual void PostInitProperties() OVERRIDE;
-	virtual bool Modify( bool bAlwaysMarkDirty=true ) OVERRIDE;
-	virtual void ProcessEvent( UFunction* Function, void* Parameters ) OVERRIDE;
-	virtual int32 GetFunctionCallspace( UFunction* Function, void* Parameters, FFrame* Stack ) OVERRIDE;
-	virtual bool CallRemoteFunction( UFunction* Function, void* Parameters, FOutParmRec* OutParms, FFrame* Stack ) OVERRIDE;
-	virtual void PostLoad() OVERRIDE;
-	virtual void PostLoadSubobjects( FObjectInstancingGraph* OuterInstanceGraph ) OVERRIDE;
-	virtual void BeginDestroy() OVERRIDE;
-	virtual bool IsReadyForFinishDestroy() OVERRIDE;
-	virtual bool Rename( const TCHAR* NewName=NULL, UObject* NewOuter=NULL, ERenameFlags Flags=REN_None ) OVERRIDE;
+	virtual bool CheckDefaultSubobjects(bool bForceCheck = false) override;
+	virtual void PostInitProperties() override;
+	virtual bool Modify( bool bAlwaysMarkDirty=true ) override;
+	virtual void ProcessEvent( UFunction* Function, void* Parameters ) override;
+	virtual int32 GetFunctionCallspace( UFunction* Function, void* Parameters, FFrame* Stack ) override;
+	virtual bool CallRemoteFunction( UFunction* Function, void* Parameters, FOutParmRec* OutParms, FFrame* Stack ) override;
+	virtual void PostLoad() override;
+	virtual void PostLoadSubobjects( FObjectInstancingGraph* OuterInstanceGraph ) override;
+	virtual void BeginDestroy() override;
+	virtual bool IsReadyForFinishDestroy() override;
+	virtual bool Rename( const TCHAR* NewName=NULL, UObject* NewOuter=NULL, ERenameFlags Flags=REN_None ) override;
 #if WITH_EDITOR
-	virtual void PreEditChange(UProperty* PropertyThatWillChange) OVERRIDE;
-	virtual void PostEditChangeProperty(FPropertyChangedEvent& PropertyChangedEvent) OVERRIDE;
-	virtual void PreEditUndo() OVERRIDE;
-	virtual void PostEditUndo() OVERRIDE;
+	virtual void PreEditChange(UProperty* PropertyThatWillChange) override;
+	virtual void PostEditChangeProperty(FPropertyChangedEvent& PropertyChangedEvent) override;
+	virtual void PreEditUndo() override;
+	virtual void PostEditUndo() override;
 
 	struct FActorRootComponentReconstructionData
 	{
@@ -1090,8 +1090,8 @@ public:
 	/* Cached pointer to the transaction annotation data from PostEditUndo to be used in the next RerunConstructionScript */
 	TSharedPtr<FActorTransactionAnnotation> CurrentTransactionAnnotation;
 
-	virtual TSharedPtr<ITransactionObjectAnnotation> GetTransactionAnnotation() const OVERRIDE;
-	virtual void PostEditUndo(TSharedPtr<ITransactionObjectAnnotation> TransactionAnnotation) OVERRIDE;
+	virtual TSharedPtr<ITransactionObjectAnnotation> GetTransactionAnnotation() const override;
+	virtual void PostEditUndo(TSharedPtr<ITransactionObjectAnnotation> TransactionAnnotation) override;
 #endif // WITH_EDITOR
 	// End UObject Interface
 
@@ -1444,19 +1444,19 @@ public:
 	virtual void LifeSpanExpired();
 
 	// Always called immediately before properties are received from the remote.
-	virtual void PreNetReceive() OVERRIDE;
+	virtual void PreNetReceive() override;
 	
 	// Always called immediately after properties are received from the remote.
-	virtual void PostNetReceive() OVERRIDE;
+	virtual void PostNetReceive() override;
 
 	/** IsNameStableForNetworking means an object can be referred to its path name (relative to outer) over the network */
-	virtual bool IsNameStableForNetworking() const OVERRIDE;
+	virtual bool IsNameStableForNetworking() const override;
 
 	/** IsSupportedForNetworking means an object can be referenced over the network */
-	virtual bool IsSupportedForNetworking() const OVERRIDE;
+	virtual bool IsSupportedForNetworking() const override;
 
 	/** Returns a list of sub-objects that have stable names for networking */
-	virtual void GetSubobjectsWithStableNamesForNetworking(TArray<UObject*> &ObjList) OVERRIDE;
+	virtual void GetSubobjectsWithStableNamesForNetworking(TArray<UObject*> &ObjList) override;
 
 	// Always called immediately after spawning and reading in replicated properties
 	virtual void PostNetInit();
@@ -1783,7 +1783,7 @@ public:
 #if !(UE_BUILD_SHIPPING || UE_BUILD_TEST)
 	private:
 	// this is the old name of the tick function. We just want to avoid mistakes with an attempt to override this
-	virtual void Tick( float DeltaTime, enum ELevelTick TickType ) FINAL
+	virtual void Tick( float DeltaTime, enum ELevelTick TickType ) final
 	{
 		check(0);
 	}
@@ -1982,7 +1982,7 @@ public:
 	bool IsInPersistentLevel(bool bIncludeLevelStreamingPersistent = false) const;
 
 	/** Getter for the cached world pointer */
-	virtual UWorld* GetWorld() const OVERRIDE;
+	virtual UWorld* GetWorld() const override;
 
 	/** Get the timer instance from the actors world */
 	class FTimerManager& GetWorldTimerManager() const;
@@ -2144,7 +2144,7 @@ FORCEINLINE FVector AActor::GetSimpleCollisionCylinderExtent() const
 	bool SetActorLocation(const FVector& NewLocation, bool bSweep=false) { return Super::SetActorLocation(NewLocation, bSweep); } \
 	bool SetActorRotation(FRotator NewRotation) { return Super::SetActorRotation(NewRotation); } \
 	bool SetActorLocationAndRotation(const FVector& NewLocation, FRotator NewRotation, bool bSweep=false) { return Super::SetActorLocationAndRotation(NewLocation, NewRotation, bSweep); } \
-	virtual bool TeleportTo( const FVector& DestLocation, const FRotator& DestRotation, bool bIsATest, bool bNoCheck ) OVERRIDE { return Super::TeleportTo(DestLocation, DestRotation, bIsATest, bNoCheck); } \
-	virtual FVector GetVelocity() const OVERRIDE { return Super::GetVelocity(); } \
+	virtual bool TeleportTo( const FVector& DestLocation, const FRotator& DestRotation, bool bIsATest, bool bNoCheck ) override { return Super::TeleportTo(DestLocation, DestRotation, bIsATest, bNoCheck); } \
+	virtual FVector GetVelocity() const override { return Super::GetVelocity(); } \
 	float GetDistanceTo(AActor* OtherActor) { return Super::GetDistanceTo(OtherActor); }
 
diff --git a/Engine/Source/Runtime/Engine/Classes/GameFramework/Character.h b/Engine/Source/Runtime/Engine/Classes/GameFramework/Character.h
--- a/Engine/Source/Runtime/Engine/Classes/GameFramework/Character.h
+++ b/Engine/Source/Runtime/Engine/Classes/GameFramework/Character.h
@@ -253,16 +253,16 @@ public:
 	float JumpMaxHoldTime;
 
 	// Begin AActor Interface.
-	virtual void TeleportSucceeded(bool bIsATest) OVERRIDE;
-	virtual void ClearCrossLevelReferences() OVERRIDE;
-	virtual void PreNetReceive() OVERRIDE;
-	virtual void PostNetReceive() OVERRIDE;
-	virtual void OnRep_ReplicatedMovement() OVERRIDE;
-	virtual void PostNetReceiveLocation() OVERRIDE;
-	virtual void GetSimpleCollisionCylinder(float& CollisionRadius, float& CollisionHalfHeight) const OVERRIDE;
-	virtual void ApplyWorldOffset(const FVector& InOffset, bool bWorldShift) OVERRIDE;
-	virtual UActorComponent* FindComponentByClass(const TSubclassOf<UActorComponent> ComponentClass) const OVERRIDE;
-	virtual void TornOff() OVERRIDE;
+	virtual void TeleportSucceeded(bool bIsATest) override;
+	virtual void ClearCrossLevelReferences() override;
+	virtual void PreNetReceive() override;
+	virtual void PostNetReceive() override;
+	virtual void OnRep_ReplicatedMovement() override;
+	virtual void PostNetReceiveLocation() override;
+	virtual void GetSimpleCollisionCylinder(float& CollisionRadius, float& CollisionHalfHeight) const override;
+	virtual void ApplyWorldOffset(const FVector& InOffset, bool bWorldShift) override;
+	virtual UActorComponent* FindComponentByClass(const TSubclassOf<UActorComponent> ComponentClass) const override;
+	virtual void TornOff() override;
 	// End AActor Interface
 
 	template<class T>
@@ -272,17 +272,17 @@ public:
 	}
 
 	// Begin APawn Interface.
-	virtual void PostInitializeComponents() OVERRIDE;
-	virtual class UPawnMovementComponent* GetMovementComponent() const OVERRIDE { return CharacterMovement; }
-	virtual class UPrimitiveComponent* GetMovementBase() const OVERRIDE FINAL { return MovementBase; }
-	virtual float GetDefaultHalfHeight() const OVERRIDE;
-	virtual void TurnOff() OVERRIDE;
-	virtual void Restart() OVERRIDE;
-	virtual void PawnClientRestart() OVERRIDE;
-	virtual void UnPossessed() OVERRIDE;
-	virtual void SetupPlayerInputComponent(class UInputComponent* InputComponent) OVERRIDE;
-	virtual void DisplayDebug(class UCanvas* Canvas, const FDebugDisplayInfo& DebugDisplay, float& YL, float& YPos) OVERRIDE;
-	virtual void RecalculateBaseEyeHeight() OVERRIDE;
+	virtual void PostInitializeComponents() override;
+	virtual class UPawnMovementComponent* GetMovementComponent() const override { return CharacterMovement; }
+	virtual class UPrimitiveComponent* GetMovementBase() const override final { return MovementBase; }
+	virtual float GetDefaultHalfHeight() const override;
+	virtual void TurnOff() override;
+	virtual void Restart() override;
+	virtual void PawnClientRestart() override;
+	virtual void UnPossessed() override;
+	virtual void SetupPlayerInputComponent(class UInputComponent* InputComponent) override;
+	virtual void DisplayDebug(class UCanvas* Canvas, const FDebugDisplayInfo& DebugDisplay, float& YL, float& YPos) override;
+	virtual void RecalculateBaseEyeHeight() override;
 	// End APawn Interface
 
 	/** Apply momentum caused by damage. */
@@ -589,5 +589,5 @@ public:
 	bool IsPlayingRootMotion() const;
 
 	/** Called on the actor right before replication occurs */
-	virtual void PreReplication( IRepChangedPropertyTracker & ChangedPropertyTracker ) OVERRIDE;
+	virtual void PreReplication( IRepChangedPropertyTracker & ChangedPropertyTracker ) override;
 };
diff --git a/Engine/Source/Runtime/Engine/Classes/GameFramework/CharacterMovementComponent.h b/Engine/Source/Runtime/Engine/Classes/GameFramework/CharacterMovementComponent.h
--- a/Engine/Source/Runtime/Engine/Classes/GameFramework/CharacterMovementComponent.h
+++ b/Engine/Source/Runtime/Engine/Classes/GameFramework/CharacterMovementComponent.h
@@ -261,10 +261,10 @@ struct FCharacterMovementComponentPreClothTickFunction : public FTickFunction
 	 * @param CurrentThread - thread we are executing on, useful to pass along as new tasks are created
 	 * @param MyCompletionGraphEvent - completion event for this task. Useful for holding the completion of this task until certain child tasks are complete.
 	 **/
-	virtual void ExecuteTick(float DeltaTime, enum ELevelTick TickType, ENamedThreads::Type CurrentThread, const FGraphEventRef& MyCompletionGraphEvent) OVERRIDE;
+	virtual void ExecuteTick(float DeltaTime, enum ELevelTick TickType, ENamedThreads::Type CurrentThread, const FGraphEventRef& MyCompletionGraphEvent) override;
 
 	/** Abstract function to describe this tick. Used to print messages about illegal cycles in the dependency graph **/
-	virtual FString DiagnosticMessage() OVERRIDE;
+	virtual FString DiagnosticMessage() override;
 };
 
 UCLASS(dependson=(UNetworkPredictionInterface))
@@ -807,24 +807,24 @@ public:
 	virtual void ApplyNetworkMovementMode(const uint8 ReceivedMode);
 
 	//Begin UActorComponent Interface
-	virtual void TickComponent(float DeltaTime, enum ELevelTick TickType, FActorComponentTickFunction *ThisTickFunction) OVERRIDE;
-	virtual void OnUnregister() OVERRIDE;
-	virtual void PostLoad() OVERRIDE;
-	virtual void RegisterComponentTickFunctions(bool bRegister) OVERRIDE;
+	virtual void TickComponent(float DeltaTime, enum ELevelTick TickType, FActorComponentTickFunction *ThisTickFunction) override;
+	virtual void OnUnregister() override;
+	virtual void PostLoad() override;
+	virtual void RegisterComponentTickFunctions(bool bRegister) override;
 	//End UActorComponent Interface
 
 	//BEGIN UMovementComponent Interface
-	virtual float GetMaxSpeed() const OVERRIDE;
-	virtual float GetMaxSpeedModifier() const OVERRIDE;
-	virtual void StopActiveMovement() OVERRIDE;
-	virtual bool IsCrouching() const OVERRIDE;
-	virtual bool IsFalling() const OVERRIDE;
-	virtual bool IsMovingOnGround() const OVERRIDE;
-	virtual bool IsSwimming() const OVERRIDE;
-	virtual bool IsFlying() const OVERRIDE;
-	virtual float GetGravityZ() const OVERRIDE;
-	virtual void AddRadialForce(const FVector& Origin, float Radius, float Strength, enum ERadialImpulseFalloff Falloff) OVERRIDE;
-	virtual void AddRadialImpulse(const FVector& Origin, float Radius, float Strength, enum ERadialImpulseFalloff Falloff, bool bVelChange) OVERRIDE;
+	virtual float GetMaxSpeed() const override;
+	virtual float GetMaxSpeedModifier() const override;
+	virtual void StopActiveMovement() override;
+	virtual bool IsCrouching() const override;
+	virtual bool IsFalling() const override;
+	virtual bool IsMovingOnGround() const override;
+	virtual bool IsSwimming() const override;
+	virtual bool IsFlying() const override;
+	virtual float GetGravityZ() const override;
+	virtual void AddRadialForce(const FVector& Origin, float Radius, float Strength, enum ERadialImpulseFalloff Falloff) override;
+	virtual void AddRadialImpulse(const FVector& Origin, float Radius, float Strength, enum ERadialImpulseFalloff Falloff, bool bVelChange) override;
 	//END UMovementComponent Interface
 
 	/** @return true if the character is in the 'Walking' movement mode. */
@@ -832,16 +832,16 @@ public:
 	bool IsWalking() const;
 
 	//BEGIN UNavMovementComponent Interface
-	virtual void RequestDirectMove(const FVector& MoveVelocity, bool bForceMaxSpeed) OVERRIDE;
-	virtual bool CanStopPathFollowing() const OVERRIDE;
+	virtual void RequestDirectMove(const FVector& MoveVelocity, bool bForceMaxSpeed) override;
+	virtual bool CanStopPathFollowing() const override;
 	//END UNaVMovementComponent Interface
 
 	//Begin UPawnMovementComponent Interface
-	virtual void NotifyBumpedPawn(APawn* BumpedPawn) OVERRIDE;
+	virtual void NotifyBumpedPawn(APawn* BumpedPawn) override;
 	//End UPawnMovementComponent Interface
 
 #if WITH_EDITOR
-	virtual void PostEditChangeProperty(FPropertyChangedEvent& PropertyChangedEvent) OVERRIDE;
+	virtual void PostEditChangeProperty(FPropertyChangedEvent& PropertyChangedEvent) override;
 #endif // WITH_EDITOR
 
 	/** Make movement impossible (sets movement mode to MOVE_None). */
@@ -1094,7 +1094,7 @@ public:
 	virtual void MoveSmooth(const FVector& InVelocity, const float DeltaSeconds, FStepDownResult* OutStepDownResult = NULL );
 
 	
-	virtual void SetUpdatedComponent(class UPrimitiveComponent* NewUpdatedComponent) OVERRIDE;
+	virtual void SetUpdatedComponent(class UPrimitiveComponent* NewUpdatedComponent) override;
 	
 	/** @Return MovementMode string */
 	virtual FString GetMovementName();
@@ -1224,25 +1224,25 @@ protected:
 	virtual void MaintainHorizontalGroundVelocity();
 
 	/** Overridden to set bJustTeleported to true, so we don't make incorrect velocity calculations based on adjusted movement. */
-	virtual bool ResolvePenetration(const FVector& Adjustment, const FHitResult& Hit, const FRotator& NewRotation) OVERRIDE;
+	virtual bool ResolvePenetration(const FVector& Adjustment, const FHitResult& Hit, const FRotator& NewRotation) override;
 
 	/** Overridden to avoid upward adjustments that can go too high when walking. */
-	virtual FVector GetPenetrationAdjustment(const FHitResult& Hit) const OVERRIDE;
+	virtual FVector GetPenetrationAdjustment(const FHitResult& Hit) const override;
 
 	/** Don't call for intermediate parts of move! */
-	virtual void HandleImpact(FHitResult const& Hit, float TimeSlice=0.f, const FVector& MoveDelta = FVector::ZeroVector) OVERRIDE;
+	virtual void HandleImpact(FHitResult const& Hit, float TimeSlice=0.f, const FVector& MoveDelta = FVector::ZeroVector) override;
 
 	/** Custom version of SlideAlongSurface that handles different movement modes separately; namely during walking physics we might not want to slide up slopes. */
-	virtual float SlideAlongSurface(const FVector& Delta, float Time, const FVector& Normal, FHitResult &Hit, bool bHandleImpact) OVERRIDE;
+	virtual float SlideAlongSurface(const FVector& Delta, float Time, const FVector& Normal, FHitResult &Hit, bool bHandleImpact) override;
 
 	/** Custom version that allows upwards slides when walking if the surface is walkable. */
-	virtual void TwoWallAdjust(FVector &Delta, const FHitResult& Hit, const FVector &OldHitNormal) const OVERRIDE;
+	virtual void TwoWallAdjust(FVector &Delta, const FHitResult& Hit, const FVector &OldHitNormal) const override;
 
 	/**
 	 * Calculate slide on a slope. Has special treatment when falling, to avoid boosting up slopes.
 	 * Calls AdjustUpperHemisphereImpact() for upward falling movement that hits the top of the capsule, because commonly we don't want this to behave like a smooth capsule.
 	 */
-	virtual FVector ComputeSlideVector(const FVector& Delta, const float Time, const FVector& Normal, const FHitResult& Hit) const OVERRIDE;
+	virtual FVector ComputeSlideVector(const FVector& Delta, const float Time, const FVector& Normal, const FHitResult& Hit) const override;
 
 	/**
 	 * Given an upward impact on the top of the capsule, allows calculation of a different movement delta.
@@ -1388,22 +1388,22 @@ public:
 	//--------------------------------
 	// Server hook
 	//--------------------------------
-	virtual void SendClientAdjustment() OVERRIDE;
-	virtual void ForcePositionUpdate(float DeltaTime) OVERRIDE;
+	virtual void SendClientAdjustment() override;
+	virtual void ForcePositionUpdate(float DeltaTime) override;
 
 	//--------------------------------
 	// Client hook
 	//--------------------------------
-	virtual void SmoothCorrection(const FVector& OldLocation) OVERRIDE;
+	virtual void SmoothCorrection(const FVector& OldLocation) override;
 
-	virtual class FNetworkPredictionData_Client* GetPredictionData_Client() const OVERRIDE;
-	virtual class FNetworkPredictionData_Server* GetPredictionData_Server() const OVERRIDE;
+	virtual class FNetworkPredictionData_Client* GetPredictionData_Client() const override;
+	virtual class FNetworkPredictionData_Server* GetPredictionData_Server() const override;
 
-	virtual bool HasPredictionData_Client() const OVERRIDE { return ClientPredictionData != NULL; }
-	virtual bool HasPredictionData_Server() const OVERRIDE { return ServerPredictionData != NULL; }
+	virtual bool HasPredictionData_Client() const override { return ClientPredictionData != NULL; }
+	virtual bool HasPredictionData_Server() const override { return ServerPredictionData != NULL; }
 
-	virtual void ResetPredictionData_Client() OVERRIDE;
-	virtual void ResetPredictionData_Server() OVERRIDE;
+	virtual void ResetPredictionData_Client() override;
+	virtual void ResetPredictionData_Server() override;
 
 protected:
 	class FNetworkPredictionData_Client_Character* ClientPredictionData;
diff --git a/Engine/Source/Runtime/Engine/Classes/GameFramework/Controller.h b/Engine/Source/Runtime/Engine/Classes/GameFramework/Controller.h
--- a/Engine/Source/Runtime/Engine/Classes/GameFramework/Controller.h
+++ b/Engine/Source/Runtime/Engine/Classes/GameFramework/Controller.h
@@ -164,17 +164,17 @@ public:
 	class APawn* GetControlledPawn() const;
 
 	// Begin AActor interface
-	virtual void TickActor( float DeltaTime, enum ELevelTick TickType, FActorTickFunction& ThisTickFunction ) OVERRIDE;
-	virtual void K2_DestroyActor() OVERRIDE;
-	virtual void DisplayDebug(class UCanvas* Canvas, const FDebugDisplayInfo& DebugDisplay, float& YL, float& YPos) OVERRIDE;
-	virtual void GetActorEyesViewPoint( FVector& out_Location, FRotator& out_Rotation ) const OVERRIDE;
-	virtual FString GetHumanReadableName() const OVERRIDE;
+	virtual void TickActor( float DeltaTime, enum ELevelTick TickType, FActorTickFunction& ThisTickFunction ) override;
+	virtual void K2_DestroyActor() override;
+	virtual void DisplayDebug(class UCanvas* Canvas, const FDebugDisplayInfo& DebugDisplay, float& YL, float& YPos) override;
+	virtual void GetActorEyesViewPoint( FVector& out_Location, FRotator& out_Rotation ) const override;
+	virtual FString GetHumanReadableName() const override;
 
 	/* Overridden to create the player replication info and perform other mundane initialization tasks. */
-	virtual void PostInitializeComponents() OVERRIDE;
+	virtual void PostInitializeComponents() override;
 
-	virtual void Reset() OVERRIDE;
-	virtual void Destroyed() OVERRIDE;
+	virtual void Reset() override;
+	virtual void Destroyed() override;
 	// End AActor interface
 
 	/** Getter for Pawn */
@@ -263,10 +263,10 @@ public:
 	virtual void FailedToSpawnPawn();
 
 	// Begin INavAgentInterface Interface
-	virtual const struct FNavAgentProperties* GetNavAgentProperties() const OVERRIDE;
-	virtual FVector GetNavAgentLocation() const OVERRIDE;
-	virtual void GetMoveGoalReachTest(class AActor* MovingActor, const FVector& MoveOffset, FVector& GoalOffset, float& GoalRadius, float& GoalHalfHeight) const OVERRIDE;
-	virtual bool ShouldPostponePathUpdates() const OVERRIDE;
+	virtual const struct FNavAgentProperties* GetNavAgentProperties() const override;
+	virtual FVector GetNavAgentLocation() const override;
+	virtual void GetMoveGoalReachTest(class AActor* MovingActor, const FVector& MoveOffset, FVector& GoalOffset, float& GoalRadius, float& GoalHalfHeight) const override;
+	virtual bool ShouldPostponePathUpdates() const override;
 	// End INavAgentInterface Interface
 	
 	/** prepares path finding and path following components */
diff --git a/Engine/Source/Runtime/Engine/Classes/GameFramework/DefaultPawn.h b/Engine/Source/Runtime/Engine/Classes/GameFramework/DefaultPawn.h
--- a/Engine/Source/Runtime/Engine/Classes/GameFramework/DefaultPawn.h
+++ b/Engine/Source/Runtime/Engine/Classes/GameFramework/DefaultPawn.h
@@ -15,8 +15,8 @@ class ENGINE_API ADefaultPawn : public APawn
 	GENERATED_UCLASS_BODY()
 
 	// Begin Pawn overrides
-	virtual class UPawnMovementComponent* GetMovementComponent() const OVERRIDE { return MovementComponent; }
-	virtual void SetupPlayerInputComponent(class UInputComponent* InputComponent) OVERRIDE;
+	virtual class UPawnMovementComponent* GetMovementComponent() const override { return MovementComponent; }
+	virtual void SetupPlayerInputComponent(class UInputComponent* InputComponent) override;
 	// End Pawn overrides
 
 	/** Input callback to move forward. */
diff --git a/Engine/Source/Runtime/Engine/Classes/GameFramework/DefaultPhysicsVolume.h b/Engine/Source/Runtime/Engine/Classes/GameFramework/DefaultPhysicsVolume.h
--- a/Engine/Source/Runtime/Engine/Classes/GameFramework/DefaultPhysicsVolume.h
+++ b/Engine/Source/Runtime/Engine/Classes/GameFramework/DefaultPhysicsVolume.h
@@ -15,7 +15,7 @@ class ADefaultPhysicsVolume : public APhysicsVolume
 	GENERATED_UCLASS_BODY()
 
 	// Begin AActor interface
-	virtual void Destroyed() OVERRIDE;
+	virtual void Destroyed() override;
 	// End AActor interface
 };
 
diff --git a/Engine/Source/Runtime/Engine/Classes/GameFramework/EngineMessage.h b/Engine/Source/Runtime/Engine/Classes/GameFramework/EngineMessage.h
--- a/Engine/Source/Runtime/Engine/Classes/GameFramework/EngineMessage.h
+++ b/Engine/Source/Runtime/Engine/Classes/GameFramework/EngineMessage.h
@@ -41,6 +41,6 @@ class  UEngineMessage : public ULocalMessage
 	FString NewSpecMessage;
 
 	// Begin ULocalMessage interface
-	virtual void ClientReceive(const FClientReceiveData& ClientData) const OVERRIDE;
+	virtual void ClientReceive(const FClientReceiveData& ClientData) const override;
 	// End ULocalMessage interface
 };
diff --git a/Engine/Source/Runtime/Engine/Classes/GameFramework/FloatingPawnMovement.h b/Engine/Source/Runtime/Engine/Classes/GameFramework/FloatingPawnMovement.h
--- a/Engine/Source/Runtime/Engine/Classes/GameFramework/FloatingPawnMovement.h
+++ b/Engine/Source/Runtime/Engine/Classes/GameFramework/FloatingPawnMovement.h
@@ -13,12 +13,12 @@ class ENGINE_API UFloatingPawnMovement : public UPawnMovementComponent
 	GENERATED_UCLASS_BODY()
 
 	//Begin UActorComponent Interface
-	virtual void TickComponent(float DeltaTime, enum ELevelTick TickType, FActorComponentTickFunction *ThisTickFunction) OVERRIDE;
+	virtual void TickComponent(float DeltaTime, enum ELevelTick TickType, FActorComponentTickFunction *ThisTickFunction) override;
 	//End UActorComponent Interface
 
 	//Begin UMovementComponent Interface
-	virtual float GetMaxSpeed() const OVERRIDE { return MaxSpeed; }
-	virtual bool ResolvePenetration(const FVector& Adjustment, const FHitResult& Hit, const FRotator& NewRotation) OVERRIDE;
+	virtual float GetMaxSpeed() const override { return MaxSpeed; }
+	virtual bool ResolvePenetration(const FVector& Adjustment, const FHitResult& Hit, const FRotator& NewRotation) override;
 	//End UMovementComponent Interface
 
 	/** Maximum velocity magnitude allowed for the controlled Pawn. */
diff --git a/Engine/Source/Runtime/Engine/Classes/GameFramework/ForceFeedbackEffect.h b/Engine/Source/Runtime/Engine/Classes/GameFramework/ForceFeedbackEffect.h
--- a/Engine/Source/Runtime/Engine/Classes/GameFramework/ForceFeedbackEffect.h
+++ b/Engine/Source/Runtime/Engine/Classes/GameFramework/ForceFeedbackEffect.h
@@ -83,7 +83,7 @@ class UForceFeedbackEffect : public UObject
 	float Duration;
 
 #if WITH_EDITOR
-	virtual void PostEditChangeChainProperty( struct FPropertyChangedChainEvent& PropertyChangedEvent) OVERRIDE;
+	virtual void PostEditChangeChainProperty( struct FPropertyChangedChainEvent& PropertyChangedEvent) override;
 #endif
 
 	float GetDuration();
diff --git a/Engine/Source/Runtime/Engine/Classes/GameFramework/GameMode.h b/Engine/Source/Runtime/Engine/Classes/GameFramework/GameMode.h
--- a/Engine/Source/Runtime/Engine/Classes/GameFramework/GameMode.h
+++ b/Engine/Source/Runtime/Engine/Classes/GameFramework/GameMode.h
@@ -248,10 +248,10 @@ public:
 	virtual void SetBandwidthLimit(float AsyncIOBandwidthLimit);
 
 	// Begin AActor interface
-	virtual void PreInitializeComponents() OVERRIDE;
-	virtual void DisplayDebug(class UCanvas* Canvas, const FDebugDisplayInfo& DebugDisplay, float& YL, float& YPos) OVERRIDE;
-	virtual void Reset() OVERRIDE;
-	virtual void Tick(float DeltaSeconds) OVERRIDE;
+	virtual void PreInitializeComponents() override;
+	virtual void DisplayDebug(class UCanvas* Canvas, const FDebugDisplayInfo& DebugDisplay, float& YL, float& YPos) override;
+	virtual void Reset() override;
+	virtual void Tick(float DeltaSeconds) override;
 	// End AActor interface
 
 	/** 
diff --git a/Engine/Source/Runtime/Engine/Classes/GameFramework/GameNetworkManager.h b/Engine/Source/Runtime/Engine/Classes/GameFramework/GameNetworkManager.h
--- a/Engine/Source/Runtime/Engine/Classes/GameFramework/GameNetworkManager.h
+++ b/Engine/Source/Runtime/Engine/Classes/GameFramework/GameNetworkManager.h
@@ -129,7 +129,7 @@ class AGameNetworkManager : public AInfo
 
 	/** @return true if velocity vector passed in is considered near zero for networking purposes */
 	bool NetworkVelocityNearZero(FVector InVelocity) const;
-	virtual void PostInitializeComponents() OVERRIDE;
+	virtual void PostInitializeComponents() override;
 
 	/** @RETURN new per/client bandwidth given number of players in the game */
 	virtual int32 CalculatedNetSpeed();
diff --git a/Engine/Source/Runtime/Engine/Classes/GameFramework/Info.h b/Engine/Source/Runtime/Engine/Classes/GameFramework/Info.h
--- a/Engine/Source/Runtime/Engine/Classes/GameFramework/Info.h
+++ b/Engine/Source/Runtime/Engine/Classes/GameFramework/Info.h
@@ -19,10 +19,10 @@ class AInfo : public AActor
 	TSubobjectPtr<class UBillboardComponent> SpriteComponent;
 #endif
 
-	virtual bool UpdateNavigationRelevancy() OVERRIDE { SetNavigationRelevancy(false); return false; }
+	virtual bool UpdateNavigationRelevancy() override { SetNavigationRelevancy(false); return false; }
 	
 	/** Indicates whether this actor should participate in level bounds calculations. */
-	virtual bool IsLevelBoundsRelevant() const OVERRIDE { return false; }
+	virtual bool IsLevelBoundsRelevant() const override { return false; }
 };
 
 
diff --git a/Engine/Source/Runtime/Engine/Classes/GameFramework/InputSettings.h b/Engine/Source/Runtime/Engine/Classes/GameFramework/InputSettings.h
--- a/Engine/Source/Runtime/Engine/Classes/GameFramework/InputSettings.h
+++ b/Engine/Source/Runtime/Engine/Classes/GameFramework/InputSettings.h
@@ -57,7 +57,7 @@ class ENGINE_API UInputSettings
 
 	// UObject interface
 #if WITH_EDITOR
-	virtual void PostEditChangeChainProperty(struct FPropertyChangedChainEvent& PropertyChangedEvent) OVERRIDE;
+	virtual void PostEditChangeChainProperty(struct FPropertyChangedChainEvent& PropertyChangedEvent) override;
 #endif
 	// End of UObject interface
 
diff --git a/Engine/Source/Runtime/Engine/Classes/GameFramework/KillZVolume.h b/Engine/Source/Runtime/Engine/Classes/GameFramework/KillZVolume.h
--- a/Engine/Source/Runtime/Engine/Classes/GameFramework/KillZVolume.h
+++ b/Engine/Source/Runtime/Engine/Classes/GameFramework/KillZVolume.h
@@ -10,7 +10,7 @@ class ENGINE_API AKillZVolume : public APhysicsVolume
 	GENERATED_UCLASS_BODY()
 	
 	//Begin PhysicsVolume Interface
-	virtual void ActorEnteredVolume(class AActor* Other) OVERRIDE;
+	virtual void ActorEnteredVolume(class AActor* Other) override;
 	//End PhysicsVolume Interface
 };
 
diff --git a/Engine/Source/Runtime/Engine/Classes/GameFramework/MovementComponent.h b/Engine/Source/Runtime/Engine/Classes/GameFramework/MovementComponent.h
--- a/Engine/Source/Runtime/Engine/Classes/GameFramework/MovementComponent.h
+++ b/Engine/Source/Runtime/Engine/Classes/GameFramework/MovementComponent.h
@@ -70,10 +70,10 @@ public:
 	uint32 bAutoRegisterUpdatedComponent:1;
 
 	// Begin ActorComponent interface 
-	virtual void RegisterComponentTickFunctions(bool bRegister) OVERRIDE;
+	virtual void RegisterComponentTickFunctions(bool bRegister) override;
 
 	/** Overridden to auto-register the updated component if it starts NULL, and we can find a root component on our owner. */
-	virtual void InitializeComponent() OVERRIDE;
+	virtual void InitializeComponent() override;
 	// End ActorComponent interface
 
 	/** @return gravity that affects this component */
diff --git a/Engine/Source/Runtime/Engine/Classes/GameFramework/PainCausingVolume.h b/Engine/Source/Runtime/Engine/Classes/GameFramework/PainCausingVolume.h
--- a/Engine/Source/Runtime/Engine/Classes/GameFramework/PainCausingVolume.h
+++ b/Engine/Source/Runtime/Engine/Classes/GameFramework/PainCausingVolume.h
@@ -48,20 +48,20 @@ class ENGINE_API APainCausingVolume : public APhysicsVolume
 
 #if WITH_EDITOR
 	//Begin AVolume Interface
-	virtual void CheckForErrors() OVERRIDE;
+	virtual void CheckForErrors() override;
 	//End AVolume Interface
 #endif
 
 	//Begin AActor Interface
-	virtual void PostInitializeComponents() OVERRIDE;
+	virtual void PostInitializeComponents() override;
 
 	/* reset actor to initial state - used when restarting level without reloading. */
-	virtual void Reset() OVERRIDE;
+	virtual void Reset() override;
 	//End AActor Interface
 
 	//Being PhysicsVolume Interface
 	/** If bEntryPain is true, call CausePainTo() on entering actor. */
-	virtual void ActorEnteredVolume(class AActor* Other) OVERRIDE;
+	virtual void ActorEnteredVolume(class AActor* Other) override;
 	//End PhysicsVolume Interface
 
 	/** damage overlapping actors if pain causing. */
diff --git a/Engine/Source/Runtime/Engine/Classes/GameFramework/Pawn.h b/Engine/Source/Runtime/Engine/Classes/GameFramework/Pawn.h
--- a/Engine/Source/Runtime/Engine/Classes/GameFramework/Pawn.h
+++ b/Engine/Source/Runtime/Engine/Classes/GameFramework/Pawn.h
@@ -132,7 +132,7 @@ public:
 	UFUNCTION(BlueprintPure, Category="Pawn", meta=(ToolTip="Gets the owning actor of the Movement Base Component on which the pawn is standing."))
 	static AActor* GetMovementBaseActor(const APawn* Pawn);
 
-	virtual bool IsBasedOnActor(const AActor * Other) const OVERRIDE;
+	virtual bool IsBasedOnActor(const AActor * Other) const override;
 
 	virtual bool ReachedDesiredRotation();
 
@@ -163,42 +163,42 @@ public:
 	uint32 bProcessingOutsideWorldBounds:1;
 
 	// Begin AActor Interface.
-	virtual FVector GetVelocity() const OVERRIDE;
-	virtual void Tick(float DeltaSeconds) OVERRIDE;
-	virtual void Reset() OVERRIDE;
-	virtual FString GetHumanReadableName() const OVERRIDE;
-	virtual float GetNetPriority(const FVector& ViewPos, const FVector& ViewDir, APlayerController* Viewer, UActorChannel* InChannel, float Time, bool bLowBandwidth) OVERRIDE;
-	virtual bool ShouldTickIfViewportsOnly() const OVERRIDE;
-	virtual bool IsNetRelevantFor(APlayerController* RealViewer, AActor* Viewer, const FVector& SrcLocation) OVERRIDE;
-	virtual void PostNetReceiveLocation() OVERRIDE;
-	virtual void PostNetReceiveVelocity(const FVector& NewVelocity) OVERRIDE;
-	virtual void DisplayDebug(class UCanvas* Canvas, const FDebugDisplayInfo& DebugDisplay, float& YL, float& YPos) OVERRIDE;
-	virtual void GetActorEyesViewPoint( FVector& Location, FRotator& Rotation ) const OVERRIDE;
-	virtual void OutsideWorldBounds() OVERRIDE;
-	virtual void Destroyed() OVERRIDE;
-	virtual void PreInitializeComponents() OVERRIDE;
-	virtual void PostInitializeComponents() OVERRIDE;
-	virtual class UPlayer* GetNetOwningPlayer() OVERRIDE;
-	virtual class UNetConnection* GetNetConnection() OVERRIDE;
-	virtual void PostInitProperties() OVERRIDE;
-	virtual void PostLoad() OVERRIDE;
-	virtual void PostRegisterAllComponents() OVERRIDE;
-	virtual float TakeDamage(float Damage, struct FDamageEvent const& DamageEvent, class AController* EventInstigator, class AActor* DamageCauser) OVERRIDE;
-	virtual void BecomeViewTarget(class APlayerController* PC) OVERRIDE;
-	virtual bool UpdateNavigationRelevancy() OVERRIDE { SetNavigationRelevancy(false); return false; }
-	virtual void EnableInput(class APlayerController* PlayerController) OVERRIDE;
-	virtual void DisableInput(class APlayerController* PlayerController) OVERRIDE;
+	virtual FVector GetVelocity() const override;
+	virtual void Tick(float DeltaSeconds) override;
+	virtual void Reset() override;
+	virtual FString GetHumanReadableName() const override;
+	virtual float GetNetPriority(const FVector& ViewPos, const FVector& ViewDir, APlayerController* Viewer, UActorChannel* InChannel, float Time, bool bLowBandwidth) override;
+	virtual bool ShouldTickIfViewportsOnly() const override;
+	virtual bool IsNetRelevantFor(APlayerController* RealViewer, AActor* Viewer, const FVector& SrcLocation) override;
+	virtual void PostNetReceiveLocation() override;
+	virtual void PostNetReceiveVelocity(const FVector& NewVelocity) override;
+	virtual void DisplayDebug(class UCanvas* Canvas, const FDebugDisplayInfo& DebugDisplay, float& YL, float& YPos) override;
+	virtual void GetActorEyesViewPoint( FVector& Location, FRotator& Rotation ) const override;
+	virtual void OutsideWorldBounds() override;
+	virtual void Destroyed() override;
+	virtual void PreInitializeComponents() override;
+	virtual void PostInitializeComponents() override;
+	virtual class UPlayer* GetNetOwningPlayer() override;
+	virtual class UNetConnection* GetNetConnection() override;
+	virtual void PostInitProperties() override;
+	virtual void PostLoad() override;
+	virtual void PostRegisterAllComponents() override;
+	virtual float TakeDamage(float Damage, struct FDamageEvent const& DamageEvent, class AController* EventInstigator, class AActor* DamageCauser) override;
+	virtual void BecomeViewTarget(class APlayerController* PC) override;
+	virtual bool UpdateNavigationRelevancy() override { SetNavigationRelevancy(false); return false; }
+	virtual void EnableInput(class APlayerController* PlayerController) override;
+	virtual void DisableInput(class APlayerController* PlayerController) override;
 
 	/** Overridden to defer to the RootComponent's CanBeCharacterBase setting if it is explicitly Yes or No. If set to Owner, will return Super::CanBeBaseForCharacter(). */
-	virtual bool CanBeBaseForCharacter(class APawn* APawn) const OVERRIDE;
+	virtual bool CanBeBaseForCharacter(class APawn* APawn) const override;
 	// End AActor Interface
 
 	// Begin INavAgentInterface Interface
-	virtual const struct FNavAgentProperties* GetNavAgentProperties() const OVERRIDE { return GetMovementComponent() ? GetMovementComponent()->GetNavAgentProperties() : NULL;}
+	virtual const struct FNavAgentProperties* GetNavAgentProperties() const override { return GetMovementComponent() ? GetMovementComponent()->GetNavAgentProperties() : NULL;}
 	/** Basically retrieved pawn's location on navmesh */
 	UFUNCTION(BlueprintCallable, Category="Pawn")
-	virtual FVector GetNavAgentLocation() const OVERRIDE { return GetActorLocation() - FVector(0.f, 0.f, BaseEyeHeight); }
-	virtual void GetMoveGoalReachTest(class AActor* MovingActor, const FVector& MoveOffset, FVector& GoalOffset, float& GoalRadius, float& GoalHalfHeight) const OVERRIDE;
+	virtual FVector GetNavAgentLocation() const override { return GetActorLocation() - FVector(0.f, 0.f, BaseEyeHeight); }
+	virtual void GetMoveGoalReachTest(class AActor* MovingActor, const FVector& MoveOffset, FVector& GoalOffset, float& GoalRadius, float& GoalHalfHeight) const override;
 	// End INavAgentInterface Interface
 
 	/** updates MovementComponent's parameters used by navigation system */
diff --git a/Engine/Source/Runtime/Engine/Classes/GameFramework/PawnMovementComponent.h b/Engine/Source/Runtime/Engine/Classes/GameFramework/PawnMovementComponent.h
--- a/Engine/Source/Runtime/Engine/Classes/GameFramework/PawnMovementComponent.h
+++ b/Engine/Source/Runtime/Engine/Classes/GameFramework/PawnMovementComponent.h
@@ -14,10 +14,10 @@ class ENGINE_API UPawnMovementComponent : public UNavMovementComponent
 	GENERATED_UCLASS_BODY()
 
 	// Overridden to only allow registration with components owned by a Pawn.
-	virtual void SetUpdatedComponent(class UPrimitiveComponent* NewUpdatedComponent) OVERRIDE;
+	virtual void SetUpdatedComponent(class UPrimitiveComponent* NewUpdatedComponent) override;
 
 	// Overridden to also call StopActiveMovement().
-	virtual void StopMovementImmediately() OVERRIDE;
+	virtual void StopMovementImmediately() override;
 
 	/** Stops applying further movement (usually zeros acceleration). */
 	UFUNCTION(BlueprintCallable, Category="Pawn|Components|PawnMovement")
diff --git a/Engine/Source/Runtime/Engine/Classes/GameFramework/PhysicsVolume.h b/Engine/Source/Runtime/Engine/Classes/GameFramework/PhysicsVolume.h
--- a/Engine/Source/Runtime/Engine/Classes/GameFramework/PhysicsVolume.h
+++ b/Engine/Source/Runtime/Engine/Classes/GameFramework/PhysicsVolume.h
@@ -16,7 +16,7 @@ class APhysicsVolume : public AVolume
 
 	// Begin UObject interface.
 #if WITH_EDITOR
-	ENGINE_API virtual void LoadedFromAnotherClass(const FName& OldClassName) OVERRIDE;
+	ENGINE_API virtual void LoadedFromAnotherClass(const FName& OldClassName) override;
 #endif // WITH_EDITOR	
 	// End UObject interface.
 
diff --git a/Engine/Source/Runtime/Engine/Classes/GameFramework/PlayerController.h b/Engine/Source/Runtime/Engine/Classes/GameFramework/PlayerController.h
--- a/Engine/Source/Runtime/Engine/Classes/GameFramework/PlayerController.h
+++ b/Engine/Source/Runtime/Engine/Classes/GameFramework/PlayerController.h
@@ -968,44 +968,44 @@ public:
 	virtual FString ConsoleCommand(const FString& Command, bool bWriteToLog = true);
 
 	// Begin UObject Interface
-	virtual void PostLoad() OVERRIDE;
+	virtual void PostLoad() override;
 	// End of UObject Interface
 
 	// Begin AActor Interface
-	virtual void GetActorEyesViewPoint(FVector& Location, FRotator& Rotation) const OVERRIDE;
-	virtual void CalcCamera(float DeltaTime, struct FMinimalViewInfo& OutResult) OVERRIDE;
-	virtual void TickActor(float DeltaTime, enum ELevelTick TickType, FActorTickFunction& ThisTickFunction) OVERRIDE;
-	virtual bool IsNetRelevantFor(APlayerController* RealViewer, AActor* Viewer, const FVector& SrcLocation) OVERRIDE;
-	virtual void FellOutOfWorld(const class UDamageType& dmgType) OVERRIDE;
-	virtual void Reset() OVERRIDE;
-	virtual void Possess(APawn* aPawn) OVERRIDE;
-	virtual void UnPossess() OVERRIDE;
-	virtual void CleanupPlayerState() OVERRIDE;
-	virtual void Destroyed() OVERRIDE;
-	virtual void OnActorChannelOpen(class FInBunch& InBunch, class UNetConnection* Connection) OVERRIDE;
-	virtual void OnSerializeNewActor(class FOutBunch& OutBunch) OVERRIDE;
-	virtual void OnNetCleanup(class UNetConnection* Connection) OVERRIDE;
-	virtual float GetNetPriority(const FVector& ViewPos, const FVector& ViewDir, APlayerController* Viewer, UActorChannel* InChannel, float Time, bool bLowBandwidth) OVERRIDE;
-	virtual class UPlayer* GetNetOwningPlayer() OVERRIDE;
-	virtual class UNetConnection* GetNetConnection() OVERRIDE;
-	virtual void DisplayDebug(class UCanvas* Canvas, const FDebugDisplayInfo& DebugDisplay, float& YL, float& YPos) OVERRIDE;
-	virtual void PostInitializeComponents() OVERRIDE;
-	virtual void EnableInput(class APlayerController* PlayerController) OVERRIDE;
-	virtual void DisableInput(class APlayerController* PlayerController) OVERRIDE;
+	virtual void GetActorEyesViewPoint(FVector& Location, FRotator& Rotation) const override;
+	virtual void CalcCamera(float DeltaTime, struct FMinimalViewInfo& OutResult) override;
+	virtual void TickActor(float DeltaTime, enum ELevelTick TickType, FActorTickFunction& ThisTickFunction) override;
+	virtual bool IsNetRelevantFor(APlayerController* RealViewer, AActor* Viewer, const FVector& SrcLocation) override;
+	virtual void FellOutOfWorld(const class UDamageType& dmgType) override;
+	virtual void Reset() override;
+	virtual void Possess(APawn* aPawn) override;
+	virtual void UnPossess() override;
+	virtual void CleanupPlayerState() override;
+	virtual void Destroyed() override;
+	virtual void OnActorChannelOpen(class FInBunch& InBunch, class UNetConnection* Connection) override;
+	virtual void OnSerializeNewActor(class FOutBunch& OutBunch) override;
+	virtual void OnNetCleanup(class UNetConnection* Connection) override;
+	virtual float GetNetPriority(const FVector& ViewPos, const FVector& ViewDir, APlayerController* Viewer, UActorChannel* InChannel, float Time, bool bLowBandwidth) override;
+	virtual class UPlayer* GetNetOwningPlayer() override;
+	virtual class UNetConnection* GetNetConnection() override;
+	virtual void DisplayDebug(class UCanvas* Canvas, const FDebugDisplayInfo& DebugDisplay, float& YL, float& YPos) override;
+	virtual void PostInitializeComponents() override;
+	virtual void EnableInput(class APlayerController* PlayerController) override;
+	virtual void DisableInput(class APlayerController* PlayerController) override;
 	// End AActor interface
 
 	// Begin AController interface
-	virtual void GameHasEnded(class AActor* EndGameFocus = NULL, bool bIsWinner = false) OVERRIDE;
-	virtual bool IsLocalPlayerController() const OVERRIDE;
-	virtual bool IsLocalController() const OVERRIDE;
-	virtual void GetPlayerViewPoint(FVector& out_Location, FRotator& out_Rotation) const OVERRIDE;
-	virtual void SetInitialLocationAndRotation(const FVector& NewLocation, const FRotator& NewRotation) OVERRIDE;
-	virtual void ChangeState(FName NewState) OVERRIDE;
-	virtual class AActor* GetViewTarget() const OVERRIDE;
-	virtual void BeginInactiveState() OVERRIDE;
-	virtual void EndInactiveState() OVERRIDE;
-	virtual void FailedToSpawnPawn() OVERRIDE;
-	virtual void SetPawn(APawn* InPawn) OVERRIDE;
+	virtual void GameHasEnded(class AActor* EndGameFocus = NULL, bool bIsWinner = false) override;
+	virtual bool IsLocalPlayerController() const override;
+	virtual bool IsLocalController() const override;
+	virtual void GetPlayerViewPoint(FVector& out_Location, FRotator& out_Rotation) const override;
+	virtual void SetInitialLocationAndRotation(const FVector& NewLocation, const FRotator& NewRotation) override;
+	virtual void ChangeState(FName NewState) override;
+	virtual class AActor* GetViewTarget() const override;
+	virtual void BeginInactiveState() override;
+	virtual void EndInactiveState() override;
+	virtual void FailedToSpawnPawn() override;
+	virtual void SetPawn(APawn* InPawn) override;
 	// End AController interface
 
 	/** called on the server when the client sends a message indicating it was unable to initialize an Actor channel,
diff --git a/Engine/Source/Runtime/Engine/Classes/GameFramework/PlayerStart.h b/Engine/Source/Runtime/Engine/Classes/GameFramework/PlayerStart.h
--- a/Engine/Source/Runtime/Engine/Classes/GameFramework/PlayerStart.h
+++ b/Engine/Source/Runtime/Engine/Classes/GameFramework/PlayerStart.h
@@ -21,9 +21,9 @@ class ENGINE_API APlayerStart : public ANavigationObjectBase
 	TSubobjectPtr<class UArrowComponent> ArrowComponent;
 #endif
 
-	virtual void PostInitializeComponents() OVERRIDE;
+	virtual void PostInitializeComponents() override;
 	
-	virtual void PostUnregisterAllComponents() OVERRIDE;
+	virtual void PostUnregisterAllComponents() override;
 };
 
 
diff --git a/Engine/Source/Runtime/Engine/Classes/GameFramework/PlayerState.h b/Engine/Source/Runtime/Engine/Classes/GameFramework/PlayerState.h
--- a/Engine/Source/Runtime/Engine/Classes/GameFramework/PlayerState.h
+++ b/Engine/Source/Runtime/Engine/Classes/GameFramework/PlayerState.h
@@ -133,10 +133,10 @@ public:
 	virtual void OnRep_UniqueId();
 
 	// Begin AActor Interface
-	virtual void PostInitializeComponents() OVERRIDE; 
-	virtual void Destroyed() OVERRIDE;
-	virtual void Reset() OVERRIDE;
-	virtual FString GetHumanReadableName() const OVERRIDE;
+	virtual void PostInitializeComponents() override; 
+	virtual void Destroyed() override;
+	virtual void Reset() override;
+	virtual FString GetHumanReadableName() const override;
 	// End AActor Interface
 
 
diff --git a/Engine/Source/Runtime/Engine/Classes/GameFramework/ProjectileMovementComponent.h b/Engine/Source/Runtime/Engine/Classes/GameFramework/ProjectileMovementComponent.h
--- a/Engine/Source/Runtime/Engine/Classes/GameFramework/ProjectileMovementComponent.h
+++ b/Engine/Source/Runtime/Engine/Classes/GameFramework/ProjectileMovementComponent.h
@@ -72,13 +72,13 @@ class ENGINE_API UProjectileMovementComponent : public UMovementComponent
 	virtual void SetVelocityInLocalSpace(FVector NewVelocity);
 
 	//Begin UActorComponent Interface
-	virtual void TickComponent(float DeltaTime, enum ELevelTick TickType, FActorComponentTickFunction *ThisTickFunction) OVERRIDE;
-	virtual void Serialize( FArchive& Ar) OVERRIDE;
+	virtual void TickComponent(float DeltaTime, enum ELevelTick TickType, FActorComponentTickFunction *ThisTickFunction) override;
+	virtual void Serialize( FArchive& Ar) override;
 	//End UActorComponent Interface
 
 	//Begin UMovementComponent Interface
-	virtual float GetMaxSpeed() const OVERRIDE { return MaxSpeed; }
-	virtual void InitializeComponent() OVERRIDE;
+	virtual float GetMaxSpeed() const override { return MaxSpeed; }
+	virtual void InitializeComponent() override;
 	//End UMovementComponent Interface
 
 	/**
@@ -106,7 +106,7 @@ protected:
 	bool HandleHitWall(const FHitResult& Hit, float TimeTick, const FVector& MoveDelta);
 
 	/** Applies bounce logic (if enabled) to affect velocity upon impact, or stops the projectile if bounces are not enabled (or velocity is below threshold). Fires applicable events. */
-	virtual void HandleImpact(const FHitResult& Hit, float TimeSlice=0.f, const FVector& MoveDelta = FVector::ZeroVector) OVERRIDE;
+	virtual void HandleImpact(const FHitResult& Hit, float TimeSlice=0.f, const FVector& MoveDelta = FVector::ZeroVector) override;
 
 	/** Computes result of a bounce and returns the new velocity. */
 	virtual FVector ComputeBounceResult(const FHitResult& Hit, float TimeSlice, const FVector& MoveDelta);
diff --git a/Engine/Source/Runtime/Engine/Classes/GameFramework/SpectatorPawn.h b/Engine/Source/Runtime/Engine/Classes/GameFramework/SpectatorPawn.h
--- a/Engine/Source/Runtime/Engine/Classes/GameFramework/SpectatorPawn.h
+++ b/Engine/Source/Runtime/Engine/Classes/GameFramework/SpectatorPawn.h
@@ -18,6 +18,6 @@ class ENGINE_API ASpectatorPawn : public ADefaultPawn
 
 	// Begin Pawn overrides
 	/** Overridden to avoid changing network role. If subclasses want networked behavior, call the Pawn::PossessedBy() instead. */
-	virtual void PossessedBy(class AController* NewController) OVERRIDE;
+	virtual void PossessedBy(class AController* NewController) override;
 	// End Pawn overrides
 };
diff --git a/Engine/Source/Runtime/Engine/Classes/GameFramework/SpringArmComponent.h b/Engine/Source/Runtime/Engine/Classes/GameFramework/SpringArmComponent.h
--- a/Engine/Source/Runtime/Engine/Classes/GameFramework/SpringArmComponent.h
+++ b/Engine/Source/Runtime/Engine/Classes/GameFramework/SpringArmComponent.h
@@ -79,14 +79,14 @@ class ENGINE_API USpringArmComponent : public USceneComponent
 	FRotator PreviousDesiredRot;
 
 	// UActorComponent interface
-	virtual void OnRegister() OVERRIDE;
-	virtual void TickComponent(float DeltaTime, enum ELevelTick TickType, FActorComponentTickFunction* ThisTickFunction) OVERRIDE;
+	virtual void OnRegister() override;
+	virtual void TickComponent(float DeltaTime, enum ELevelTick TickType, FActorComponentTickFunction* ThisTickFunction) override;
 	// End of UActorComponent interface
 
 	// USceneComponent interface
-	virtual bool HasAnySockets() const OVERRIDE;
-	virtual FTransform GetSocketTransform(FName InSocketName, ERelativeTransformSpace TransformSpace = RTS_World) const OVERRIDE;
-	virtual void QuerySupportedSockets(TArray<FComponentSocketDescription>& OutSockets) const OVERRIDE;
+	virtual bool HasAnySockets() const override;
+	virtual FTransform GetSocketTransform(FName InSocketName, ERelativeTransformSpace TransformSpace = RTS_World) const override;
+	virtual void QuerySupportedSockets(TArray<FComponentSocketDescription>& OutSockets) const override;
 	// End of USceneComponent interface
 
 	/** The name of the socket at the end of the spring arm (looking back towards the spring arm origin) */
diff --git a/Engine/Source/Runtime/Engine/Classes/GameFramework/Volume.h b/Engine/Source/Runtime/Engine/Classes/GameFramework/Volume.h
--- a/Engine/Source/Runtime/Engine/Classes/GameFramework/Volume.h
+++ b/Engine/Source/Runtime/Engine/Classes/GameFramework/Volume.h
@@ -31,15 +31,15 @@ public:
 	* Function that gets called from within Map_Check to allow this actor to check itself
 	* for any potential errors and register them with map check dialog.
 	*/
-	ENGINE_API virtual void CheckForErrors() OVERRIDE;
+	ENGINE_API virtual void CheckForErrors() override;
 #endif // WITH_EDITOR
 	
-	ENGINE_API virtual bool IsLevelBoundsRelevant() const OVERRIDE;
+	ENGINE_API virtual bool IsLevelBoundsRelevant() const override;
 	// End AActor Interface
 
 	// Begin Brush Interface
-	ENGINE_API virtual bool IsStaticBrush() const OVERRIDE;
-	ENGINE_API virtual bool IsVolumeBrush() const OVERRIDE;
+	ENGINE_API virtual bool IsStaticBrush() const override;
+	ENGINE_API virtual bool IsVolumeBrush() const override;
 	// End Brush Interface
 
 	/** @returns true if a sphere/point (with optional radius CheckRadius) overlaps this volume */
@@ -47,8 +47,8 @@ public:
 
 	//Begin UObject Interface
 #if WITH_EDITOR
-	ENGINE_API virtual void PostEditImport() OVERRIDE;
-	ENGINE_API virtual void PostEditChangeProperty(FPropertyChangedEvent& PropertyChangedEvent) OVERRIDE;
+	ENGINE_API virtual void PostEditImport() override;
+	ENGINE_API virtual void PostEditChangeProperty(FPropertyChangedEvent& PropertyChangedEvent) override;
 #endif // WITH_EDITOR
 	//End UObject Interface
 
diff --git a/Engine/Source/Runtime/Engine/Classes/GameFramework/WheeledVehicle.h b/Engine/Source/Runtime/Engine/Classes/GameFramework/WheeledVehicle.h
--- a/Engine/Source/Runtime/Engine/Classes/GameFramework/WheeledVehicle.h
+++ b/Engine/Source/Runtime/Engine/Classes/GameFramework/WheeledVehicle.h
@@ -29,6 +29,6 @@ class ENGINE_API AWheeledVehicle : public APawn
 	}
 
 	// Begin AActor interface
-	virtual void DisplayDebug(class UCanvas* Canvas, const FDebugDisplayInfo& DebugDisplay, float& YL, float& YPos) OVERRIDE;
+	virtual void DisplayDebug(class UCanvas* Canvas, const FDebugDisplayInfo& DebugDisplay, float& YL, float& YPos) override;
 	// End Actor interface
 };
diff --git a/Engine/Source/Runtime/Engine/Classes/GameFramework/WorldSettings.h b/Engine/Source/Runtime/Engine/Classes/GameFramework/WorldSettings.h
--- a/Engine/Source/Runtime/Engine/Classes/GameFramework/WorldSettings.h
+++ b/Engine/Source/Runtime/Engine/Classes/GameFramework/WorldSettings.h
@@ -412,17 +412,17 @@ class ENGINE_API AWorldSettings : public AInfo, public IInterface_AssetUserData
 public:
 	// Begin UObject interface.
 #if WITH_EDITOR
-	virtual void PostEditChangeProperty(FPropertyChangedEvent& PropertyChangedEvent) OVERRIDE;
+	virtual void PostEditChangeProperty(FPropertyChangedEvent& PropertyChangedEvent) override;
 #endif // WITH_EDITOR
 	// End UObject interface.
 
 
 	// Begin AActor interface.
 #if WITH_EDITOR
-	virtual void CheckForErrors() OVERRIDE;
+	virtual void CheckForErrors() override;
 #endif
-	virtual void PreInitializeComponents() OVERRIDE;
-	virtual void PostInitializeComponents() OVERRIDE;
+	virtual void PreInitializeComponents() override;
+	virtual void PostInitializeComponents() override;
 	// End AActor interface.
 
 	/**
@@ -449,9 +449,9 @@ public:
 	virtual void NotifyMatchStarted();
 
 	// Begin IInterface_AssetUserData Interface
-	virtual void AddAssetUserData(UAssetUserData* InUserData) OVERRIDE;
-	virtual void RemoveUserDataOfClass(TSubclassOf<UAssetUserData> InUserDataClass) OVERRIDE;
-	virtual UAssetUserData* GetAssetUserDataOfClass(TSubclassOf<UAssetUserData> InUserDataClass) OVERRIDE;
+	virtual void AddAssetUserData(UAssetUserData* InUserData) override;
+	virtual void RemoveUserDataOfClass(TSubclassOf<UAssetUserData> InUserDataClass) override;
+	virtual UAssetUserData* GetAssetUserDataOfClass(TSubclassOf<UAssetUserData> InUserDataClass) override;
 	// End IInterface_AssetUserData Interface
 
 
@@ -460,7 +460,7 @@ private:
 	// Hidden functions that don't make sense to use on this class.
 	HIDE_ACTOR_TRANSFORM_FUNCTIONS();
 
-	virtual void Serialize( FArchive& Ar ) OVERRIDE;
+	virtual void Serialize( FArchive& Ar ) override;
 
 	/** Toggles world composition mode */
 	void EnabledWorldComposition(bool bEnable);
diff --git a/Engine/Source/Runtime/Engine/Classes/Kismet/BlueprintFunctionLibrary.h b/Engine/Source/Runtime/Engine/Classes/Kismet/BlueprintFunctionLibrary.h
--- a/Engine/Source/Runtime/Engine/Classes/Kismet/BlueprintFunctionLibrary.h
+++ b/Engine/Source/Runtime/Engine/Classes/Kismet/BlueprintFunctionLibrary.h
@@ -11,7 +11,7 @@ class UBlueprintFunctionLibrary : public UObject
 	GENERATED_UCLASS_BODY()
 
 	// UObject interface
-	ENGINE_API virtual int32 GetFunctionCallspace(UFunction* Function, void* Parms, FFrame* Stack) OVERRIDE;
+	ENGINE_API virtual int32 GetFunctionCallspace(UFunction* Function, void* Parms, FFrame* Stack) override;
 	// End of UObject interface
 
 protected:
diff --git a/Engine/Source/Runtime/Engine/Classes/Landscape/ControlPointMeshComponent.h b/Engine/Source/Runtime/Engine/Classes/Landscape/ControlPointMeshComponent.h
--- a/Engine/Source/Runtime/Engine/Classes/Landscape/ControlPointMeshComponent.h
+++ b/Engine/Source/Runtime/Engine/Classes/Landscape/ControlPointMeshComponent.h
@@ -16,7 +16,7 @@ class UControlPointMeshComponent : public UStaticMeshComponent
 
 	//Begin UPrimitiveComponent Interface
 #if WITH_EDITOR
-	virtual bool ShouldRenderSelected() const OVERRIDE
+	virtual bool ShouldRenderSelected() const override
 	{
 		return Super::ShouldRenderSelected() || bSelected;
 	}
diff --git a/Engine/Source/Runtime/Engine/Classes/Landscape/Landscape.h b/Engine/Source/Runtime/Engine/Classes/Landscape/Landscape.h
--- a/Engine/Source/Runtime/Engine/Classes/Landscape/Landscape.h
+++ b/Engine/Source/Runtime/Engine/Classes/Landscape/Landscape.h
@@ -29,16 +29,16 @@ class ALandscape : public ALandscapeProxy
 
 	// Begin AActor Interface
 #if WITH_EDITOR
-	virtual void CheckForErrors() OVERRIDE;
-	virtual void Destroyed() OVERRIDE;
+	virtual void CheckForErrors() override;
+	virtual void Destroyed() override;
 #endif
 	// End AActor Interface
 
 	// Begin ALandscapeProxy Interface
-	ENGINE_API virtual ALandscape* GetLandscapeActor() OVERRIDE;
+	ENGINE_API virtual ALandscape* GetLandscapeActor() override;
 #if WITH_EDITOR
-	virtual UMaterialInterface* GetLandscapeMaterial() const OVERRIDE;
-	virtual UMaterialInterface* GetLandscapeHoleMaterial() const OVERRIDE;
+	virtual UMaterialInterface* GetLandscapeMaterial() const override;
+	virtual UMaterialInterface* GetLandscapeHoleMaterial() const override;
 	// End ALandscapeProxy Interface
 
 	ENGINE_API bool HasAllComponent(); // determine all component is in this actor
@@ -54,11 +54,11 @@ class ALandscape : public ALandscapeProxy
 	static void SplitHeightmap(ULandscapeComponent* Comp, bool bMoveToCurrentLevel = false);
 	
 	// Begin UObject interface.
-	virtual void PreSave() OVERRIDE;
-	virtual void PostEditChangeProperty(FPropertyChangedEvent& PropertyChangedEvent) OVERRIDE;
-	virtual void PostEditMove(bool bFinished) OVERRIDE;
+	virtual void PreSave() override;
+	virtual void PostEditChangeProperty(FPropertyChangedEvent& PropertyChangedEvent) override;
+	virtual void PostEditMove(bool bFinished) override;
 #endif
-	virtual void PostLoad() OVERRIDE;
+	virtual void PostLoad() override;
 	// End UObject Interface
 
 
diff --git a/Engine/Source/Runtime/Engine/Classes/Landscape/LandscapeComponent.h b/Engine/Source/Runtime/Engine/Classes/Landscape/LandscapeComponent.h
--- a/Engine/Source/Runtime/Engine/Classes/Landscape/LandscapeComponent.h
+++ b/Engine/Source/Runtime/Engine/Classes/Landscape/LandscapeComponent.h
@@ -238,16 +238,16 @@ public:
 	FLandscapeComponentDerivedData PlatformData;
 
 	// Begin UObject interface.	
-	virtual void PostInitProperties() OVERRIDE;	
-	virtual void Serialize(FArchive& Ar) OVERRIDE;
-	virtual void BeginDestroy() OVERRIDE;
+	virtual void PostInitProperties() override;	
+	virtual void Serialize(FArchive& Ar) override;
+	virtual void BeginDestroy() override;
 	static void AddReferencedObjects(UObject* InThis, FReferenceCollector& Collector);
-	virtual void PostDuplicate(bool bDuplicateForPIE) OVERRIDE;
+	virtual void PostDuplicate(bool bDuplicateForPIE) override;
 #if WITH_EDITOR
-	virtual void PostLoad() OVERRIDE;
-	virtual void PostEditUndo() OVERRIDE;
-	virtual void PreEditChange(UProperty* PropertyThatWillChange) OVERRIDE;
-	virtual void PostEditChangeProperty(FPropertyChangedEvent& PropertyChangedEvent) OVERRIDE;
+	virtual void PostLoad() override;
+	virtual void PostEditUndo() override;
+	virtual void PreEditChange(UProperty* PropertyThatWillChange) override;
+	virtual void PostEditChangeProperty(FPropertyChangedEvent& PropertyChangedEvent) override;
 	// End UObject interface
 
 	/** Fix up component layers, weightmaps
@@ -255,30 +255,30 @@ public:
 	ENGINE_API void FixupWeightmaps();
 	
 	// Begin UPrimitiveComponent interface.
-	virtual bool GetLightMapResolution( int32& Width, int32& Height ) const OVERRIDE;
-	virtual void GetLightAndShadowMapMemoryUsage( int32& LightMapMemoryUsage, int32& ShadowMapMemoryUsage ) const OVERRIDE;
-	virtual void GetStaticLightingInfo(FStaticLightingPrimitiveInfo& OutPrimitiveInfo,const TArray<ULightComponent*>& InRelevantLights,const FLightingBuildOptions& Options) OVERRIDE;
+	virtual bool GetLightMapResolution( int32& Width, int32& Height ) const override;
+	virtual void GetLightAndShadowMapMemoryUsage( int32& LightMapMemoryUsage, int32& ShadowMapMemoryUsage ) const override;
+	virtual void GetStaticLightingInfo(FStaticLightingPrimitiveInfo& OutPrimitiveInfo,const TArray<ULightComponent*>& InRelevantLights,const FLightingBuildOptions& Options) override;
 #endif
-	virtual void GetUsedMaterials( TArray<UMaterialInterface*>& OutMaterials ) const OVERRIDE;
-	virtual FPrimitiveSceneProxy* CreateSceneProxy() OVERRIDE;
-	virtual ELightMapInteractionType GetStaticLightingType() const OVERRIDE { return LMIT_Texture;	}
-	virtual void GetStreamingTextureInfo(TArray<FStreamingTexturePrimitiveInfo>& OutStreamingTextures) const OVERRIDE;
+	virtual void GetUsedMaterials( TArray<UMaterialInterface*>& OutMaterials ) const override;
+	virtual FPrimitiveSceneProxy* CreateSceneProxy() override;
+	virtual ELightMapInteractionType GetStaticLightingType() const override { return LMIT_Texture;	}
+	virtual void GetStreamingTextureInfo(TArray<FStreamingTexturePrimitiveInfo>& OutStreamingTextures) const override;
 
 #if WITH_EDITOR
-	virtual int32 GetNumMaterials() const OVERRIDE;
-	virtual class UMaterialInterface* GetMaterial(int32 ElementIndex) const OVERRIDE;
-	virtual void SetMaterial(int32 ElementIndex, class UMaterialInterface* Material) OVERRIDE;
+	virtual int32 GetNumMaterials() const override;
+	virtual class UMaterialInterface* GetMaterial(int32 ElementIndex) const override;
+	virtual void SetMaterial(int32 ElementIndex, class UMaterialInterface* Material) override;
 #endif
 	// End UPrimitiveComponent interface.
 
 	// Begin USceneComponent interface.
-	virtual void DestroyComponent() OVERRIDE;
-	virtual FBoxSphereBounds CalcBounds(const FTransform & LocalToWorld) const OVERRIDE;
+	virtual void DestroyComponent() override;
+	virtual FBoxSphereBounds CalcBounds(const FTransform & LocalToWorld) const override;
 	// End USceneComponent interface.
 
 	// Begin UActorComponent interface.
-	virtual void OnRegister() OVERRIDE;
-	virtual void OnUnregister() OVERRIDE;
+	virtual void OnRegister() override;
+	virtual void OnUnregister() override;
 	// End UActorComponent interface.
 
 
diff --git a/Engine/Source/Runtime/Engine/Classes/Landscape/LandscapeGizmoActiveActor.h b/Engine/Source/Runtime/Engine/Classes/Landscape/LandscapeGizmoActiveActor.h
--- a/Engine/Source/Runtime/Engine/Classes/Landscape/LandscapeGizmoActiveActor.h
+++ b/Engine/Source/Runtime/Engine/Classes/Landscape/LandscapeGizmoActiveActor.h
@@ -104,15 +104,15 @@ public:
 
 #if WITH_EDITOR
 	// Begin UObject interface.
-	virtual void PostEditChangeProperty(FPropertyChangedEvent& PropertyChangedEvent) OVERRIDE;
-	virtual void PostEditMove(bool bFinished) OVERRIDE;
+	virtual void PostEditChangeProperty(FPropertyChangedEvent& PropertyChangedEvent) override;
+	virtual void PostEditMove(bool bFinished) override;
 	// End UObject interface.
 
 	virtual FVector SnapToLandscapeGrid(const FVector& GizmoLocation) const;
 	virtual FRotator SnapToLandscapeGrid(const FRotator& GizmoRotation) const;
 
-	virtual void EditorApplyTranslation(const FVector& DeltaTranslation, bool bAltDown, bool bShiftDown, bool bCtrlDown) OVERRIDE;
-	virtual void EditorApplyRotation(const FRotator& DeltaRotation, bool bAltDown, bool bShiftDown, bool bCtrlDown) OVERRIDE;
+	virtual void EditorApplyTranslation(const FVector& DeltaTranslation, bool bAltDown, bool bShiftDown, bool bCtrlDown) override;
+	virtual void EditorApplyRotation(const FRotator& DeltaRotation, bool bAltDown, bool bShiftDown, bool bCtrlDown) override;
 
 	/**
 	 * Whenever the decal actor has moved:
diff --git a/Engine/Source/Runtime/Engine/Classes/Landscape/LandscapeGizmoActor.h b/Engine/Source/Runtime/Engine/Classes/Landscape/LandscapeGizmoActor.h
--- a/Engine/Source/Runtime/Engine/Classes/Landscape/LandscapeGizmoActor.h
+++ b/Engine/Source/Runtime/Engine/Classes/Landscape/LandscapeGizmoActor.h
@@ -40,7 +40,7 @@ class ALandscapeGizmoActor : public AActor
 	virtual void Duplicate(ALandscapeGizmoActor* Gizmo); 
 	//virtual void EditorApplyTranslation(const FVector& DeltaTranslation, bool bAltDown, bool bShiftDown, bool bCtrlDown);
 
-	bool EditorCanAttachTo(const AActor* InParent, FText& OutReason) const OVERRIDE
+	bool EditorCanAttachTo(const AActor* InParent, FText& OutReason) const override
 	{
 		return false;
 	}
@@ -49,7 +49,7 @@ class ALandscapeGizmoActor : public AActor
 	/** 
 	 * Indicates whether this actor should participate in level bounds calculations
 	 */
-	virtual bool IsLevelBoundsRelevant() const OVERRIDE { return false; }
+	virtual bool IsLevelBoundsRelevant() const override { return false; }
 };
 
 
diff --git a/Engine/Source/Runtime/Engine/Classes/Landscape/LandscapeGizmoRenderComponent.h b/Engine/Source/Runtime/Engine/Classes/Landscape/LandscapeGizmoRenderComponent.h
--- a/Engine/Source/Runtime/Engine/Classes/Landscape/LandscapeGizmoRenderComponent.h
+++ b/Engine/Source/Runtime/Engine/Classes/Landscape/LandscapeGizmoRenderComponent.h
@@ -11,11 +11,11 @@ class ULandscapeGizmoRenderComponent : public UPrimitiveComponent
 
 
 	// Begin UPrimitiveComponent Interface
-	virtual FPrimitiveSceneProxy* CreateSceneProxy() OVERRIDE;
+	virtual FPrimitiveSceneProxy* CreateSceneProxy() override;
 	// End UPrimitiveComponent Interface
 
 	// Begin USceneComponent interface.
-	virtual FBoxSphereBounds CalcBounds(const FTransform & LocalToWorld) const OVERRIDE;
+	virtual FBoxSphereBounds CalcBounds(const FTransform & LocalToWorld) const override;
 	// End USceneComponent interface.
 };
 
diff --git a/Engine/Source/Runtime/Engine/Classes/Landscape/LandscapeHeightfieldCollisionComponent.h b/Engine/Source/Runtime/Engine/Classes/Landscape/LandscapeHeightfieldCollisionComponent.h
--- a/Engine/Source/Runtime/Engine/Classes/Landscape/LandscapeHeightfieldCollisionComponent.h
+++ b/Engine/Source/Runtime/Engine/Classes/Landscape/LandscapeHeightfieldCollisionComponent.h
@@ -124,33 +124,33 @@ class ULandscapeHeightfieldCollisionComponent : public UPrimitiveComponent
 	};
 
 	// Begin UActorComponent interface.
-	virtual void CreatePhysicsState() OVERRIDE;
-	virtual void ApplyWorldOffset(const FVector& InOffset, bool bWorldShift) OVERRIDE;
+	virtual void CreatePhysicsState() override;
+	virtual void ApplyWorldOffset(const FVector& InOffset, bool bWorldShift) override;
 	// End UActorComponent interface.
 
 	// Begin USceneComponent interface.
-	virtual void DestroyComponent() OVERRIDE;
-	virtual FBoxSphereBounds CalcBounds(const FTransform &BoundTransform) const OVERRIDE;
+	virtual void DestroyComponent() override;
+	virtual FBoxSphereBounds CalcBounds(const FTransform &BoundTransform) const override;
 
-	virtual ECollisionEnabled::Type GetCollisionEnabled() const OVERRIDE;
-	virtual ECollisionResponse GetCollisionResponseToChannel(ECollisionChannel Channel) const OVERRIDE;
-	virtual ECollisionChannel GetCollisionObjectType() const OVERRIDE;
-	virtual const FCollisionResponseContainer& GetCollisionResponseToChannels() const OVERRIDE;
+	virtual ECollisionEnabled::Type GetCollisionEnabled() const override;
+	virtual ECollisionResponse GetCollisionResponseToChannel(ECollisionChannel Channel) const override;
+	virtual ECollisionChannel GetCollisionObjectType() const override;
+	virtual const FCollisionResponseContainer& GetCollisionResponseToChannels() const override;
 	// End USceneComponent interface.
 
 	// Begin UPrimitiveComponent interface
-	virtual bool DoCustomNavigableGeometryExport(struct FNavigableGeometryExport* GeomExport) const OVERRIDE;
+	virtual bool DoCustomNavigableGeometryExport(struct FNavigableGeometryExport* GeomExport) const override;
 	//End UPrimitiveComponent interface
 
 	// Begin UObject Interface.
-	virtual void Serialize(FArchive& Ar) OVERRIDE;
-	virtual void BeginDestroy() OVERRIDE;
-	virtual void PostLoad() OVERRIDE;
+	virtual void Serialize(FArchive& Ar) override;
+	virtual void BeginDestroy() override;
+	virtual void PostLoad() override;
 #if WITH_EDITOR
-	virtual void ExportCustomProperties(FOutputDevice& Out, uint32 Indent) OVERRIDE;
-	virtual void ImportCustomProperties(const TCHAR* SourceText, FFeedbackContext* Warn) OVERRIDE;
-	virtual void PostEditImport() OVERRIDE;
-	virtual void PostEditUndo() OVERRIDE;
+	virtual void ExportCustomProperties(FOutputDevice& Out, uint32 Indent) override;
+	virtual void ImportCustomProperties(const TCHAR* SourceText, FFeedbackContext* Warn) override;
+	virtual void PostEditImport() override;
+	virtual void PostEditUndo() override;
 	// End UObject Interface.
 
 	// Update Collision object for add LandscapeComponent tool
diff --git a/Engine/Source/Runtime/Engine/Classes/Landscape/LandscapeInfo.h b/Engine/Source/Runtime/Engine/Classes/Landscape/LandscapeInfo.h
--- a/Engine/Source/Runtime/Engine/Classes/Landscape/LandscapeInfo.h
+++ b/Engine/Source/Runtime/Engine/Classes/Landscape/LandscapeInfo.h
@@ -150,10 +150,10 @@ public:
 	TMap<FIntPoint,float> SelectedRegion;
 
 	// Begin UObject Interface.
-	virtual void Serialize(FArchive& Ar) OVERRIDE;
-	virtual void BeginDestroy() OVERRIDE;
+	virtual void Serialize(FArchive& Ar) override;
+	virtual void BeginDestroy() override;
 #if WITH_EDITOR
-	virtual void PostEditUndo() OVERRIDE;
+	virtual void PostEditUndo() override;
 	// End UObject Interface
 
 	// @todo document 
diff --git a/Engine/Source/Runtime/Engine/Classes/Landscape/LandscapeLayerInfoObject.h b/Engine/Source/Runtime/Engine/Classes/Landscape/LandscapeLayerInfoObject.h
--- a/Engine/Source/Runtime/Engine/Classes/Landscape/LandscapeLayerInfoObject.h
+++ b/Engine/Source/Runtime/Engine/Classes/Landscape/LandscapeLayerInfoObject.h
@@ -25,8 +25,8 @@ class ULandscapeLayerInfoObject : public UObject
 
 #if WITH_EDITOR
 	// Begin UObject Interface
-	virtual void PostEditChangeProperty(FPropertyChangedEvent& PropertyChangedEvent) OVERRIDE;
-	virtual void PostLoad() OVERRIDE;
+	virtual void PostEditChangeProperty(FPropertyChangedEvent& PropertyChangedEvent) override;
+	virtual void PostLoad() override;
 	// End UObject Interface
 #endif
 };
diff --git a/Engine/Source/Runtime/Engine/Classes/Landscape/LandscapeMeshCollisionComponent.h b/Engine/Source/Runtime/Engine/Classes/Landscape/LandscapeMeshCollisionComponent.h
--- a/Engine/Source/Runtime/Engine/Classes/Landscape/LandscapeMeshCollisionComponent.h
+++ b/Engine/Source/Runtime/Engine/Classes/Landscape/LandscapeMeshCollisionComponent.h
@@ -66,32 +66,32 @@ class ULandscapeMeshCollisionComponent : public ULandscapeHeightfieldCollisionCo
 	TRefCountPtr<struct FPhysXMeshRef>			MeshRef;
 
 	// Begin UActorComponent interface.
-	virtual void CreatePhysicsState() OVERRIDE;
-	virtual void ApplyWorldOffset(const FVector& InOffset, bool bWorldShift) OVERRIDE;
+	virtual void CreatePhysicsState() override;
+	virtual void ApplyWorldOffset(const FVector& InOffset, bool bWorldShift) override;
 	// End UActorComponent interface.
 
 	// Begin USceneComponent interface.
-	virtual void DestroyComponent() OVERRIDE;
+	virtual void DestroyComponent() override;
 	// End USceneComponent interface.
 
 	// Begin UPrimitiveComponent interface
-	virtual bool DoCustomNavigableGeometryExport(struct FNavigableGeometryExport* GeomExport) const OVERRIDE;
+	virtual bool DoCustomNavigableGeometryExport(struct FNavigableGeometryExport* GeomExport) const override;
 	//End UPrimitiveComponent interface
 
 	// Begin UObject Interface.
-	virtual void Serialize(FArchive& Ar) OVERRIDE;
-	virtual void BeginDestroy() OVERRIDE;
+	virtual void Serialize(FArchive& Ar) override;
+	virtual void BeginDestroy() override;
 #if WITH_EDITOR
-	virtual void ExportCustomProperties(FOutputDevice& Out, uint32 Indent) OVERRIDE;
-	virtual void ImportCustomProperties(const TCHAR* SourceText, FFeedbackContext* Warn) OVERRIDE;
+	virtual void ExportCustomProperties(FOutputDevice& Out, uint32 Indent) override;
+	virtual void ImportCustomProperties(const TCHAR* SourceText, FFeedbackContext* Warn) override;
 
-	virtual bool CookCollsionData(const FName& Format, bool bUseOnlyDefMaterial, TArray<uint8>& OutCookedData, TArray<UPhysicalMaterial*>& OutMaterails) const OVERRIDE;
+	virtual bool CookCollsionData(const FName& Format, bool bUseOnlyDefMaterial, TArray<uint8>& OutCookedData, TArray<UPhysicalMaterial*>& OutMaterails) const override;
 #endif
 	// End UObject Interface.
 
 	// Begin ULandscapeHeightfieldCollisionComponent Interface
-	virtual void CreateCollisionObject() OVERRIDE;
-	virtual void RecreateCollision(bool bUpdateAddCollision = true) OVERRIDE;
+	virtual void CreateCollisionObject() override;
+	virtual void RecreateCollision(bool bUpdateAddCollision = true) override;
 	// End ULandscapeHeightfieldCollisionComponent Interface
 };
 
diff --git a/Engine/Source/Runtime/Engine/Classes/Landscape/LandscapeProxy.h b/Engine/Source/Runtime/Engine/Classes/Landscape/LandscapeProxy.h
--- a/Engine/Source/Runtime/Engine/Classes/Landscape/LandscapeProxy.h
+++ b/Engine/Source/Runtime/Engine/Classes/Landscape/LandscapeProxy.h
@@ -331,19 +331,19 @@ public:
 	TMap<UTexture2D*, struct FLandscapeWeightmapUsage> WeightmapUsageMap;
 
 	// Begin AActor Interface
-	virtual void UnregisterAllComponents() OVERRIDE;
-	virtual void RegisterAllComponents() OVERRIDE;
-	virtual void RerunConstructionScripts() OVERRIDE {}
-	virtual bool IsLevelBoundsRelevant() const OVERRIDE { return true; }
-	virtual bool UpdateNavigationRelevancy() OVERRIDE;
-	virtual FBox GetComponentsBoundingBox(bool bNonColliding = false) const OVERRIDE;
+	virtual void UnregisterAllComponents() override;
+	virtual void RegisterAllComponents() override;
+	virtual void RerunConstructionScripts() override {}
+	virtual bool IsLevelBoundsRelevant() const override { return true; }
+	virtual bool UpdateNavigationRelevancy() override;
+	virtual FBox GetComponentsBoundingBox(bool bNonColliding = false) const override;
 #if WITH_EDITOR
-	virtual void Destroyed() OVERRIDE;
-	virtual void EditorApplyScale(const FVector& DeltaScale, const FVector* PivotLocation, bool bAltDown, bool bShiftDown, bool bCtrlDown) OVERRIDE;
-	virtual void PostEditMove(bool bFinished) OVERRIDE;
-	virtual bool ShouldImport(FString* ActorPropString, bool IsMovingLevel) OVERRIDE;
-	virtual bool ShouldExport() OVERRIDE;
-	virtual bool GetSelectedComponents(TArray<UObject*>& SelectedObjects) OVERRIDE;
+	virtual void Destroyed() override;
+	virtual void EditorApplyScale(const FVector& DeltaScale, const FVector* PivotLocation, bool bAltDown, bool bShiftDown, bool bCtrlDown) override;
+	virtual void PostEditMove(bool bFinished) override;
+	virtual bool ShouldImport(FString* ActorPropString, bool IsMovingLevel) override;
+	virtual bool ShouldExport() override;
+	virtual bool GetSelectedComponents(TArray<UObject*>& SelectedObjects) override;
 	// End AActor Interface
 #endif	//WITH_EDITOR
 
@@ -351,20 +351,20 @@ public:
 	virtual class ALandscape* GetLandscapeActor();
 
 	// Begin INavRelevantActorInterface Interface
-	virtual bool DoesSupplyPerComponentNavigationCollision() const OVERRIDE{ return true; }
+	virtual bool DoesSupplyPerComponentNavigationCollision() const override{ return true; }
 	// End INavRelevantActorInterface Interface
 
 	// Begin UObject interface.
-	virtual void Serialize(FArchive& Ar) OVERRIDE;
+	virtual void Serialize(FArchive& Ar) override;
 	static void AddReferencedObjects(UObject* InThis, FReferenceCollector& Collector);
-	virtual void PostLoad() OVERRIDE;
+	virtual void PostLoad() override;
 
 #if WITH_EDITOR
-	virtual void PostEditChangeProperty(FPropertyChangedEvent& PropertyChangedEvent) OVERRIDE;
-	virtual void PostEditChangeChainProperty(FPropertyChangedChainEvent& PropertyChangedEvent) OVERRIDE;
-	virtual void PreEditUndo() OVERRIDE;
-	virtual void PostEditUndo() OVERRIDE;
-	virtual void PostEditImport() OVERRIDE;
+	virtual void PostEditChangeProperty(FPropertyChangedEvent& PropertyChangedEvent) override;
+	virtual void PostEditChangeChainProperty(FPropertyChangedChainEvent& PropertyChangedEvent) override;
+	virtual void PreEditUndo() override;
+	virtual void PostEditUndo() override;
+	virtual void PostEditImport() override;
 	// End UObject Interface
 
 	FLandscapeLayerStruct* GetLayerInfo_Deprecated(FName LayerName);
diff --git a/Engine/Source/Runtime/Engine/Classes/Landscape/LandscapeSplineControlPoint.h b/Engine/Source/Runtime/Engine/Classes/Landscape/LandscapeSplineControlPoint.h
--- a/Engine/Source/Runtime/Engine/Classes/Landscape/LandscapeSplineControlPoint.h
+++ b/Engine/Source/Runtime/Engine/Classes/Landscape/LandscapeSplineControlPoint.h
@@ -161,12 +161,12 @@ public:
 #endif // WITH_EDITOR
 
 	// Begin UObject Interface
-	virtual void PostLoad() OVERRIDE;
+	virtual void PostLoad() override;
 #if WITH_EDITOR
-	virtual void PostEditUndo() OVERRIDE;
-	virtual void PostDuplicate(bool bDuplicateForPIE) OVERRIDE;
-	virtual void PostEditImport() OVERRIDE;
-	virtual void PostEditChangeProperty(FPropertyChangedEvent& PropertyChangedEvent) OVERRIDE;
+	virtual void PostEditUndo() override;
+	virtual void PostDuplicate(bool bDuplicateForPIE) override;
+	virtual void PostEditImport() override;
+	virtual void PostEditChangeProperty(FPropertyChangedEvent& PropertyChangedEvent) override;
 	// End UObject Interface
 #endif // WITH_EDITOR
 
diff --git a/Engine/Source/Runtime/Engine/Classes/Landscape/LandscapeSplineSegment.h b/Engine/Source/Runtime/Engine/Classes/Landscape/LandscapeSplineSegment.h
--- a/Engine/Source/Runtime/Engine/Classes/Landscape/LandscapeSplineSegment.h
+++ b/Engine/Source/Runtime/Engine/Classes/Landscape/LandscapeSplineSegment.h
@@ -228,17 +228,17 @@ public:
 	virtual void FindNearest( const FVector& InLocation, float& t, FVector& OutLocation, FVector& OutTangent );
 
 	// Begin UObject Interface
-	virtual void Serialize(FArchive& Ar) OVERRIDE;
-	virtual void PostLoad() OVERRIDE;
-	virtual bool Modify(bool bAlwaysMarkDirty = true) OVERRIDE;
+	virtual void Serialize(FArchive& Ar) override;
+	virtual void PostLoad() override;
+	virtual bool Modify(bool bAlwaysMarkDirty = true) override;
 #if WITH_EDITOR
-	virtual void PostEditUndo() OVERRIDE;
-	virtual void PostDuplicate(bool bDuplicateForPIE) OVERRIDE;
-	virtual void PostEditImport() OVERRIDE;
-	virtual void PostEditChangeProperty(FPropertyChangedEvent& PropertyChangedEvent) OVERRIDE;
+	virtual void PostEditUndo() override;
+	virtual void PostDuplicate(bool bDuplicateForPIE) override;
+	virtual void PostEditImport() override;
+	virtual void PostEditChangeProperty(FPropertyChangedEvent& PropertyChangedEvent) override;
 #endif // WITH_EDITOR
 protected:
-	virtual void PostInitProperties() OVERRIDE;
+	virtual void PostInitProperties() override;
 public:
 	// End UObject Interface
 
diff --git a/Engine/Source/Runtime/Engine/Classes/Landscape/LandscapeSplinesComponent.h b/Engine/Source/Runtime/Engine/Classes/Landscape/LandscapeSplinesComponent.h
--- a/Engine/Source/Runtime/Engine/Classes/Landscape/LandscapeSplinesComponent.h
+++ b/Engine/Source/Runtime/Engine/Classes/Landscape/LandscapeSplinesComponent.h
@@ -47,21 +47,21 @@ public:
 #endif
 
 	// Begin UObject Interface
-	virtual void PostLoad() OVERRIDE;
+	virtual void PostLoad() override;
 #if WITH_EDITOR
-	virtual void PostEditChangeProperty(FPropertyChangedEvent& PropertyChangedEvent) OVERRIDE;
+	virtual void PostEditChangeProperty(FPropertyChangedEvent& PropertyChangedEvent) override;
 #endif // WITH_EDITOR
 	// End UObject Interface
 
 	// Begin UActorComponent interface
-	virtual void OnRegister() OVERRIDE;
-	virtual void OnUnregister() OVERRIDE;
+	virtual void OnRegister() override;
+	virtual void OnUnregister() override;
 	// End UActorComponent interface
 	
 	// Begin UPrimitiveComponent interface.
 #if WITH_EDITOR
-	virtual FPrimitiveSceneProxy* CreateSceneProxy() OVERRIDE;
+	virtual FPrimitiveSceneProxy* CreateSceneProxy() override;
 #endif
-	virtual FBoxSphereBounds CalcBounds(const FTransform& LocalToWorld) const OVERRIDE;
+	virtual FBoxSphereBounds CalcBounds(const FTransform& LocalToWorld) const override;
 	// End UPrimitiveComponent interface.
 };
diff --git a/Engine/Source/Runtime/Engine/Classes/Materials/Material.h b/Engine/Source/Runtime/Engine/Classes/Materials/Material.h
--- a/Engine/Source/Runtime/Engine/Classes/Materials/Material.h
+++ b/Engine/Source/Runtime/Engine/Classes/Materials/Material.h
@@ -721,32 +721,32 @@ private:
 public:
 
 	// Begin UMaterialInterface interface.
-	ENGINE_API virtual UMaterial* GetMaterial() OVERRIDE;
-	ENGINE_API virtual const UMaterial* GetMaterial() const OVERRIDE;
-	ENGINE_API virtual const UMaterial* GetMaterial_Concurrent(TMicRecursionGuard& RecursionGuard) const OVERRIDE;
+	ENGINE_API virtual UMaterial* GetMaterial() override;
+	ENGINE_API virtual const UMaterial* GetMaterial() const override;
+	ENGINE_API virtual const UMaterial* GetMaterial_Concurrent(TMicRecursionGuard& RecursionGuard) const override;
 	ENGINE_API virtual bool GetParameterDesc(FName ParameterName, FString& OutDesc) const;
-	ENGINE_API virtual bool GetVectorParameterValue(FName ParameterName,FLinearColor& OutValue) const OVERRIDE;
-	ENGINE_API virtual bool GetScalarParameterValue(FName ParameterName,float& OutValue) const OVERRIDE;
-	ENGINE_API virtual bool GetTextureParameterValue(FName ParameterName,class UTexture*& OutValue) const OVERRIDE;
-	ENGINE_API virtual bool GetFontParameterValue(FName ParameterName,class UFont*& OutFontValue,int32& OutFontPage) const OVERRIDE;
-	ENGINE_API virtual bool GetGroupName(FName ParameterName, FName& OutDesc) const OVERRIDE;
-	ENGINE_API virtual bool GetRefractionSettings(float& OutBiasValue) const OVERRIDE;
-	ENGINE_API virtual FMaterialRenderProxy* GetRenderProxy(bool Selected, bool bHovered=false) const OVERRIDE;
-	ENGINE_API virtual UPhysicalMaterial* GetPhysicalMaterial() const OVERRIDE;
-	ENGINE_API virtual void GetUsedTextures(TArray<UTexture*>& OutTextures, EMaterialQualityLevel::Type QualityLevel, bool bAllQualityLevels) const OVERRIDE;
-	ENGINE_API virtual void OverrideTexture( const UTexture* InTextureToOverride, UTexture* OverrideTexture ) OVERRIDE;
-	ENGINE_API virtual bool CheckMaterialUsage(const EMaterialUsage Usage, const bool bSkipPrim = false) OVERRIDE;
-	ENGINE_API virtual bool CheckMaterialUsage_Concurrent(const EMaterialUsage Usage, const bool bSkipPrim = false) const OVERRIDE;
+	ENGINE_API virtual bool GetVectorParameterValue(FName ParameterName,FLinearColor& OutValue) const override;
+	ENGINE_API virtual bool GetScalarParameterValue(FName ParameterName,float& OutValue) const override;
+	ENGINE_API virtual bool GetTextureParameterValue(FName ParameterName,class UTexture*& OutValue) const override;
+	ENGINE_API virtual bool GetFontParameterValue(FName ParameterName,class UFont*& OutFontValue,int32& OutFontPage) const override;
+	ENGINE_API virtual bool GetGroupName(FName ParameterName, FName& OutDesc) const override;
+	ENGINE_API virtual bool GetRefractionSettings(float& OutBiasValue) const override;
+	ENGINE_API virtual FMaterialRenderProxy* GetRenderProxy(bool Selected, bool bHovered=false) const override;
+	ENGINE_API virtual UPhysicalMaterial* GetPhysicalMaterial() const override;
+	ENGINE_API virtual void GetUsedTextures(TArray<UTexture*>& OutTextures, EMaterialQualityLevel::Type QualityLevel, bool bAllQualityLevels) const override;
+	ENGINE_API virtual void OverrideTexture( const UTexture* InTextureToOverride, UTexture* OverrideTexture ) override;
+	ENGINE_API virtual bool CheckMaterialUsage(const EMaterialUsage Usage, const bool bSkipPrim = false) override;
+	ENGINE_API virtual bool CheckMaterialUsage_Concurrent(const EMaterialUsage Usage, const bool bSkipPrim = false) const override;
 	ENGINE_API virtual FMaterialResource* AllocateResource();
-	ENGINE_API virtual FMaterialResource* GetMaterialResource(ERHIFeatureLevel::Type InFeatureLevel, EMaterialQualityLevel::Type QualityLevel = EMaterialQualityLevel::Num) OVERRIDE;
-	ENGINE_API virtual const FMaterialResource* GetMaterialResource(ERHIFeatureLevel::Type InFeatureLevel, EMaterialQualityLevel::Type QualityLevel = EMaterialQualityLevel::Num) const OVERRIDE;
-	ENGINE_API virtual bool GetStaticSwitchParameterValue(FName ParameterName,bool &OutValue,FGuid &OutExpressionGuid) OVERRIDE;
-	ENGINE_API virtual bool GetStaticComponentMaskParameterValue(FName ParameterName, bool &R, bool &G, bool &B, bool &A, FGuid &OutExpressionGuid) OVERRIDE;
-	ENGINE_API virtual bool GetTerrainLayerWeightParameterValue(FName ParameterName, int32& OutWeightmapIndex, FGuid &OutExpressionGuid) OVERRIDE;
-	ENGINE_API virtual bool UpdateLightmassTextureTracking() OVERRIDE;
+	ENGINE_API virtual FMaterialResource* GetMaterialResource(ERHIFeatureLevel::Type InFeatureLevel, EMaterialQualityLevel::Type QualityLevel = EMaterialQualityLevel::Num) override;
+	ENGINE_API virtual const FMaterialResource* GetMaterialResource(ERHIFeatureLevel::Type InFeatureLevel, EMaterialQualityLevel::Type QualityLevel = EMaterialQualityLevel::Num) const override;
+	ENGINE_API virtual bool GetStaticSwitchParameterValue(FName ParameterName,bool &OutValue,FGuid &OutExpressionGuid) override;
+	ENGINE_API virtual bool GetStaticComponentMaskParameterValue(FName ParameterName, bool &R, bool &G, bool &B, bool &A, FGuid &OutExpressionGuid) override;
+	ENGINE_API virtual bool GetTerrainLayerWeightParameterValue(FName ParameterName, int32& OutWeightmapIndex, FGuid &OutExpressionGuid) override;
+	ENGINE_API virtual bool UpdateLightmassTextureTracking() override;
 	ENGINE_API virtual bool GetTexturesInPropertyChain(EMaterialProperty InProperty, TArray<UTexture*>& OutTextures, 
-		TArray<FName>* OutTextureParamNames, class FStaticParameterSet* InStaticParameterSet) OVERRIDE;
-	ENGINE_API virtual void RecacheUniformExpressions() const OVERRIDE;
+		TArray<FName>* OutTextureParamNames, class FStaticParameterSet* InStaticParameterSet) override;
+	ENGINE_API virtual void RecacheUniformExpressions() const override;
 
 	ENGINE_API virtual float GetOpacityMaskClipValue_Internal() const;
 	ENGINE_API virtual EBlendMode GetBlendMode_Internal() const;
@@ -759,28 +759,28 @@ public:
 	ENGINE_API virtual bool IsPropertyActive(EMaterialProperty InProperty) const;
 	/** Allows material properties to be compiled with the option of being overridden by the material attributes input. */
 	ENGINE_API virtual int32 CompileProperty( class FMaterialCompiler* Compiler, EMaterialProperty Property, float DefaultFloat, FLinearColor DefaultColor, const FVector4& DefaultVector );
-	ENGINE_API virtual void ForceRecompileForRendering() OVERRIDE;
+	ENGINE_API virtual void ForceRecompileForRendering() override;
 	// End UMaterialInterface interface.
 
 	// Begin UObject Interface
 	ENGINE_API virtual void PreSave();
 	ENGINE_API virtual void PostInitProperties();	
-	ENGINE_API virtual void Serialize(FArchive& Ar) OVERRIDE;
-	ENGINE_API virtual void PostDuplicate(bool bDuplicateForPIE) OVERRIDE;
-	ENGINE_API virtual void PostLoad() OVERRIDE;
-	ENGINE_API virtual void BeginCacheForCookedPlatformData( const ITargetPlatform *TargetPlatform ) OVERRIDE;
-	ENGINE_API virtual void ClearCachedCookedPlatformData( const ITargetPlatform *TargetPlatform ) OVERRIDE;
-	ENGINE_API virtual void ClearAllCachedCookedPlatformData() OVERRIDE;
+	ENGINE_API virtual void Serialize(FArchive& Ar) override;
+	ENGINE_API virtual void PostDuplicate(bool bDuplicateForPIE) override;
+	ENGINE_API virtual void PostLoad() override;
+	ENGINE_API virtual void BeginCacheForCookedPlatformData( const ITargetPlatform *TargetPlatform ) override;
+	ENGINE_API virtual void ClearCachedCookedPlatformData( const ITargetPlatform *TargetPlatform ) override;
+	ENGINE_API virtual void ClearAllCachedCookedPlatformData() override;
 
 #if WITH_EDITOR
-	ENGINE_API virtual void PreEditChange(UProperty* PropertyAboutToChange) OVERRIDE;
-	ENGINE_API virtual void PostEditChangeProperty(FPropertyChangedEvent& PropertyChangedEvent) OVERRIDE;
-	ENGINE_API virtual bool CanEditChange(const UProperty* InProperty) const OVERRIDE;
+	ENGINE_API virtual void PreEditChange(UProperty* PropertyAboutToChange) override;
+	ENGINE_API virtual void PostEditChangeProperty(FPropertyChangedEvent& PropertyChangedEvent) override;
+	ENGINE_API virtual bool CanEditChange(const UProperty* InProperty) const override;
 #endif // WITH_EDITOR
-	ENGINE_API virtual void BeginDestroy() OVERRIDE;
-	ENGINE_API virtual bool IsReadyForFinishDestroy() OVERRIDE;
-	ENGINE_API virtual void FinishDestroy() OVERRIDE;
-	ENGINE_API virtual SIZE_T GetResourceSize(EResourceSizeMode::Type Mode) OVERRIDE;
+	ENGINE_API virtual void BeginDestroy() override;
+	ENGINE_API virtual bool IsReadyForFinishDestroy() override;
+	ENGINE_API virtual void FinishDestroy() override;
+	ENGINE_API virtual SIZE_T GetResourceSize(EResourceSizeMode::Type Mode) override;
 	ENGINE_API static void AddReferencedObjects(UObject* InThis, FReferenceCollector& Collector);
 	// End UObject Interface
 
diff --git a/Engine/Source/Runtime/Engine/Classes/Materials/MaterialExpression.h b/Engine/Source/Runtime/Engine/Classes/Materials/MaterialExpression.h
--- a/Engine/Source/Runtime/Engine/Classes/Materials/MaterialExpression.h
+++ b/Engine/Source/Runtime/Engine/Classes/Materials/MaterialExpression.h
@@ -163,16 +163,16 @@ class UMaterialExpression : public UObject
 	TArray<FExpressionOutput> Outputs;
 
 	// Begin UObject interface.
-	virtual void PostInitProperties() OVERRIDE;
-	virtual void PostLoad() OVERRIDE;
-	virtual void PostDuplicate(bool bDuplicateForPIE) OVERRIDE;
+	virtual void PostInitProperties() override;
+	virtual void PostLoad() override;
+	virtual void PostDuplicate(bool bDuplicateForPIE) override;
 #if WITH_EDITOR
-	virtual void PostEditChangeProperty(FPropertyChangedEvent& PropertyChangedEvent) OVERRIDE;
-	virtual void PostEditImport() OVERRIDE;
-	virtual bool CanEditChange( const UProperty* InProperty ) const OVERRIDE;
+	virtual void PostEditChangeProperty(FPropertyChangedEvent& PropertyChangedEvent) override;
+	virtual void PostEditImport() override;
+	virtual bool CanEditChange( const UProperty* InProperty ) const override;
 #endif // WITH_EDITOR
 
-	virtual bool Modify( bool bAlwaysMarkDirty=true ) OVERRIDE;
+	virtual bool Modify( bool bAlwaysMarkDirty=true ) override;
 	virtual void Serialize( FArchive& Ar );
 	// End UObject interface.
 
diff --git a/Engine/Source/Runtime/Engine/Classes/Materials/MaterialExpressionAbs.h b/Engine/Source/Runtime/Engine/Classes/Materials/MaterialExpressionAbs.h
--- a/Engine/Source/Runtime/Engine/Classes/Materials/MaterialExpressionAbs.h
+++ b/Engine/Source/Runtime/Engine/Classes/Materials/MaterialExpressionAbs.h
@@ -20,8 +20,8 @@ class UMaterialExpressionAbs : public UMaterialExpression
 
 
 	// Begin UMaterialExpression Interface
-	virtual int32 Compile(class FMaterialCompiler* Compiler, int32 OutputIndex, int32 MultiplexIndex) OVERRIDE;
-	virtual void GetCaption(TArray<FString>& OutCaptions) const OVERRIDE;
+	virtual int32 Compile(class FMaterialCompiler* Compiler, int32 OutputIndex, int32 MultiplexIndex) override;
+	virtual void GetCaption(TArray<FString>& OutCaptions) const override;
 	// End UMaterialExpression Interface
 
 };
diff --git a/Engine/Source/Runtime/Engine/Classes/Materials/MaterialExpressionActorPositionWS.h b/Engine/Source/Runtime/Engine/Classes/Materials/MaterialExpressionActorPositionWS.h
--- a/Engine/Source/Runtime/Engine/Classes/Materials/MaterialExpressionActorPositionWS.h
+++ b/Engine/Source/Runtime/Engine/Classes/Materials/MaterialExpressionActorPositionWS.h
@@ -12,8 +12,8 @@ class UMaterialExpressionActorPositionWS : public UMaterialExpression
 
 
 	// Begin UMaterialExpression Interface
-	virtual int32 Compile(class FMaterialCompiler* Compiler, int32 OutputIndex, int32 MultiplexIndex) OVERRIDE;
-	virtual void GetCaption(TArray<FString>& OutCaptions) const OVERRIDE;
+	virtual int32 Compile(class FMaterialCompiler* Compiler, int32 OutputIndex, int32 MultiplexIndex) override;
+	virtual void GetCaption(TArray<FString>& OutCaptions) const override;
 	// End UMaterialExpression Interface
 };
 
diff --git a/Engine/Source/Runtime/Engine/Classes/Materials/MaterialExpressionAntialiasedTextureMask.h b/Engine/Source/Runtime/Engine/Classes/Materials/MaterialExpressionAntialiasedTextureMask.h
--- a/Engine/Source/Runtime/Engine/Classes/Materials/MaterialExpressionAntialiasedTextureMask.h
+++ b/Engine/Source/Runtime/Engine/Classes/Materials/MaterialExpressionAntialiasedTextureMask.h
@@ -28,14 +28,14 @@ class UMaterialExpressionAntialiasedTextureMask : public UMaterialExpressionText
 
 
 	// Begin UMaterialExpressionTextureSampleParameter Interface
-	virtual bool TextureIsValid( UTexture* InTexture ) OVERRIDE;
-	virtual const TCHAR* GetRequirements() OVERRIDE;
-	virtual void SetDefaultTexture() OVERRIDE;
+	virtual bool TextureIsValid( UTexture* InTexture ) override;
+	virtual const TCHAR* GetRequirements() override;
+	virtual void SetDefaultTexture() override;
 	// End UMaterialExpressionTextureSampleParameter Interface
 
 	// Begin UMaterialExpression Interface
-	virtual int32 Compile(class FMaterialCompiler* Compiler, int32 OutputIndex, int32 MultiplexIndex) OVERRIDE;
-	virtual void GetCaption(TArray<FString>& OutCaptions) const OVERRIDE;
+	virtual int32 Compile(class FMaterialCompiler* Compiler, int32 OutputIndex, int32 MultiplexIndex) override;
+	virtual void GetCaption(TArray<FString>& OutCaptions) const override;
 	// End UMaterialExpression Interface
 };
 
diff --git a/Engine/Source/Runtime/Engine/Classes/Materials/MaterialExpressionBreakMaterialAttributes.h b/Engine/Source/Runtime/Engine/Classes/Materials/MaterialExpressionBreakMaterialAttributes.h
--- a/Engine/Source/Runtime/Engine/Classes/Materials/MaterialExpressionBreakMaterialAttributes.h
+++ b/Engine/Source/Runtime/Engine/Classes/Materials/MaterialExpressionBreakMaterialAttributes.h
@@ -17,18 +17,18 @@ class UMaterialExpressionBreakMaterialAttributes : public UMaterialExpression
  	FMaterialAttributesInput MaterialAttributes;
  
 	// Begin UObject interface.
- 	virtual void Serialize( FArchive& Ar ) OVERRIDE;
+ 	virtual void Serialize( FArchive& Ar ) override;
 	// End UObject interface.
 
 	// Begin UMaterialExpression Interface
-	virtual int32 Compile(class FMaterialCompiler* Compiler, int32 OutputIndex, int32 MultiplexIndex) OVERRIDE;
-	virtual void GetCaption(TArray<FString>& OutCaptions) const OVERRIDE;
-	virtual const TArray<FExpressionInput*> GetInputs()OVERRIDE;
-	virtual FExpressionInput* GetInput(int32 InputIndex)OVERRIDE;
-	virtual FString GetInputName(int32 InputIndex) const OVERRIDE;
-	virtual bool IsInputConnectionRequired(int32 InputIndex) const OVERRIDE;
+	virtual int32 Compile(class FMaterialCompiler* Compiler, int32 OutputIndex, int32 MultiplexIndex) override;
+	virtual void GetCaption(TArray<FString>& OutCaptions) const override;
+	virtual const TArray<FExpressionInput*> GetInputs()override;
+	virtual FExpressionInput* GetInput(int32 InputIndex)override;
+	virtual FString GetInputName(int32 InputIndex) const override;
+	virtual bool IsInputConnectionRequired(int32 InputIndex) const override;
 #if WITH_EDITOR
-	virtual uint32 GetInputType(int32 InputIndex) OVERRIDE {return MCT_MaterialAttributes;}
+	virtual uint32 GetInputType(int32 InputIndex) override {return MCT_MaterialAttributes;}
 #endif // WITH_EDITOR
 	// End UMaterialExpression Interface
 };
diff --git a/Engine/Source/Runtime/Engine/Classes/Materials/MaterialExpressionClamp.h b/Engine/Source/Runtime/Engine/Classes/Materials/MaterialExpressionClamp.h
--- a/Engine/Source/Runtime/Engine/Classes/Materials/MaterialExpressionClamp.h
+++ b/Engine/Source/Runtime/Engine/Classes/Materials/MaterialExpressionClamp.h
@@ -39,12 +39,12 @@ class UMaterialExpressionClamp : public UMaterialExpression
 
 
 	// Begin UObject Interface
-	virtual void Serialize( FArchive& Ar ) OVERRIDE;
+	virtual void Serialize( FArchive& Ar ) override;
 	// End UObject Interface
 
 	// Begin UMaterialExpression Interface
-	virtual int32 Compile(class FMaterialCompiler* Compiler, int32 OutputIndex, int32 MultiplexIndex) OVERRIDE;
-	virtual void GetCaption(TArray<FString>& OutCaptions) const OVERRIDE;
+	virtual int32 Compile(class FMaterialCompiler* Compiler, int32 OutputIndex, int32 MultiplexIndex) override;
+	virtual void GetCaption(TArray<FString>& OutCaptions) const override;
 	// End UMaterialExpression Interface
 };
 
diff --git a/Engine/Source/Runtime/Engine/Classes/Materials/MaterialExpressionCollectionParameter.h b/Engine/Source/Runtime/Engine/Classes/Materials/MaterialExpressionCollectionParameter.h
--- a/Engine/Source/Runtime/Engine/Classes/Materials/MaterialExpressionCollectionParameter.h
+++ b/Engine/Source/Runtime/Engine/Classes/Materials/MaterialExpressionCollectionParameter.h
@@ -27,15 +27,15 @@ class UMaterialExpressionCollectionParameter : public UMaterialExpression
 
 	// Begin UObject interface.
 #if WITH_EDITOR
-	virtual void PostEditChangeProperty(FPropertyChangedEvent& PropertyChangedEvent) OVERRIDE;
+	virtual void PostEditChangeProperty(FPropertyChangedEvent& PropertyChangedEvent) override;
 #endif // WITH_EDITOR
-	virtual void PostLoad() OVERRIDE;
+	virtual void PostLoad() override;
 	// End UObject interface.
 
 	// Begin UMaterialExpression Interface
-	virtual int32 Compile(class FMaterialCompiler* Compiler, int32 OutputIndex, int32 MultiplexIndex) OVERRIDE;
-	virtual void GetCaption(TArray<FString>& OutCaptions) const OVERRIDE;
-	virtual bool MatchesSearchQuery( const TCHAR* SearchQuery ) OVERRIDE;
+	virtual int32 Compile(class FMaterialCompiler* Compiler, int32 OutputIndex, int32 MultiplexIndex) override;
+	virtual void GetCaption(TArray<FString>& OutCaptions) const override;
+	virtual bool MatchesSearchQuery( const TCHAR* SearchQuery ) override;
 	// End UMaterialExpression Interface
 };
 
diff --git a/Engine/Source/Runtime/Engine/Classes/Materials/MaterialExpressionComment.h b/Engine/Source/Runtime/Engine/Classes/Materials/MaterialExpressionComment.h
--- a/Engine/Source/Runtime/Engine/Classes/Materials/MaterialExpressionComment.h
+++ b/Engine/Source/Runtime/Engine/Classes/Materials/MaterialExpressionComment.h
@@ -25,14 +25,14 @@ class UMaterialExpressionComment : public UMaterialExpression
 
 	// Begin UObject Interface
 #if WITH_EDITOR
-	virtual void PostEditChangeProperty(FPropertyChangedEvent& PropertyChangedEvent) OVERRIDE;
+	virtual void PostEditChangeProperty(FPropertyChangedEvent& PropertyChangedEvent) override;
 #endif // WITH_EDITOR
-	virtual bool Modify( bool bAlwaysMarkDirty=true ) OVERRIDE;
+	virtual bool Modify( bool bAlwaysMarkDirty=true ) override;
 	// End UObject Interface
 
 	// Begin UMaterialExpression Interface
-	virtual void GetCaption(TArray<FString>& OutCaptions) const OVERRIDE;
-	virtual bool MatchesSearchQuery( const TCHAR* SearchQuery ) OVERRIDE;
+	virtual void GetCaption(TArray<FString>& OutCaptions) const override;
+	virtual bool MatchesSearchQuery( const TCHAR* SearchQuery ) override;
 	// End UMaterialExpression Interface
 };
 
diff --git a/Engine/Source/Runtime/Engine/Classes/Materials/MaterialExpressionConstant.h b/Engine/Source/Runtime/Engine/Classes/Materials/MaterialExpressionConstant.h
--- a/Engine/Source/Runtime/Engine/Classes/Materials/MaterialExpressionConstant.h
+++ b/Engine/Source/Runtime/Engine/Classes/Materials/MaterialExpressionConstant.h
@@ -15,10 +15,10 @@ class UMaterialExpressionConstant : public UMaterialExpression
 
 
 	// Begin UMaterialExpression Interface
-	virtual int32 Compile(class FMaterialCompiler* Compiler, int32 OutputIndex, int32 MultiplexIndex) OVERRIDE;
-	virtual void GetCaption(TArray<FString>& OutCaptions) const OVERRIDE;
+	virtual int32 Compile(class FMaterialCompiler* Compiler, int32 OutputIndex, int32 MultiplexIndex) override;
+	virtual void GetCaption(TArray<FString>& OutCaptions) const override;
 #if WITH_EDITOR
-	virtual FString GetDescription() const OVERRIDE;
+	virtual FString GetDescription() const override;
 #endif // WITH_EDITOR
 
 	// End UMaterialExpression Interface
