UE4 Refactoring. Changed OVERRIDE and FINAL macros to keywords override and final.

diff --git a/Engine/Source/Runtime/Engine/Classes/Materials/MaterialExpressionConstant2Vector.h b/Engine/Source/Runtime/Engine/Classes/Materials/MaterialExpressionConstant2Vector.h
--- a/Engine/Source/Runtime/Engine/Classes/Materials/MaterialExpressionConstant2Vector.h
+++ b/Engine/Source/Runtime/Engine/Classes/Materials/MaterialExpressionConstant2Vector.h
@@ -18,11 +18,11 @@ class UMaterialExpressionConstant2Vector : public UMaterialExpression
 
 
 	// Begin UMaterialExpression Interface
-	virtual int32 Compile(class FMaterialCompiler* Compiler, int32 OutputIndex, int32 MultiplexIndex) OVERRIDE;
-	virtual void GetCaption(TArray<FString>& OutCaptions) const OVERRIDE;
+	virtual int32 Compile(class FMaterialCompiler* Compiler, int32 OutputIndex, int32 MultiplexIndex) override;
+	virtual void GetCaption(TArray<FString>& OutCaptions) const override;
 #if WITH_EDITOR
-	virtual FString GetDescription() const OVERRIDE;
-	virtual uint32 GetOutputType(int32 OutputIndex) OVERRIDE {return MCT_Float2;}
+	virtual FString GetDescription() const override;
+	virtual uint32 GetOutputType(int32 OutputIndex) override {return MCT_Float2;}
 #endif // WITH_EDITOR
 	// End UMaterialExpression Interface
 };
diff --git a/Engine/Source/Runtime/Engine/Classes/Materials/MaterialExpressionConstant3Vector.h b/Engine/Source/Runtime/Engine/Classes/Materials/MaterialExpressionConstant3Vector.h
--- a/Engine/Source/Runtime/Engine/Classes/Materials/MaterialExpressionConstant3Vector.h
+++ b/Engine/Source/Runtime/Engine/Classes/Materials/MaterialExpressionConstant3Vector.h
@@ -23,15 +23,15 @@ class UMaterialExpressionConstant3Vector : public UMaterialExpression
 	FLinearColor Constant;
 
 	// Begin UObject interface.
-	virtual void PostLoad() OVERRIDE;
+	virtual void PostLoad() override;
 	// End UObject interface.
 
 	// Begin UMaterialExpression Interface
-	virtual int32 Compile(class FMaterialCompiler* Compiler, int32 OutputIndex, int32 MultiplexIndex) OVERRIDE;
-	virtual void GetCaption(TArray<FString>& OutCaptions) const OVERRIDE;
+	virtual int32 Compile(class FMaterialCompiler* Compiler, int32 OutputIndex, int32 MultiplexIndex) override;
+	virtual void GetCaption(TArray<FString>& OutCaptions) const override;
 #if WITH_EDITOR
-	virtual FString GetDescription() const OVERRIDE;
-	virtual uint32 GetOutputType(int32 OutputIndex) OVERRIDE {return MCT_Float3;}
+	virtual FString GetDescription() const override;
+	virtual uint32 GetOutputType(int32 OutputIndex) override {return MCT_Float3;}
 #endif // WITH_EDITOR
 	// End UMaterialExpression Interface
 };
diff --git a/Engine/Source/Runtime/Engine/Classes/Materials/MaterialExpressionConstant4Vector.h b/Engine/Source/Runtime/Engine/Classes/Materials/MaterialExpressionConstant4Vector.h
--- a/Engine/Source/Runtime/Engine/Classes/Materials/MaterialExpressionConstant4Vector.h
+++ b/Engine/Source/Runtime/Engine/Classes/Materials/MaterialExpressionConstant4Vector.h
@@ -26,15 +26,15 @@ class UMaterialExpressionConstant4Vector : public UMaterialExpression
 	FLinearColor Constant;
 
 	// Begin UObject interface.
-	virtual void PostLoad() OVERRIDE;
+	virtual void PostLoad() override;
 	// End UObject interface.
 
 	// Begin UMaterialExpression Interface
-	virtual int32 Compile(class FMaterialCompiler* Compiler, int32 OutputIndex, int32 MultiplexIndex) OVERRIDE;
-	virtual void GetCaption(TArray<FString>& OutCaptions) const OVERRIDE;
+	virtual int32 Compile(class FMaterialCompiler* Compiler, int32 OutputIndex, int32 MultiplexIndex) override;
+	virtual void GetCaption(TArray<FString>& OutCaptions) const override;
 #if WITH_EDITOR
-	virtual FString GetDescription() const OVERRIDE;
-	virtual uint32 GetOutputType(int32 OutputIndex) OVERRIDE {return MCT_Float4;}
+	virtual FString GetDescription() const override;
+	virtual uint32 GetOutputType(int32 OutputIndex) override {return MCT_Float4;}
 #endif // WITH_EDITOR
 	// End UMaterialExpression Interface
 };
diff --git a/Engine/Source/Runtime/Engine/Classes/Materials/MaterialExpressionCustom.h b/Engine/Source/Runtime/Engine/Classes/Materials/MaterialExpressionCustom.h
--- a/Engine/Source/Runtime/Engine/Classes/Materials/MaterialExpressionCustom.h
+++ b/Engine/Source/Runtime/Engine/Classes/Materials/MaterialExpressionCustom.h
@@ -48,19 +48,19 @@ class UMaterialExpressionCustom : public UMaterialExpression
 
 	// Begin UObject interface.
 #if WITH_EDITOR
-	virtual void PostEditChangeProperty(FPropertyChangedEvent& PropertyChangedEvent) OVERRIDE;
+	virtual void PostEditChangeProperty(FPropertyChangedEvent& PropertyChangedEvent) override;
 #endif // WITH_EDITOR
 	// End UObject interface.
 	
 	// Begin UMaterialExpression Interface
-	virtual int32 Compile(class FMaterialCompiler* Compiler, int32 OutputIndex, int32 MultiplexIndex) OVERRIDE;
-	virtual void GetCaption(TArray<FString>& OutCaptions) const OVERRIDE;
-	virtual const TArray<FExpressionInput*> GetInputs() OVERRIDE;
-	virtual FExpressionInput* GetInput(int32 InputIndex) OVERRIDE;
-	virtual FString GetInputName(int32 InputIndex) const OVERRIDE;
+	virtual int32 Compile(class FMaterialCompiler* Compiler, int32 OutputIndex, int32 MultiplexIndex) override;
+	virtual void GetCaption(TArray<FString>& OutCaptions) const override;
+	virtual const TArray<FExpressionInput*> GetInputs() override;
+	virtual FExpressionInput* GetInput(int32 InputIndex) override;
+	virtual FString GetInputName(int32 InputIndex) const override;
 #if WITH_EDITOR
-	virtual uint32 GetInputType(int32 InputIndex) OVERRIDE {return MCT_Unknown;}
-	virtual uint32 GetOutputType(int32 OutputIndex) OVERRIDE;
+	virtual uint32 GetInputType(int32 InputIndex) override {return MCT_Unknown;}
+	virtual uint32 GetOutputType(int32 OutputIndex) override;
 #endif // WITH_EDITOR
 	// End UMaterialExpression Interface
 };
diff --git a/Engine/Source/Runtime/Engine/Classes/Materials/MaterialExpressionCustomTexture.h b/Engine/Source/Runtime/Engine/Classes/Materials/MaterialExpressionCustomTexture.h
--- a/Engine/Source/Runtime/Engine/Classes/Materials/MaterialExpressionCustomTexture.h
+++ b/Engine/Source/Runtime/Engine/Classes/Materials/MaterialExpressionCustomTexture.h
@@ -15,25 +15,25 @@ class UMaterialExpressionCustomTexture : public UMaterialExpression
 
 
 	// Begin UMaterialExpression Interface
-	virtual int32 Compile(class FMaterialCompiler* Compiler, int32 OutputIndex, int32 MultiplexIndex) OVERRIDE;
-	virtual void GetCaption(TArray<FString>& OutCaptions) const OVERRIDE;
-	virtual int32 CompilePreview(class FMaterialCompiler* Compiler, int32 OutputIndex, int32 MultiplexIndex) OVERRIDE;
-	virtual int32 GetWidth() const OVERRIDE;
-	virtual int32 GetLabelPadding() OVERRIDE { return 8; }
-	virtual bool MatchesSearchQuery( const TCHAR* SearchQuery ) OVERRIDE;
+	virtual int32 Compile(class FMaterialCompiler* Compiler, int32 OutputIndex, int32 MultiplexIndex) override;
+	virtual void GetCaption(TArray<FString>& OutCaptions) const override;
+	virtual int32 CompilePreview(class FMaterialCompiler* Compiler, int32 OutputIndex, int32 MultiplexIndex) override;
+	virtual int32 GetWidth() const override;
+	virtual int32 GetLabelPadding() override { return 8; }
+	virtual bool MatchesSearchQuery( const TCHAR* SearchQuery ) override;
 
 	/** 
 	 * Callback to get any texture reference this expression emits.
 	 * This is used to link the compiled uniform expressions with their default texture values. 
 	 * Any UMaterialExpression whose compilation creates a texture uniform expression (eg Compiler->Texture, Compiler->TextureParameter) must implement this.
 	 */
-	virtual UTexture* GetReferencedTexture() OVERRIDE
+	virtual UTexture* GetReferencedTexture() override
 	{
 		return Texture;
 	}
 
 #if WITH_EDITOR
-	virtual uint32 GetOutputType(int32 OutputIndex) OVERRIDE;
+	virtual uint32 GetOutputType(int32 OutputIndex) override;
 #endif // WITH_EDITOR
 	// End UMaterialExpression Interface
 };
diff --git a/Engine/Source/Runtime/Engine/Classes/Materials/MaterialExpressionDepthFade.h b/Engine/Source/Runtime/Engine/Classes/Materials/MaterialExpressionDepthFade.h
--- a/Engine/Source/Runtime/Engine/Classes/Materials/MaterialExpressionDepthFade.h
+++ b/Engine/Source/Runtime/Engine/Classes/Materials/MaterialExpressionDepthFade.h
@@ -31,7 +31,7 @@ class UMaterialExpressionDepthFade : public UMaterialExpression
 
 
 	// Begin UMaterialExpression Interface
-	virtual FString GetInputName(int32 InputIndex) const OVERRIDE
+	virtual FString GetInputName(int32 InputIndex) const override
 	{
 		if (InputIndex == 0)
 		{
@@ -42,8 +42,8 @@ class UMaterialExpressionDepthFade : public UMaterialExpression
 			return Super::GetInputName(InputIndex);
 		}
 	}
-	virtual int32 Compile(class FMaterialCompiler* Compiler, int32 OutputIndex, int32 MultiplexIndex) OVERRIDE;
-	virtual void GetCaption(TArray<FString>& OutCaptions) const OVERRIDE
+	virtual int32 Compile(class FMaterialCompiler* Compiler, int32 OutputIndex, int32 MultiplexIndex) override;
+	virtual void GetCaption(TArray<FString>& OutCaptions) const override
 	{
 		OutCaptions.Add(TEXT("Depth Fade"));
 	}
diff --git a/Engine/Source/Runtime/Engine/Classes/Materials/MaterialExpressionDesaturation.h b/Engine/Source/Runtime/Engine/Classes/Materials/MaterialExpressionDesaturation.h
--- a/Engine/Source/Runtime/Engine/Classes/Materials/MaterialExpressionDesaturation.h
+++ b/Engine/Source/Runtime/Engine/Classes/Materials/MaterialExpressionDesaturation.h
@@ -22,8 +22,8 @@ class UMaterialExpressionDesaturation : public UMaterialExpression
 
 
 	// Begin UMaterialExpression Interface
-	virtual int32 Compile(class FMaterialCompiler* Compiler, int32 OutputIndex, int32 MultiplexIndex) OVERRIDE;
-	virtual void GetCaption(TArray<FString>& OutCaptions) const OVERRIDE
+	virtual int32 Compile(class FMaterialCompiler* Compiler, int32 OutputIndex, int32 MultiplexIndex) override;
+	virtual void GetCaption(TArray<FString>& OutCaptions) const override
 	{
 		OutCaptions.Add(TEXT("Desaturation"));
 	}
diff --git a/Engine/Source/Runtime/Engine/Classes/Materials/MaterialExpressionDynamicParameter.h b/Engine/Source/Runtime/Engine/Classes/Materials/MaterialExpressionDynamicParameter.h
--- a/Engine/Source/Runtime/Engine/Classes/Materials/MaterialExpressionDynamicParameter.h
+++ b/Engine/Source/Runtime/Engine/Classes/Materials/MaterialExpressionDynamicParameter.h
@@ -25,17 +25,17 @@ class UMaterialExpressionDynamicParameter : public UMaterialExpression
 
 	// Begin UObject Interface
 #if WITH_EDITOR
-	virtual void PostEditChangeProperty(FPropertyChangedEvent& PropertyChangedEvent) OVERRIDE;
+	virtual void PostEditChangeProperty(FPropertyChangedEvent& PropertyChangedEvent) override;
 #endif // WITH_EDITOR
 	// End UObject Interface
 
 	// Begin UMaterialExpression Interface
-	virtual int32 Compile(class FMaterialCompiler* Compiler, int32 OutputIndex, int32 MultiplexIndex) OVERRIDE;
-	virtual void GetCaption(TArray<FString>& OutCaptions) const OVERRIDE;
-	virtual TArray<FExpressionOutput>& GetOutputs() OVERRIDE;
-	virtual int32 GetWidth() const OVERRIDE;
-	virtual int32 GetLabelPadding() OVERRIDE { return 8; }
-	virtual bool MatchesSearchQuery( const TCHAR* SearchQuery ) OVERRIDE;
+	virtual int32 Compile(class FMaterialCompiler* Compiler, int32 OutputIndex, int32 MultiplexIndex) override;
+	virtual void GetCaption(TArray<FString>& OutCaptions) const override;
+	virtual TArray<FExpressionOutput>& GetOutputs() override;
+	virtual int32 GetWidth() const override;
+	virtual int32 GetLabelPadding() override { return 8; }
+	virtual bool MatchesSearchQuery( const TCHAR* SearchQuery ) override;
 	// End UMaterialExpression Interface
 };
 
diff --git a/Engine/Source/Runtime/Engine/Classes/Materials/MaterialExpressionFeatureLevelSwitch.h b/Engine/Source/Runtime/Engine/Classes/Materials/MaterialExpressionFeatureLevelSwitch.h
--- a/Engine/Source/Runtime/Engine/Classes/Materials/MaterialExpressionFeatureLevelSwitch.h
+++ b/Engine/Source/Runtime/Engine/Classes/Materials/MaterialExpressionFeatureLevelSwitch.h
@@ -18,17 +18,17 @@ class UMaterialExpressionFeatureLevelSwitch : public UMaterialExpression
 	FExpressionInput Inputs[ERHIFeatureLevel::Num];
 
 	// Begin UMaterialExpression Interface
-	virtual int32 Compile(class FMaterialCompiler* Compiler, int32 OutputIndex, int32 MultiplexIndex) OVERRIDE;
-	virtual void GetCaption(TArray<FString>& OutCaptions) const OVERRIDE;
+	virtual int32 Compile(class FMaterialCompiler* Compiler, int32 OutputIndex, int32 MultiplexIndex) override;
+	virtual void GetCaption(TArray<FString>& OutCaptions) const override;
 	
-	virtual const TArray<FExpressionInput*> GetInputs() OVERRIDE;
-	virtual FExpressionInput* GetInput(int32 InputIndex) OVERRIDE;
-	virtual FString GetInputName(int32 InputIndex) const OVERRIDE;
-	virtual bool IsInputConnectionRequired(int32 InputIndex) const OVERRIDE;
-	virtual bool IsResultMaterialAttributes(int32 OutputIndex) OVERRIDE;
+	virtual const TArray<FExpressionInput*> GetInputs() override;
+	virtual FExpressionInput* GetInput(int32 InputIndex) override;
+	virtual FString GetInputName(int32 InputIndex) const override;
+	virtual bool IsInputConnectionRequired(int32 InputIndex) const override;
+	virtual bool IsResultMaterialAttributes(int32 OutputIndex) override;
 #if WITH_EDITOR
-	virtual uint32 GetInputType(int32 InputIndex) OVERRIDE {return MCT_Unknown;}
-	virtual uint32 GetOutputType(int32 OutputIndex) OVERRIDE {return MCT_Unknown;}
+	virtual uint32 GetInputType(int32 InputIndex) override {return MCT_Unknown;}
+	virtual uint32 GetOutputType(int32 OutputIndex) override {return MCT_Unknown;}
 #endif
 	// End UMaterialExpression Interface
 };
diff --git a/Engine/Source/Runtime/Engine/Classes/Materials/MaterialExpressionFontSample.h b/Engine/Source/Runtime/Engine/Classes/Materials/MaterialExpressionFontSample.h
--- a/Engine/Source/Runtime/Engine/Classes/Materials/MaterialExpressionFontSample.h
+++ b/Engine/Source/Runtime/Engine/Classes/Materials/MaterialExpressionFontSample.h
@@ -20,18 +20,18 @@ class UMaterialExpressionFontSample : public UMaterialExpression
 
 
 	// Begin UMaterialExpression Interface
-	virtual int32 Compile(class FMaterialCompiler* Compiler, int32 OutputIndex, int32 MultiplexIndex) OVERRIDE;
-	virtual void GetCaption(TArray<FString>& OutCaptions) const OVERRIDE;
-	virtual int32 GetWidth() const OVERRIDE;
-	virtual int32 GetLabelPadding() OVERRIDE { return 8; }
-	virtual bool MatchesSearchQuery( const TCHAR* SearchQuery ) OVERRIDE;
+	virtual int32 Compile(class FMaterialCompiler* Compiler, int32 OutputIndex, int32 MultiplexIndex) override;
+	virtual void GetCaption(TArray<FString>& OutCaptions) const override;
+	virtual int32 GetWidth() const override;
+	virtual int32 GetLabelPadding() override { return 8; }
+	virtual bool MatchesSearchQuery( const TCHAR* SearchQuery ) override;
 
 	/** 
 	 * Callback to get any texture reference this expression emits.
 	 * This is used to link the compiled uniform expressions with their default texture values. 
 	 * Any UMaterialExpression whose compilation creates a texture uniform expression (eg Compiler->Texture, Compiler->TextureParameter) must implement this.
 	 */
-	virtual UTexture* GetReferencedTexture() OVERRIDE;
+	virtual UTexture* GetReferencedTexture() override;
 	// End UMaterialExpression Interface
 };
 
diff --git a/Engine/Source/Runtime/Engine/Classes/Materials/MaterialExpressionFontSampleParameter.h b/Engine/Source/Runtime/Engine/Classes/Materials/MaterialExpressionFontSampleParameter.h
--- a/Engine/Source/Runtime/Engine/Classes/Materials/MaterialExpressionFontSampleParameter.h
+++ b/Engine/Source/Runtime/Engine/Classes/Materials/MaterialExpressionFontSampleParameter.h
@@ -23,9 +23,9 @@ class UMaterialExpressionFontSampleParameter : public UMaterialExpressionFontSam
 	FName Group;
 
 	// Begin UMaterialExpression Interface
-	virtual int32 Compile(class FMaterialCompiler* Compiler, int32 OutputIndex, int32 MultiplexIndex) OVERRIDE;
-	virtual void GetCaption(TArray<FString>& OutCaptions) const OVERRIDE;
-	virtual bool MatchesSearchQuery( const TCHAR* SearchQuery ) OVERRIDE;
+	virtual int32 Compile(class FMaterialCompiler* Compiler, int32 OutputIndex, int32 MultiplexIndex) override;
+	virtual void GetCaption(TArray<FString>& OutCaptions) const override;
+	virtual bool MatchesSearchQuery( const TCHAR* SearchQuery ) override;
 	// End UMaterialExpression Interface
 	
 	/**
@@ -33,7 +33,7 @@ class UMaterialExpressionFontSampleParameter : public UMaterialExpressionFontSam
 	*/
 	virtual void SetDefaultFont();
 	
-	ENGINE_API virtual FGuid& GetParameterExpressionId() OVERRIDE
+	ENGINE_API virtual FGuid& GetParameterExpressionId() override
 	{
 		return ExpressionGUID;
 	}
diff --git a/Engine/Source/Runtime/Engine/Classes/Materials/MaterialExpressionFresnel.h b/Engine/Source/Runtime/Engine/Classes/Materials/MaterialExpressionFresnel.h
--- a/Engine/Source/Runtime/Engine/Classes/Materials/MaterialExpressionFresnel.h
+++ b/Engine/Source/Runtime/Engine/Classes/Materials/MaterialExpressionFresnel.h
@@ -39,8 +39,8 @@ class UMaterialExpressionFresnel : public UMaterialExpression
 
 
 	// Begin UMaterialExpression Interface
-	virtual int32 Compile(class FMaterialCompiler* Compiler, int32 OutputIndex, int32 MultiplexIndex) OVERRIDE;
-	virtual void GetCaption(TArray<FString>& OutCaptions) const OVERRIDE
+	virtual int32 Compile(class FMaterialCompiler* Compiler, int32 OutputIndex, int32 MultiplexIndex) override;
+	virtual void GetCaption(TArray<FString>& OutCaptions) const override
 	{
 		OutCaptions.Add(FString(TEXT("Fresnel")));
 	}
diff --git a/Engine/Source/Runtime/Engine/Classes/Materials/MaterialExpressionFunctionInput.h b/Engine/Source/Runtime/Engine/Classes/Materials/MaterialExpressionFunctionInput.h
--- a/Engine/Source/Runtime/Engine/Classes/Materials/MaterialExpressionFunctionInput.h
+++ b/Engine/Source/Runtime/Engine/Classes/Materials/MaterialExpressionFunctionInput.h
@@ -70,28 +70,28 @@ class UMaterialExpressionFunctionInput : public UMaterialExpression
 
 
 	// Begin UObject interface.
-	virtual void PostLoad() OVERRIDE;
-	virtual void PostDuplicate(bool bDuplicateForPIE) OVERRIDE;
+	virtual void PostLoad() override;
+	virtual void PostDuplicate(bool bDuplicateForPIE) override;
 #if WITH_EDITOR
-	virtual void PostEditImport() OVERRIDE;
-	virtual void PreEditChange(UProperty* PropertyAboutToChange) OVERRIDE;
-	virtual void PostEditChangeProperty(FPropertyChangedEvent& PropertyChangedEvent) OVERRIDE;
+	virtual void PostEditImport() override;
+	virtual void PreEditChange(UProperty* PropertyAboutToChange) override;
+	virtual void PostEditChangeProperty(FPropertyChangedEvent& PropertyChangedEvent) override;
 #endif // WITH_EDITOR
 	// End UObject interface.
 
 	// Begin UMaterialExpression Interface
 #if WITH_EDITOR
-	virtual void GetExpressionToolTip(TArray<FString>& OutToolTip) OVERRIDE;
+	virtual void GetExpressionToolTip(TArray<FString>& OutToolTip) override;
 #endif
 
-	virtual int32 CompilePreview(class FMaterialCompiler* Compiler, int32 OutputIndex, int32 MultiplexIndex) OVERRIDE;
-	virtual int32 Compile(class FMaterialCompiler* Compiler, int32 OutputIndex, int32 MultiplexIndex) OVERRIDE;
+	virtual int32 CompilePreview(class FMaterialCompiler* Compiler, int32 OutputIndex, int32 MultiplexIndex) override;
+	virtual int32 Compile(class FMaterialCompiler* Compiler, int32 OutputIndex, int32 MultiplexIndex) override;
 
-	virtual void GetCaption(TArray<FString>& OutCaptions) const OVERRIDE;
-	virtual bool IsResultMaterialAttributes(int32 OutputIndex) OVERRIDE;
+	virtual void GetCaption(TArray<FString>& OutCaptions) const override;
+	virtual bool IsResultMaterialAttributes(int32 OutputIndex) override;
 #if WITH_EDITOR
-	virtual uint32 GetInputType(int32 InputIndex) OVERRIDE;
-	virtual uint32 GetOutputType(int32 OutputIndex) OVERRIDE;
+	virtual uint32 GetInputType(int32 InputIndex) override;
+	virtual uint32 GetOutputType(int32 OutputIndex) override;
 #endif
 	// End UMaterialExpression Interface
 
diff --git a/Engine/Source/Runtime/Engine/Classes/Materials/MaterialExpressionFunctionOutput.h b/Engine/Source/Runtime/Engine/Classes/Materials/MaterialExpressionFunctionOutput.h
--- a/Engine/Source/Runtime/Engine/Classes/Materials/MaterialExpressionFunctionOutput.h
+++ b/Engine/Source/Runtime/Engine/Classes/Materials/MaterialExpressionFunctionOutput.h
@@ -39,24 +39,24 @@ class UMaterialExpressionFunctionOutput : public UMaterialExpression
 	virtual void PostLoad();
 	virtual void PostDuplicate(bool bDuplicateForPIE);
 #if WITH_EDITOR
-	virtual void PostEditImport() OVERRIDE;
-	virtual void PreEditChange(UProperty* PropertyAboutToChange) OVERRIDE;
-	virtual void PostEditChangeProperty(FPropertyChangedEvent& PropertyChangedEvent) OVERRIDE;
+	virtual void PostEditImport() override;
+	virtual void PreEditChange(UProperty* PropertyAboutToChange) override;
+	virtual void PostEditChangeProperty(FPropertyChangedEvent& PropertyChangedEvent) override;
 #endif // WITH_EDITOR
 	// End UObject interface.
 
 	// Begin UMaterialExpression Interface
-	virtual int32 Compile(class FMaterialCompiler* Compiler, int32 OutputIndex, int32 MultiplexIndex) OVERRIDE;
-	virtual void GetCaption(TArray<FString>& OutCaptions) const OVERRIDE;
-	virtual FString GetInputName(int32 InputIndex) const OVERRIDE
+	virtual int32 Compile(class FMaterialCompiler* Compiler, int32 OutputIndex, int32 MultiplexIndex) override;
+	virtual void GetCaption(TArray<FString>& OutCaptions) const override;
+	virtual FString GetInputName(int32 InputIndex) const override
 	{
 		return TEXT("");
 	}
 #if WITH_EDITOR
-	virtual void GetExpressionToolTip(TArray<FString>& OutToolTip) OVERRIDE;
-	virtual uint32 GetInputType(int32 InputIndex) OVERRIDE;
+	virtual void GetExpressionToolTip(TArray<FString>& OutToolTip) override;
+	virtual uint32 GetInputType(int32 InputIndex) override;
 #endif
-	virtual bool IsResultMaterialAttributes(int32 OutputIndex) OVERRIDE;
+	virtual bool IsResultMaterialAttributes(int32 OutputIndex) override;
 	// End UMaterialExpression Interface
 
 
diff --git a/Engine/Source/Runtime/Engine/Classes/Materials/MaterialExpressionIf.h b/Engine/Source/Runtime/Engine/Classes/Materials/MaterialExpressionIf.h
--- a/Engine/Source/Runtime/Engine/Classes/Materials/MaterialExpressionIf.h
+++ b/Engine/Source/Runtime/Engine/Classes/Materials/MaterialExpressionIf.h
@@ -37,11 +37,11 @@ class UMaterialExpressionIf : public UMaterialExpression
 	float ConstAEqualsB;
 
 	// Begin UMaterialExpression Interface
-	virtual int32 Compile(class FMaterialCompiler* Compiler, int32 OutputIndex, int32 MultiplexIndex) OVERRIDE;
-	virtual void GetCaption(TArray<FString>& OutCaptions) const OVERRIDE;
+	virtual int32 Compile(class FMaterialCompiler* Compiler, int32 OutputIndex, int32 MultiplexIndex) override;
+	virtual void GetCaption(TArray<FString>& OutCaptions) const override;
 #if WITH_EDITOR
-	virtual uint32 GetInputType(int32 InputIndex) OVERRIDE;
-	virtual uint32 GetOutputType(int32 InputIndex) OVERRIDE {return MCT_Unknown;}
+	virtual uint32 GetInputType(int32 InputIndex) override;
+	virtual uint32 GetOutputType(int32 InputIndex) override {return MCT_Unknown;}
 #endif
 	// End UMaterialExpression Interface
 };
diff --git a/Engine/Source/Runtime/Engine/Classes/Materials/MaterialExpressionLandscapeLayerBlend.h b/Engine/Source/Runtime/Engine/Classes/Materials/MaterialExpressionLandscapeLayerBlend.h
--- a/Engine/Source/Runtime/Engine/Classes/Materials/MaterialExpressionLandscapeLayerBlend.h
+++ b/Engine/Source/Runtime/Engine/Classes/Materials/MaterialExpressionLandscapeLayerBlend.h
@@ -65,21 +65,21 @@ class UMaterialExpressionLandscapeLayerBlend : public UMaterialExpression
 
 	// Begin UObject Interface
 #if WITH_EDITOR
-	virtual void PostEditChangeProperty(FPropertyChangedEvent& PropertyChangedEvent) OVERRIDE;
+	virtual void PostEditChangeProperty(FPropertyChangedEvent& PropertyChangedEvent) override;
 	// End UObject Interface
 
 	// Begin UMaterialExpression Interface
 #endif
-	virtual bool IsResultMaterialAttributes(int32 OutputIndex) OVERRIDE;
-	virtual int32 Compile(class FMaterialCompiler* Compiler, int32 OutputIndex, int32 MultiplexIndex) OVERRIDE;
-	virtual void GetCaption(TArray<FString>& OutCaptions) const OVERRIDE;
-	virtual const TArray<FExpressionInput*> GetInputs() OVERRIDE;
-	virtual FExpressionInput* GetInput(int32 InputIndex) OVERRIDE;
-	virtual FString GetInputName(int32 InputIndex) const OVERRIDE;
-	virtual UTexture* GetReferencedTexture() OVERRIDE;
+	virtual bool IsResultMaterialAttributes(int32 OutputIndex) override;
+	virtual int32 Compile(class FMaterialCompiler* Compiler, int32 OutputIndex, int32 MultiplexIndex) override;
+	virtual void GetCaption(TArray<FString>& OutCaptions) const override;
+	virtual const TArray<FExpressionInput*> GetInputs() override;
+	virtual FExpressionInput* GetInput(int32 InputIndex) override;
+	virtual FString GetInputName(int32 InputIndex) const override;
+	virtual UTexture* GetReferencedTexture() override;
 	// End UMaterialExpression Interface
 
-	ENGINE_API virtual FGuid& GetParameterExpressionId() OVERRIDE
+	ENGINE_API virtual FGuid& GetParameterExpressionId() override
 	{
 		return ExpressionGUID;
 	}
diff --git a/Engine/Source/Runtime/Engine/Classes/Materials/MaterialExpressionLandscapeLayerSwitch.h b/Engine/Source/Runtime/Engine/Classes/Materials/MaterialExpressionLandscapeLayerSwitch.h
--- a/Engine/Source/Runtime/Engine/Classes/Materials/MaterialExpressionLandscapeLayerSwitch.h
+++ b/Engine/Source/Runtime/Engine/Classes/Materials/MaterialExpressionLandscapeLayerSwitch.h
@@ -30,21 +30,21 @@ class UMaterialExpressionLandscapeLayerSwitch : public UMaterialExpression
 public:
 
 	// Begin UObject Interface
-	virtual void Serialize(FArchive& Ar) OVERRIDE;
+	virtual void Serialize(FArchive& Ar) override;
 	// End UObject Interface
 
 	// Begin UMaterialExpression Interface
-	virtual bool IsResultMaterialAttributes(int32 OutputIndex) OVERRIDE;
-	virtual int32 Compile(class FMaterialCompiler* Compiler, int32 OutputIndex, int32 MultiplexIndex) OVERRIDE;
-	virtual void GetCaption(TArray<FString>& OutCaptions) const OVERRIDE;
-	virtual UTexture* GetReferencedTexture() OVERRIDE;
+	virtual bool IsResultMaterialAttributes(int32 OutputIndex) override;
+	virtual int32 Compile(class FMaterialCompiler* Compiler, int32 OutputIndex, int32 MultiplexIndex) override;
+	virtual void GetCaption(TArray<FString>& OutCaptions) const override;
+	virtual UTexture* GetReferencedTexture() override;
 #if WITH_EDITOR
-	virtual uint32 GetInputType(int32 InputIndex) OVERRIDE {return MCT_Unknown;}
-	virtual uint32 GetOutputType(int32 InputIndex) OVERRIDE {return MCT_Unknown;}
+	virtual uint32 GetInputType(int32 InputIndex) override {return MCT_Unknown;}
+	virtual uint32 GetOutputType(int32 InputIndex) override {return MCT_Unknown;}
 #endif
 	// End UMaterialExpression Interface
 
-	ENGINE_API virtual FGuid& GetParameterExpressionId() OVERRIDE
+	ENGINE_API virtual FGuid& GetParameterExpressionId() override
 	{
 		return ExpressionGUID;
 	}
diff --git a/Engine/Source/Runtime/Engine/Classes/Materials/MaterialExpressionLandscapeLayerWeight.h b/Engine/Source/Runtime/Engine/Classes/Materials/MaterialExpressionLandscapeLayerWeight.h
--- a/Engine/Source/Runtime/Engine/Classes/Materials/MaterialExpressionLandscapeLayerWeight.h
+++ b/Engine/Source/Runtime/Engine/Classes/Materials/MaterialExpressionLandscapeLayerWeight.h
@@ -33,16 +33,16 @@ class UMaterialExpressionLandscapeLayerWeight : public UMaterialExpression
 public:
 
 	// Begin UMaterialExpression Interface
-	virtual bool IsResultMaterialAttributes(int32 OutputIndex) OVERRIDE;
-	virtual int32 Compile(class FMaterialCompiler* Compiler, int32 OutputIndex, int32 MultiplexIndex) OVERRIDE;
-	virtual void GetCaption(TArray<FString>& OutCaptions) const OVERRIDE;
-	virtual UTexture* GetReferencedTexture() OVERRIDE;
+	virtual bool IsResultMaterialAttributes(int32 OutputIndex) override;
+	virtual int32 Compile(class FMaterialCompiler* Compiler, int32 OutputIndex, int32 MultiplexIndex) override;
+	virtual void GetCaption(TArray<FString>& OutCaptions) const override;
+	virtual UTexture* GetReferencedTexture() override;
 #if WITH_EDITOR
-	virtual uint32 GetInputType(int32 InputIndex) OVERRIDE {return MCT_Float | MCT_MaterialAttributes;}
+	virtual uint32 GetInputType(int32 InputIndex) override {return MCT_Float | MCT_MaterialAttributes;}
 #endif //WITH_EDITOR
 	// End UMaterialExpression Interface
 
-	ENGINE_API virtual FGuid& GetParameterExpressionId() OVERRIDE
+	ENGINE_API virtual FGuid& GetParameterExpressionId() override
 	{
 		return ExpressionGUID;
 	}
diff --git a/Engine/Source/Runtime/Engine/Classes/Materials/MaterialExpressionLandscapeVisibilityMask.h b/Engine/Source/Runtime/Engine/Classes/Materials/MaterialExpressionLandscapeVisibilityMask.h
--- a/Engine/Source/Runtime/Engine/Classes/Materials/MaterialExpressionLandscapeVisibilityMask.h
+++ b/Engine/Source/Runtime/Engine/Classes/Materials/MaterialExpressionLandscapeVisibilityMask.h
@@ -19,12 +19,12 @@ public:
 	ENGINE_API static FName ParameterName;
 
 	// Begin UMaterialExpression Interface
-	virtual int32 Compile(class FMaterialCompiler* Compiler, int32 OutputIndex, int32 MultiplexIndex) OVERRIDE;
-	virtual void GetCaption(TArray<FString>& OutCaptions) const OVERRIDE;
-	virtual UTexture* GetReferencedTexture() OVERRIDE;
+	virtual int32 Compile(class FMaterialCompiler* Compiler, int32 OutputIndex, int32 MultiplexIndex) override;
+	virtual void GetCaption(TArray<FString>& OutCaptions) const override;
+	virtual UTexture* GetReferencedTexture() override;
 	// End UMaterialExpression Interface
 
-	ENGINE_API virtual FGuid& GetParameterExpressionId() OVERRIDE
+	ENGINE_API virtual FGuid& GetParameterExpressionId() override
 	{
 		return ExpressionGUID;
 	}
diff --git a/Engine/Source/Runtime/Engine/Classes/Materials/MaterialExpressionMaterialFunctionCall.h b/Engine/Source/Runtime/Engine/Classes/Materials/MaterialExpressionMaterialFunctionCall.h
--- a/Engine/Source/Runtime/Engine/Classes/Materials/MaterialExpressionMaterialFunctionCall.h
+++ b/Engine/Source/Runtime/Engine/Classes/Materials/MaterialExpressionMaterialFunctionCall.h
@@ -86,28 +86,28 @@ class UMaterialExpressionMaterialFunctionCall : public UMaterialExpression
 
 	// Begin UObject interface.
 #if WITH_EDITOR
-	virtual void PreEditChange(UProperty* PropertyAboutToChange) OVERRIDE;
-	virtual void PostEditChangeProperty(FPropertyChangedEvent& PropertyChangedEvent) OVERRIDE;
+	virtual void PreEditChange(UProperty* PropertyAboutToChange) override;
+	virtual void PostEditChangeProperty(FPropertyChangedEvent& PropertyChangedEvent) override;
 #endif // WITH_EDITOR
-	virtual void PostLoad() OVERRIDE;
+	virtual void PostLoad() override;
 	// End UObject interface.
 
 	// Begin UMaterialExpression Interface
-	virtual int32 Compile(class FMaterialCompiler* Compiler, int32 OutputIndex, int32 MultiplexIndex) OVERRIDE;
-	virtual void GetCaption(TArray<FString>& OutCaptions) const OVERRIDE;
-	virtual const TArray<FExpressionInput*> GetInputs() OVERRIDE;
-	virtual FExpressionInput* GetInput(int32 InputIndex) OVERRIDE;
-	virtual FString GetInputName(int32 InputIndex) const OVERRIDE;
-	virtual bool IsInputConnectionRequired(int32 InputIndex) const OVERRIDE;
+	virtual int32 Compile(class FMaterialCompiler* Compiler, int32 OutputIndex, int32 MultiplexIndex) override;
+	virtual void GetCaption(TArray<FString>& OutCaptions) const override;
+	virtual const TArray<FExpressionInput*> GetInputs() override;
+	virtual FExpressionInput* GetInput(int32 InputIndex) override;
+	virtual FString GetInputName(int32 InputIndex) const override;
+	virtual bool IsInputConnectionRequired(int32 InputIndex) const override;
 #if WITH_EDITOR
-	virtual FString GetDescription() const OVERRIDE;
-	virtual void GetConnectorToolTip(int32 InputIndex, int32 OutputIndex, TArray<FString>& OutToolTip) OVERRIDE;
-	virtual void GetExpressionToolTip(TArray<FString>& OutToolTip) OVERRIDE;
+	virtual FString GetDescription() const override;
+	virtual void GetConnectorToolTip(int32 InputIndex, int32 OutputIndex, TArray<FString>& OutToolTip) override;
+	virtual void GetExpressionToolTip(TArray<FString>& OutToolTip) override;
 #endif
-	virtual bool MatchesSearchQuery( const TCHAR* SearchQuery ) OVERRIDE;
-	virtual bool IsResultMaterialAttributes(int32 OutputIndex) OVERRIDE;
+	virtual bool MatchesSearchQuery( const TCHAR* SearchQuery ) override;
+	virtual bool IsResultMaterialAttributes(int32 OutputIndex) override;
 #if WITH_EDITOR
-	virtual uint32 GetInputType(int32 InputIndex) OVERRIDE;
+	virtual uint32 GetInputType(int32 InputIndex) override;
 #endif // WITH_EDITOR
 	// End UMaterialExpression Interface
 
diff --git a/Engine/Source/Runtime/Engine/Classes/Materials/MaterialExpressionNormalize.h b/Engine/Source/Runtime/Engine/Classes/Materials/MaterialExpressionNormalize.h
--- a/Engine/Source/Runtime/Engine/Classes/Materials/MaterialExpressionNormalize.h
+++ b/Engine/Source/Runtime/Engine/Classes/Materials/MaterialExpressionNormalize.h
@@ -15,8 +15,8 @@ class UMaterialExpressionNormalize : public UMaterialExpression
 
 
 	// Begin UMaterialExpression Interface
-	virtual int32 Compile(class FMaterialCompiler* Compiler, int32 OutputIndex, int32 MultiplexIndex) OVERRIDE;
-	virtual void GetCaption(TArray<FString>& OutCaptions) const OVERRIDE { OutCaptions.Add(TEXT("Normalize")); }
+	virtual int32 Compile(class FMaterialCompiler* Compiler, int32 OutputIndex, int32 MultiplexIndex) override;
+	virtual void GetCaption(TArray<FString>& OutCaptions) const override { OutCaptions.Add(TEXT("Normalize")); }
 	// End UMaterialExpression Interface
 };
 
diff --git a/Engine/Source/Runtime/Engine/Classes/Materials/MaterialExpressionPanner.h b/Engine/Source/Runtime/Engine/Classes/Materials/MaterialExpressionPanner.h
--- a/Engine/Source/Runtime/Engine/Classes/Materials/MaterialExpressionPanner.h
+++ b/Engine/Source/Runtime/Engine/Classes/Materials/MaterialExpressionPanner.h
@@ -32,9 +32,9 @@ class UMaterialExpressionPanner : public UMaterialExpression
 	bool bFractionalPart;
 
 	// Begin UMaterialExpression Interface
-	virtual int32 Compile(class FMaterialCompiler* Compiler, int32 OutputIndex, int32 MultiplexIndex) OVERRIDE;
-	virtual void GetCaption(TArray<FString>& OutCaptions) const OVERRIDE;
-	virtual bool NeedsRealtimePreview() OVERRIDE { return Time.Expression==NULL && (SpeedX != 0.f || SpeedY != 0.f); }
+	virtual int32 Compile(class FMaterialCompiler* Compiler, int32 OutputIndex, int32 MultiplexIndex) override;
+	virtual void GetCaption(TArray<FString>& OutCaptions) const override;
+	virtual bool NeedsRealtimePreview() override { return Time.Expression==NULL && (SpeedX != 0.f || SpeedY != 0.f); }
 	// End UMaterialExpression Interface
 
 };
diff --git a/Engine/Source/Runtime/Engine/Classes/Materials/MaterialExpressionParameter.h b/Engine/Source/Runtime/Engine/Classes/Materials/MaterialExpressionParameter.h
--- a/Engine/Source/Runtime/Engine/Classes/Materials/MaterialExpressionParameter.h
+++ b/Engine/Source/Runtime/Engine/Classes/Materials/MaterialExpressionParameter.h
@@ -23,10 +23,10 @@ class UMaterialExpressionParameter : public UMaterialExpression
 	FName Group;
 
 	// Begin UMaterialExpression Interface
-	virtual bool MatchesSearchQuery( const TCHAR* SearchQuery ) OVERRIDE;
+	virtual bool MatchesSearchQuery( const TCHAR* SearchQuery ) override;
 	// End UMaterialExpression Interface
 
-	ENGINE_API virtual FGuid& GetParameterExpressionId() OVERRIDE
+	ENGINE_API virtual FGuid& GetParameterExpressionId() override
 	{
 		return ExpressionGUID;
 	}
diff --git a/Engine/Source/Runtime/Engine/Classes/Materials/MaterialExpressionParticleSubUV.h b/Engine/Source/Runtime/Engine/Classes/Materials/MaterialExpressionParticleSubUV.h
--- a/Engine/Source/Runtime/Engine/Classes/Materials/MaterialExpressionParticleSubUV.h
+++ b/Engine/Source/Runtime/Engine/Classes/Materials/MaterialExpressionParticleSubUV.h
@@ -15,10 +15,10 @@ class UMaterialExpressionParticleSubUV : public UMaterialExpressionTextureSample
 
 
 	// Begin UMaterialExpression Interface
-	virtual int32 Compile(class FMaterialCompiler* Compiler, int32 OutputIndex, int32 MultiplexIndex) OVERRIDE;
-	virtual int32 GetWidth() const OVERRIDE;
-	virtual void GetCaption(TArray<FString>& OutCaptions) const OVERRIDE;
-	virtual int32 GetLabelPadding() OVERRIDE { return 8; }
+	virtual int32 Compile(class FMaterialCompiler* Compiler, int32 OutputIndex, int32 MultiplexIndex) override;
+	virtual int32 GetWidth() const override;
+	virtual void GetCaption(TArray<FString>& OutCaptions) const override;
+	virtual int32 GetLabelPadding() override { return 8; }
 	// End UMaterialExpression Interface
 };
 
diff --git a/Engine/Source/Runtime/Engine/Classes/Materials/MaterialExpressionQualitySwitch.h b/Engine/Source/Runtime/Engine/Classes/Materials/MaterialExpressionQualitySwitch.h
--- a/Engine/Source/Runtime/Engine/Classes/Materials/MaterialExpressionQualitySwitch.h
+++ b/Engine/Source/Runtime/Engine/Classes/Materials/MaterialExpressionQualitySwitch.h
@@ -18,17 +18,17 @@ class UMaterialExpressionQualitySwitch : public UMaterialExpression
 	FExpressionInput Inputs[EMaterialQualityLevel::Num];
 
 	// Begin UMaterialExpression Interface
-	virtual int32 Compile(class FMaterialCompiler* Compiler, int32 OutputIndex, int32 MultiplexIndex) OVERRIDE;
-	virtual void GetCaption(TArray<FString>& OutCaptions) const OVERRIDE;
+	virtual int32 Compile(class FMaterialCompiler* Compiler, int32 OutputIndex, int32 MultiplexIndex) override;
+	virtual void GetCaption(TArray<FString>& OutCaptions) const override;
 	
-	virtual const TArray<FExpressionInput*> GetInputs() OVERRIDE;
-	virtual FExpressionInput* GetInput(int32 InputIndex) OVERRIDE;
-	virtual FString GetInputName(int32 InputIndex) const OVERRIDE;
-	virtual bool IsInputConnectionRequired(int32 InputIndex) const OVERRIDE;
-	virtual bool IsResultMaterialAttributes(int32 OutputIndex) OVERRIDE;
+	virtual const TArray<FExpressionInput*> GetInputs() override;
+	virtual FExpressionInput* GetInput(int32 InputIndex) override;
+	virtual FString GetInputName(int32 InputIndex) const override;
+	virtual bool IsInputConnectionRequired(int32 InputIndex) const override;
+	virtual bool IsResultMaterialAttributes(int32 OutputIndex) override;
 #if WITH_EDITOR
-	virtual uint32 GetInputType(int32 InputIndex) OVERRIDE {return MCT_Unknown;}
-	virtual uint32 GetOutputType(int32 InputIndex) OVERRIDE {return MCT_Unknown;}
+	virtual uint32 GetInputType(int32 InputIndex) override {return MCT_Unknown;}
+	virtual uint32 GetOutputType(int32 InputIndex) override {return MCT_Unknown;}
 #endif
 	// End UMaterialExpression Interface
 };
diff --git a/Engine/Source/Runtime/Engine/Classes/Materials/MaterialExpressionRotator.h b/Engine/Source/Runtime/Engine/Classes/Materials/MaterialExpressionRotator.h
--- a/Engine/Source/Runtime/Engine/Classes/Materials/MaterialExpressionRotator.h
+++ b/Engine/Source/Runtime/Engine/Classes/Materials/MaterialExpressionRotator.h
@@ -30,9 +30,9 @@ class UMaterialExpressionRotator : public UMaterialExpression
 	uint32 ConstCoordinate;
 
 	// Begin UMaterialExpression Interface
-	virtual int32 Compile(class FMaterialCompiler* Compiler, int32 OutputIndex, int32 MultiplexIndex) OVERRIDE;
-	virtual void GetCaption(TArray<FString>& OutCaptions) const OVERRIDE;
-	virtual bool NeedsRealtimePreview() OVERRIDE { return Time.Expression==NULL && Speed != 0.f; }
+	virtual int32 Compile(class FMaterialCompiler* Compiler, int32 OutputIndex, int32 MultiplexIndex) override;
+	virtual void GetCaption(TArray<FString>& OutCaptions) const override;
+	virtual bool NeedsRealtimePreview() override { return Time.Expression==NULL && Speed != 0.f; }
 	// End UMaterialExpression Interface
 
 };
diff --git a/Engine/Source/Runtime/Engine/Classes/Materials/MaterialExpressionSceneColor.h b/Engine/Source/Runtime/Engine/Classes/Materials/MaterialExpressionSceneColor.h
--- a/Engine/Source/Runtime/Engine/Classes/Materials/MaterialExpressionSceneColor.h
+++ b/Engine/Source/Runtime/Engine/Classes/Materials/MaterialExpressionSceneColor.h
@@ -43,13 +43,13 @@ class UMaterialExpressionSceneColor : public UMaterialExpression
 	FVector2D ConstInput;
 
 	// Begin UObject interface.
-	virtual void PostLoad() OVERRIDE;
+	virtual void PostLoad() override;
 	// End UObject interface.
 
 	// Begin UMaterialExpression Interface
-	virtual int32 Compile(class FMaterialCompiler* Compiler, int32 OutputIndex, int32 MultiplexIndex) OVERRIDE;
-	virtual void GetCaption(TArray<FString>& OutCaptions) const OVERRIDE;
-	virtual FString GetInputName(int32 InputIndex) const OVERRIDE;
+	virtual int32 Compile(class FMaterialCompiler* Compiler, int32 OutputIndex, int32 MultiplexIndex) override;
+	virtual void GetCaption(TArray<FString>& OutCaptions) const override;
+	virtual FString GetInputName(int32 InputIndex) const override;
 	// End UMaterialExpression Interface
 };
 
diff --git a/Engine/Source/Runtime/Engine/Classes/Materials/MaterialExpressionSceneDepth.h b/Engine/Source/Runtime/Engine/Classes/Materials/MaterialExpressionSceneDepth.h
--- a/Engine/Source/Runtime/Engine/Classes/Materials/MaterialExpressionSceneDepth.h
+++ b/Engine/Source/Runtime/Engine/Classes/Materials/MaterialExpressionSceneDepth.h
@@ -33,13 +33,13 @@ class UMaterialExpressionSceneDepth : public UMaterialExpression
 	FVector2D ConstInput;
 
 	// Begin UObject interface.
-	virtual void PostLoad() OVERRIDE;
+	virtual void PostLoad() override;
 	// End UObject interface.
 
 	// Begin UMaterialExpression Interface
-	virtual int32 Compile(class FMaterialCompiler* Compiler, int32 OutputIndex, int32 MultiplexIndex) OVERRIDE;
-	virtual void GetCaption(TArray<FString>& OutCaptions) const OVERRIDE;
-	virtual FString GetInputName(int32 InputIndex) const OVERRIDE;
+	virtual int32 Compile(class FMaterialCompiler* Compiler, int32 OutputIndex, int32 MultiplexIndex) override;
+	virtual void GetCaption(TArray<FString>& OutCaptions) const override;
+	virtual FString GetInputName(int32 InputIndex) const override;
 	// End UMaterialExpression Interface
 };
 
diff --git a/Engine/Source/Runtime/Engine/Classes/Materials/MaterialExpressionSpeedTree.h b/Engine/Source/Runtime/Engine/Classes/Materials/MaterialExpressionSpeedTree.h
--- a/Engine/Source/Runtime/Engine/Classes/Materials/MaterialExpressionSpeedTree.h
+++ b/Engine/Source/Runtime/Engine/Classes/Materials/MaterialExpressionSpeedTree.h
@@ -51,15 +51,15 @@ class UMaterialExpressionSpeedTree : public UMaterialExpression
 	float BillboardThreshold;
 
 	// Begin UObject Interface
-	virtual void Serialize(FArchive& Ar) OVERRIDE;
+	virtual void Serialize(FArchive& Ar) override;
 #if WITH_EDITOR
-	virtual bool CanEditChange(const UProperty* InProperty) const OVERRIDE;
+	virtual bool CanEditChange(const UProperty* InProperty) const override;
 #endif // WITH_EDITOR
 	// End UObject Interface
 
 	// Begin UMaterialExpression Interface
-	virtual int32 Compile(class FMaterialCompiler* Compiler, int32 OutputIndex, int32 MultiplexIndex) OVERRIDE;
-	virtual void GetCaption(TArray<FString>& OutCaptions) const OVERRIDE;
+	virtual int32 Compile(class FMaterialCompiler* Compiler, int32 OutputIndex, int32 MultiplexIndex) override;
+	virtual void GetCaption(TArray<FString>& OutCaptions) const override;
 	// End UMaterialExpression Interface
 };
 
diff --git a/Engine/Source/Runtime/Engine/Classes/Materials/MaterialExpressionStaticBool.h b/Engine/Source/Runtime/Engine/Classes/Materials/MaterialExpressionStaticBool.h
--- a/Engine/Source/Runtime/Engine/Classes/Materials/MaterialExpressionStaticBool.h
+++ b/Engine/Source/Runtime/Engine/Classes/Materials/MaterialExpressionStaticBool.h
@@ -15,11 +15,11 @@ class UMaterialExpressionStaticBool : public UMaterialExpression
 
 
 	// Begin UMaterialExpression Interface
-	virtual int32 CompilePreview(class FMaterialCompiler* Compiler, int32 OutputIndex, int32 MultiplexIndex) OVERRIDE;
-	virtual int32 Compile(class FMaterialCompiler* Compiler, int32 OutputIndex, int32 MultiplexIndex) OVERRIDE;
-	virtual void GetCaption(TArray<FString>& OutCaptions) const OVERRIDE;
+	virtual int32 CompilePreview(class FMaterialCompiler* Compiler, int32 OutputIndex, int32 MultiplexIndex) override;
+	virtual int32 Compile(class FMaterialCompiler* Compiler, int32 OutputIndex, int32 MultiplexIndex) override;
+	virtual void GetCaption(TArray<FString>& OutCaptions) const override;
 #if WITH_EDITOR
-	virtual uint32 GetOutputType(int32 OutputIndex) OVERRIDE {return MCT_StaticBool;}
+	virtual uint32 GetOutputType(int32 OutputIndex) override {return MCT_StaticBool;}
 #endif
 	// End UMaterialExpression Interface
 };
diff --git a/Engine/Source/Runtime/Engine/Classes/Materials/MaterialExpressionStaticBoolParameter.h b/Engine/Source/Runtime/Engine/Classes/Materials/MaterialExpressionStaticBoolParameter.h
--- a/Engine/Source/Runtime/Engine/Classes/Materials/MaterialExpressionStaticBoolParameter.h
+++ b/Engine/Source/Runtime/Engine/Classes/Materials/MaterialExpressionStaticBoolParameter.h
@@ -15,11 +15,11 @@ class UMaterialExpressionStaticBoolParameter : public UMaterialExpressionParamet
 
 public:
 	// Begin UMaterialExpression Interface
-	virtual int32 Compile(class FMaterialCompiler* Compiler, int32 OutputIndex, int32 MultiplexIndex) OVERRIDE;
-	virtual int32 CompilePreview(class FMaterialCompiler* Compiler, int32 OutputIndex, int32 MultiplexIndex) OVERRIDE;
-	virtual void GetCaption(TArray<FString>& OutCaptions) const OVERRIDE;
+	virtual int32 Compile(class FMaterialCompiler* Compiler, int32 OutputIndex, int32 MultiplexIndex) override;
+	virtual int32 CompilePreview(class FMaterialCompiler* Compiler, int32 OutputIndex, int32 MultiplexIndex) override;
+	virtual void GetCaption(TArray<FString>& OutCaptions) const override;
 #if WITH_EDITOR
-	virtual uint32 GetOutputType(int32 OutputIndex) OVERRIDE {return MCT_StaticBool;}
+	virtual uint32 GetOutputType(int32 OutputIndex) override {return MCT_StaticBool;}
 #endif
 	// End UMaterialExpression Interface
 
diff --git a/Engine/Source/Runtime/Engine/Classes/Materials/MaterialExpressionStaticSwitch.h b/Engine/Source/Runtime/Engine/Classes/Materials/MaterialExpressionStaticSwitch.h
--- a/Engine/Source/Runtime/Engine/Classes/Materials/MaterialExpressionStaticSwitch.h
+++ b/Engine/Source/Runtime/Engine/Classes/Materials/MaterialExpressionStaticSwitch.h
@@ -24,13 +24,13 @@ class UMaterialExpressionStaticSwitch : public UMaterialExpression
 
 
 	// Begin UMaterialExpression Interface
-	virtual int32 Compile(class FMaterialCompiler* Compiler, int32 OutputIndex, int32 MultiplexIndex) OVERRIDE;
-	virtual void GetCaption(TArray<FString>& OutCaptions) const OVERRIDE;
-	virtual FString GetInputName(int32 InputIndex) const OVERRIDE;
-	virtual bool IsResultMaterialAttributes(int32 OutputIndex) OVERRIDE;
+	virtual int32 Compile(class FMaterialCompiler* Compiler, int32 OutputIndex, int32 MultiplexIndex) override;
+	virtual void GetCaption(TArray<FString>& OutCaptions) const override;
+	virtual FString GetInputName(int32 InputIndex) const override;
+	virtual bool IsResultMaterialAttributes(int32 OutputIndex) override;
 #if WITH_EDITOR
-	virtual uint32 GetInputType(int32 InputIndex) OVERRIDE;
-	virtual uint32 GetOutputType(int32 OutputIndex) OVERRIDE {return MCT_Unknown;}
+	virtual uint32 GetInputType(int32 InputIndex) override;
+	virtual uint32 GetOutputType(int32 OutputIndex) override {return MCT_Unknown;}
 #endif
 	// End UMaterialExpression Interface
 };
diff --git a/Engine/Source/Runtime/Engine/Classes/Materials/MaterialExpressionStaticSwitchParameter.h b/Engine/Source/Runtime/Engine/Classes/Materials/MaterialExpressionStaticSwitchParameter.h
--- a/Engine/Source/Runtime/Engine/Classes/Materials/MaterialExpressionStaticSwitchParameter.h
+++ b/Engine/Source/Runtime/Engine/Classes/Materials/MaterialExpressionStaticSwitchParameter.h
@@ -18,13 +18,13 @@ class UMaterialExpressionStaticSwitchParameter : public UMaterialExpressionStati
 
 
 	// Begin UMaterialExpression Interface
-	virtual FString GetInputName(int32 InputIndex) const OVERRIDE;
-	virtual int32 Compile(class FMaterialCompiler* Compiler, int32 OutputIndex, int32 MultiplexIndex) OVERRIDE;
-	virtual void GetCaption(TArray<FString>& OutCaptions) const OVERRIDE;
-	virtual bool IsResultMaterialAttributes(int32 OutputIndex) OVERRIDE;
+	virtual FString GetInputName(int32 InputIndex) const override;
+	virtual int32 Compile(class FMaterialCompiler* Compiler, int32 OutputIndex, int32 MultiplexIndex) override;
+	virtual void GetCaption(TArray<FString>& OutCaptions) const override;
+	virtual bool IsResultMaterialAttributes(int32 OutputIndex) override;
 #if WITH_EDITOR
-	virtual uint32 GetInputType(int32 InputIndex) OVERRIDE {return MCT_Unknown;}
-	virtual uint32 GetOutputType(int32 OutputIndex) OVERRIDE {return MCT_Unknown;}
+	virtual uint32 GetInputType(int32 InputIndex) override {return MCT_Unknown;}
+	virtual uint32 GetOutputType(int32 OutputIndex) override {return MCT_Unknown;}
 #endif
 	// End UMaterialExpression Interface
 };
diff --git a/Engine/Source/Runtime/Engine/Classes/Materials/MaterialExpressionTextureBase.h b/Engine/Source/Runtime/Engine/Classes/Materials/MaterialExpressionTextureBase.h
--- a/Engine/Source/Runtime/Engine/Classes/Materials/MaterialExpressionTextureBase.h
+++ b/Engine/Source/Runtime/Engine/Classes/Materials/MaterialExpressionTextureBase.h
@@ -26,13 +26,13 @@ class UMaterialExpressionTextureBase : public UMaterialExpression
 	
 	// Begin UObject Interface
 #if WITH_EDITOR
-	virtual void PostEditChangeProperty(FPropertyChangedEvent& PropertyChangedEvent) OVERRIDE;
+	virtual void PostEditChangeProperty(FPropertyChangedEvent& PropertyChangedEvent) override;
 #endif // WITH_EDITOR
 	// End UObject Interface
 
 	// Begin UMaterialExpression Interface
 #if WITH_EDITOR
-	virtual FString GetDescription() const OVERRIDE;
+	virtual FString GetDescription() const override;
 #endif
 	// End UMaterialExpression Interface
 
@@ -41,7 +41,7 @@ class UMaterialExpressionTextureBase : public UMaterialExpression
 	 * This is used to link the compiled uniform expressions with their default texture values. 
 	 * Any UMaterialExpression whose compilation creates a texture uniform expression (eg Compiler->Texture, Compiler->TextureParameter) must implement this.
 	 */
-	virtual UTexture* GetReferencedTexture() OVERRIDE
+	virtual UTexture* GetReferencedTexture() override
 	{
 		return Texture;
 	}
diff --git a/Engine/Source/Runtime/Engine/Classes/Materials/MaterialExpressionTextureObject.h b/Engine/Source/Runtime/Engine/Classes/Materials/MaterialExpressionTextureObject.h
--- a/Engine/Source/Runtime/Engine/Classes/Materials/MaterialExpressionTextureObject.h
+++ b/Engine/Source/Runtime/Engine/Classes/Materials/MaterialExpressionTextureObject.h
@@ -15,13 +15,13 @@ class UMaterialExpressionTextureObject : public UMaterialExpressionTextureBase
 	GENERATED_UCLASS_BODY()
 
 #if WITH_EDITOR
-	virtual void PostEditChangeProperty(FPropertyChangedEvent& PropertyChangedEvent) OVERRIDE;
+	virtual void PostEditChangeProperty(FPropertyChangedEvent& PropertyChangedEvent) override;
 #endif // WITH_EDITOR
 	// Begin UMaterialExpression Interface
-	virtual int32 Compile(class FMaterialCompiler* Compiler, int32 OutputIndex, int32 MultiplexIndex) OVERRIDE;
-	virtual void GetCaption(TArray<FString>& OutCaptions) const OVERRIDE;
+	virtual int32 Compile(class FMaterialCompiler* Compiler, int32 OutputIndex, int32 MultiplexIndex) override;
+	virtual void GetCaption(TArray<FString>& OutCaptions) const override;
 #if WITH_EDITOR
-	virtual uint32 GetOutputType(int32 OutputIndex) OVERRIDE;
+	virtual uint32 GetOutputType(int32 OutputIndex) override;
 #endif // WITH_EDITOR
 	// End UMaterialExpression Interface
 	virtual int32 CompilePreview(class FMaterialCompiler* Compiler, int32 OutputIndex, int32 MultiplexIndex);
diff --git a/Engine/Source/Runtime/Engine/Classes/Materials/MaterialExpressionTextureObjectParameter.h b/Engine/Source/Runtime/Engine/Classes/Materials/MaterialExpressionTextureObjectParameter.h
--- a/Engine/Source/Runtime/Engine/Classes/Materials/MaterialExpressionTextureObjectParameter.h
+++ b/Engine/Source/Runtime/Engine/Classes/Materials/MaterialExpressionTextureObjectParameter.h
@@ -16,17 +16,17 @@ class UMaterialExpressionTextureObjectParameter : public UMaterialExpressionText
 
 
 	// Begin UMaterialExpression Interface
-	virtual void GetCaption(TArray<FString>& OutCaptions) const OVERRIDE;
-	virtual const TArray<FExpressionInput*> GetInputs() OVERRIDE;
-	virtual int32 Compile(class FMaterialCompiler* Compiler, int32 OutputIndex, int32 MultiplexIndex) OVERRIDE;
-	virtual int32 CompilePreview(class FMaterialCompiler* Compiler, int32 OutputIndex, int32 MultiplexIndex) OVERRIDE;
+	virtual void GetCaption(TArray<FString>& OutCaptions) const override;
+	virtual const TArray<FExpressionInput*> GetInputs() override;
+	virtual int32 Compile(class FMaterialCompiler* Compiler, int32 OutputIndex, int32 MultiplexIndex) override;
+	virtual int32 CompilePreview(class FMaterialCompiler* Compiler, int32 OutputIndex, int32 MultiplexIndex) override;
 #if WITH_EDITOR
-	virtual uint32 GetOutputType(int32 InputIndex) OVERRIDE {return MCT_Texture;}
+	virtual uint32 GetOutputType(int32 InputIndex) override {return MCT_Texture;}
 #endif
 	// End UMaterialExpression Interface
 
 	// Begin UMaterialExpressionTextureSampleParameter Interface
-	virtual const TCHAR* GetRequirements() OVERRIDE;
+	virtual const TCHAR* GetRequirements() override;
 	// End UMaterialExpressionTextureSampleParameter Interface
 
 };
diff --git a/Engine/Source/Runtime/Engine/Classes/Materials/MaterialExpressionTextureSample.h b/Engine/Source/Runtime/Engine/Classes/Materials/MaterialExpressionTextureSample.h
--- a/Engine/Source/Runtime/Engine/Classes/Materials/MaterialExpressionTextureSample.h
+++ b/Engine/Source/Runtime/Engine/Classes/Materials/MaterialExpressionTextureSample.h
@@ -49,20 +49,20 @@ class UMaterialExpressionTextureSample : public UMaterialExpressionTextureBase
 
 	// Begin UObject Interface
 #if WITH_EDITOR
-	virtual bool CanEditChange(const UProperty* InProperty) const OVERRIDE;
-	virtual void PostEditChangeProperty(FPropertyChangedEvent& PropertyChangedEvent) OVERRIDE;
+	virtual bool CanEditChange(const UProperty* InProperty) const override;
+	virtual void PostEditChangeProperty(FPropertyChangedEvent& PropertyChangedEvent) override;
 #endif // WITH_EDITOR
 	// End UObject Interface
 
 	// Begin UMaterialExpression Interface
-	virtual const TArray<FExpressionInput*> GetInputs() OVERRIDE;
-	virtual FExpressionInput* GetInput(int32 InputIndex) OVERRIDE;
-	virtual FString GetInputName(int32 InputIndex) const OVERRIDE;
-	virtual int32 GetWidth() const OVERRIDE;
-	virtual int32 GetLabelPadding() OVERRIDE { return 8; }
-	virtual int32 Compile(class FMaterialCompiler* Compiler, int32 OutputIndex, int32 MultiplexIndex) OVERRIDE;
-	virtual void GetCaption(TArray<FString>& OutCaptions) const OVERRIDE;
-	virtual bool MatchesSearchQuery( const TCHAR* SearchQuery ) OVERRIDE;
+	virtual const TArray<FExpressionInput*> GetInputs() override;
+	virtual FExpressionInput* GetInput(int32 InputIndex) override;
+	virtual FString GetInputName(int32 InputIndex) const override;
+	virtual int32 GetWidth() const override;
+	virtual int32 GetLabelPadding() override { return 8; }
+	virtual int32 Compile(class FMaterialCompiler* Compiler, int32 OutputIndex, int32 MultiplexIndex) override;
+	virtual void GetCaption(TArray<FString>& OutCaptions) const override;
+	virtual bool MatchesSearchQuery( const TCHAR* SearchQuery ) override;
 #if WITH_EDITOR
 	virtual uint32 GetInputType(int32 InputIndex);
 #endif
diff --git a/Engine/Source/Runtime/Engine/Classes/Materials/MaterialExpressionTextureSampleParameter.h b/Engine/Source/Runtime/Engine/Classes/Materials/MaterialExpressionTextureSampleParameter.h
--- a/Engine/Source/Runtime/Engine/Classes/Materials/MaterialExpressionTextureSampleParameter.h
+++ b/Engine/Source/Runtime/Engine/Classes/Materials/MaterialExpressionTextureSampleParameter.h
@@ -22,8 +22,8 @@ class UMaterialExpressionTextureSampleParameter : public UMaterialExpressionText
 	FName Group;
 
 	// Begin UMaterialExpression Interface
-	virtual int32 Compile(class FMaterialCompiler* Compiler, int32 OutputIndex, int32 MultiplexIndex) OVERRIDE;
-	virtual void GetCaption(TArray<FString>& OutCaptions) const OVERRIDE;
+	virtual int32 Compile(class FMaterialCompiler* Compiler, int32 OutputIndex, int32 MultiplexIndex) override;
+	virtual void GetCaption(TArray<FString>& OutCaptions) const override;
 	// End UMaterialExpression Interface
 
 	/**
@@ -46,7 +46,7 @@ class UMaterialExpressionTextureSampleParameter : public UMaterialExpressionText
 	 */
 	virtual void SetDefaultTexture();
 
-	ENGINE_API virtual FGuid& GetParameterExpressionId() OVERRIDE
+	ENGINE_API virtual FGuid& GetParameterExpressionId() override
 	{
 		return ExpressionGUID;
 	}
diff --git a/Engine/Source/Runtime/Engine/Classes/Materials/MaterialExpressionTextureSampleParameter2D.h b/Engine/Source/Runtime/Engine/Classes/Materials/MaterialExpressionTextureSampleParameter2D.h
--- a/Engine/Source/Runtime/Engine/Classes/Materials/MaterialExpressionTextureSampleParameter2D.h
+++ b/Engine/Source/Runtime/Engine/Classes/Materials/MaterialExpressionTextureSampleParameter2D.h
@@ -12,13 +12,13 @@ class UMaterialExpressionTextureSampleParameter2D : public UMaterialExpressionTe
 
 
 	// Begin UMaterialExpression Interface
-	virtual void GetCaption(TArray<FString>& OutCaptions) const OVERRIDE;
+	virtual void GetCaption(TArray<FString>& OutCaptions) const override;
 	// End UMaterialExpression Interface
 	
 	// Begin UMaterialExpressionTextureSampleParameter Interface
-	virtual bool TextureIsValid( UTexture* InTexture ) OVERRIDE;
-	virtual const TCHAR* GetRequirements() OVERRIDE;
-	virtual void SetDefaultTexture() OVERRIDE;
+	virtual bool TextureIsValid( UTexture* InTexture ) override;
+	virtual const TCHAR* GetRequirements() override;
+	virtual void SetDefaultTexture() override;
 	// End UMaterialExpressionTextureSampleParameter Interface
 };
 
diff --git a/Engine/Source/Runtime/Engine/Classes/Materials/MaterialExpressionTextureSampleParameterCube.h b/Engine/Source/Runtime/Engine/Classes/Materials/MaterialExpressionTextureSampleParameterCube.h
--- a/Engine/Source/Runtime/Engine/Classes/Materials/MaterialExpressionTextureSampleParameterCube.h
+++ b/Engine/Source/Runtime/Engine/Classes/Materials/MaterialExpressionTextureSampleParameterCube.h
@@ -14,14 +14,14 @@ class UMaterialExpressionTextureSampleParameterCube : public UMaterialExpression
 
 
 	// Begin UMaterialExpression Interface
-	virtual void GetCaption(TArray<FString>& OutCaptions) const OVERRIDE;
-	virtual int32 Compile(class FMaterialCompiler* Compiler, int32 OutputIndex, int32 MultiplexIndex) OVERRIDE;
+	virtual void GetCaption(TArray<FString>& OutCaptions) const override;
+	virtual int32 Compile(class FMaterialCompiler* Compiler, int32 OutputIndex, int32 MultiplexIndex) override;
 	// End UMaterialExpression Interface
 
 	// Begin UMaterialExpressionTextureSampleParameter Interface
-	virtual bool TextureIsValid( UTexture* InTexture ) OVERRIDE;
-	virtual const TCHAR* GetRequirements() OVERRIDE;
-	virtual void SetDefaultTexture() OVERRIDE;
+	virtual bool TextureIsValid( UTexture* InTexture ) override;
+	virtual const TCHAR* GetRequirements() override;
+	virtual void SetDefaultTexture() override;
 	// End UMaterialExpressionTextureSampleParameter Interface
 };
 
diff --git a/Engine/Source/Runtime/Engine/Classes/Materials/MaterialExpressionTextureSampleParameterMovie.h b/Engine/Source/Runtime/Engine/Classes/Materials/MaterialExpressionTextureSampleParameterMovie.h
--- a/Engine/Source/Runtime/Engine/Classes/Materials/MaterialExpressionTextureSampleParameterMovie.h
+++ b/Engine/Source/Runtime/Engine/Classes/Materials/MaterialExpressionTextureSampleParameterMovie.h
@@ -16,12 +16,12 @@ class UMaterialExpressionTextureSampleParameterMovie : public UMaterialExpressio
 
 
 	// Begin UMaterialExpression Interface
-	virtual void GetCaption(TArray<FString>& OutCaptions) const OVERRIDE;
+	virtual void GetCaption(TArray<FString>& OutCaptions) const override;
 	// End UMaterialExpression Interface
 
 	// Begin UMaterialExpressionTextureSampleParameter Interface
-	virtual bool TextureIsValid( UTexture* InTexture ) OVERRIDE;
-	virtual const TCHAR* GetRequirements() OVERRIDE;
+	virtual bool TextureIsValid( UTexture* InTexture ) override;
+	virtual const TCHAR* GetRequirements() override;
 	// End UMaterialExpressionTextureSampleParameter Interface
 };
 
diff --git a/Engine/Source/Runtime/Engine/Classes/Materials/MaterialExpressionTextureSampleParameterSubUV.h b/Engine/Source/Runtime/Engine/Classes/Materials/MaterialExpressionTextureSampleParameterSubUV.h
--- a/Engine/Source/Runtime/Engine/Classes/Materials/MaterialExpressionTextureSampleParameterSubUV.h
+++ b/Engine/Source/Runtime/Engine/Classes/Materials/MaterialExpressionTextureSampleParameterSubUV.h
@@ -15,13 +15,13 @@ class UMaterialExpressionTextureSampleParameterSubUV : public UMaterialExpressio
 
 
 	// Begin UMaterialExpressionTextureSampleParameter Interface
-	virtual bool TextureIsValid( UTexture* InTexture ) OVERRIDE;
-	virtual const TCHAR* GetRequirements() OVERRIDE;
+	virtual bool TextureIsValid( UTexture* InTexture ) override;
+	virtual const TCHAR* GetRequirements() override;
 	// End UMaterialExpressionTextureSampleParameter Interface
 
 	// Begin UMaterialExpression Interface
-	virtual int32 Compile(class FMaterialCompiler* Compiler, int32 OutputIndex, int32 MultiplexIndex) OVERRIDE;
-	virtual void GetCaption(TArray<FString>& OutCaptions) const OVERRIDE;
+	virtual int32 Compile(class FMaterialCompiler* Compiler, int32 OutputIndex, int32 MultiplexIndex) override;
+	virtual void GetCaption(TArray<FString>& OutCaptions) const override;
 	// End UMaterialExpression Interface
 };
 
diff --git a/Engine/Source/Runtime/Engine/Classes/Materials/MaterialExpressionTime.h b/Engine/Source/Runtime/Engine/Classes/Materials/MaterialExpressionTime.h
--- a/Engine/Source/Runtime/Engine/Classes/Materials/MaterialExpressionTime.h
+++ b/Engine/Source/Runtime/Engine/Classes/Materials/MaterialExpressionTime.h
@@ -16,9 +16,9 @@ class UMaterialExpressionTime : public UMaterialExpression
 
 
 	// Begin UMaterialExpression Interface
-	virtual int32 Compile(class FMaterialCompiler* Compiler, int32 OutputIndex, int32 MultiplexIndex) OVERRIDE;
-	virtual void GetCaption(TArray<FString>& OutCaptions) const OVERRIDE;
-	virtual bool NeedsRealtimePreview() OVERRIDE { return true; }
+	virtual int32 Compile(class FMaterialCompiler* Compiler, int32 OutputIndex, int32 MultiplexIndex) override;
+	virtual void GetCaption(TArray<FString>& OutCaptions) const override;
+	virtual bool NeedsRealtimePreview() override { return true; }
 	// End UMaterialExpression Interface
 
 };
diff --git a/Engine/Source/Runtime/Engine/Classes/Materials/MaterialFunction.h b/Engine/Source/Runtime/Engine/Classes/Materials/MaterialFunction.h
--- a/Engine/Source/Runtime/Engine/Classes/Materials/MaterialFunction.h
+++ b/Engine/Source/Runtime/Engine/Classes/Materials/MaterialFunction.h
@@ -57,11 +57,11 @@ public:
 
 	// Begin UObject interface.
 #if WITH_EDITOR
-	virtual void PostEditChangeProperty(FPropertyChangedEvent& PropertyChangedEvent) OVERRIDE;
+	virtual void PostEditChangeProperty(FPropertyChangedEvent& PropertyChangedEvent) override;
 #endif // WITH_EDITOR
-	virtual void Serialize(FArchive& Ar) OVERRIDE;
-	virtual void PostLoad() OVERRIDE;
-	virtual void GetAssetRegistryTags(TArray<FAssetRegistryTag>& OutTags) const OVERRIDE;
+	virtual void Serialize(FArchive& Ar) override;
+	virtual void PostLoad() override;
+	virtual void GetAssetRegistryTags(TArray<FAssetRegistryTag>& OutTags) const override;
 	// End UObject interface.
 
 	/** Recursively update all function call expressions in this function, or in nested functions. */
diff --git a/Engine/Source/Runtime/Engine/Classes/Materials/MaterialInstance.h b/Engine/Source/Runtime/Engine/Classes/Materials/MaterialInstance.h
--- a/Engine/Source/Runtime/Engine/Classes/Materials/MaterialInstance.h
+++ b/Engine/Source/Runtime/Engine/Classes/Materials/MaterialInstance.h
@@ -208,36 +208,36 @@ public:
 	// End interface IBlendableInterface
 
 	// Begin UMaterialInterface interface.
-	virtual UMaterial* GetMaterial() OVERRIDE;
-	virtual const UMaterial* GetMaterial() const OVERRIDE;
-	virtual const UMaterial* GetMaterial_Concurrent(TMicRecursionGuard& RecursionGuard) const OVERRIDE;
-	virtual FMaterialResource* GetMaterialResource(ERHIFeatureLevel::Type InFeatureLevel, EMaterialQualityLevel::Type QualityLevel = EMaterialQualityLevel::Num) OVERRIDE;
-	virtual const FMaterialResource* GetMaterialResource(ERHIFeatureLevel::Type InFeatureLevel, EMaterialQualityLevel::Type QualityLevel = EMaterialQualityLevel::Num) const OVERRIDE;
-	virtual bool GetFontParameterValue(FName ParameterName,class UFont*& OutFontValue, int32& OutFontPage) const OVERRIDE;
-	virtual bool GetScalarParameterValue(FName ParameterName,float& OutValue) const OVERRIDE;
-	virtual bool GetTextureParameterValue(FName ParameterName,class UTexture*& OutValue) const OVERRIDE;
-	virtual bool GetVectorParameterValue(FName ParameterName,FLinearColor& OutValue) const OVERRIDE;
-	virtual void GetUsedTextures(TArray<UTexture*>& OutTextures, EMaterialQualityLevel::Type QualityLevel, bool bAllQualityLevels) const OVERRIDE;
-	virtual void OverrideTexture( const UTexture* InTextureToOverride, UTexture* OverrideTexture ) OVERRIDE;
-	virtual bool CheckMaterialUsage(const EMaterialUsage Usage, const bool bSkipPrim = false) OVERRIDE;
-	virtual bool CheckMaterialUsage_Concurrent(const EMaterialUsage Usage, const bool bSkipPrim = false) const OVERRIDE;
-	virtual bool GetStaticSwitchParameterValue(FName ParameterName,bool &OutValue,FGuid &OutExpressionGuid) OVERRIDE;
-	virtual bool GetStaticComponentMaskParameterValue(FName ParameterName, bool &R, bool &G, bool &B, bool &A,FGuid &OutExpressionGuid) OVERRIDE;
-	virtual bool GetTerrainLayerWeightParameterValue(FName ParameterName, int32& OutWeightmapIndex, FGuid &OutExpressionGuid) OVERRIDE;
-	virtual bool IsDependent(UMaterialInterface* TestDependency) OVERRIDE;
-	virtual FMaterialRenderProxy* GetRenderProxy(bool Selected, bool bHovered=false) const OVERRIDE;
-	virtual UPhysicalMaterial* GetPhysicalMaterial() const OVERRIDE;
-	virtual bool UpdateLightmassTextureTracking() OVERRIDE;
-	virtual bool GetCastShadowAsMasked() const OVERRIDE;
-	virtual float GetEmissiveBoost() const OVERRIDE;
-	virtual float GetDiffuseBoost() const OVERRIDE;
-	virtual float GetExportResolutionScale() const OVERRIDE;
-	virtual float GetDistanceFieldPenumbraScale() const OVERRIDE;
+	virtual UMaterial* GetMaterial() override;
+	virtual const UMaterial* GetMaterial() const override;
+	virtual const UMaterial* GetMaterial_Concurrent(TMicRecursionGuard& RecursionGuard) const override;
+	virtual FMaterialResource* GetMaterialResource(ERHIFeatureLevel::Type InFeatureLevel, EMaterialQualityLevel::Type QualityLevel = EMaterialQualityLevel::Num) override;
+	virtual const FMaterialResource* GetMaterialResource(ERHIFeatureLevel::Type InFeatureLevel, EMaterialQualityLevel::Type QualityLevel = EMaterialQualityLevel::Num) const override;
+	virtual bool GetFontParameterValue(FName ParameterName,class UFont*& OutFontValue, int32& OutFontPage) const override;
+	virtual bool GetScalarParameterValue(FName ParameterName,float& OutValue) const override;
+	virtual bool GetTextureParameterValue(FName ParameterName,class UTexture*& OutValue) const override;
+	virtual bool GetVectorParameterValue(FName ParameterName,FLinearColor& OutValue) const override;
+	virtual void GetUsedTextures(TArray<UTexture*>& OutTextures, EMaterialQualityLevel::Type QualityLevel, bool bAllQualityLevels) const override;
+	virtual void OverrideTexture( const UTexture* InTextureToOverride, UTexture* OverrideTexture ) override;
+	virtual bool CheckMaterialUsage(const EMaterialUsage Usage, const bool bSkipPrim = false) override;
+	virtual bool CheckMaterialUsage_Concurrent(const EMaterialUsage Usage, const bool bSkipPrim = false) const override;
+	virtual bool GetStaticSwitchParameterValue(FName ParameterName,bool &OutValue,FGuid &OutExpressionGuid) override;
+	virtual bool GetStaticComponentMaskParameterValue(FName ParameterName, bool &R, bool &G, bool &B, bool &A,FGuid &OutExpressionGuid) override;
+	virtual bool GetTerrainLayerWeightParameterValue(FName ParameterName, int32& OutWeightmapIndex, FGuid &OutExpressionGuid) override;
+	virtual bool IsDependent(UMaterialInterface* TestDependency) override;
+	virtual FMaterialRenderProxy* GetRenderProxy(bool Selected, bool bHovered=false) const override;
+	virtual UPhysicalMaterial* GetPhysicalMaterial() const override;
+	virtual bool UpdateLightmassTextureTracking() override;
+	virtual bool GetCastShadowAsMasked() const override;
+	virtual float GetEmissiveBoost() const override;
+	virtual float GetDiffuseBoost() const override;
+	virtual float GetExportResolutionScale() const override;
+	virtual float GetDistanceFieldPenumbraScale() const override;
 	virtual bool GetTexturesInPropertyChain(EMaterialProperty InProperty, TArray<UTexture*>& OutTextures, 
-		TArray<FName>* OutTextureParamNames, class FStaticParameterSet* InStaticParameterSet) OVERRIDE;
-	virtual void RecacheUniformExpressions() const OVERRIDE;
-	virtual bool GetRefractionSettings(float& OutBiasValue) const OVERRIDE;
-	ENGINE_API virtual void ForceRecompileForRendering() OVERRIDE;
+		TArray<FName>* OutTextureParamNames, class FStaticParameterSet* InStaticParameterSet) override;
+	virtual void RecacheUniformExpressions() const override;
+	virtual bool GetRefractionSettings(float& OutBiasValue) const override;
+	ENGINE_API virtual void ForceRecompileForRendering() override;
 	
 	ENGINE_API virtual float GetOpacityMaskClipValue_Internal() const;
 	ENGINE_API virtual EBlendMode GetBlendMode_Internal() const;
@@ -260,19 +260,19 @@ public:
 	// End UMaterialInterface interface.
 
 	// Begin UObject interface.
-	virtual SIZE_T GetResourceSize(EResourceSizeMode::Type Mode) OVERRIDE;
-	virtual void PostInitProperties() OVERRIDE;	
-	virtual void BeginCacheForCookedPlatformData(  const ITargetPlatform *TargetPlatform ) OVERRIDE;
-	virtual void ClearCachedCookedPlatformData( const ITargetPlatform *TargetPlatform ) OVERRIDE;
-	virtual void ClearAllCachedCookedPlatformData() OVERRIDE;
-	virtual void Serialize(FArchive& Ar) OVERRIDE;
-	virtual void PostLoad() OVERRIDE;
-	virtual void BeginDestroy() OVERRIDE;
-	virtual bool IsReadyForFinishDestroy() OVERRIDE;
-	virtual void FinishDestroy() OVERRIDE;
+	virtual SIZE_T GetResourceSize(EResourceSizeMode::Type Mode) override;
+	virtual void PostInitProperties() override;	
+	virtual void BeginCacheForCookedPlatformData(  const ITargetPlatform *TargetPlatform ) override;
+	virtual void ClearCachedCookedPlatformData( const ITargetPlatform *TargetPlatform ) override;
+	virtual void ClearAllCachedCookedPlatformData() override;
+	virtual void Serialize(FArchive& Ar) override;
+	virtual void PostLoad() override;
+	virtual void BeginDestroy() override;
+	virtual bool IsReadyForFinishDestroy() override;
+	virtual void FinishDestroy() override;
 	ENGINE_API static void AddReferencedObjects(UObject* InThis, FReferenceCollector& Collector);
 #if WITH_EDITOR
-	virtual void PostEditChangeProperty(FPropertyChangedEvent& PropertyChangedEvent) OVERRIDE;
+	virtual void PostEditChangeProperty(FPropertyChangedEvent& PropertyChangedEvent) override;
 
 	/**
 	 * Sets new static parameter overrides on the instance and recompiles the static permutation resources if needed (can be forced with bForceRecompile).
diff --git a/Engine/Source/Runtime/Engine/Classes/Materials/MaterialInstanceActor.h b/Engine/Source/Runtime/Engine/Classes/Materials/MaterialInstanceActor.h
--- a/Engine/Source/Runtime/Engine/Classes/Materials/MaterialInstanceActor.h
+++ b/Engine/Source/Runtime/Engine/Classes/Materials/MaterialInstanceActor.h
@@ -23,9 +23,9 @@ class AMaterialInstanceActor : public AActor
 #endif
 
 	// Begin UObject Interface
-	virtual void PostLoad() OVERRIDE;
+	virtual void PostLoad() override;
 #if WITH_EDITOR
-	virtual void PostEditChangeProperty(FPropertyChangedEvent& PropertyChangedEvent) OVERRIDE;
+	virtual void PostEditChangeProperty(FPropertyChangedEvent& PropertyChangedEvent) override;
 #endif // WITH_EDITOR
 	// End UObject Interface
 };
diff --git a/Engine/Source/Runtime/Engine/Classes/Materials/MaterialInterface.h b/Engine/Source/Runtime/Engine/Classes/Materials/MaterialInterface.h
--- a/Engine/Source/Runtime/Engine/Classes/Materials/MaterialInterface.h
+++ b/Engine/Source/Runtime/Engine/Classes/Materials/MaterialInterface.h
@@ -137,12 +137,12 @@ public:
 	ENGINE_API void SetFeatureLevelToCompile(ERHIFeatureLevel::Type FeatureLevel, bool bShouldCompile);
 
 	// Begin UObject interface.
-	ENGINE_API virtual void BeginDestroy() OVERRIDE;
-	ENGINE_API virtual bool IsReadyForFinishDestroy() OVERRIDE;
-	ENGINE_API virtual void PostLoad() OVERRIDE;
-	ENGINE_API virtual void PostDuplicate(bool bDuplicateForPIE) OVERRIDE;
+	ENGINE_API virtual void BeginDestroy() override;
+	ENGINE_API virtual bool IsReadyForFinishDestroy() override;
+	ENGINE_API virtual void PostLoad() override;
+	ENGINE_API virtual void PostDuplicate(bool bDuplicateForPIE) override;
 #if WITH_EDITOR
-	ENGINE_API virtual void PostEditChangeProperty(FPropertyChangedEvent& PropertyChangedEvent) OVERRIDE;
+	ENGINE_API virtual void PostEditChangeProperty(FPropertyChangedEvent& PropertyChangedEvent) override;
 #endif // WITH_EDITOR
 	// End UObject interface.
 
diff --git a/Engine/Source/Runtime/Engine/Classes/Materials/MaterialParameterCollection.h b/Engine/Source/Runtime/Engine/Classes/Materials/MaterialParameterCollection.h
--- a/Engine/Source/Runtime/Engine/Classes/Materials/MaterialParameterCollection.h
+++ b/Engine/Source/Runtime/Engine/Classes/Materials/MaterialParameterCollection.h
@@ -68,11 +68,11 @@ class UMaterialParameterCollection : public UObject
 
 	// Begin UObject interface.
 #if WITH_EDITOR
-	virtual void PreEditChange(class FEditPropertyChain& PropertyAboutToChange) OVERRIDE;
-	virtual void PostEditChangeProperty(FPropertyChangedEvent& PropertyChangedEvent) OVERRIDE;
-	virtual void PostEditChangeChainProperty(struct FPropertyChangedChainEvent& PropertyChangedEvent) OVERRIDE;
+	virtual void PreEditChange(class FEditPropertyChain& PropertyAboutToChange) override;
+	virtual void PostEditChangeProperty(FPropertyChangedEvent& PropertyChangedEvent) override;
+	virtual void PostEditChangeChainProperty(struct FPropertyChangedChainEvent& PropertyChangedEvent) override;
 #endif // WITH_EDITOR
-	virtual void PostLoad() OVERRIDE;
+	virtual void PostLoad() override;
 	// End UObject interface.
 
 	/** Finds a parameter name given an Id, returns NAME_None if the parameter was not found. */
diff --git a/Engine/Source/Runtime/Engine/Classes/Materials/MaterialParameterCollectionInstance.h b/Engine/Source/Runtime/Engine/Classes/Materials/MaterialParameterCollectionInstance.h
--- a/Engine/Source/Runtime/Engine/Classes/Materials/MaterialParameterCollectionInstance.h
+++ b/Engine/Source/Runtime/Engine/Classes/Materials/MaterialParameterCollectionInstance.h
@@ -18,7 +18,7 @@ class UMaterialParameterCollectionInstance : public UObject
 
 	// Begin UObject interface.
 	ENGINE_API virtual void PostInitProperties();	
-	ENGINE_API virtual void FinishDestroy() OVERRIDE;
+	ENGINE_API virtual void FinishDestroy() override;
 	// End UObject interface.
 
 	/** Initializes the instance with the collection it is based off of and the world it is owned by. */
diff --git a/Engine/Source/Runtime/Engine/Classes/Matinee/InterpData.h b/Engine/Source/Runtime/Engine/Classes/Matinee/InterpData.h
--- a/Engine/Source/Runtime/Engine/Classes/Matinee/InterpData.h
+++ b/Engine/Source/Runtime/Engine/Classes/Matinee/InterpData.h
@@ -96,8 +96,8 @@ class UInterpData : public UObject
 
 
 	// Begin UObject interface
-	virtual void PostLoad(void) OVERRIDE;
-	virtual void PostInitProperties() OVERRIDE;
+	virtual void PostLoad(void) override;
+	virtual void PostInitProperties() override;
 	// End UObject interface
 
 	/** Search through all InterpGroups in this InterpData to find a group whose GroupName matches the given name. Returns INDEX_NONE if group not found. */
diff --git a/Engine/Source/Runtime/Engine/Classes/Matinee/InterpFilter_Classes.h b/Engine/Source/Runtime/Engine/Classes/Matinee/InterpFilter_Classes.h
--- a/Engine/Source/Runtime/Engine/Classes/Matinee/InterpFilter_Classes.h
+++ b/Engine/Source/Runtime/Engine/Classes/Matinee/InterpFilter_Classes.h
@@ -28,7 +28,7 @@ class UInterpFilter_Classes : public UInterpFilter
 #endif // WITH_EDITORONLY_DATA
 
 	// Begin UInterpFilter Interface
-	virtual void FilterData(class AMatineeActor* InMatineeActor) OVERRIDE;
+	virtual void FilterData(class AMatineeActor* InMatineeActor) override;
 	// End UInterpFilter Interface
 };
 
diff --git a/Engine/Source/Runtime/Engine/Classes/Matinee/InterpFilter_Custom.h b/Engine/Source/Runtime/Engine/Classes/Matinee/InterpFilter_Custom.h
--- a/Engine/Source/Runtime/Engine/Classes/Matinee/InterpFilter_Custom.h
+++ b/Engine/Source/Runtime/Engine/Classes/Matinee/InterpFilter_Custom.h
@@ -23,7 +23,7 @@ class UInterpFilter_Custom : public UInterpFilter
 #endif // WITH_EDITORONLY_DATA
 
 	// Begin UInterpFilter Interface
-	virtual void FilterData(class AMatineeActor* InMatineeActor) OVERRIDE;
+	virtual void FilterData(class AMatineeActor* InMatineeActor) override;
 	// End UInterpFilter Interface
 };
 
diff --git a/Engine/Source/Runtime/Engine/Classes/Matinee/InterpGroup.h b/Engine/Source/Runtime/Engine/Classes/Matinee/InterpGroup.h
--- a/Engine/Source/Runtime/Engine/Classes/Matinee/InterpGroup.h
+++ b/Engine/Source/Runtime/Engine/Classes/Matinee/InterpGroup.h
@@ -116,9 +116,9 @@ private:
 
 public:
 	// Begin UObject interface.
-	virtual void PostLoad() OVERRIDE;
+	virtual void PostLoad() override;
 #if WITH_EDITOR
-	virtual void PostEditChangeProperty(FPropertyChangedEvent& PropertyChangedEvent) OVERRIDE;
+	virtual void PostEditChangeProperty(FPropertyChangedEvent& PropertyChangedEvent) override;
 #endif // WITH_EDITOR
 	// End UObject interface.
 
diff --git a/Engine/Source/Runtime/Engine/Classes/Matinee/InterpGroupDirector.h b/Engine/Source/Runtime/Engine/Classes/Matinee/InterpGroupDirector.h
--- a/Engine/Source/Runtime/Engine/Classes/Matinee/InterpGroupDirector.h
+++ b/Engine/Source/Runtime/Engine/Classes/Matinee/InterpGroupDirector.h
@@ -15,8 +15,8 @@ class UInterpGroupDirector : public UInterpGroup
 	GENERATED_UCLASS_BODY()
 
 	// Begin UInterpGroup Interface
-	virtual AActor* SelectGroupActor( class UInterpGroupInst* GrInst, bool bDeselectActors ) OVERRIDE;
-	virtual AActor* DeselectGroupActor( class UInterpGroupInst* GrInst ) OVERRIDE;
+	virtual AActor* SelectGroupActor( class UInterpGroupInst* GrInst, bool bDeselectActors ) override;
+	virtual AActor* DeselectGroupActor( class UInterpGroupInst* GrInst ) override;
 	// End UInterpGroup Interface
 
 	/** @return the director track inside this Director group - if present. */
diff --git a/Engine/Source/Runtime/Engine/Classes/Matinee/InterpTrack.h b/Engine/Source/Runtime/Engine/Classes/Matinee/InterpTrack.h
--- a/Engine/Source/Runtime/Engine/Classes/Matinee/InterpTrack.h
+++ b/Engine/Source/Runtime/Engine/Classes/Matinee/InterpTrack.h
@@ -164,7 +164,7 @@ public:
 
 public:
 	// Begin FInterpEdInputInterface Interface
-	virtual UObject* GetUObject() OVERRIDE {return this;}
+	virtual UObject* GetUObject() override {return this;}
 	// End FInterpEdInputInterface Interface
 
 	/** @return  The total number of keyframes currently in this track. */
diff --git a/Engine/Source/Runtime/Engine/Classes/Matinee/InterpTrackAnimControl.h b/Engine/Source/Runtime/Engine/Classes/Matinee/InterpTrackAnimControl.h
--- a/Engine/Source/Runtime/Engine/Classes/Matinee/InterpTrackAnimControl.h
+++ b/Engine/Source/Runtime/Engine/Classes/Matinee/InterpTrackAnimControl.h
@@ -79,30 +79,30 @@ class UInterpTrackAnimControl : public UInterpTrackFloatBase
 	// End UObject interface.
 
 	// Begin UInterpTrack interface.
-	virtual int32 GetNumKeyframes() const OVERRIDE;
-	virtual void GetTimeRange(float& StartTime, float& EndTime) const OVERRIDE;
-	virtual float GetTrackEndTime() const OVERRIDE;
-	virtual float GetKeyframeTime(int32 KeyIndex) const OVERRIDE;
-	virtual int32 GetKeyframeIndex( float KeyTime ) const OVERRIDE;
-	virtual int32 AddKeyframe(float Time, class UInterpTrackInst* TrInst, EInterpCurveMode InitInterpMode) OVERRIDE;
-	virtual int32 SetKeyframeTime(int32 KeyIndex, float NewKeyTime, bool bUpdateOrder=true) OVERRIDE;
-	virtual void RemoveKeyframe(int32 KeyIndex) OVERRIDE;
-	virtual int32 DuplicateKeyframe(int32 KeyIndex, float NewKeyTime, UInterpTrack* ToTrack = NULL) OVERRIDE;
-	virtual bool GetClosestSnapPosition(float InPosition, TArray<int32> &IgnoreKeys, float& OutPosition) OVERRIDE;
-	virtual FColor GetKeyframeColor(int32 KeyIndex) const OVERRIDE;
-	virtual void PreviewUpdateTrack(float NewPosition, class UInterpTrackInst* TrInst) OVERRIDE;
-	virtual void UpdateTrack(float NewPosition, UInterpTrackInst* TrInst, bool bJump) OVERRIDE;
-	virtual const FString	GetEdHelperClassName() const OVERRIDE;
-	virtual const FString	GetSlateHelperClassName() const OVERRIDE;
-	virtual class UTexture2D* GetTrackIcon() const OVERRIDE;
-	virtual void DrawTrack( FCanvas* Canvas, UInterpGroup* Group, const FInterpTrackDrawParams& Params ) OVERRIDE;
+	virtual int32 GetNumKeyframes() const override;
+	virtual void GetTimeRange(float& StartTime, float& EndTime) const override;
+	virtual float GetTrackEndTime() const override;
+	virtual float GetKeyframeTime(int32 KeyIndex) const override;
+	virtual int32 GetKeyframeIndex( float KeyTime ) const override;
+	virtual int32 AddKeyframe(float Time, class UInterpTrackInst* TrInst, EInterpCurveMode InitInterpMode) override;
+	virtual int32 SetKeyframeTime(int32 KeyIndex, float NewKeyTime, bool bUpdateOrder=true) override;
+	virtual void RemoveKeyframe(int32 KeyIndex) override;
+	virtual int32 DuplicateKeyframe(int32 KeyIndex, float NewKeyTime, UInterpTrack* ToTrack = NULL) override;
+	virtual bool GetClosestSnapPosition(float InPosition, TArray<int32> &IgnoreKeys, float& OutPosition) override;
+	virtual FColor GetKeyframeColor(int32 KeyIndex) const override;
+	virtual void PreviewUpdateTrack(float NewPosition, class UInterpTrackInst* TrInst) override;
+	virtual void UpdateTrack(float NewPosition, UInterpTrackInst* TrInst, bool bJump) override;
+	virtual const FString	GetEdHelperClassName() const override;
+	virtual const FString	GetSlateHelperClassName() const override;
+	virtual class UTexture2D* GetTrackIcon() const override;
+	virtual void DrawTrack( FCanvas* Canvas, UInterpGroup* Group, const FInterpTrackDrawParams& Params ) override;
 	// End UInterpTrack interface.
 
 	// Begin FInterpEdInputInterface Interface
-	virtual void BeginDrag(FInterpEdInputData &InputData) OVERRIDE;
-	virtual void EndDrag(FInterpEdInputData &InputData) OVERRIDE;
-	virtual EMouseCursor::Type GetMouseCursor(FInterpEdInputData &InputData) OVERRIDE;
-	virtual void ObjectDragged(FInterpEdInputData& InputData) OVERRIDE;
+	virtual void BeginDrag(FInterpEdInputData &InputData) override;
+	virtual void EndDrag(FInterpEdInputData &InputData) override;
+	virtual EMouseCursor::Type GetMouseCursor(FInterpEdInputData &InputData) override;
+	virtual void ObjectDragged(FInterpEdInputData& InputData) override;
 	// End FInterpEdInputInterface Interface
 
 	/** 
diff --git a/Engine/Source/Runtime/Engine/Classes/Matinee/InterpTrackAudioMaster.h b/Engine/Source/Runtime/Engine/Classes/Matinee/InterpTrackAudioMaster.h
--- a/Engine/Source/Runtime/Engine/Classes/Matinee/InterpTrackAudioMaster.h
+++ b/Engine/Source/Runtime/Engine/Classes/Matinee/InterpTrackAudioMaster.h
@@ -12,12 +12,12 @@ class UInterpTrackAudioMaster : public UInterpTrackVectorBase
 
 
 	// Begin UInterpTrack interface.
-	virtual int32 AddKeyframe(float Time, UInterpTrackInst* TrInst, EInterpCurveMode InitInterpMode) OVERRIDE;
-	virtual void UpdateKeyframe(int32 KeyIndex, UInterpTrackInst* TrInst) OVERRIDE;
-	virtual void PreviewUpdateTrack(float NewPosition, UInterpTrackInst* TrInst) OVERRIDE;
-	virtual void UpdateTrack(float NewPosition, UInterpTrackInst* TrInst, bool bJump) OVERRIDE;
-	virtual void SetTrackToSensibleDefault() OVERRIDE;
-	virtual class UTexture2D* GetTrackIcon() const OVERRIDE;
+	virtual int32 AddKeyframe(float Time, UInterpTrackInst* TrInst, EInterpCurveMode InitInterpMode) override;
+	virtual void UpdateKeyframe(int32 KeyIndex, UInterpTrackInst* TrInst) override;
+	virtual void PreviewUpdateTrack(float NewPosition, UInterpTrackInst* TrInst) override;
+	virtual void UpdateTrack(float NewPosition, UInterpTrackInst* TrInst, bool bJump) override;
+	virtual void SetTrackToSensibleDefault() override;
+	virtual class UTexture2D* GetTrackIcon() const override;
 	// End UInterpTrack interface.
 
 	/** Return the sound volume scale for the specified time */
diff --git a/Engine/Source/Runtime/Engine/Classes/Matinee/InterpTrackBoolProp.h b/Engine/Source/Runtime/Engine/Classes/Matinee/InterpTrackBoolProp.h
--- a/Engine/Source/Runtime/Engine/Classes/Matinee/InterpTrackBoolProp.h
+++ b/Engine/Source/Runtime/Engine/Classes/Matinee/InterpTrackBoolProp.h
@@ -40,24 +40,24 @@ class UInterpTrackBoolProp : public UInterpTrack
 
 
 	// Begin UInterpTrack Interface
-	virtual int32 GetNumKeyframes() const OVERRIDE;
-	virtual float GetTrackEndTime() const OVERRIDE;
-	virtual float GetKeyframeTime( int32 KeyIndex ) const OVERRIDE;
-	virtual int32 GetKeyframeIndex( float KeyTime ) const OVERRIDE;
-	virtual void GetTimeRange( float& StartTime, float& EndTime ) const OVERRIDE;
-	virtual int32 SetKeyframeTime( int32 KeyIndex, float NewKeyTime, bool bUpdateOrder = true ) OVERRIDE;
-	virtual void RemoveKeyframe( int32 KeyIndex ) OVERRIDE;
-	virtual int32 DuplicateKeyframe( int32 KeyIndex, float NewKeyTime, UInterpTrack* ToTrack = NULL ) OVERRIDE;
-	virtual bool GetClosestSnapPosition( float InPosition, TArray<int32>& IgnoreKeys, float& OutPosition ) OVERRIDE;
-	virtual int32 AddKeyframe( float Time, UInterpTrackInst* TrackInst, EInterpCurveMode InitInterpMode ) OVERRIDE;
-	virtual bool CanAddKeyframe( UInterpTrackInst* TrackInst ) OVERRIDE;
-	virtual void UpdateKeyframe( int32 KeyIndex, UInterpTrackInst* TrackInst ) OVERRIDE;
-	virtual void PreviewUpdateTrack( float NewPosition, UInterpTrackInst* TrackInst ) OVERRIDE;
-	virtual void UpdateTrack( float NewPosition, UInterpTrackInst* TrackInst, bool bJump ) OVERRIDE;
-	virtual bool AllowStaticActors() OVERRIDE { return true; }
-	virtual const FString GetEdHelperClassName() const OVERRIDE;
-	virtual const FString GetSlateHelperClassName() const OVERRIDE;
-	virtual class UTexture2D* GetTrackIcon() const OVERRIDE;
+	virtual int32 GetNumKeyframes() const override;
+	virtual float GetTrackEndTime() const override;
+	virtual float GetKeyframeTime( int32 KeyIndex ) const override;
+	virtual int32 GetKeyframeIndex( float KeyTime ) const override;
+	virtual void GetTimeRange( float& StartTime, float& EndTime ) const override;
+	virtual int32 SetKeyframeTime( int32 KeyIndex, float NewKeyTime, bool bUpdateOrder = true ) override;
+	virtual void RemoveKeyframe( int32 KeyIndex ) override;
+	virtual int32 DuplicateKeyframe( int32 KeyIndex, float NewKeyTime, UInterpTrack* ToTrack = NULL ) override;
+	virtual bool GetClosestSnapPosition( float InPosition, TArray<int32>& IgnoreKeys, float& OutPosition ) override;
+	virtual int32 AddKeyframe( float Time, UInterpTrackInst* TrackInst, EInterpCurveMode InitInterpMode ) override;
+	virtual bool CanAddKeyframe( UInterpTrackInst* TrackInst ) override;
+	virtual void UpdateKeyframe( int32 KeyIndex, UInterpTrackInst* TrackInst ) override;
+	virtual void PreviewUpdateTrack( float NewPosition, UInterpTrackInst* TrackInst ) override;
+	virtual void UpdateTrack( float NewPosition, UInterpTrackInst* TrackInst, bool bJump ) override;
+	virtual bool AllowStaticActors() override { return true; }
+	virtual const FString GetEdHelperClassName() const override;
+	virtual const FString GetSlateHelperClassName() const override;
+	virtual class UTexture2D* GetTrackIcon() const override;
 	// End UInterpTrack Interface
 };
 
diff --git a/Engine/Source/Runtime/Engine/Classes/Matinee/InterpTrackColorProp.h b/Engine/Source/Runtime/Engine/Classes/Matinee/InterpTrackColorProp.h
--- a/Engine/Source/Runtime/Engine/Classes/Matinee/InterpTrackColorProp.h
+++ b/Engine/Source/Runtime/Engine/Classes/Matinee/InterpTrackColorProp.h
@@ -16,14 +16,14 @@ class UInterpTrackColorProp : public UInterpTrackVectorBase
 
 
 	// Begin UInterpTrack interface
-	virtual int32 AddKeyframe(float Time, UInterpTrackInst* TrInst, EInterpCurveMode InitInterpMode) OVERRIDE;
-	virtual bool CanAddKeyframe( UInterpTrackInst* TrackInst ) OVERRIDE;
-	virtual void UpdateKeyframe(int32 KeyIndex, UInterpTrackInst* TrInst) OVERRIDE;
-	virtual void PreviewUpdateTrack(float NewPosition, UInterpTrackInst* TrInst) OVERRIDE;
-	virtual void UpdateTrack(float NewPosition, UInterpTrackInst* TrInst, bool bJump) OVERRIDE;
-	virtual const FString	GetEdHelperClassName() const OVERRIDE;
-	virtual const FString	GetSlateHelperClassName() const OVERRIDE;
-	virtual class UTexture2D* GetTrackIcon() const OVERRIDE;
+	virtual int32 AddKeyframe(float Time, UInterpTrackInst* TrInst, EInterpCurveMode InitInterpMode) override;
+	virtual bool CanAddKeyframe( UInterpTrackInst* TrackInst ) override;
+	virtual void UpdateKeyframe(int32 KeyIndex, UInterpTrackInst* TrInst) override;
+	virtual void PreviewUpdateTrack(float NewPosition, UInterpTrackInst* TrInst) override;
+	virtual void UpdateTrack(float NewPosition, UInterpTrackInst* TrInst, bool bJump) override;
+	virtual const FString	GetEdHelperClassName() const override;
+	virtual const FString	GetSlateHelperClassName() const override;
+	virtual class UTexture2D* GetTrackIcon() const override;
 	// End UInterpTrack interface
 };
 
diff --git a/Engine/Source/Runtime/Engine/Classes/Matinee/InterpTrackColorScale.h b/Engine/Source/Runtime/Engine/Classes/Matinee/InterpTrackColorScale.h
--- a/Engine/Source/Runtime/Engine/Classes/Matinee/InterpTrackColorScale.h
+++ b/Engine/Source/Runtime/Engine/Classes/Matinee/InterpTrackColorScale.h
@@ -16,12 +16,12 @@ class UInterpTrackColorScale : public UInterpTrackVectorBase
 
 
 	// Begin UInterpTrack interface.
-	virtual int32 AddKeyframe(float Time, UInterpTrackInst* TrInst, EInterpCurveMode InitInterpMode) OVERRIDE;
-	virtual void UpdateKeyframe(int32 KeyIndex, UInterpTrackInst* TrInst) OVERRIDE;
-	virtual void PreviewUpdateTrack(float NewPosition, UInterpTrackInst* TrInst) OVERRIDE;
-	virtual void UpdateTrack(float NewPosition, UInterpTrackInst* TrInst, bool bJump) OVERRIDE;
-	virtual void SetTrackToSensibleDefault() OVERRIDE;
-	virtual class UTexture2D* GetTrackIcon() const OVERRIDE;
+	virtual int32 AddKeyframe(float Time, UInterpTrackInst* TrInst, EInterpCurveMode InitInterpMode) override;
+	virtual void UpdateKeyframe(int32 KeyIndex, UInterpTrackInst* TrInst) override;
+	virtual void PreviewUpdateTrack(float NewPosition, UInterpTrackInst* TrInst) override;
+	virtual void UpdateTrack(float NewPosition, UInterpTrackInst* TrInst, bool bJump) override;
+	virtual void SetTrackToSensibleDefault() override;
+	virtual class UTexture2D* GetTrackIcon() const override;
 	// Begin UInterpTrack interface.
 
 	/** Return the blur alpha we want at the given time. */
diff --git a/Engine/Source/Runtime/Engine/Classes/Matinee/InterpTrackDirector.h b/Engine/Source/Runtime/Engine/Classes/Matinee/InterpTrackDirector.h
--- a/Engine/Source/Runtime/Engine/Classes/Matinee/InterpTrackDirector.h
+++ b/Engine/Source/Runtime/Engine/Classes/Matinee/InterpTrackDirector.h
@@ -68,22 +68,22 @@ public:
 	// End UObject interface
 	
 	// Begin UInterpTrack Interface
-	virtual int32 GetNumKeyframes() const OVERRIDE;
-	virtual void GetTimeRange( float& StartTime, float& EndTime ) const OVERRIDE;
-	virtual float GetTrackEndTime() const OVERRIDE;
-	virtual float GetKeyframeTime( int32 KeyIndex ) const OVERRIDE;
-	virtual int32 GetKeyframeIndex( float KeyTime ) const OVERRIDE;
-	virtual int32 AddKeyframe( float Time, UInterpTrackInst* TrackInst, EInterpCurveMode InitInterpMode ) OVERRIDE;
-	virtual int32 SetKeyframeTime( int32 KeyIndex, float NewKeyTime, bool bUpdateOrder = true ) OVERRIDE;
-	virtual void RemoveKeyframe( int32 KeyIndex ) OVERRIDE;
-	virtual int32 DuplicateKeyframe(int32 KeyIndex, float NewKeyTime, UInterpTrack* ToTrack = NULL) OVERRIDE;
-	virtual bool GetClosestSnapPosition( float InPosition, TArray<int32>& IgnoreKeys, float& OutPosition ) OVERRIDE;
-	virtual void PreviewUpdateTrack(float NewPosition, class UInterpTrackInst* TrackInst) OVERRIDE;
-	virtual void UpdateTrack( float NewPosition, UInterpTrackInst* TrackInst, bool bJump ) OVERRIDE;
-	virtual const FString GetEdHelperClassName() const OVERRIDE;
-	virtual const FString GetSlateHelperClassName() const OVERRIDE;
-	virtual class UTexture2D* GetTrackIcon() const OVERRIDE;
-	virtual void DrawTrack( FCanvas* Canvas, UInterpGroup* Group, const FInterpTrackDrawParams& Params ) OVERRIDE;
+	virtual int32 GetNumKeyframes() const override;
+	virtual void GetTimeRange( float& StartTime, float& EndTime ) const override;
+	virtual float GetTrackEndTime() const override;
+	virtual float GetKeyframeTime( int32 KeyIndex ) const override;
+	virtual int32 GetKeyframeIndex( float KeyTime ) const override;
+	virtual int32 AddKeyframe( float Time, UInterpTrackInst* TrackInst, EInterpCurveMode InitInterpMode ) override;
+	virtual int32 SetKeyframeTime( int32 KeyIndex, float NewKeyTime, bool bUpdateOrder = true ) override;
+	virtual void RemoveKeyframe( int32 KeyIndex ) override;
+	virtual int32 DuplicateKeyframe(int32 KeyIndex, float NewKeyTime, UInterpTrack* ToTrack = NULL) override;
+	virtual bool GetClosestSnapPosition( float InPosition, TArray<int32>& IgnoreKeys, float& OutPosition ) override;
+	virtual void PreviewUpdateTrack(float NewPosition, class UInterpTrackInst* TrackInst) override;
+	virtual void UpdateTrack( float NewPosition, UInterpTrackInst* TrackInst, bool bJump ) override;
+	virtual const FString GetEdHelperClassName() const override;
+	virtual const FString GetSlateHelperClassName() const override;
+	virtual class UTexture2D* GetTrackIcon() const override;
+	virtual void DrawTrack( FCanvas* Canvas, UInterpGroup* Group, const FInterpTrackDrawParams& Params ) override;
 	// End UInterpTrack Interface
 
 	/** Get the keyframe index nearest the time, without going over */
diff --git a/Engine/Source/Runtime/Engine/Classes/Matinee/InterpTrackEvent.h b/Engine/Source/Runtime/Engine/Classes/Matinee/InterpTrackEvent.h
--- a/Engine/Source/Runtime/Engine/Classes/Matinee/InterpTrackEvent.h
+++ b/Engine/Source/Runtime/Engine/Classes/Matinee/InterpTrackEvent.h
@@ -53,22 +53,22 @@ class UInterpTrackEvent : public UInterpTrack
 
 
 	// Begin UInterpTrack Interface
-	virtual int32 GetNumKeyframes() const OVERRIDE;
-	virtual void GetTimeRange( float& StartTime, float& EndTime ) const OVERRIDE;
-	virtual float GetTrackEndTime() const OVERRIDE;
-	virtual float GetKeyframeTime( int32 KeyIndex ) const OVERRIDE;
-	virtual int32 GetKeyframeIndex( float KeyTime ) const OVERRIDE;
-	virtual int32 AddKeyframe( float Time, UInterpTrackInst* TrackInst, EInterpCurveMode InitInterpMode ) OVERRIDE;
-	virtual int32 SetKeyframeTime( int32 KeyIndex, float NewKeyTime, bool bUpdateOrder = true ) OVERRIDE;
-	virtual void RemoveKeyframe( int32 KeyIndex ) OVERRIDE;
-	virtual int32 DuplicateKeyframe(int32 KeyIndex, float NewKeyTime, UInterpTrack* ToTrack = NULL) OVERRIDE;
-	virtual bool GetClosestSnapPosition( float InPosition, TArray<int32>& IgnoreKeys, float& OutPosition ) OVERRIDE;
-	virtual void UpdateTrack( float NewPosition, UInterpTrackInst* TrackInst, bool bJump ) OVERRIDE;
-	virtual const FString GetEdHelperClassName() const OVERRIDE;
-	virtual const FString GetSlateHelperClassName() const OVERRIDE;
-	virtual class UTexture2D* GetTrackIcon() const OVERRIDE;
-	virtual bool AllowStaticActors() OVERRIDE { return true; }
-	virtual void DrawTrack( FCanvas* Canvas, UInterpGroup* Group, const FInterpTrackDrawParams& Params ) OVERRIDE;
+	virtual int32 GetNumKeyframes() const override;
+	virtual void GetTimeRange( float& StartTime, float& EndTime ) const override;
+	virtual float GetTrackEndTime() const override;
+	virtual float GetKeyframeTime( int32 KeyIndex ) const override;
+	virtual int32 GetKeyframeIndex( float KeyTime ) const override;
+	virtual int32 AddKeyframe( float Time, UInterpTrackInst* TrackInst, EInterpCurveMode InitInterpMode ) override;
+	virtual int32 SetKeyframeTime( int32 KeyIndex, float NewKeyTime, bool bUpdateOrder = true ) override;
+	virtual void RemoveKeyframe( int32 KeyIndex ) override;
+	virtual int32 DuplicateKeyframe(int32 KeyIndex, float NewKeyTime, UInterpTrack* ToTrack = NULL) override;
+	virtual bool GetClosestSnapPosition( float InPosition, TArray<int32>& IgnoreKeys, float& OutPosition ) override;
+	virtual void UpdateTrack( float NewPosition, UInterpTrackInst* TrackInst, bool bJump ) override;
+	virtual const FString GetEdHelperClassName() const override;
+	virtual const FString GetSlateHelperClassName() const override;
+	virtual class UTexture2D* GetTrackIcon() const override;
+	virtual bool AllowStaticActors() override { return true; }
+	virtual void DrawTrack( FCanvas* Canvas, UInterpGroup* Group, const FInterpTrackDrawParams& Params ) override;
 	// End UInterpTrack Interface
 };
 
diff --git a/Engine/Source/Runtime/Engine/Classes/Matinee/InterpTrackFade.h b/Engine/Source/Runtime/Engine/Classes/Matinee/InterpTrackFade.h
--- a/Engine/Source/Runtime/Engine/Classes/Matinee/InterpTrackFade.h
+++ b/Engine/Source/Runtime/Engine/Classes/Matinee/InterpTrackFade.h
@@ -21,11 +21,11 @@ class UInterpTrackFade : public UInterpTrackFloatBase
 
 
 	// Begin UInterpTrack interface.
-	virtual int32 AddKeyframe(float Time, UInterpTrackInst* TrInst, EInterpCurveMode InitInterpMode) OVERRIDE;
-	virtual void UpdateKeyframe(int32 KeyIndex, UInterpTrackInst* TrInst) OVERRIDE;
-	virtual void PreviewUpdateTrack(float NewPosition, UInterpTrackInst* TrInst) OVERRIDE;
-	virtual void UpdateTrack(float NewPosition, UInterpTrackInst* TrInst, bool bJump) OVERRIDE;
-	virtual class UTexture2D* GetTrackIcon() const OVERRIDE;
+	virtual int32 AddKeyframe(float Time, UInterpTrackInst* TrInst, EInterpCurveMode InitInterpMode) override;
+	virtual void UpdateKeyframe(int32 KeyIndex, UInterpTrackInst* TrInst) override;
+	virtual void PreviewUpdateTrack(float NewPosition, UInterpTrackInst* TrInst) override;
+	virtual void UpdateTrack(float NewPosition, UInterpTrackInst* TrInst, bool bJump) override;
+	virtual class UTexture2D* GetTrackIcon() const override;
 	// End UInterpTrack interface.
 
 	/** @return the amount of fading we want at the given time. */
