UE4 Refactoring. Changed OVERRIDE and FINAL macros to keywords override and final.

diff --git a/Engine/Source/Runtime/Engine/Classes/Matinee/InterpTrackFloatBase.h b/Engine/Source/Runtime/Engine/Classes/Matinee/InterpTrackFloatBase.h
--- a/Engine/Source/Runtime/Engine/Classes/Matinee/InterpTrackFloatBase.h
+++ b/Engine/Source/Runtime/Engine/Classes/Matinee/InterpTrackFloatBase.h
@@ -25,34 +25,34 @@ class UInterpTrackFloatBase : public UInterpTrack
 	// End UObject interface.
 
 	// Begin UInterpTrack interface.
-	virtual int32 GetNumKeyframes() const OVERRIDE;
-	virtual void GetTimeRange(float& StartTime, float& EndTime) const OVERRIDE;
-	virtual float GetTrackEndTime() const OVERRIDE;
-	virtual float GetKeyframeTime(int32 KeyIndex) const OVERRIDE;
-	virtual int32 GetKeyframeIndex( float KeyTime ) const OVERRIDE;
-	virtual int32 SetKeyframeTime(int32 KeyIndex, float NewKeyTime, bool bUpdateOrder=true) OVERRIDE;
-	virtual void RemoveKeyframe(int32 KeyIndex) OVERRIDE;
-	virtual int32 DuplicateKeyframe(int32 KeyIndex, float NewKeyTime, UInterpTrack* ToTrack = NULL) OVERRIDE;
-	virtual bool GetClosestSnapPosition(float InPosition, TArray<int32> &IgnoreKeys, float& OutPosition) OVERRIDE;
-	virtual FColor GetKeyframeColor(int32 KeyIndex) const OVERRIDE;
+	virtual int32 GetNumKeyframes() const override;
+	virtual void GetTimeRange(float& StartTime, float& EndTime) const override;
+	virtual float GetTrackEndTime() const override;
+	virtual float GetKeyframeTime(int32 KeyIndex) const override;
+	virtual int32 GetKeyframeIndex( float KeyTime ) const override;
+	virtual int32 SetKeyframeTime(int32 KeyIndex, float NewKeyTime, bool bUpdateOrder=true) override;
+	virtual void RemoveKeyframe(int32 KeyIndex) override;
+	virtual int32 DuplicateKeyframe(int32 KeyIndex, float NewKeyTime, UInterpTrack* ToTrack = NULL) override;
+	virtual bool GetClosestSnapPosition(float InPosition, TArray<int32> &IgnoreKeys, float& OutPosition) override;
+	virtual FColor GetKeyframeColor(int32 KeyIndex) const override;
 	// End UInterpTrack interface.
 
 	// Begin FCurveEdInterface interface.
-	virtual int32		GetNumKeys() const OVERRIDE;
-	virtual int32		GetNumSubCurves() const OVERRIDE;
-	virtual float	GetKeyIn(int32 KeyIndex) OVERRIDE;
-	virtual float	GetKeyOut(int32 SubIndex, int32 KeyIndex) OVERRIDE;
-	virtual void	GetInRange(float& MinIn, float& MaxIn) const OVERRIDE;
-	virtual void	GetOutRange(float& MinOut, float& MaxOut) const OVERRIDE;
-	virtual EInterpCurveMode	GetKeyInterpMode(int32 KeyIndex) const OVERRIDE;
-	virtual void	GetTangents(int32 SubIndex, int32 KeyIndex, float& ArriveTangent, float& LeaveTangent) const OVERRIDE;
-	virtual float	EvalSub(int32 SubIndex, float InVal) OVERRIDE;
-	virtual int32		CreateNewKey(float KeyIn) OVERRIDE;
-	virtual void	DeleteKey(int32 KeyIndex) OVERRIDE;
-	virtual int32		SetKeyIn(int32 KeyIndex, float NewInVal) OVERRIDE;
-	virtual void	SetKeyOut(int32 SubIndex, int32 KeyIndex, float NewOutVal) OVERRIDE;
-	virtual void	SetKeyInterpMode(int32 KeyIndex, EInterpCurveMode NewMode) OVERRIDE;
-	virtual void	SetTangents(int32 SubIndex, int32 KeyIndex, float ArriveTangent, float LeaveTangent) OVERRIDE;
+	virtual int32		GetNumKeys() const override;
+	virtual int32		GetNumSubCurves() const override;
+	virtual float	GetKeyIn(int32 KeyIndex) override;
+	virtual float	GetKeyOut(int32 SubIndex, int32 KeyIndex) override;
+	virtual void	GetInRange(float& MinIn, float& MaxIn) const override;
+	virtual void	GetOutRange(float& MinOut, float& MaxOut) const override;
+	virtual EInterpCurveMode	GetKeyInterpMode(int32 KeyIndex) const override;
+	virtual void	GetTangents(int32 SubIndex, int32 KeyIndex, float& ArriveTangent, float& LeaveTangent) const override;
+	virtual float	EvalSub(int32 SubIndex, float InVal) override;
+	virtual int32		CreateNewKey(float KeyIn) override;
+	virtual void	DeleteKey(int32 KeyIndex) override;
+	virtual int32		SetKeyIn(int32 KeyIndex, float NewInVal) override;
+	virtual void	SetKeyOut(int32 SubIndex, int32 KeyIndex, float NewOutVal) override;
+	virtual void	SetKeyInterpMode(int32 KeyIndex, EInterpCurveMode NewMode) override;
+	virtual void	SetTangents(int32 SubIndex, int32 KeyIndex, float ArriveTangent, float LeaveTangent) override;
 	// End FCurveEdInterface interface.
 };
 
diff --git a/Engine/Source/Runtime/Engine/Classes/Matinee/InterpTrackFloatMaterialParam.h b/Engine/Source/Runtime/Engine/Classes/Matinee/InterpTrackFloatMaterialParam.h
--- a/Engine/Source/Runtime/Engine/Classes/Matinee/InterpTrackFloatMaterialParam.h
+++ b/Engine/Source/Runtime/Engine/Classes/Matinee/InterpTrackFloatMaterialParam.h
@@ -19,15 +19,15 @@ class UInterpTrackFloatMaterialParam : public UInterpTrackFloatBase
 
 	// Begin UObject Interface
 #if WITH_EDITOR
-	virtual void PreEditChange(UProperty* PropertyThatWillChange) OVERRIDE;
-	virtual void PostEditChangeProperty(FPropertyChangedEvent& PropertyChangedEvent) OVERRIDE;
+	virtual void PreEditChange(UProperty* PropertyThatWillChange) override;
+	virtual void PostEditChangeProperty(FPropertyChangedEvent& PropertyChangedEvent) override;
 #endif // WITH_EDITOR
 	// End UObject Interface
 
 	// Begin UInterpTrack interface.
-	virtual int32 AddKeyframe(float Time, UInterpTrackInst* TrInst, EInterpCurveMode InitInterpMode) OVERRIDE;
-	virtual void PreviewUpdateTrack(float NewPosition, UInterpTrackInst* TrInst) OVERRIDE;
-	virtual void UpdateTrack(float NewPosition, UInterpTrackInst* TrInst, bool bJump) OVERRIDE;
+	virtual int32 AddKeyframe(float Time, UInterpTrackInst* TrInst, EInterpCurveMode InitInterpMode) override;
+	virtual void PreviewUpdateTrack(float NewPosition, UInterpTrackInst* TrInst) override;
+	virtual void UpdateTrack(float NewPosition, UInterpTrackInst* TrInst, bool bJump) override;
 	// End UInterpTrack interface.
 };
 
diff --git a/Engine/Source/Runtime/Engine/Classes/Matinee/InterpTrackFloatParticleParam.h b/Engine/Source/Runtime/Engine/Classes/Matinee/InterpTrackFloatParticleParam.h
--- a/Engine/Source/Runtime/Engine/Classes/Matinee/InterpTrackFloatParticleParam.h
+++ b/Engine/Source/Runtime/Engine/Classes/Matinee/InterpTrackFloatParticleParam.h
@@ -14,9 +14,9 @@ class UInterpTrackFloatParticleParam : public UInterpTrackFloatBase
 
 
 	// Begin UInterpTrack interface.
-	virtual int32 AddKeyframe(float Time, UInterpTrackInst* TrInst, EInterpCurveMode InitInterpMode) OVERRIDE;
-	virtual void PreviewUpdateTrack(float NewPosition, UInterpTrackInst* TrInst) OVERRIDE;
-	virtual void UpdateTrack(float NewPosition, UInterpTrackInst* TrInst, bool bJump) OVERRIDE;
+	virtual int32 AddKeyframe(float Time, UInterpTrackInst* TrInst, EInterpCurveMode InitInterpMode) override;
+	virtual void PreviewUpdateTrack(float NewPosition, UInterpTrackInst* TrInst) override;
+	virtual void UpdateTrack(float NewPosition, UInterpTrackInst* TrInst, bool bJump) override;
 	// End UInterpTrack interface.
 };
 
diff --git a/Engine/Source/Runtime/Engine/Classes/Matinee/InterpTrackFloatProp.h b/Engine/Source/Runtime/Engine/Classes/Matinee/InterpTrackFloatProp.h
--- a/Engine/Source/Runtime/Engine/Classes/Matinee/InterpTrackFloatProp.h
+++ b/Engine/Source/Runtime/Engine/Classes/Matinee/InterpTrackFloatProp.h
@@ -15,15 +15,15 @@ class UInterpTrackFloatProp : public UInterpTrackFloatBase
 
 
 	// Begin InterpTrack interface.
-	virtual int32 AddKeyframe(float Time, UInterpTrackInst* TrInst, EInterpCurveMode InitInterpMode) OVERRIDE;
-	virtual bool CanAddKeyframe( UInterpTrackInst* TrackInst ) OVERRIDE;
-	virtual void UpdateKeyframe(int32 KeyIndex, UInterpTrackInst* TrInst) OVERRIDE;
-	virtual void PreviewUpdateTrack(float NewPosition, UInterpTrackInst* TrInst) OVERRIDE;
-	virtual void UpdateTrack(float NewPosition, UInterpTrackInst* TrInst, bool bJump) OVERRIDE;
-	virtual const FString	GetEdHelperClassName() const OVERRIDE;
-	virtual const FString	GetSlateHelperClassName() const OVERRIDE;
-	virtual class UTexture2D* GetTrackIcon() const OVERRIDE;
-	virtual void ReduceKeys( float IntervalStart, float IntervalEnd, float Tolerance ) OVERRIDE;
+	virtual int32 AddKeyframe(float Time, UInterpTrackInst* TrInst, EInterpCurveMode InitInterpMode) override;
+	virtual bool CanAddKeyframe( UInterpTrackInst* TrackInst ) override;
+	virtual void UpdateKeyframe(int32 KeyIndex, UInterpTrackInst* TrInst) override;
+	virtual void PreviewUpdateTrack(float NewPosition, UInterpTrackInst* TrInst) override;
+	virtual void UpdateTrack(float NewPosition, UInterpTrackInst* TrInst, bool bJump) override;
+	virtual const FString	GetEdHelperClassName() const override;
+	virtual const FString	GetSlateHelperClassName() const override;
+	virtual class UTexture2D* GetTrackIcon() const override;
+	virtual void ReduceKeys( float IntervalStart, float IntervalEnd, float Tolerance ) override;
 	// End InterpTrack interface.
 };
 
diff --git a/Engine/Source/Runtime/Engine/Classes/Matinee/InterpTrackInst.h b/Engine/Source/Runtime/Engine/Classes/Matinee/InterpTrackInst.h
--- a/Engine/Source/Runtime/Engine/Classes/Matinee/InterpTrackInst.h
+++ b/Engine/Source/Runtime/Engine/Classes/Matinee/InterpTrackInst.h
@@ -56,6 +56,6 @@ class UInterpTrackInst : public UObject
 	*
 	* returns the GroupActor's world
 	*/
-	virtual UWorld* GetWorld() const OVERRIDE;
+	virtual UWorld* GetWorld() const override;
 };
 
diff --git a/Engine/Source/Runtime/Engine/Classes/Matinee/InterpTrackInstAnimControl.h b/Engine/Source/Runtime/Engine/Classes/Matinee/InterpTrackInstAnimControl.h
--- a/Engine/Source/Runtime/Engine/Classes/Matinee/InterpTrackInstAnimControl.h
+++ b/Engine/Source/Runtime/Engine/Classes/Matinee/InterpTrackInstAnimControl.h
@@ -23,7 +23,7 @@ class UInterpTrackInstAnimControl : public UInterpTrackInst
 #endif // WITH_EDITORONLY_DATA
 
 	// Begin UInterpTrackInst Instance
-	virtual void InitTrackInst(UInterpTrack* Track) OVERRIDE;
+	virtual void InitTrackInst(UInterpTrack* Track) override;
 	// End UInterpTrackInst Instance
 };
 
diff --git a/Engine/Source/Runtime/Engine/Classes/Matinee/InterpTrackInstAudioMaster.h b/Engine/Source/Runtime/Engine/Classes/Matinee/InterpTrackInstAudioMaster.h
--- a/Engine/Source/Runtime/Engine/Classes/Matinee/InterpTrackInstAudioMaster.h
+++ b/Engine/Source/Runtime/Engine/Classes/Matinee/InterpTrackInstAudioMaster.h
@@ -11,8 +11,8 @@ class UInterpTrackInstAudioMaster : public UInterpTrackInst
 
 
 	// Begin UInterpTrackInst Instance
-	virtual void InitTrackInst(UInterpTrack* Track) OVERRIDE;
-	virtual void TermTrackInst(UInterpTrack* Track) OVERRIDE;
+	virtual void InitTrackInst(UInterpTrack* Track) override;
+	virtual void TermTrackInst(UInterpTrack* Track) override;
 	// End UInterpTrackInst Instance
 };
 
diff --git a/Engine/Source/Runtime/Engine/Classes/Matinee/InterpTrackInstBoolProp.h b/Engine/Source/Runtime/Engine/Classes/Matinee/InterpTrackInstBoolProp.h
--- a/Engine/Source/Runtime/Engine/Classes/Matinee/InterpTrackInstBoolProp.h
+++ b/Engine/Source/Runtime/Engine/Classes/Matinee/InterpTrackInstBoolProp.h
@@ -22,9 +22,9 @@ class UInterpTrackInstBoolProp : public UInterpTrackInstProperty
 
 
 	// Begin UInterpTrackInst Instance
-	virtual void InitTrackInst(UInterpTrack* Track) OVERRIDE;
-	virtual void SaveActorState(UInterpTrack* Track) OVERRIDE;
-	virtual void RestoreActorState(UInterpTrack* Track) OVERRIDE;
+	virtual void InitTrackInst(UInterpTrack* Track) override;
+	virtual void SaveActorState(UInterpTrack* Track) override;
+	virtual void RestoreActorState(UInterpTrack* Track) override;
 	// End UInterpTrackInst Instance
 };
 
diff --git a/Engine/Source/Runtime/Engine/Classes/Matinee/InterpTrackInstColorProp.h b/Engine/Source/Runtime/Engine/Classes/Matinee/InterpTrackInstColorProp.h
--- a/Engine/Source/Runtime/Engine/Classes/Matinee/InterpTrackInstColorProp.h
+++ b/Engine/Source/Runtime/Engine/Classes/Matinee/InterpTrackInstColorProp.h
@@ -18,9 +18,9 @@ class UInterpTrackInstColorProp : public UInterpTrackInstProperty
 
 
 	// Begin UInterpTrackInst Instance
-	virtual void InitTrackInst(UInterpTrack* Track) OVERRIDE;
-	virtual void SaveActorState(UInterpTrack* Track) OVERRIDE;
-	virtual void RestoreActorState(UInterpTrack* Track) OVERRIDE;
+	virtual void InitTrackInst(UInterpTrack* Track) override;
+	virtual void SaveActorState(UInterpTrack* Track) override;
+	virtual void RestoreActorState(UInterpTrack* Track) override;
 	// End UInterpTrackInst Instance
 };
 
diff --git a/Engine/Source/Runtime/Engine/Classes/Matinee/InterpTrackInstColorScale.h b/Engine/Source/Runtime/Engine/Classes/Matinee/InterpTrackInstColorScale.h
--- a/Engine/Source/Runtime/Engine/Classes/Matinee/InterpTrackInstColorScale.h
+++ b/Engine/Source/Runtime/Engine/Classes/Matinee/InterpTrackInstColorScale.h
@@ -10,7 +10,7 @@ class UInterpTrackInstColorScale : public UInterpTrackInst
 
 
 	// Begin UInterpTrackInst Instance
-	virtual void TermTrackInst(UInterpTrack* Track) OVERRIDE;
+	virtual void TermTrackInst(UInterpTrack* Track) override;
 	// End UInterpTrackInst Instance
 };
 
diff --git a/Engine/Source/Runtime/Engine/Classes/Matinee/InterpTrackInstDirector.h b/Engine/Source/Runtime/Engine/Classes/Matinee/InterpTrackInstDirector.h
--- a/Engine/Source/Runtime/Engine/Classes/Matinee/InterpTrackInstDirector.h
+++ b/Engine/Source/Runtime/Engine/Classes/Matinee/InterpTrackInstDirector.h
@@ -15,7 +15,7 @@ class UInterpTrackInstDirector : public UInterpTrackInst
 
 
 	// Begin UInterpTrackInst Instance
-	virtual void TermTrackInst(class UInterpTrack* Track) OVERRIDE;
+	virtual void TermTrackInst(class UInterpTrack* Track) override;
 	// End UInterpTrackInst Instance
 };
 
diff --git a/Engine/Source/Runtime/Engine/Classes/Matinee/InterpTrackInstEvent.h b/Engine/Source/Runtime/Engine/Classes/Matinee/InterpTrackInstEvent.h
--- a/Engine/Source/Runtime/Engine/Classes/Matinee/InterpTrackInstEvent.h
+++ b/Engine/Source/Runtime/Engine/Classes/Matinee/InterpTrackInstEvent.h
@@ -18,7 +18,7 @@ class UInterpTrackInstEvent : public UInterpTrackInst
 
 
 	// Begin UInterpTrackInst Instance
-	virtual void InitTrackInst(UInterpTrack* Track) OVERRIDE;
+	virtual void InitTrackInst(UInterpTrack* Track) override;
 	// End UInterpTrackInst Instance
 };
 
diff --git a/Engine/Source/Runtime/Engine/Classes/Matinee/InterpTrackInstFade.h b/Engine/Source/Runtime/Engine/Classes/Matinee/InterpTrackInstFade.h
--- a/Engine/Source/Runtime/Engine/Classes/Matinee/InterpTrackInstFade.h
+++ b/Engine/Source/Runtime/Engine/Classes/Matinee/InterpTrackInstFade.h
@@ -11,7 +11,7 @@ class UInterpTrackInstFade : public UInterpTrackInst
 
 
 	// Begin UInterpTrackInst Instance
-	virtual void TermTrackInst(UInterpTrack* Track) OVERRIDE;
+	virtual void TermTrackInst(UInterpTrack* Track) override;
 	// End UInterpTrackInst Instance
 };
 
diff --git a/Engine/Source/Runtime/Engine/Classes/Matinee/InterpTrackInstFloatMaterialParam.h b/Engine/Source/Runtime/Engine/Classes/Matinee/InterpTrackInstFloatMaterialParam.h
--- a/Engine/Source/Runtime/Engine/Classes/Matinee/InterpTrackInstFloatMaterialParam.h
+++ b/Engine/Source/Runtime/Engine/Classes/Matinee/InterpTrackInstFloatMaterialParam.h
@@ -27,10 +27,10 @@ class UInterpTrackInstFloatMaterialParam : public UInterpTrackInst
 
 
 	// Begin UInterpTrackInst Instance
-	virtual void InitTrackInst(UInterpTrack* Track) OVERRIDE;
-	virtual void TermTrackInst(UInterpTrack* Track) OVERRIDE;
-	virtual void SaveActorState(UInterpTrack* Track) OVERRIDE;
-	virtual void RestoreActorState(UInterpTrack* Track) OVERRIDE;
+	virtual void InitTrackInst(UInterpTrack* Track) override;
+	virtual void TermTrackInst(UInterpTrack* Track) override;
+	virtual void SaveActorState(UInterpTrack* Track) override;
+	virtual void RestoreActorState(UInterpTrack* Track) override;
 	// End UInterpTrackInst Instance
 };
 
diff --git a/Engine/Source/Runtime/Engine/Classes/Matinee/InterpTrackInstFloatParticleParam.h b/Engine/Source/Runtime/Engine/Classes/Matinee/InterpTrackInstFloatParticleParam.h
--- a/Engine/Source/Runtime/Engine/Classes/Matinee/InterpTrackInstFloatParticleParam.h
+++ b/Engine/Source/Runtime/Engine/Classes/Matinee/InterpTrackInstFloatParticleParam.h
@@ -15,8 +15,8 @@ class UInterpTrackInstFloatParticleParam : public UInterpTrackInst
 
 
 	// Begin UInterpTrackInst Instance
-	virtual void SaveActorState(UInterpTrack* Track) OVERRIDE;
-	virtual void RestoreActorState(UInterpTrack* Track) OVERRIDE;
+	virtual void SaveActorState(UInterpTrack* Track) override;
+	virtual void RestoreActorState(UInterpTrack* Track) override;
 	// End UInterpTrackInst Instance
 };
 
diff --git a/Engine/Source/Runtime/Engine/Classes/Matinee/InterpTrackInstParticleReplay.h b/Engine/Source/Runtime/Engine/Classes/Matinee/InterpTrackInstParticleReplay.h
--- a/Engine/Source/Runtime/Engine/Classes/Matinee/InterpTrackInstParticleReplay.h
+++ b/Engine/Source/Runtime/Engine/Classes/Matinee/InterpTrackInstParticleReplay.h
@@ -18,8 +18,8 @@ class UInterpTrackInstParticleReplay : public UInterpTrackInst
 
 
 	// Begin UInterpTrackInst Instance
-	virtual void InitTrackInst(UInterpTrack* Track) OVERRIDE;
-	virtual void RestoreActorState(UInterpTrack* Track) OVERRIDE;
+	virtual void InitTrackInst(UInterpTrack* Track) override;
+	virtual void RestoreActorState(UInterpTrack* Track) override;
 	// End UInterpTrackInst Instance
 };
 
diff --git a/Engine/Source/Runtime/Engine/Classes/Matinee/InterpTrackInstSlomo.h b/Engine/Source/Runtime/Engine/Classes/Matinee/InterpTrackInstSlomo.h
--- a/Engine/Source/Runtime/Engine/Classes/Matinee/InterpTrackInstSlomo.h
+++ b/Engine/Source/Runtime/Engine/Classes/Matinee/InterpTrackInstSlomo.h
@@ -15,10 +15,10 @@ class UInterpTrackInstSlomo : public UInterpTrackInst
 
 
 	// Begin UInterpTrackInst Instance
-	virtual void InitTrackInst(UInterpTrack* Track) OVERRIDE;
-	virtual void TermTrackInst(UInterpTrack* Track) OVERRIDE;
-	virtual void SaveActorState(UInterpTrack* Track) OVERRIDE;
-	virtual void RestoreActorState(UInterpTrack* Track) OVERRIDE;
+	virtual void InitTrackInst(UInterpTrack* Track) override;
+	virtual void TermTrackInst(UInterpTrack* Track) override;
+	virtual void SaveActorState(UInterpTrack* Track) override;
+	virtual void RestoreActorState(UInterpTrack* Track) override;
 	// End UInterpTrackInst Instance
 
 	/** @return whether the slomo track's effects should actually be applied. We want to only do this once for the server
diff --git a/Engine/Source/Runtime/Engine/Classes/Matinee/InterpTrackInstSound.h b/Engine/Source/Runtime/Engine/Classes/Matinee/InterpTrackInstSound.h
--- a/Engine/Source/Runtime/Engine/Classes/Matinee/InterpTrackInstSound.h
+++ b/Engine/Source/Runtime/Engine/Classes/Matinee/InterpTrackInstSound.h
@@ -17,8 +17,8 @@ class UInterpTrackInstSound : public UInterpTrackInst
 
 
 	// Begin UInterpTrackInst Instance
-	virtual void InitTrackInst(UInterpTrack* Track) OVERRIDE;
-	virtual void TermTrackInst(UInterpTrack* Track) OVERRIDE;
+	virtual void InitTrackInst(UInterpTrack* Track) override;
+	virtual void TermTrackInst(UInterpTrack* Track) override;
 	// End UInterpTrackInst Instance
 };
 
diff --git a/Engine/Source/Runtime/Engine/Classes/Matinee/InterpTrackLinearColorBase.h b/Engine/Source/Runtime/Engine/Classes/Matinee/InterpTrackLinearColorBase.h
--- a/Engine/Source/Runtime/Engine/Classes/Matinee/InterpTrackLinearColorBase.h
+++ b/Engine/Source/Runtime/Engine/Classes/Matinee/InterpTrackLinearColorBase.h
@@ -20,45 +20,45 @@ class UInterpTrackLinearColorBase : public UInterpTrack
 
 	// Begin UObject interface.
 #if WITH_EDITOR
-	virtual void PostEditChangeProperty(FPropertyChangedEvent& PropertyChangedEvent) OVERRIDE;
+	virtual void PostEditChangeProperty(FPropertyChangedEvent& PropertyChangedEvent) override;
 #endif // WITH_EDITOR
-	virtual void Serialize(FArchive& Ar) OVERRIDE
+	virtual void Serialize(FArchive& Ar) override
 	{
 		Super::Serialize(Ar);
 	}
 	// End UObject interface.
 
 	// Begin UInterpTrack interface.
-	virtual int32 GetNumKeyframes() const OVERRIDE;
-	virtual void GetTimeRange(float& StartTime, float& EndTime) const OVERRIDE; 
-	virtual float GetTrackEndTime() const OVERRIDE;
-	virtual float GetKeyframeTime(int32 KeyIndex) const OVERRIDE;
-	virtual int32 GetKeyframeIndex( float KeyTime ) const OVERRIDE;
-	virtual int32 SetKeyframeTime(int32 KeyIndex, float NewKeyTime, bool bUpdateOrder=true) OVERRIDE;
-	virtual void RemoveKeyframe(int32 KeyIndex) OVERRIDE;
-	virtual int32 DuplicateKeyframe(int32 KeyIndex, float NewKeyTime, UInterpTrack* ToTrack = NULL) OVERRIDE;
-	virtual bool GetClosestSnapPosition(float InPosition, TArray<int32> &IgnoreKeys, float& OutPosition) OVERRIDE;
-	virtual FColor GetKeyframeColor(int32 KeyIndex) const OVERRIDE;
+	virtual int32 GetNumKeyframes() const override;
+	virtual void GetTimeRange(float& StartTime, float& EndTime) const override; 
+	virtual float GetTrackEndTime() const override;
+	virtual float GetKeyframeTime(int32 KeyIndex) const override;
+	virtual int32 GetKeyframeIndex( float KeyTime ) const override;
+	virtual int32 SetKeyframeTime(int32 KeyIndex, float NewKeyTime, bool bUpdateOrder=true) override;
+	virtual void RemoveKeyframe(int32 KeyIndex) override;
+	virtual int32 DuplicateKeyframe(int32 KeyIndex, float NewKeyTime, UInterpTrack* ToTrack = NULL) override;
+	virtual bool GetClosestSnapPosition(float InPosition, TArray<int32> &IgnoreKeys, float& OutPosition) override;
+	virtual FColor GetKeyframeColor(int32 KeyIndex) const override;
 	// End UInterpTrack interface.
 
 	// Begin FCurveEdInterface interface
-	virtual int32		GetNumKeys() const OVERRIDE;
-	virtual int32		GetNumSubCurves() const OVERRIDE;
-	virtual FColor	GetSubCurveButtonColor(int32 SubCurveIndex, bool bIsSubCurveHidden) const OVERRIDE;
-	virtual float	GetKeyIn(int32 KeyIndex) OVERRIDE;
-	virtual float	GetKeyOut(int32 SubIndex, int32 KeyIndex) OVERRIDE;
-	virtual void	GetInRange(float& MinIn, float& MaxIn) const OVERRIDE;
-	virtual void	GetOutRange(float& MinOut, float& MaxOut) const OVERRIDE;
-	virtual FColor	GetKeyColor(int32 SubIndex, int32 KeyIndex, const FColor& CurveColor) OVERRIDE;
-	virtual EInterpCurveMode	GetKeyInterpMode(int32 KeyIndex) const OVERRIDE;
-	virtual void	GetTangents(int32 SubIndex, int32 KeyIndex, float& ArriveTangent, float& LeaveTangent) const OVERRIDE;
-	virtual float	EvalSub(int32 SubIndex, float InVal) OVERRIDE;
-	virtual int32		CreateNewKey(float KeyIn) OVERRIDE;
-	virtual void	DeleteKey(int32 KeyIndex) OVERRIDE;
-	virtual int32		SetKeyIn(int32 KeyIndex, float NewInVal) OVERRIDE;
-	virtual void	SetKeyOut(int32 SubIndex, int32 KeyIndex, float NewOutVal) OVERRIDE;
-	virtual void	SetKeyInterpMode(int32 KeyIndex, EInterpCurveMode NewMode) OVERRIDE;
-	virtual void	SetTangents(int32 SubIndex, int32 KeyIndex, float ArriveTangent, float LeaveTangent) OVERRIDE;
+	virtual int32		GetNumKeys() const override;
+	virtual int32		GetNumSubCurves() const override;
+	virtual FColor	GetSubCurveButtonColor(int32 SubCurveIndex, bool bIsSubCurveHidden) const override;
+	virtual float	GetKeyIn(int32 KeyIndex) override;
+	virtual float	GetKeyOut(int32 SubIndex, int32 KeyIndex) override;
+	virtual void	GetInRange(float& MinIn, float& MaxIn) const override;
+	virtual void	GetOutRange(float& MinOut, float& MaxOut) const override;
+	virtual FColor	GetKeyColor(int32 SubIndex, int32 KeyIndex, const FColor& CurveColor) override;
+	virtual EInterpCurveMode	GetKeyInterpMode(int32 KeyIndex) const override;
+	virtual void	GetTangents(int32 SubIndex, int32 KeyIndex, float& ArriveTangent, float& LeaveTangent) const override;
+	virtual float	EvalSub(int32 SubIndex, float InVal) override;
+	virtual int32		CreateNewKey(float KeyIn) override;
+	virtual void	DeleteKey(int32 KeyIndex) override;
+	virtual int32		SetKeyIn(int32 KeyIndex, float NewInVal) override;
+	virtual void	SetKeyOut(int32 SubIndex, int32 KeyIndex, float NewOutVal) override;
+	virtual void	SetKeyInterpMode(int32 KeyIndex, EInterpCurveMode NewMode) override;
+	virtual void	SetTangents(int32 SubIndex, int32 KeyIndex, float ArriveTangent, float LeaveTangent) override;
 	// End FCurveEdInterface interface
 };
 
diff --git a/Engine/Source/Runtime/Engine/Classes/Matinee/InterpTrackLinearColorProp.h b/Engine/Source/Runtime/Engine/Classes/Matinee/InterpTrackLinearColorProp.h
--- a/Engine/Source/Runtime/Engine/Classes/Matinee/InterpTrackLinearColorProp.h
+++ b/Engine/Source/Runtime/Engine/Classes/Matinee/InterpTrackLinearColorProp.h
@@ -16,14 +16,14 @@ class UInterpTrackLinearColorProp : public UInterpTrackLinearColorBase
 
 
 	// Begin UInterpTrack interface.
-	virtual int32 AddKeyframe(float Time, UInterpTrackInst* TrInst, EInterpCurveMode InitInterpMode) OVERRIDE;
-	virtual bool CanAddKeyframe( UInterpTrackInst* TrackInst ) OVERRIDE;
-	virtual void UpdateKeyframe(int32 KeyIndex, UInterpTrackInst* TrInst) OVERRIDE;
-	virtual void PreviewUpdateTrack(float NewPosition, UInterpTrackInst* TrInst) OVERRIDE;
-	virtual void UpdateTrack(float NewPosition, UInterpTrackInst* TrInst, bool bJump) OVERRIDE;
-	virtual const FString	GetEdHelperClassName() const OVERRIDE;
-	virtual const FString	GetSlateHelperClassName() const OVERRIDE;
-	virtual class UTexture2D* GetTrackIcon() const OVERRIDE;
+	virtual int32 AddKeyframe(float Time, UInterpTrackInst* TrInst, EInterpCurveMode InitInterpMode) override;
+	virtual bool CanAddKeyframe( UInterpTrackInst* TrackInst ) override;
+	virtual void UpdateKeyframe(int32 KeyIndex, UInterpTrackInst* TrInst) override;
+	virtual void PreviewUpdateTrack(float NewPosition, UInterpTrackInst* TrInst) override;
+	virtual void UpdateTrack(float NewPosition, UInterpTrackInst* TrInst, bool bJump) override;
+	virtual const FString	GetEdHelperClassName() const override;
+	virtual const FString	GetSlateHelperClassName() const override;
+	virtual class UTexture2D* GetTrackIcon() const override;
 	// End UInterpTrack interface.
 };
 
diff --git a/Engine/Source/Runtime/Engine/Classes/Matinee/InterpTrackMove.h b/Engine/Source/Runtime/Engine/Classes/Matinee/InterpTrackMove.h
--- a/Engine/Source/Runtime/Engine/Classes/Matinee/InterpTrackMove.h
+++ b/Engine/Source/Runtime/Engine/Classes/Matinee/InterpTrackMove.h
@@ -148,56 +148,56 @@ class UInterpTrackMove : public UInterpTrack
 
 	// Begin UObject interface.
 #if WITH_EDITOR
-	virtual void PostEditChangeProperty(FPropertyChangedEvent& PropertyChangedEvent) OVERRIDE;
-	virtual void PostEditImport() OVERRIDE;
+	virtual void PostEditChangeProperty(FPropertyChangedEvent& PropertyChangedEvent) override;
+	virtual void PostEditImport() override;
 #endif
 	// End UObject interface.
 
 	// Begin UInterpTrack interface.
-	virtual int32 GetNumKeyframes() const OVERRIDE;
-	virtual void GetTimeRange(float& StartTime, float& EndTime) const OVERRIDE;
-	virtual float GetTrackEndTime() const OVERRIDE;
-	virtual float GetKeyframeTime(int32 KeyIndex) const OVERRIDE;
-	virtual int32 GetKeyframeIndex( float KeyTime ) const OVERRIDE;
-	virtual int32 AddKeyframe(float Time, UInterpTrackInst* TrInst, EInterpCurveMode InitInterpMode) OVERRIDE;
-	virtual int32 AddChildKeyframe(class UInterpTrack* ChildTrack, float Time, UInterpTrackInst* TrackInst, EInterpCurveMode InitInterpMode) OVERRIDE;
-	virtual bool CanAddKeyframe( UInterpTrackInst* TrackInst ) OVERRIDE;
-	virtual bool CanAddChildKeyframe( UInterpTrackInst* TrackInst ) OVERRIDE;
-	virtual void UpdateKeyframe(int32 KeyIndex, UInterpTrackInst* TrInst) OVERRIDE;
-	virtual void UpdateChildKeyframe( class UInterpTrack* ChildTrack, int32 KeyIndex, UInterpTrackInst* TrackInst ) OVERRIDE;
-	virtual int32 SetKeyframeTime(int32 KeyIndex, float NewKeyTime, bool bUpdateOrder=true) OVERRIDE;
-	virtual void RemoveKeyframe(int32 KeyIndex) OVERRIDE;
-	virtual int32 DuplicateKeyframe(int32 KeyIndex, float NewKeyTime, UInterpTrack* ToTrack = NULL) OVERRIDE;
-	virtual bool GetClosestSnapPosition(float InPosition, TArray<int32> &IgnoreKeys, float& OutPosition) OVERRIDE;
-	virtual void ConditionalPreviewUpdateTrack(float NewPosition, class UInterpTrackInst* TrInst) OVERRIDE;
-	virtual void PreviewUpdateTrack(float NewPosition, UInterpTrackInst* TrInst) OVERRIDE;
-	virtual void UpdateTrack(float NewPosition, UInterpTrackInst* TrInst, bool bJump) OVERRIDE;
-	virtual class UTexture2D* GetTrackIcon() const OVERRIDE;
-	virtual FColor GetKeyframeColor(int32 KeyIndex) const OVERRIDE;
-	virtual void DrawTrack( FCanvas* Canvas, UInterpGroup* Group, const FInterpTrackDrawParams& Params ) OVERRIDE;
-	virtual void Render3DTrack(UInterpTrackInst* TrInst, const FSceneView* View, class FPrimitiveDrawInterface* PDI, int32 TrackIndex, const FColor& TrackColor, TArray<struct FInterpEdSelKey>& SelectedKeys) OVERRIDE;
-	virtual void SetTrackToSensibleDefault() OVERRIDE;
-	virtual void ApplyWorldOffset(const FVector& InOffset, bool bWorldShift) OVERRIDE;
+	virtual int32 GetNumKeyframes() const override;
+	virtual void GetTimeRange(float& StartTime, float& EndTime) const override;
+	virtual float GetTrackEndTime() const override;
+	virtual float GetKeyframeTime(int32 KeyIndex) const override;
+	virtual int32 GetKeyframeIndex( float KeyTime ) const override;
+	virtual int32 AddKeyframe(float Time, UInterpTrackInst* TrInst, EInterpCurveMode InitInterpMode) override;
+	virtual int32 AddChildKeyframe(class UInterpTrack* ChildTrack, float Time, UInterpTrackInst* TrackInst, EInterpCurveMode InitInterpMode) override;
+	virtual bool CanAddKeyframe( UInterpTrackInst* TrackInst ) override;
+	virtual bool CanAddChildKeyframe( UInterpTrackInst* TrackInst ) override;
+	virtual void UpdateKeyframe(int32 KeyIndex, UInterpTrackInst* TrInst) override;
+	virtual void UpdateChildKeyframe( class UInterpTrack* ChildTrack, int32 KeyIndex, UInterpTrackInst* TrackInst ) override;
+	virtual int32 SetKeyframeTime(int32 KeyIndex, float NewKeyTime, bool bUpdateOrder=true) override;
+	virtual void RemoveKeyframe(int32 KeyIndex) override;
+	virtual int32 DuplicateKeyframe(int32 KeyIndex, float NewKeyTime, UInterpTrack* ToTrack = NULL) override;
+	virtual bool GetClosestSnapPosition(float InPosition, TArray<int32> &IgnoreKeys, float& OutPosition) override;
+	virtual void ConditionalPreviewUpdateTrack(float NewPosition, class UInterpTrackInst* TrInst) override;
+	virtual void PreviewUpdateTrack(float NewPosition, UInterpTrackInst* TrInst) override;
+	virtual void UpdateTrack(float NewPosition, UInterpTrackInst* TrInst, bool bJump) override;
+	virtual class UTexture2D* GetTrackIcon() const override;
+	virtual FColor GetKeyframeColor(int32 KeyIndex) const override;
+	virtual void DrawTrack( FCanvas* Canvas, UInterpGroup* Group, const FInterpTrackDrawParams& Params ) override;
+	virtual void Render3DTrack(UInterpTrackInst* TrInst, const FSceneView* View, class FPrimitiveDrawInterface* PDI, int32 TrackIndex, const FColor& TrackColor, TArray<struct FInterpEdSelKey>& SelectedKeys) override;
+	virtual void SetTrackToSensibleDefault() override;
+	virtual void ApplyWorldOffset(const FVector& InOffset, bool bWorldShift) override;
 	// End UInterpTrack interface.
 
 	// Begin FCurveEdInterface interface.
-	virtual int32		GetNumKeys() const OVERRIDE;
-	virtual int32		GetNumSubCurves() const OVERRIDE;
-	virtual FColor	GetSubCurveButtonColor(int32 SubCurveIndex, bool bIsSubCurveHidden) const OVERRIDE;
-	virtual float	GetKeyIn(int32 KeyIndex) OVERRIDE;
-	virtual float	GetKeyOut(int32 SubIndex, int32 KeyIndex) OVERRIDE;
-	virtual void	GetInRange(float& MinIn, float& MaxIn) const OVERRIDE;
-	virtual void	GetOutRange(float& MinOut, float& MaxOut) const OVERRIDE;
-	virtual FColor	GetKeyColor(int32 SubIndex, int32 KeyIndex, const FColor& CurveColor) OVERRIDE;
-	virtual EInterpCurveMode	GetKeyInterpMode(int32 KeyIndex) const OVERRIDE;
-	virtual void	GetTangents(int32 SubIndex, int32 KeyIndex, float& ArriveTangent, float& LeaveTangent) const OVERRIDE;
-	virtual float	EvalSub(int32 SubIndex, float InVal) OVERRIDE;
-	virtual int32		CreateNewKey(float KeyIn) OVERRIDE;
-	virtual void	DeleteKey(int32 KeyIndex) OVERRIDE;
-	virtual int32		SetKeyIn(int32 KeyIndex, float NewInVal) OVERRIDE;
-	virtual void	SetKeyOut(int32 SubIndex, int32 KeyIndex, float NewOutVal) OVERRIDE;
-	virtual void	SetKeyInterpMode(int32 KeyIndex, EInterpCurveMode NewMode) OVERRIDE;
-	virtual void	SetTangents(int32 SubIndex, int32 KeyIndex, float ArriveTangent, float LeaveTangent) OVERRIDE;
+	virtual int32		GetNumKeys() const override;
+	virtual int32		GetNumSubCurves() const override;
+	virtual FColor	GetSubCurveButtonColor(int32 SubCurveIndex, bool bIsSubCurveHidden) const override;
+	virtual float	GetKeyIn(int32 KeyIndex) override;
+	virtual float	GetKeyOut(int32 SubIndex, int32 KeyIndex) override;
+	virtual void	GetInRange(float& MinIn, float& MaxIn) const override;
+	virtual void	GetOutRange(float& MinOut, float& MaxOut) const override;
+	virtual FColor	GetKeyColor(int32 SubIndex, int32 KeyIndex, const FColor& CurveColor) override;
+	virtual EInterpCurveMode	GetKeyInterpMode(int32 KeyIndex) const override;
+	virtual void	GetTangents(int32 SubIndex, int32 KeyIndex, float& ArriveTangent, float& LeaveTangent) const override;
+	virtual float	EvalSub(int32 SubIndex, float InVal) override;
+	virtual int32		CreateNewKey(float KeyIn) override;
+	virtual void	DeleteKey(int32 KeyIndex) override;
+	virtual int32		SetKeyIn(int32 KeyIndex, float NewInVal) override;
+	virtual void	SetKeyOut(int32 SubIndex, int32 KeyIndex, float NewOutVal) override;
+	virtual void	SetKeyInterpMode(int32 KeyIndex, EInterpCurveMode NewMode) override;
+	virtual void	SetTangents(int32 SubIndex, int32 KeyIndex, float ArriveTangent, float LeaveTangent) override;
 	// End FCurveEdInterface interface.
 
 
diff --git a/Engine/Source/Runtime/Engine/Classes/Matinee/InterpTrackMoveAxis.h b/Engine/Source/Runtime/Engine/Classes/Matinee/InterpTrackMoveAxis.h
--- a/Engine/Source/Runtime/Engine/Classes/Matinee/InterpTrackMoveAxis.h
+++ b/Engine/Source/Runtime/Engine/Classes/Matinee/InterpTrackMoveAxis.h
@@ -38,23 +38,23 @@ class UInterpTrackMoveAxis : public UInterpTrackFloatBase
 
 
 	// Begin UInterpTrack Interface
-	virtual int32 AddKeyframe( float Time, UInterpTrackInst* TrInst, EInterpCurveMode InitInterpMode ) OVERRIDE;
-	virtual void UpdateKeyframe(int32 KeyIndex, UInterpTrackInst* TrInst) OVERRIDE;
-	virtual int32 SetKeyframeTime( int32 KeyIndex, float NewKeyTime, bool bUpdateOrder ) OVERRIDE;
-	virtual void RemoveKeyframe( int32 KeyIndex ) OVERRIDE;
-	virtual int32 DuplicateKeyframe( int32 KeyIndex, float NewKeyTime, UInterpTrack* ToTrack = NULL ) OVERRIDE;
-	virtual class UTexture2D* GetTrackIcon() const OVERRIDE;
-	virtual void ReduceKeys( float IntervalStart, float IntervalEnd, float Tolerance ) OVERRIDE;
+	virtual int32 AddKeyframe( float Time, UInterpTrackInst* TrInst, EInterpCurveMode InitInterpMode ) override;
+	virtual void UpdateKeyframe(int32 KeyIndex, UInterpTrackInst* TrInst) override;
+	virtual int32 SetKeyframeTime( int32 KeyIndex, float NewKeyTime, bool bUpdateOrder ) override;
+	virtual void RemoveKeyframe( int32 KeyIndex ) override;
+	virtual int32 DuplicateKeyframe( int32 KeyIndex, float NewKeyTime, UInterpTrack* ToTrack = NULL ) override;
+	virtual class UTexture2D* GetTrackIcon() const override;
+	virtual void ReduceKeys( float IntervalStart, float IntervalEnd, float Tolerance ) override;
 	// End UInterpTrack Interface
 	
 
 
 	// Begin FCurveEdInterface interface.
-	virtual FColor GetSubCurveButtonColor( int32 SubCurveIndex, bool bIsSubCurveHidden ) const OVERRIDE;
-	virtual int32 CreateNewKey( float KeyIn ) OVERRIDE;
-	virtual void DeleteKey( int32 KeyIndex ) OVERRIDE;
-	virtual int32 SetKeyIn( int32 KeyIndex, float NewInVal ) OVERRIDE;
-	virtual FColor GetKeyColor(int32 SubIndex, int32 KeyIndex, const FColor& CurveColor) OVERRIDE;
+	virtual FColor GetSubCurveButtonColor( int32 SubCurveIndex, bool bIsSubCurveHidden ) const override;
+	virtual int32 CreateNewKey( float KeyIn ) override;
+	virtual void DeleteKey( int32 KeyIndex ) override;
+	virtual int32 SetKeyIn( int32 KeyIndex, float NewInVal ) override;
+	virtual FColor GetKeyColor(int32 SubIndex, int32 KeyIndex, const FColor& CurveColor) override;
 	// End FCurveEdInterface interface.
 
 	/** @todo document */
diff --git a/Engine/Source/Runtime/Engine/Classes/Matinee/InterpTrackParticleReplay.h b/Engine/Source/Runtime/Engine/Classes/Matinee/InterpTrackParticleReplay.h
--- a/Engine/Source/Runtime/Engine/Classes/Matinee/InterpTrackParticleReplay.h
+++ b/Engine/Source/Runtime/Engine/Classes/Matinee/InterpTrackParticleReplay.h
@@ -57,30 +57,30 @@ class UInterpTrackParticleReplay : public UInterpTrack
 #endif // WITH_EDITORONLY_DATA
 
 	// Begin UInterpTrack interface.
-	virtual int32 GetNumKeyframes() const OVERRIDE;
-	virtual void GetTimeRange(float& StartTime, float& EndTime) const OVERRIDE;
-	virtual float GetTrackEndTime() const OVERRIDE;
-	virtual float GetKeyframeTime(int32 KeyIndex) const OVERRIDE;
-	virtual int32 GetKeyframeIndex( float KeyTime ) const OVERRIDE;
-	virtual int32 AddKeyframe(float Time, UInterpTrackInst* TrInst, EInterpCurveMode InitInterpMode) OVERRIDE;
-	virtual int32 SetKeyframeTime(int32 KeyIndex, float NewKeyTime, bool bUpdateOrder=true) OVERRIDE;
-	virtual void RemoveKeyframe(int32 KeyIndex) OVERRIDE;
-	virtual int32 DuplicateKeyframe(int32 KeyIndex, float NewKeyTime, UInterpTrack* ToTrack = NULL) OVERRIDE;
-	virtual bool GetClosestSnapPosition(float InPosition, TArray<int32> &IgnoreKeys, float& OutPosition) OVERRIDE;
-	virtual void PreviewUpdateTrack(float NewPosition, UInterpTrackInst* TrInst) OVERRIDE;
-	virtual void UpdateTrack(float NewPosition, UInterpTrackInst* TrInst, bool bJump) OVERRIDE;
-	virtual const FString	GetEdHelperClassName() const OVERRIDE;
-	virtual const FString	GetSlateHelperClassName() const OVERRIDE;
-	virtual class UTexture2D* GetTrackIcon() const OVERRIDE;
-	virtual bool AllowStaticActors() OVERRIDE { return true; }
+	virtual int32 GetNumKeyframes() const override;
+	virtual void GetTimeRange(float& StartTime, float& EndTime) const override;
+	virtual float GetTrackEndTime() const override;
+	virtual float GetKeyframeTime(int32 KeyIndex) const override;
+	virtual int32 GetKeyframeIndex( float KeyTime ) const override;
+	virtual int32 AddKeyframe(float Time, UInterpTrackInst* TrInst, EInterpCurveMode InitInterpMode) override;
+	virtual int32 SetKeyframeTime(int32 KeyIndex, float NewKeyTime, bool bUpdateOrder=true) override;
+	virtual void RemoveKeyframe(int32 KeyIndex) override;
+	virtual int32 DuplicateKeyframe(int32 KeyIndex, float NewKeyTime, UInterpTrack* ToTrack = NULL) override;
+	virtual bool GetClosestSnapPosition(float InPosition, TArray<int32> &IgnoreKeys, float& OutPosition) override;
+	virtual void PreviewUpdateTrack(float NewPosition, UInterpTrackInst* TrInst) override;
+	virtual void UpdateTrack(float NewPosition, UInterpTrackInst* TrInst, bool bJump) override;
+	virtual const FString	GetEdHelperClassName() const override;
+	virtual const FString	GetSlateHelperClassName() const override;
+	virtual class UTexture2D* GetTrackIcon() const override;
+	virtual bool AllowStaticActors() override { return true; }
 	virtual void DrawTrack( FCanvas* Canvas, UInterpGroup* Group, const FInterpTrackDrawParams& Params );
 	// End UInterpTrack interface.
 
 	// Begin FInterpEdInputInterface Interface
-	virtual void BeginDrag(FInterpEdInputData &InputData) OVERRIDE;
-	virtual void EndDrag(FInterpEdInputData &InputData) OVERRIDE;
-	virtual EMouseCursor::Type GetMouseCursor(FInterpEdInputData &InputData) OVERRIDE;
-	virtual void ObjectDragged(FInterpEdInputData& InputData) OVERRIDE;
+	virtual void BeginDrag(FInterpEdInputData &InputData) override;
+	virtual void EndDrag(FInterpEdInputData &InputData) override;
+	virtual EMouseCursor::Type GetMouseCursor(FInterpEdInputData &InputData) override;
+	virtual void ObjectDragged(FInterpEdInputData& InputData) override;
 	// End FInterpEdInputInterface Interface
 
 };
diff --git a/Engine/Source/Runtime/Engine/Classes/Matinee/InterpTrackSound.h b/Engine/Source/Runtime/Engine/Classes/Matinee/InterpTrackSound.h
--- a/Engine/Source/Runtime/Engine/Classes/Matinee/InterpTrackSound.h
+++ b/Engine/Source/Runtime/Engine/Classes/Matinee/InterpTrackSound.h
@@ -68,29 +68,29 @@ class UInterpTrackSound : public UInterpTrackVectorBase
 	uint32 bAttach:1;
 
 	// Begin UObject Interface
-	virtual void PostLoad() OVERRIDE;
+	virtual void PostLoad() override;
 	// End UObject Interface
 
 	// Begin UInterpTrack interface
-	virtual int32 GetNumKeyframes() const OVERRIDE;
-	virtual void GetTimeRange(float& StartTime, float& EndTime) const OVERRIDE;
-	virtual float GetTrackEndTime() const OVERRIDE;
-	virtual float GetKeyframeTime(int32 KeyIndex) const OVERRIDE;
-	virtual int32 GetKeyframeIndex( float KeyTime ) const OVERRIDE;
-	virtual int32 AddKeyframe(float Time, UInterpTrackInst* TrInst, EInterpCurveMode InitInterpMode) OVERRIDE;
-	virtual int32 SetKeyframeTime(int32 KeyIndex, float NewKeyTime, bool bUpdateOrder=true) OVERRIDE;
-	virtual void RemoveKeyframe(int32 KeyIndex) OVERRIDE;
-	virtual int32 DuplicateKeyframe(int32 KeyIndex, float NewKeyTime, UInterpTrack* ToTrack = NULL) OVERRIDE;
-	virtual bool GetClosestSnapPosition(float InPosition, TArray<int32> &IgnoreKeys, float& OutPosition) OVERRIDE;
-	virtual void PreviewUpdateTrack(float NewPosition, UInterpTrackInst* TrInst) OVERRIDE;
-	virtual void UpdateTrack(float NewPosition, UInterpTrackInst* TrInst, bool bJump) OVERRIDE;
-	virtual void PreviewStopPlayback(class UInterpTrackInst* TrInst) OVERRIDE;
-	virtual const FString	GetEdHelperClassName() const OVERRIDE;
-	virtual const FString	GetSlateHelperClassName() const OVERRIDE;
-	virtual class UTexture2D* GetTrackIcon() const OVERRIDE;
-	virtual void DrawTrack( FCanvas* Canvas, UInterpGroup* Group, const FInterpTrackDrawParams& Params ) OVERRIDE;
-	virtual bool AllowStaticActors() OVERRIDE { return true; }
-	virtual void SetTrackToSensibleDefault() OVERRIDE;
+	virtual int32 GetNumKeyframes() const override;
+	virtual void GetTimeRange(float& StartTime, float& EndTime) const override;
+	virtual float GetTrackEndTime() const override;
+	virtual float GetKeyframeTime(int32 KeyIndex) const override;
+	virtual int32 GetKeyframeIndex( float KeyTime ) const override;
+	virtual int32 AddKeyframe(float Time, UInterpTrackInst* TrInst, EInterpCurveMode InitInterpMode) override;
+	virtual int32 SetKeyframeTime(int32 KeyIndex, float NewKeyTime, bool bUpdateOrder=true) override;
+	virtual void RemoveKeyframe(int32 KeyIndex) override;
+	virtual int32 DuplicateKeyframe(int32 KeyIndex, float NewKeyTime, UInterpTrack* ToTrack = NULL) override;
+	virtual bool GetClosestSnapPosition(float InPosition, TArray<int32> &IgnoreKeys, float& OutPosition) override;
+	virtual void PreviewUpdateTrack(float NewPosition, UInterpTrackInst* TrInst) override;
+	virtual void UpdateTrack(float NewPosition, UInterpTrackInst* TrInst, bool bJump) override;
+	virtual void PreviewStopPlayback(class UInterpTrackInst* TrInst) override;
+	virtual const FString	GetEdHelperClassName() const override;
+	virtual const FString	GetSlateHelperClassName() const override;
+	virtual class UTexture2D* GetTrackIcon() const override;
+	virtual void DrawTrack( FCanvas* Canvas, UInterpGroup* Group, const FInterpTrackDrawParams& Params ) override;
+	virtual bool AllowStaticActors() override { return true; }
+	virtual void SetTrackToSensibleDefault() override;
 	// End UInterpTrack interface
 
 	/**
diff --git a/Engine/Source/Runtime/Engine/Classes/Matinee/InterpTrackToggle.h b/Engine/Source/Runtime/Engine/Classes/Matinee/InterpTrackToggle.h
--- a/Engine/Source/Runtime/Engine/Classes/Matinee/InterpTrackToggle.h
+++ b/Engine/Source/Runtime/Engine/Classes/Matinee/InterpTrackToggle.h
@@ -78,23 +78,23 @@ class UInterpTrackToggle : public UInterpTrack
 
 
 	// Begin UInterpTrack interface.
-	virtual int32 GetNumKeyframes() const OVERRIDE;
-	virtual void GetTimeRange(float& StartTime, float& EndTime) const OVERRIDE;
-	virtual float GetTrackEndTime() const OVERRIDE;
-	virtual float GetKeyframeTime(int32 KeyIndex) const OVERRIDE;
-	virtual int32 GetKeyframeIndex( float KeyTime ) const OVERRIDE;
-	virtual int32 AddKeyframe(float Time, UInterpTrackInst* TrInst, EInterpCurveMode InitInterpMode) OVERRIDE;
-	virtual int32 SetKeyframeTime(int32 KeyIndex, float NewKeyTime, bool bUpdateOrder=true) OVERRIDE;
-	virtual void RemoveKeyframe(int32 KeyIndex) OVERRIDE;
-	virtual int32 DuplicateKeyframe(int32 KeyIndex, float NewKeyTime, UInterpTrack* ToTrack = NULL) OVERRIDE;
-	virtual bool GetClosestSnapPosition(float InPosition, TArray<int32> &IgnoreKeys, float& OutPosition) OVERRIDE;
-	virtual void PreviewUpdateTrack(float NewPosition, UInterpTrackInst* TrInst) OVERRIDE;
-	virtual void UpdateTrack(float NewPosition, UInterpTrackInst* TrInst, bool bJump) OVERRIDE;
-	virtual const FString	GetEdHelperClassName() const OVERRIDE;
-	virtual const FString	GetSlateHelperClassName() const OVERRIDE;
-	virtual class UTexture2D* GetTrackIcon() const OVERRIDE;
-	virtual bool AllowStaticActors() OVERRIDE { return true; }
-	virtual void DrawTrack( FCanvas* Canvas, UInterpGroup* Group, const FInterpTrackDrawParams& Params ) OVERRIDE;
+	virtual int32 GetNumKeyframes() const override;
+	virtual void GetTimeRange(float& StartTime, float& EndTime) const override;
+	virtual float GetTrackEndTime() const override;
+	virtual float GetKeyframeTime(int32 KeyIndex) const override;
+	virtual int32 GetKeyframeIndex( float KeyTime ) const override;
+	virtual int32 AddKeyframe(float Time, UInterpTrackInst* TrInst, EInterpCurveMode InitInterpMode) override;
+	virtual int32 SetKeyframeTime(int32 KeyIndex, float NewKeyTime, bool bUpdateOrder=true) override;
+	virtual void RemoveKeyframe(int32 KeyIndex) override;
+	virtual int32 DuplicateKeyframe(int32 KeyIndex, float NewKeyTime, UInterpTrack* ToTrack = NULL) override;
+	virtual bool GetClosestSnapPosition(float InPosition, TArray<int32> &IgnoreKeys, float& OutPosition) override;
+	virtual void PreviewUpdateTrack(float NewPosition, UInterpTrackInst* TrInst) override;
+	virtual void UpdateTrack(float NewPosition, UInterpTrackInst* TrInst, bool bJump) override;
+	virtual const FString	GetEdHelperClassName() const override;
+	virtual const FString	GetSlateHelperClassName() const override;
+	virtual class UTexture2D* GetTrackIcon() const override;
+	virtual bool AllowStaticActors() override { return true; }
+	virtual void DrawTrack( FCanvas* Canvas, UInterpGroup* Group, const FInterpTrackDrawParams& Params ) override;
 	// Begin UInterpTrack interface.
 };
 
diff --git a/Engine/Source/Runtime/Engine/Classes/Matinee/InterpTrackVectorBase.h b/Engine/Source/Runtime/Engine/Classes/Matinee/InterpTrackVectorBase.h
--- a/Engine/Source/Runtime/Engine/Classes/Matinee/InterpTrackVectorBase.h
+++ b/Engine/Source/Runtime/Engine/Classes/Matinee/InterpTrackVectorBase.h
@@ -20,41 +20,41 @@ class UInterpTrackVectorBase : public UInterpTrack
 
 	// Begin UObject interface.
 #if WITH_EDITOR
-	virtual void PostEditChangeProperty(FPropertyChangedEvent& PropertyChangedEvent) OVERRIDE;
+	virtual void PostEditChangeProperty(FPropertyChangedEvent& PropertyChangedEvent) override;
 #endif // WITH_EDITOR
 	// End UObject interface.
 
 	// Begin InterpTrack interface.
-	virtual int32 GetNumKeyframes() const OVERRIDE;
-	virtual void GetTimeRange(float& StartTime, float& EndTime) const OVERRIDE;
-	virtual float GetTrackEndTime() const OVERRIDE;
-	virtual float GetKeyframeTime(int32 KeyIndex) const OVERRIDE;
-	virtual int32 GetKeyframeIndex( float KeyTime ) const OVERRIDE;
-	virtual int32 SetKeyframeTime(int32 KeyIndex, float NewKeyTime, bool bUpdateOrder=true) OVERRIDE;
-	virtual void RemoveKeyframe(int32 KeyIndex) OVERRIDE;
-	virtual int32 DuplicateKeyframe(int32 KeyIndex, float NewKeyTime, UInterpTrack* ToTrack = NULL) OVERRIDE;
-	virtual bool GetClosestSnapPosition(float InPosition, TArray<int32> &IgnoreKeys, float& OutPosition) OVERRIDE;
-	virtual FColor GetKeyframeColor(int32 KeyIndex) const OVERRIDE;
+	virtual int32 GetNumKeyframes() const override;
+	virtual void GetTimeRange(float& StartTime, float& EndTime) const override;
+	virtual float GetTrackEndTime() const override;
+	virtual float GetKeyframeTime(int32 KeyIndex) const override;
+	virtual int32 GetKeyframeIndex( float KeyTime ) const override;
+	virtual int32 SetKeyframeTime(int32 KeyIndex, float NewKeyTime, bool bUpdateOrder=true) override;
+	virtual void RemoveKeyframe(int32 KeyIndex) override;
+	virtual int32 DuplicateKeyframe(int32 KeyIndex, float NewKeyTime, UInterpTrack* ToTrack = NULL) override;
+	virtual bool GetClosestSnapPosition(float InPosition, TArray<int32> &IgnoreKeys, float& OutPosition) override;
+	virtual FColor GetKeyframeColor(int32 KeyIndex) const override;
 	// End InterpTrack interface.
 
 	// Begin FCurveEdInterface interface.
-	virtual int32		GetNumKeys() const OVERRIDE;
-	virtual int32		GetNumSubCurves() const OVERRIDE;
-	virtual FColor	GetSubCurveButtonColor(int32 SubCurveIndex, bool bIsSubCurveHidden) const OVERRIDE;
-	virtual float	GetKeyIn(int32 KeyIndex) OVERRIDE;
-	virtual float	GetKeyOut(int32 SubIndex, int32 KeyIndex) OVERRIDE;
-	virtual void	GetInRange(float& MinIn, float& MaxIn) const OVERRIDE;
-	virtual void	GetOutRange(float& MinOut, float& MaxOut) const OVERRIDE;
-	virtual FColor	GetKeyColor(int32 SubIndex, int32 KeyIndex, const FColor& CurveColor) OVERRIDE;
-	virtual EInterpCurveMode	GetKeyInterpMode(int32 KeyIndex) const OVERRIDE;
-	virtual void	GetTangents(int32 SubIndex, int32 KeyIndex, float& ArriveTangent, float& LeaveTangent) const OVERRIDE;
-	virtual float	EvalSub(int32 SubIndex, float InVal) OVERRIDE;
-	virtual int32		CreateNewKey(float KeyIn) OVERRIDE;
-	virtual void	DeleteKey(int32 KeyIndex) OVERRIDE;
-	virtual int32		SetKeyIn(int32 KeyIndex, float NewInVal) OVERRIDE;
-	virtual void	SetKeyOut(int32 SubIndex, int32 KeyIndex, float NewOutVal) OVERRIDE;
-	virtual void	SetKeyInterpMode(int32 KeyIndex, EInterpCurveMode NewMode) OVERRIDE;
-	virtual void	SetTangents(int32 SubIndex, int32 KeyIndex, float ArriveTangent, float LeaveTangent) OVERRIDE;
+	virtual int32		GetNumKeys() const override;
+	virtual int32		GetNumSubCurves() const override;
+	virtual FColor	GetSubCurveButtonColor(int32 SubCurveIndex, bool bIsSubCurveHidden) const override;
+	virtual float	GetKeyIn(int32 KeyIndex) override;
+	virtual float	GetKeyOut(int32 SubIndex, int32 KeyIndex) override;
+	virtual void	GetInRange(float& MinIn, float& MaxIn) const override;
+	virtual void	GetOutRange(float& MinOut, float& MaxOut) const override;
+	virtual FColor	GetKeyColor(int32 SubIndex, int32 KeyIndex, const FColor& CurveColor) override;
+	virtual EInterpCurveMode	GetKeyInterpMode(int32 KeyIndex) const override;
+	virtual void	GetTangents(int32 SubIndex, int32 KeyIndex, float& ArriveTangent, float& LeaveTangent) const override;
+	virtual float	EvalSub(int32 SubIndex, float InVal) override;
+	virtual int32		CreateNewKey(float KeyIn) override;
+	virtual void	DeleteKey(int32 KeyIndex) override;
+	virtual int32		SetKeyIn(int32 KeyIndex, float NewInVal) override;
+	virtual void	SetKeyOut(int32 SubIndex, int32 KeyIndex, float NewOutVal) override;
+	virtual void	SetKeyInterpMode(int32 KeyIndex, EInterpCurveMode NewMode) override;
+	virtual void	SetTangents(int32 SubIndex, int32 KeyIndex, float ArriveTangent, float LeaveTangent) override;
 	// End FCurveEdInterface interface.
 };
 
diff --git a/Engine/Source/Runtime/Engine/Classes/Matinee/InterpTrackVectorMaterialParam.h b/Engine/Source/Runtime/Engine/Classes/Matinee/InterpTrackVectorMaterialParam.h
--- a/Engine/Source/Runtime/Engine/Classes/Matinee/InterpTrackVectorMaterialParam.h
+++ b/Engine/Source/Runtime/Engine/Classes/Matinee/InterpTrackVectorMaterialParam.h
@@ -19,15 +19,15 @@ class UInterpTrackVectorMaterialParam : public UInterpTrackVectorBase
 
 	//Begin UObject Interface
 #if WITH_EDITOR
-	virtual void PreEditChange(UProperty* PropertyThatWillChange) OVERRIDE;
-	virtual void PostEditChangeProperty(FPropertyChangedEvent& PropertyChangedEvent) OVERRIDE;
+	virtual void PreEditChange(UProperty* PropertyThatWillChange) override;
+	virtual void PostEditChangeProperty(FPropertyChangedEvent& PropertyChangedEvent) override;
 #endif // WITH_EDITOR
 	//End UObject Interface
 
 	// Begin UInterpTrack interface.
-	virtual int32 AddKeyframe(float Time, UInterpTrackInst* TrInst, EInterpCurveMode InitInterpMode) OVERRIDE;
-	virtual void PreviewUpdateTrack(float NewPosition, UInterpTrackInst* TrInst) OVERRIDE;
-	virtual void UpdateTrack(float NewPosition, UInterpTrackInst* TrInst, bool bJump) OVERRIDE;
+	virtual int32 AddKeyframe(float Time, UInterpTrackInst* TrInst, EInterpCurveMode InitInterpMode) override;
+	virtual void PreviewUpdateTrack(float NewPosition, UInterpTrackInst* TrInst) override;
+	virtual void UpdateTrack(float NewPosition, UInterpTrackInst* TrInst, bool bJump) override;
 	// End UInterpTrack interface.
 };
 
diff --git a/Engine/Source/Runtime/Engine/Classes/Matinee/InterpTrackVisibility.h b/Engine/Source/Runtime/Engine/Classes/Matinee/InterpTrackVisibility.h
--- a/Engine/Source/Runtime/Engine/Classes/Matinee/InterpTrackVisibility.h
+++ b/Engine/Source/Runtime/Engine/Classes/Matinee/InterpTrackVisibility.h
@@ -84,23 +84,23 @@ class UInterpTrackVisibility : public UInterpTrack
 
 
 	// Begin UInterpTrack interface.
-	virtual int32 GetNumKeyframes() const OVERRIDE;
-	virtual void GetTimeRange(float& StartTime, float& EndTime) const OVERRIDE;
-	virtual float GetTrackEndTime() const OVERRIDE;
-	virtual float GetKeyframeTime(int32 KeyIndex) const OVERRIDE;
-	virtual int32 GetKeyframeIndex( float KeyTime ) const OVERRIDE;
-	virtual int32 AddKeyframe(float Time, UInterpTrackInst* TrInst, EInterpCurveMode InitInterpMode) OVERRIDE;
-	virtual int32 SetKeyframeTime(int32 KeyIndex, float NewKeyTime, bool bUpdateOrder=true) OVERRIDE;
-	virtual void RemoveKeyframe(int32 KeyIndex) OVERRIDE;
-	virtual int32 DuplicateKeyframe(int32 KeyIndex, float NewKeyTime, UInterpTrack* ToTrack = NULL) OVERRIDE;
-	virtual bool GetClosestSnapPosition(float InPosition, TArray<int32> &IgnoreKeys, float& OutPosition) OVERRIDE;
-	virtual void PreviewUpdateTrack(float NewPosition, UInterpTrackInst* TrInst) OVERRIDE;
-	virtual void UpdateTrack(float NewPosition, UInterpTrackInst* TrInst, bool bJump) OVERRIDE;
-	virtual const FString	GetEdHelperClassName() const OVERRIDE;
-	virtual const FString	GetSlateHelperClassName() const OVERRIDE;
-	virtual class UTexture2D* GetTrackIcon() const OVERRIDE;
-	virtual bool AllowStaticActors() OVERRIDE { return true; }
-	virtual void DrawTrack( FCanvas* Canvas, UInterpGroup* Group, const FInterpTrackDrawParams& Params ) OVERRIDE;
+	virtual int32 GetNumKeyframes() const override;
+	virtual void GetTimeRange(float& StartTime, float& EndTime) const override;
+	virtual float GetTrackEndTime() const override;
+	virtual float GetKeyframeTime(int32 KeyIndex) const override;
+	virtual int32 GetKeyframeIndex( float KeyTime ) const override;
+	virtual int32 AddKeyframe(float Time, UInterpTrackInst* TrInst, EInterpCurveMode InitInterpMode) override;
+	virtual int32 SetKeyframeTime(int32 KeyIndex, float NewKeyTime, bool bUpdateOrder=true) override;
+	virtual void RemoveKeyframe(int32 KeyIndex) override;
+	virtual int32 DuplicateKeyframe(int32 KeyIndex, float NewKeyTime, UInterpTrack* ToTrack = NULL) override;
+	virtual bool GetClosestSnapPosition(float InPosition, TArray<int32> &IgnoreKeys, float& OutPosition) override;
+	virtual void PreviewUpdateTrack(float NewPosition, UInterpTrackInst* TrInst) override;
+	virtual void UpdateTrack(float NewPosition, UInterpTrackInst* TrInst, bool bJump) override;
+	virtual const FString	GetEdHelperClassName() const override;
+	virtual const FString	GetSlateHelperClassName() const override;
+	virtual class UTexture2D* GetTrackIcon() const override;
+	virtual bool AllowStaticActors() override { return true; }
+	virtual void DrawTrack( FCanvas* Canvas, UInterpGroup* Group, const FInterpTrackDrawParams& Params ) override;
 	// End UInterpTrack interface.
 
 	/** Shows or hides the actor */
diff --git a/Engine/Source/Runtime/Engine/Classes/Matinee/MatineeActor.h b/Engine/Source/Runtime/Engine/Classes/Matinee/MatineeActor.h
--- a/Engine/Source/Runtime/Engine/Classes/Matinee/MatineeActor.h
+++ b/Engine/Source/Runtime/Engine/Classes/Matinee/MatineeActor.h
@@ -297,23 +297,23 @@ public:
 #endif
 
 	// Begin AActor Interface
-	virtual float GetNetPriority(const FVector& ViewPos, const FVector& ViewDir, APlayerController* Viewer, UActorChannel* InChannel, float Time, bool bLowBandwidth) OVERRIDE;
-	virtual void Tick(float DeltaTime) OVERRIDE;
-	virtual void PreNetReceive() OVERRIDE;
-	virtual void PostNetReceive() OVERRIDE;
-	virtual void BeginPlay() OVERRIDE;
-	virtual void ApplyWorldOffset(const FVector& InOffset, bool bWorldShift) OVERRIDE;
+	virtual float GetNetPriority(const FVector& ViewPos, const FVector& ViewDir, APlayerController* Viewer, UActorChannel* InChannel, float Time, bool bLowBandwidth) override;
+	virtual void Tick(float DeltaTime) override;
+	virtual void PreNetReceive() override;
+	virtual void PostNetReceive() override;
+	virtual void BeginPlay() override;
+	virtual void ApplyWorldOffset(const FVector& InOffset, bool bWorldShift) override;
 #if WITH_EDITOR
-	virtual bool GetReferencedContentObjects( TArray<UObject*>& Objects ) const OVERRIDE;
+	virtual bool GetReferencedContentObjects( TArray<UObject*>& Objects ) const override;
 #endif
 	// Begin AActor Interface
 
 	// Begin UObject Interface
 #if WITH_EDITOR
-	virtual void PostEditChangeProperty( FPropertyChangedEvent& PropertyChangedEvent ) OVERRIDE; 
-	virtual bool CanEditChange( const UProperty* Property ) const OVERRIDE;
+	virtual void PostEditChangeProperty( FPropertyChangedEvent& PropertyChangedEvent ) override; 
+	virtual bool CanEditChange( const UProperty* Property ) const override;
 #endif // WITH_EDITOR
-	virtual void PostLoadSubobjects( FObjectInstancingGraph* OuterInstanceGraph ) OVERRIDE;
+	virtual void PostLoadSubobjects( FObjectInstancingGraph* OuterInstanceGraph ) override;
 	// End UObject Interface
 
 	/** Increment track forwards by given timestep and iterate over each track updating any properties. */
diff --git a/Engine/Source/Runtime/Engine/Classes/Matinee/MatineeActorCameraAnim.h b/Engine/Source/Runtime/Engine/Classes/Matinee/MatineeActorCameraAnim.h
--- a/Engine/Source/Runtime/Engine/Classes/Matinee/MatineeActorCameraAnim.h
+++ b/Engine/Source/Runtime/Engine/Classes/Matinee/MatineeActorCameraAnim.h
@@ -17,12 +17,12 @@ class AMatineeActorCameraAnim : public AMatineeActor
 	class UCameraAnim* CameraAnim;
 
 	// Begin UObject interface
-	virtual bool NeedsLoadForClient() const OVERRIDE
+	virtual bool NeedsLoadForClient() const override
 	{ 
 		return false; 
 	}
 
-	virtual bool NeedsLoadForServer() const OVERRIDE
+	virtual bool NeedsLoadForServer() const override
 	{ 
 		return false;
 	}
diff --git a/Engine/Source/Runtime/Engine/Classes/Particles/Acceleration/ParticleModuleAcceleration.h b/Engine/Source/Runtime/Engine/Classes/Particles/Acceleration/ParticleModuleAcceleration.h
--- a/Engine/Source/Runtime/Engine/Classes/Particles/Acceleration/ParticleModuleAcceleration.h
+++ b/Engine/Source/Runtime/Engine/Classes/Particles/Acceleration/ParticleModuleAcceleration.h
@@ -33,20 +33,20 @@ class UParticleModuleAcceleration : public UParticleModuleAccelerationBase
 
 	//Begin UObject Interface
 #if WITH_EDITOR
-	virtual void PostEditChangeProperty(FPropertyChangedEvent& PropertyChangedEvent) OVERRIDE;
+	virtual void PostEditChangeProperty(FPropertyChangedEvent& PropertyChangedEvent) override;
 #endif // WITH_EDITOR
-	virtual void PostInitProperties() OVERRIDE;
-	virtual void Serialize(FArchive& Ar) OVERRIDE;
+	virtual void PostInitProperties() override;
+	virtual void Serialize(FArchive& Ar) override;
 	//End UObject Interface
 
 	//Begin UParticleModule Interface
-	virtual void CompileModule( struct FParticleEmitterBuildInfo& EmitterInfo ) OVERRIDE;
-	virtual void Spawn(FParticleEmitterInstance* Owner, int32 Offset, float SpawnTime, FBaseParticle* ParticleBase) OVERRIDE;
-	virtual void Update(FParticleEmitterInstance* Owner, int32 Offset, float DeltaTime) OVERRIDE;
-	virtual uint32 RequiredBytes(FParticleEmitterInstance* Owner = NULL) OVERRIDE;
+	virtual void CompileModule( struct FParticleEmitterBuildInfo& EmitterInfo ) override;
+	virtual void Spawn(FParticleEmitterInstance* Owner, int32 Offset, float SpawnTime, FBaseParticle* ParticleBase) override;
+	virtual void Update(FParticleEmitterInstance* Owner, int32 Offset, float DeltaTime) override;
+	virtual uint32 RequiredBytes(FParticleEmitterInstance* Owner = NULL) override;
 
 #if WITH_EDITOR
-	virtual bool IsValidForLODLevel(UParticleLODLevel* LODLevel, FString& OutErrorString) OVERRIDE;
+	virtual bool IsValidForLODLevel(UParticleLODLevel* LODLevel, FString& OutErrorString) override;
 #endif
 	//End UParticleModule Interface
 
diff --git a/Engine/Source/Runtime/Engine/Classes/Particles/Acceleration/ParticleModuleAccelerationBase.h b/Engine/Source/Runtime/Engine/Classes/Particles/Acceleration/ParticleModuleAccelerationBase.h
--- a/Engine/Source/Runtime/Engine/Classes/Particles/Acceleration/ParticleModuleAccelerationBase.h
+++ b/Engine/Source/Runtime/Engine/Classes/Particles/Acceleration/ParticleModuleAccelerationBase.h
@@ -18,7 +18,7 @@ class UParticleModuleAccelerationBase : public UParticleModule
 
 
 	// Begin UParticleModule Interface
-	virtual void SetToSensibleDefaults(UParticleEmitter* Owner) OVERRIDE;
+	virtual void SetToSensibleDefaults(UParticleEmitter* Owner) override;
 	// End UParticleModule Interface
 };
 
diff --git a/Engine/Source/Runtime/Engine/Classes/Particles/Acceleration/ParticleModuleAccelerationConstant.h b/Engine/Source/Runtime/Engine/Classes/Particles/Acceleration/ParticleModuleAccelerationConstant.h
--- a/Engine/Source/Runtime/Engine/Classes/Particles/Acceleration/ParticleModuleAccelerationConstant.h
+++ b/Engine/Source/Runtime/Engine/Classes/Particles/Acceleration/ParticleModuleAccelerationConstant.h
@@ -19,9 +19,9 @@ class UParticleModuleAccelerationConstant : public UParticleModuleAccelerationBa
 
 
 	//Begin UParticleModule Interface
-	virtual void CompileModule( struct FParticleEmitterBuildInfo& EmitterInfo ) OVERRIDE;
-	virtual void Spawn(FParticleEmitterInstance* Owner, int32 Offset, float SpawnTime, FBaseParticle* ParticleBase) OVERRIDE;
-	virtual void Update(FParticleEmitterInstance* Owner, int32 Offset, float DeltaTime) OVERRIDE;
+	virtual void CompileModule( struct FParticleEmitterBuildInfo& EmitterInfo ) override;
+	virtual void Spawn(FParticleEmitterInstance* Owner, int32 Offset, float SpawnTime, FBaseParticle* ParticleBase) override;
+	virtual void Update(FParticleEmitterInstance* Owner, int32 Offset, float DeltaTime) override;
 	//End UParticleModule Interface
 };
 
diff --git a/Engine/Source/Runtime/Engine/Classes/Particles/Acceleration/ParticleModuleAccelerationDrag.h b/Engine/Source/Runtime/Engine/Classes/Particles/Acceleration/ParticleModuleAccelerationDrag.h
--- a/Engine/Source/Runtime/Engine/Classes/Particles/Acceleration/ParticleModuleAccelerationDrag.h
+++ b/Engine/Source/Runtime/Engine/Classes/Particles/Acceleration/ParticleModuleAccelerationDrag.h
@@ -22,19 +22,19 @@ class UParticleModuleAccelerationDrag : public UParticleModuleAccelerationBase
 
 	//Begin UObject Interface
 #if WITH_EDITOR
-	virtual void PostEditChangeProperty(FPropertyChangedEvent& PropertyChangedEvent) OVERRIDE;
+	virtual void PostEditChangeProperty(FPropertyChangedEvent& PropertyChangedEvent) override;
 #endif // WITH_EDITOR
-	virtual void PostInitProperties() OVERRIDE;
-	virtual void Serialize(FArchive& Ar) OVERRIDE;
+	virtual void PostInitProperties() override;
+	virtual void Serialize(FArchive& Ar) override;
 	//End UObject Interface
 
 	//Begin UParticleModule Interface
-	virtual void CompileModule( struct FParticleEmitterBuildInfo& EmitterInfo ) OVERRIDE;
-	virtual void Update( FParticleEmitterInstance* Owner, int32 Offset, float DeltaTime ) OVERRIDE;
+	virtual void CompileModule( struct FParticleEmitterBuildInfo& EmitterInfo ) override;
+	virtual void Update( FParticleEmitterInstance* Owner, int32 Offset, float DeltaTime ) override;
 	//End UParticleModule Interface
 
 #if WITH_EDITOR
-	virtual bool IsValidForLODLevel(UParticleLODLevel* LODLevel, FString& OutErrorString) OVERRIDE;
+	virtual bool IsValidForLODLevel(UParticleLODLevel* LODLevel, FString& OutErrorString) override;
 #endif
 
 protected:
diff --git a/Engine/Source/Runtime/Engine/Classes/Particles/Acceleration/ParticleModuleAccelerationDragScaleOverLife.h b/Engine/Source/Runtime/Engine/Classes/Particles/Acceleration/ParticleModuleAccelerationDragScaleOverLife.h
--- a/Engine/Source/Runtime/Engine/Classes/Particles/Acceleration/ParticleModuleAccelerationDragScaleOverLife.h
+++ b/Engine/Source/Runtime/Engine/Classes/Particles/Acceleration/ParticleModuleAccelerationDragScaleOverLife.h
@@ -22,18 +22,18 @@ class UParticleModuleAccelerationDragScaleOverLife : public UParticleModuleAccel
 
 	//Begin UObject Interface
 #if WITH_EDITOR
-	virtual void PostEditChangeProperty(FPropertyChangedEvent& PropertyChangedEvent) OVERRIDE;
+	virtual void PostEditChangeProperty(FPropertyChangedEvent& PropertyChangedEvent) override;
 #endif // WITH_EDITOR
-	virtual void PostInitProperties() OVERRIDE;
-	virtual void Serialize(FArchive& Ar) OVERRIDE;
+	virtual void PostInitProperties() override;
+	virtual void Serialize(FArchive& Ar) override;
 	//End UObject Interface
 
 	//Begin UParticleModule Interface
-	virtual void CompileModule( struct FParticleEmitterBuildInfo& EmitterInfo ) OVERRIDE;
+	virtual void CompileModule( struct FParticleEmitterBuildInfo& EmitterInfo ) override;
 	//End UParticleModule Interface
 
 #if WITH_EDITOR
-	virtual bool IsValidForLODLevel(UParticleLODLevel* LODLevel, FString& OutErrorString) OVERRIDE;
+	virtual bool IsValidForLODLevel(UParticleLODLevel* LODLevel, FString& OutErrorString) override;
 #endif
 
 protected:
diff --git a/Engine/Source/Runtime/Engine/Classes/Particles/Acceleration/ParticleModuleAccelerationOverLifetime.h b/Engine/Source/Runtime/Engine/Classes/Particles/Acceleration/ParticleModuleAccelerationOverLifetime.h
--- a/Engine/Source/Runtime/Engine/Classes/Particles/Acceleration/ParticleModuleAccelerationOverLifetime.h
+++ b/Engine/Source/Runtime/Engine/Classes/Particles/Acceleration/ParticleModuleAccelerationOverLifetime.h
@@ -26,13 +26,13 @@ class UParticleModuleAccelerationOverLifetime : public UParticleModuleAccelerati
 
 	//Begin UObject Interface
 #if WITH_EDITOR
-	virtual void PostEditChangeProperty(FPropertyChangedEvent& PropertyChangedEvent) OVERRIDE;
+	virtual void PostEditChangeProperty(FPropertyChangedEvent& PropertyChangedEvent) override;
 #endif // WITH_EDITOR
-	virtual void PostInitProperties() OVERRIDE;
+	virtual void PostInitProperties() override;
 	//End UObject Interface
 
 	//Begin UParticleModule Interface
-	virtual void Update(FParticleEmitterInstance* Owner, int32 Offset, float DeltaTime) OVERRIDE;
+	virtual void Update(FParticleEmitterInstance* Owner, int32 Offset, float DeltaTime) override;
 	//End UParticleModule Interface
 };
 
diff --git a/Engine/Source/Runtime/Engine/Classes/Particles/Attractor/ParticleModuleAttractorLine.h b/Engine/Source/Runtime/Engine/Classes/Particles/Attractor/ParticleModuleAttractorLine.h
--- a/Engine/Source/Runtime/Engine/Classes/Particles/Attractor/ParticleModuleAttractorLine.h
+++ b/Engine/Source/Runtime/Engine/Classes/Particles/Attractor/ParticleModuleAttractorLine.h
@@ -31,15 +31,15 @@ class UParticleModuleAttractorLine : public UParticleModuleAttractorBase
 
 	//Begin UObject Interface
 #if WITH_EDITOR
-	virtual void PostEditChangeProperty(FPropertyChangedEvent& PropertyChangedEvent) OVERRIDE;
+	virtual void PostEditChangeProperty(FPropertyChangedEvent& PropertyChangedEvent) override;
 #endif // WITH_EDITOR
-	virtual void PostInitProperties() OVERRIDE;
-	virtual void Serialize(FArchive& Ar) OVERRIDE;
+	virtual void PostInitProperties() override;
+	virtual void Serialize(FArchive& Ar) override;
 	//End UObject Interface
 
 	//Begin UParticleModule Interface
-	virtual void Update(FParticleEmitterInstance* Owner, int32 Offset, float DeltaTime) OVERRIDE;
-	virtual void Render3DPreview(FParticleEmitterInstance* Owner, const FSceneView* View,FPrimitiveDrawInterface* PDI) OVERRIDE;
+	virtual void Update(FParticleEmitterInstance* Owner, int32 Offset, float DeltaTime) override;
+	virtual void Render3DPreview(FParticleEmitterInstance* Owner, const FSceneView* View,FPrimitiveDrawInterface* PDI) override;
 	//End UParticleModule Interface
 };
 
diff --git a/Engine/Source/Runtime/Engine/Classes/Particles/Attractor/ParticleModuleAttractorParticle.h b/Engine/Source/Runtime/Engine/Classes/Particles/Attractor/ParticleModuleAttractorParticle.h
--- a/Engine/Source/Runtime/Engine/Classes/Particles/Attractor/ParticleModuleAttractorParticle.h
+++ b/Engine/Source/Runtime/Engine/Classes/Particles/Attractor/ParticleModuleAttractorParticle.h
@@ -77,16 +77,16 @@ class UParticleModuleAttractorParticle : public UParticleModuleAttractorBase
 
 	//Begin UObject Interface
 #if WITH_EDITOR
-	virtual void	PostEditChangeProperty(FPropertyChangedEvent& PropertyChangedEvent) OVERRIDE;
+	virtual void	PostEditChangeProperty(FPropertyChangedEvent& PropertyChangedEvent) override;
 #endif // WITH_EDITOR
-	virtual void PostInitProperties() OVERRIDE;
-	virtual void Serialize(FArchive& Ar) OVERRIDE;
+	virtual void PostInitProperties() override;
+	virtual void Serialize(FArchive& Ar) override;
 	//End UObject Interface
 
 	//Begin UParticleModule Interface
-	virtual void Spawn(FParticleEmitterInstance* Owner, int32 Offset, float SpawnTime, FBaseParticle* ParticleBase) OVERRIDE;
-	virtual void Update(FParticleEmitterInstance* Owner, int32 Offset, float DeltaTime) OVERRIDE;
-	virtual uint32 RequiredBytes(FParticleEmitterInstance* Owner = NULL) OVERRIDE;
+	virtual void Spawn(FParticleEmitterInstance* Owner, int32 Offset, float SpawnTime, FBaseParticle* ParticleBase) override;
+	virtual void Update(FParticleEmitterInstance* Owner, int32 Offset, float DeltaTime) override;
+	virtual uint32 RequiredBytes(FParticleEmitterInstance* Owner = NULL) override;
 	//End UParticleModule Interface
 };
 
diff --git a/Engine/Source/Runtime/Engine/Classes/Particles/Attractor/ParticleModuleAttractorPoint.h b/Engine/Source/Runtime/Engine/Classes/Particles/Attractor/ParticleModuleAttractorPoint.h
--- a/Engine/Source/Runtime/Engine/Classes/Particles/Attractor/ParticleModuleAttractorPoint.h
+++ b/Engine/Source/Runtime/Engine/Classes/Particles/Attractor/ParticleModuleAttractorPoint.h
@@ -67,15 +67,15 @@ class UParticleModuleAttractorPoint : public UParticleModuleAttractorBase
 
 	//Begin UObject Interface
 #if WITH_EDITOR
-	virtual void PostEditChangeProperty(FPropertyChangedEvent& PropertyChangedEvent) OVERRIDE;
+	virtual void PostEditChangeProperty(FPropertyChangedEvent& PropertyChangedEvent) override;
 #endif // WITH_EDITOR
-	virtual void PostInitProperties() OVERRIDE;
-	virtual void Serialize(FArchive& Ar) OVERRIDE;
+	virtual void PostInitProperties() override;
+	virtual void Serialize(FArchive& Ar) override;
 	//End UObject Interface
 
 	//Begin UParticleModule Interface
-	virtual void Update(FParticleEmitterInstance* Owner, int32 Offset, float DeltaTime) OVERRIDE;
-	virtual void Render3DPreview(FParticleEmitterInstance* Owner, const FSceneView* View,FPrimitiveDrawInterface* PDI) OVERRIDE;
+	virtual void Update(FParticleEmitterInstance* Owner, int32 Offset, float DeltaTime) override;
+	virtual void Render3DPreview(FParticleEmitterInstance* Owner, const FSceneView* View,FPrimitiveDrawInterface* PDI) override;
 	//End UParticleModule Interface
 };
 
diff --git a/Engine/Source/Runtime/Engine/Classes/Particles/Attractor/ParticleModuleAttractorPointGravity.h b/Engine/Source/Runtime/Engine/Classes/Particles/Attractor/ParticleModuleAttractorPointGravity.h
--- a/Engine/Source/Runtime/Engine/Classes/Particles/Attractor/ParticleModuleAttractorPointGravity.h
+++ b/Engine/Source/Runtime/Engine/Classes/Particles/Attractor/ParticleModuleAttractorPointGravity.h
@@ -31,15 +31,15 @@ class UParticleModuleAttractorPointGravity : public UParticleModuleAttractorBase
 
 	//Begin UObject Interface
 #if WITH_EDITOR
-	virtual void PostEditChangeProperty(FPropertyChangedEvent& PropertyChangedEvent) OVERRIDE;
+	virtual void PostEditChangeProperty(FPropertyChangedEvent& PropertyChangedEvent) override;
 #endif // WITH_EDITOR
-	virtual void PostInitProperties() OVERRIDE;
-	virtual void Serialize(FArchive& Ar) OVERRIDE;
+	virtual void PostInitProperties() override;
+	virtual void Serialize(FArchive& Ar) override;
 	//End UObject Interface
 
 	//Begin UParticleModule Interface
-	virtual void CompileModule(struct FParticleEmitterBuildInfo& EmitterInfo) OVERRIDE;
-	virtual void Render3DPreview(FParticleEmitterInstance* Owner, const FSceneView* View,FPrimitiveDrawInterface* PDI) OVERRIDE;
+	virtual void CompileModule(struct FParticleEmitterBuildInfo& EmitterInfo) override;
+	virtual void Render3DPreview(FParticleEmitterInstance* Owner, const FSceneView* View,FPrimitiveDrawInterface* PDI) override;
 	//End UParticleModule Interface
 };
 
diff --git a/Engine/Source/Runtime/Engine/Classes/Particles/Beam/ParticleModuleBeamBase.h b/Engine/Source/Runtime/Engine/Classes/Particles/Beam/ParticleModuleBeamBase.h
--- a/Engine/Source/Runtime/Engine/Classes/Particles/Beam/ParticleModuleBeamBase.h
+++ b/Engine/Source/Runtime/Engine/Classes/Particles/Beam/ParticleModuleBeamBase.h
@@ -59,9 +59,9 @@ class UParticleModuleBeamBase : public UParticleModule
 
 
 	//Begin UParticleModule Interface
-	virtual EModuleType	GetModuleType() const OVERRIDE {	return EPMT_Beam;	}
+	virtual EModuleType	GetModuleType() const override {	return EPMT_Beam;	}
 
-	virtual bool CanTickInAnyThread() OVERRIDE
+	virtual bool CanTickInAnyThread() override
 	{
 		return false;
 	}
diff --git a/Engine/Source/Runtime/Engine/Classes/Particles/Beam/ParticleModuleBeamModifier.h b/Engine/Source/Runtime/Engine/Classes/Particles/Beam/ParticleModuleBeamModifier.h
--- a/Engine/Source/Runtime/Engine/Classes/Particles/Beam/ParticleModuleBeamModifier.h
+++ b/Engine/Source/Runtime/Engine/Classes/Particles/Beam/ParticleModuleBeamModifier.h
@@ -91,20 +91,20 @@ class UParticleModuleBeamModifier : public UParticleModuleBeamBase
 
 	//Begin UObject Interface
 #if WITH_EDITOR
-	virtual void PostEditChangeProperty(FPropertyChangedEvent& PropertyChangedEvent) OVERRIDE;
+	virtual void PostEditChangeProperty(FPropertyChangedEvent& PropertyChangedEvent) override;
 #endif // WITH_EDITOR
-	virtual void PostInitProperties() OVERRIDE;
-	virtual void Serialize(FArchive& Ar) OVERRIDE;
+	virtual void PostInitProperties() override;
+	virtual void Serialize(FArchive& Ar) override;
 	//End Uobject Interface
 
 	//Begin UParticleModule Interface
-	virtual void Spawn(FParticleEmitterInstance* Owner, int32 Offset, float SpawnTime, FBaseParticle* ParticleBase) OVERRIDE;
-	virtual void Update(FParticleEmitterInstance* Owner, int32 Offset, float DeltaTime) OVERRIDE;
-	virtual uint32 RequiredBytes(FParticleEmitterInstance* Owner = NULL) OVERRIDE;
-	virtual void AutoPopulateInstanceProperties(UParticleSystemComponent* PSysComp) OVERRIDE;
-	virtual void GetCurveObjects(TArray<FParticleCurvePair>& OutCurves) OVERRIDE;
-	virtual	bool AddModuleCurvesToEditor(UInterpCurveEdSetup* EdSetup, TArray<const FCurveEdEntry*>& OutCurveEntries) OVERRIDE;
-	virtual void GetParticleSysParamsUtilized(TArray<FString>& ParticleSysParamList) OVERRIDE;
+	virtual void Spawn(FParticleEmitterInstance* Owner, int32 Offset, float SpawnTime, FBaseParticle* ParticleBase) override;
+	virtual void Update(FParticleEmitterInstance* Owner, int32 Offset, float DeltaTime) override;
+	virtual uint32 RequiredBytes(FParticleEmitterInstance* Owner = NULL) override;
+	virtual void AutoPopulateInstanceProperties(UParticleSystemComponent* PSysComp) override;
+	virtual void GetCurveObjects(TArray<FParticleCurvePair>& OutCurves) override;
+	virtual	bool AddModuleCurvesToEditor(UInterpCurveEdSetup* EdSetup, TArray<const FCurveEdEntry*>& OutCurveEntries) override;
+	virtual void GetParticleSysParamsUtilized(TArray<FString>& ParticleSysParamList) override;
 	//End UParticleModule Interface
 
 
diff --git a/Engine/Source/Runtime/Engine/Classes/Particles/Beam/ParticleModuleBeamNoise.h b/Engine/Source/Runtime/Engine/Classes/Particles/Beam/ParticleModuleBeamNoise.h
--- a/Engine/Source/Runtime/Engine/Classes/Particles/Beam/ParticleModuleBeamNoise.h
+++ b/Engine/Source/Runtime/Engine/Classes/Particles/Beam/ParticleModuleBeamNoise.h
@@ -119,16 +119,16 @@ class UParticleModuleBeamNoise : public UParticleModuleBeamBase
 
 	//Begin UObject Interface
 #if WITH_EDITOR
-	virtual void PostEditChangeProperty(FPropertyChangedEvent& PropertyChangedEvent) OVERRIDE;
+	virtual void PostEditChangeProperty(FPropertyChangedEvent& PropertyChangedEvent) override;
 #endif // WITH_EDITOR
-	virtual void PostInitProperties() OVERRIDE;
-	virtual void Serialize(FArchive& Ar) OVERRIDE;
+	virtual void PostInitProperties() override;
+	virtual void Serialize(FArchive& Ar) override;
 	//End UObject Interface
 
 	//Begin UParticleModule Interface
-	virtual void Spawn(FParticleEmitterInstance* Owner, int32 Offset, float SpawnTime, FBaseParticle* ParticleBase) OVERRIDE;
-	virtual void Update(FParticleEmitterInstance* Owner, int32 Offset, float DeltaTime) OVERRIDE;
-	virtual void SetToSensibleDefaults(UParticleEmitter* Owner) OVERRIDE;
+	virtual void Spawn(FParticleEmitterInstance* Owner, int32 Offset, float SpawnTime, FBaseParticle* ParticleBase) override;
+	virtual void Update(FParticleEmitterInstance* Owner, int32 Offset, float DeltaTime) override;
+	virtual void SetToSensibleDefaults(UParticleEmitter* Owner) override;
 	//End UParticleModule Interface
 
 	// @todo document
diff --git a/Engine/Source/Runtime/Engine/Classes/Particles/Beam/ParticleModuleBeamSource.h b/Engine/Source/Runtime/Engine/Classes/Particles/Beam/ParticleModuleBeamSource.h
--- a/Engine/Source/Runtime/Engine/Classes/Particles/Beam/ParticleModuleBeamSource.h
+++ b/Engine/Source/Runtime/Engine/Classes/Particles/Beam/ParticleModuleBeamSource.h
@@ -61,18 +61,18 @@ class UParticleModuleBeamSource : public UParticleModuleBeamBase
 
 	//Begin UObject Interface
 #if WITH_EDITOR
-	virtual void PostEditChangeProperty(FPropertyChangedEvent& PropertyChangedEvent) OVERRIDE;
+	virtual void PostEditChangeProperty(FPropertyChangedEvent& PropertyChangedEvent) override;
 #endif // WITH_EDITOR
-	virtual void PostInitProperties() OVERRIDE;
-	virtual void Serialize(FArchive& Ar) OVERRIDE;
+	virtual void PostInitProperties() override;
+	virtual void Serialize(FArchive& Ar) override;
 	//End UObject Interface
 
 	//Begin UParticleModule Interface
-	virtual void Spawn(FParticleEmitterInstance* Owner, int32 Offset, float SpawnTime, FBaseParticle* ParticleBase) OVERRIDE;
-	virtual void Update(FParticleEmitterInstance* Owner, int32 Offset, float DeltaTime) OVERRIDE;
-	virtual uint32 RequiredBytes(FParticleEmitterInstance* Owner = NULL) OVERRIDE;
-	virtual void AutoPopulateInstanceProperties(UParticleSystemComponent* PSysComp) OVERRIDE;
-	virtual void GetParticleSysParamsUtilized(TArray<FString>& ParticleSysParamList) OVERRIDE;
+	virtual void Spawn(FParticleEmitterInstance* Owner, int32 Offset, float SpawnTime, FBaseParticle* ParticleBase) override;
+	virtual void Update(FParticleEmitterInstance* Owner, int32 Offset, float DeltaTime) override;
+	virtual uint32 RequiredBytes(FParticleEmitterInstance* Owner = NULL) override;
+	virtual void AutoPopulateInstanceProperties(UParticleSystemComponent* PSysComp) override;
+	virtual void GetParticleSysParamsUtilized(TArray<FString>& ParticleSysParamList) override;
 	//End UParticleModule Interface
 
 
diff --git a/Engine/Source/Runtime/Engine/Classes/Particles/Beam/ParticleModuleBeamTarget.h b/Engine/Source/Runtime/Engine/Classes/Particles/Beam/ParticleModuleBeamTarget.h
--- a/Engine/Source/Runtime/Engine/Classes/Particles/Beam/ParticleModuleBeamTarget.h
+++ b/Engine/Source/Runtime/Engine/Classes/Particles/Beam/ParticleModuleBeamTarget.h
@@ -65,17 +65,17 @@ class UParticleModuleBeamTarget : public UParticleModuleBeamBase
 
 	//Begin UObject Interface
 #if WITH_EDITOR
-	virtual void PostEditChangeProperty(FPropertyChangedEvent& PropertyChangedEvent) OVERRIDE;
+	virtual void PostEditChangeProperty(FPropertyChangedEvent& PropertyChangedEvent) override;
 #endif // WITH_EDITOR
-	virtual void PostInitProperties() OVERRIDE;
-	virtual void Serialize(FArchive& Ar) OVERRIDE;
+	virtual void PostInitProperties() override;
+	virtual void Serialize(FArchive& Ar) override;
 	//End UObject Interface
 
 	//Begin UParticleModule Interface
-	virtual void Spawn(FParticleEmitterInstance* Owner, int32 Offset, float SpawnTime, FBaseParticle* ParticleBase) OVERRIDE;
-	virtual void Update(FParticleEmitterInstance* Owner, int32 Offset, float DeltaTime) OVERRIDE;
-	virtual void AutoPopulateInstanceProperties(UParticleSystemComponent* PSysComp) OVERRIDE;
-	virtual void GetParticleSysParamsUtilized(TArray<FString>& ParticleSysParamList) OVERRIDE;
+	virtual void Spawn(FParticleEmitterInstance* Owner, int32 Offset, float SpawnTime, FBaseParticle* ParticleBase) override;
+	virtual void Update(FParticleEmitterInstance* Owner, int32 Offset, float DeltaTime) override;
+	virtual void AutoPopulateInstanceProperties(UParticleSystemComponent* PSysComp) override;
+	virtual void GetParticleSysParamsUtilized(TArray<FString>& ParticleSysParamList) override;
 	//End UParticleModule Interface
 
 	// @todo document
diff --git a/Engine/Source/Runtime/Engine/Classes/Particles/Camera/ParticleModuleCameraBase.h b/Engine/Source/Runtime/Engine/Classes/Particles/Camera/ParticleModuleCameraBase.h
--- a/Engine/Source/Runtime/Engine/Classes/Particles/Camera/ParticleModuleCameraBase.h
+++ b/Engine/Source/Runtime/Engine/Classes/Particles/Camera/ParticleModuleCameraBase.h
@@ -10,7 +10,7 @@ class UParticleModuleCameraBase : public UParticleModule
 {
 	GENERATED_UCLASS_BODY()
 
-	virtual bool CanTickInAnyThread() OVERRIDE
+	virtual bool CanTickInAnyThread() override
 	{
 		return false;
 	}
diff --git a/Engine/Source/Runtime/Engine/Classes/Particles/Camera/ParticleModuleCameraOffset.h b/Engine/Source/Runtime/Engine/Classes/Particles/Camera/ParticleModuleCameraOffset.h
--- a/Engine/Source/Runtime/Engine/Classes/Particles/Camera/ParticleModuleCameraOffset.h
+++ b/Engine/Source/Runtime/Engine/Classes/Particles/Camera/ParticleModuleCameraOffset.h
@@ -46,16 +46,16 @@ class UParticleModuleCameraOffset : public UParticleModuleCameraBase
 
 	//Begin UObject Interface
 #if WITH_EDITOR
-	virtual void PostEditChangeProperty(FPropertyChangedEvent& PropertyChangedEvent) OVERRIDE;
+	virtual void PostEditChangeProperty(FPropertyChangedEvent& PropertyChangedEvent) override;
 #endif // WITH_EDITOR
-	virtual void PostInitProperties() OVERRIDE;
-	virtual void Serialize(FArchive& Ar) OVERRIDE;
+	virtual void PostInitProperties() override;
+	virtual void Serialize(FArchive& Ar) override;
 	//End UObject Interface
 
 	//Begin UParticleModule Interface
-	virtual void Spawn(FParticleEmitterInstance* Owner, int32 Offset, float SpawnTime, FBaseParticle* ParticleBase) OVERRIDE;
-	virtual void Update(FParticleEmitterInstance* Owner, int32 Offset, float DeltaTime) OVERRIDE;
-	virtual uint32 RequiredBytes(FParticleEmitterInstance* Owner = NULL) OVERRIDE;
+	virtual void Spawn(FParticleEmitterInstance* Owner, int32 Offset, float SpawnTime, FBaseParticle* ParticleBase) override;
+	virtual void Update(FParticleEmitterInstance* Owner, int32 Offset, float DeltaTime) override;
+	virtual uint32 RequiredBytes(FParticleEmitterInstance* Owner = NULL) override;
 	//End UParticleModule Interface
 };
 
diff --git a/Engine/Source/Runtime/Engine/Classes/Particles/Collision/ParticleModuleCollision.h b/Engine/Source/Runtime/Engine/Classes/Particles/Collision/ParticleModuleCollision.h
--- a/Engine/Source/Runtime/Engine/Classes/Particles/Collision/ParticleModuleCollision.h
+++ b/Engine/Source/Runtime/Engine/Classes/Particles/Collision/ParticleModuleCollision.h
@@ -135,21 +135,21 @@ class UParticleModuleCollision : public UParticleModuleCollisionBase
 
 	//Begin UObject Interface
 #if WITH_EDITOR
-	virtual void PostEditChangeProperty(FPropertyChangedEvent& PropertyChangedEvent) OVERRIDE;
+	virtual void PostEditChangeProperty(FPropertyChangedEvent& PropertyChangedEvent) override;
 #endif // WITH_EDITOR
-	virtual void PostInitProperties() OVERRIDE;
-	virtual void Serialize(FArchive& Ar) OVERRIDE;
+	virtual void PostInitProperties() override;
+	virtual void Serialize(FArchive& Ar) override;
 	//End UObject Interface
 
 	//Begin UParticleModule Interface
-	virtual void Spawn(FParticleEmitterInstance* Owner, int32 Offset, float SpawnTime, FBaseParticle* ParticleBase) OVERRIDE;
-	virtual void Update(FParticleEmitterInstance* Owner, int32 Offset, float DeltaTime) OVERRIDE;
-	virtual uint32 RequiredBytes(FParticleEmitterInstance* Owner = NULL) OVERRIDE;
-	virtual uint32 RequiredBytesPerInstance(FParticleEmitterInstance* Owner = NULL) OVERRIDE;
-	virtual uint32 PrepPerInstanceBlock(FParticleEmitterInstance* Owner, void* InstData) OVERRIDE;
-	virtual void SetToSensibleDefaults(UParticleEmitter* Owner) OVERRIDE;
-	virtual bool GenerateLODModuleValues(UParticleModule* SourceModule, float Percentage, UParticleLODLevel* LODLevel) OVERRIDE;
-	virtual bool CanTickInAnyThread() OVERRIDE
+	virtual void Spawn(FParticleEmitterInstance* Owner, int32 Offset, float SpawnTime, FBaseParticle* ParticleBase) override;
+	virtual void Update(FParticleEmitterInstance* Owner, int32 Offset, float DeltaTime) override;
+	virtual uint32 RequiredBytes(FParticleEmitterInstance* Owner = NULL) override;
+	virtual uint32 RequiredBytesPerInstance(FParticleEmitterInstance* Owner = NULL) override;
+	virtual uint32 PrepPerInstanceBlock(FParticleEmitterInstance* Owner, void* InstData) override;
+	virtual void SetToSensibleDefaults(UParticleEmitter* Owner) override;
+	virtual bool GenerateLODModuleValues(UParticleModule* SourceModule, float Percentage, UParticleLODLevel* LODLevel) override;
+	virtual bool CanTickInAnyThread() override
 	{
 		return false;
 	}
diff --git a/Engine/Source/Runtime/Engine/Classes/Particles/Collision/ParticleModuleCollisionGPU.h b/Engine/Source/Runtime/Engine/Classes/Particles/Collision/ParticleModuleCollisionGPU.h
--- a/Engine/Source/Runtime/Engine/Classes/Particles/Collision/ParticleModuleCollisionGPU.h
+++ b/Engine/Source/Runtime/Engine/Classes/Particles/Collision/ParticleModuleCollisionGPU.h
@@ -69,16 +69,16 @@ class UParticleModuleCollisionGPU : public UParticleModuleCollisionBase
 
 	//Begin UObject Interface
 #if WITH_EDITOR
-	virtual void PostEditChangeProperty(FPropertyChangedEvent& PropertyChangedEvent) OVERRIDE;
+	virtual void PostEditChangeProperty(FPropertyChangedEvent& PropertyChangedEvent) override;
 #endif // WITH_EDITOR
-	virtual void PostInitProperties() OVERRIDE;
+	virtual void PostInitProperties() override;
 	//End UObject Interface
 
 	//Begin UParticleModule Interface
-	virtual void SetToSensibleDefaults(UParticleEmitter* Owner) OVERRIDE;
-	virtual void CompileModule(struct FParticleEmitterBuildInfo& EmitterInfo) OVERRIDE;
+	virtual void SetToSensibleDefaults(UParticleEmitter* Owner) override;
+	virtual void CompileModule(struct FParticleEmitterBuildInfo& EmitterInfo) override;
 #if WITH_EDITOR
-	virtual bool IsValidForLODLevel(UParticleLODLevel* LODLevel, FString& OutErrorString) OVERRIDE;
+	virtual bool IsValidForLODLevel(UParticleLODLevel* LODLevel, FString& OutErrorString) override;
 #endif
 	//End UParticleModule Interface
 
diff --git a/Engine/Source/Runtime/Engine/Classes/Particles/Color/ParticleModuleColor.h b/Engine/Source/Runtime/Engine/Classes/Particles/Color/ParticleModuleColor.h
--- a/Engine/Source/Runtime/Engine/Classes/Particles/Color/ParticleModuleColor.h
+++ b/Engine/Source/Runtime/Engine/Classes/Particles/Color/ParticleModuleColor.h
@@ -27,18 +27,18 @@ class UParticleModuleColor : public UParticleModuleColorBase
 
 	// Begin UObject Interface
 #if WITH_EDITOR
-	virtual	void PostEditChangeProperty(FPropertyChangedEvent& PropertyChangedEvent) OVERRIDE;
+	virtual	void PostEditChangeProperty(FPropertyChangedEvent& PropertyChangedEvent) override;
 #endif // WITH_EDITOR
-	virtual void PostInitProperties() OVERRIDE;
-	virtual void Serialize(FArchive& Ar) OVERRIDE;
+	virtual void PostInitProperties() override;
+	virtual void Serialize(FArchive& Ar) override;
 	// End UObject Interface
 
 
 	//Begin UParticleModule Interface
-	virtual	bool AddModuleCurvesToEditor(UInterpCurveEdSetup* EdSetup, TArray<const FCurveEdEntry*>& OutCurveEntries) OVERRIDE;
-	virtual void Spawn(FParticleEmitterInstance* Owner, int32 Offset, float SpawnTime, FBaseParticle* ParticleBase) OVERRIDE;
-	virtual void CompileModule( struct FParticleEmitterBuildInfo& EmitterInfo ) OVERRIDE;
-	virtual void SetToSensibleDefaults(UParticleEmitter* Owner) OVERRIDE;
+	virtual	bool AddModuleCurvesToEditor(UInterpCurveEdSetup* EdSetup, TArray<const FCurveEdEntry*>& OutCurveEntries) override;
+	virtual void Spawn(FParticleEmitterInstance* Owner, int32 Offset, float SpawnTime, FBaseParticle* ParticleBase) override;
+	virtual void CompileModule( struct FParticleEmitterBuildInfo& EmitterInfo ) override;
+	virtual void SetToSensibleDefaults(UParticleEmitter* Owner) override;
 	//End UParticleModule Interface
 
 	/**
diff --git a/Engine/Source/Runtime/Engine/Classes/Particles/Color/ParticleModuleColorOverLife.h b/Engine/Source/Runtime/Engine/Classes/Particles/Color/ParticleModuleColorOverLife.h
--- a/Engine/Source/Runtime/Engine/Classes/Particles/Color/ParticleModuleColorOverLife.h
+++ b/Engine/Source/Runtime/Engine/Classes/Particles/Color/ParticleModuleColorOverLife.h
@@ -27,19 +27,19 @@ class UParticleModuleColorOverLife : public UParticleModuleColorBase
 
 	// Begin UObject Interface
 #if WITH_EDITOR
-	virtual	void PostEditChangeProperty(FPropertyChangedEvent& PropertyChangedEvent) OVERRIDE;
+	virtual	void PostEditChangeProperty(FPropertyChangedEvent& PropertyChangedEvent) override;
 #endif // WITH_EDITOR
-	virtual void PostInitProperties() OVERRIDE;
-	virtual void Serialize(FArchive& Ar) OVERRIDE;
+	virtual void PostInitProperties() override;
+	virtual void Serialize(FArchive& Ar) override;
 	// End UObject Interface
 
 
 	//Begin UParticleModule Interface
-	virtual	bool AddModuleCurvesToEditor(UInterpCurveEdSetup* EdSetup, TArray<const FCurveEdEntry*>& OutCurveEntries) OVERRIDE;
-	virtual void Spawn(FParticleEmitterInstance* Owner, int32 Offset, float SpawnTime, FBaseParticle* ParticleBase) OVERRIDE;
-	virtual void Update(FParticleEmitterInstance* Owner, int32 Offset, float DeltaTime) OVERRIDE;
-	virtual void CompileModule( struct FParticleEmitterBuildInfo& EmitterInfo ) OVERRIDE;
-	virtual void SetToSensibleDefaults(UParticleEmitter* Owner) OVERRIDE;
+	virtual	bool AddModuleCurvesToEditor(UInterpCurveEdSetup* EdSetup, TArray<const FCurveEdEntry*>& OutCurveEntries) override;
+	virtual void Spawn(FParticleEmitterInstance* Owner, int32 Offset, float SpawnTime, FBaseParticle* ParticleBase) override;
+	virtual void Update(FParticleEmitterInstance* Owner, int32 Offset, float DeltaTime) override;
+	virtual void CompileModule( struct FParticleEmitterBuildInfo& EmitterInfo ) override;
+	virtual void SetToSensibleDefaults(UParticleEmitter* Owner) override;
 	//End UParticleModule Interface
 
 };
diff --git a/Engine/Source/Runtime/Engine/Classes/Particles/Color/ParticleModuleColorScaleOverLife.h b/Engine/Source/Runtime/Engine/Classes/Particles/Color/ParticleModuleColorScaleOverLife.h
--- a/Engine/Source/Runtime/Engine/Classes/Particles/Color/ParticleModuleColorScaleOverLife.h
+++ b/Engine/Source/Runtime/Engine/Classes/Particles/Color/ParticleModuleColorScaleOverLife.h
@@ -33,21 +33,21 @@ class UParticleModuleColorScaleOverLife : public UParticleModuleColorBase
 
 	//Begin UObject Interface
 #if WITH_EDITOR
-	virtual void PostEditChangeProperty(FPropertyChangedEvent& PropertyChangedEvent) OVERRIDE;
+	virtual void PostEditChangeProperty(FPropertyChangedEvent& PropertyChangedEvent) override;
 #endif // WITH_EDITOR
-	virtual void PostInitProperties() OVERRIDE;
-	virtual void Serialize(FArchive& Ar) OVERRIDE;
+	virtual void PostInitProperties() override;
+	virtual void Serialize(FArchive& Ar) override;
 	//End UObject Interface
 
 	//Begin UParticleModule Interface
-	virtual void CompileModule( struct FParticleEmitterBuildInfo& EmitterInfo ) OVERRIDE;
-	virtual void Spawn(FParticleEmitterInstance* Owner, int32 Offset, float SpawnTime, FBaseParticle* ParticleBase) OVERRIDE;
-	virtual void Update(FParticleEmitterInstance* Owner, int32 Offset, float DeltaTime) OVERRIDE;
-	virtual void SetToSensibleDefaults(UParticleEmitter* Owner) OVERRIDE;
+	virtual void CompileModule( struct FParticleEmitterBuildInfo& EmitterInfo ) override;
+	virtual void Spawn(FParticleEmitterInstance* Owner, int32 Offset, float SpawnTime, FBaseParticle* ParticleBase) override;
+	virtual void Update(FParticleEmitterInstance* Owner, int32 Offset, float DeltaTime) override;
+	virtual void SetToSensibleDefaults(UParticleEmitter* Owner) override;
 #if WITH_EDITOR
-	virtual int32 GetNumberOfCustomMenuOptions() const OVERRIDE;
-	virtual bool GetCustomMenuEntryDisplayString(int32 InEntryIndex, FString& OutDisplayString) const OVERRIDE;
-	virtual bool PerformCustomMenuEntry(int32 InEntryIndex) OVERRIDE;
+	virtual int32 GetNumberOfCustomMenuOptions() const override;
+	virtual bool GetCustomMenuEntryDisplayString(int32 InEntryIndex, FString& OutDisplayString) const override;
+	virtual bool PerformCustomMenuEntry(int32 InEntryIndex) override;
 #endif
 	//End UParticleModule Interface
 };
diff --git a/Engine/Source/Runtime/Engine/Classes/Particles/Color/ParticleModuleColor_Seeded.h b/Engine/Source/Runtime/Engine/Classes/Particles/Color/ParticleModuleColor_Seeded.h
--- a/Engine/Source/Runtime/Engine/Classes/Particles/Color/ParticleModuleColor_Seeded.h
+++ b/Engine/Source/Runtime/Engine/Classes/Particles/Color/ParticleModuleColor_Seeded.h
@@ -16,14 +16,14 @@ class UParticleModuleColor_Seeded : public UParticleModuleColor
 
 
 	//Begin UParticleModule Interface
-	virtual void Spawn(FParticleEmitterInstance* Owner, int32 Offset, float SpawnTime, FBaseParticle* ParticleBase) OVERRIDE;
-	virtual uint32	RequiredBytesPerInstance(FParticleEmitterInstance* Owner = NULL) OVERRIDE;
-	virtual uint32	PrepPerInstanceBlock(FParticleEmitterInstance* Owner, void* InstData) OVERRIDE;
-	virtual FParticleRandomSeedInfo* GetRandomSeedInfo() OVERRIDE
+	virtual void Spawn(FParticleEmitterInstance* Owner, int32 Offset, float SpawnTime, FBaseParticle* ParticleBase) override;
+	virtual uint32	RequiredBytesPerInstance(FParticleEmitterInstance* Owner = NULL) override;
+	virtual uint32	PrepPerInstanceBlock(FParticleEmitterInstance* Owner, void* InstData) override;
+	virtual FParticleRandomSeedInfo* GetRandomSeedInfo() override
 	{
 		return &RandomSeedInfo;
 	}
-	virtual void EmitterLoopingNotify(FParticleEmitterInstance* Owner) OVERRIDE;
+	virtual void EmitterLoopingNotify(FParticleEmitterInstance* Owner) override;
 	//End UParticleModule Interface
 };
 
diff --git a/Engine/Source/Runtime/Engine/Classes/Particles/Emitter.h b/Engine/Source/Runtime/Engine/Classes/Particles/Emitter.h
--- a/Engine/Source/Runtime/Engine/Classes/Particles/Emitter.h
+++ b/Engine/Source/Runtime/Engine/Classes/Particles/Emitter.h
@@ -118,16 +118,16 @@ class ENGINE_API AEmitter : public AActor
 	void AutoPopulateInstanceProperties();
 
 	// Begin UObject Interface
-	virtual FString GetDetailedInfoInternal() const OVERRIDE;
+	virtual FString GetDetailedInfoInternal() const override;
 	// End UObject Interface
 
 
 	// Begin AActor interface
-	virtual void PostActorCreated() OVERRIDE;
-	virtual void PostInitializeComponents() OVERRIDE;
+	virtual void PostActorCreated() override;
+	virtual void PostInitializeComponents() override;
 #if WITH_EDITOR
-	virtual void CheckForErrors() OVERRIDE;
-	virtual bool GetReferencedContentObjects( TArray<UObject*>& Objects ) const OVERRIDE;
+	virtual void CheckForErrors() override;
+	virtual bool GetReferencedContentObjects( TArray<UObject*>& Objects ) const override;
 	// End AActor interface
 
 	/**
diff --git a/Engine/Source/Runtime/Engine/Classes/Particles/EmitterCameraLensEffectBase.h b/Engine/Source/Runtime/Engine/Classes/Particles/EmitterCameraLensEffectBase.h
--- a/Engine/Source/Runtime/Engine/Classes/Particles/EmitterCameraLensEffectBase.h
+++ b/Engine/Source/Runtime/Engine/Classes/Particles/EmitterCameraLensEffectBase.h
@@ -56,8 +56,8 @@ public:
 
 
 	// Begin AActor Interface
-	ENGINE_API virtual void Destroyed() OVERRIDE;
-	ENGINE_API virtual void PostInitializeComponents() OVERRIDE;
+	ENGINE_API virtual void Destroyed() override;
+	ENGINE_API virtual void PostInitializeComponents() override;
 	// End AActor Interface
 
 	/** Tell the emitter what camera it is attached to. */
diff --git a/Engine/Source/Runtime/Engine/Classes/Particles/Event/ParticleModuleEventBase.h b/Engine/Source/Runtime/Engine/Classes/Particles/Event/ParticleModuleEventBase.h
--- a/Engine/Source/Runtime/Engine/Classes/Particles/Event/ParticleModuleEventBase.h
+++ b/Engine/Source/Runtime/Engine/Classes/Particles/Event/ParticleModuleEventBase.h
@@ -12,7 +12,7 @@ class UParticleModuleEventBase : public UParticleModule
 
 
 	// Begin UParticleModule Interface
-	virtual EModuleType	GetModuleType() const OVERRIDE {	return EPMT_Event;	}
+	virtual EModuleType	GetModuleType() const override {	return EPMT_Event;	}
 	// End UParticleModule Interface
 };
 
diff --git a/Engine/Source/Runtime/Engine/Classes/Particles/Event/ParticleModuleEventGenerator.h b/Engine/Source/Runtime/Engine/Classes/Particles/Event/ParticleModuleEventGenerator.h
--- a/Engine/Source/Runtime/Engine/Classes/Particles/Event/ParticleModuleEventGenerator.h
+++ b/Engine/Source/Runtime/Engine/Classes/Particles/Event/ParticleModuleEventGenerator.h
@@ -72,16 +72,16 @@ class UParticleModuleEventGenerator : public UParticleModuleEventBase
 
 	//Begin UObject Interface
 #if WITH_EDITOR
-	virtual void	PostEditChangeProperty(FPropertyChangedEvent& PropertyChangedEvent) OVERRIDE;
+	virtual void	PostEditChangeProperty(FPropertyChangedEvent& PropertyChangedEvent) override;
 #endif // WITH_EDITOR
 	//End UObject Interface
 
 	//Begin UParticleModule Interface
-	virtual void	Spawn(FParticleEmitterInstance* Owner, int32 Offset, float SpawnTime, FBaseParticle* ParticleBase) OVERRIDE;
-	virtual void	Update(FParticleEmitterInstance* Owner, int32 Offset, float DeltaTime) OVERRIDE;
-	virtual uint32	RequiredBytes(FParticleEmitterInstance* Owner = NULL) OVERRIDE;
-	virtual uint32	RequiredBytesPerInstance(FParticleEmitterInstance* Owner = NULL) OVERRIDE;
-	virtual uint32	PrepPerInstanceBlock(FParticleEmitterInstance* Owner, void* InstData) OVERRIDE;
+	virtual void	Spawn(FParticleEmitterInstance* Owner, int32 Offset, float SpawnTime, FBaseParticle* ParticleBase) override;
+	virtual void	Update(FParticleEmitterInstance* Owner, int32 Offset, float DeltaTime) override;
+	virtual uint32	RequiredBytes(FParticleEmitterInstance* Owner = NULL) override;
+	virtual uint32	RequiredBytesPerInstance(FParticleEmitterInstance* Owner = NULL) override;
+	virtual uint32	PrepPerInstanceBlock(FParticleEmitterInstance* Owner, void* InstData) override;
 	//End UParticleModule Interface
 
 
