UE4 Refactoring. Changed OVERRIDE and FINAL macros to keywords override and final.

diff --git a/Engine/Source/Runtime/Engine/Classes/Particles/Event/ParticleModuleEventReceiverKillParticles.h b/Engine/Source/Runtime/Engine/Classes/Particles/Event/ParticleModuleEventReceiverKillParticles.h
--- a/Engine/Source/Runtime/Engine/Classes/Particles/Event/ParticleModuleEventReceiverKillParticles.h
+++ b/Engine/Source/Runtime/Engine/Classes/Particles/Event/ParticleModuleEventReceiverKillParticles.h
@@ -16,7 +16,7 @@ class UParticleModuleEventReceiverKillParticles : public UParticleModuleEventRec
 
 
 	// Begin UParticleModuleEventBase Interface
-	virtual bool ProcessParticleEvent(FParticleEmitterInstance* Owner, FParticleEventData& InEvent, float DeltaTime) OVERRIDE;
+	virtual bool ProcessParticleEvent(FParticleEmitterInstance* Owner, FParticleEventData& InEvent, float DeltaTime) override;
 	// End UParticleModuleEventBase Interface
 };
 
diff --git a/Engine/Source/Runtime/Engine/Classes/Particles/Event/ParticleModuleEventReceiverSpawn.h b/Engine/Source/Runtime/Engine/Classes/Particles/Event/ParticleModuleEventReceiverSpawn.h
--- a/Engine/Source/Runtime/Engine/Classes/Particles/Event/ParticleModuleEventReceiverSpawn.h
+++ b/Engine/Source/Runtime/Engine/Classes/Particles/Event/ParticleModuleEventReceiverSpawn.h
@@ -48,14 +48,14 @@ class UParticleModuleEventReceiverSpawn : public UParticleModuleEventReceiverBas
 
 	//Begin UObject Interface
 #if WITH_EDITOR
-	virtual void PostEditChangeProperty(FPropertyChangedEvent& PropertyChangedEvent) OVERRIDE;
+	virtual void PostEditChangeProperty(FPropertyChangedEvent& PropertyChangedEvent) override;
 #endif // WITH_EDITOR
-	virtual void PostInitProperties() OVERRIDE;
-	virtual void Serialize(FArchive& Ar) OVERRIDE;
+	virtual void PostInitProperties() override;
+	virtual void Serialize(FArchive& Ar) override;
 	//End UObject Interface
 
 	// Begin UParticleModuleEventBase Interface
-	virtual bool ProcessParticleEvent(FParticleEmitterInstance* Owner, FParticleEventData& InEvent, float DeltaTime) OVERRIDE;
+	virtual bool ProcessParticleEvent(FParticleEmitterInstance* Owner, FParticleEventData& InEvent, float DeltaTime) override;
 	// End UParticleModuleEventBase Interface
 };
 
diff --git a/Engine/Source/Runtime/Engine/Classes/Particles/Kill/ParticleModuleKillBox.h b/Engine/Source/Runtime/Engine/Classes/Particles/Kill/ParticleModuleKillBox.h
--- a/Engine/Source/Runtime/Engine/Classes/Particles/Kill/ParticleModuleKillBox.h
+++ b/Engine/Source/Runtime/Engine/Classes/Particles/Kill/ParticleModuleKillBox.h
@@ -38,15 +38,15 @@ class UParticleModuleKillBox : public UParticleModuleKillBase
 
 	//Begin UObject Interface
 #if WITH_EDITOR
-	virtual void PostEditChangeProperty(FPropertyChangedEvent& PropertyChangedEvent) OVERRIDE;
+	virtual void PostEditChangeProperty(FPropertyChangedEvent& PropertyChangedEvent) override;
 #endif // WITH_EDITOR
-	virtual void PostInitProperties() OVERRIDE;
-	virtual void Serialize(FArchive& Ar) OVERRIDE;
+	virtual void PostInitProperties() override;
+	virtual void Serialize(FArchive& Ar) override;
 	//End UObject Interface
 
 	//Begin UParticleModule Interface
-	virtual void	Update(FParticleEmitterInstance* Owner, int32 Offset, float DeltaTime) OVERRIDE;
-	virtual void	Render3DPreview(FParticleEmitterInstance* Owner, const FSceneView* View,FPrimitiveDrawInterface* PDI) OVERRIDE;
+	virtual void	Update(FParticleEmitterInstance* Owner, int32 Offset, float DeltaTime) override;
+	virtual void	Render3DPreview(FParticleEmitterInstance* Owner, const FSceneView* View,FPrimitiveDrawInterface* PDI) override;
 	//End UParticleModule Interface
 };
 
diff --git a/Engine/Source/Runtime/Engine/Classes/Particles/Kill/ParticleModuleKillHeight.h b/Engine/Source/Runtime/Engine/Classes/Particles/Kill/ParticleModuleKillHeight.h
--- a/Engine/Source/Runtime/Engine/Classes/Particles/Kill/ParticleModuleKillHeight.h
+++ b/Engine/Source/Runtime/Engine/Classes/Particles/Kill/ParticleModuleKillHeight.h
@@ -34,15 +34,15 @@ class UParticleModuleKillHeight : public UParticleModuleKillBase
 
 	//Begin UObject Interface
 #if WITH_EDITOR
-	virtual void PostEditChangeProperty(FPropertyChangedEvent& PropertyChangedEvent) OVERRIDE;
+	virtual void PostEditChangeProperty(FPropertyChangedEvent& PropertyChangedEvent) override;
 #endif // WITH_EDITOR
-	virtual void PostInitProperties() OVERRIDE;
-	virtual void Serialize(FArchive& Ar) OVERRIDE;
+	virtual void PostInitProperties() override;
+	virtual void Serialize(FArchive& Ar) override;
 	//End UObject Interface
 
 	//Begin UParticleModule Interface
-	virtual void	Update(FParticleEmitterInstance* Owner, int32 Offset, float DeltaTime) OVERRIDE;
-	virtual void	Render3DPreview(FParticleEmitterInstance* Owner, const FSceneView* View,FPrimitiveDrawInterface* PDI) OVERRIDE;
+	virtual void	Update(FParticleEmitterInstance* Owner, int32 Offset, float DeltaTime) override;
+	virtual void	Render3DPreview(FParticleEmitterInstance* Owner, const FSceneView* View,FPrimitiveDrawInterface* PDI) override;
 	//End UParticleModule Interface
 };
 
diff --git a/Engine/Source/Runtime/Engine/Classes/Particles/Lifetime/ParticleModuleLifetime.h b/Engine/Source/Runtime/Engine/Classes/Particles/Lifetime/ParticleModuleLifetime.h
--- a/Engine/Source/Runtime/Engine/Classes/Particles/Lifetime/ParticleModuleLifetime.h
+++ b/Engine/Source/Runtime/Engine/Classes/Particles/Lifetime/ParticleModuleLifetime.h
@@ -19,21 +19,21 @@ class UParticleModuleLifetime : public UParticleModuleLifetimeBase
 
 	//Begin UObject Interface
 #if WITH_EDITOR
-	virtual void	PostEditChangeProperty(FPropertyChangedEvent& PropertyChangedEvent) OVERRIDE;
+	virtual void	PostEditChangeProperty(FPropertyChangedEvent& PropertyChangedEvent) override;
 #endif // WITH_EDITOR
-	virtual void PostInitProperties() OVERRIDE;
-	virtual void Serialize(FArchive& Ar) OVERRIDE;
+	virtual void PostInitProperties() override;
+	virtual void Serialize(FArchive& Ar) override;
 	//End UObject Interface
 
 	//Begin UParticleModule Interface
-	virtual void CompileModule( struct FParticleEmitterBuildInfo& EmitterInfo ) OVERRIDE;
-	virtual void Spawn(FParticleEmitterInstance* Owner, int32 Offset, float SpawnTime, FBaseParticle* ParticleBase) OVERRIDE;
-	virtual void SetToSensibleDefaults(UParticleEmitter* Owner) OVERRIDE;
+	virtual void CompileModule( struct FParticleEmitterBuildInfo& EmitterInfo ) override;
+	virtual void Spawn(FParticleEmitterInstance* Owner, int32 Offset, float SpawnTime, FBaseParticle* ParticleBase) override;
+	virtual void SetToSensibleDefaults(UParticleEmitter* Owner) override;
 	//End UParticleModule Interface
 
 	// Begin UParticleModuleLifetimeBase Interface
-	virtual float GetMaxLifetime() OVERRIDE;
-	virtual float GetLifetimeValue(FParticleEmitterInstance* Owner, float InTime, UObject* Data = NULL) OVERRIDE;
+	virtual float GetMaxLifetime() override;
+	virtual float GetLifetimeValue(FParticleEmitterInstance* Owner, float InTime, UObject* Data = NULL) override;
 	// End UParticleModuleLifetimeBase Interface
 
 	/**
diff --git a/Engine/Source/Runtime/Engine/Classes/Particles/Lifetime/ParticleModuleLifetime_Seeded.h b/Engine/Source/Runtime/Engine/Classes/Particles/Lifetime/ParticleModuleLifetime_Seeded.h
--- a/Engine/Source/Runtime/Engine/Classes/Particles/Lifetime/ParticleModuleLifetime_Seeded.h
+++ b/Engine/Source/Runtime/Engine/Classes/Particles/Lifetime/ParticleModuleLifetime_Seeded.h
@@ -16,18 +16,18 @@ class UParticleModuleLifetime_Seeded : public UParticleModuleLifetime
 
 
 	//Begin UParticleModule Interface
-	virtual void Spawn(FParticleEmitterInstance* Owner, int32 Offset, float SpawnTime, FBaseParticle* ParticleBase) OVERRIDE;
-	virtual uint32	RequiredBytesPerInstance(FParticleEmitterInstance* Owner = NULL) OVERRIDE;
-	virtual uint32	PrepPerInstanceBlock(FParticleEmitterInstance* Owner, void* InstData) OVERRIDE;
-	virtual FParticleRandomSeedInfo* GetRandomSeedInfo() OVERRIDE
+	virtual void Spawn(FParticleEmitterInstance* Owner, int32 Offset, float SpawnTime, FBaseParticle* ParticleBase) override;
+	virtual uint32	RequiredBytesPerInstance(FParticleEmitterInstance* Owner = NULL) override;
+	virtual uint32	PrepPerInstanceBlock(FParticleEmitterInstance* Owner, void* InstData) override;
+	virtual FParticleRandomSeedInfo* GetRandomSeedInfo() override
 	{
 		return &RandomSeedInfo;
 	}
-	virtual void EmitterLoopingNotify(FParticleEmitterInstance* Owner) OVERRIDE;
+	virtual void EmitterLoopingNotify(FParticleEmitterInstance* Owner) override;
 	//End UParticleModule Interface
 
 	// Begin UParticleModuleLifetimeBase Interface
-	virtual float	GetLifetimeValue(FParticleEmitterInstance* Owner, float InTime, UObject* Data = NULL) OVERRIDE;
+	virtual float	GetLifetimeValue(FParticleEmitterInstance* Owner, float InTime, UObject* Data = NULL) override;
 	// End UParticleModuleLifetimeBase Interface
 };
 
diff --git a/Engine/Source/Runtime/Engine/Classes/Particles/Light/ParticleModuleLight.h b/Engine/Source/Runtime/Engine/Classes/Particles/Light/ParticleModuleLight.h
--- a/Engine/Source/Runtime/Engine/Classes/Particles/Light/ParticleModuleLight.h
+++ b/Engine/Source/Runtime/Engine/Classes/Particles/Light/ParticleModuleLight.h
@@ -54,19 +54,19 @@ class UParticleModuleLight : public UParticleModuleLightBase
 
 	// Begin UObject Interface
 #if WITH_EDITOR
-	virtual void PostEditChangeProperty(FPropertyChangedEvent& PropertyChangedEvent) OVERRIDE;
+	virtual void PostEditChangeProperty(FPropertyChangedEvent& PropertyChangedEvent) override;
 #endif // WITH_EDITOR
-	virtual void PostInitProperties() OVERRIDE;
+	virtual void PostInitProperties() override;
 	// End UObject Interface
 
 
 	//Begin UParticleModule Interface
-	virtual void Spawn(FParticleEmitterInstance* Owner, int32 Offset, float SpawnTime, FBaseParticle* ParticleBase) OVERRIDE;
-	virtual void Update(FParticleEmitterInstance* Owner, int32 Offset, float DeltaTime) OVERRIDE;
-	virtual uint32 RequiredBytes(FParticleEmitterInstance* Owner = NULL) OVERRIDE;
-	virtual void SetToSensibleDefaults(UParticleEmitter* Owner) OVERRIDE;
-	virtual EModuleType	GetModuleType() const OVERRIDE { return EPMT_Light; }
-	virtual void Render3DPreview(FParticleEmitterInstance* Owner, const FSceneView* View,FPrimitiveDrawInterface* PDI) OVERRIDE;
+	virtual void Spawn(FParticleEmitterInstance* Owner, int32 Offset, float SpawnTime, FBaseParticle* ParticleBase) override;
+	virtual void Update(FParticleEmitterInstance* Owner, int32 Offset, float DeltaTime) override;
+	virtual uint32 RequiredBytes(FParticleEmitterInstance* Owner = NULL) override;
+	virtual void SetToSensibleDefaults(UParticleEmitter* Owner) override;
+	virtual EModuleType	GetModuleType() const override { return EPMT_Light; }
+	virtual void Render3DPreview(FParticleEmitterInstance* Owner, const FSceneView* View,FPrimitiveDrawInterface* PDI) override;
 	//End UParticleModule Interface
 
 	void SpawnEx(FParticleEmitterInstance* Owner, int32 Offset, float SpawnTime, class FRandomStream* InRandomStream, FBaseParticle* ParticleBase);
diff --git a/Engine/Source/Runtime/Engine/Classes/Particles/Light/ParticleModuleLightBase.h b/Engine/Source/Runtime/Engine/Classes/Particles/Light/ParticleModuleLightBase.h
--- a/Engine/Source/Runtime/Engine/Classes/Particles/Light/ParticleModuleLightBase.h
+++ b/Engine/Source/Runtime/Engine/Classes/Particles/Light/ParticleModuleLightBase.h
@@ -9,7 +9,7 @@ UCLASS(editinlinenew, hidecategories=Object, abstract, meta=(DisplayName = "Ligh
 class UParticleModuleLightBase : public UParticleModule
 {
 	GENERATED_UCLASS_BODY()
-	virtual bool CanTickInAnyThread() OVERRIDE
+	virtual bool CanTickInAnyThread() override
 	{
 		return false;
 	}
diff --git a/Engine/Source/Runtime/Engine/Classes/Particles/Light/ParticleModuleLight_Seeded.h b/Engine/Source/Runtime/Engine/Classes/Particles/Light/ParticleModuleLight_Seeded.h
--- a/Engine/Source/Runtime/Engine/Classes/Particles/Light/ParticleModuleLight_Seeded.h
+++ b/Engine/Source/Runtime/Engine/Classes/Particles/Light/ParticleModuleLight_Seeded.h
@@ -16,14 +16,14 @@ class UParticleModuleLight_Seeded : public UParticleModuleLight
 
 
 	//Begin UParticleModule Interface
-	virtual void Spawn(FParticleEmitterInstance* Owner, int32 Offset, float SpawnTime, FBaseParticle* ParticleBase) OVERRIDE;
-	virtual uint32	RequiredBytesPerInstance(FParticleEmitterInstance* Owner = NULL) OVERRIDE;
-	virtual uint32	PrepPerInstanceBlock(FParticleEmitterInstance* Owner, void* InstData) OVERRIDE;
-	virtual FParticleRandomSeedInfo* GetRandomSeedInfo() OVERRIDE
+	virtual void Spawn(FParticleEmitterInstance* Owner, int32 Offset, float SpawnTime, FBaseParticle* ParticleBase) override;
+	virtual uint32	RequiredBytesPerInstance(FParticleEmitterInstance* Owner = NULL) override;
+	virtual uint32	PrepPerInstanceBlock(FParticleEmitterInstance* Owner, void* InstData) override;
+	virtual FParticleRandomSeedInfo* GetRandomSeedInfo() override
 	{
 		return &RandomSeedInfo;
 	}
-	virtual void EmitterLoopingNotify(FParticleEmitterInstance* Owner) OVERRIDE;
+	virtual void EmitterLoopingNotify(FParticleEmitterInstance* Owner) override;
 	//End UParticleModule Interface
 };
 
diff --git a/Engine/Source/Runtime/Engine/Classes/Particles/Location/ParticleModuleLocation.h b/Engine/Source/Runtime/Engine/Classes/Particles/Location/ParticleModuleLocation.h
--- a/Engine/Source/Runtime/Engine/Classes/Particles/Location/ParticleModuleLocation.h
+++ b/Engine/Source/Runtime/Engine/Classes/Particles/Location/ParticleModuleLocation.h
@@ -51,15 +51,15 @@ protected:
 public:
 	//Begin UObject Interface
 #if WITH_EDITOR
-	virtual void PostEditChangeProperty(FPropertyChangedEvent& PropertyChangedEvent) OVERRIDE;
+	virtual void PostEditChangeProperty(FPropertyChangedEvent& PropertyChangedEvent) override;
 #endif // WITH_EDITOR
-	virtual void PostInitProperties() OVERRIDE;
-	virtual void Serialize(FArchive& Ar) OVERRIDE;
+	virtual void PostInitProperties() override;
+	virtual void Serialize(FArchive& Ar) override;
 	//End UObject Interface
 
 	//Begin UParticleModule Interface
-	virtual void Spawn(FParticleEmitterInstance* Owner, int32 Offset, float SpawnTime, FBaseParticle* ParticleBase) OVERRIDE;
-	virtual void Render3DPreview(FParticleEmitterInstance* Owner, const FSceneView* View,FPrimitiveDrawInterface* PDI) OVERRIDE;
+	virtual void Spawn(FParticleEmitterInstance* Owner, int32 Offset, float SpawnTime, FBaseParticle* ParticleBase) override;
+	virtual void Render3DPreview(FParticleEmitterInstance* Owner, const FSceneView* View,FPrimitiveDrawInterface* PDI) override;
 	//End UParticleModule Interface
 };
 
diff --git a/Engine/Source/Runtime/Engine/Classes/Particles/Location/ParticleModuleLocationBoneSocket.h b/Engine/Source/Runtime/Engine/Classes/Particles/Location/ParticleModuleLocationBoneSocket.h
--- a/Engine/Source/Runtime/Engine/Classes/Particles/Location/ParticleModuleLocationBoneSocket.h
+++ b/Engine/Source/Runtime/Engine/Classes/Particles/Location/ParticleModuleLocationBoneSocket.h
@@ -99,22 +99,22 @@ class UParticleModuleLocationBoneSocket : public UParticleModuleLocationBase
 #endif // WITH_EDITORONLY_DATA
 
 	// Begin UParticleModule Interface
-	virtual void	Spawn(FParticleEmitterInstance* Owner, int32 Offset, float SpawnTime, FBaseParticle* ParticleBase) OVERRIDE;
-	virtual void	Update(FParticleEmitterInstance* Owner, int32 Offset, float DeltaTime) OVERRIDE;
-	virtual void	FinalUpdate(FParticleEmitterInstance* Owner, int32 Offset, float DeltaTime) OVERRIDE;
-	virtual uint32	RequiredBytes(FParticleEmitterInstance* Owner = NULL) OVERRIDE;
-	virtual uint32	RequiredBytesPerInstance(FParticleEmitterInstance* Owner = NULL) OVERRIDE;
-	virtual uint32	PrepPerInstanceBlock(FParticleEmitterInstance* Owner, void* InstData) OVERRIDE;
-	virtual bool	TouchesMeshRotation() const OVERRIDE { return true; }
-	virtual void	AutoPopulateInstanceProperties(UParticleSystemComponent* PSysComp) OVERRIDE;
-	virtual bool CanTickInAnyThread() OVERRIDE
+	virtual void	Spawn(FParticleEmitterInstance* Owner, int32 Offset, float SpawnTime, FBaseParticle* ParticleBase) override;
+	virtual void	Update(FParticleEmitterInstance* Owner, int32 Offset, float DeltaTime) override;
+	virtual void	FinalUpdate(FParticleEmitterInstance* Owner, int32 Offset, float DeltaTime) override;
+	virtual uint32	RequiredBytes(FParticleEmitterInstance* Owner = NULL) override;
+	virtual uint32	RequiredBytesPerInstance(FParticleEmitterInstance* Owner = NULL) override;
+	virtual uint32	PrepPerInstanceBlock(FParticleEmitterInstance* Owner, void* InstData) override;
+	virtual bool	TouchesMeshRotation() const override { return true; }
+	virtual void	AutoPopulateInstanceProperties(UParticleSystemComponent* PSysComp) override;
+	virtual bool CanTickInAnyThread() override
 	{
 		return false;
 	}
 #if WITH_EDITOR
-	virtual int32 GetNumberOfCustomMenuOptions() const OVERRIDE;
-	virtual bool GetCustomMenuEntryDisplayString(int32 InEntryIndex, FString& OutDisplayString) const OVERRIDE;
-	virtual bool PerformCustomMenuEntry(int32 InEntryIndex) OVERRIDE;
+	virtual int32 GetNumberOfCustomMenuOptions() const override;
+	virtual bool GetCustomMenuEntryDisplayString(int32 InEntryIndex, FString& OutDisplayString) const override;
+	virtual bool PerformCustomMenuEntry(int32 InEntryIndex) override;
 #endif
 	// End UParticleModule Interface
 
diff --git a/Engine/Source/Runtime/Engine/Classes/Particles/Location/ParticleModuleLocationDirect.h b/Engine/Source/Runtime/Engine/Classes/Particles/Location/ParticleModuleLocationDirect.h
--- a/Engine/Source/Runtime/Engine/Classes/Particles/Location/ParticleModuleLocationDirect.h
+++ b/Engine/Source/Runtime/Engine/Classes/Particles/Location/ParticleModuleLocationDirect.h
@@ -46,16 +46,16 @@ class UParticleModuleLocationDirect : public UParticleModuleLocationBase
 
 	// Begin UObject Interface
 #if WITH_EDITOR
-	virtual void PostEditChangeProperty(FPropertyChangedEvent& PropertyChangedEvent) OVERRIDE;
+	virtual void PostEditChangeProperty(FPropertyChangedEvent& PropertyChangedEvent) override;
 #endif // WITH_EDITOR
-	virtual void	PostInitProperties() OVERRIDE;
-	virtual void	Serialize(FArchive& Ar) OVERRIDE;
+	virtual void	PostInitProperties() override;
+	virtual void	Serialize(FArchive& Ar) override;
 	// End UObject Interface
 
 	// Begin UParticleModule Interface
-	virtual void	Spawn(FParticleEmitterInstance* Owner, int32 Offset, float SpawnTime, FBaseParticle* ParticleBase) OVERRIDE;
-	virtual void	Update(FParticleEmitterInstance* Owner, int32 Offset, float DeltaTime) OVERRIDE;
-	virtual uint32	RequiredBytes(FParticleEmitterInstance* Owner = NULL) OVERRIDE;
+	virtual void	Spawn(FParticleEmitterInstance* Owner, int32 Offset, float SpawnTime, FBaseParticle* ParticleBase) override;
+	virtual void	Update(FParticleEmitterInstance* Owner, int32 Offset, float DeltaTime) override;
+	virtual uint32	RequiredBytes(FParticleEmitterInstance* Owner = NULL) override;
 	// End UParticleModule Interface
 };
 
diff --git a/Engine/Source/Runtime/Engine/Classes/Particles/Location/ParticleModuleLocationEmitter.h b/Engine/Source/Runtime/Engine/Classes/Particles/Location/ParticleModuleLocationEmitter.h
--- a/Engine/Source/Runtime/Engine/Classes/Particles/Location/ParticleModuleLocationEmitter.h
+++ b/Engine/Source/Runtime/Engine/Classes/Particles/Location/ParticleModuleLocationEmitter.h
@@ -61,8 +61,8 @@ class UParticleModuleLocationEmitter : public UParticleModuleLocationBase
 
 
 	// Begin UParticleModule Interface
-	virtual void	Spawn(FParticleEmitterInstance* Owner, int32 Offset, float SpawnTime, FBaseParticle* ParticleBase) OVERRIDE;
-	virtual uint32	RequiredBytesPerInstance(FParticleEmitterInstance* Owner = NULL) OVERRIDE;
+	virtual void	Spawn(FParticleEmitterInstance* Owner, int32 Offset, float SpawnTime, FBaseParticle* ParticleBase) override;
+	virtual uint32	RequiredBytesPerInstance(FParticleEmitterInstance* Owner = NULL) override;
 	// End UParticleModule Interface
 };
 
diff --git a/Engine/Source/Runtime/Engine/Classes/Particles/Location/ParticleModuleLocationEmitterDirect.h b/Engine/Source/Runtime/Engine/Classes/Particles/Location/ParticleModuleLocationEmitterDirect.h
--- a/Engine/Source/Runtime/Engine/Classes/Particles/Location/ParticleModuleLocationEmitterDirect.h
+++ b/Engine/Source/Runtime/Engine/Classes/Particles/Location/ParticleModuleLocationEmitterDirect.h
@@ -23,8 +23,8 @@ class UParticleModuleLocationEmitterDirect : public UParticleModuleLocationBase
 
 
 	// Begin UParticleModule Interface
-	virtual void	Spawn(FParticleEmitterInstance* Owner, int32 Offset, float SpawnTime, FBaseParticle* ParticleBase) OVERRIDE;
-	virtual void	Update(FParticleEmitterInstance* Owner, int32 Offset, float DeltaTime) OVERRIDE;
+	virtual void	Spawn(FParticleEmitterInstance* Owner, int32 Offset, float SpawnTime, FBaseParticle* ParticleBase) override;
+	virtual void	Update(FParticleEmitterInstance* Owner, int32 Offset, float DeltaTime) override;
 	// End UParticleModule Interface
 
 };
diff --git a/Engine/Source/Runtime/Engine/Classes/Particles/Location/ParticleModuleLocationPrimitiveBase.h b/Engine/Source/Runtime/Engine/Classes/Particles/Location/ParticleModuleLocationPrimitiveBase.h
--- a/Engine/Source/Runtime/Engine/Classes/Particles/Location/ParticleModuleLocationPrimitiveBase.h
+++ b/Engine/Source/Runtime/Engine/Classes/Particles/Location/ParticleModuleLocationPrimitiveBase.h
@@ -59,10 +59,10 @@ class UParticleModuleLocationPrimitiveBase : public UParticleModuleLocationBase
 
 	//Begin UObject Interface
 #if WITH_EDITOR
-	virtual void PostEditChangeProperty(FPropertyChangedEvent& PropertyChangedEvent) OVERRIDE;
+	virtual void PostEditChangeProperty(FPropertyChangedEvent& PropertyChangedEvent) override;
 #endif // WITH_EDITOR
-	virtual void PostInitProperties() OVERRIDE;
-	virtual void Serialize(FArchive& Ar) OVERRIDE;
+	virtual void PostInitProperties() override;
+	virtual void Serialize(FArchive& Ar) override;
 	//End UObject Interface
 
 	//@todo document
diff --git a/Engine/Source/Runtime/Engine/Classes/Particles/Location/ParticleModuleLocationPrimitiveCylinder.h b/Engine/Source/Runtime/Engine/Classes/Particles/Location/ParticleModuleLocationPrimitiveCylinder.h
--- a/Engine/Source/Runtime/Engine/Classes/Particles/Location/ParticleModuleLocationPrimitiveCylinder.h
+++ b/Engine/Source/Runtime/Engine/Classes/Particles/Location/ParticleModuleLocationPrimitiveCylinder.h
@@ -50,15 +50,15 @@ class UParticleModuleLocationPrimitiveCylinder : public UParticleModuleLocationP
 
 	// Begin UObject Interface
 #if WITH_EDITOR
-	virtual void PostEditChangeProperty(FPropertyChangedEvent& PropertyChangedEvent) OVERRIDE;
+	virtual void PostEditChangeProperty(FPropertyChangedEvent& PropertyChangedEvent) override;
 #endif // WITH_EDITOR
-	virtual void PostInitProperties() OVERRIDE;
-	virtual void Serialize(FArchive& Ar) OVERRIDE;
+	virtual void PostInitProperties() override;
+	virtual void Serialize(FArchive& Ar) override;
 	// End UObject Interface
 
 	// Begin UParticleModule Interface
-	virtual void	Spawn(FParticleEmitterInstance* Owner, int32 Offset, float SpawnTime, FBaseParticle* ParticleBase) OVERRIDE;
-	virtual void	Render3DPreview(FParticleEmitterInstance* Owner, const FSceneView* View,FPrimitiveDrawInterface* PDI) OVERRIDE;
+	virtual void	Spawn(FParticleEmitterInstance* Owner, int32 Offset, float SpawnTime, FBaseParticle* ParticleBase) override;
+	virtual void	Render3DPreview(FParticleEmitterInstance* Owner, const FSceneView* View,FPrimitiveDrawInterface* PDI) override;
 	// End UParticleModule Interface
 
 	/**
diff --git a/Engine/Source/Runtime/Engine/Classes/Particles/Location/ParticleModuleLocationPrimitiveSphere.h b/Engine/Source/Runtime/Engine/Classes/Particles/Location/ParticleModuleLocationPrimitiveSphere.h
--- a/Engine/Source/Runtime/Engine/Classes/Particles/Location/ParticleModuleLocationPrimitiveSphere.h
+++ b/Engine/Source/Runtime/Engine/Classes/Particles/Location/ParticleModuleLocationPrimitiveSphere.h
@@ -23,15 +23,15 @@ class UParticleModuleLocationPrimitiveSphere : public UParticleModuleLocationPri
 
 	// Begin UObject Interface
 #if WITH_EDITOR
-	virtual void PostEditChangeProperty(FPropertyChangedEvent& PropertyChangedEvent) OVERRIDE;
+	virtual void PostEditChangeProperty(FPropertyChangedEvent& PropertyChangedEvent) override;
 #endif // WITH_EDITOR
-	virtual void PostInitProperties() OVERRIDE;
-	virtual void Serialize(FArchive& Ar) OVERRIDE;
+	virtual void PostInitProperties() override;
+	virtual void Serialize(FArchive& Ar) override;
 	// End UObject Interface
 
 	// Begin UParticleModule Interface
-	virtual void	Spawn(FParticleEmitterInstance* Owner, int32 Offset, float SpawnTime, FBaseParticle* ParticleBase) OVERRIDE;
-	virtual void	Render3DPreview(FParticleEmitterInstance* Owner, const FSceneView* View,FPrimitiveDrawInterface* PDI) OVERRIDE;
+	virtual void	Spawn(FParticleEmitterInstance* Owner, int32 Offset, float SpawnTime, FBaseParticle* ParticleBase) override;
+	virtual void	Render3DPreview(FParticleEmitterInstance* Owner, const FSceneView* View,FPrimitiveDrawInterface* PDI) override;
 	// End UParticleModule Interface
 
 	/**
diff --git a/Engine/Source/Runtime/Engine/Classes/Particles/Location/ParticleModuleLocationSkelVertSurface.h b/Engine/Source/Runtime/Engine/Classes/Particles/Location/ParticleModuleLocationSkelVertSurface.h
--- a/Engine/Source/Runtime/Engine/Classes/Particles/Location/ParticleModuleLocationSkelVertSurface.h
+++ b/Engine/Source/Runtime/Engine/Classes/Particles/Location/ParticleModuleLocationSkelVertSurface.h
@@ -88,30 +88,30 @@ class UParticleModuleLocationSkelVertSurface : public UParticleModuleLocationBas
 	uint32 bInheritVertexColor : 1;
 
 	// Begin UObject Interface
-	virtual void PostLoad() OVERRIDE;
+	virtual void PostLoad() override;
 #if WITH_EDITOR
-	virtual void PostEditChangeProperty(FPropertyChangedEvent& PropertyChangedEvent) OVERRIDE;
+	virtual void PostEditChangeProperty(FPropertyChangedEvent& PropertyChangedEvent) override;
 #endif // WITH_EDITOR
 	// End UObject Interface
 
 	//Begin UParticleModule Interface
-	virtual void Spawn(FParticleEmitterInstance* Owner, int32 Offset, float SpawnTime, FBaseParticle* ParticleBase) OVERRIDE;
-	virtual void Update(FParticleEmitterInstance* Owner, int32 Offset, float DeltaTime) OVERRIDE;
-	virtual void FinalUpdate(FParticleEmitterInstance* Owner, int32 Offset, float DeltaTime) OVERRIDE;
-	virtual uint32	PrepPerInstanceBlock(FParticleEmitterInstance* Owner, void* InstData) OVERRIDE;
-	virtual uint32	RequiredBytes(FParticleEmitterInstance* Owner = NULL) OVERRIDE;
-	virtual uint32	RequiredBytesPerInstance(FParticleEmitterInstance* Owner = NULL) OVERRIDE;
-	virtual bool	TouchesMeshRotation() const OVERRIDE { return true; }
-	virtual void	AutoPopulateInstanceProperties(UParticleSystemComponent* PSysComp) OVERRIDE;
-	virtual bool CanTickInAnyThread() OVERRIDE
+	virtual void Spawn(FParticleEmitterInstance* Owner, int32 Offset, float SpawnTime, FBaseParticle* ParticleBase) override;
+	virtual void Update(FParticleEmitterInstance* Owner, int32 Offset, float DeltaTime) override;
+	virtual void FinalUpdate(FParticleEmitterInstance* Owner, int32 Offset, float DeltaTime) override;
+	virtual uint32	PrepPerInstanceBlock(FParticleEmitterInstance* Owner, void* InstData) override;
+	virtual uint32	RequiredBytes(FParticleEmitterInstance* Owner = NULL) override;
+	virtual uint32	RequiredBytesPerInstance(FParticleEmitterInstance* Owner = NULL) override;
+	virtual bool	TouchesMeshRotation() const override { return true; }
+	virtual void	AutoPopulateInstanceProperties(UParticleSystemComponent* PSysComp) override;
+	virtual bool CanTickInAnyThread() override
 	{
 		return false;
 	}
 #if WITH_EDITOR
-	virtual int32 GetNumberOfCustomMenuOptions() const OVERRIDE;
-	virtual bool GetCustomMenuEntryDisplayString(int32 InEntryIndex, FString& OutDisplayString) const OVERRIDE;
-	virtual bool PerformCustomMenuEntry(int32 InEntryIndex) OVERRIDE;
-	virtual bool IsValidForLODLevel(UParticleLODLevel* LODLevel, FString& OutErrorString) OVERRIDE;
+	virtual int32 GetNumberOfCustomMenuOptions() const override;
+	virtual bool GetCustomMenuEntryDisplayString(int32 InEntryIndex, FString& OutDisplayString) const override;
+	virtual bool PerformCustomMenuEntry(int32 InEntryIndex) override;
+	virtual bool IsValidForLODLevel(UParticleLODLevel* LODLevel, FString& OutErrorString) override;
 #endif
 	//End UParticleModule Interface
 
diff --git a/Engine/Source/Runtime/Engine/Classes/Particles/Location/ParticleModuleLocationWorldOffset.h b/Engine/Source/Runtime/Engine/Classes/Particles/Location/ParticleModuleLocationWorldOffset.h
--- a/Engine/Source/Runtime/Engine/Classes/Particles/Location/ParticleModuleLocationWorldOffset.h
+++ b/Engine/Source/Runtime/Engine/Classes/Particles/Location/ParticleModuleLocationWorldOffset.h
@@ -13,7 +13,7 @@ class UParticleModuleLocationWorldOffset : public UParticleModuleLocation
 
 protected:
 	//Begin UParticleModuleLocation Interface
-	virtual void SpawnEx(FParticleEmitterInstance* Owner, int32 Offset, float SpawnTime, class FRandomStream* InRandomStream, FBaseParticle* ParticleBase) OVERRIDE;
+	virtual void SpawnEx(FParticleEmitterInstance* Owner, int32 Offset, float SpawnTime, class FRandomStream* InRandomStream, FBaseParticle* ParticleBase) override;
 	//End UParticleModuleLocation Interface
 };
 
diff --git a/Engine/Source/Runtime/Engine/Classes/Particles/Location/ParticleModuleLocation_Seeded.h b/Engine/Source/Runtime/Engine/Classes/Particles/Location/ParticleModuleLocation_Seeded.h
--- a/Engine/Source/Runtime/Engine/Classes/Particles/Location/ParticleModuleLocation_Seeded.h
+++ b/Engine/Source/Runtime/Engine/Classes/Particles/Location/ParticleModuleLocation_Seeded.h
@@ -16,14 +16,14 @@ class UParticleModuleLocation_Seeded : public UParticleModuleLocation
 
 
 	//Begin UParticleModule Interface
-	virtual void Spawn(FParticleEmitterInstance* Owner, int32 Offset, float SpawnTime, FBaseParticle* ParticleBase) OVERRIDE;
-	virtual uint32 RequiredBytesPerInstance(FParticleEmitterInstance* Owner = NULL) OVERRIDE;
-	virtual uint32 PrepPerInstanceBlock(FParticleEmitterInstance* Owner, void* InstData) OVERRIDE;
-	virtual FParticleRandomSeedInfo* GetRandomSeedInfo() OVERRIDE
+	virtual void Spawn(FParticleEmitterInstance* Owner, int32 Offset, float SpawnTime, FBaseParticle* ParticleBase) override;
+	virtual uint32 RequiredBytesPerInstance(FParticleEmitterInstance* Owner = NULL) override;
+	virtual uint32 PrepPerInstanceBlock(FParticleEmitterInstance* Owner, void* InstData) override;
+	virtual FParticleRandomSeedInfo* GetRandomSeedInfo() override
 	{
 		return &RandomSeedInfo;
 	}
-	virtual void EmitterLoopingNotify(FParticleEmitterInstance* Owner) OVERRIDE;
+	virtual void EmitterLoopingNotify(FParticleEmitterInstance* Owner) override;
 	//End UParticleModule Interface
 };
 
diff --git a/Engine/Source/Runtime/Engine/Classes/Particles/Location/ParticleModuleSourceMovement.h b/Engine/Source/Runtime/Engine/Classes/Particles/Location/ParticleModuleSourceMovement.h
--- a/Engine/Source/Runtime/Engine/Classes/Particles/Location/ParticleModuleSourceMovement.h
+++ b/Engine/Source/Runtime/Engine/Classes/Particles/Location/ParticleModuleSourceMovement.h
@@ -18,13 +18,13 @@ class UParticleModuleSourceMovement : public UParticleModuleLocationBase
 	struct FRawDistributionVector SourceMovementScale;
 
 	//Begin UObject Interface
-	virtual void PostInitProperties() OVERRIDE;
-	virtual void Serialize(FArchive& Ar) OVERRIDE;
+	virtual void PostInitProperties() override;
+	virtual void Serialize(FArchive& Ar) override;
 	//End UObject Interface
 
 	//Begin UParticleModule Interface
-	virtual void FinalUpdate(FParticleEmitterInstance* Owner, int32 Offset, float DeltaTime) OVERRIDE;
-	virtual bool CanTickInAnyThread() OVERRIDE
+	virtual void FinalUpdate(FParticleEmitterInstance* Owner, int32 Offset, float DeltaTime) override;
+	virtual bool CanTickInAnyThread() override
 	{
 		return false;
 	}
diff --git a/Engine/Source/Runtime/Engine/Classes/Particles/Modules/Location/ParticleModulePivotOffset.h b/Engine/Source/Runtime/Engine/Classes/Particles/Modules/Location/ParticleModulePivotOffset.h
--- a/Engine/Source/Runtime/Engine/Classes/Particles/Modules/Location/ParticleModulePivotOffset.h
+++ b/Engine/Source/Runtime/Engine/Classes/Particles/Modules/Location/ParticleModulePivotOffset.h
@@ -17,15 +17,15 @@ class UParticleModulePivotOffset : public UParticleModuleLocationBase
 	void InitializeDefaults();
 
 	//Begin UObject Interface
-	virtual void PostInitProperties() OVERRIDE;
+	virtual void PostInitProperties() override;
 	//End UObject Interface
 
 	//Begin UParticleModule Interface
-	virtual void CompileModule( struct FParticleEmitterBuildInfo& EmitterInfo ) OVERRIDE;
+	virtual void CompileModule( struct FParticleEmitterBuildInfo& EmitterInfo ) override;
 	//End UParticleModule Interface
 
 #if WITH_EDITOR
-	virtual bool IsValidForLODLevel(UParticleLODLevel* LODLevel, FString& OutErrorString) OVERRIDE;
+	virtual bool IsValidForLODLevel(UParticleLODLevel* LODLevel, FString& OutErrorString) override;
 #endif
 
 };
diff --git a/Engine/Source/Runtime/Engine/Classes/Particles/Orbit/ParticleModuleOrbit.h b/Engine/Source/Runtime/Engine/Classes/Particles/Orbit/ParticleModuleOrbit.h
--- a/Engine/Source/Runtime/Engine/Classes/Particles/Orbit/ParticleModuleOrbit.h
+++ b/Engine/Source/Runtime/Engine/Classes/Particles/Orbit/ParticleModuleOrbit.h
@@ -109,20 +109,20 @@ class UParticleModuleOrbit : public UParticleModuleOrbitBase
 
 	//Begin UObject Interface
 #if WITH_EDITOR
-	virtual void PostEditChangeProperty(FPropertyChangedEvent& PropertyChangedEvent) OVERRIDE;
+	virtual void PostEditChangeProperty(FPropertyChangedEvent& PropertyChangedEvent) override;
 #endif // WITH_EDITOR
-	virtual void PostInitProperties() OVERRIDE;
-	virtual void Serialize(FArchive& Ar) OVERRIDE;
+	virtual void PostInitProperties() override;
+	virtual void Serialize(FArchive& Ar) override;
 	// End UObject Interface
 
 	//Begin UParticleModule Interface
-	virtual void CompileModule( struct FParticleEmitterBuildInfo& EmitterInfo ) OVERRIDE;
-	virtual void	Spawn(FParticleEmitterInstance* Owner, int32 Offset, float SpawnTime, FBaseParticle* ParticleBase) OVERRIDE;
-	virtual void	Update(FParticleEmitterInstance* Owner, int32 Offset, float DeltaTime) OVERRIDE;
-	virtual uint32	RequiredBytes(FParticleEmitterInstance* Owner = NULL) OVERRIDE;
-	virtual uint32	RequiredBytesPerInstance(FParticleEmitterInstance* Owner = NULL) OVERRIDE;
+	virtual void CompileModule( struct FParticleEmitterBuildInfo& EmitterInfo ) override;
+	virtual void	Spawn(FParticleEmitterInstance* Owner, int32 Offset, float SpawnTime, FBaseParticle* ParticleBase) override;
+	virtual void	Update(FParticleEmitterInstance* Owner, int32 Offset, float DeltaTime) override;
+	virtual uint32	RequiredBytes(FParticleEmitterInstance* Owner = NULL) override;
+	virtual uint32	RequiredBytesPerInstance(FParticleEmitterInstance* Owner = NULL) override;
 #if WITH_EDITOR
-	virtual bool IsValidForLODLevel(UParticleLODLevel* LODLevel, FString& OutErrorString) OVERRIDE;
+	virtual bool IsValidForLODLevel(UParticleLODLevel* LODLevel, FString& OutErrorString) override;
 #endif
 	//End UParticleModule Interface
 
diff --git a/Engine/Source/Runtime/Engine/Classes/Particles/Orientation/ParticleModuleOrientationAxisLock.h b/Engine/Source/Runtime/Engine/Classes/Particles/Orientation/ParticleModuleOrientationAxisLock.h
--- a/Engine/Source/Runtime/Engine/Classes/Particles/Orientation/ParticleModuleOrientationAxisLock.h
+++ b/Engine/Source/Runtime/Engine/Classes/Particles/Orientation/ParticleModuleOrientationAxisLock.h
@@ -57,13 +57,13 @@ class UParticleModuleOrientationAxisLock : public UParticleModuleOrientationBase
 
 	// Begin UObject Interface
 #if WITH_EDITOR
-	virtual void	PostEditChangeProperty(FPropertyChangedEvent& PropertyChangedEvent) OVERRIDE;
+	virtual void	PostEditChangeProperty(FPropertyChangedEvent& PropertyChangedEvent) override;
 #endif // WITH_EDITOR
 	// End UObject Interface
 
 	// Begin UParticleModule Interface
-	virtual void	Spawn(FParticleEmitterInstance* Owner, int32 Offset, float SpawnTime, FBaseParticle* ParticleBase) OVERRIDE;
-	virtual void	Update(FParticleEmitterInstance* Owner, int32 Offset, float DeltaTime) OVERRIDE;
+	virtual void	Spawn(FParticleEmitterInstance* Owner, int32 Offset, float SpawnTime, FBaseParticle* ParticleBase) override;
+	virtual void	Update(FParticleEmitterInstance* Owner, int32 Offset, float DeltaTime) override;
 	// End UParticleModule Interface
 
 	//@todo document
diff --git a/Engine/Source/Runtime/Engine/Classes/Particles/Parameter/ParticleModuleParameterDynamic.h b/Engine/Source/Runtime/Engine/Classes/Particles/Parameter/ParticleModuleParameterDynamic.h
--- a/Engine/Source/Runtime/Engine/Classes/Particles/Parameter/ParticleModuleParameterDynamic.h
+++ b/Engine/Source/Runtime/Engine/Classes/Particles/Parameter/ParticleModuleParameterDynamic.h
@@ -100,29 +100,29 @@ class UParticleModuleParameterDynamic : public UParticleModuleParameterBase
 	void InitializeDefaults();
 
 	//Begin UObject Interface
-	virtual void	PostLoad() OVERRIDE;
+	virtual void	PostLoad() override;
 #if WITH_EDITOR
-	virtual void	PostEditChangeProperty(FPropertyChangedEvent& PropertyChangedEvent) OVERRIDE;
+	virtual void	PostEditChangeProperty(FPropertyChangedEvent& PropertyChangedEvent) override;
 #endif // WITH_EDITOR
-	virtual void	PostInitProperties() OVERRIDE;
-	virtual void	Serialize(FArchive& Ar) OVERRIDE;
+	virtual void	PostInitProperties() override;
+	virtual void	Serialize(FArchive& Ar) override;
 	// Begin UObject Interface
 
 	// Begin UParticleModule Interface
-	virtual void	Spawn(FParticleEmitterInstance* Owner, int32 Offset, float SpawnTime, FBaseParticle* ParticleBase) OVERRIDE;
-	virtual void	Update(FParticleEmitterInstance* Owner, int32 Offset, float DeltaTime) OVERRIDE;
-	virtual uint32	RequiredBytes(FParticleEmitterInstance* Owner = NULL) OVERRIDE;
-	virtual void SetToSensibleDefaults(UParticleEmitter* Owner) OVERRIDE;
-	virtual void	GetCurveObjects(TArray<FParticleCurvePair>& OutCurves) OVERRIDE;
-	virtual bool WillGeneratedModuleBeIdentical(UParticleLODLevel* SourceLODLevel, UParticleLODLevel* DestLODLevel, float Percentage) OVERRIDE
+	virtual void	Spawn(FParticleEmitterInstance* Owner, int32 Offset, float SpawnTime, FBaseParticle* ParticleBase) override;
+	virtual void	Update(FParticleEmitterInstance* Owner, int32 Offset, float DeltaTime) override;
+	virtual uint32	RequiredBytes(FParticleEmitterInstance* Owner = NULL) override;
+	virtual void SetToSensibleDefaults(UParticleEmitter* Owner) override;
+	virtual void	GetCurveObjects(TArray<FParticleCurvePair>& OutCurves) override;
+	virtual bool WillGeneratedModuleBeIdentical(UParticleLODLevel* SourceLODLevel, UParticleLODLevel* DestLODLevel, float Percentage) override
 	{
 		// The assumption is that at 100%, ANY module will be identical...
 		// (Although this is virtual to allow over-riding that assumption on a case-by-case basis!)
 		return true;
 	}
-	virtual void GetParticleSysParamsUtilized(TArray<FString>& ParticleSysParamList) OVERRIDE;
-	virtual void GetParticleParametersUtilized(TArray<FString>& ParticleParameterList) OVERRIDE;
-	virtual void RefreshModule(UInterpCurveEdSetup* EdSetup, UParticleEmitter* InEmitter, int32 InLODLevel) OVERRIDE;
+	virtual void GetParticleSysParamsUtilized(TArray<FString>& ParticleSysParamList) override;
+	virtual void GetParticleParametersUtilized(TArray<FString>& ParticleParameterList) override;
+	virtual void RefreshModule(UInterpCurveEdSetup* EdSetup, UParticleEmitter* InEmitter, int32 InLODLevel) override;
 	// End UParticleModule Interface
 
 	/**
diff --git a/Engine/Source/Runtime/Engine/Classes/Particles/ParticleEmitter.h b/Engine/Source/Runtime/Engine/Classes/Particles/ParticleEmitter.h
--- a/Engine/Source/Runtime/Engine/Classes/Particles/ParticleEmitter.h
+++ b/Engine/Source/Runtime/Engine/Classes/Particles/ParticleEmitter.h
@@ -165,9 +165,9 @@ class UParticleEmitter : public UObject
 
 	// Begin UObject Interface
 #if WITH_EDITOR
-	virtual void PostEditChangeProperty(FPropertyChangedEvent& PropertyChangedEvent) OVERRIDE;
+	virtual void PostEditChangeProperty(FPropertyChangedEvent& PropertyChangedEvent) override;
 #endif // WITH_EDITOR
-	virtual void PostLoad() OVERRIDE;
+	virtual void PostLoad() override;
 	// End UObject Interface
 
 	// @todo document
diff --git a/Engine/Source/Runtime/Engine/Classes/Particles/ParticleLODLevel.h b/Engine/Source/Runtime/Engine/Classes/Particles/ParticleLODLevel.h
--- a/Engine/Source/Runtime/Engine/Classes/Particles/ParticleLODLevel.h
+++ b/Engine/Source/Runtime/Engine/Classes/Particles/ParticleLODLevel.h
@@ -71,7 +71,7 @@ class UParticleLODLevel : public UObject
 
 
 	// Begin UObject Interface
-	virtual void	PostLoad() OVERRIDE;
+	virtual void	PostLoad() override;
 	// End UObject Interface
 
 	// @todo document
diff --git a/Engine/Source/Runtime/Engine/Classes/Particles/ParticleModule.h b/Engine/Source/Runtime/Engine/Classes/Particles/ParticleModule.h
--- a/Engine/Source/Runtime/Engine/Classes/Particles/ParticleModule.h
+++ b/Engine/Source/Runtime/Engine/Classes/Particles/ParticleModule.h
@@ -209,7 +209,7 @@ class UParticleModule : public UObject
 
 	// Begin UObject Interface
 #if WITH_EDITOR
-	virtual void PostEditChangeProperty(FPropertyChangedEvent& PropertyChangedEvent) OVERRIDE;
+	virtual void PostEditChangeProperty(FPropertyChangedEvent& PropertyChangedEvent) override;
 #endif // WITH_EDITOR
 	// End UObject Interface
 
@@ -549,7 +549,7 @@ class UParticleModule : public UObject
 	ENGINE_API bool IsUsedInGPUEmitter()const;
 
 #if WITH_EDITOR
-	virtual void PostLoadSubobjects( FObjectInstancingGraph* OuterInstanceGraph ) OVERRIDE;
+	virtual void PostLoadSubobjects( FObjectInstancingGraph* OuterInstanceGraph ) override;
 
 	/**
 	 *	Custom Cascade module menu entries support
diff --git a/Engine/Source/Runtime/Engine/Classes/Particles/ParticleModuleRequired.h b/Engine/Source/Runtime/Engine/Classes/Particles/ParticleModuleRequired.h
--- a/Engine/Source/Runtime/Engine/Classes/Particles/ParticleModuleRequired.h
+++ b/Engine/Source/Runtime/Engine/Classes/Particles/ParticleModuleRequired.h
@@ -282,23 +282,23 @@ class UParticleModuleRequired : public UParticleModule
 
 	// Begin UObject Interface
 #if WITH_EDITOR
-	virtual void	PostEditChangeProperty(FPropertyChangedEvent& PropertyChangedEvent) OVERRIDE;
-	virtual bool IsValidForLODLevel(UParticleLODLevel* LODLevel, FString& OutErrorString) OVERRIDE;
+	virtual void	PostEditChangeProperty(FPropertyChangedEvent& PropertyChangedEvent) override;
+	virtual bool IsValidForLODLevel(UParticleLODLevel* LODLevel, FString& OutErrorString) override;
 #endif // WITH_EDITOR
-	virtual void	PostLoad() OVERRIDE;
-	virtual void	PostInitProperties() OVERRIDE;
-	virtual void	Serialize(FArchive& Ar) OVERRIDE;
+	virtual void	PostLoad() override;
+	virtual void	PostInitProperties() override;
+	virtual void	Serialize(FArchive& Ar) override;
 	// End UObject Interface
 
 	// Begin UParticleModule Interface
-	virtual void SetToSensibleDefaults(UParticleEmitter* Owner) OVERRIDE;
-	virtual	bool AddModuleCurvesToEditor(UInterpCurveEdSetup* EdSetup, TArray<const FCurveEdEntry*>& OutCurveEntries) OVERRIDE
+	virtual void SetToSensibleDefaults(UParticleEmitter* Owner) override;
+	virtual	bool AddModuleCurvesToEditor(UInterpCurveEdSetup* EdSetup, TArray<const FCurveEdEntry*>& OutCurveEntries) override
 	{
 		// Overide the base implementation to prevent old SpawnRate from being added...
 		return true;
 	}
-	virtual EModuleType	GetModuleType() const OVERRIDE {	return EPMT_Required;	}
-	virtual bool	GenerateLODModuleValues(UParticleModule* SourceModule, float Percentage, UParticleLODLevel* LODLevel) OVERRIDE;
+	virtual EModuleType	GetModuleType() const override {	return EPMT_Required;	}
+	virtual bool	GenerateLODModuleValues(UParticleModule* SourceModule, float Percentage, UParticleLODLevel* LODLevel) override;
 	// End UParticleModule Interface
 
 protected:
diff --git a/Engine/Source/Runtime/Engine/Classes/Particles/ParticleSpriteEmitter.h b/Engine/Source/Runtime/Engine/Classes/Particles/ParticleSpriteEmitter.h
--- a/Engine/Source/Runtime/Engine/Classes/Particles/ParticleSpriteEmitter.h
+++ b/Engine/Source/Runtime/Engine/Classes/Particles/ParticleSpriteEmitter.h
@@ -23,12 +23,12 @@ class UParticleSpriteEmitter : public UParticleEmitter
 
 
 	// Begin UObject Interface
-	virtual void PostLoad() OVERRIDE;
+	virtual void PostLoad() override;
 	// End UObject Interface
 
 	// Begin UParticleEmitter Interface
-	virtual FParticleEmitterInstance* CreateInstance(UParticleSystemComponent* InComponent) OVERRIDE;
-	virtual void SetToSensibleDefaults() OVERRIDE;
+	virtual FParticleEmitterInstance* CreateInstance(UParticleSystemComponent* InComponent) override;
+	virtual void SetToSensibleDefaults() override;
 	// End UParticleEmitter Interface
 };
 
diff --git a/Engine/Source/Runtime/Engine/Classes/Particles/ParticleSystem.h b/Engine/Source/Runtime/Engine/Classes/Particles/ParticleSystem.h
--- a/Engine/Source/Runtime/Engine/Classes/Particles/ParticleSystem.h
+++ b/Engine/Source/Runtime/Engine/Classes/Particles/ParticleSystem.h
@@ -300,12 +300,12 @@ public:
 
 	// Begin UObject interface.
 #if WITH_EDITOR
-	virtual void PostEditChangeProperty(FPropertyChangedEvent& PropertyChangedEvent) OVERRIDE;
+	virtual void PostEditChangeProperty(FPropertyChangedEvent& PropertyChangedEvent) override;
 #endif // WITH_EDITOR
-	virtual void PreSave() OVERRIDE;
-	virtual void PostLoad() OVERRIDE;
-	virtual SIZE_T GetResourceSize(EResourceSizeMode::Type Mode) OVERRIDE;
-	virtual void GetAssetRegistryTags(TArray<FAssetRegistryTag>& OutTags) const OVERRIDE;
+	virtual void PreSave() override;
+	virtual void PostLoad() override;
+	virtual SIZE_T GetResourceSize(EResourceSizeMode::Type Mode) override;
+	virtual void GetAssetRegistryTags(TArray<FAssetRegistryTag>& OutTags) const override;
 	// End UObject interface.
 
 
diff --git a/Engine/Source/Runtime/Engine/Classes/Particles/ParticleSystemComponent.h b/Engine/Source/Runtime/Engine/Classes/Particles/ParticleSystemComponent.h
--- a/Engine/Source/Runtime/Engine/Classes/Particles/ParticleSystemComponent.h
+++ b/Engine/Source/Runtime/Engine/Classes/Particles/ParticleSystemComponent.h
@@ -680,20 +680,20 @@ public:
 
 	// Begin UActorComponent interface.
 #if WITH_EDITOR
-	virtual void CheckForErrors() OVERRIDE;
+	virtual void CheckForErrors() override;
 #endif
-	virtual void TickComponent(float DeltaTime, enum ELevelTick TickType, FActorComponentTickFunction *ThisTickFunction) OVERRIDE;
-	virtual UObject const* AdditionalStatObject() const OVERRIDE
+	virtual void TickComponent(float DeltaTime, enum ELevelTick TickType, FActorComponentTickFunction *ThisTickFunction) override;
+	virtual UObject const* AdditionalStatObject() const override
 	{
 		return Template;
 	}
 protected:
-	virtual void CreateRenderState_Concurrent() OVERRIDE;
-	virtual void SendRenderTransform_Concurrent() OVERRIDE;
-	virtual void DestroyRenderState_Concurrent() OVERRIDE;
-	virtual void OnRegister() OVERRIDE;
-	virtual void OnUnregister()  OVERRIDE;
-	virtual void SendRenderDynamicData_Concurrent() OVERRIDE;
+	virtual void CreateRenderState_Concurrent() override;
+	virtual void SendRenderTransform_Concurrent() override;
+	virtual void DestroyRenderState_Concurrent() override;
+	virtual void OnRegister() override;
+	virtual void OnUnregister()  override;
+	virtual void SendRenderDynamicData_Concurrent() override;
 public:
 	// End UActorComponent interface.
 
@@ -733,35 +733,35 @@ private:
 public:
 
 	// Begin UObject interface.
-	virtual void PostLoad() OVERRIDE;
-	virtual void BeginDestroy() OVERRIDE;
-	virtual void FinishDestroy() OVERRIDE;
+	virtual void PostLoad() override;
+	virtual void BeginDestroy() override;
+	virtual void FinishDestroy() override;
 #if WITH_EDITOR
-	virtual void PreEditChange(UProperty* PropertyThatWillChange) OVERRIDE;
-	virtual void PostEditChangeChainProperty(FPropertyChangedChainEvent& PropertyChangedEvent) OVERRIDE;
+	virtual void PreEditChange(UProperty* PropertyThatWillChange) override;
+	virtual void PostEditChangeChainProperty(FPropertyChangedChainEvent& PropertyChangedEvent) override;
 #endif // WITH_EDITOR
-	virtual void Serialize(FArchive& Ar) OVERRIDE;
-	virtual SIZE_T GetResourceSize(EResourceSizeMode::Type Mode) OVERRIDE;
-	virtual FString GetDetailedInfoInternal() const OVERRIDE;
+	virtual void Serialize(FArchive& Ar) override;
+	virtual SIZE_T GetResourceSize(EResourceSizeMode::Type Mode) override;
+	virtual FString GetDetailedInfoInternal() const override;
 	// End UObject interface.
 
 	//Begin UPrimitiveComponent Interface
-	virtual int32 GetNumMaterials() const OVERRIDE; 
-	virtual UMaterialInterface* GetMaterial(int32 ElementIndex) const OVERRIDE;
-	virtual void SetMaterial(int32 ElementIndex, UMaterialInterface* Material) OVERRIDE;
-	virtual FBoxSphereBounds CalcBounds(const FTransform & LocalToWorld) const OVERRIDE;
-	virtual FPrimitiveSceneProxy* CreateSceneProxy() OVERRIDE;
-	virtual void GetUsedMaterials( TArray<UMaterialInterface*>& OutMaterials ) const OVERRIDE;
+	virtual int32 GetNumMaterials() const override; 
+	virtual UMaterialInterface* GetMaterial(int32 ElementIndex) const override;
+	virtual void SetMaterial(int32 ElementIndex, UMaterialInterface* Material) override;
+	virtual FBoxSphereBounds CalcBounds(const FTransform & LocalToWorld) const override;
+	virtual FPrimitiveSceneProxy* CreateSceneProxy() override;
+	virtual void GetUsedMaterials( TArray<UMaterialInterface*>& OutMaterials ) const override;
 	//End UPrimitiveComponent Interface
 
 	// Begin USceneComonent Interface
 protected:
-	virtual bool ShouldActivate() const OVERRIDE;
+	virtual bool ShouldActivate() const override;
 
 public:
-	virtual void Activate(bool bReset=false) OVERRIDE;
-	virtual void Deactivate() OVERRIDE;
-	virtual void ApplyWorldOffset(const FVector& InOffset, bool bWorldShift) OVERRIDE;
+	virtual void Activate(bool bReset=false) override;
+	virtual void Deactivate() override;
+	virtual void ApplyWorldOffset(const FVector& InOffset, bool bWorldShift) override;
 	// End USceneComponent Interface
 
 	/** Activate the system */
diff --git a/Engine/Source/Runtime/Engine/Classes/Particles/ParticleSystemReplay.h b/Engine/Source/Runtime/Engine/Classes/Particles/ParticleSystemReplay.h
--- a/Engine/Source/Runtime/Engine/Classes/Particles/ParticleSystemReplay.h
+++ b/Engine/Source/Runtime/Engine/Classes/Particles/ParticleSystemReplay.h
@@ -72,7 +72,7 @@ class UParticleSystemReplay : public UObject
 
 
 	// Begin UObject Interface
-	virtual void Serialize( FArchive& Ar ) OVERRIDE;
+	virtual void Serialize( FArchive& Ar ) override;
 	// End UObject Interface
 };
 
diff --git a/Engine/Source/Runtime/Engine/Classes/Particles/Rotation/ParticleModuleMeshRotation.h b/Engine/Source/Runtime/Engine/Classes/Particles/Rotation/ParticleModuleMeshRotation.h
--- a/Engine/Source/Runtime/Engine/Classes/Particles/Rotation/ParticleModuleMeshRotation.h
+++ b/Engine/Source/Runtime/Engine/Classes/Particles/Rotation/ParticleModuleMeshRotation.h
@@ -22,13 +22,13 @@ class UParticleModuleMeshRotation : public UParticleModuleRotationBase
 	uint32 bInheritParent:1;
 
 	//Begin UObject Interface
-	virtual void	PostInitProperties() OVERRIDE;
-	virtual void	Serialize(FArchive& Ar) OVERRIDE;
+	virtual void	PostInitProperties() override;
+	virtual void	Serialize(FArchive& Ar) override;
 	//End UObject Interface
 
 	//Begin UParticleModule Interface
-	virtual void	Spawn(FParticleEmitterInstance* Owner, int32 Offset, float SpawnTime, FBaseParticle* ParticleBase) OVERRIDE;
-	virtual bool	TouchesMeshRotation() const OVERRIDE { return true; }
+	virtual void	Spawn(FParticleEmitterInstance* Owner, int32 Offset, float SpawnTime, FBaseParticle* ParticleBase) override;
+	virtual bool	TouchesMeshRotation() const override { return true; }
 	//End UParticleModule Interface
 
 	/**
diff --git a/Engine/Source/Runtime/Engine/Classes/Particles/Rotation/ParticleModuleRotation.h b/Engine/Source/Runtime/Engine/Classes/Particles/Rotation/ParticleModuleRotation.h
--- a/Engine/Source/Runtime/Engine/Classes/Particles/Rotation/ParticleModuleRotation.h
+++ b/Engine/Source/Runtime/Engine/Classes/Particles/Rotation/ParticleModuleRotation.h
@@ -22,14 +22,14 @@ class UParticleModuleRotation : public UParticleModuleRotationBase
 
 	//Begin UObject Interface
 #if WITH_EDITOR
-	virtual void PostEditChangeProperty(FPropertyChangedEvent& PropertyChangedEvent) OVERRIDE;
+	virtual void PostEditChangeProperty(FPropertyChangedEvent& PropertyChangedEvent) override;
 #endif // WITH_EDITOR
-	virtual void	PostInitProperties() OVERRIDE;
-	virtual void	Serialize(FArchive& Ar) OVERRIDE;
+	virtual void	PostInitProperties() override;
+	virtual void	Serialize(FArchive& Ar) override;
 	//End UObject Interface
 
 	//Begin UParticleModule Interface
-	virtual void	Spawn(FParticleEmitterInstance* Owner, int32 Offset, float SpawnTime, FBaseParticle* ParticleBase) OVERRIDE;
+	virtual void	Spawn(FParticleEmitterInstance* Owner, int32 Offset, float SpawnTime, FBaseParticle* ParticleBase) override;
 	//End UParticleModule Interface
 
 	/**
diff --git a/Engine/Source/Runtime/Engine/Classes/Particles/Rotation/ParticleModuleRotationOverLifetime.h b/Engine/Source/Runtime/Engine/Classes/Particles/Rotation/ParticleModuleRotationOverLifetime.h
--- a/Engine/Source/Runtime/Engine/Classes/Particles/Rotation/ParticleModuleRotationOverLifetime.h
+++ b/Engine/Source/Runtime/Engine/Classes/Particles/Rotation/ParticleModuleRotationOverLifetime.h
@@ -29,13 +29,13 @@ class UParticleModuleRotationOverLifetime : public UParticleModuleRotationBase
 
 	//Begin UObject Interface
 #if WITH_EDITOR
-	virtual void PostEditChangeProperty(FPropertyChangedEvent& PropertyChangedEvent) OVERRIDE;
+	virtual void PostEditChangeProperty(FPropertyChangedEvent& PropertyChangedEvent) override;
 #endif // WITH_EDITOR
-	virtual void	PostInitProperties() OVERRIDE;
+	virtual void	PostInitProperties() override;
 	//End UObject Interface
 
 	//Begin UParticleModule Interface
-	virtual void	Update(FParticleEmitterInstance* Owner, int32 Offset, float DeltaTime) OVERRIDE;
+	virtual void	Update(FParticleEmitterInstance* Owner, int32 Offset, float DeltaTime) override;
 	//End UParticleModule Interface
 };
 
diff --git a/Engine/Source/Runtime/Engine/Classes/Particles/RotationRate/ParticleModuleMeshRotationRate.h b/Engine/Source/Runtime/Engine/Classes/Particles/RotationRate/ParticleModuleMeshRotationRate.h
--- a/Engine/Source/Runtime/Engine/Classes/Particles/RotationRate/ParticleModuleMeshRotationRate.h
+++ b/Engine/Source/Runtime/Engine/Classes/Particles/RotationRate/ParticleModuleMeshRotationRate.h
@@ -18,14 +18,14 @@ class UParticleModuleMeshRotationRate : public UParticleModuleRotationRateBase
 	struct FRawDistributionVector StartRotationRate;
 
 	//Begin UObject Interface
-	virtual void	PostInitProperties() OVERRIDE;
-	virtual void	Serialize(FArchive& Ar) OVERRIDE;
+	virtual void	PostInitProperties() override;
+	virtual void	Serialize(FArchive& Ar) override;
 	//End UObject Interface
 
 	// Begin UParticleModule Interface
-	virtual void	Spawn(FParticleEmitterInstance* Owner, int32 Offset, float SpawnTime, FBaseParticle* ParticleBase) OVERRIDE;
-	virtual void SetToSensibleDefaults(UParticleEmitter* Owner) OVERRIDE;
-	virtual bool	TouchesMeshRotation() const OVERRIDE { return true; }
+	virtual void	Spawn(FParticleEmitterInstance* Owner, int32 Offset, float SpawnTime, FBaseParticle* ParticleBase) override;
+	virtual void SetToSensibleDefaults(UParticleEmitter* Owner) override;
+	virtual bool	TouchesMeshRotation() const override { return true; }
 	// End UParticleModule Interface
 
 	/**
diff --git a/Engine/Source/Runtime/Engine/Classes/Particles/RotationRate/ParticleModuleMeshRotationRateMultiplyLife.h b/Engine/Source/Runtime/Engine/Classes/Particles/RotationRate/ParticleModuleMeshRotationRateMultiplyLife.h
--- a/Engine/Source/Runtime/Engine/Classes/Particles/RotationRate/ParticleModuleMeshRotationRateMultiplyLife.h
+++ b/Engine/Source/Runtime/Engine/Classes/Particles/RotationRate/ParticleModuleMeshRotationRateMultiplyLife.h
@@ -19,14 +19,14 @@ class UParticleModuleMeshRotationRateMultiplyLife : public UParticleModuleRotati
 
 
 	//Begin UObject Interface
-	virtual void	PostInitProperties() OVERRIDE;
-	virtual void	Serialize(FArchive& Ar) OVERRIDE;
+	virtual void	PostInitProperties() override;
+	virtual void	Serialize(FArchive& Ar) override;
 	//End UObject Interface
 
 	// Begin UParticleModule Interface
-	virtual void	Spawn(FParticleEmitterInstance* Owner, int32 Offset, float SpawnTime, FBaseParticle* ParticleBase) OVERRIDE;
-	virtual void	Update(FParticleEmitterInstance* Owner, int32 Offset, float DeltaTime) OVERRIDE;
-	virtual void SetToSensibleDefaults(UParticleEmitter* Owner) OVERRIDE;
+	virtual void	Spawn(FParticleEmitterInstance* Owner, int32 Offset, float SpawnTime, FBaseParticle* ParticleBase) override;
+	virtual void	Update(FParticleEmitterInstance* Owner, int32 Offset, float DeltaTime) override;
+	virtual void SetToSensibleDefaults(UParticleEmitter* Owner) override;
 	// End UParticleModule Interface
 };
 
diff --git a/Engine/Source/Runtime/Engine/Classes/Particles/RotationRate/ParticleModuleMeshRotationRateOverLife.h b/Engine/Source/Runtime/Engine/Classes/Particles/RotationRate/ParticleModuleMeshRotationRateOverLife.h
--- a/Engine/Source/Runtime/Engine/Classes/Particles/RotationRate/ParticleModuleMeshRotationRateOverLife.h
+++ b/Engine/Source/Runtime/Engine/Classes/Particles/RotationRate/ParticleModuleMeshRotationRateOverLife.h
@@ -25,14 +25,14 @@ class UParticleModuleMeshRotationRateOverLife : public UParticleModuleRotationRa
 	uint32 bScaleRotRate:1;
 
 	//Begin UObject Interface
-	virtual void	PostInitProperties() OVERRIDE;
+	virtual void	PostInitProperties() override;
 	//End UObject Interface
 
 	// Begin UParticleModule Interface
-	virtual void Spawn(FParticleEmitterInstance* Owner, int32 Offset, float SpawnTime, FBaseParticle* ParticleBase) OVERRIDE;
-	virtual void Update(FParticleEmitterInstance* Owner, int32 Offset, float DeltaTime) OVERRIDE;
-	virtual void SetToSensibleDefaults(UParticleEmitter* Owner) OVERRIDE;
-	virtual bool TouchesMeshRotation() const OVERRIDE { return true; }
+	virtual void Spawn(FParticleEmitterInstance* Owner, int32 Offset, float SpawnTime, FBaseParticle* ParticleBase) override;
+	virtual void Update(FParticleEmitterInstance* Owner, int32 Offset, float DeltaTime) override;
+	virtual void SetToSensibleDefaults(UParticleEmitter* Owner) override;
+	virtual bool TouchesMeshRotation() const override { return true; }
 	// End UParticleModule Interface
 };
 
diff --git a/Engine/Source/Runtime/Engine/Classes/Particles/RotationRate/ParticleModuleRotationRate.h b/Engine/Source/Runtime/Engine/Classes/Particles/RotationRate/ParticleModuleRotationRate.h
--- a/Engine/Source/Runtime/Engine/Classes/Particles/RotationRate/ParticleModuleRotationRate.h
+++ b/Engine/Source/Runtime/Engine/Classes/Particles/RotationRate/ParticleModuleRotationRate.h
@@ -22,16 +22,16 @@ class UParticleModuleRotationRate : public UParticleModuleRotationRateBase
 
 	//Begin UObject Interface
 #if WITH_EDITOR
-	virtual void PostEditChangeProperty(FPropertyChangedEvent& PropertyChangedEvent) OVERRIDE;
+	virtual void PostEditChangeProperty(FPropertyChangedEvent& PropertyChangedEvent) override;
 #endif // WITH_EDITOR
-	virtual void	PostInitProperties() OVERRIDE;
-	virtual void	Serialize(FArchive& Ar) OVERRIDE;
+	virtual void	PostInitProperties() override;
+	virtual void	Serialize(FArchive& Ar) override;
 	//End UObject Interface
 
 	// Begin UParticleModule Interface
-	virtual void CompileModule( struct FParticleEmitterBuildInfo& EmitterInfo ) OVERRIDE;
-	virtual void Spawn(FParticleEmitterInstance* Owner, int32 Offset, float SpawnTime, FBaseParticle* ParticleBase) OVERRIDE;
-	virtual void SetToSensibleDefaults(UParticleEmitter* Owner) OVERRIDE;
+	virtual void CompileModule( struct FParticleEmitterBuildInfo& EmitterInfo ) override;
+	virtual void Spawn(FParticleEmitterInstance* Owner, int32 Offset, float SpawnTime, FBaseParticle* ParticleBase) override;
+	virtual void SetToSensibleDefaults(UParticleEmitter* Owner) override;
 	// Begin UParticleModule Interface
 
 	/**
diff --git a/Engine/Source/Runtime/Engine/Classes/Particles/RotationRate/ParticleModuleRotationRateMultiplyLife.h b/Engine/Source/Runtime/Engine/Classes/Particles/RotationRate/ParticleModuleRotationRateMultiplyLife.h
--- a/Engine/Source/Runtime/Engine/Classes/Particles/RotationRate/ParticleModuleRotationRateMultiplyLife.h
+++ b/Engine/Source/Runtime/Engine/Classes/Particles/RotationRate/ParticleModuleRotationRateMultiplyLife.h
@@ -22,16 +22,16 @@ class UParticleModuleRotationRateMultiplyLife : public UParticleModuleRotationRa
 
 	//Begin UObject Interface
 #if WITH_EDITOR
-	virtual void PostEditChangeProperty(FPropertyChangedEvent& PropertyChangedEvent) OVERRIDE;
+	virtual void PostEditChangeProperty(FPropertyChangedEvent& PropertyChangedEvent) override;
 #endif // WITH_EDITOR
-	virtual void PostInitProperties() OVERRIDE;
-	virtual void Serialize(FArchive& Ar) OVERRIDE;
+	virtual void PostInitProperties() override;
+	virtual void Serialize(FArchive& Ar) override;
 	//End UObject Interface
 
 	// Begin UParticleModule Interface
-	virtual void Spawn(FParticleEmitterInstance* Owner, int32 Offset, float SpawnTime, FBaseParticle* ParticleBase) OVERRIDE;
-	virtual void Update(FParticleEmitterInstance* Owner, int32 Offset, float DeltaTime) OVERRIDE;
-	virtual void SetToSensibleDefaults(UParticleEmitter* Owner) OVERRIDE;
+	virtual void Spawn(FParticleEmitterInstance* Owner, int32 Offset, float SpawnTime, FBaseParticle* ParticleBase) override;
+	virtual void Update(FParticleEmitterInstance* Owner, int32 Offset, float DeltaTime) override;
+	virtual void SetToSensibleDefaults(UParticleEmitter* Owner) override;
 	// End UParticleModule Interface
 };
 
diff --git a/Engine/Source/Runtime/Engine/Classes/Particles/Size/ParticleModuleSize.h b/Engine/Source/Runtime/Engine/Classes/Particles/Size/ParticleModuleSize.h
--- a/Engine/Source/Runtime/Engine/Classes/Particles/Size/ParticleModuleSize.h
+++ b/Engine/Source/Runtime/Engine/Classes/Particles/Size/ParticleModuleSize.h
@@ -23,15 +23,15 @@ class UParticleModuleSize : public UParticleModuleSizeBase
 
 	//Begin UObject Interface
 #if WITH_EDITOR
-	virtual void PostEditChangeProperty(FPropertyChangedEvent& PropertyChangedEvent) OVERRIDE;
+	virtual void PostEditChangeProperty(FPropertyChangedEvent& PropertyChangedEvent) override;
 #endif // WITH_EDITOR
-	virtual void PostInitProperties() OVERRIDE;
-	virtual void Serialize(FArchive& Ar) OVERRIDE;
+	virtual void PostInitProperties() override;
+	virtual void Serialize(FArchive& Ar) override;
 	// End UObject Interface
 
 	// Begin UParticleModule Interface
-	virtual void CompileModule( struct FParticleEmitterBuildInfo& EmitterInfo ) OVERRIDE;
-	virtual void Spawn(FParticleEmitterInstance* Owner, int32 Offset, float SpawnTime, FBaseParticle* ParticleBase) OVERRIDE;
+	virtual void CompileModule( struct FParticleEmitterBuildInfo& EmitterInfo ) override;
+	virtual void Spawn(FParticleEmitterInstance* Owner, int32 Offset, float SpawnTime, FBaseParticle* ParticleBase) override;
 	// End UParticleModule Interface
 
 	/**
diff --git a/Engine/Source/Runtime/Engine/Classes/Particles/Size/ParticleModuleSizeMultiplyLife.h b/Engine/Source/Runtime/Engine/Classes/Particles/Size/ParticleModuleSizeMultiplyLife.h
--- a/Engine/Source/Runtime/Engine/Classes/Particles/Size/ParticleModuleSizeMultiplyLife.h
+++ b/Engine/Source/Runtime/Engine/Classes/Particles/Size/ParticleModuleSizeMultiplyLife.h
@@ -43,21 +43,21 @@ class UParticleModuleSizeMultiplyLife : public UParticleModuleSizeBase
 
 	//Begin UObject Interface
 #if WITH_EDITOR
-	virtual void PostEditChangeProperty(FPropertyChangedEvent& PropertyChangedEvent) OVERRIDE;
+	virtual void PostEditChangeProperty(FPropertyChangedEvent& PropertyChangedEvent) override;
 #endif // WITH_EDITOR
-	virtual void PostInitProperties() OVERRIDE;
-	virtual void Serialize(FArchive& Ar) OVERRIDE;
+	virtual void PostInitProperties() override;
+	virtual void Serialize(FArchive& Ar) override;
 	// End UObject Interface
 
 	// Begin UParticleModule Interface
-	virtual void CompileModule( struct FParticleEmitterBuildInfo& EmitterInfo ) OVERRIDE;
-	virtual void Spawn(FParticleEmitterInstance* Owner, int32 Offset, float SpawnTime, FBaseParticle* ParticleBase) OVERRIDE;
-	virtual void	Update(FParticleEmitterInstance* Owner, int32 Offset, float DeltaTime) OVERRIDE;
-	virtual void SetToSensibleDefaults(UParticleEmitter* Owner) OVERRIDE;
-	virtual bool   IsSizeMultiplyLife() OVERRIDE { return true; };
+	virtual void CompileModule( struct FParticleEmitterBuildInfo& EmitterInfo ) override;
+	virtual void Spawn(FParticleEmitterInstance* Owner, int32 Offset, float SpawnTime, FBaseParticle* ParticleBase) override;
+	virtual void	Update(FParticleEmitterInstance* Owner, int32 Offset, float DeltaTime) override;
+	virtual void SetToSensibleDefaults(UParticleEmitter* Owner) override;
+	virtual bool   IsSizeMultiplyLife() override { return true; };
 
 #if WITH_EDITOR
-	virtual bool IsValidForLODLevel(UParticleLODLevel* LODLevel, FString& OutErrorString) OVERRIDE;
+	virtual bool IsValidForLODLevel(UParticleLODLevel* LODLevel, FString& OutErrorString) override;
 #endif
 
 	// End UParticleModule Interface
diff --git a/Engine/Source/Runtime/Engine/Classes/Particles/Size/ParticleModuleSizeScale.h b/Engine/Source/Runtime/Engine/Classes/Particles/Size/ParticleModuleSizeScale.h
--- a/Engine/Source/Runtime/Engine/Classes/Particles/Size/ParticleModuleSizeScale.h
+++ b/Engine/Source/Runtime/Engine/Classes/Particles/Size/ParticleModuleSizeScale.h
@@ -35,20 +35,20 @@ class UParticleModuleSizeScale : public UParticleModuleSizeBase
 
 	//Begin UObject Interface
 #if WITH_EDITOR
-	virtual void PostEditChangeProperty(FPropertyChangedEvent& PropertyChangedEvent) OVERRIDE;
+	virtual void PostEditChangeProperty(FPropertyChangedEvent& PropertyChangedEvent) override;
 #endif // WITH_EDITOR
-	virtual void PostInitProperties() OVERRIDE;
-	virtual void Serialize(FArchive& Ar) OVERRIDE;
+	virtual void PostInitProperties() override;
+	virtual void Serialize(FArchive& Ar) override;
 	// End UObject Interface
 
 	// Begin UParticleModule Interface
-	virtual void CompileModule( struct FParticleEmitterBuildInfo& EmitterInfo ) OVERRIDE;
-	virtual void Spawn(FParticleEmitterInstance* Owner, int32 Offset, float SpawnTime, FBaseParticle* ParticleBase) OVERRIDE;
-	virtual void Update(FParticleEmitterInstance* Owner, int32 Offset, float DeltaTime) OVERRIDE;
-	virtual void SetToSensibleDefaults(UParticleEmitter* Owner) OVERRIDE;
+	virtual void CompileModule( struct FParticleEmitterBuildInfo& EmitterInfo ) override;
+	virtual void Spawn(FParticleEmitterInstance* Owner, int32 Offset, float SpawnTime, FBaseParticle* ParticleBase) override;
+	virtual void Update(FParticleEmitterInstance* Owner, int32 Offset, float DeltaTime) override;
+	virtual void SetToSensibleDefaults(UParticleEmitter* Owner) override;
 
 #if WITH_EDITOR
-	virtual bool IsValidForLODLevel(UParticleLODLevel* LODLevel, FString& OutErrorString) OVERRIDE;
+	virtual bool IsValidForLODLevel(UParticleLODLevel* LODLevel, FString& OutErrorString) override;
 #endif
 
 	// End UParticleModule Interface
diff --git a/Engine/Source/Runtime/Engine/Classes/Particles/Size/ParticleModuleSizeScaleBySpeed.h b/Engine/Source/Runtime/Engine/Classes/Particles/Size/ParticleModuleSizeScaleBySpeed.h
--- a/Engine/Source/Runtime/Engine/Classes/Particles/Size/ParticleModuleSizeScaleBySpeed.h
+++ b/Engine/Source/Runtime/Engine/Classes/Particles/Size/ParticleModuleSizeScaleBySpeed.h
@@ -23,8 +23,8 @@ class UParticleModuleSizeScaleBySpeed : public UParticleModuleSizeBase
 
 
 	// Begin UParticleModule Interface
-	virtual void Update(FParticleEmitterInstance* Owner, int32 Offset, float DeltaTime) OVERRIDE;
-	virtual void CompileModule(struct FParticleEmitterBuildInfo& EmitterInfo) OVERRIDE;
+	virtual void Update(FParticleEmitterInstance* Owner, int32 Offset, float DeltaTime) override;
+	virtual void CompileModule(struct FParticleEmitterBuildInfo& EmitterInfo) override;
 	// End UParticleModule Interface
 };
 
diff --git a/Engine/Source/Runtime/Engine/Classes/Particles/Spawn/ParticleModuleSpawn.h b/Engine/Source/Runtime/Engine/Classes/Particles/Spawn/ParticleModuleSpawn.h
--- a/Engine/Source/Runtime/Engine/Classes/Particles/Spawn/ParticleModuleSpawn.h
+++ b/Engine/Source/Runtime/Engine/Classes/Particles/Spawn/ParticleModuleSpawn.h
@@ -38,22 +38,22 @@ class UParticleModuleSpawn : public UParticleModuleSpawnBase
 
 	// Begin UObject Interface
 #if WITH_EDITOR
-	virtual void	PostEditChangeProperty(FPropertyChangedEvent& PropertyChangedEvent) OVERRIDE;
+	virtual void	PostEditChangeProperty(FPropertyChangedEvent& PropertyChangedEvent) override;
 #endif // WITH_EDITOR
-	virtual void	PostInitProperties() OVERRIDE;
-	virtual void	Serialize(FArchive& Ar) OVERRIDE;
+	virtual void	PostInitProperties() override;
+	virtual void	Serialize(FArchive& Ar) override;
 	// End UObject Interface
 
 	// Begin UParticleModule Interface
-	virtual bool	GenerateLODModuleValues(UParticleModule* SourceModule, float Percentage, UParticleLODLevel* LODLevel) OVERRIDE;
+	virtual bool	GenerateLODModuleValues(UParticleModule* SourceModule, float Percentage, UParticleLODLevel* LODLevel) override;
 	// End UParticleModule Interface
 
 	// Begin UParticleModuleSpawnBase Interface
 	virtual bool GetSpawnAmount(FParticleEmitterInstance* Owner, int32 Offset, float OldLeftover, 
-		float DeltaTime, int32& Number, float& Rate) OVERRIDE;
-	virtual float GetMaximumSpawnRate() OVERRIDE;
-	virtual float GetEstimatedSpawnRate() OVERRIDE;
-	virtual int32 GetMaximumBurstCount() OVERRIDE;
+		float DeltaTime, int32& Number, float& Rate) override;
+	virtual float GetMaximumSpawnRate() override;
+	virtual float GetEstimatedSpawnRate() override;
+	virtual int32 GetMaximumBurstCount() override;
 	// End UParticleModuleSpawnBase Interface
 };
 
diff --git a/Engine/Source/Runtime/Engine/Classes/Particles/Spawn/ParticleModuleSpawnBase.h b/Engine/Source/Runtime/Engine/Classes/Particles/Spawn/ParticleModuleSpawnBase.h
--- a/Engine/Source/Runtime/Engine/Classes/Particles/Spawn/ParticleModuleSpawnBase.h
+++ b/Engine/Source/Runtime/Engine/Classes/Particles/Spawn/ParticleModuleSpawnBase.h
@@ -30,7 +30,7 @@ class UParticleModuleSpawnBase : public UParticleModule
 
 
 	// Begin UParticleModule Interface
-	virtual EModuleType	GetModuleType() const OVERRIDE {	return EPMT_Spawn;	}
+	virtual EModuleType	GetModuleType() const override {	return EPMT_Spawn;	}
 	// End UParticleModule Interface
 
 	/**
diff --git a/Engine/Source/Runtime/Engine/Classes/Particles/Spawn/ParticleModuleSpawnPerUnit.h b/Engine/Source/Runtime/Engine/Classes/Particles/Spawn/ParticleModuleSpawnPerUnit.h
--- a/Engine/Source/Runtime/Engine/Classes/Particles/Spawn/ParticleModuleSpawnPerUnit.h
+++ b/Engine/Source/Runtime/Engine/Classes/Particles/Spawn/ParticleModuleSpawnPerUnit.h
@@ -68,20 +68,20 @@ class UParticleModuleSpawnPerUnit : public UParticleModuleSpawnBase
 
 	// Begin UObject Interface
 #if WITH_EDITOR
-	virtual void PostEditChangeProperty(FPropertyChangedEvent& PropertyChangedEvent) OVERRIDE;
+	virtual void PostEditChangeProperty(FPropertyChangedEvent& PropertyChangedEvent) override;
 #endif // WITH_EDITOR
-	virtual void PostInitProperties() OVERRIDE;
-	virtual void Serialize(FArchive& Ar) OVERRIDE;
+	virtual void PostInitProperties() override;
+	virtual void Serialize(FArchive& Ar) override;
 	// End UObject Interface
 
 	// Begin UParticleModule Interface
-	virtual void CompileModule( struct FParticleEmitterBuildInfo& EmitterInfo ) OVERRIDE;
-	virtual uint32	RequiredBytesPerInstance(FParticleEmitterInstance* Owner = NULL) OVERRIDE;
+	virtual void CompileModule( struct FParticleEmitterBuildInfo& EmitterInfo ) override;
+	virtual uint32	RequiredBytesPerInstance(FParticleEmitterInstance* Owner = NULL) override;
 	// End UParticleModule Interface
 
 	// Begin UParticleModuleSpawnBase Interface
 	virtual bool GetSpawnAmount(FParticleEmitterInstance* Owner, int32 Offset, float OldLeftover, 
-		float DeltaTime, int32& Number, float& Rate) OVERRIDE;
+		float DeltaTime, int32& Number, float& Rate) override;
 	// End UParticleModuleSpawnBase Interface
 };
 
diff --git a/Engine/Source/Runtime/Engine/Classes/Particles/SubUV/ParticleModuleSubUV.h b/Engine/Source/Runtime/Engine/Classes/Particles/SubUV/ParticleModuleSubUV.h
--- a/Engine/Source/Runtime/Engine/Classes/Particles/SubUV/ParticleModuleSubUV.h
+++ b/Engine/Source/Runtime/Engine/Classes/Particles/SubUV/ParticleModuleSubUV.h
@@ -30,17 +30,17 @@ class UParticleModuleSubUV : public UParticleModuleSubUVBase
 	
 	//Begin UObject Interface
 #if WITH_EDITOR
-	virtual void PostEditChangeProperty(FPropertyChangedEvent& PropertyChangedEvent) OVERRIDE;
+	virtual void PostEditChangeProperty(FPropertyChangedEvent& PropertyChangedEvent) override;
 #endif // WITH_EDITOR
-	virtual void	PostInitProperties() OVERRIDE;
-	virtual void	Serialize(FArchive& Ar) OVERRIDE;
+	virtual void	PostInitProperties() override;
+	virtual void	Serialize(FArchive& Ar) override;
 	//End UObject Interface
 
 	// Begin UParticleModule Interface
-	virtual void CompileModule( struct FParticleEmitterBuildInfo& EmitterInfo ) OVERRIDE;
-	virtual void Spawn(FParticleEmitterInstance* Owner, int32 Offset, float SpawnTime, FBaseParticle* ParticleBase) OVERRIDE;
-	virtual void Update(FParticleEmitterInstance* Owner, int32 Offset, float DeltaTime) OVERRIDE;
-	virtual void SetToSensibleDefaults(UParticleEmitter* Owner) OVERRIDE;
+	virtual void CompileModule( struct FParticleEmitterBuildInfo& EmitterInfo ) override;
+	virtual void Spawn(FParticleEmitterInstance* Owner, int32 Offset, float SpawnTime, FBaseParticle* ParticleBase) override;
+	virtual void Update(FParticleEmitterInstance* Owner, int32 Offset, float DeltaTime) override;
+	virtual void SetToSensibleDefaults(UParticleEmitter* Owner) override;
 	// End UParticleModule Interface
 
 	/**
@@ -58,7 +58,7 @@ class UParticleModuleSubUV : public UParticleModuleSubUVBase
 		EParticleSubUVInterpMethod InterpMethod, FFullSubUVPayload& SubUVPayload, float DeltaTime);
 
 #if WITH_EDITOR
-	virtual bool IsValidForLODLevel(UParticleLODLevel* LODLevel, FString& OutErrorString) OVERRIDE;
+	virtual bool IsValidForLODLevel(UParticleLODLevel* LODLevel, FString& OutErrorString) override;
 #endif
 
 protected:
diff --git a/Engine/Source/Runtime/Engine/Classes/Particles/SubUV/ParticleModuleSubUVMovie.h b/Engine/Source/Runtime/Engine/Classes/Particles/SubUV/ParticleModuleSubUVMovie.h
--- a/Engine/Source/Runtime/Engine/Classes/Particles/SubUV/ParticleModuleSubUVMovie.h
+++ b/Engine/Source/Runtime/Engine/Classes/Particles/SubUV/ParticleModuleSubUVMovie.h
@@ -38,22 +38,22 @@ class UParticleModuleSubUVMovie : public UParticleModuleSubUV
 
 	//Begin UObject Interface
 #if WITH_EDITOR
-	virtual void PostEditChangeProperty(FPropertyChangedEvent& PropertyChangedEvent) OVERRIDE;
+	virtual void PostEditChangeProperty(FPropertyChangedEvent& PropertyChangedEvent) override;
 #endif // WITH_EDITOR
-	virtual void PostInitProperties() OVERRIDE;
-	virtual void Serialize(FArchive& Ar) OVERRIDE;
+	virtual void PostInitProperties() override;
+	virtual void Serialize(FArchive& Ar) override;
 	//End UObject Interface
 
 	// Begin UParticleModule Interface
-	virtual void Spawn(FParticleEmitterInstance* Owner, int32 Offset, float SpawnTime, FBaseParticle* ParticleBase) OVERRIDE;
-	virtual uint32 RequiredBytes(FParticleEmitterInstance* Owner = NULL) OVERRIDE;
-	virtual void SetToSensibleDefaults(UParticleEmitter* Owner) OVERRIDE;
-	virtual void GetCurveObjects(TArray<FParticleCurvePair>& OutCurves) OVERRIDE;
+	virtual void Spawn(FParticleEmitterInstance* Owner, int32 Offset, float SpawnTime, FBaseParticle* ParticleBase) override;
+	virtual uint32 RequiredBytes(FParticleEmitterInstance* Owner = NULL) override;
+	virtual void SetToSensibleDefaults(UParticleEmitter* Owner) override;
+	virtual void GetCurveObjects(TArray<FParticleCurvePair>& OutCurves) override;
 	// End UParticleModule Interface
 	
 	// Begin UParticleModuleSubUV Interface
 	virtual float DetermineImageIndex(FParticleEmitterInstance* Owner, int32 Offset, FBaseParticle* Particle, 
-		EParticleSubUVInterpMethod InterpMethod, FFullSubUVPayload& SubUVPayload, float DeltaTime) OVERRIDE;
+		EParticleSubUVInterpMethod InterpMethod, FFullSubUVPayload& SubUVPayload, float DeltaTime) override;
 	// End UParticleModuleSubUV Interface
 	
 };
diff --git a/Engine/Source/Runtime/Engine/Classes/Particles/Trail/ParticleModuleTrailBase.h b/Engine/Source/Runtime/Engine/Classes/Particles/Trail/ParticleModuleTrailBase.h
--- a/Engine/Source/Runtime/Engine/Classes/Particles/Trail/ParticleModuleTrailBase.h
+++ b/Engine/Source/Runtime/Engine/Classes/Particles/Trail/ParticleModuleTrailBase.h
@@ -17,8 +17,8 @@ class UParticleModuleTrailBase : public UParticleModule
 
 
 	// Begin UParticleModule Interface
-	virtual EModuleType	GetModuleType() const OVERRIDE {	return EPMT_Trail;	}
-	virtual bool CanTickInAnyThread() OVERRIDE
+	virtual EModuleType	GetModuleType() const override {	return EPMT_Trail;	}
+	virtual bool CanTickInAnyThread() override
 	{
 		return false;
 	}
diff --git a/Engine/Source/Runtime/Engine/Classes/Particles/Trail/ParticleModuleTrailSource.h b/Engine/Source/Runtime/Engine/Classes/Particles/Trail/ParticleModuleTrailSource.h
--- a/Engine/Source/Runtime/Engine/Classes/Particles/Trail/ParticleModuleTrailSource.h
+++ b/Engine/Source/Runtime/Engine/Classes/Particles/Trail/ParticleModuleTrailSource.h
@@ -79,15 +79,15 @@ class UParticleModuleTrailSource : public UParticleModuleTrailBase
 
 	// Begin UObject Interface
 #if WITH_EDITOR
-	virtual void	PostEditChangeProperty(FPropertyChangedEvent& PropertyChangedEvent) OVERRIDE;
+	virtual void	PostEditChangeProperty(FPropertyChangedEvent& PropertyChangedEvent) override;
 #endif // WITH_EDITOR
-	virtual void	PostInitProperties() OVERRIDE;
-	virtual void 	Serialize(FArchive& Ar) OVERRIDE;
+	virtual void	PostInitProperties() override;
+	virtual void 	Serialize(FArchive& Ar) override;
 	// End UObject Interface
 
 	// Begin UParticleModule Interface
-	virtual void	AutoPopulateInstanceProperties(UParticleSystemComponent* PSysComp) OVERRIDE;
-	virtual void GetParticleSysParamsUtilized(TArray<FString>& ParticleSysParamList) OVERRIDE;
+	virtual void	AutoPopulateInstanceProperties(UParticleSystemComponent* PSysComp) override;
+	virtual void GetParticleSysParamsUtilized(TArray<FString>& ParticleSysParamList) override;
 	// End UParticleModule Interface
 
 	/**
diff --git a/Engine/Source/Runtime/Engine/Classes/Particles/TypeData/ParticleModuleTypeDataAnimTrail.h b/Engine/Source/Runtime/Engine/Classes/Particles/TypeData/ParticleModuleTypeDataAnimTrail.h
--- a/Engine/Source/Runtime/Engine/Classes/Particles/TypeData/ParticleModuleTypeDataAnimTrail.h
+++ b/Engine/Source/Runtime/Engine/Classes/Particles/TypeData/ParticleModuleTypeDataAnimTrail.h
@@ -65,15 +65,15 @@ class UParticleModuleTypeDataAnimTrail : public UParticleModuleTypeDataBase
 
 
 	// Begin UParticleModule Interface
-	virtual uint32 RequiredBytes(FParticleEmitterInstance* Owner = NULL) OVERRIDE;
-	virtual bool CanTickInAnyThread() OVERRIDE
+	virtual uint32 RequiredBytes(FParticleEmitterInstance* Owner = NULL) override;
+	virtual bool CanTickInAnyThread() override
 	{
 		return false;
 	}
 	// End UParticleModule Interface
 
 	// Begin UParticleModuleTypeDataBase Interface
-	virtual FParticleEmitterInstance* CreateInstance(UParticleEmitter* InEmitterParent, UParticleSystemComponent* InComponent) OVERRIDE;
+	virtual FParticleEmitterInstance* CreateInstance(UParticleEmitter* InEmitterParent, UParticleSystemComponent* InComponent) override;
 	// End UParticleModuleTypeDataBase Interface
 };
 
diff --git a/Engine/Source/Runtime/Engine/Classes/Particles/TypeData/ParticleModuleTypeDataBase.h b/Engine/Source/Runtime/Engine/Classes/Particles/TypeData/ParticleModuleTypeDataBase.h
--- a/Engine/Source/Runtime/Engine/Classes/Particles/TypeData/ParticleModuleTypeDataBase.h
+++ b/Engine/Source/Runtime/Engine/Classes/Particles/TypeData/ParticleModuleTypeDataBase.h
@@ -14,7 +14,7 @@ class UParticleModuleTypeDataBase : public UParticleModule
 
 
 	// Begin UParticleModule Interface
-	virtual EModuleType	GetModuleType() const OVERRIDE {	return EPMT_TypeData;	}
+	virtual EModuleType	GetModuleType() const override {	return EPMT_TypeData;	}
 	// End UParticleModule Interface
 
 	/**
diff --git a/Engine/Source/Runtime/Engine/Classes/Particles/TypeData/ParticleModuleTypeDataBeam2.h b/Engine/Source/Runtime/Engine/Classes/Particles/TypeData/ParticleModuleTypeDataBeam2.h
--- a/Engine/Source/Runtime/Engine/Classes/Particles/TypeData/ParticleModuleTypeDataBeam2.h
+++ b/Engine/Source/Runtime/Engine/Classes/Particles/TypeData/ParticleModuleTypeDataBeam2.h
@@ -176,19 +176,19 @@ class UParticleModuleTypeDataBeam2 : public UParticleModuleTypeDataBase
 
 	// Begin UObject Interface
 #if WITH_EDITOR
-	virtual void	PostEditChangeProperty(FPropertyChangedEvent& PropertyChangedEvent) OVERRIDE;
+	virtual void	PostEditChangeProperty(FPropertyChangedEvent& PropertyChangedEvent) override;
 #endif // WITH_EDITOR
-	virtual void	PostInitProperties() OVERRIDE;
-	virtual void	Serialize(FArchive& Ar) OVERRIDE;
+	virtual void	PostInitProperties() override;
+	virtual void	Serialize(FArchive& Ar) override;
 	// End UObject Interface
 
 
 	// Begin UParticleModule Interface
-	virtual void	Spawn(FParticleEmitterInstance* Owner, int32 Offset, float SpawnTime, FBaseParticle* ParticleBase) OVERRIDE;
-	virtual void	Update(FParticleEmitterInstance* Owner, int32 Offset, float DeltaTime) OVERRIDE;
-	virtual uint32	RequiredBytes(FParticleEmitterInstance* Owner = NULL) OVERRIDE;
-	virtual	bool	AddModuleCurvesToEditor(UInterpCurveEdSetup* EdSetup, TArray<const FCurveEdEntry*>& OutCurveEntries) OVERRIDE;
-	virtual bool CanTickInAnyThread() OVERRIDE
+	virtual void	Spawn(FParticleEmitterInstance* Owner, int32 Offset, float SpawnTime, FBaseParticle* ParticleBase) override;
+	virtual void	Update(FParticleEmitterInstance* Owner, int32 Offset, float DeltaTime) override;
+	virtual uint32	RequiredBytes(FParticleEmitterInstance* Owner = NULL) override;
+	virtual	bool	AddModuleCurvesToEditor(UInterpCurveEdSetup* EdSetup, TArray<const FCurveEdEntry*>& OutCurveEntries) override;
+	virtual bool CanTickInAnyThread() override
 	{
 		return false;
 	}
@@ -196,7 +196,7 @@ class UParticleModuleTypeDataBeam2 : public UParticleModuleTypeDataBase
 
 
 	// Begin UParticleModuleTypeDataBase Interface
-	virtual FParticleEmitterInstance* CreateInstance(UParticleEmitter* InEmitterParent, UParticleSystemComponent* InComponent) OVERRIDE;
+	virtual FParticleEmitterInstance* CreateInstance(UParticleEmitter* InEmitterParent, UParticleSystemComponent* InComponent) override;
 	// End UParticleModuleTypeDataBase Interface
 
 
diff --git a/Engine/Source/Runtime/Engine/Classes/Particles/TypeData/ParticleModuleTypeDataGpu.h b/Engine/Source/Runtime/Engine/Classes/Particles/TypeData/ParticleModuleTypeDataGpu.h
--- a/Engine/Source/Runtime/Engine/Classes/Particles/TypeData/ParticleModuleTypeDataGpu.h
+++ b/Engine/Source/Runtime/Engine/Classes/Particles/TypeData/ParticleModuleTypeDataGpu.h
@@ -408,14 +408,14 @@ class UParticleModuleTypeDataGpu : public UParticleModuleTypeDataBase
 
 
 	// Begin UObject Interface
-	virtual void PostLoad() OVERRIDE;
-	virtual void BeginDestroy() OVERRIDE;
+	virtual void PostLoad() override;
+	virtual void BeginDestroy() override;
 	// End UObject Interface
 
 	// Begin UParticleModuleTypeDataBase Interface
-	virtual void Build( struct FParticleEmitterBuildInfo& EmitterBuildInfo ) OVERRIDE;
+	virtual void Build( struct FParticleEmitterBuildInfo& EmitterBuildInfo ) override;
 	virtual bool RequiresBuild() const { return true; }
-	virtual FParticleEmitterInstance* CreateInstance(UParticleEmitter* InEmitterParent, UParticleSystemComponent* InComponent) OVERRIDE;
+	virtual FParticleEmitterInstance* CreateInstance(UParticleEmitter* InEmitterParent, UParticleSystemComponent* InComponent) override;
 	// End UParticleModuleTypeDataBase Interface
 };
 
diff --git a/Engine/Source/Runtime/Engine/Classes/Particles/TypeData/ParticleModuleTypeDataMesh.h b/Engine/Source/Runtime/Engine/Classes/Particles/TypeData/ParticleModuleTypeDataMesh.h
--- a/Engine/Source/Runtime/Engine/Classes/Particles/TypeData/ParticleModuleTypeDataMesh.h
+++ b/Engine/Source/Runtime/Engine/Classes/Particles/TypeData/ParticleModuleTypeDataMesh.h
@@ -179,7 +179,7 @@ class UParticleModuleTypeDataMesh : public UParticleModuleTypeDataBase
 
 	// Begin UObject Interface
 #if WITH_EDITOR
-	virtual void	PostEditChangeProperty(FPropertyChangedEvent& PropertyChangedEvent) OVERRIDE;
+	virtual void	PostEditChangeProperty(FPropertyChangedEvent& PropertyChangedEvent) override;
 #endif // WITH_EDITOR
 
 	virtual void	Serialize(FArchive& Ar) override;
@@ -189,14 +189,14 @@ class UParticleModuleTypeDataMesh : public UParticleModuleTypeDataBase
 	void CreateDistribution();
 
 	// Begin UParticleModule Interface
-	virtual void	SetToSensibleDefaults(UParticleEmitter* Owner) OVERRIDE;
+	virtual void	SetToSensibleDefaults(UParticleEmitter* Owner) override;
 	// End UParticleModule Interface
 
 	// Begin UParticleModuleTypeDataBase Interface
-	virtual FParticleEmitterInstance*	CreateInstance(UParticleEmitter* InEmitterParent, UParticleSystemComponent* InComponent) OVERRIDE;
-	virtual bool	SupportsSpecificScreenAlignmentFlags() const OVERRIDE {	return true;	}	
-	virtual bool	SupportsSubUV() const OVERRIDE { return true; }
-	virtual bool	IsAMeshEmitter() const OVERRIDE { return true; }
+	virtual FParticleEmitterInstance*	CreateInstance(UParticleEmitter* InEmitterParent, UParticleSystemComponent* InComponent) override;
+	virtual bool	SupportsSpecificScreenAlignmentFlags() const override {	return true;	}	
+	virtual bool	SupportsSubUV() const override { return true; }
+	virtual bool	IsAMeshEmitter() const override { return true; }
 	// End UParticleModuleTypeDataBase Interface
 };
 
diff --git a/Engine/Source/Runtime/Engine/Classes/Particles/TypeData/ParticleModuleTypeDataRibbon.h b/Engine/Source/Runtime/Engine/Classes/Particles/TypeData/ParticleModuleTypeDataRibbon.h
--- a/Engine/Source/Runtime/Engine/Classes/Particles/TypeData/ParticleModuleTypeDataRibbon.h
+++ b/Engine/Source/Runtime/Engine/Classes/Particles/TypeData/ParticleModuleTypeDataRibbon.h
@@ -161,16 +161,16 @@ class UParticleModuleTypeDataRibbon : public UParticleModuleTypeDataBase
 
 	// Begin UObject Interface
 #if WITH_EDITOR
-	virtual void PostEditChangeProperty(FPropertyChangedEvent& PropertyChangedEvent) OVERRIDE;
+	virtual void PostEditChangeProperty(FPropertyChangedEvent& PropertyChangedEvent) override;
 #endif // WITH_EDITOR
 	// End UObject Interface
 
 	// Begin UParticleModule Interface
-	virtual uint32 RequiredBytes(FParticleEmitterInstance* Owner = NULL) OVERRIDE;
+	virtual uint32 RequiredBytes(FParticleEmitterInstance* Owner = NULL) override;
 	// End UParticleModule Interface
 
 	// Begin UParticleModuleTypeDataBase Interface
-	virtual FParticleEmitterInstance* CreateInstance(UParticleEmitter* InEmitterParent, UParticleSystemComponent* InComponent) OVERRIDE;
+	virtual FParticleEmitterInstance* CreateInstance(UParticleEmitter* InEmitterParent, UParticleSystemComponent* InComponent) override;
 	// End UParticleModuleTypeDataBase Interface
 };
 
diff --git a/Engine/Source/Runtime/Engine/Classes/Particles/VectorField/ParticleModuleVectorFieldGlobal.h b/Engine/Source/Runtime/Engine/Classes/Particles/VectorField/ParticleModuleVectorFieldGlobal.h
--- a/Engine/Source/Runtime/Engine/Classes/Particles/VectorField/ParticleModuleVectorFieldGlobal.h
+++ b/Engine/Source/Runtime/Engine/Classes/Particles/VectorField/ParticleModuleVectorFieldGlobal.h
@@ -26,7 +26,7 @@ class UParticleModuleVectorFieldGlobal : public UParticleModuleVectorFieldBase
 	float GlobalVectorFieldTightness;
 
 	// Begin UParticleModule Interface
-	virtual void CompileModule(struct FParticleEmitterBuildInfo& EmitterInfo) OVERRIDE;
+	virtual void CompileModule(struct FParticleEmitterBuildInfo& EmitterInfo) override;
 	// Begin UParticleModule Interface
 };
 
diff --git a/Engine/Source/Runtime/Engine/Classes/Particles/VectorField/ParticleModuleVectorFieldLocal.h b/Engine/Source/Runtime/Engine/Classes/Particles/VectorField/ParticleModuleVectorFieldLocal.h
--- a/Engine/Source/Runtime/Engine/Classes/Particles/VectorField/ParticleModuleVectorFieldLocal.h
+++ b/Engine/Source/Runtime/Engine/Classes/Particles/VectorField/ParticleModuleVectorFieldLocal.h
@@ -51,7 +51,7 @@ class UParticleModuleVectorFieldLocal : public UParticleModuleVectorFieldBase
 	uint32 bTileZ:1;
 
 	// Begin UParticleModule Interface
-	virtual void CompileModule(struct FParticleEmitterBuildInfo& EmitterInfo) OVERRIDE;
+	virtual void CompileModule(struct FParticleEmitterBuildInfo& EmitterInfo) override;
 	// Begin UParticleModule Interface
 };
 
diff --git a/Engine/Source/Runtime/Engine/Classes/Particles/VectorField/ParticleModuleVectorFieldScale.h b/Engine/Source/Runtime/Engine/Classes/Particles/VectorField/ParticleModuleVectorFieldScale.h
--- a/Engine/Source/Runtime/Engine/Classes/Particles/VectorField/ParticleModuleVectorFieldScale.h
+++ b/Engine/Source/Runtime/Engine/Classes/Particles/VectorField/ParticleModuleVectorFieldScale.h
@@ -18,15 +18,15 @@ class UParticleModuleVectorFieldScale : public UParticleModuleVectorFieldBase
 	class UDistributionFloat* VectorFieldScale;
 
 	// Begin UObject Interface
-	virtual void PostInitProperties() OVERRIDE;
-	virtual void Serialize(FArchive& Ar) OVERRIDE;
+	virtual void PostInitProperties() override;
+	virtual void Serialize(FArchive& Ar) override;
 	// End UObject Interface
 
 	// Begin UParticleModule Interface
-	virtual void CompileModule(struct FParticleEmitterBuildInfo& EmitterInfo) OVERRIDE;
+	virtual void CompileModule(struct FParticleEmitterBuildInfo& EmitterInfo) override;
 
 #if WITH_EDITOR
-	virtual bool IsValidForLODLevel(UParticleLODLevel* LODLevel, FString& OutErrorString) OVERRIDE;
+	virtual bool IsValidForLODLevel(UParticleLODLevel* LODLevel, FString& OutErrorString) override;
 #endif
 
 	// End UParticleModule Interface
diff --git a/Engine/Source/Runtime/Engine/Classes/Particles/VectorField/ParticleModuleVectorFieldScaleOverLife.h b/Engine/Source/Runtime/Engine/Classes/Particles/VectorField/ParticleModuleVectorFieldScaleOverLife.h
--- a/Engine/Source/Runtime/Engine/Classes/Particles/VectorField/ParticleModuleVectorFieldScaleOverLife.h
+++ b/Engine/Source/Runtime/Engine/Classes/Particles/VectorField/ParticleModuleVectorFieldScaleOverLife.h
@@ -18,15 +18,15 @@ class UParticleModuleVectorFieldScaleOverLife : public UParticleModuleVectorFiel
 	class UDistributionFloat* VectorFieldScaleOverLife;
 
 	// Begin UObject Interface
-	virtual void PostInitProperties() OVERRIDE;
-	virtual void Serialize(FArchive& Ar) OVERRIDE;
+	virtual void PostInitProperties() override;
+	virtual void Serialize(FArchive& Ar) override;
 	// End UObject Interface
 
 	// Begin UParticleModule Interface
-	virtual void CompileModule(struct FParticleEmitterBuildInfo& EmitterInfo) OVERRIDE;
+	virtual void CompileModule(struct FParticleEmitterBuildInfo& EmitterInfo) override;
 
 #if WITH_EDITOR
-	virtual bool IsValidForLODLevel(UParticleLODLevel* LODLevel, FString& OutErrorString) OVERRIDE;
+	virtual bool IsValidForLODLevel(UParticleLODLevel* LODLevel, FString& OutErrorString) override;
 #endif
 	// End UParticleModule Interface
 
diff --git a/Engine/Source/Runtime/Engine/Classes/Particles/Velocity/ParticleModuleVelocity.h b/Engine/Source/Runtime/Engine/Classes/Particles/Velocity/ParticleModuleVelocity.h
--- a/Engine/Source/Runtime/Engine/Classes/Particles/Velocity/ParticleModuleVelocity.h
+++ b/Engine/Source/Runtime/Engine/Classes/Particles/Velocity/ParticleModuleVelocity.h
@@ -30,14 +30,14 @@ class UParticleModuleVelocity : public UParticleModuleVelocityBase
 	
 	// Begin UObject Interface
 #if WITH_EDITOR
-	virtual void PostEditChangeProperty(FPropertyChangedEvent& PropertyChangedEvent) OVERRIDE;
+	virtual void PostEditChangeProperty(FPropertyChangedEvent& PropertyChangedEvent) override;
 #endif // WITH_EDITOR
-	virtual void PostInitProperties() OVERRIDE;
-	virtual void Serialize(FArchive& Ar) OVERRIDE;
+	virtual void PostInitProperties() override;
+	virtual void Serialize(FArchive& Ar) override;
 	// End UObject Interface
 
 	// Begin UParticleModule Interface
-	virtual void Spawn(FParticleEmitterInstance* Owner, int32 Offset, float SpawnTime, FBaseParticle* ParticleBase) OVERRIDE;
+	virtual void Spawn(FParticleEmitterInstance* Owner, int32 Offset, float SpawnTime, FBaseParticle* ParticleBase) override;
 	// Begin UParticleModule Interface
 
 	/**
diff --git a/Engine/Source/Runtime/Engine/Classes/Particles/Velocity/ParticleModuleVelocityCone.h b/Engine/Source/Runtime/Engine/Classes/Particles/Velocity/ParticleModuleVelocityCone.h
--- a/Engine/Source/Runtime/Engine/Classes/Particles/Velocity/ParticleModuleVelocityCone.h
+++ b/Engine/Source/Runtime/Engine/Classes/Particles/Velocity/ParticleModuleVelocityCone.h
@@ -27,15 +27,15 @@ class UParticleModuleVelocityCone : public UParticleModuleVelocityBase
 
 	// Begin UObject Interface
 #if WITH_EDITOR
-	virtual void PostEditChangeProperty(FPropertyChangedEvent& PropertyChangedEvent) OVERRIDE;
+	virtual void PostEditChangeProperty(FPropertyChangedEvent& PropertyChangedEvent) override;
 #endif // WITH_EDITOR
-	virtual void PostInitProperties() OVERRIDE;
-	virtual void Serialize(FArchive& Ar) OVERRIDE;
+	virtual void PostInitProperties() override;
+	virtual void Serialize(FArchive& Ar) override;
 	// End UObject Interface
 
 	// Begin UParticleModule Interface
-	virtual void Spawn(FParticleEmitterInstance* Owner, int32 Offset, float SpawnTime, FBaseParticle* ParticleBase) OVERRIDE;
-	virtual void	Render3DPreview(FParticleEmitterInstance* Owner, const FSceneView* View,FPrimitiveDrawInterface* PDI) OVERRIDE;
+	virtual void Spawn(FParticleEmitterInstance* Owner, int32 Offset, float SpawnTime, FBaseParticle* ParticleBase) override;
+	virtual void	Render3DPreview(FParticleEmitterInstance* Owner, const FSceneView* View,FPrimitiveDrawInterface* PDI) override;
 	// Begin UParticleModule Interface
 	
 	/**
diff --git a/Engine/Source/Runtime/Engine/Classes/Particles/Velocity/ParticleModuleVelocityInheritParent.h b/Engine/Source/Runtime/Engine/Classes/Particles/Velocity/ParticleModuleVelocityInheritParent.h
--- a/Engine/Source/Runtime/Engine/Classes/Particles/Velocity/ParticleModuleVelocityInheritParent.h
+++ b/Engine/Source/Runtime/Engine/Classes/Particles/Velocity/ParticleModuleVelocityInheritParent.h
@@ -22,14 +22,14 @@ class UParticleModuleVelocityInheritParent : public UParticleModuleVelocityBase
 
 	//Begin UObject Interface
 #if WITH_EDITOR
-	virtual void PostEditChangeProperty(FPropertyChangedEvent& PropertyChangedEvent) OVERRIDE;
+	virtual void PostEditChangeProperty(FPropertyChangedEvent& PropertyChangedEvent) override;
 #endif // WITH_EDITOR
-	virtual void PostInitProperties() OVERRIDE;
-	virtual void Serialize(FArchive& Ar) OVERRIDE;
+	virtual void PostInitProperties() override;
+	virtual void Serialize(FArchive& Ar) override;
 	// End UObject Interface
 
 	// Begin UParticleModule Interface
-	virtual void Spawn(FParticleEmitterInstance* Owner, int32 Offset, float SpawnTime, FBaseParticle* ParticleBase) OVERRIDE;
+	virtual void Spawn(FParticleEmitterInstance* Owner, int32 Offset, float SpawnTime, FBaseParticle* ParticleBase) override;
 	// Begin UParticleModule Interface
 };
 
diff --git a/Engine/Source/Runtime/Engine/Classes/Particles/Velocity/ParticleModuleVelocityOverLifetime.h b/Engine/Source/Runtime/Engine/Classes/Particles/Velocity/ParticleModuleVelocityOverLifetime.h
--- a/Engine/Source/Runtime/Engine/Classes/Particles/Velocity/ParticleModuleVelocityOverLifetime.h
+++ b/Engine/Source/Runtime/Engine/Classes/Particles/Velocity/ParticleModuleVelocityOverLifetime.h
@@ -29,14 +29,14 @@ class UParticleModuleVelocityOverLifetime : public UParticleModuleVelocityBase
 
 	//Begin UObject Interface
 #if WITH_EDITOR
-	virtual void PostEditChangeProperty(FPropertyChangedEvent& PropertyChangedEvent) OVERRIDE;
+	virtual void PostEditChangeProperty(FPropertyChangedEvent& PropertyChangedEvent) override;
 #endif // WITH_EDITOR
-	virtual void PostInitProperties() OVERRIDE;
+	virtual void PostInitProperties() override;
 	// End UObject Interface
 
 	// Begin UParticleModule Interface
-	virtual void Spawn(FParticleEmitterInstance* Owner, int32 Offset, float SpawnTime, FBaseParticle* ParticleBase) OVERRIDE;
-	virtual void	Update(FParticleEmitterInstance* Owner, int32 Offset, float DeltaTime) OVERRIDE;
+	virtual void Spawn(FParticleEmitterInstance* Owner, int32 Offset, float SpawnTime, FBaseParticle* ParticleBase) override;
+	virtual void	Update(FParticleEmitterInstance* Owner, int32 Offset, float DeltaTime) override;
 	// Begin UParticleModule Interface
 };
 
diff --git a/Engine/Source/Runtime/Engine/Classes/PhysicalMaterials/PhysicalMaterial.h b/Engine/Source/Runtime/Engine/Classes/PhysicalMaterials/PhysicalMaterial.h
--- a/Engine/Source/Runtime/Engine/Classes/PhysicalMaterials/PhysicalMaterial.h
+++ b/Engine/Source/Runtime/Engine/Classes/PhysicalMaterials/PhysicalMaterial.h
@@ -107,11 +107,11 @@ public:
 
 	// Begin UObject interface
 #if WITH_EDITOR
-	virtual void PostEditChangeProperty(FPropertyChangedEvent& PropertyChangedEvent) OVERRIDE;
+	virtual void PostEditChangeProperty(FPropertyChangedEvent& PropertyChangedEvent) override;
 	static void RebuildPhysicalMaterials();
 #endif // WITH_EDITOR
-	virtual void PostLoad() OVERRIDE;
-	virtual void FinishDestroy() OVERRIDE;
+	virtual void PostLoad() override;
+	virtual void FinishDestroy() override;
 	// End UObject interface
 
 #if WITH_PHYSX
diff --git a/Engine/Source/Runtime/Engine/Classes/PhysicalMaterials/PhysicalMaterialPropertyBase.h b/Engine/Source/Runtime/Engine/Classes/PhysicalMaterials/PhysicalMaterialPropertyBase.h
--- a/Engine/Source/Runtime/Engine/Classes/PhysicalMaterials/PhysicalMaterialPropertyBase.h
+++ b/Engine/Source/Runtime/Engine/Classes/PhysicalMaterials/PhysicalMaterialPropertyBase.h
@@ -29,7 +29,7 @@ class UDEPRECATED_PhysicalMaterialPropertyBase : public UObject
 	GENERATED_UCLASS_BODY()
 
 	// Begin UObject Interface
-	virtual bool IsAsset() const OVERRIDE { return false; }
+	virtual bool IsAsset() const override { return false; }
 	// End UObject Interface
 
 	// compatibility support function
diff --git a/Engine/Source/Runtime/Engine/Classes/PhysicsEngine/BodySetup.h b/Engine/Source/Runtime/Engine/Classes/PhysicsEngine/BodySetup.h
--- a/Engine/Source/Runtime/Engine/Classes/PhysicsEngine/BodySetup.h
+++ b/Engine/Source/Runtime/Engine/Classes/PhysicsEngine/BodySetup.h
@@ -492,15 +492,15 @@ class UBodySetup : public UObject
 
 public:
 	// Begin UObject interface.
-	virtual void Serialize(FArchive& Ar) OVERRIDE;
-	virtual void BeginDestroy() OVERRIDE;
-	virtual void FinishDestroy() OVERRIDE;
-	virtual void PostLoad() OVERRIDE;
-	virtual void PostInitProperties() OVERRIDE;
+	virtual void Serialize(FArchive& Ar) override;
+	virtual void BeginDestroy() override;
+	virtual void FinishDestroy() override;
+	virtual void PostLoad() override;
+	virtual void PostInitProperties() override;
 #if WITH_EDITOR
-	virtual void PostEditUndo() OVERRIDE;
+	virtual void PostEditUndo() override;
 #endif // WITH_EDITOR
-	virtual SIZE_T GetResourceSize(EResourceSizeMode::Type Mode) OVERRIDE;
+	virtual SIZE_T GetResourceSize(EResourceSizeMode::Type Mode) override;
 	// End UObject interface.
 
 	//
