UE4 Refactoring. Changed OVERRIDE and FINAL macros to keywords override and final.

diff --git a/Engine/Source/Runtime/Engine/Classes/PhysicsEngine/BodySetup2D.h b/Engine/Source/Runtime/Engine/Classes/PhysicsEngine/BodySetup2D.h
--- a/Engine/Source/Runtime/Engine/Classes/PhysicsEngine/BodySetup2D.h
+++ b/Engine/Source/Runtime/Engine/Classes/PhysicsEngine/BodySetup2D.h
@@ -15,13 +15,13 @@ class ENGINE_API UBodySetup2D : public UBodySetup
 	FAggregateGeometry2D AggGeom2D;
 
 	// UBodySetup interface
-	virtual void CreatePhysicsMeshes() OVERRIDE;
-	virtual float GetVolume(const FVector& Scale) const OVERRIDE;
+	virtual void CreatePhysicsMeshes() override;
+	virtual float GetVolume(const FVector& Scale) const override;
 	// End of UBodySetup interface
 
 #if WITH_EDITOR
 	// UBodySetup interface
-	virtual void InvalidatePhysicsData() OVERRIDE;
+	virtual void InvalidatePhysicsData() override;
 	// End of UBodySetup interface
 #endif
 };
diff --git a/Engine/Source/Runtime/Engine/Classes/PhysicsEngine/DestructibleActor.h b/Engine/Source/Runtime/Engine/Classes/PhysicsEngine/DestructibleActor.h
--- a/Engine/Source/Runtime/Engine/Classes/PhysicsEngine/DestructibleActor.h
+++ b/Engine/Source/Runtime/Engine/Classes/PhysicsEngine/DestructibleActor.h
@@ -23,9 +23,9 @@ class ADestructibleActor : public AActor
 	uint32 bAffectNavigation : 1;
 
 	// Begin AActor interface.
-	virtual bool UpdateNavigationRelevancy() OVERRIDE { SetNavigationRelevancy(!!bAffectNavigation); return !!bAffectNavigation; }
+	virtual bool UpdateNavigationRelevancy() override { SetNavigationRelevancy(!!bAffectNavigation); return !!bAffectNavigation; }
 #if WITH_EDITOR
-	ENGINE_API virtual bool GetReferencedContentObjects( TArray<UObject*>& Objects ) const OVERRIDE;
+	ENGINE_API virtual bool GetReferencedContentObjects( TArray<UObject*>& Objects ) const override;
 #endif // WITH_EDITOR
 	// End AActor interface.
 
diff --git a/Engine/Source/Runtime/Engine/Classes/PhysicsEngine/PhysicsAsset.h b/Engine/Source/Runtime/Engine/Classes/PhysicsEngine/PhysicsAsset.h
--- a/Engine/Source/Runtime/Engine/Classes/PhysicsEngine/PhysicsAsset.h
+++ b/Engine/Source/Runtime/Engine/Classes/PhysicsEngine/PhysicsAsset.h
@@ -116,12 +116,12 @@ public:
 	TMap<FRigidBodyIndexPair,bool>		CollisionDisableTable;
 
 	// Begin UObject interface
-	virtual void Serialize(FArchive& Ar) OVERRIDE;
-	virtual void PostLoad() OVERRIDE;
-	virtual FString GetDesc() OVERRIDE;
-	virtual void GetAssetRegistryTags(TArray<FAssetRegistryTag>& OutTags) const OVERRIDE;
+	virtual void Serialize(FArchive& Ar) override;
+	virtual void PostLoad() override;
+	virtual FString GetDesc() override;
+	virtual void GetAssetRegistryTags(TArray<FAssetRegistryTag>& OutTags) const override;
 #if WITH_EDITOR
-	virtual void PostEditUndo() OVERRIDE;
+	virtual void PostEditUndo() override;
 #endif
 	// End UObject interface
 
diff --git a/Engine/Source/Runtime/Engine/Classes/PhysicsEngine/PhysicsCollisionHandler.h b/Engine/Source/Runtime/Engine/Classes/PhysicsEngine/PhysicsCollisionHandler.h
--- a/Engine/Source/Runtime/Engine/Classes/PhysicsEngine/PhysicsCollisionHandler.h
+++ b/Engine/Source/Runtime/Engine/Classes/PhysicsEngine/PhysicsCollisionHandler.h
@@ -29,7 +29,7 @@ class ENGINE_API UPhysicsCollisionHandler : public UObject
 	float LastImpactSoundTime;
 
 	/** Get the world we are handling collisions for */
-	virtual UWorld* GetWorld() const OVERRIDE
+	virtual UWorld* GetWorld() const override
 	{
 		return CastChecked<UWorld>(GetOuter());
 	}
diff --git a/Engine/Source/Runtime/Engine/Classes/PhysicsEngine/PhysicsConstraintActor.h b/Engine/Source/Runtime/Engine/Classes/PhysicsEngine/PhysicsConstraintActor.h
--- a/Engine/Source/Runtime/Engine/Classes/PhysicsEngine/PhysicsConstraintActor.h
+++ b/Engine/Source/Runtime/Engine/Classes/PhysicsEngine/PhysicsConstraintActor.h
@@ -24,9 +24,9 @@ class APhysicsConstraintActor : public ARigidBodyBase
 	uint32 bDisableCollision_DEPRECATED:1;
 
 	// Begin UObject Interface
-	virtual void PostLoad() OVERRIDE;
+	virtual void PostLoad() override;
 #if WITH_EDITOR
-	virtual void LoadedFromAnotherClass(const FName& OldClassName) OVERRIDE;
+	virtual void LoadedFromAnotherClass(const FName& OldClassName) override;
 #endif // WITH_EDITOR	
 	// End UObject Interface
 };
diff --git a/Engine/Source/Runtime/Engine/Classes/PhysicsEngine/PhysicsConstraintComponent.h b/Engine/Source/Runtime/Engine/Classes/PhysicsEngine/PhysicsConstraintComponent.h
--- a/Engine/Source/Runtime/Engine/Classes/PhysicsEngine/PhysicsConstraintComponent.h
+++ b/Engine/Source/Runtime/Engine/Classes/PhysicsEngine/PhysicsConstraintComponent.h
@@ -58,25 +58,25 @@ public:
 	FConstraintInstance			ConstraintInstance;
 
 	//Begin UObject Interface
-	virtual void BeginDestroy() OVERRIDE;
-	virtual void PostLoad() OVERRIDE;
+	virtual void BeginDestroy() override;
+	virtual void PostLoad() override;
 #if WITH_EDITOR
-	virtual void PostEditChangeProperty(FPropertyChangedEvent& PropertyChangedEvent) OVERRIDE;
+	virtual void PostEditChangeProperty(FPropertyChangedEvent& PropertyChangedEvent) override;
 #endif // WITH_EDITOR
 	//End UObject interface
 
 	//Begin ActorComponent interface
 #if WITH_EDITOR
-	virtual void CheckForErrors() OVERRIDE;
+	virtual void CheckForErrors() override;
 #endif // WITH_EDITOR
-	virtual void OnRegister() OVERRIDE;
-	virtual void OnUnregister() OVERRIDE;
-	virtual void InitializeComponent() OVERRIDE;
+	virtual void OnRegister() override;
+	virtual void OnUnregister() override;
+	virtual void InitializeComponent() override;
 	//End ActorComponent interface
 
 	// Begin SceneComponent interface
 #if WITH_EDITOR
-	virtual void PostEditComponentMove(bool bFinished) OVERRIDE;
+	virtual void PostEditComponentMove(bool bFinished) override;
 #endif // WITH_EDITOR
 	// End SceneComponent interface
 
diff --git a/Engine/Source/Runtime/Engine/Classes/PhysicsEngine/PhysicsConstraintTemplate.h b/Engine/Source/Runtime/Engine/Classes/PhysicsEngine/PhysicsConstraintTemplate.h
--- a/Engine/Source/Runtime/Engine/Classes/PhysicsEngine/PhysicsConstraintTemplate.h
+++ b/Engine/Source/Runtime/Engine/Classes/PhysicsEngine/PhysicsConstraintTemplate.h
@@ -99,8 +99,8 @@ class UPhysicsConstraintTemplate : public UObject
 
 
 	// Begin UObject interface.
-	virtual void PostLoad() OVERRIDE;
-	virtual void Serialize(FArchive& Ar) OVERRIDE;
+	virtual void PostLoad() override;
+	virtual void Serialize(FArchive& Ar) override;
 	// End UObject interface.
 
 private:
diff --git a/Engine/Source/Runtime/Engine/Classes/PhysicsEngine/PhysicsHandleComponent.h b/Engine/Source/Runtime/Engine/Classes/PhysicsEngine/PhysicsHandleComponent.h
--- a/Engine/Source/Runtime/Engine/Classes/PhysicsEngine/PhysicsHandleComponent.h
+++ b/Engine/Source/Runtime/Engine/Classes/PhysicsEngine/PhysicsHandleComponent.h
@@ -65,8 +65,8 @@ protected:
 	physx::PxRigidDynamic* KinActorData;
 
 	// Begin UActorComponent interface.
-	virtual void OnUnregister() OVERRIDE;
-	virtual void TickComponent(float DeltaTime, enum ELevelTick TickType, FActorComponentTickFunction *ThisTickFunction) OVERRIDE;
+	virtual void OnUnregister() override;
+	virtual void TickComponent(float DeltaTime, enum ELevelTick TickType, FActorComponentTickFunction *ThisTickFunction) override;
 	// End UActorComponent interface.
 
 public:
diff --git a/Engine/Source/Runtime/Engine/Classes/PhysicsEngine/PhysicsSettings.h b/Engine/Source/Runtime/Engine/Classes/PhysicsEngine/PhysicsSettings.h
--- a/Engine/Source/Runtime/Engine/Classes/PhysicsEngine/PhysicsSettings.h
+++ b/Engine/Source/Runtime/Engine/Classes/PhysicsEngine/PhysicsSettings.h
@@ -113,12 +113,12 @@ public:
 
 	static UPhysicsSettings * Get() { return CastChecked<UPhysicsSettings>(UPhysicsSettings::StaticClass()->GetDefaultObject()); }
 
-	virtual void PostInitProperties() OVERRIDE;
+	virtual void PostInitProperties() override;
 
 #if WITH_EDITOR
 
-	virtual bool CanEditChange( const UProperty* Property ) const OVERRIDE;
-	virtual void PostEditChangeProperty(struct FPropertyChangedEvent& PropertyChangedEvent) OVERRIDE;
+	virtual bool CanEditChange( const UProperty* Property ) const override;
+	virtual void PostEditChangeProperty(struct FPropertyChangedEvent& PropertyChangedEvent) override;
 
 	/** Load Material Type data from INI file **/
 	/** this changes displayname meta data. That means we won't need it outside of editor*/
diff --git a/Engine/Source/Runtime/Engine/Classes/PhysicsEngine/PhysicsThrusterComponent.h b/Engine/Source/Runtime/Engine/Classes/PhysicsEngine/PhysicsThrusterComponent.h
--- a/Engine/Source/Runtime/Engine/Classes/PhysicsEngine/PhysicsThrusterComponent.h
+++ b/Engine/Source/Runtime/Engine/Classes/PhysicsEngine/PhysicsThrusterComponent.h
@@ -23,11 +23,11 @@ class UPhysicsThrusterComponent : public USceneComponent
 
 
 	// Begin UActorComponent Interface
-	virtual void TickComponent(float DeltaTime, enum ELevelTick TickType, FActorComponentTickFunction *ThisTickFunction) OVERRIDE;
+	virtual void TickComponent(float DeltaTime, enum ELevelTick TickType, FActorComponentTickFunction *ThisTickFunction) override;
 	// End UActorComponent Interface
 
 	// Begin UObject interface.
-	virtual void PostLoad() OVERRIDE;
+	virtual void PostLoad() override;
 	// End UObject interface.
 };
 
diff --git a/Engine/Source/Runtime/Engine/Classes/PhysicsEngine/RadialForceComponent.h b/Engine/Source/Runtime/Engine/Classes/PhysicsEngine/RadialForceComponent.h
--- a/Engine/Source/Runtime/Engine/Classes/PhysicsEngine/RadialForceComponent.h
+++ b/Engine/Source/Runtime/Engine/Classes/PhysicsEngine/RadialForceComponent.h
@@ -65,13 +65,13 @@ protected:
 
 protected:
 	// Begin UActorComponent interface.
-	virtual void TickComponent(float DeltaTime, enum ELevelTick TickType, FActorComponentTickFunction *ThisTickFunction) OVERRIDE;
+	virtual void TickComponent(float DeltaTime, enum ELevelTick TickType, FActorComponentTickFunction *ThisTickFunction) override;
 	// End UActorComponent interface.
 
 	// Begin UObject interface.
-	virtual void PostLoad() OVERRIDE;
+	virtual void PostLoad() override;
 #if WITH_EDITOR
-	virtual void PostEditChangeProperty(FPropertyChangedEvent& PropertyChangedEvent) OVERRIDE;
+	virtual void PostEditChangeProperty(FPropertyChangedEvent& PropertyChangedEvent) override;
 #endif
 	// End UObject interface.
 
diff --git a/Engine/Source/Runtime/Engine/Classes/Sound/AmbientSound.h b/Engine/Source/Runtime/Engine/Classes/Sound/AmbientSound.h
--- a/Engine/Source/Runtime/Engine/Classes/Sound/AmbientSound.h
+++ b/Engine/Source/Runtime/Engine/Classes/Sound/AmbientSound.h
@@ -29,11 +29,11 @@ class AAmbientSound : public AActor
 
 	// Begin AActor interface.
 #if WITH_EDITOR
-	virtual void CheckForErrors() OVERRIDE;
-	virtual bool GetReferencedContentObjects( TArray<UObject*>& Objects ) const OVERRIDE;
+	virtual void CheckForErrors() override;
+	virtual bool GetReferencedContentObjects( TArray<UObject*>& Objects ) const override;
 #endif
-	virtual void PostLoad() OVERRIDE;
-	virtual void PostRegisterAllComponents() OVERRIDE;
+	virtual void PostLoad() override;
+	virtual void PostRegisterAllComponents() override;
 	// End AActor interface.
 
 	// BEGIN DEPRECATED (use component functions now in level script)
diff --git a/Engine/Source/Runtime/Engine/Classes/Sound/DialogueVoice.h b/Engine/Source/Runtime/Engine/Classes/Sound/DialogueVoice.h
--- a/Engine/Source/Runtime/Engine/Classes/Sound/DialogueVoice.h
+++ b/Engine/Source/Runtime/Engine/Classes/Sound/DialogueVoice.h
@@ -24,10 +24,10 @@ class UDialogueVoice : public UObject
 
 public:
 	// Begin UObject interface. 
-	virtual bool IsReadyForFinishDestroy() OVERRIDE;
-	virtual FName GetExporterName() OVERRIDE;
-	virtual FString GetDesc() OVERRIDE;
-	virtual void GetAssetRegistryTags(TArray<FAssetRegistryTag>& OutTags) const OVERRIDE;
+	virtual bool IsReadyForFinishDestroy() override;
+	virtual FName GetExporterName() override;
+	virtual FString GetDesc() override;
+	virtual void GetAssetRegistryTags(TArray<FAssetRegistryTag>& OutTags) const override;
 
 	virtual void PostDuplicate(bool bDuplicateForPIE);
 	// End UObject interface. 
diff --git a/Engine/Source/Runtime/Engine/Classes/Sound/DialogueWave.h b/Engine/Source/Runtime/Engine/Classes/Sound/DialogueWave.h
--- a/Engine/Source/Runtime/Engine/Classes/Sound/DialogueWave.h
+++ b/Engine/Source/Runtime/Engine/Classes/Sound/DialogueWave.h
@@ -62,16 +62,16 @@ class UDialogueWave : public UObject
 
 public:
 	// Begin UObject interface. 
-	virtual void Serialize( FArchive& Ar ) OVERRIDE;
-	virtual bool IsReadyForFinishDestroy() OVERRIDE;
-	virtual FString GetDesc() OVERRIDE;
-	virtual void GetAssetRegistryTags(TArray<FAssetRegistryTag>& OutTags) const OVERRIDE;
+	virtual void Serialize( FArchive& Ar ) override;
+	virtual bool IsReadyForFinishDestroy() override;
+	virtual FString GetDesc() override;
+	virtual void GetAssetRegistryTags(TArray<FAssetRegistryTag>& OutTags) const override;
 
 	virtual void PostDuplicate(bool bDuplicateForPIE);
 	virtual void PostLoad();
 
 #if WITH_EDITOR
-	virtual void PostEditChangeChainProperty(FPropertyChangedChainEvent& PropertyChangedEvent) OVERRIDE;
+	virtual void PostEditChangeChainProperty(FPropertyChangedChainEvent& PropertyChangedEvent) override;
 #endif
 
 	// End UObject interface. 
diff --git a/Engine/Source/Runtime/Engine/Classes/Sound/ReverbVolume.h b/Engine/Source/Runtime/Engine/Classes/Sound/ReverbVolume.h
--- a/Engine/Source/Runtime/Engine/Classes/Sound/ReverbVolume.h
+++ b/Engine/Source/Runtime/Engine/Classes/Sound/ReverbVolume.h
@@ -174,14 +174,14 @@ class AReverbVolume : public AVolume
 
 	// Begin UObject Interface
 #if WITH_EDITOR
-	virtual void PostEditChangeProperty(FPropertyChangedEvent& PropertyChangedEvent) OVERRIDE;
+	virtual void PostEditChangeProperty(FPropertyChangedEvent& PropertyChangedEvent) override;
 #endif // WITH_EDITOR
 	// End UObject Interface
 
 	// Begin AActor Interface
-	virtual void PostUnregisterAllComponents() OVERRIDE;
+	virtual void PostUnregisterAllComponents() override;
 protected:
-	virtual void PostRegisterAllComponents() OVERRIDE;
+	virtual void PostRegisterAllComponents() override;
 public:
 	// End AActor Interface
 };
diff --git a/Engine/Source/Runtime/Engine/Classes/Sound/SoundClass.h b/Engine/Source/Runtime/Engine/Classes/Sound/SoundClass.h
--- a/Engine/Source/Runtime/Engine/Classes/Sound/SoundClass.h
+++ b/Engine/Source/Runtime/Engine/Classes/Sound/SoundClass.h
@@ -186,13 +186,13 @@ public:
 protected:
 
 	// Begin UObject interface.
-	virtual void Serialize( FArchive& Ar ) OVERRIDE;
-	virtual FString GetDesc( void ) OVERRIDE;
-	virtual void BeginDestroy() OVERRIDE;
-	virtual void PostLoad() OVERRIDE;
+	virtual void Serialize( FArchive& Ar ) override;
+	virtual FString GetDesc( void ) override;
+	virtual void BeginDestroy() override;
+	virtual void PostLoad() override;
 #if WITH_EDITOR
-	virtual void PreEditChange(UProperty* PropertyAboutToChange) OVERRIDE;
-	virtual void PostEditChangeProperty(struct FPropertyChangedEvent& PropertyChangedEvent) OVERRIDE;
+	virtual void PreEditChange(UProperty* PropertyAboutToChange) override;
+	virtual void PostEditChangeProperty(struct FPropertyChangedEvent& PropertyChangedEvent) override;
 #endif
 	// End UObject interface.
 
diff --git a/Engine/Source/Runtime/Engine/Classes/Sound/SoundCue.h b/Engine/Source/Runtime/Engine/Classes/Sound/SoundCue.h
--- a/Engine/Source/Runtime/Engine/Classes/Sound/SoundCue.h
+++ b/Engine/Source/Runtime/Engine/Classes/Sound/SoundCue.h
@@ -74,25 +74,25 @@ private:
 public:
 
 	// Begin UObject interface.
-	virtual SIZE_T GetResourceSize(EResourceSizeMode::Type Mode) OVERRIDE;
-	virtual FString GetDesc() OVERRIDE;
+	virtual SIZE_T GetResourceSize(EResourceSizeMode::Type Mode) override;
+	virtual FString GetDesc() override;
 #if WITH_EDITOR
-	virtual void PostInitProperties() OVERRIDE;
-	virtual void PostEditChangeProperty(struct FPropertyChangedEvent& PropertyChangedEvent) OVERRIDE;
-	virtual void Serialize( FArchive& Ar ) OVERRIDE;
+	virtual void PostInitProperties() override;
+	virtual void PostEditChangeProperty(struct FPropertyChangedEvent& PropertyChangedEvent) override;
+	virtual void Serialize( FArchive& Ar ) override;
 	static void AddReferencedObjects(UObject* InThis, FReferenceCollector& Collector);
-	virtual void PostLoad() OVERRIDE;
+	virtual void PostLoad() override;
 #endif
 	// End UObject interface.
 
 	// Begin USoundBase interface.
-	virtual bool IsPlayable() const OVERRIDE;
-	virtual void Parse( class FAudioDevice* AudioDevice, const UPTRINT NodeWaveInstanceHash, FActiveSound& ActiveSound, const FSoundParseParameters& ParseParams, TArray<FWaveInstance*>& WaveInstances ) OVERRIDE;
-	virtual float GetVolumeMultiplier() OVERRIDE;
-	virtual float GetPitchMultiplier() OVERRIDE;
-	virtual float GetMaxAudibleDistance() OVERRIDE;
-	virtual float GetDuration() OVERRIDE;
-	virtual const FAttenuationSettings* GetAttenuationSettingsToApply() const OVERRIDE;
+	virtual bool IsPlayable() const override;
+	virtual void Parse( class FAudioDevice* AudioDevice, const UPTRINT NodeWaveInstanceHash, FActiveSound& ActiveSound, const FSoundParseParameters& ParseParams, TArray<FWaveInstance*>& WaveInstances ) override;
+	virtual float GetVolumeMultiplier() override;
+	virtual float GetPitchMultiplier() override;
+	virtual float GetMaxAudibleDistance() override;
+	virtual float GetDuration() override;
+	virtual const FAttenuationSettings* GetAttenuationSettingsToApply() const override;
 	// End USoundBase interface.
 
 	/** Construct and initialize a node within this Cue */
diff --git a/Engine/Source/Runtime/Engine/Classes/Sound/SoundMix.h b/Engine/Source/Runtime/Engine/Classes/Sound/SoundMix.h
--- a/Engine/Source/Runtime/Engine/Classes/Sound/SoundMix.h
+++ b/Engine/Source/Runtime/Engine/Classes/Sound/SoundMix.h
@@ -153,11 +153,11 @@ class USoundMix : public UObject
 
 protected:
 	// Begin UObject interface.
-	virtual FString GetDesc( void ) OVERRIDE;
+	virtual FString GetDesc( void ) override;
 #if WITH_EDITOR
-	virtual void PostEditChangeProperty(FPropertyChangedEvent& PropertyChangedEvent) OVERRIDE;
+	virtual void PostEditChangeProperty(FPropertyChangedEvent& PropertyChangedEvent) override;
 #endif // WITH_EDITOR
-	virtual void BeginDestroy() OVERRIDE;
+	virtual void BeginDestroy() override;
 	// End UObject interface.
 
 #if WITH_EDITOR
diff --git a/Engine/Source/Runtime/Engine/Classes/Sound/SoundNode.h b/Engine/Source/Runtime/Engine/Classes/Sound/SoundNode.h
--- a/Engine/Source/Runtime/Engine/Classes/Sound/SoundNode.h
+++ b/Engine/Source/Runtime/Engine/Classes/Sound/SoundNode.h
@@ -64,10 +64,10 @@ class ENGINE_API USoundNode : public UObject
 public:
 	// Begin UObject Interface
 #if WITH_EDITOR
-	virtual void PostEditChangeProperty(FPropertyChangedEvent& PropertyChangedEvent) OVERRIDE;
+	virtual void PostEditChangeProperty(FPropertyChangedEvent& PropertyChangedEvent) override;
 #endif
 #if WITH_EDITORONLY_DATA
-	virtual void Serialize(FArchive& Ar) OVERRIDE;
+	virtual void Serialize(FArchive& Ar) override;
 	static void AddReferencedObjects(UObject* InThis, FReferenceCollector& Collector);
 #endif //WITH_EDITORONLY_DATA
 	// End UObject Interface
diff --git a/Engine/Source/Runtime/Engine/Classes/Sound/SoundNodeAttenuation.h b/Engine/Source/Runtime/Engine/Classes/Sound/SoundNodeAttenuation.h
--- a/Engine/Source/Runtime/Engine/Classes/Sound/SoundNodeAttenuation.h
+++ b/Engine/Source/Runtime/Engine/Classes/Sound/SoundNodeAttenuation.h
@@ -55,13 +55,13 @@ class USoundNodeAttenuation : public USoundNode
 
 public:
 	// Begin UObject interface.
-	virtual void Serialize(FArchive& Ar) OVERRIDE;
+	virtual void Serialize(FArchive& Ar) override;
 	// End UObject interface.
 
 	// Begin USoundNode interface. 
-	virtual void ParseNodes( class FAudioDevice* AudioDevice, const UPTRINT NodeWaveInstanceHash, FActiveSound& ActiveSound, const FSoundParseParameters& ParseParams, TArray<FWaveInstance*>& WaveInstances ) OVERRIDE;
-	virtual float MaxAudibleDistance( float CurrentMaxDistance ) OVERRIDE;
-	virtual FString GetUniqueString() const OVERRIDE;
+	virtual void ParseNodes( class FAudioDevice* AudioDevice, const UPTRINT NodeWaveInstanceHash, FActiveSound& ActiveSound, const FSoundParseParameters& ParseParams, TArray<FWaveInstance*>& WaveInstances ) override;
+	virtual float MaxAudibleDistance( float CurrentMaxDistance ) override;
+	virtual FString GetUniqueString() const override;
 	// End USoundNode interface. 
 
 	ENGINE_API FAttenuationSettings* GetAttenuationSettingsToApply();
diff --git a/Engine/Source/Runtime/Engine/Classes/Sound/SoundNodeBranch.h b/Engine/Source/Runtime/Engine/Classes/Sound/SoundNodeBranch.h
--- a/Engine/Source/Runtime/Engine/Classes/Sound/SoundNodeBranch.h
+++ b/Engine/Source/Runtime/Engine/Classes/Sound/SoundNodeBranch.h
@@ -28,25 +28,25 @@ private:
 
 public:
 	// Begin USoundNode interface.
-	virtual void ParseNodes( FAudioDevice* AudioDevice, const UPTRINT NodeWaveInstanceHash, FActiveSound& ActiveSound, const FSoundParseParameters& ParseParams, TArray<FWaveInstance*>& WaveInstances ) OVERRIDE;
-	virtual int32 GetMaxChildNodes( void ) const OVERRIDE 
+	virtual void ParseNodes( FAudioDevice* AudioDevice, const UPTRINT NodeWaveInstanceHash, FActiveSound& ActiveSound, const FSoundParseParameters& ParseParams, TArray<FWaveInstance*>& WaveInstances ) override;
+	virtual int32 GetMaxChildNodes( void ) const override 
 	{ 
 		return BranchPurpose::MAX; 
 	}
-	virtual int32 GetMinChildNodes() const OVERRIDE
+	virtual int32 GetMinChildNodes() const override
 	{ 
 		return BranchPurpose::MAX;
 	}
-	virtual void CreateStartingConnectors( void ) OVERRIDE;
+	virtual void CreateStartingConnectors( void ) override;
 
-	virtual void RemoveChildNode( int32 Index ) OVERRIDE
+	virtual void RemoveChildNode( int32 Index ) override
 	{
 		// Do nothing - we do not want to allow deleting of children nodes
 	}
 
 #if WITH_EDITOR
-	virtual FString GetInputPinName(int32 PinIndex) const OVERRIDE;
-	virtual FString GetTitle() const OVERRIDE;
+	virtual FString GetInputPinName(int32 PinIndex) const override;
+	virtual FString GetTitle() const override;
 #endif //WITH_EDITOR
 	// End USoundNode interface.
 };
diff --git a/Engine/Source/Runtime/Engine/Classes/Sound/SoundNodeConcatenator.h b/Engine/Source/Runtime/Engine/Classes/Sound/SoundNodeConcatenator.h
--- a/Engine/Source/Runtime/Engine/Classes/Sound/SoundNodeConcatenator.h
+++ b/Engine/Source/Runtime/Engine/Classes/Sound/SoundNodeConcatenator.h
@@ -24,21 +24,21 @@ class USoundNodeConcatenator : public USoundNode
 
 public:	
 	// Begin USoundNode interface. 
-	virtual bool NotifyWaveInstanceFinished( struct FWaveInstance* WaveInstance ) OVERRIDE;
-	virtual float GetDuration( void ) OVERRIDE;
-	virtual void ParseNodes( FAudioDevice* AudioDevice, const UPTRINT NodeWaveInstanceHash, FActiveSound& ActiveSound, const FSoundParseParameters& ParseParams, TArray<FWaveInstance*>& WaveInstances ) OVERRIDE;
-	virtual int32 GetMaxChildNodes() const OVERRIDE 
+	virtual bool NotifyWaveInstanceFinished( struct FWaveInstance* WaveInstance ) override;
+	virtual float GetDuration( void ) override;
+	virtual void ParseNodes( FAudioDevice* AudioDevice, const UPTRINT NodeWaveInstanceHash, FActiveSound& ActiveSound, const FSoundParseParameters& ParseParams, TArray<FWaveInstance*>& WaveInstances ) override;
+	virtual int32 GetMaxChildNodes() const override 
 	{ 
 		return MAX_ALLOWED_CHILD_NODES; 
 	}
-	virtual void CreateStartingConnectors( void ) OVERRIDE;
-	virtual void InsertChildNode( int32 Index ) OVERRIDE;
-	virtual void RemoveChildNode( int32 Index ) OVERRIDE;
+	virtual void CreateStartingConnectors( void ) override;
+	virtual void InsertChildNode( int32 Index ) override;
+	virtual void RemoveChildNode( int32 Index ) override;
 #if WITH_EDITOR
 	/** Ensure amount of inputs matches new amount of children */
-	virtual void SetChildNodes(TArray<USoundNode*>& InChildNodes) OVERRIDE;
+	virtual void SetChildNodes(TArray<USoundNode*>& InChildNodes) override;
 #endif
-	virtual FString GetUniqueString() const OVERRIDE;
+	virtual FString GetUniqueString() const override;
 	// End USoundNode interface. 
 };
 
diff --git a/Engine/Source/Runtime/Engine/Classes/Sound/SoundNodeDelay.h b/Engine/Source/Runtime/Engine/Classes/Sound/SoundNodeDelay.h
--- a/Engine/Source/Runtime/Engine/Classes/Sound/SoundNodeDelay.h
+++ b/Engine/Source/Runtime/Engine/Classes/Sound/SoundNodeDelay.h
@@ -22,9 +22,9 @@ class USoundNodeDelay : public USoundNode
 
 public:
 	// Begin USoundNode interface. 
-	virtual void ParseNodes( FAudioDevice* AudioDevice, const UPTRINT NodeWaveInstanceHash, FActiveSound& ActiveSound, const FSoundParseParameters& ParseParams, TArray<FWaveInstance*>& WaveInstances ) OVERRIDE;
-	virtual float GetDuration( void ) OVERRIDE;
-	virtual FString GetUniqueString() const OVERRIDE;
+	virtual void ParseNodes( FAudioDevice* AudioDevice, const UPTRINT NodeWaveInstanceHash, FActiveSound& ActiveSound, const FSoundParseParameters& ParseParams, TArray<FWaveInstance*>& WaveInstances ) override;
+	virtual float GetDuration( void ) override;
+	virtual FString GetUniqueString() const override;
 	// End USoundNode interface. 
 };
 
diff --git a/Engine/Source/Runtime/Engine/Classes/Sound/SoundNodeDialoguePlayer.h b/Engine/Source/Runtime/Engine/Classes/Sound/SoundNodeDialoguePlayer.h
--- a/Engine/Source/Runtime/Engine/Classes/Sound/SoundNodeDialoguePlayer.h
+++ b/Engine/Source/Runtime/Engine/Classes/Sound/SoundNodeDialoguePlayer.h
@@ -22,10 +22,10 @@ class USoundNodeDialoguePlayer : public USoundNode
 
 public:	
 	// Begin USoundNode Interface
-	virtual int32 GetMaxChildNodes() const OVERRIDE;
-	virtual float GetDuration() OVERRIDE;
-	virtual void ParseNodes( FAudioDevice* AudioDevice, const UPTRINT NodeWaveInstanceHash, FActiveSound& ActiveSound, const FSoundParseParameters& ParseParams, TArray<FWaveInstance*>& WaveInstances ) OVERRIDE;
-	virtual FString GetUniqueString() const OVERRIDE;
+	virtual int32 GetMaxChildNodes() const override;
+	virtual float GetDuration() override;
+	virtual void ParseNodes( FAudioDevice* AudioDevice, const UPTRINT NodeWaveInstanceHash, FActiveSound& ActiveSound, const FSoundParseParameters& ParseParams, TArray<FWaveInstance*>& WaveInstances ) override;
+	virtual FString GetUniqueString() const override;
 	// End USoundNode Interface
 
 };
diff --git a/Engine/Source/Runtime/Engine/Classes/Sound/SoundNodeDistanceCrossFade.h b/Engine/Source/Runtime/Engine/Classes/Sound/SoundNodeDistanceCrossFade.h
--- a/Engine/Source/Runtime/Engine/Classes/Sound/SoundNodeDistanceCrossFade.h
+++ b/Engine/Source/Runtime/Engine/Classes/Sound/SoundNodeDistanceCrossFade.h
@@ -85,20 +85,20 @@ class USoundNodeDistanceCrossFade : public USoundNode
 
 public:
 	// Begin USoundNode interface. 
-	virtual void ParseNodes( FAudioDevice* AudioDevice, const UPTRINT NodeWaveInstanceHash, FActiveSound& ActiveSound, const FSoundParseParameters& ParseParams, TArray<FWaveInstance*>& WaveInstances ) OVERRIDE;
-	virtual int32 GetMaxChildNodes() const OVERRIDE 
+	virtual void ParseNodes( FAudioDevice* AudioDevice, const UPTRINT NodeWaveInstanceHash, FActiveSound& ActiveSound, const FSoundParseParameters& ParseParams, TArray<FWaveInstance*>& WaveInstances ) override;
+	virtual int32 GetMaxChildNodes() const override 
 	{ 
 		return MAX_ALLOWED_CHILD_NODES;  
 	}
-	virtual void CreateStartingConnectors( void ) OVERRIDE;
-	virtual void InsertChildNode( int32 Index ) OVERRIDE;
-	virtual void RemoveChildNode( int32 Index ) OVERRIDE;
+	virtual void CreateStartingConnectors( void ) override;
+	virtual void InsertChildNode( int32 Index ) override;
+	virtual void RemoveChildNode( int32 Index ) override;
 #if WITH_EDITOR
 	/** Ensure amount of Cross Fade inputs matches new amount of children */
-	virtual void SetChildNodes(TArray<USoundNode*>& InChildNodes) OVERRIDE;
+	virtual void SetChildNodes(TArray<USoundNode*>& InChildNodes) override;
 #endif //WITH_EDITOR
-	virtual float MaxAudibleDistance( float CurrentMaxDistance ) OVERRIDE;
-	virtual FString GetUniqueString() const OVERRIDE;
+	virtual float MaxAudibleDistance( float CurrentMaxDistance ) override;
+	virtual FString GetUniqueString() const override;
 	// End USoundNode interface. 
 
 	virtual float GetCurrentDistance(FAudioDevice* AudioDevice, FActiveSound& ActiveSound, const FSoundParseParameters& ParseParams) const;
diff --git a/Engine/Source/Runtime/Engine/Classes/Sound/SoundNodeDoppler.h b/Engine/Source/Runtime/Engine/Classes/Sound/SoundNodeDoppler.h
--- a/Engine/Source/Runtime/Engine/Classes/Sound/SoundNodeDoppler.h
+++ b/Engine/Source/Runtime/Engine/Classes/Sound/SoundNodeDoppler.h
@@ -21,8 +21,8 @@ class USoundNodeDoppler : public USoundNode
 
 public:
 	// Begin USoundNode interface. 
-	virtual void ParseNodes( FAudioDevice* AudioDevice, const UPTRINT NodeWaveInstanceHash, FActiveSound& ActiveSound, const FSoundParseParameters& ParseParams, TArray<FWaveInstance*>& WaveInstances ) OVERRIDE;
-	virtual FString GetUniqueString() const OVERRIDE;
+	virtual void ParseNodes( FAudioDevice* AudioDevice, const UPTRINT NodeWaveInstanceHash, FActiveSound& ActiveSound, const FSoundParseParameters& ParseParams, TArray<FWaveInstance*>& WaveInstances ) override;
+	virtual FString GetUniqueString() const override;
 	// End USoundNode interface. 
 
 protected:
diff --git a/Engine/Source/Runtime/Engine/Classes/Sound/SoundNodeEnveloper.h b/Engine/Source/Runtime/Engine/Classes/Sound/SoundNodeEnveloper.h
--- a/Engine/Source/Runtime/Engine/Classes/Sound/SoundNodeEnveloper.h
+++ b/Engine/Source/Runtime/Engine/Classes/Sound/SoundNodeEnveloper.h
@@ -70,17 +70,17 @@ class USoundNodeEnveloper : public USoundNode
 public:	
 	// Begin UObject interface. 
 #if WITH_EDITOR
-	virtual void PostEditChangeProperty( struct FPropertyChangedEvent& PropertyChangedEvent) OVERRIDE;
+	virtual void PostEditChangeProperty( struct FPropertyChangedEvent& PropertyChangedEvent) override;
 #endif // WITH_EDITOR
-	virtual void PostInitProperties() OVERRIDE;
-	virtual void Serialize(FArchive& Ar) OVERRIDE;
+	virtual void PostInitProperties() override;
+	virtual void Serialize(FArchive& Ar) override;
 	// End UObject interface. 
 
 
 	// Begin USoundNode interface. 
-	virtual void ParseNodes( FAudioDevice* AudioDevice, const UPTRINT NodeWaveInstanceHash, FActiveSound& ActiveSound, const FSoundParseParameters& ParseParams, TArray<FWaveInstance*>& WaveInstances ) OVERRIDE;
-	virtual FString GetUniqueString() const OVERRIDE;
-	virtual float GetDuration( void ) OVERRIDE;
+	virtual void ParseNodes( FAudioDevice* AudioDevice, const UPTRINT NodeWaveInstanceHash, FActiveSound& ActiveSound, const FSoundParseParameters& ParseParams, TArray<FWaveInstance*>& WaveInstances ) override;
+	virtual FString GetUniqueString() const override;
+	virtual float GetDuration( void ) override;
 	// End USoundNode interface. 
 
 };
diff --git a/Engine/Source/Runtime/Engine/Classes/Sound/SoundNodeGroupControl.h b/Engine/Source/Runtime/Engine/Classes/Sound/SoundNodeGroupControl.h
--- a/Engine/Source/Runtime/Engine/Classes/Sound/SoundNodeGroupControl.h
+++ b/Engine/Source/Runtime/Engine/Classes/Sound/SoundNodeGroupControl.h
@@ -20,19 +20,19 @@ class USoundNodeGroupControl : public USoundNode
 public:
 	// Begin USoundNode interface.
 	virtual bool NotifyWaveInstanceFinished( FWaveInstance* WaveInstance );
-	virtual void ParseNodes( FAudioDevice* AudioDevice, const UPTRINT NodeWaveInstanceHash, FActiveSound& ActiveSound, const FSoundParseParameters& ParseParams, TArray<FWaveInstance*>& WaveInstances ) OVERRIDE;
-	virtual int32 GetMaxChildNodes() const OVERRIDE 
+	virtual void ParseNodes( FAudioDevice* AudioDevice, const UPTRINT NodeWaveInstanceHash, FActiveSound& ActiveSound, const FSoundParseParameters& ParseParams, TArray<FWaveInstance*>& WaveInstances ) override;
+	virtual int32 GetMaxChildNodes() const override 
 	{ 
 		return MAX_ALLOWED_CHILD_NODES; 
 	}
-	virtual void InsertChildNode( int32 Index ) OVERRIDE;
-	virtual void RemoveChildNode( int32 Index ) OVERRIDE;
+	virtual void InsertChildNode( int32 Index ) override;
+	virtual void RemoveChildNode( int32 Index ) override;
 #if WITH_EDITOR
 	/** Ensure amount of Groups matches new amount of children */
-	virtual void SetChildNodes(TArray<USoundNode*>& InChildNodes) OVERRIDE;
+	virtual void SetChildNodes(TArray<USoundNode*>& InChildNodes) override;
 #endif //WITH_EDITOR
-	virtual void CreateStartingConnectors() OVERRIDE;
-	virtual FString GetUniqueString() const OVERRIDE;
+	virtual void CreateStartingConnectors() override;
+	virtual FString GetUniqueString() const override;
 	// End USoundNode interface.
 
 private:
diff --git a/Engine/Source/Runtime/Engine/Classes/Sound/SoundNodeLooping.h b/Engine/Source/Runtime/Engine/Classes/Sound/SoundNodeLooping.h
--- a/Engine/Source/Runtime/Engine/Classes/Sound/SoundNodeLooping.h
+++ b/Engine/Source/Runtime/Engine/Classes/Sound/SoundNodeLooping.h
@@ -15,14 +15,14 @@ class USoundNodeLooping : public USoundNode
 
 public:	
 	// Begin USoundNode interface. 
-	virtual bool NotifyWaveInstanceFinished( struct FWaveInstance* WaveInstance ) OVERRIDE;
-	virtual float MaxAudibleDistance( float CurrentMaxDistance ) OVERRIDE 
+	virtual bool NotifyWaveInstanceFinished( struct FWaveInstance* WaveInstance ) override;
+	virtual float MaxAudibleDistance( float CurrentMaxDistance ) override 
 	{ 
 		return( WORLD_MAX ); 
 	}
-	virtual float GetDuration( void ) OVERRIDE;
-	virtual void ParseNodes( FAudioDevice* AudioDevice, const UPTRINT NodeWaveInstanceHash, FActiveSound& ActiveSound, const FSoundParseParameters& ParseParams, TArray<FWaveInstance*>& WaveInstances ) OVERRIDE;
-	virtual FString GetUniqueString() const OVERRIDE;
+	virtual float GetDuration( void ) override;
+	virtual void ParseNodes( FAudioDevice* AudioDevice, const UPTRINT NodeWaveInstanceHash, FActiveSound& ActiveSound, const FSoundParseParameters& ParseParams, TArray<FWaveInstance*>& WaveInstances ) override;
+	virtual FString GetUniqueString() const override;
 	// End USoundNode interface. 
 };
 
diff --git a/Engine/Source/Runtime/Engine/Classes/Sound/SoundNodeMature.h b/Engine/Source/Runtime/Engine/Classes/Sound/SoundNodeMature.h
--- a/Engine/Source/Runtime/Engine/Classes/Sound/SoundNodeMature.h
+++ b/Engine/Source/Runtime/Engine/Classes/Sound/SoundNodeMature.h
@@ -17,14 +17,14 @@ class USoundNodeMature : public USoundNode
 
 public:
 	// Begin UObject Interface
-	virtual void PostLoad() OVERRIDE;
+	virtual void PostLoad() override;
 	// End UObject Interface
 
 	// Begin USoundNode interface.
-	virtual void ParseNodes( FAudioDevice* AudioDevice, const UPTRINT NodeWaveInstanceHash, FActiveSound& ActiveSound, const FSoundParseParameters& ParseParams, TArray<FWaveInstance*>& WaveInstances ) OVERRIDE;
-	virtual void CreateStartingConnectors( void ) OVERRIDE;
-	virtual int32 GetMaxChildNodes() const OVERRIDE;
-	virtual FString GetUniqueString() const OVERRIDE;
+	virtual void ParseNodes( FAudioDevice* AudioDevice, const UPTRINT NodeWaveInstanceHash, FActiveSound& ActiveSound, const FSoundParseParameters& ParseParams, TArray<FWaveInstance*>& WaveInstances ) override;
+	virtual void CreateStartingConnectors( void ) override;
+	virtual int32 GetMaxChildNodes() const override;
+	virtual FString GetUniqueString() const override;
 	// End USoundNode interface.
 };
 
diff --git a/Engine/Source/Runtime/Engine/Classes/Sound/SoundNodeMixer.h b/Engine/Source/Runtime/Engine/Classes/Sound/SoundNodeMixer.h
--- a/Engine/Source/Runtime/Engine/Classes/Sound/SoundNodeMixer.h
+++ b/Engine/Source/Runtime/Engine/Classes/Sound/SoundNodeMixer.h
@@ -20,19 +20,19 @@ class USoundNodeMixer : public USoundNode
 
 public:
 	// Begin USoundNode interface.
-	virtual void ParseNodes( FAudioDevice* AudioDevice, const UPTRINT NodeWaveInstanceHash, FActiveSound& ActiveSound, const FSoundParseParameters& ParseParams, TArray<FWaveInstance*>& WaveInstances ) OVERRIDE;
-	virtual int32 GetMaxChildNodes() const OVERRIDE 
+	virtual void ParseNodes( FAudioDevice* AudioDevice, const UPTRINT NodeWaveInstanceHash, FActiveSound& ActiveSound, const FSoundParseParameters& ParseParams, TArray<FWaveInstance*>& WaveInstances ) override;
+	virtual int32 GetMaxChildNodes() const override 
 	{ 
 		return MAX_ALLOWED_CHILD_NODES; 
 	}
-	virtual void CreateStartingConnectors( void ) OVERRIDE;
-	virtual void InsertChildNode( int32 Index ) OVERRIDE;
-	virtual void RemoveChildNode( int32 Index ) OVERRIDE;
+	virtual void CreateStartingConnectors( void ) override;
+	virtual void InsertChildNode( int32 Index ) override;
+	virtual void RemoveChildNode( int32 Index ) override;
 #if WITH_EDITOR
 	/** Ensure amount of inputs matches new amount of children */
-	virtual void SetChildNodes(TArray<USoundNode*>& InChildNodes) OVERRIDE;
+	virtual void SetChildNodes(TArray<USoundNode*>& InChildNodes) override;
 #endif //WITH_EDITOR
-	virtual FString GetUniqueString() const OVERRIDE;
+	virtual FString GetUniqueString() const override;
 	// End USoundNode interface.
 };
 
diff --git a/Engine/Source/Runtime/Engine/Classes/Sound/SoundNodeModulator.h b/Engine/Source/Runtime/Engine/Classes/Sound/SoundNodeModulator.h
--- a/Engine/Source/Runtime/Engine/Classes/Sound/SoundNodeModulator.h
+++ b/Engine/Source/Runtime/Engine/Classes/Sound/SoundNodeModulator.h
@@ -31,8 +31,8 @@ class USoundNodeModulator : public USoundNode
 
 public:
 	// Begin USoundNode interface. 
-	virtual void ParseNodes( FAudioDevice* AudioDevice, const UPTRINT NodeWaveInstanceHash, FActiveSound& ActiveSound, const FSoundParseParameters& ParseParams, TArray<FWaveInstance*>& WaveInstances ) OVERRIDE;
-	virtual FString GetUniqueString() const OVERRIDE;
+	virtual void ParseNodes( FAudioDevice* AudioDevice, const UPTRINT NodeWaveInstanceHash, FActiveSound& ActiveSound, const FSoundParseParameters& ParseParams, TArray<FWaveInstance*>& WaveInstances ) override;
+	virtual FString GetUniqueString() const override;
 	// End USoundNode interface. 
 };
 
diff --git a/Engine/Source/Runtime/Engine/Classes/Sound/SoundNodeModulatorContinuous.h b/Engine/Source/Runtime/Engine/Classes/Sound/SoundNodeModulatorContinuous.h
--- a/Engine/Source/Runtime/Engine/Classes/Sound/SoundNodeModulatorContinuous.h
+++ b/Engine/Source/Runtime/Engine/Classes/Sound/SoundNodeModulatorContinuous.h
@@ -85,12 +85,12 @@ class USoundNodeModulatorContinuous : public USoundNode
 
 public:
 	// Begin UObject Interface
-	virtual void Serialize(FArchive& Ar) OVERRIDE;
+	virtual void Serialize(FArchive& Ar) override;
 	// End UObject Interface
 
 	// Begin USoundNode interface.
-	virtual void ParseNodes( FAudioDevice* AudioDevice, const UPTRINT NodeWaveInstanceHash, FActiveSound& ActiveSound, const FSoundParseParameters& ParseParams, TArray<FWaveInstance*>& WaveInstances ) OVERRIDE;
-	virtual FString GetUniqueString() const OVERRIDE;
+	virtual void ParseNodes( FAudioDevice* AudioDevice, const UPTRINT NodeWaveInstanceHash, FActiveSound& ActiveSound, const FSoundParseParameters& ParseParams, TArray<FWaveInstance*>& WaveInstances ) override;
+	virtual FString GetUniqueString() const override;
 	// End USoundNode interface.
 };
 
diff --git a/Engine/Source/Runtime/Engine/Classes/Sound/SoundNodeParamCrossFade.h b/Engine/Source/Runtime/Engine/Classes/Sound/SoundNodeParamCrossFade.h
--- a/Engine/Source/Runtime/Engine/Classes/Sound/SoundNodeParamCrossFade.h
+++ b/Engine/Source/Runtime/Engine/Classes/Sound/SoundNodeParamCrossFade.h
@@ -18,6 +18,6 @@ public:
 	UPROPERTY(EditAnywhere, Category=CrossFade )
 	FName ParamName;
 
-	virtual FString GetUniqueString() const OVERRIDE;
-	virtual float GetCurrentDistance(FAudioDevice* AudioDevice, FActiveSound& ActiveSound, const FSoundParseParameters& ParseParams) const OVERRIDE;
+	virtual FString GetUniqueString() const override;
+	virtual float GetCurrentDistance(FAudioDevice* AudioDevice, FActiveSound& ActiveSound, const FSoundParseParameters& ParseParams) const override;
 };
diff --git a/Engine/Source/Runtime/Engine/Classes/Sound/SoundNodeRandom.h b/Engine/Source/Runtime/Engine/Classes/Sound/SoundNodeRandom.h
--- a/Engine/Source/Runtime/Engine/Classes/Sound/SoundNodeRandom.h
+++ b/Engine/Source/Runtime/Engine/Classes/Sound/SoundNodeRandom.h
@@ -56,24 +56,24 @@ class USoundNodeRandom : public USoundNode
 
 public:
 	// Begin UObject Interface
-	virtual void PostLoad() OVERRIDE;
+	virtual void PostLoad() override;
 	// End UObject Interface
 
 	// Begin USoundNode interface.
-	virtual void ParseNodes( FAudioDevice* AudioDevice, const UPTRINT NodeWaveInstanceHash, FActiveSound& ActiveSound, const FSoundParseParameters& ParseParams, TArray<FWaveInstance*>& WaveInstances ) OVERRIDE;
-	virtual int32 GetMaxChildNodes() const OVERRIDE 
+	virtual void ParseNodes( FAudioDevice* AudioDevice, const UPTRINT NodeWaveInstanceHash, FActiveSound& ActiveSound, const FSoundParseParameters& ParseParams, TArray<FWaveInstance*>& WaveInstances ) override;
+	virtual int32 GetMaxChildNodes() const override 
 	{ 
 		return MAX_ALLOWED_CHILD_NODES; 
 	}
-	virtual void InsertChildNode( int32 Index ) OVERRIDE;
-	virtual void RemoveChildNode( int32 Index ) OVERRIDE;
+	virtual void InsertChildNode( int32 Index ) override;
+	virtual void RemoveChildNode( int32 Index ) override;
 #if WITH_EDITOR
 	/** Ensure Random weights and usage array matches new amount of children */
-	virtual void SetChildNodes(TArray<USoundNode*>& InChildNodes) OVERRIDE;
-	virtual void OnBeginPIE(const bool bIsSimulating) OVERRIDE;
+	virtual void SetChildNodes(TArray<USoundNode*>& InChildNodes) override;
+	virtual void OnBeginPIE(const bool bIsSimulating) override;
 #endif //WITH_EDITOR
-	virtual void CreateStartingConnectors( void ) OVERRIDE;
-	virtual FString GetUniqueString() const OVERRIDE;
+	virtual void CreateStartingConnectors( void ) override;
+	virtual FString GetUniqueString() const override;
 	// End USoundNode interface.
 
 	// @todo document
diff --git a/Engine/Source/Runtime/Engine/Classes/Sound/SoundNodeSoundClass.h b/Engine/Source/Runtime/Engine/Classes/Sound/SoundNodeSoundClass.h
--- a/Engine/Source/Runtime/Engine/Classes/Sound/SoundNodeSoundClass.h
+++ b/Engine/Source/Runtime/Engine/Classes/Sound/SoundNodeSoundClass.h
@@ -18,7 +18,7 @@ class USoundNodeSoundClass : public USoundNode
 
 public:
 	// Begin USoundNode interface. 
-	virtual void ParseNodes( class FAudioDevice* AudioDevice, const UPTRINT NodeWaveInstanceHash, FActiveSound& ActiveSound, const FSoundParseParameters& ParseParams, TArray<FWaveInstance*>& WaveInstances ) OVERRIDE;
-	virtual FString GetUniqueString() const OVERRIDE;
+	virtual void ParseNodes( class FAudioDevice* AudioDevice, const UPTRINT NodeWaveInstanceHash, FActiveSound& ActiveSound, const FSoundParseParameters& ParseParams, TArray<FWaveInstance*>& WaveInstances ) override;
+	virtual FString GetUniqueString() const override;
 	// End USoundNode interface. 
 };
diff --git a/Engine/Source/Runtime/Engine/Classes/Sound/SoundNodeSwitch.h b/Engine/Source/Runtime/Engine/Classes/Sound/SoundNodeSwitch.h
--- a/Engine/Source/Runtime/Engine/Classes/Sound/SoundNodeSwitch.h
+++ b/Engine/Source/Runtime/Engine/Classes/Sound/SoundNodeSwitch.h
@@ -18,36 +18,36 @@ class USoundNodeSwitch : public USoundNode
 
 public:
 	// Begin USoundNode interface.
-	virtual void ParseNodes( FAudioDevice* AudioDevice, const UPTRINT NodeWaveInstanceHash, FActiveSound& ActiveSound, const FSoundParseParameters& ParseParams, TArray<FWaveInstance*>& WaveInstances ) OVERRIDE;
-	virtual int32 GetMaxChildNodes( void ) const OVERRIDE 
+	virtual void ParseNodes( FAudioDevice* AudioDevice, const UPTRINT NodeWaveInstanceHash, FActiveSound& ActiveSound, const FSoundParseParameters& ParseParams, TArray<FWaveInstance*>& WaveInstances ) override;
+	virtual int32 GetMaxChildNodes( void ) const override 
 	{ 
 		return MAX_ALLOWED_CHILD_NODES; 
 	}
-	virtual int32 GetMinChildNodes() const OVERRIDE
+	virtual int32 GetMinChildNodes() const override
 	{ 
 		return 1;
 	}
-	virtual void CreateStartingConnectors( void ) OVERRIDE;
+	virtual void CreateStartingConnectors( void ) override;
 	
 #if WITH_EDITOR
 	void RenamePins();
 
-	virtual void InsertChildNode( int32 Index ) OVERRIDE
+	virtual void InsertChildNode( int32 Index ) override
 	{
 		Super::InsertChildNode(Index);
 
 		RenamePins();
 	}
 
-	virtual void RemoveChildNode( int32 Index ) OVERRIDE
+	virtual void RemoveChildNode( int32 Index ) override
 	{
 		Super::RemoveChildNode(Index);
 		
 		RenamePins();
 	}
 
-	virtual FString GetInputPinName(int32 PinIndex) const OVERRIDE;
-	virtual FString GetTitle() const OVERRIDE;
+	virtual FString GetInputPinName(int32 PinIndex) const override;
+	virtual FString GetTitle() const override;
 #endif //WITH_EDITOR
 	// End USoundNode interface.
 };
diff --git a/Engine/Source/Runtime/Engine/Classes/Sound/SoundNodeWaveParam.h b/Engine/Source/Runtime/Engine/Classes/Sound/SoundNodeWaveParam.h
--- a/Engine/Source/Runtime/Engine/Classes/Sound/SoundNodeWaveParam.h
+++ b/Engine/Source/Runtime/Engine/Classes/Sound/SoundNodeWaveParam.h
@@ -19,9 +19,9 @@ class USoundNodeWaveParam : public USoundNode
 
 public:	
 	// Begin USoundNode Interface
-	virtual float GetDuration( void ) OVERRIDE;
-	virtual void ParseNodes( FAudioDevice* AudioDevice, const UPTRINT NodeWaveInstanceHash, FActiveSound& ActiveSound, const FSoundParseParameters& ParseParams, TArray<FWaveInstance*>& WaveInstances ) OVERRIDE;
-	virtual FString GetUniqueString() const OVERRIDE;
+	virtual float GetDuration( void ) override;
+	virtual void ParseNodes( FAudioDevice* AudioDevice, const UPTRINT NodeWaveInstanceHash, FActiveSound& ActiveSound, const FSoundParseParameters& ParseParams, TArray<FWaveInstance*>& WaveInstances ) override;
+	virtual FString GetUniqueString() const override;
 	// End USoundNode Interface
 };
 
diff --git a/Engine/Source/Runtime/Engine/Classes/Sound/SoundNodeWavePlayer.h b/Engine/Source/Runtime/Engine/Classes/Sound/SoundNodeWavePlayer.h
--- a/Engine/Source/Runtime/Engine/Classes/Sound/SoundNodeWavePlayer.h
+++ b/Engine/Source/Runtime/Engine/Classes/Sound/SoundNodeWavePlayer.h
@@ -21,12 +21,12 @@ class USoundNodeWavePlayer : public USoundNode
 
 public:	
 	// Begin USoundNode Interface
-	virtual int32 GetMaxChildNodes() const OVERRIDE;
-	virtual float GetDuration() OVERRIDE;
-	virtual void ParseNodes( FAudioDevice* AudioDevice, const UPTRINT NodeWaveInstanceHash, FActiveSound& ActiveSound, const FSoundParseParameters& ParseParams, TArray<FWaveInstance*>& WaveInstances ) OVERRIDE;
-	virtual FString GetUniqueString() const OVERRIDE;
+	virtual int32 GetMaxChildNodes() const override;
+	virtual float GetDuration() override;
+	virtual void ParseNodes( FAudioDevice* AudioDevice, const UPTRINT NodeWaveInstanceHash, FActiveSound& ActiveSound, const FSoundParseParameters& ParseParams, TArray<FWaveInstance*>& WaveInstances ) override;
+	virtual FString GetUniqueString() const override;
 #if WITH_EDITOR
-	virtual FString GetTitle() const OVERRIDE;
+	virtual FString GetTitle() const override;
 #endif
 	// End USoundNode Interface
 
diff --git a/Engine/Source/Runtime/Engine/Classes/Sound/SoundWave.h b/Engine/Source/Runtime/Engine/Classes/Sound/SoundWave.h
--- a/Engine/Source/Runtime/Engine/Classes/Sound/SoundWave.h
+++ b/Engine/Source/Runtime/Engine/Classes/Sound/SoundWave.h
@@ -248,27 +248,27 @@ public:
 	TMap<FString, FStreamedAudioPlatformData*> CookedPlatformData;
 
 	// Begin UObject interface. 
-	virtual void Serialize( FArchive& Ar ) OVERRIDE;
-	virtual void PostInitProperties() OVERRIDE;
-	virtual bool IsReadyForFinishDestroy() OVERRIDE;
-	virtual void FinishDestroy() OVERRIDE;
-	virtual void PostLoad() OVERRIDE;
+	virtual void Serialize( FArchive& Ar ) override;
+	virtual void PostInitProperties() override;
+	virtual bool IsReadyForFinishDestroy() override;
+	virtual void FinishDestroy() override;
+	virtual void PostLoad() override;
 #if WITH_EDITOR
-	virtual void CookerWillNeverCookAgain() OVERRIDE;
-	virtual void PostEditChangeProperty(FPropertyChangedEvent& PropertyChangedEvent) OVERRIDE;	
+	virtual void CookerWillNeverCookAgain() override;
+	virtual void PostEditChangeProperty(FPropertyChangedEvent& PropertyChangedEvent) override;	
 #endif // WITH_EDITOR
-	virtual bool IsLocalizedResource() OVERRIDE;
-	virtual SIZE_T GetResourceSize(EResourceSizeMode::Type Mode) OVERRIDE;
-	virtual FName GetExporterName() OVERRIDE;
-	virtual FString GetDesc() OVERRIDE;
-	virtual void GetAssetRegistryTags(TArray<FAssetRegistryTag>& OutTags) const OVERRIDE;
+	virtual bool IsLocalizedResource() override;
+	virtual SIZE_T GetResourceSize(EResourceSizeMode::Type Mode) override;
+	virtual FName GetExporterName() override;
+	virtual FString GetDesc() override;
+	virtual void GetAssetRegistryTags(TArray<FAssetRegistryTag>& OutTags) const override;
 	// End UObject interface. 
 
 	// Begin USoundBase interface.
-	virtual bool IsPlayable() const OVERRIDE;
-	virtual void Parse( class FAudioDevice* AudioDevice, const UPTRINT NodeWaveInstanceHash, FActiveSound& ActiveSound, const FSoundParseParameters& ParseParams, TArray<FWaveInstance*>& WaveInstances ) OVERRIDE;
-	virtual float GetMaxAudibleDistance() OVERRIDE;
-	virtual float GetDuration() OVERRIDE;
+	virtual bool IsPlayable() const override;
+	virtual void Parse( class FAudioDevice* AudioDevice, const UPTRINT NodeWaveInstanceHash, FActiveSound& ActiveSound, const FSoundParseParameters& ParseParams, TArray<FWaveInstance*>& WaveInstances ) override;
+	virtual float GetMaxAudibleDistance() override;
+	virtual float GetDuration() override;
 	// End USoundBase interface.
 
 	/**
@@ -375,7 +375,7 @@ public:
 	/**
 	 * Begins caching platform data in the background for the platform requested
 	 */
-	virtual void BeginCacheForCookedPlatformData(  const ITargetPlatform *TargetPlatform ) OVERRIDE;
+	virtual void BeginCacheForCookedPlatformData(  const ITargetPlatform *TargetPlatform ) override;
 
 	/**
 	 * Begins caching platform data in the background.
diff --git a/Engine/Source/Runtime/Engine/Classes/Sound/SoundWaveStreaming.h b/Engine/Source/Runtime/Engine/Classes/Sound/SoundWaveStreaming.h
--- a/Engine/Source/Runtime/Engine/Classes/Sound/SoundWaveStreaming.h
+++ b/Engine/Source/Runtime/Engine/Classes/Sound/SoundWaveStreaming.h
@@ -21,17 +21,17 @@ private:
 public:
 
 	// Begin UObject interface. 
-	virtual void Serialize( FArchive& Ar ) OVERRIDE;
-	virtual void GetAssetRegistryTags(TArray<FAssetRegistryTag>& OutTags) const OVERRIDE;
-	virtual SIZE_T GetResourceSize(EResourceSizeMode::Type Mode) OVERRIDE;
+	virtual void Serialize( FArchive& Ar ) override;
+	virtual void GetAssetRegistryTags(TArray<FAssetRegistryTag>& OutTags) const override;
+	virtual SIZE_T GetResourceSize(EResourceSizeMode::Type Mode) override;
 	// End UObject interface. 
 
 	// Begin USoundWave interface.
-	virtual int32 GeneratePCMData(uint8* PCMData, const int32 SamplesNeeded) OVERRIDE;
-	virtual FByteBulkData* GetCompressedData(FName Format) OVERRIDE;
-	virtual void InitAudioResource( FByteBulkData& CompressedData ) OVERRIDE;
-	virtual bool InitAudioResource(FName Format) OVERRIDE;
-	virtual int32 GetResourceSizeForFormat(FName Format) OVERRIDE;
+	virtual int32 GeneratePCMData(uint8* PCMData, const int32 SamplesNeeded) override;
+	virtual FByteBulkData* GetCompressedData(FName Format) override;
+	virtual void InitAudioResource( FByteBulkData& CompressedData ) override;
+	virtual bool InitAudioResource(FName Format) override;
+	virtual int32 GetResourceSizeForFormat(FName Format) override;
 	// End USoundWave interface.
 
 
diff --git a/Engine/Source/Runtime/Engine/Classes/VectorField/VectorFieldAnimated.h b/Engine/Source/Runtime/Engine/Classes/VectorField/VectorFieldAnimated.h
--- a/Engine/Source/Runtime/Engine/Classes/VectorField/VectorFieldAnimated.h
+++ b/Engine/Source/Runtime/Engine/Classes/VectorField/VectorFieldAnimated.h
@@ -79,12 +79,12 @@ class UVectorFieldAnimated : public UVectorField
 
 	// Begin UObject interface.
 #if WITH_EDITOR
-	virtual void PostEditChangeProperty(FPropertyChangedEvent& PropertyChangedEvent) OVERRIDE;
+	virtual void PostEditChangeProperty(FPropertyChangedEvent& PropertyChangedEvent) override;
 #endif // WITH_EDITOR
 	// End UObject interface
 
 	// Begin UVectorField Interface
-	virtual void InitInstance(class FVectorFieldInstance* Instance, bool bPreviewInstance) OVERRIDE;
+	virtual void InitInstance(class FVectorFieldInstance* Instance, bool bPreviewInstance) override;
 	// End UVectorField Interface
 
 };
diff --git a/Engine/Source/Runtime/Engine/Classes/VectorField/VectorFieldStatic.h b/Engine/Source/Runtime/Engine/Classes/VectorField/VectorFieldStatic.h
--- a/Engine/Source/Runtime/Engine/Classes/VectorField/VectorFieldStatic.h
+++ b/Engine/Source/Runtime/Engine/Classes/VectorField/VectorFieldStatic.h
@@ -44,16 +44,16 @@ public:
 #endif // WITH_EDITORONLY_DATA
 
 	// Begin UObject interface.
-	virtual void PostLoad() OVERRIDE;
-	virtual void BeginDestroy() OVERRIDE;
+	virtual void PostLoad() override;
+	virtual void BeginDestroy() override;
 #if WITH_EDITOR
-	virtual void PostEditChangeProperty(FPropertyChangedEvent& PropertyChangedEvent) OVERRIDE;
+	virtual void PostEditChangeProperty(FPropertyChangedEvent& PropertyChangedEvent) override;
 #endif // WITH_EDITOR
-	virtual void Serialize(FArchive& Ar) OVERRIDE;
+	virtual void Serialize(FArchive& Ar) override;
 	// End UObject interface.
 
 	// Begin UVectorField Interface
-	virtual void InitInstance(class FVectorFieldInstance* Instance, bool bPreviewInstance) OVERRIDE;
+	virtual void InitInstance(class FVectorFieldInstance* Instance, bool bPreviewInstance) override;
 	// End UVectorField Interface
 
 	/**
diff --git a/Engine/Source/Runtime/Engine/Classes/Vehicles/TireType.h b/Engine/Source/Runtime/Engine/Classes/Vehicles/TireType.h
--- a/Engine/Source/Runtime/Engine/Classes/Vehicles/TireType.h
+++ b/Engine/Source/Runtime/Engine/Classes/Vehicles/TireType.h
@@ -47,20 +47,20 @@ public:
 	 * Called after the C++ constructor and after the properties have been initialized, but before the config has been loaded, etc.
 	 * mainly this is to emulate some behavior of when the constructor was called after the properties were intialized.
 	 */
-	virtual void PostInitProperties() OVERRIDE;
+	virtual void PostInitProperties() override;
 	
 	/**
 	 * Called before destroying the object.  This is called immediately upon deciding to destroy the object, to allow the object to begin an
 	 * asynchronous cleanup process.
 	 */
-	virtual void BeginDestroy() OVERRIDE;
+	virtual void BeginDestroy() override;
 
 #if WITH_EDITOR
 
 	/**
 	 * Respond to a property change in editor
 	 */
-	virtual void PostEditChangeProperty( struct FPropertyChangedEvent& PropertyChangedEvent ) OVERRIDE;
+	virtual void PostEditChangeProperty( struct FPropertyChangedEvent& PropertyChangedEvent ) override;
 
 #endif //WITH_EDITOR
 
diff --git a/Engine/Source/Runtime/Engine/Classes/Vehicles/WheeledVehicleMovementComponent.h b/Engine/Source/Runtime/Engine/Classes/Vehicles/WheeledVehicleMovementComponent.h
--- a/Engine/Source/Runtime/Engine/Classes/Vehicles/WheeledVehicleMovementComponent.h
+++ b/Engine/Source/Runtime/Engine/Classes/Vehicles/WheeledVehicleMovementComponent.h
@@ -265,14 +265,14 @@ class ENGINE_API UWheeledVehicleMovementComponent : public UPawnMovementComponen
 	virtual void UpdateDrag( float DeltaTime );
 
 	/** Used to create any physics engine information for this component */
-	virtual void CreatePhysicsState() OVERRIDE;
+	virtual void CreatePhysicsState() override;
 
 	/** Used to shut down and pysics engine structure for this component */
-	virtual void DestroyPhysicsState() OVERRIDE;
+	virtual void DestroyPhysicsState() override;
 
-	virtual bool ShouldCreatePhysicsState() const OVERRIDE;
+	virtual bool ShouldCreatePhysicsState() const override;
 
-	virtual bool HasValidPhysicsState() const OVERRIDE;
+	virtual bool HasValidPhysicsState() const override;
 
 	/** Draw debug text for the wheels and suspension */
 	virtual void DrawDebug(UCanvas* Canvas, float& YL, float& YPos);
@@ -285,7 +285,7 @@ class ENGINE_API UWheeledVehicleMovementComponent : public UPawnMovementComponen
 
 #if WITH_EDITOR
 	/** Respond to a property change in editor */
-	virtual void PostEditChangeProperty(struct FPropertyChangedEvent& PropertyChangedEvent) OVERRIDE;
+	virtual void PostEditChangeProperty(struct FPropertyChangedEvent& PropertyChangedEvent) override;
 #endif //WITH_EDITOR
 
 #endif // WITH_VEHICLE
diff --git a/Engine/Source/Runtime/Engine/Classes/Vehicles/WheeledVehicleMovementComponent4W.h b/Engine/Source/Runtime/Engine/Classes/Vehicles/WheeledVehicleMovementComponent4W.h
--- a/Engine/Source/Runtime/Engine/Classes/Vehicles/WheeledVehicleMovementComponent4W.h
+++ b/Engine/Source/Runtime/Engine/Classes/Vehicles/WheeledVehicleMovementComponent4W.h
@@ -170,10 +170,10 @@ class ENGINE_API UWheeledVehicleMovementComponent4W : public UWheeledVehicleMove
 	float AckermannAccuracy;
 
 
-	virtual void Serialize(FArchive & Ar) OVERRIDE;
-	virtual void ComputeConstants() OVERRIDE;
+	virtual void Serialize(FArchive & Ar) override;
+	virtual void ComputeConstants() override;
 #if WITH_EDITOR
-	virtual void PostEditChangeProperty(struct FPropertyChangedEvent& PropertyChangedEvent) OVERRIDE;
+	virtual void PostEditChangeProperty(struct FPropertyChangedEvent& PropertyChangedEvent) override;
 #endif
 
 protected:
@@ -181,9 +181,9 @@ protected:
 #if WITH_VEHICLE
 
 	/** Allocate and setup the PhysX vehicle */
-	virtual void SetupVehicle() OVERRIDE;
+	virtual void SetupVehicle() override;
 
-	virtual void UpdateSimulation(float DeltaTime) OVERRIDE;
+	virtual void UpdateSimulation(float DeltaTime) override;
 
 #endif // WITH_VEHICLE
 
diff --git a/Engine/Source/Runtime/Engine/Private/AI/Navigation/NavCollision.cpp b/Engine/Source/Runtime/Engine/Private/AI/Navigation/NavCollision.cpp
--- a/Engine/Source/Runtime/Engine/Private/AI/Navigation/NavCollision.cpp
+++ b/Engine/Source/Runtime/Engine/Private/AI/Navigation/NavCollision.cpp
@@ -57,17 +57,17 @@ private:
 public:
 	FDerivedDataNavCollisionCooker(FName InFormat, UNavCollision* InInstance);
 
-	virtual const TCHAR* GetPluginName() const OVERRIDE
+	virtual const TCHAR* GetPluginName() const override
 	{
 		return TEXT("NavCollision");
 	}
 
-	virtual const TCHAR* GetVersionString() const OVERRIDE
+	virtual const TCHAR* GetVersionString() const override
 	{
 		return TEXT("B89838347A4348138EE337A847529C5C");
 	}
 
-	virtual FString GetPluginSpecificCacheKeySuffix() const OVERRIDE
+	virtual FString GetPluginSpecificCacheKeySuffix() const override
 	{
 		const uint16 Version = 13;
 
@@ -79,12 +79,12 @@ public:
 			);
 	}
 
-	virtual bool IsBuildThreadsafe() const OVERRIDE
+	virtual bool IsBuildThreadsafe() const override
 	{
 		return false;
 	}
 
-	virtual bool Build( TArray<uint8>& OutData ) OVERRIDE;
+	virtual bool Build( TArray<uint8>& OutData ) override;
 
 	/** Return true if we can build **/
 	bool CanBuild()
diff --git a/Engine/Source/Runtime/Engine/Private/AI/Navigation/NavGraphGenerator.h b/Engine/Source/Runtime/Engine/Private/AI/Navigation/NavGraphGenerator.h
--- a/Engine/Source/Runtime/Engine/Private/AI/Navigation/NavGraphGenerator.h
+++ b/Engine/Source/Runtime/Engine/Private/AI/Navigation/NavGraphGenerator.h
@@ -23,11 +23,11 @@ public:
 	// FNavDataGenerator overrides
 	//----------------------------------------------------------------------//
 	// Triggers navmesh building process
-	virtual bool Generate() OVERRIDE;
+	virtual bool Generate() override;
 
-	virtual bool IsBuildInProgress(bool bCheckDirtyToo = false) const OVERRIDE;
+	virtual bool IsBuildInProgress(bool bCheckDirtyToo = false) const override;
 
-	virtual void OnWorldInitDone(bool bAllowedToRebuild) OVERRIDE;
+	virtual void OnWorldInitDone(bool bAllowedToRebuild) override;
 		
 	//--- accessors --- //
 
@@ -44,11 +44,11 @@ private:
 	void TriggerGeneration();
 
 public:
-	virtual void OnNavigationBuildingLocked() OVERRIDE;
-	virtual void OnNavigationBuildingUnlocked(bool bForce) OVERRIDE;
-	virtual void OnNavigationBoundsUpdated(class AVolume* Volume) OVERRIDE;
+	virtual void OnNavigationBuildingLocked() override;
+	virtual void OnNavigationBuildingUnlocked(bool bForce) override;
+	virtual void OnNavigationBoundsUpdated(class AVolume* Volume) override;
 
-	virtual void OnNavigationDataDestroyed(class ANavigationData* NavData) OVERRIDE;
+	virtual void OnNavigationDataDestroyed(class ANavigationData* NavData) override;
 
 private:
 	/** Bounding geometry definition. */
diff --git a/Engine/Source/Runtime/Engine/Private/AI/Navigation/NavTestRenderingComponent.cpp b/Engine/Source/Runtime/Engine/Private/AI/Navigation/NavTestRenderingComponent.cpp
--- a/Engine/Source/Runtime/Engine/Private/AI/Navigation/NavTestRenderingComponent.cpp
+++ b/Engine/Source/Runtime/Engine/Private/AI/Navigation/NavTestRenderingComponent.cpp
@@ -58,13 +58,13 @@ public:
 	{
 	}
 
-	virtual void RegisterDebugDrawDelgate() OVERRIDE
+	virtual void RegisterDebugDrawDelgate() override
 	{
 		DebugTextDrawingDelegate = FDebugDrawDelegate::CreateRaw(this, &FNavTestSceneProxy::DrawDebugLabels);
 		UDebugDrawService::Register(TEXT("Navigation"), DebugTextDrawingDelegate);
 	}
 
-	virtual void UnregisterDebugDrawDelgate() OVERRIDE
+	virtual void UnregisterDebugDrawDelgate() override
 	{
 		if (DebugTextDrawingDelegate.IsBound())
 		{
@@ -72,7 +72,7 @@ public:
 		}
 	}
 
-	virtual void DrawDynamicElements(FPrimitiveDrawInterface* PDI,const FSceneView* View) OVERRIDE
+	virtual void DrawDynamicElements(FPrimitiveDrawInterface* PDI,const FSceneView* View) override
 	{
 		if (NavTestActor)
 		{
@@ -355,7 +355,7 @@ public:
 		return Result;
 	}
 
-	virtual uint32 GetMemoryFootprint( void ) const OVERRIDE { return sizeof( *this ) + GetAllocatedSize(); }
+	virtual uint32 GetMemoryFootprint( void ) const override { return sizeof( *this ) + GetAllocatedSize(); }
 
 	uint32 GetAllocatedSize( void ) const 
 	{
diff --git a/Engine/Source/Runtime/Engine/Private/AI/Navigation/PImplRecastNavMesh.h b/Engine/Source/Runtime/Engine/Private/AI/Navigation/PImplRecastNavMesh.h
--- a/Engine/Source/Runtime/Engine/Private/AI/Navigation/PImplRecastNavMesh.h
+++ b/Engine/Source/Runtime/Engine/Private/AI/Navigation/PImplRecastNavMesh.h
@@ -22,21 +22,21 @@ class ENGINE_API FRecastQueryFilter : public INavigationQueryFilterInterface, pu
 public:
 	virtual ~FRecastQueryFilter(){}
 
-	virtual void Reset() OVERRIDE;
-
-	virtual void SetAreaCost(uint8 AreaType, float Cost) OVERRIDE;
-	virtual void SetFixedAreaEnteringCost(uint8 AreaType, float Cost) OVERRIDE;
-	virtual void SetExcludedArea(uint8 AreaType) OVERRIDE;
-	virtual void SetAllAreaCosts(const float* CostArray, const int32 Count) OVERRIDE;
-	virtual void GetAllAreaCosts(float* CostArray, float* FixedCostArray, const int32 Count) const OVERRIDE;
-	virtual void SetBacktrackingEnabled(const bool bBacktracking) OVERRIDE;
-	virtual bool IsBacktrackingEnabled() const OVERRIDE;
-	virtual bool IsEqual(const INavigationQueryFilterInterface* Other) const OVERRIDE;
-	virtual void SetIncludeFlags(uint16 Flags) OVERRIDE;
-	virtual uint16 GetIncludeFlags() const OVERRIDE;
-	virtual void SetExcludeFlags(uint16 Flags) OVERRIDE;
-	virtual uint16 GetExcludeFlags() const OVERRIDE;
-	virtual INavigationQueryFilterInterface* CreateCopy() const OVERRIDE;
+	virtual void Reset() override;
+
+	virtual void SetAreaCost(uint8 AreaType, float Cost) override;
+	virtual void SetFixedAreaEnteringCost(uint8 AreaType, float Cost) override;
+	virtual void SetExcludedArea(uint8 AreaType) override;
+	virtual void SetAllAreaCosts(const float* CostArray, const int32 Count) override;
+	virtual void GetAllAreaCosts(float* CostArray, float* FixedCostArray, const int32 Count) const override;
+	virtual void SetBacktrackingEnabled(const bool bBacktracking) override;
+	virtual bool IsBacktrackingEnabled() const override;
+	virtual bool IsEqual(const INavigationQueryFilterInterface* Other) const override;
+	virtual void SetIncludeFlags(uint16 Flags) override;
+	virtual uint16 GetIncludeFlags() const override;
+	virtual void SetExcludeFlags(uint16 Flags) override;
+	virtual uint16 GetExcludeFlags() const override;
+	virtual INavigationQueryFilterInterface* CreateCopy() const override;
 
 	const dtQueryFilter* GetAsDetourQueryFilter() const { return this; }
 };
@@ -44,7 +44,7 @@ public:
 struct ENGINE_API FRecastSpeciaLinkFilter : public dtQuerySpecialLinkFilter
 {
 	FRecastSpeciaLinkFilter(UNavigationSystem* NavSystem, const UObject* Owner) : NavSys(NavSystem), SearchOwner(Owner) {}
-	virtual bool isLinkAllowed(const int32 UserId) const OVERRIDE;
+	virtual bool isLinkAllowed(const int32 UserId) const override;
 
 	UNavigationSystem* NavSys;
 	const UObject* SearchOwner;
diff --git a/Engine/Source/Runtime/Engine/Private/AI/Navigation/RecastNavMeshGenerator.cpp b/Engine/Source/Runtime/Engine/Private/AI/Navigation/RecastNavMeshGenerator.cpp
--- a/Engine/Source/Runtime/Engine/Private/AI/Navigation/RecastNavMeshGenerator.cpp
+++ b/Engine/Source/Runtime/Engine/Private/AI/Navigation/RecastNavMeshGenerator.cpp
@@ -304,14 +304,14 @@ struct FRecastGeometryExport : public FNavigableGeometryExport
 	FWalkableSlopeOverride SlopeOverride;
 
 #if WITH_PHYSX
-	virtual void ExportPxTriMesh16Bit(class physx::PxTriangleMesh const * const TriMesh, const FTransform& LocalToWorld) OVERRIDE;
-	virtual void ExportPxTriMesh32Bit(class physx::PxTriangleMesh const * const TriMesh, const FTransform& LocalToWorld) OVERRIDE;
-	virtual void ExportPxConvexMesh(class physx::PxConvexMesh const * const ConvexMesh, const FTransform& LocalToWorld) OVERRIDE;
-	virtual void ExportPxHeightField(physx::PxHeightField const * const HeightField, const FTransform& LocalToWorld) OVERRIDE;
+	virtual void ExportPxTriMesh16Bit(class physx::PxTriangleMesh const * const TriMesh, const FTransform& LocalToWorld) override;
+	virtual void ExportPxTriMesh32Bit(class physx::PxTriangleMesh const * const TriMesh, const FTransform& LocalToWorld) override;
+	virtual void ExportPxConvexMesh(class physx::PxConvexMesh const * const ConvexMesh, const FTransform& LocalToWorld) override;
+	virtual void ExportPxHeightField(physx::PxHeightField const * const HeightField, const FTransform& LocalToWorld) override;
 #endif // WITH_PHYSX
-	virtual void ExportCustomMesh(const FVector* InVertices, int32 NumVerts, const int32* InIndices, int32 NumIndices, const FTransform& LocalToWorld) OVERRIDE;
-	virtual void ExportRigidBodySetup(UBodySetup& BodySetup, const FTransform& LocalToWorld) OVERRIDE;
-	virtual void AddNavModifiers(const FCompositeNavModifier& Modifiers) OVERRIDE;
+	virtual void ExportCustomMesh(const FVector* InVertices, int32 NumVerts, const int32* InIndices, int32 NumIndices, const FTransform& LocalToWorld) override;
+	virtual void ExportRigidBodySetup(UBodySetup& BodySetup, const FTransform& LocalToWorld) override;
+	virtual void AddNavModifiers(const FCompositeNavModifier& Modifiers) override;
 };
 
 FRecastVoxelCache::FRecastVoxelCache(const uint8* Memory)
@@ -4795,7 +4795,7 @@ static class FNavigationGeomExec : private FSelfRegisteringExec
 {
 public:
 	/** Console commands, see embeded usage statement **/
-	virtual bool Exec( UWorld* InWorld, const TCHAR* Cmd, FOutputDevice& Ar ) OVERRIDE
+	virtual bool Exec( UWorld* InWorld, const TCHAR* Cmd, FOutputDevice& Ar ) override
 	{
 #if ALLOW_DEBUG_FILES && !(UE_BUILD_SHIPPING || UE_BUILD_TEST)
 		bool bCorrectCmd = FParse::Command(&Cmd, TEXT("ExportNavigation"));
diff --git a/Engine/Source/Runtime/Engine/Private/AI/Navigation/RecastNavMeshGenerator.h b/Engine/Source/Runtime/Engine/Private/AI/Navigation/RecastNavMeshGenerator.h
--- a/Engine/Source/Runtime/Engine/Private/AI/Navigation/RecastNavMeshGenerator.h
+++ b/Engine/Source/Runtime/Engine/Private/AI/Navigation/RecastNavMeshGenerator.h
@@ -341,14 +341,14 @@ private:
 public:
 
 	// Triggers navmesh building process
-	virtual bool Generate() OVERRIDE;
+	virtual bool Generate() override;
 
 	/** Asks generator to update navigation affected by DirtyAreas */
-	virtual void RebuildDirtyAreas(const TArray<FNavigationDirtyArea>& DirtyAreas) OVERRIDE;
+	virtual void RebuildDirtyAreas(const TArray<FNavigationDirtyArea>& DirtyAreas) override;
 
-	virtual void TriggerGeneration() OVERRIDE;
+	virtual void TriggerGeneration() override;
 
-	virtual bool IsBuildInProgress(bool bCheckDirtyToo = false) const OVERRIDE;
+	virtual bool IsBuildInProgress(bool bCheckDirtyToo = false) const override;
 
 	bool AreAnyTilesBeingBuilt(bool bCheckDirtyToo = false) const;
 	bool IsAsyncBuildInProgress() const;
@@ -357,7 +357,7 @@ public:
 	 *	If TileIndex is out of bounds function returns false (i.e. not being built) */
 	bool IsTileFresh(int32 X, int32 Y, float FreshnessTime = FRecastNavMeshGenerator::DefaultFreshness) const;
 
-	virtual void OnWorldInitDone(bool bAllowedToRebuild) OVERRIDE;
+	virtual void OnWorldInitDone(bool bAllowedToRebuild) override;
 
 	/** Moves data from this generator to it's destination navmesh object. 
 	 *	@return fail if DestNavMesh is not valid in some way, true otherwise */
@@ -459,7 +459,7 @@ private:
 	//void NotifyBoxDirty(FBox Bounds, const AActor* Actor);
 	//void MarkBoxDirty(FBox Bounds, const AActor* Actor);
 
-	virtual void RebuildAll() OVERRIDE;
+	virtual void RebuildAll() override;
 
 	void RequestDirtyTilesRebuild();
 	
@@ -491,16 +491,16 @@ public:
 	/** Called to trigger generation at earliest convenience */
 	void RequestGeneration();
 
-	virtual void OnNavigationBuildingLocked() OVERRIDE;
-	virtual void OnNavigationBuildingUnlocked(bool bForce) OVERRIDE;
-	virtual void OnNavigationBoundsUpdated(class AVolume* Volume) OVERRIDE;
+	virtual void OnNavigationBuildingLocked() override;
+	virtual void OnNavigationBuildingUnlocked(bool bForce) override;
+	virtual void OnNavigationBoundsUpdated(class AVolume* Volume) override;
 
-	virtual void OnNavigationDataDestroyed(class ANavigationData* NavData) OVERRIDE;
+	virtual void OnNavigationDataDestroyed(class ANavigationData* NavData) override;
 
 	//----------------------------------------------------------------------//
 	// debug
 	//----------------------------------------------------------------------//
-	virtual uint32 LogMemUsed() const OVERRIDE;
+	virtual uint32 LogMemUsed() const override;
 
 private:
 	/** Parameters defining navmesh tiles */
diff --git a/Engine/Source/Runtime/Engine/Private/ArrowComponent.cpp b/Engine/Source/Runtime/Engine/Private/ArrowComponent.cpp
--- a/Engine/Source/Runtime/Engine/Private/ArrowComponent.cpp
+++ b/Engine/Source/Runtime/Engine/Private/ArrowComponent.cpp
@@ -247,7 +247,7 @@ public:
 		return Result;
 	}
 
-	virtual void OnTransformChanged() OVERRIDE
+	virtual void OnTransformChanged() override
 	{
 		Origin = GetLocalToWorld().GetOrigin();
 	}
diff --git a/Engine/Source/Runtime/Engine/Private/Atmosphere/Atmosphere.cpp b/Engine/Source/Runtime/Engine/Private/Atmosphere/Atmosphere.cpp
--- a/Engine/Source/Runtime/Engine/Private/Atmosphere/Atmosphere.cpp
+++ b/Engine/Source/Runtime/Engine/Private/Atmosphere/Atmosphere.cpp
@@ -185,7 +185,7 @@ public:
 	}
 
 	/** FTickableEditorObject interface */
-	virtual void Tick( float DeltaTime ) OVERRIDE
+	virtual void Tick( float DeltaTime ) override
 	{
 		if( Component && Component->PrecomputeCounter.GetValue() == UAtmosphericFogComponent::EFinishedComputation ) 
 		{
@@ -198,7 +198,7 @@ public:
 		return true;
 	}
 
-	TStatId GetStatId() const OVERRIDE
+	TStatId GetStatId() const override
 	{
 		RETURN_QUICK_DECLARE_CYCLE_STAT(FAtmospherePrecomputeDataHandler, STATGROUP_Tickables);
 	}
@@ -693,7 +693,7 @@ public:
 	virtual ~FAtmospherePrecomputeInstanceData()
 	{}
 
-	virtual bool MatchesComponent(const UActorComponent* Component) const OVERRIDE
+	virtual bool MatchesComponent(const UActorComponent* Component) const override
 	{
 		return (PrecomputeParameter == CastChecked<UAtmosphericFogComponent>(Component)->GetPrecomputeParameters());
 	}
diff --git a/Engine/Source/Runtime/Engine/Private/Atmosphere/Atmosphere.h b/Engine/Source/Runtime/Engine/Private/Atmosphere/Atmosphere.h
--- a/Engine/Source/Runtime/Engine/Private/Atmosphere/Atmosphere.h
+++ b/Engine/Source/Runtime/Engine/Private/Atmosphere/Atmosphere.h
@@ -21,7 +21,7 @@ public:
 	/** 
 	 * @return ptr to the resource memory which has been preallocated
 	 */
-	virtual const void* GetResourceBulkData() const OVERRIDE
+	virtual const void* GetResourceBulkData() const override
 	{
 		check(BulkData != NULL);
 		return BulkData;
@@ -30,7 +30,7 @@ public:
 	/** 
 	 * @return size of resource memory
 	 */
-	virtual uint32 GetResourceBulkDataSize() const OVERRIDE
+	virtual uint32 GetResourceBulkDataSize() const override
 	{
 		check(BulkDataSize > 0);
 		return BulkDataSize;
@@ -39,7 +39,7 @@ public:
 	/**
 	 * Free memory after it has been used to initialize RHI resource 
 	 */
-	virtual void Discard() OVERRIDE
+	virtual void Discard() override
 	{
 	}
 
@@ -131,7 +131,7 @@ public:
 	/**
 	 * Initialize RHI resources.
 	 */
-	virtual void InitRHI() OVERRIDE
+	virtual void InitRHI() override
 	{
 		if ( TextureData && IsFeatureLevelSupported(GRHIShaderPlatform, ERHIFeatureLevel::SM3))
 		{
@@ -186,7 +186,7 @@ public:
 	/**
 	 * Release RHI resources.
 	 */
-	virtual void ReleaseRHI() OVERRIDE
+	virtual void ReleaseRHI() override
 	{
 		TextureRHI.SafeRelease();
 	}
diff --git a/Engine/Source/Runtime/Engine/Private/BillboardComponent.cpp b/Engine/Source/Runtime/Engine/Private/BillboardComponent.cpp
--- a/Engine/Source/Runtime/Engine/Private/BillboardComponent.cpp
+++ b/Engine/Source/Runtime/Engine/Private/BillboardComponent.cpp
@@ -109,7 +109,7 @@ public:
 	 * @param	PDI - draw interface to render to
 	 * @param	View - current view
 	 */
-	virtual void DrawDynamicElements(FPrimitiveDrawInterface* PDI, const FSceneView* View) OVERRIDE
+	virtual void DrawDynamicElements(FPrimitiveDrawInterface* PDI, const FSceneView* View) override
 	{
 		QUICK_SCOPE_CYCLE_COUNTER( STAT_SpriteSceneProxy_DrawDynamicElements );
 
@@ -175,7 +175,7 @@ public:
 				);
 		}
 	}
-	virtual FPrimitiveViewRelevance GetViewRelevance(const FSceneView* View) OVERRIDE
+	virtual FPrimitiveViewRelevance GetViewRelevance(const FSceneView* View) override
 	{
 		bool bVisible = View->Family->EngineShowFlags.BillboardSprites;
 #if WITH_EDITOR
@@ -193,11 +193,11 @@ public:
 		Result.bEditorPrimitiveRelevance = UseEditorCompositing(View);
 		return Result;
 	}
-	virtual void OnTransformChanged() OVERRIDE
+	virtual void OnTransformChanged() override
 	{
 		Origin = GetLocalToWorld().GetOrigin();
 	}
-	virtual uint32 GetMemoryFootprint(void) const OVERRIDE { return(sizeof(*this) + GetAllocatedSize()); }
+	virtual uint32 GetMemoryFootprint(void) const override { return(sizeof(*this) + GetAllocatedSize()); }
 	
 	uint32 GetAllocatedSize(void) const { return( FPrimitiveSceneProxy::GetAllocatedSize() ); }
 
diff --git a/Engine/Source/Runtime/Engine/Private/BoxComponent.cpp b/Engine/Source/Runtime/Engine/Private/BoxComponent.cpp
--- a/Engine/Source/Runtime/Engine/Private/BoxComponent.cpp
+++ b/Engine/Source/Runtime/Engine/Private/BoxComponent.cpp
@@ -96,7 +96,7 @@ FPrimitiveSceneProxy* UBoxComponent::CreateSceneProxy()
 		}
 
 
-		virtual void DrawDynamicElements(FPrimitiveDrawInterface* PDI,const FSceneView* View) OVERRIDE
+		virtual void DrawDynamicElements(FPrimitiveDrawInterface* PDI,const FSceneView* View) override
 		{
 			QUICK_SCOPE_CYCLE_COUNTER( STAT_BoxSceneProxy_DrawDynamicElements );
 
@@ -105,7 +105,7 @@ FPrimitiveSceneProxy* UBoxComponent::CreateSceneProxy()
 			DrawOrientedWireBox(PDI, LocalToWorld.GetOrigin(), LocalToWorld.GetScaledAxis( EAxis::X ), LocalToWorld.GetScaledAxis( EAxis::Y ), LocalToWorld.GetScaledAxis( EAxis::Z ), BoxExtents, DrawColor, SDPG_World);
 		}
 
-		virtual FPrimitiveViewRelevance GetViewRelevance(const FSceneView* View) OVERRIDE
+		virtual FPrimitiveViewRelevance GetViewRelevance(const FSceneView* View) override
 		{
 			const bool bVisible = !bDrawOnlyIfSelected || IsSelected();
 
@@ -116,7 +116,7 @@ FPrimitiveSceneProxy* UBoxComponent::CreateSceneProxy()
 			Result.bEditorPrimitiveRelevance = UseEditorCompositing(View);
 			return Result;
 		}
-		virtual uint32 GetMemoryFootprint( void ) const OVERRIDE { return( sizeof( *this ) + GetAllocatedSize() ); }
+		virtual uint32 GetMemoryFootprint( void ) const override { return( sizeof( *this ) + GetAllocatedSize() ); }
 		uint32 GetAllocatedSize( void ) const { return( FPrimitiveSceneProxy::GetAllocatedSize() ); }
 
 	private:
diff --git a/Engine/Source/Runtime/Engine/Private/CapsuleComponent.cpp b/Engine/Source/Runtime/Engine/Private/CapsuleComponent.cpp
--- a/Engine/Source/Runtime/Engine/Private/CapsuleComponent.cpp
+++ b/Engine/Source/Runtime/Engine/Private/CapsuleComponent.cpp
@@ -38,7 +38,7 @@ FPrimitiveSceneProxy* UCapsuleComponent::CreateSceneProxy()
 		* @param	PDI - draw interface to render to
 		* @param	View - current view
 		*/
-		virtual void DrawDynamicElements(FPrimitiveDrawInterface* PDI,const FSceneView* View) OVERRIDE
+		virtual void DrawDynamicElements(FPrimitiveDrawInterface* PDI,const FSceneView* View) override
 		{
 			QUICK_SCOPE_CYCLE_COUNTER( STAT_DrawCylinderSceneProxy_DrawDynamicElements );
 
@@ -48,7 +48,7 @@ FPrimitiveSceneProxy* UCapsuleComponent::CreateSceneProxy()
 			DrawWireCapsule( PDI, LocalToWorld.GetOrigin(), LocalToWorld.GetScaledAxis( EAxis::X ), LocalToWorld.GetScaledAxis( EAxis::Y ), LocalToWorld.GetScaledAxis( EAxis::Z ), DrawCapsuleColor, CapsuleRadius, CapsuleHalfHeight, CapsuleSides, SDPG_World );
 		}
 
-		virtual FPrimitiveViewRelevance GetViewRelevance(const FSceneView* View) OVERRIDE
+		virtual FPrimitiveViewRelevance GetViewRelevance(const FSceneView* View) override
 		{
 			const bool bVisible = !bDrawOnlyIfSelected || IsSelected();
 			FPrimitiveViewRelevance Result;
@@ -58,7 +58,7 @@ FPrimitiveSceneProxy* UCapsuleComponent::CreateSceneProxy()
 			Result.bEditorPrimitiveRelevance = UseEditorCompositing(View);
 			return Result;
 		}
-		virtual uint32 GetMemoryFootprint( void ) const OVERRIDE { return( sizeof( *this ) + GetAllocatedSize() ); }
+		virtual uint32 GetMemoryFootprint( void ) const override { return( sizeof( *this ) + GetAllocatedSize() ); }
 		uint32 GetAllocatedSize( void ) const { return( FPrimitiveSceneProxy::GetAllocatedSize() ); }
 
 	private:
diff --git a/Engine/Source/Runtime/Engine/Private/Collision/PhysXCollision.h b/Engine/Source/Runtime/Engine/Private/Collision/PhysXCollision.h
--- a/Engine/Source/Runtime/Engine/Private/Collision/PhysXCollision.h
+++ b/Engine/Source/Runtime/Engine/Private/Collision/PhysXCollision.h
@@ -109,10 +109,10 @@ public:
 	 */
 	static PxSceneQueryHitType::Enum CalcQueryHitType(const PxFilterData &PQueryFilter, const PxFilterData &PShapeFilter);
 	
-	virtual PxSceneQueryHitType::Enum preFilter(const PxFilterData& filterData, const PxShape* shape, const PxRigidActor* actor, PxSceneQueryFlags& queryFlags) OVERRIDE;
+	virtual PxSceneQueryHitType::Enum preFilter(const PxFilterData& filterData, const PxShape* shape, const PxRigidActor* actor, PxSceneQueryFlags& queryFlags) override;
 
 
-	virtual PxSceneQueryHitType::Enum postFilter(const PxFilterData& filterData, const PxSceneQueryHit& hit) OVERRIDE
+	virtual PxSceneQueryHitType::Enum postFilter(const PxFilterData& filterData, const PxSceneQueryHit& hit) override
 	{
 		// Currently not used
 		return PxSceneQueryHitType::eBLOCK;
@@ -131,7 +131,7 @@ public:
 		DiscardInitialOverlaps = false;
 	}
 
-	virtual PxSceneQueryHitType::Enum postFilter(const PxFilterData& filterData, const PxSceneQueryHit& hit) OVERRIDE;
+	virtual PxSceneQueryHitType::Enum postFilter(const PxFilterData& filterData, const PxSceneQueryHit& hit) override;
 };
 
 // MISC
diff --git a/Engine/Source/Runtime/Engine/Private/DeferRegisterComponents.h b/Engine/Source/Runtime/Engine/Private/DeferRegisterComponents.h
--- a/Engine/Source/Runtime/Engine/Private/DeferRegisterComponents.h
+++ b/Engine/Source/Runtime/Engine/Private/DeferRegisterComponents.h
@@ -44,7 +44,7 @@ public:
 	void RegisterComponents(AActor* Actor);
 
 	// FGCObject interface
-	virtual void AddReferencedObjects(FReferenceCollector& Collector) OVERRIDE;
+	virtual void AddReferencedObjects(FReferenceCollector& Collector) override;
 
 	/* Gets the FDeferRegisterComponents singleton. */
 	static FDeferRegisterComponents& Get();
diff --git a/Engine/Source/Runtime/Engine/Private/DestructibleFractureSettings.cpp b/Engine/Source/Runtime/Engine/Private/DestructibleFractureSettings.cpp
--- a/Engine/Source/Runtime/Engine/Private/DestructibleFractureSettings.cpp
+++ b/Engine/Source/Runtime/Engine/Private/DestructibleFractureSettings.cpp
@@ -192,15 +192,15 @@ class FProgressListener : public IProgressListener
 {
 public:
 	// Begin IProgressListener interface
-	virtual void	setProgress(int progress, const char* taskName = NULL) OVERRIDE {}
+	virtual void	setProgress(int progress, const char* taskName = NULL) override {}
 	// End IProgressListener interface
 };
 
 class FExplicitHierarchicalMeshEmbedding : public IExplicitHierarchicalMesh::IEmbedding
 {
 public:
-	virtual void	serialize(physx::general_PxIOStream2::PxFileBuf& stream, IExplicitHierarchicalMesh::IEmbedding::DataType type) const OVERRIDE {}
-	virtual void	deserialize(physx::general_PxIOStream2::PxFileBuf& stream, IExplicitHierarchicalMesh::IEmbedding::DataType type, physx::PxU32 version) OVERRIDE {}
+	virtual void	serialize(physx::general_PxIOStream2::PxFileBuf& stream, IExplicitHierarchicalMesh::IEmbedding::DataType type) const override {}
+	virtual void	deserialize(physx::general_PxIOStream2::PxFileBuf& stream, IExplicitHierarchicalMesh::IEmbedding::DataType type, physx::PxU32 version) override {}
 };
 
 #endif // WITH_APEX
diff --git a/Engine/Source/Runtime/Engine/Private/DirectionalLightComponent.cpp b/Engine/Source/Runtime/Engine/Private/DirectionalLightComponent.cpp
--- a/Engine/Source/Runtime/Engine/Private/DirectionalLightComponent.cpp
+++ b/Engine/Source/Runtime/Engine/Private/DirectionalLightComponent.cpp
@@ -160,14 +160,14 @@ public:
 		LightMinRoughness = MinRoughness;
 	}
 
-	virtual bool GetLightShaftOcclusionParameters(float& OutOcclusionMaskDarkness, float& OutOcclusionDepthRange) const OVERRIDE
+	virtual bool GetLightShaftOcclusionParameters(float& OutOcclusionMaskDarkness, float& OutOcclusionDepthRange) const override
 	{
 		OutOcclusionMaskDarkness = OcclusionMaskDarkness;
 		OutOcclusionDepthRange = OcclusionDepthRange;
 		return bEnableLightShaftOcclusion;
 	}
 
-	virtual FVector GetLightPositionForLightShafts(FVector ViewOrigin) const OVERRIDE
+	virtual FVector GetLightPositionForLightShafts(FVector ViewOrigin) const override
 	{
 		const FVector EffectiveDirection = LightShaftOverrideDirection.SizeSquared() > 0 ? LightShaftOverrideDirection : GetDirection();
 		return ViewOrigin - EffectiveDirection * WORLD_MAX;
diff --git a/Engine/Source/Runtime/Engine/Private/DynamicMeshBuilder.cpp b/Engine/Source/Runtime/Engine/Private/DynamicMeshBuilder.cpp
--- a/Engine/Source/Runtime/Engine/Private/DynamicMeshBuilder.cpp
+++ b/Engine/Source/Runtime/Engine/Private/DynamicMeshBuilder.cpp
@@ -116,7 +116,7 @@ public:
 	}
 	
 public: // From FTickableObjectRenderThread
-	virtual TStatId GetStatId() const OVERRIDE
+	virtual TStatId GetStatId() const override
 	{
 		RETURN_QUICK_DECLARE_CYCLE_STAT(FGlobalDynamicMeshIndexPool, STATGROUP_Tickables);
 	}
@@ -167,7 +167,7 @@ public:
 	}
 	
 public: // From FTickableObjectRenderThread
-	virtual TStatId GetStatId() const OVERRIDE
+	virtual TStatId GetStatId() const override
 	{
 		RETURN_QUICK_DECLARE_CYCLE_STAT(FGlobalDynamicMeshVertexPool, STATGROUP_Tickables);
 	}
diff --git a/Engine/Source/Runtime/Engine/Private/GPUSkinCache.h b/Engine/Source/Runtime/Engine/Private/GPUSkinCache.h
--- a/Engine/Source/Runtime/Engine/Private/GPUSkinCache.h
+++ b/Engine/Source/Runtime/Engine/Private/GPUSkinCache.h
@@ -28,7 +28,7 @@ public:
 	~FGPUSkinCache();
 
 	// FRenderResource overrides
-	virtual void ReleaseRHI() OVERRIDE;
+	virtual void ReleaseRHI() override;
 
 	void	Initialize();
 	void	Cleanup();
diff --git a/Engine/Source/Runtime/Engine/Private/GPUSkinVertexFactory.cpp b/Engine/Source/Runtime/Engine/Private/GPUSkinVertexFactory.cpp
--- a/Engine/Source/Runtime/Engine/Private/GPUSkinVertexFactory.cpp
+++ b/Engine/Source/Runtime/Engine/Private/GPUSkinVertexFactory.cpp
@@ -401,7 +401,7 @@ public:
 	* Bind shader constants by name
 	* @param	ParameterMap - mapping of named shader constants to indices
 	*/
-	virtual void Bind(const FShaderParameterMap& ParameterMap) OVERRIDE
+	virtual void Bind(const FShaderParameterMap& ParameterMap) override
 	{
 		BoneIndexOffset.Bind(ParameterMap,TEXT("BoneIndexOffset"));
 		MeshOriginParameter.Bind(ParameterMap,TEXT("MeshOrigin"));
@@ -414,7 +414,7 @@ public:
 	* Serialize shader params to an archive
 	* @param	Ar - archive to serialize to
 	*/
-	virtual void Serialize(FArchive& Ar) OVERRIDE
+	virtual void Serialize(FArchive& Ar) override
 	{
 		Ar << BoneIndexOffset;
 		Ar << MeshOriginParameter;
@@ -426,7 +426,7 @@ public:
 	/**
 	* Set any shader data specific to this vertex factory
 	*/
-	virtual void SetMesh(FRHICommandList& RHICmdList, FShader* Shader,const FVertexFactory* VertexFactory,const FSceneView& View,const FMeshBatchElement& BatchElement,uint32 DataFlags) const OVERRIDE
+	virtual void SetMesh(FRHICommandList& RHICmdList, FShader* Shader,const FVertexFactory* VertexFactory,const FSceneView& View,const FMeshBatchElement& BatchElement,uint32 DataFlags) const override
 	{
 		RHICmdList.CheckIsNull();
 
@@ -552,7 +552,7 @@ public:
 	* Bind shader constants by name
 	* @param	ParameterMap - mapping of named shader constants to indices
 	*/
-	virtual void Bind(const FShaderParameterMap& ParameterMap) OVERRIDE
+	virtual void Bind(const FShaderParameterMap& ParameterMap) override
 	{
 		FGPUSkinVertexFactoryShaderParameters::Bind(ParameterMap);
 		GPUSkinCacheStreamFloatOffset.Bind(ParameterMap,TEXT("GPUSkinCacheStreamFloatOffset"));
@@ -563,7 +563,7 @@ public:
 	* Serialize shader params to an archive
 	* @param	Ar - archive to serialize to
 	*/
-	virtual void Serialize(FArchive& Ar) OVERRIDE
+	virtual void Serialize(FArchive& Ar) override
 	{
 		FGPUSkinVertexFactoryShaderParameters::Serialize(Ar);
 		Ar << GPUSkinCacheStreamFloatOffset;
@@ -573,7 +573,7 @@ public:
 	/**
 	* Set any shader data specific to this vertex factory
 	*/
-	virtual void SetMesh(FRHICommandList& RHICmdList, FShader* Shader,const FVertexFactory* VertexFactory,const FSceneView& View,const FMeshBatchElement& BatchElement,uint32 DataFlags) const OVERRIDE
+	virtual void SetMesh(FRHICommandList& RHICmdList, FShader* Shader,const FVertexFactory* VertexFactory,const FSceneView& View,const FMeshBatchElement& BatchElement,uint32 DataFlags) const override
 	{
 		FGPUSkinVertexFactoryShaderParameters::SetMesh(RHICmdList, Shader, VertexFactory, View, BatchElement, DataFlags);
 
@@ -692,7 +692,7 @@ public:
 	* Bind shader constants by name
 	* @param	ParameterMap - mapping of named shader constants to indices
 	*/
-	virtual void Bind(const FShaderParameterMap& ParameterMap) OVERRIDE
+	virtual void Bind(const FShaderParameterMap& ParameterMap) override
 	{
 		FGPUSkinVertexFactoryShaderParameters::Bind(ParameterMap);
 		ClothSimulPositionsParameter.Bind(ParameterMap,TEXT("ClothSimulVertsPositions"));
@@ -703,7 +703,7 @@ public:
 	* Serialize shader params to an archive
 	* @param	Ar - archive to serialize to
 	*/
-	virtual void Serialize(FArchive& Ar) OVERRIDE
+	virtual void Serialize(FArchive& Ar) override
 	{ 
 		FGPUSkinVertexFactoryShaderParameters::Serialize(Ar);
 		Ar << ClothSimulPositionsParameter;
@@ -711,7 +711,7 @@ public:
 		Ar << ClothBlendWeightParameter;
 	}
 
-	virtual void SetMesh(FRHICommandList& RHICmdList, FShader* Shader,const FVertexFactory* VertexFactory,const FSceneView& View,const FMeshBatchElement& BatchElement,uint32 DataFlags) const OVERRIDE
+	virtual void SetMesh(FRHICommandList& RHICmdList, FShader* Shader,const FVertexFactory* VertexFactory,const FSceneView& View,const FMeshBatchElement& BatchElement,uint32 DataFlags) const override
 	{
 		if(Shader->GetVertexShader())
 		{
diff --git a/Engine/Source/Runtime/Engine/Private/GPUSort.cpp b/Engine/Source/Runtime/Engine/Private/GPUSort.cpp
--- a/Engine/Source/Runtime/Engine/Private/GPUSort.cpp
+++ b/Engine/Source/Runtime/Engine/Private/GPUSort.cpp
@@ -79,7 +79,7 @@ public:
 	/**
 	 * Initialize RHI resources.
 	 */
-	virtual void InitRHI() OVERRIDE
+	virtual void InitRHI() override
 	{
 		const int32 OffsetsCount = DIGIT_COUNT * MAX_GROUP_COUNT;
 		const int32 OffsetsBufferSize = OffsetsCount * sizeof(uint32);
@@ -107,7 +107,7 @@ public:
 	/**
 	 * Release RHI resources.
 	 */
-	virtual void ReleaseRHI() OVERRIDE
+	virtual void ReleaseRHI() override
 	{
 		for ( int32 BufferIndex = 0; BufferIndex < 2; ++BufferIndex )
 		{
@@ -180,7 +180,7 @@ public:
 	/**
 	 * Initialize RHI resources.
 	 */
-	virtual void InitRHI() OVERRIDE
+	virtual void InitRHI() override
 	{
 		if (GRHIFeatureLevel == ERHIFeatureLevel::SM5)
 		{
@@ -198,7 +198,7 @@ public:
 	/**
 	 * Release RHI resources.
 	 */
-	virtual void ReleaseRHI() OVERRIDE
+	virtual void ReleaseRHI() override
 	{
 		SortParametersBufferSRV.SafeRelease();
 		SortParametersBufferRHI.SafeRelease();
@@ -246,7 +246,7 @@ public:
 	}
 
 	/** Serialization. */
-	virtual bool Serialize( FArchive& Ar ) OVERRIDE
+	virtual bool Serialize( FArchive& Ar ) override
 	{
 		bool bShaderHasOutdatedParameters = FGlobalShader::Serialize( Ar );
 		Ar << OutOffsets;
@@ -333,7 +333,7 @@ public:
 	}
 
 	/** Serialization. */
-	virtual bool Serialize( FArchive& Ar ) OVERRIDE
+	virtual bool Serialize( FArchive& Ar ) override
 	{
 		bool bShaderHasOutdatedParameters = FGlobalShader::Serialize( Ar );
 		Ar << RadixSortParameterBuffer;
@@ -451,7 +451,7 @@ public:
 	}
 
 	/** Serialization. */
-	virtual bool Serialize( FArchive& Ar ) OVERRIDE
+	virtual bool Serialize( FArchive& Ar ) override
 	{
 		bool bShaderHasOutdatedParameters = FGlobalShader::Serialize( Ar );
 		Ar << InOffsets;
@@ -555,7 +555,7 @@ public:
 	}
 
 	/** Serialization. */
-	virtual bool Serialize( FArchive& Ar ) OVERRIDE
+	virtual bool Serialize( FArchive& Ar ) override
 	{
 		bool bShaderHasOutdatedParameters = FGlobalShader::Serialize( Ar );
 		Ar << RadixSortParameterBuffer;
diff --git a/Engine/Source/Runtime/Engine/Private/InstancedStaticMesh.cpp b/Engine/Source/Runtime/Engine/Private/InstancedStaticMesh.cpp
--- a/Engine/Source/Runtime/Engine/Private/InstancedStaticMesh.cpp
+++ b/Engine/Source/Runtime/Engine/Private/InstancedStaticMesh.cpp
@@ -506,13 +506,13 @@ void FInstancedStaticMeshVertexFactory::InitRHI()
 
 class FInstancedStaticMeshVertexFactoryShaderParameters : public FVertexFactoryShaderParameters
 {
-	virtual void Bind(const FShaderParameterMap& ParameterMap) OVERRIDE
+	virtual void Bind(const FShaderParameterMap& ParameterMap) override
 	{
 		InstancedViewTranslationParameter.Bind(ParameterMap,TEXT("InstancedViewTranslation"));
 		InstancingFadeOutParamsParameter.Bind(ParameterMap,TEXT("InstancingFadeOutParams"));
 	}
 
-	virtual void SetMesh(FRHICommandList& RHICmdList, FShader* VertexShader,const class FVertexFactory* VertexFactory,const class FSceneView& View,const struct FMeshBatchElement& BatchElement,uint32 DataFlags) const OVERRIDE
+	virtual void SetMesh(FRHICommandList& RHICmdList, FShader* VertexShader,const class FVertexFactory* VertexFactory,const class FSceneView& View,const struct FMeshBatchElement& BatchElement,uint32 DataFlags) const override
 	{
 		if( InstancedViewTranslationParameter.IsBound() )
 		{
@@ -933,7 +933,7 @@ public:
 
 	// FPrimitiveSceneProxy interface.
 	
-	virtual FPrimitiveViewRelevance GetViewRelevance(const FSceneView* View) OVERRIDE
+	virtual FPrimitiveViewRelevance GetViewRelevance(const FSceneView* View) override
 	{
 		FPrimitiveViewRelevance Result;
 		if(View->Family->EngineShowFlags.InstancedStaticMeshes)
@@ -951,13 +951,13 @@ public:
 	}
 
 	/** Draw the scene proxy as a dynamic element */
-	virtual void DrawDynamicElements(FPrimitiveDrawInterface* PDI,const FSceneView* View) OVERRIDE;
+	virtual void DrawDynamicElements(FPrimitiveDrawInterface* PDI,const FSceneView* View) override;
 
 	/** Sets up a FMeshBatch for a specific LOD and element. */
-	virtual bool GetMeshElement(int32 LODIndex,int32 ElementIndex,uint8 InDepthPriorityGroup,FMeshBatch& OutMeshElement, const bool bUseSelectedMaterial, const bool bUseHoveredMaterial) const OVERRIDE;
+	virtual bool GetMeshElement(int32 LODIndex,int32 ElementIndex,uint8 InDepthPriorityGroup,FMeshBatch& OutMeshElement, const bool bUseSelectedMaterial, const bool bUseHoveredMaterial) const override;
 
 	/** Sets up a wireframe FMeshBatch for a specific LOD. */
-	virtual bool GetWireframeMeshElement(int32 LODIndex, const FMaterialRenderProxy* WireframeRenderProxy, uint8 InDepthPriorityGroup, FMeshBatch& OutMeshElement) const OVERRIDE;
+	virtual bool GetWireframeMeshElement(int32 LODIndex, const FMaterialRenderProxy* WireframeRenderProxy, uint8 InDepthPriorityGroup, FMeshBatch& OutMeshElement) const override;
 
 	/**
 	 * Creates the hit proxies are used when DrawDynamicElements is called.
@@ -965,7 +965,7 @@ public:
 	 * @param OutHitProxies - Hit proxes which are created should be added to this array.
 	 * @return The hit proxy to use by default for elements drawn by DrawDynamicElements.
 	 */
-	virtual HHitProxy* CreateHitProxies(UPrimitiveComponent* Component,TArray<TRefCountPtr<HHitProxy> >& OutHitProxies) OVERRIDE
+	virtual HHitProxy* CreateHitProxies(UPrimitiveComponent* Component,TArray<TRefCountPtr<HHitProxy> >& OutHitProxies) override
 	{
 		if( InstancedRenderData.HitProxies.Num() )
 		{
@@ -981,7 +981,7 @@ public:
 		}
 	}
 
-	virtual bool IsDetailMesh() const OVERRIDE
+	virtual bool IsDetailMesh() const override
 	{
 		return true;
 	}
diff --git a/Engine/Source/Runtime/Engine/Private/InterpolateComponentToAction.h b/Engine/Source/Runtime/Engine/Private/InterpolateComponentToAction.h
--- a/Engine/Source/Runtime/Engine/Private/InterpolateComponentToAction.h
+++ b/Engine/Source/Runtime/Engine/Private/InterpolateComponentToAction.h
@@ -121,7 +121,7 @@ public:
 
 #if WITH_EDITOR
 	// Returns a human readable description of the latent operation's current state
-	virtual FString GetDescription() const OVERRIDE
+	virtual FString GetDescription() const override
 	{
 		return FString::Printf( *NSLOCTEXT("FInterpolateComponentToAction", "ActionTime", "Move (%.3f seconds left)").ToString(), TotalTime-TimeElapsed);
 	}
diff --git a/Engine/Source/Runtime/Engine/Private/Landscape/LandscapeRender.cpp b/Engine/Source/Runtime/Engine/Private/Landscape/LandscapeRender.cpp
--- a/Engine/Source/Runtime/Engine/Private/Landscape/LandscapeRender.cpp
+++ b/Engine/Source/Runtime/Engine/Private/Landscape/LandscapeRender.cpp
@@ -1961,7 +1961,7 @@ public:
 	* Bind shader constants by name
 	* @param	ParameterMap - mapping of named shader constants to indices
 	*/
-	virtual void Bind(const FShaderParameterMap& ParameterMap) OVERRIDE
+	virtual void Bind(const FShaderParameterMap& ParameterMap) override
 	{
 		HeightmapTextureParameter.Bind(ParameterMap,TEXT("HeightmapTexture"));
 		HeightmapTextureParameterSampler.Bind(ParameterMap,TEXT("HeightmapTextureSampler"));
@@ -1977,7 +1977,7 @@ public:
 	* Serialize shader params to an archive
 	* @param	Ar - archive to serialize to
 	*/
-	virtual void Serialize(FArchive& Ar) OVERRIDE
+	virtual void Serialize(FArchive& Ar) override
 	{
 		Ar << HeightmapTextureParameter;
 		Ar << HeightmapTextureParameterSampler;
@@ -1992,7 +1992,7 @@ public:
 	/**
 	* Set any shader data specific to this vertex factory
 	*/
-	virtual void SetMesh(FRHICommandList& RHICmdList, FShader* VertexShader,const class FVertexFactory* VertexFactory,const class FSceneView& View,const struct FMeshBatchElement& BatchElement,uint32 DataFlags) const OVERRIDE
+	virtual void SetMesh(FRHICommandList& RHICmdList, FShader* VertexShader,const class FVertexFactory* VertexFactory,const class FSceneView& View,const struct FMeshBatchElement& BatchElement,uint32 DataFlags) const override
 	{
 		SCOPE_CYCLE_COUNTER(STAT_LandscapeVFDrawTime);
 
@@ -2073,7 +2073,7 @@ public:
 		}
 	}
 
-	virtual uint32 GetSize() const OVERRIDE
+	virtual uint32 GetSize() const override
 	{
 		return sizeof(*this);
 	}
@@ -2097,7 +2097,7 @@ public:
 	/**
 	* Set any shader data specific to this vertex factory
 	*/
-	virtual void SetMesh(FRHICommandList& RHICmdList, FShader* VertexShader,const class FVertexFactory* VertexFactory,const class FSceneView& View,const struct FMeshBatchElement& BatchElement,uint32 DataFlags) const OVERRIDE
+	virtual void SetMesh(FRHICommandList& RHICmdList, FShader* VertexShader,const class FVertexFactory* VertexFactory,const class FSceneView& View,const struct FMeshBatchElement& BatchElement,uint32 DataFlags) const override
 	{
 		SCOPE_CYCLE_COUNTER(STAT_LandscapeVFDrawTime);
 
diff --git a/Engine/Source/Runtime/Engine/Private/Landscape/LandscapeRenderMobile.cpp b/Engine/Source/Runtime/Engine/Private/Landscape/LandscapeRenderMobile.cpp
--- a/Engine/Source/Runtime/Engine/Private/Landscape/LandscapeRenderMobile.cpp
+++ b/Engine/Source/Runtime/Engine/Private/Landscape/LandscapeRenderMobile.cpp
@@ -39,7 +39,7 @@ public:
 	* Bind shader constants by name
 	* @param	ParameterMap - mapping of named shader constants to indices
 	*/
-	virtual void Bind(const FShaderParameterMap& ParameterMap) OVERRIDE
+	virtual void Bind(const FShaderParameterMap& ParameterMap) override
 	{
 		LodValuesParameter.Bind(ParameterMap,TEXT("LodValues"));
 		NeighborSectionLodParameter.Bind(ParameterMap,TEXT("NeighborSectionLod"));
@@ -51,7 +51,7 @@ public:
 	* Serialize shader params to an archive
 	* @param	Ar - archive to serialize to
 	*/
-	virtual void Serialize(FArchive& Ar) OVERRIDE
+	virtual void Serialize(FArchive& Ar) override
 	{
 		Ar << LodValuesParameter;
 		Ar << NeighborSectionLodParameter;
@@ -61,7 +61,7 @@ public:
 	/**
 	* Set any shader data specific to this vertex factory
 	*/
-	virtual void SetMesh(FRHICommandList& RHICmdList, FShader* VertexShader,const class FVertexFactory* VertexFactory,const class FSceneView& View,const struct FMeshBatchElement& BatchElement,uint32 DataFlags) const OVERRIDE
+	virtual void SetMesh(FRHICommandList& RHICmdList, FShader* VertexShader,const class FVertexFactory* VertexFactory,const class FSceneView& View,const struct FMeshBatchElement& BatchElement,uint32 DataFlags) const override
 	{
 		SCOPE_CYCLE_COUNTER(STAT_LandscapeVFDrawTime);
 
diff --git a/Engine/Source/Runtime/Engine/Private/Level.cpp b/Engine/Source/Runtime/Engine/Private/Level.cpp
--- a/Engine/Source/Runtime/Engine/Private/Level.cpp
+++ b/Engine/Source/Runtime/Engine/Private/Level.cpp
@@ -1959,7 +1959,7 @@ public:
 	  * @param	PDI - draw interface to render to
 	  * @param	View - current view
 	  */
-	  virtual void DrawDynamicElements(FPrimitiveDrawInterface* PDI,const FSceneView* View) OVERRIDE
+	  virtual void DrawDynamicElements(FPrimitiveDrawInterface* PDI,const FSceneView* View) override
 	  {
 		  QUICK_SCOPE_CYCLE_COUNTER( STAT_LineBatcherSceneProxy_DrawDynamicElements );
 
diff --git a/Engine/Source/Runtime/Engine/Private/LightComponent.cpp b/Engine/Source/Runtime/Engine/Private/LightComponent.cpp
--- a/Engine/Source/Runtime/Engine/Private/LightComponent.cpp
+++ b/Engine/Source/Runtime/Engine/Private/LightComponent.cpp
@@ -765,7 +765,7 @@ public:
 	{}
 
 	// Begin FComponentInstanceDataBase interface
-	virtual bool MatchesComponent(const UActorComponent* Component) const OVERRIDE
+	virtual bool MatchesComponent(const UActorComponent* Component) const override
 	{
 		return Transform.Equals(CastChecked<ULightComponent>(Component)->ComponentToWorld);
 	}
