UE4 Refactoring. Changed OVERRIDE and FINAL macros to keywords override and final.

diff --git a/Engine/Source/Runtime/Engine/Private/LocalVertexFactory.cpp b/Engine/Source/Runtime/Engine/Private/LocalVertexFactory.cpp
--- a/Engine/Source/Runtime/Engine/Private/LocalVertexFactory.cpp
+++ b/Engine/Source/Runtime/Engine/Private/LocalVertexFactory.cpp
@@ -12,17 +12,17 @@ class FSpeedTreeVertexFactoryShaderParameters : public FVertexFactoryShaderParam
 {
 public:
 
-	virtual void Bind(const FShaderParameterMap& ParameterMap) OVERRIDE
+	virtual void Bind(const FShaderParameterMap& ParameterMap) override
 	{
 		LODParameter.Bind(ParameterMap, TEXT("SpeedTreeLODInfo"));
 	}
 
-	virtual void Serialize(FArchive& Ar) OVERRIDE
+	virtual void Serialize(FArchive& Ar) override
 	{
 		Ar << LODParameter;
 	}
 
-	virtual void SetMesh(FRHICommandList& RHICmdList, FShader* Shader,const FVertexFactory* VertexFactory,const FSceneView& View,const FMeshBatchElement& BatchElement,uint32 DataFlags) const OVERRIDE
+	virtual void SetMesh(FRHICommandList& RHICmdList, FShader* Shader,const FVertexFactory* VertexFactory,const FSceneView& View,const FMeshBatchElement& BatchElement,uint32 DataFlags) const override
 	{
 		if (View.Family != NULL && View.Family->Scene != NULL)
 		{
diff --git a/Engine/Source/Runtime/Engine/Private/Materials/HLSLMaterialTranslator.h b/Engine/Source/Runtime/Engine/Private/Materials/HLSLMaterialTranslator.h
--- a/Engine/Source/Runtime/Engine/Private/Materials/HLSLMaterialTranslator.h
+++ b/Engine/Source/Runtime/Engine/Private/Materials/HLSLMaterialTranslator.h
@@ -1198,7 +1198,7 @@ protected:
 		return INDEX_NONE;
 	}
 
-	virtual int32 CallExpression(FMaterialExpressionKey ExpressionKey,FMaterialCompiler* Compiler) OVERRIDE
+	virtual int32 CallExpression(FMaterialExpressionKey ExpressionKey,FMaterialCompiler* Compiler) override
 	{
 		// Check if this expression has already been translated.
 		int32* ExistingCodeIndex = FunctionStack.Last().ExpressionCodeMap[MaterialProperty][ShaderFrequency].Find(ExpressionKey);
@@ -1233,7 +1233,7 @@ protected:
 		}
 	}
 
-	virtual EMaterialValueType GetType(int32 Code) OVERRIDE
+	virtual EMaterialValueType GetType(int32 Code) override
 	{
 		if(Code != INDEX_NONE)
 		{
@@ -1245,17 +1245,17 @@ protected:
 		}
 	}
 
-	virtual EMaterialQualityLevel::Type GetQualityLevel() OVERRIDE
+	virtual EMaterialQualityLevel::Type GetQualityLevel() override
 	{
 		return QualityLevel;
 	}
 
-	virtual ERHIFeatureLevel::Type GetFeatureLevel() OVERRIDE
+	virtual ERHIFeatureLevel::Type GetFeatureLevel() override
 	{
 		return FeatureLevel;
 	}
 
-	virtual float GetRefractionDepthBiasValue() OVERRIDE
+	virtual float GetRefractionDepthBiasValue() override
 	{
 		return Material->GetRefractionDepthBiasValue();
 	}
@@ -1264,7 +1264,7 @@ protected:
 	 * Casts the passed in code to DestType, or generates a compile error if the cast is not valid. 
 	 * This will truncate a type (float4 -> float3) but not add components (float2 -> float3), however a float1 can be cast to any float type by replication. 
 	 */
-	virtual int32 ValidCast(int32 Code,EMaterialValueType DestType) OVERRIDE
+	virtual int32 ValidCast(int32 Code,EMaterialValueType DestType) override
 	{
 		if(Code == INDEX_NONE)
 		{
@@ -1345,7 +1345,7 @@ protected:
 		return CompiledResult;
 	}
 
-	virtual int32 ForceCast(int32 Code,EMaterialValueType DestType,bool bExactMatch=false,bool bReplicateValue=false) OVERRIDE
+	virtual int32 ForceCast(int32 Code,EMaterialValueType DestType,bool bExactMatch=false,bool bReplicateValue=false) override
 	{
 		if(Code == INDEX_NONE)
 		{
@@ -1414,18 +1414,18 @@ protected:
 	}
 
 	/** Pushes a function onto the compiler's function stack, which indicates that compilation is entering a function. */
-	virtual void PushFunction(const FMaterialFunctionCompileState& FunctionState) OVERRIDE
+	virtual void PushFunction(const FMaterialFunctionCompileState& FunctionState) override
 	{
 		FunctionStack.Push(FunctionState);
 	}	
 
 	/** Pops a function from the compiler's function stack, which indicates that compilation is leaving a function. */
-	virtual FMaterialFunctionCompileState PopFunction() OVERRIDE
+	virtual FMaterialFunctionCompileState PopFunction() override
 	{
 		return FunctionStack.Pop();
 	}
 
-	virtual int32 AccessCollectionParameter(UMaterialParameterCollection* ParameterCollection, int32 ParameterIndex, int32 ComponentIndex) OVERRIDE
+	virtual int32 AccessCollectionParameter(UMaterialParameterCollection* ParameterCollection, int32 ParameterIndex, int32 ComponentIndex) override
 	{
 		if (!ParameterCollection || ParameterIndex == -1)
 		{
@@ -1454,47 +1454,47 @@ protected:
 			ComponentIndex == -1 ? true : ComponentIndex % 4 == 3);
 	}
 
-	virtual int32 VectorParameter(FName ParameterName,const FLinearColor& DefaultValue) OVERRIDE
+	virtual int32 VectorParameter(FName ParameterName,const FLinearColor& DefaultValue) override
 	{
 		return AddUniformExpression(new FMaterialUniformExpressionVectorParameter(ParameterName,DefaultValue),MCT_Float4,TEXT(""));
 	}
 
-	virtual int32 ScalarParameter(FName ParameterName,float DefaultValue) OVERRIDE
+	virtual int32 ScalarParameter(FName ParameterName,float DefaultValue) override
 	{
 		return AddUniformExpression(new FMaterialUniformExpressionScalarParameter(ParameterName,DefaultValue),MCT_Float,TEXT(""));
 	}
 
-	virtual int32 Constant(float X) OVERRIDE
+	virtual int32 Constant(float X) override
 	{
 		return AddUniformExpression(new FMaterialUniformExpressionConstant(FLinearColor(X,X,X,X),MCT_Float),MCT_Float,TEXT("%0.8f"),X);
 	}
 
-	virtual int32 Constant2(float X,float Y) OVERRIDE
+	virtual int32 Constant2(float X,float Y) override
 	{
 		return AddUniformExpression(new FMaterialUniformExpressionConstant(FLinearColor(X,Y,0,0),MCT_Float2),MCT_Float2,TEXT("MaterialFloat2(%0.8f,%0.8f)"),X,Y);
 	}
 
-	virtual int32 Constant3(float X,float Y,float Z) OVERRIDE
+	virtual int32 Constant3(float X,float Y,float Z) override
 	{
 		return AddUniformExpression(new FMaterialUniformExpressionConstant(FLinearColor(X,Y,Z,0),MCT_Float3),MCT_Float3,TEXT("MaterialFloat3(%0.8f,%0.8f,%0.8f)"),X,Y,Z);
 	}
 
-	virtual int32 Constant4(float X,float Y,float Z,float W) OVERRIDE
+	virtual int32 Constant4(float X,float Y,float Z,float W) override
 	{
 		return AddUniformExpression(new FMaterialUniformExpressionConstant(FLinearColor(X,Y,Z,W),MCT_Float4),MCT_Float4,TEXT("MaterialFloat4(%0.8f,%0.8f,%0.8f,%0.8f)"),X,Y,Z,W);
 	}
 
-	virtual int32 GameTime() OVERRIDE
+	virtual int32 GameTime() override
 	{
 		return AddInlinedCodeChunk(MCT_Float, TEXT("View.GameTime"));
 	}
 
-	virtual int32 RealTime() OVERRIDE
+	virtual int32 RealTime() override
 	{
 		return AddInlinedCodeChunk(MCT_Float, TEXT("View.RealTime"));
 	}
 
-	virtual int32 PeriodicHint(int32 PeriodicCode) OVERRIDE
+	virtual int32 PeriodicHint(int32 PeriodicCode) override
 	{
 		if(PeriodicCode == INDEX_NONE)
 		{
@@ -1511,7 +1511,7 @@ protected:
 		}
 	}
 
-	virtual int32 Sine(int32 X) OVERRIDE
+	virtual int32 Sine(int32 X) override
 	{
 		if(X == INDEX_NONE)
 		{
@@ -1528,7 +1528,7 @@ protected:
 		}
 	}
 
-	virtual int32 Cosine(int32 X) OVERRIDE
+	virtual int32 Cosine(int32 X) override
 	{
 		if(X == INDEX_NONE)
 		{
@@ -1545,7 +1545,7 @@ protected:
 		}
 	}
 
-	virtual int32 Floor(int32 X) OVERRIDE
+	virtual int32 Floor(int32 X) override
 	{
 		if(X == INDEX_NONE)
 		{
@@ -1562,7 +1562,7 @@ protected:
 		}
 	}
 
-	virtual int32 Ceil(int32 X) OVERRIDE
+	virtual int32 Ceil(int32 X) override
 	{
 		if(X == INDEX_NONE)
 		{
@@ -1579,7 +1579,7 @@ protected:
 		}
 	}
 
-	virtual int32 Frac(int32 X) OVERRIDE
+	virtual int32 Frac(int32 X) override
 	{
 		if(X == INDEX_NONE)
 		{
@@ -1596,7 +1596,7 @@ protected:
 		}
 	}
 
-	virtual int32 Fmod(int32 A, int32 B) OVERRIDE
+	virtual int32 Fmod(int32 A, int32 B) override
 	{
 		if ((A == INDEX_NONE) || (B == INDEX_NONE))
 		{
@@ -1621,7 +1621,7 @@ protected:
 	* @param	X - Index to the FMaterialCompiler::CodeChunk entry for the input expression
 	* @return	Index to the new FMaterialCompiler::CodeChunk entry for this expression
 	*/	
-	virtual int32 Abs( int32 X ) OVERRIDE
+	virtual int32 Abs( int32 X ) override
 	{
 		if(X == INDEX_NONE)
 		{
@@ -1641,7 +1641,7 @@ protected:
 		}
 	}
 
-	virtual int32 ReflectionVector() OVERRIDE
+	virtual int32 ReflectionVector() override
 	{
 		if (ShaderFrequency != SF_Pixel && ShaderFrequency != SF_Compute && ShaderFrequency != SF_Domain)
 		{
@@ -1651,7 +1651,7 @@ protected:
 		return AddInlinedCodeChunk(MCT_Float3,TEXT("Parameters.ReflectionVector"));
 	}
 
-	virtual int32 ReflectionAboutCustomWorldNormal(int32 CustomWorldNormal, int32 bNormalizeCustomWorldNormal) OVERRIDE
+	virtual int32 ReflectionAboutCustomWorldNormal(int32 CustomWorldNormal, int32 bNormalizeCustomWorldNormal) override
 	{
 		if (ShaderFrequency != SF_Pixel && ShaderFrequency != SF_Compute && ShaderFrequency != SF_Domain)
 		{
@@ -1668,7 +1668,7 @@ protected:
 		return AddCodeChunk(MCT_Float3,TEXT("ReflectionAboutCustomWorldNormal(Parameters, %s, %s)"), *GetParameterCode(CustomWorldNormal), ShouldNormalize);
 	}
 
-	virtual int32 CameraVector() OVERRIDE
+	virtual int32 CameraVector() override
 	{
 		if (ShaderFrequency != SF_Pixel && ShaderFrequency != SF_Compute && ShaderFrequency != SF_Domain)
 		{
@@ -1681,12 +1681,12 @@ protected:
 		return AddInlinedCodeChunk(MCT_Float3,TEXT("Parameters.CameraVector"));
 	}
 
-	virtual int32 CameraWorldPosition() OVERRIDE
+	virtual int32 CameraWorldPosition() override
 	{
 		return AddInlinedCodeChunk(MCT_Float3,TEXT("View.ViewOrigin.xyz"));
 	}
 
-	virtual int32 LightVector() OVERRIDE
+	virtual int32 LightVector() override
 	{
 		if (ShaderFrequency != SF_Pixel && ShaderFrequency != SF_Compute)
 		{
@@ -1701,7 +1701,7 @@ protected:
 		return AddInlinedCodeChunk(MCT_Float3,TEXT("Parameters.LightVector"));
 	}
 
-	virtual int32 ScreenPosition() OVERRIDE
+	virtual int32 ScreenPosition() override
 	{
 		if (ShaderFrequency != SF_Pixel && ShaderFrequency != SF_Compute)
 		{
@@ -1711,17 +1711,17 @@ protected:
 		return AddCodeChunk(MCT_Float2,TEXT("ScreenAlignedPosition(Parameters.ScreenPosition).xy"));		
 	}
 
-	virtual int32 ViewSize() OVERRIDE
+	virtual int32 ViewSize() override
 	{
 		return AddCodeChunk(MCT_Float2,TEXT("View.ViewSizeAndSceneTexelSize.xy"));
 	}
 
-	virtual int32 SceneTexelSize() OVERRIDE
+	virtual int32 SceneTexelSize() override
 	{
 		return AddCodeChunk(MCT_Float2,TEXT("View.ViewSizeAndSceneTexelSize.zw"));
 	}
 
-	virtual int32 ParticleMacroUV() OVERRIDE 
+	virtual int32 ParticleMacroUV() override 
 	{
 		if (ShaderFrequency != SF_Pixel && ShaderFrequency != SF_Compute)
 		{
@@ -1731,7 +1731,7 @@ protected:
 		return AddCodeChunk(MCT_Float2,TEXT("GetParticleMacroUV(Parameters)"));
 	}
 
-	virtual int32 ParticleSubUV(int32 TextureIndex, EMaterialSamplerType SamplerType, bool bBlend) OVERRIDE
+	virtual int32 ParticleSubUV(int32 TextureIndex, EMaterialSamplerType SamplerType, bool bBlend) override
 	{
 		if (ShaderFrequency != SF_Pixel && ShaderFrequency != SF_Compute)
 		{
@@ -1770,7 +1770,7 @@ protected:
 		return ParticleSubUV;
 	}
 
-	virtual int32 ParticleColor() OVERRIDE
+	virtual int32 ParticleColor() override
 	{
 		if (ShaderFrequency != SF_Vertex && ShaderFrequency != SF_Pixel && ShaderFrequency != SF_Compute)
 		{
@@ -1780,7 +1780,7 @@ protected:
 		return AddInlinedCodeChunk(MCT_Float4,TEXT("Parameters.Particle.Color"));	
 	}
 
-	virtual int32 ParticlePosition() OVERRIDE
+	virtual int32 ParticlePosition() override
 	{
 		if (ShaderFrequency != SF_Vertex && ShaderFrequency != SF_Pixel && ShaderFrequency != SF_Compute)
 		{
@@ -1790,7 +1790,7 @@ protected:
 		return AddInlinedCodeChunk(MCT_Float3,TEXT("Parameters.Particle.PositionAndSize.xyz"));	
 	}
 
-	virtual int32 ParticleRadius() OVERRIDE
+	virtual int32 ParticleRadius() override
 	{
 		if (ShaderFrequency != SF_Vertex && ShaderFrequency != SF_Pixel && ShaderFrequency != SF_Compute)
 		{
@@ -1800,7 +1800,7 @@ protected:
 		return AddInlinedCodeChunk(MCT_Float,TEXT("max(Parameters.Particle.PositionAndSize.w, .001f)"));	
 	}
 
-	virtual int32 SphericalParticleOpacity(int32 Density) OVERRIDE
+	virtual int32 SphericalParticleOpacity(int32 Density) override
 	{
 		if (ShaderFrequency != SF_Pixel && ShaderFrequency != SF_Compute)
 		{
@@ -1817,7 +1817,7 @@ protected:
 		return AddCodeChunk(MCT_Float, TEXT("GetSphericalParticleOpacity(Parameters,%s)"), *GetParameterCode(Density));
 	}
 
-	virtual int32 ParticleRelativeTime() OVERRIDE
+	virtual int32 ParticleRelativeTime() override
 	{
 		if (ShaderFrequency != SF_Vertex && ShaderFrequency != SF_Pixel && ShaderFrequency != SF_Compute)
 		{
@@ -1827,7 +1827,7 @@ protected:
 		return AddInlinedCodeChunk(MCT_Float,TEXT("Parameters.Particle.RelativeTime"));
 	}
 
-	virtual int32 ParticleMotionBlurFade() OVERRIDE
+	virtual int32 ParticleMotionBlurFade() override
 	{
 		if (ShaderFrequency != SF_Vertex && ShaderFrequency != SF_Pixel && ShaderFrequency != SF_Compute)
 		{
@@ -1837,7 +1837,7 @@ protected:
 		return AddInlinedCodeChunk(MCT_Float,TEXT("Parameters.Particle.MotionBlurFade"));
 	}
 
-	virtual int32 ParticleDirection() OVERRIDE
+	virtual int32 ParticleDirection() override
 	{
 		if (ShaderFrequency != SF_Vertex && ShaderFrequency != SF_Pixel && ShaderFrequency != SF_Compute)
 		{
@@ -1847,7 +1847,7 @@ protected:
 		return AddInlinedCodeChunk(MCT_Float3,TEXT("Parameters.Particle.Velocity.xyz"));
 	}
 
-	virtual int32 ParticleSpeed() OVERRIDE
+	virtual int32 ParticleSpeed() override
 	{
 		if (ShaderFrequency != SF_Vertex && ShaderFrequency != SF_Pixel && ShaderFrequency != SF_Compute)
 		{
@@ -1857,7 +1857,7 @@ protected:
 		return AddInlinedCodeChunk(MCT_Float,TEXT("Parameters.Particle.Velocity.w"));
 	}
 
-	virtual int32 ParticleSize() OVERRIDE
+	virtual int32 ParticleSize() override
 	{
 		if (ShaderFrequency != SF_Vertex && ShaderFrequency != SF_Pixel && ShaderFrequency != SF_Compute)
 		{
@@ -1867,7 +1867,7 @@ protected:
 		return AddInlinedCodeChunk(MCT_Float2,TEXT("Parameters.Particle.Size"));
 	}
 
-	virtual int32 WorldPosition(EWorldPositionIncludedOffsets WorldPositionIncludedOffsets) OVERRIDE
+	virtual int32 WorldPosition(EWorldPositionIncludedOffsets WorldPositionIncludedOffsets) override
 	{
 		if (ShaderFrequency == SF_Pixel)
 		{
@@ -1915,32 +1915,32 @@ protected:
 		}
 	}
 
-	virtual int32 ObjectWorldPosition() OVERRIDE
+	virtual int32 ObjectWorldPosition() override
 	{
 		return AddInlinedCodeChunk(MCT_Float3,TEXT("GetObjectWorldPosition(Parameters)"));		
 	}
 
-	virtual int32 ObjectRadius() OVERRIDE
+	virtual int32 ObjectRadius() override
 	{
 		return AddInlinedCodeChunk(MCT_Float,TEXT("Primitive.ObjectWorldPositionAndRadius.w"));		
 	}
 
-	virtual int32 ObjectBounds() OVERRIDE
+	virtual int32 ObjectBounds() override
 	{
 		return AddInlinedCodeChunk(MCT_Float3,TEXT("Primitive.ObjectBounds.xyz"));		
 	}
 
-	virtual int32 DistanceCullFade() OVERRIDE
+	virtual int32 DistanceCullFade() override
 	{
 		return AddInlinedCodeChunk(MCT_Float,TEXT("GetDistanceCullFade()"));		
 	}
 
-	virtual int32 ActorWorldPosition() OVERRIDE
+	virtual int32 ActorWorldPosition() override
 	{
 		return AddInlinedCodeChunk(MCT_Float3,TEXT("Primitive.ActorWorldPosition"));		
 	}
 
-	virtual int32 If(int32 A,int32 B,int32 AGreaterThanB,int32 AEqualsB,int32 ALessThanB,int32 ThresholdArg) OVERRIDE
+	virtual int32 If(int32 A,int32 B,int32 AGreaterThanB,int32 AEqualsB,int32 ALessThanB,int32 ThresholdArg) override
 	{
 		if(A == INDEX_NONE || B == INDEX_NONE || AGreaterThanB == INDEX_NONE || ALessThanB == INDEX_NONE || ThresholdArg == INDEX_NONE)
 		{
@@ -1996,7 +1996,7 @@ protected:
 		}
 	}
 
-	virtual int32 TextureCoordinate(uint32 CoordinateIndex, bool UnMirrorU, bool UnMirrorV) OVERRIDE
+	virtual int32 TextureCoordinate(uint32 CoordinateIndex, bool UnMirrorU, bool UnMirrorV) override
 	{
 		const uint32 MaxNumCoordinates = FeatureLevel == ERHIFeatureLevel::ES2 ? 3 : 8;
 
@@ -2041,7 +2041,7 @@ protected:
 				);
 		}
 
-	virtual int32 TextureSample(int32 TextureIndex,int32 CoordinateIndex,EMaterialSamplerType SamplerType,int32 MipValueIndex=INDEX_NONE,ETextureMipValueMode MipValueMode=TMVM_None) OVERRIDE
+	virtual int32 TextureSample(int32 TextureIndex,int32 CoordinateIndex,EMaterialSamplerType SamplerType,int32 MipValueIndex=INDEX_NONE,ETextureMipValueMode MipValueMode=TMVM_None) override
 	{
 		if(TextureIndex == INDEX_NONE || CoordinateIndex == INDEX_NONE)
 		{
@@ -2151,7 +2151,7 @@ protected:
 			);
 	}
 
-	virtual int32 PixelDepth() OVERRIDE
+	virtual int32 PixelDepth() override
 	{
 		if (ShaderFrequency != SF_Pixel && ShaderFrequency != SF_Compute)
 		{
@@ -2175,7 +2175,7 @@ protected:
 		}
 	}
 
-	virtual int32 SceneDepth(int32 Offset, int32 UV, bool bUseOffset) OVERRIDE
+	virtual int32 SceneDepth(int32 Offset, int32 UV, bool bUseOffset) override
 	{
 		if (ShaderFrequency != SF_Pixel && ShaderFrequency != SF_Compute)
 		{
@@ -2200,7 +2200,7 @@ protected:
 	}
 	
 	// @param SceneTextureId of type ESceneTextureId e.g. PPI_SubsurfaceColor
-	virtual int32 SceneTextureLookup(int32 UV, uint32 InSceneTextureId) OVERRIDE
+	virtual int32 SceneTextureLookup(int32 UV, uint32 InSceneTextureId) override
 	{
 		if (ErrorUnlessFeatureLevelSupported(ERHIFeatureLevel::SM3) == INDEX_NONE)
 		{
@@ -2228,7 +2228,7 @@ protected:
 	}
 
 	// @param SceneTextureId of type ESceneTextureId e.g. PPI_SubsurfaceColor
-	virtual int32 SceneTextureSize(uint32 InSceneTextureId, bool bInvert) OVERRIDE
+	virtual int32 SceneTextureSize(uint32 InSceneTextureId, bool bInvert) override
 	{
 		if (ShaderFrequency != SF_Pixel)
 		{
@@ -2268,7 +2268,7 @@ protected:
 	}
 
 	// @param SceneTextureId of type ESceneTextureId e.g. PPI_SubsurfaceColor
-	virtual int32 SceneTextureMin(uint32 InSceneTextureId) OVERRIDE
+	virtual int32 SceneTextureMin(uint32 InSceneTextureId) override
 	{
 		if (ShaderFrequency != SF_Pixel)
 		{
@@ -2292,7 +2292,7 @@ protected:
 		}
 	}
 
-	virtual int32 SceneTextureMax(uint32 InSceneTextureId) OVERRIDE
+	virtual int32 SceneTextureMax(uint32 InSceneTextureId) override
 	{
 		if (ShaderFrequency != SF_Pixel)
 		{
@@ -2349,7 +2349,7 @@ protected:
 		//   PPI_SeparateTranslucency, PPI_CustomDepth, PPI_AmbientOcclusion
 	}
 
-	virtual int32 SceneColor(int32 Offset, int32 UV, bool bUseOffset) OVERRIDE
+	virtual int32 SceneColor(int32 Offset, int32 UV, bool bUseOffset) override
 	{
 		if (Offset == INDEX_NONE && bUseOffset)
 		{
@@ -2376,7 +2376,7 @@ protected:
 			);
 	}
 
-	virtual int32 Texture(UTexture* InTexture) OVERRIDE
+	virtual int32 Texture(UTexture* InTexture) override
 	{
 		if (ShaderFrequency != SF_Pixel
 			&& ErrorUnlessFeatureLevelSupported(ERHIFeatureLevel::SM4) == INDEX_NONE)
@@ -2390,7 +2390,7 @@ protected:
 		return AddUniformExpression(new FMaterialUniformExpressionTexture(TextureReferenceIndex),ShaderType,TEXT(""));
 	}
 
-	virtual int32 TextureParameter(FName ParameterName,UTexture* DefaultValue) OVERRIDE
+	virtual int32 TextureParameter(FName ParameterName,UTexture* DefaultValue) override
 	{
 		if (ShaderFrequency != SF_Pixel
 			&& ErrorUnlessFeatureLevelSupported(ERHIFeatureLevel::SM4) == INDEX_NONE)
@@ -2404,12 +2404,12 @@ protected:
 		return AddUniformExpression(new FMaterialUniformExpressionTextureParameter(ParameterName,TextureReferenceIndex),ShaderType,TEXT(""));
 	}
 
-	virtual int32 StaticBool(bool bValue) OVERRIDE
+	virtual int32 StaticBool(bool bValue) override
 	{
 		return AddInlinedCodeChunk(MCT_StaticBool,(bValue ? TEXT("true") : TEXT("false")));
 	}
 
-	virtual int32 StaticBoolParameter(FName ParameterName,bool bDefaultValue) OVERRIDE
+	virtual int32 StaticBoolParameter(FName ParameterName,bool bDefaultValue) override
 	{
 		// Look up the value we are compiling with for this static parameter.
 		bool bValue = bDefaultValue;
@@ -2426,7 +2426,7 @@ protected:
 		return StaticBool(bValue);
 	}
 	
-	virtual int32 StaticComponentMask(int32 Vector,FName ParameterName,bool bDefaultR,bool bDefaultG,bool bDefaultB,bool bDefaultA) OVERRIDE
+	virtual int32 StaticComponentMask(int32 Vector,FName ParameterName,bool bDefaultR,bool bDefaultG,bool bDefaultB,bool bDefaultA) override
 	{
 		// Look up the value we are compiling with for this static parameter.
 		bool bValueR = bDefaultR;
@@ -2449,7 +2449,7 @@ protected:
 		return ComponentMask(Vector,bValueR,bValueG,bValueB,bValueA);
 	}
 
-	virtual bool GetStaticBoolValue(int32 BoolIndex, bool& bSucceeded) OVERRIDE
+	virtual bool GetStaticBoolValue(int32 BoolIndex, bool& bSucceeded) override
 	{
 		bSucceeded = true;
 		if (BoolIndex == INDEX_NONE)
@@ -2472,7 +2472,7 @@ protected:
 		return false;
 	}
 
-	virtual int32 StaticTerrainLayerWeight(FName ParameterName,int32 Default) OVERRIDE
+	virtual int32 StaticTerrainLayerWeight(FName ParameterName,int32 Default) override
 	{
 		// Look up the weight-map index for this static parameter.
 		int32 WeightmapIndex = INDEX_NONE;
@@ -2512,13 +2512,13 @@ protected:
 		}
 	}
 
-	virtual int32 VertexColor() OVERRIDE
+	virtual int32 VertexColor() override
 	{
 		bUsesVertexColor |= (ShaderFrequency != SF_Vertex);
 		return AddInlinedCodeChunk(MCT_Float4,TEXT("Parameters.VertexColor"));
 	}
 
-	virtual int32 Add(int32 A,int32 B) OVERRIDE
+	virtual int32 Add(int32 A,int32 B) override
 	{
 		if(A == INDEX_NONE || B == INDEX_NONE)
 		{
@@ -2535,7 +2535,7 @@ protected:
 		}
 	}
 
-	virtual int32 Sub(int32 A,int32 B) OVERRIDE
+	virtual int32 Sub(int32 A,int32 B) override
 	{
 		if(A == INDEX_NONE || B == INDEX_NONE)
 		{
@@ -2552,7 +2552,7 @@ protected:
 		}
 	}
 
-	virtual int32 Mul(int32 A,int32 B) OVERRIDE
+	virtual int32 Mul(int32 A,int32 B) override
 	{
 		if(A == INDEX_NONE || B == INDEX_NONE)
 		{
@@ -2569,7 +2569,7 @@ protected:
 		}
 	}
 
-	virtual int32 Div(int32 A,int32 B) OVERRIDE
+	virtual int32 Div(int32 A,int32 B) override
 	{
 		if(A == INDEX_NONE || B == INDEX_NONE)
 		{
@@ -2586,7 +2586,7 @@ protected:
 		}
 	}
 
-	virtual int32 Dot(int32 A,int32 B) OVERRIDE
+	virtual int32 Dot(int32 A,int32 B) override
 	{
 		if(A == INDEX_NONE || B == INDEX_NONE)
 		{
@@ -2638,7 +2638,7 @@ protected:
 		}
 	}
 
-	virtual int32 Cross(int32 A,int32 B) OVERRIDE
+	virtual int32 Cross(int32 A,int32 B) override
 	{
 		if(A == INDEX_NONE || B == INDEX_NONE)
 		{
@@ -2648,7 +2648,7 @@ protected:
 		return AddCodeChunk(MCT_Float3,TEXT("cross(%s,%s)"),*CoerceParameter(A,MCT_Float3),*CoerceParameter(B,MCT_Float3));
 	}
 
-	virtual int32 Power(int32 Base,int32 Exponent) OVERRIDE
+	virtual int32 Power(int32 Base,int32 Exponent) override
 	{
 		if(Base == INDEX_NONE || Exponent == INDEX_NONE)
 		{
@@ -2659,7 +2659,7 @@ protected:
 		return AddCodeChunk(GetParameterType(Base),TEXT("ClampedPow(%s,%s)"),*GetParameterCode(Base),*CoerceParameter(Exponent,MCT_Float));
 	}
 
-	virtual int32 SquareRoot(int32 X) OVERRIDE
+	virtual int32 SquareRoot(int32 X) override
 	{
 		if(X == INDEX_NONE)
 		{
@@ -2676,7 +2676,7 @@ protected:
 		}
 	}
 
-	virtual int32 Length(int32 X) OVERRIDE
+	virtual int32 Length(int32 X) override
 	{
 		if(X == INDEX_NONE)
 		{
@@ -2693,7 +2693,7 @@ protected:
 		}
 	}
 
-	virtual int32 Lerp(int32 X,int32 Y,int32 A) OVERRIDE
+	virtual int32 Lerp(int32 X,int32 Y,int32 A) override
 	{
 		if(X == INDEX_NONE || Y == INDEX_NONE || A == INDEX_NONE)
 		{
@@ -2705,7 +2705,7 @@ protected:
 		return AddCodeChunk(ResultType,TEXT("lerp(%s,%s,%s)"),*CoerceParameter(X,ResultType),*CoerceParameter(Y,ResultType),*CoerceParameter(A,AlphaType));
 	}
 
-	virtual int32 Min(int32 A,int32 B) OVERRIDE
+	virtual int32 Min(int32 A,int32 B) override
 	{
 		if(A == INDEX_NONE || B == INDEX_NONE)
 		{
@@ -2722,7 +2722,7 @@ protected:
 		}
 	}
 
-	virtual int32 Max(int32 A,int32 B) OVERRIDE
+	virtual int32 Max(int32 A,int32 B) override
 	{
 		if(A == INDEX_NONE || B == INDEX_NONE)
 		{
@@ -2739,7 +2739,7 @@ protected:
 		}
 	}
 
-	virtual int32 Clamp(int32 X,int32 A,int32 B) OVERRIDE
+	virtual int32 Clamp(int32 X,int32 A,int32 B) override
 	{
 		if(X == INDEX_NONE || A == INDEX_NONE || B == INDEX_NONE)
 		{
@@ -2756,7 +2756,7 @@ protected:
 		}
 	}
 
-	virtual int32 ComponentMask(int32 Vector,bool R,bool G,bool B,bool A) OVERRIDE
+	virtual int32 ComponentMask(int32 Vector,bool R,bool G,bool B,bool A) override
 	{
 		if(Vector == INDEX_NONE)
 		{
@@ -2796,7 +2796,7 @@ protected:
 			);
 	}
 
-	virtual int32 AppendVector(int32 A,int32 B) OVERRIDE
+	virtual int32 AppendVector(int32 A,int32 B) override
 	{
 		if(A == INDEX_NONE || B == INDEX_NONE)
 		{
@@ -2819,7 +2819,7 @@ protected:
 	/**
 	* Generate shader code for transforming a vector
 	*/
-	virtual int32 TransformVector(uint8 SourceCoordType,uint8 DestCoordType,int32 A) OVERRIDE
+	virtual int32 TransformVector(uint8 SourceCoordType,uint8 DestCoordType,int32 A) override
 	{
 		if (ShaderFrequency != SF_Pixel && ShaderFrequency != SF_Compute && ShaderFrequency != SF_Domain && ShaderFrequency != SF_Vertex)
 		{
@@ -2999,7 +2999,7 @@ protected:
 	* @param	CoordType - type of transform to apply. see EMaterialExpressionTransformPosition 
 	* @param	A - index for input vector parameter's code
 	*/
-	virtual int32 TransformPosition(uint8 SourceCoordType, uint8 DestCoordType, int32 A) OVERRIDE
+	virtual int32 TransformPosition(uint8 SourceCoordType, uint8 DestCoordType, int32 A) override
 	{
 		const EMaterialPositionTransformSource SourceCoordinateSpace = (EMaterialPositionTransformSource)SourceCoordType;
 		const EMaterialPositionTransformSource DestinationCoordinateSpace = (EMaterialPositionTransformSource)DestCoordType;
@@ -3033,7 +3033,7 @@ protected:
 		return Result; 
 	}
 
-	virtual int32 DynamicParameter() OVERRIDE
+	virtual int32 DynamicParameter() override
 	{
 		if (ShaderFrequency != SF_Vertex && ShaderFrequency != SF_Pixel && ShaderFrequency != SF_Compute)
 		{
@@ -3048,7 +3048,7 @@ protected:
 			);
 	}
 
-	virtual int32 LightmapUVs() OVERRIDE
+	virtual int32 LightmapUVs() override
 	{
 		if (ShaderFrequency != SF_Pixel && ShaderFrequency != SF_Compute)
 		{
@@ -3071,9 +3071,9 @@ protected:
 		return ResultIdx;
 	}
 
-	virtual int32 LightmassReplace(int32 Realtime, int32 Lightmass) OVERRIDE { return Realtime; }
+	virtual int32 LightmassReplace(int32 Realtime, int32 Lightmass) override { return Realtime; }
 
-	virtual int32 GIReplace(int32 Direct, int32 StaticIndirect, int32 DynamicIndirect) OVERRIDE 
+	virtual int32 GIReplace(int32 Direct, int32 StaticIndirect, int32 DynamicIndirect) override 
 	{ 
 		if(Direct == INDEX_NONE || DynamicIndirect == INDEX_NONE)
 		{
@@ -3085,12 +3085,12 @@ protected:
 		return AddCodeChunk(ResultType,TEXT("(GetGIReplaceState() ? (%s) : (%s))"), *GetParameterCode(DynamicIndirect), *GetParameterCode(Direct));
 	}
 
-	virtual int32 ObjectOrientation() OVERRIDE
+	virtual int32 ObjectOrientation() override
 	{ 
 		return AddInlinedCodeChunk(MCT_Float3,TEXT("Primitive.ObjectOrientation.xyz"));	
 	}
 
-	virtual int32 RotateAboutAxis(int32 NormalizedRotationAxisAndAngleIndex, int32 PositionOnAxisIndex, int32 PositionIndex) OVERRIDE
+	virtual int32 RotateAboutAxis(int32 NormalizedRotationAxisAndAngleIndex, int32 PositionOnAxisIndex, int32 PositionIndex) override
 	{
 		if (NormalizedRotationAxisAndAngleIndex == INDEX_NONE
 			|| PositionOnAxisIndex == INDEX_NONE
@@ -3110,7 +3110,7 @@ protected:
 		}
 	}
 
-	virtual int32 TwoSidedSign() OVERRIDE
+	virtual int32 TwoSidedSign() override
 	{
 		if (ShaderFrequency != SF_Pixel && ShaderFrequency != SF_Compute)
 		{
@@ -3119,7 +3119,7 @@ protected:
 		return AddInlinedCodeChunk(MCT_Float,TEXT("Parameters.TwoSidedSign"));	
 	}
 
-	virtual int32 VertexNormal() OVERRIDE
+	virtual int32 VertexNormal() override
 	{
 		if (ShaderFrequency != SF_Vertex)
 		{
@@ -3128,7 +3128,7 @@ protected:
 		return AddInlinedCodeChunk(MCT_Float3,TEXT("Parameters.TangentToWorld[2]"));	
 	}
 
-	virtual int32 PixelNormalWS() OVERRIDE
+	virtual int32 PixelNormalWS() override
 	{
 		if (ShaderFrequency != SF_Pixel && ShaderFrequency != SF_Compute)
 		{
@@ -3145,7 +3145,7 @@ protected:
 		return AddInlinedCodeChunk(MCT_Float3,TEXT("Parameters.WorldNormal"));	
 	}
 
-	virtual int32 DDX( int32 X ) OVERRIDE
+	virtual int32 DDX( int32 X ) override
 	{
 		if (ErrorUnlessFeatureLevelSupported(ERHIFeatureLevel::SM3) == INDEX_NONE)
 		{
@@ -3171,7 +3171,7 @@ protected:
 		return AddCodeChunk(GetParameterType(X),TEXT("ddx(%s)"),*GetParameterCode(X));
 	}
 
-	virtual int32 DDY( int32 X ) OVERRIDE
+	virtual int32 DDY( int32 X ) override
 	{
 		if (ErrorUnlessFeatureLevelSupported(ERHIFeatureLevel::SM3) == INDEX_NONE)
 		{
@@ -3196,7 +3196,7 @@ protected:
 		return AddCodeChunk(GetParameterType(X),TEXT("ddy(%s)"),*GetParameterCode(X));
 	}
 
-	virtual int32 AntialiasedTextureMask(int32 Tex, int32 UV, float Threshold, uint8 Channel) OVERRIDE
+	virtual int32 AntialiasedTextureMask(int32 Tex, int32 UV, float Threshold, uint8 Channel) override
 	{
 		if (ErrorUnlessFeatureLevelSupported(ERHIFeatureLevel::SM3) == INDEX_NONE)
 		{
@@ -3221,7 +3221,7 @@ protected:
 			*GetParameterCode(ChannelConst));
 	}
 
-	virtual int32 DepthOfFieldFunction(int32 Depth, int32 FunctionValueIndex) OVERRIDE
+	virtual int32 DepthOfFieldFunction(int32 Depth, int32 FunctionValueIndex) override
 	{
 		if (ShaderFrequency == SF_Hull)
 		{
@@ -3238,7 +3238,7 @@ protected:
 			*GetParameterCode(Depth), FunctionValueIndex);
 	}
 
-	virtual int32 Noise(int32 Position, float Scale, int32 Quality, uint8 NoiseFunction, bool bTurbulence, int32 Levels, float OutputMin, float OutputMax, float LevelScale, int32 FilterWidth) OVERRIDE
+	virtual int32 Noise(int32 Position, float Scale, int32 Quality, uint8 NoiseFunction, bool bTurbulence, int32 Levels, float OutputMin, float OutputMax, float LevelScale, int32 FilterWidth) override
 	{
 		if (ErrorUnlessFeatureLevelSupported(ERHIFeatureLevel::SM3) == INDEX_NONE)
 		{
@@ -3276,7 +3276,7 @@ protected:
 			*GetParameterCode(FilterWidth));
 	}
 
-	virtual int32 BlackBody( int32 Temp ) OVERRIDE
+	virtual int32 BlackBody( int32 Temp ) override
 	{
 		if (ErrorUnlessFeatureLevelSupported(ERHIFeatureLevel::SM3) == INDEX_NONE)
 		{
@@ -3291,7 +3291,7 @@ protected:
 		return AddCodeChunk( MCT_Float3, TEXT("MaterialExpressionBlackBody(%s)"), *GetParameterCode(Temp) );
 	}
 
-	virtual int32 AtmosphericFogColor( int32 WorldPosition ) OVERRIDE
+	virtual int32 AtmosphericFogColor( int32 WorldPosition ) override
 	{
 		if (ErrorUnlessFeatureLevelSupported(ERHIFeatureLevel::SM3) == INDEX_NONE)
 		{
@@ -3309,7 +3309,7 @@ protected:
 		}
 	}
 
-	virtual int32 CustomExpression( class UMaterialExpressionCustom* Custom, TArray<int32>& CompiledInputs ) OVERRIDE
+	virtual int32 CustomExpression( class UMaterialExpressionCustom* Custom, TArray<int32>& CompiledInputs ) override
 	{
 		int32 ResultIdx = INDEX_NONE;
 
@@ -3406,7 +3406,7 @@ protected:
 	 *
 	 * @return	Code index
 	 */
-	virtual int32 PerInstanceRandom() OVERRIDE
+	virtual int32 PerInstanceRandom() override
 	{
 		if (ErrorUnlessFeatureLevelSupported(ERHIFeatureLevel::SM3) == INDEX_NONE)
 		{
@@ -3428,7 +3428,7 @@ protected:
 	 *
 	 * @return	Code index
 	 */
-	virtual int32 PerInstanceFadeAmount() OVERRIDE
+	virtual int32 PerInstanceFadeAmount() override
 	{
 		if (ErrorUnlessFeatureLevelSupported(ERHIFeatureLevel::SM3) == INDEX_NONE)
 		{
@@ -3450,7 +3450,7 @@ protected:
 	*
 	* @return	Code index
 	*/
-	virtual int32 SpeedTree(ESpeedTreeGeometryType GeometryType, ESpeedTreeWindType WindType, ESpeedTreeLODType LODType, float BillboardThreshold) OVERRIDE 
+	virtual int32 SpeedTree(ESpeedTreeGeometryType GeometryType, ESpeedTreeWindType WindType, ESpeedTreeLODType LODType, float BillboardThreshold) override 
 	{ 
 		if (ErrorUnlessFeatureLevelSupported(ERHIFeatureLevel::SM3) == INDEX_NONE)
 		{
@@ -3475,7 +3475,7 @@ protected:
 	 *
 	 * @return	Code index
 	 */
-	virtual int32 TextureCoordinateOffset() OVERRIDE
+	virtual int32 TextureCoordinateOffset() override
 	{
 		if (FeatureLevel == ERHIFeatureLevel::ES2 && ShaderFrequency == SF_Vertex)
 		{
@@ -3488,7 +3488,7 @@ protected:
 	}
 
 	/**Experimental access to the EyeAdaptation RT for Post Process materials. Can be one frame behind depending on the value of BlendableLocation. */
-	virtual int32 EyeAdaptation() OVERRIDE
+	virtual int32 EyeAdaptation() override
 	{
 		if (ErrorUnlessFeatureLevelSupported(ERHIFeatureLevel::SM5) == INDEX_NONE)
 		{
diff --git a/Engine/Source/Runtime/Engine/Private/Materials/MaterialInstanceSupport.h b/Engine/Source/Runtime/Engine/Private/Materials/MaterialInstanceSupport.h
--- a/Engine/Source/Runtime/Engine/Private/Materials/MaterialInstanceSupport.h
+++ b/Engine/Source/Runtime/Engine/Private/Materials/MaterialInstanceSupport.h
@@ -80,17 +80,17 @@ public:
 	}
 
 	// FRenderResource interface.
-	virtual FString GetFriendlyName() const OVERRIDE { return Owner->GetName(); }
+	virtual FString GetFriendlyName() const override { return Owner->GetName(); }
 
 	// FMaterialRenderProxy interface.
 	/** Get the FMaterial to use for rendering.  Must return a valid FMaterial, even if it had to fall back to the default material. */
-	virtual const FMaterial* GetMaterial(ERHIFeatureLevel::Type FeatureLevel) const OVERRIDE;
+	virtual const FMaterial* GetMaterial(ERHIFeatureLevel::Type FeatureLevel) const override;
 	/** Get the FMaterial that should be used for rendering, but might not be in a valid state to actually use.  Can return NULL. */
-	virtual FMaterial* GetMaterialNoFallback(ERHIFeatureLevel::Type FeatureLevel) const OVERRIDE;
-	virtual bool GetVectorValue(const FName ParameterName, FLinearColor* OutValue, const FMaterialRenderContext& Context) const OVERRIDE;
-	virtual bool GetScalarValue(const FName ParameterName,float* OutValue, const FMaterialRenderContext& Context) const OVERRIDE;
-	virtual bool GetTextureValue(const FName ParameterName,const UTexture** OutValue, const FMaterialRenderContext& Context) const OVERRIDE;
-	virtual float GetDistanceFieldPenumbraScale() const OVERRIDE { return DistanceFieldPenumbraScale; }
+	virtual FMaterial* GetMaterialNoFallback(ERHIFeatureLevel::Type FeatureLevel) const override;
+	virtual bool GetVectorValue(const FName ParameterName, FLinearColor* OutValue, const FMaterialRenderContext& Context) const override;
+	virtual bool GetScalarValue(const FName ParameterName,float* OutValue, const FMaterialRenderContext& Context) const override;
+	virtual bool GetTextureValue(const FName ParameterName,const UTexture** OutValue, const FMaterialRenderContext& Context) const override;
+	virtual float GetDistanceFieldPenumbraScale() const override { return DistanceFieldPenumbraScale; }
 
 	/** Sets the material instance's parent. */
 	void GameThread_SetParent(UMaterialInterface* InParent);
diff --git a/Engine/Source/Runtime/Engine/Private/MatineeUtils.cpp b/Engine/Source/Runtime/Engine/Private/MatineeUtils.cpp
--- a/Engine/Source/Runtime/Engine/Private/MatineeUtils.cpp
+++ b/Engine/Source/Runtime/Engine/Private/MatineeUtils.cpp
@@ -91,7 +91,7 @@ public:
 	}
 
 protected:
-	virtual bool IsDesiredProperty(UProperty* Property) const OVERRIDE
+	virtual bool IsDesiredProperty(UProperty* Property) const override
 	{
 		return Cast<BasicPropertyType>(Property) != NULL;
 	}
@@ -113,7 +113,7 @@ public:
 protected:
 	FName DesiredStructName;
 protected:
-	virtual bool IsDesiredProperty(UProperty* Property) const OVERRIDE
+	virtual bool IsDesiredProperty(UProperty* Property) const override
 	{
 		if (UStructProperty* StructProperty = Cast<UStructProperty>(Property))
 		{
diff --git a/Engine/Source/Runtime/Engine/Private/MeshParticleVertexFactory.cpp b/Engine/Source/Runtime/Engine/Private/MeshParticleVertexFactory.cpp
--- a/Engine/Source/Runtime/Engine/Private/MeshParticleVertexFactory.cpp
+++ b/Engine/Source/Runtime/Engine/Private/MeshParticleVertexFactory.cpp
@@ -14,7 +14,7 @@ class FMeshParticleVertexFactoryShaderParameters : public FVertexFactoryShaderPa
 {
 public:
 
-	virtual void Bind(const FShaderParameterMap& ParameterMap) OVERRIDE
+	virtual void Bind(const FShaderParameterMap& ParameterMap) override
 	{
 		Transform1.Bind(ParameterMap,TEXT("Transform1"));
 		Transform2.Bind(ParameterMap,TEXT("Transform2"));
@@ -25,7 +25,7 @@ public:
 		ParticleColor.Bind(ParameterMap,TEXT("ParticleColor"));
 	}
 
-	virtual void Serialize(FArchive& Ar) OVERRIDE
+	virtual void Serialize(FArchive& Ar) override
 	{
 		Ar << Transform1;
 		Ar << Transform2;
@@ -36,7 +36,7 @@ public:
 		Ar << ParticleColor;
 	}
 
-	virtual void SetMesh(FRHICommandList& RHICmdList, FShader* Shader,const FVertexFactory* VertexFactory,const FSceneView& View,const FMeshBatchElement& BatchElement,uint32 DataFlags) const OVERRIDE
+	virtual void SetMesh(FRHICommandList& RHICmdList, FShader* Shader,const FVertexFactory* VertexFactory,const FSceneView& View,const FMeshBatchElement& BatchElement,uint32 DataFlags) const override
 	{
 		const bool bInstanced = View.GetFeatureLevel() >= ERHIFeatureLevel::SM3;
 		FMeshParticleVertexFactory* MeshParticleVF = (FMeshParticleVertexFactory*)VertexFactory;
diff --git a/Engine/Source/Runtime/Engine/Private/ModelRender.cpp b/Engine/Source/Runtime/Engine/Private/ModelRender.cpp
--- a/Engine/Source/Runtime/Engine/Private/ModelRender.cpp
+++ b/Engine/Source/Runtime/Engine/Private/ModelRender.cpp
@@ -203,7 +203,7 @@ public:
 		return ModelHitProxy;
 	}
 
-	virtual void PreRenderView(const FSceneViewFamily* ViewFamily, const uint32 VisibilityMap, int32 FrameNumber) OVERRIDE
+	virtual void PreRenderView(const FSceneViewFamily* ViewFamily, const uint32 VisibilityMap, int32 FrameNumber) override
 	{
 		// Reset any batches leftover from last frame.
 		PerViewBatches.Reset();
@@ -482,7 +482,7 @@ public:
 		return Result;
 	}
 
-	virtual bool CanBeOccluded() const OVERRIDE
+	virtual bool CanBeOccluded() const override
 	{
 		return !MaterialRelevance.bDisableDepthTest;
 	}
diff --git a/Engine/Source/Runtime/Engine/Private/NavMeshRenderingHelpers.h b/Engine/Source/Runtime/Engine/Private/NavMeshRenderingHelpers.h
--- a/Engine/Source/Runtime/Engine/Private/NavMeshRenderingHelpers.h
+++ b/Engine/Source/Runtime/Engine/Private/NavMeshRenderingHelpers.h
@@ -193,7 +193,7 @@ public:
 		return FColor(r*63, g*63, b*63, 164);
 	}
 
-	virtual void DrawDynamicElements(FPrimitiveDrawInterface* PDI,const FSceneView* View) OVERRIDE
+	virtual void DrawDynamicElements(FPrimitiveDrawInterface* PDI,const FSceneView* View) override
 	{
 		QUICK_SCOPE_CYCLE_COUNTER( STAT_RecastRenderingSceneProxy_DrawDynamicElements );
 
diff --git a/Engine/Source/Runtime/Engine/Private/NiagaraComponent.cpp b/Engine/Source/Runtime/Engine/Private/NiagaraComponent.cpp
--- a/Engine/Source/Runtime/Engine/Private/NiagaraComponent.cpp
+++ b/Engine/Source/Runtime/Engine/Private/NiagaraComponent.cpp
@@ -70,7 +70,7 @@ private:
 	}
 
 	// FPrimitiveSceneProxy interface.
-	virtual void CreateRenderThreadResources() OVERRIDE
+	virtual void CreateRenderThreadResources() override
 	{
 		VertexFactory.InitResource();
 
@@ -86,17 +86,17 @@ private:
 		UniformParameters.NormalsCylinderUnitDirection = FVector4(0.0f, 0.0f, 1.0f, 0.0f);
 	}
 
-	virtual void OnActorPositionChanged() OVERRIDE
+	virtual void OnActorPositionChanged() override
 	{
 		WorldSpacePrimitiveUniformBuffer.ReleaseResource();
 	}
 
-	virtual void OnTransformChanged() OVERRIDE
+	virtual void OnTransformChanged() override
 	{
 		WorldSpacePrimitiveUniformBuffer.ReleaseResource();
 	}
 
-	virtual void PreRenderView(const FSceneViewFamily* ViewFamily, const uint32 VisibilityMap, int32 FrameNumber) OVERRIDE
+	virtual void PreRenderView(const FSceneViewFamily* ViewFamily, const uint32 VisibilityMap, int32 FrameNumber) override
 	{
 		SCOPE_CYCLE_COUNTER(STAT_NiagaraPreRenderView);
 
@@ -147,7 +147,7 @@ private:
 		}
 	}
 		  
-	virtual void DrawDynamicElements(FPrimitiveDrawInterface* PDI,const FSceneView* View) OVERRIDE
+	virtual void DrawDynamicElements(FPrimitiveDrawInterface* PDI,const FSceneView* View) override
 	{
 		SCOPE_CYCLE_COUNTER(STAT_NiagaraRender);
 
@@ -209,7 +209,7 @@ private:
 		*/
 	}
 
-	virtual FPrimitiveViewRelevance GetViewRelevance(const FSceneView* View)  OVERRIDE
+	virtual FPrimitiveViewRelevance GetViewRelevance(const FSceneView* View)  override
 	{
 		FPrimitiveViewRelevance Result;
 		bool bHasDynamicData = DynamicData && DynamicData->VertexData.Num() > 0;
@@ -227,12 +227,12 @@ private:
 		return Result;
 	}
 
-	virtual bool CanBeOccluded() const OVERRIDE
+	virtual bool CanBeOccluded() const override
 	{
 		return !MaterialRelevance.bDisableDepthTest;
 	}
 
-	virtual uint32 GetMemoryFootprint() const OVERRIDE 
+	virtual uint32 GetMemoryFootprint() const override 
 	{ 
 		return (sizeof(*this) + GetAllocatedSize()); 
 	}
diff --git a/Engine/Source/Runtime/Engine/Private/Particles/FXSystemPrivate.h b/Engine/Source/Runtime/Engine/Private/Particles/FXSystemPrivate.h
--- a/Engine/Source/Runtime/Engine/Private/Particles/FXSystemPrivate.h
+++ b/Engine/Source/Runtime/Engine/Private/Particles/FXSystemPrivate.h
@@ -54,19 +54,19 @@ public:
 	virtual ~FFXSystem();
 
 	// Begin FFXSystemInterface.
-	virtual void Tick(float DeltaSeconds) OVERRIDE;
+	virtual void Tick(float DeltaSeconds) override;
 #if WITH_EDITOR
-	virtual void Suspend() OVERRIDE;
-	virtual void Resume() OVERRIDE;
+	virtual void Suspend() override;
+	virtual void Resume() override;
 #endif // #if WITH_EDITOR
-	virtual void DrawDebug(FCanvas* Canvas) OVERRIDE;
-	virtual void AddVectorField(UVectorFieldComponent* VectorFieldComponent) OVERRIDE;
-	virtual void RemoveVectorField(UVectorFieldComponent* VectorFieldComponent) OVERRIDE;
-	virtual void UpdateVectorField(UVectorFieldComponent* VectorFieldComponent) OVERRIDE;
-	virtual FParticleEmitterInstance* CreateGPUSpriteEmitterInstance(FGPUSpriteEmitterInfo& EmitterInfo) OVERRIDE;
-	virtual void PreInitViews() OVERRIDE;
-	virtual void PreRender() OVERRIDE;
-	virtual void PostRenderOpaque(const class FSceneView* CollisionView, FTexture2DRHIParamRef SceneDepthTexture, FTexture2DRHIParamRef GBufferATexture) OVERRIDE;
+	virtual void DrawDebug(FCanvas* Canvas) override;
+	virtual void AddVectorField(UVectorFieldComponent* VectorFieldComponent) override;
+	virtual void RemoveVectorField(UVectorFieldComponent* VectorFieldComponent) override;
+	virtual void UpdateVectorField(UVectorFieldComponent* VectorFieldComponent) override;
+	virtual FParticleEmitterInstance* CreateGPUSpriteEmitterInstance(FGPUSpriteEmitterInfo& EmitterInfo) override;
+	virtual void PreInitViews() override;
+	virtual void PreRender() override;
+	virtual void PostRenderOpaque(const class FSceneView* CollisionView, FTexture2DRHIParamRef SceneDepthTexture, FTexture2DRHIParamRef GBufferATexture) override;
 	// End FFXSystemInterface.
 
 	/*--------------------------------------------------------------------------
diff --git a/Engine/Source/Runtime/Engine/Private/Particles/ParticleCurveTexture.cpp b/Engine/Source/Runtime/Engine/Private/Particles/ParticleCurveTexture.cpp
--- a/Engine/Source/Runtime/Engine/Private/Particles/ParticleCurveTexture.cpp
+++ b/Engine/Source/Runtime/Engine/Private/Particles/ParticleCurveTexture.cpp
@@ -113,7 +113,7 @@ public:
 	}
 
 	/** Serialization. */
-	virtual bool Serialize( FArchive& Ar ) OVERRIDE
+	virtual bool Serialize( FArchive& Ar ) override
 	{
 		bool bShaderHasOutdatedParameters = FGlobalShader::Serialize( Ar );
 		return bShaderHasOutdatedParameters;
diff --git a/Engine/Source/Runtime/Engine/Private/Particles/ParticleCurveTexture.h b/Engine/Source/Runtime/Engine/Private/Particles/ParticleCurveTexture.h
--- a/Engine/Source/Runtime/Engine/Private/Particles/ParticleCurveTexture.h
+++ b/Engine/Source/Runtime/Engine/Private/Particles/ParticleCurveTexture.h
@@ -126,12 +126,12 @@ public:
 	/**
 	 * Initialize RHI resources for the curve texture.
 	 */
-	virtual void InitRHI() OVERRIDE;
+	virtual void InitRHI() override;
 
 	/**
 	 * Releases RHI resources.
 	 */
-	virtual void ReleaseRHI() OVERRIDE;
+	virtual void ReleaseRHI() override;
 
 	/**
 	 * Adds a curve to the texture.
diff --git a/Engine/Source/Runtime/Engine/Private/Particles/ParticleGpuSimulation.cpp b/Engine/Source/Runtime/Engine/Private/Particles/ParticleGpuSimulation.cpp
--- a/Engine/Source/Runtime/Engine/Private/Particles/ParticleGpuSimulation.cpp
+++ b/Engine/Source/Runtime/Engine/Private/Particles/ParticleGpuSimulation.cpp
@@ -171,7 +171,7 @@ public:
 	/**
 	 * Initialize RHI resources used for particle simulation.
 	 */
-	virtual void InitRHI() OVERRIDE
+	virtual void InitRHI() override
 	{
 		const int32 SizeX = GParticleSimulationTextureSizeX;
 		const int32 SizeY = GParticleSimulationTextureSizeY;
@@ -217,7 +217,7 @@ public:
 	/**
 	 * Releases RHI resources used for particle simulation.
 	 */
-	virtual void ReleaseRHI() OVERRIDE
+	virtual void ReleaseRHI() override
 	{
 		// Release textures.
 		PositionTextureTargetRHI.SafeRelease();
@@ -243,7 +243,7 @@ public:
 	/**
 	 * Initialize RHI resources used for particle simulation.
 	 */
-	virtual void InitRHI() OVERRIDE
+	virtual void InitRHI() override
 	{
 		const int32 SizeX = GParticleSimulationTextureSizeX;
 		const int32 SizeY = GParticleSimulationTextureSizeY;
@@ -266,7 +266,7 @@ public:
 	/**
 	 * Releases RHI resources used for particle simulation.
 	 */
-	virtual void ReleaseRHI() OVERRIDE
+	virtual void ReleaseRHI() override
 	{
 		TextureTargetRHI.SafeRelease();
 		TextureRHI.SafeRelease();
@@ -284,7 +284,7 @@ public:
 	FShaderResourceViewRHIRef VertexBufferSRV;
 
 	/** Release RHI resources. */
-	virtual void ReleaseRHI() OVERRIDE
+	virtual void ReleaseRHI() override
 	{
 		VertexBufferSRV.SafeRelease();
 		FVertexBuffer::ReleaseRHI();
@@ -465,7 +465,7 @@ typedef TUniformBufferRef<FGPUSpriteEmitterDynamicUniformParameters> FGPUSpriteE
 class FGPUSpriteVertexFactoryShaderParameters : public FVertexFactoryShaderParameters
 {
 public:
-	virtual void Bind( const FShaderParameterMap& ParameterMap ) OVERRIDE
+	virtual void Bind( const FShaderParameterMap& ParameterMap ) override
 	{
 		ParticleIndices.Bind(ParameterMap, TEXT("ParticleIndices"));
 		ParticleIndicesOffset.Bind(ParameterMap, TEXT("ParticleIndicesOffset"));
@@ -481,7 +481,7 @@ public:
 		CurveTextureSampler.Bind(ParameterMap, TEXT("CurveTextureSampler"));
 	}
 
-	virtual void Serialize(FArchive& Ar) OVERRIDE
+	virtual void Serialize(FArchive& Ar) override
 	{
 		Ar << ParticleIndices;
 		Ar << ParticleIndicesOffset;
@@ -497,7 +497,7 @@ public:
 		Ar << CurveTextureSampler;
 	}
 
-	virtual void SetMesh(FRHICommandList& RHICmdList, FShader* Shader,const FVertexFactory* VertexFactory,const FSceneView& View,const FMeshBatchElement& BatchElement,uint32 DataFlags) const OVERRIDE;
+	virtual void SetMesh(FRHICommandList& RHICmdList, FShader* Shader,const FVertexFactory* VertexFactory,const FSceneView& View,const FMeshBatchElement& BatchElement,uint32 DataFlags) const override;
 
 	virtual uint32 GetSize() const { return sizeof(*this); }
 
@@ -536,7 +536,7 @@ public:
 	/**
 	 * Initialize RHI resources.
 	 */
-	virtual void InitRHI() OVERRIDE
+	virtual void InitRHI() override
 	{
 		FVertexDeclarationElementList Elements;
 
@@ -549,7 +549,7 @@ public:
 	/**
 	 * Release RHI resources.
 	 */
-	virtual void ReleaseRHI() OVERRIDE
+	virtual void ReleaseRHI() override
 	{
 		VertexDeclarationRHI.SafeRelease();
 	}
@@ -586,7 +586,7 @@ public:
 	/**
 	 * Constructs render resources for this vertex factory.
 	 */
-	virtual void InitRHI() OVERRIDE
+	virtual void InitRHI() override
 	{
 		FVertexStream Stream;
 
@@ -797,7 +797,7 @@ public:
 	}
 
 	/** Serialization. */
-	virtual bool Serialize( FArchive& Ar ) OVERRIDE
+	virtual bool Serialize( FArchive& Ar ) override
 	{
 		bool bShaderHasOutdatedParameters = FGlobalShader::Serialize( Ar );
 		Ar << TileOffsets;
@@ -878,7 +878,7 @@ public:
 	}
 
 	/** Serialization. */
-	virtual bool Serialize(FArchive& Ar) OVERRIDE
+	virtual bool Serialize(FArchive& Ar) override
 	{
 		bool bShaderHasOutdatedParameters = FGlobalShader::Serialize( Ar );
 		Ar << PositionTexture;
@@ -1067,7 +1067,7 @@ public:
 	}
 
 	/** Serialization. */
-	virtual bool Serialize( FArchive& Ar ) OVERRIDE
+	virtual bool Serialize( FArchive& Ar ) override
 	{
 		bool bShaderHasOutdatedParameters = FGlobalShader::Serialize( Ar );
 		return bShaderHasOutdatedParameters;
@@ -1090,7 +1090,7 @@ public:
 	/** The vertex declaration. */
 	FVertexDeclarationRHIRef VertexDeclarationRHI;
 
-	virtual void InitRHI() OVERRIDE
+	virtual void InitRHI() override
 	{
 		FVertexDeclarationElementList Elements;
 		// TexCoord.
@@ -1098,7 +1098,7 @@ public:
 		VertexDeclarationRHI = RHICreateVertexDeclaration( Elements );
 	}
 
-	virtual void ReleaseRHI() OVERRIDE
+	virtual void ReleaseRHI() override
 	{
 		VertexDeclarationRHI.SafeRelease();
 	}
@@ -1437,7 +1437,7 @@ public:
 	}
 
 	/** Serialization. */
-	virtual bool Serialize( FArchive& Ar ) OVERRIDE
+	virtual bool Serialize( FArchive& Ar ) override
 	{
 		bool bShaderHasOutdatedParameters = FGlobalShader::Serialize( Ar );
 		return bShaderHasOutdatedParameters;
@@ -1659,7 +1659,7 @@ public:
 	}
 
 	/** Serialization. */
-	virtual bool Serialize( FArchive& Ar ) OVERRIDE
+	virtual bool Serialize( FArchive& Ar ) override
 	{
 		bool bShaderHasOutdatedParameters = FGlobalShader::Serialize( Ar );
 		Ar << VisualizationMode;
@@ -1911,7 +1911,7 @@ public:
 	}
 
 	/** Serialization. */
-	virtual bool Serialize( FArchive& Ar ) OVERRIDE
+	virtual bool Serialize( FArchive& Ar ) override
 	{
 		bool bShaderHasOutdatedParameters = FGlobalShader::Serialize( Ar );
 		Ar << InParticleIndices;
@@ -2171,7 +2171,7 @@ public:
 	/**
 	 * Initialize RHI resources.
 	 */
-	virtual void InitRHI() OVERRIDE
+	virtual void InitRHI() override
 	{
 		if ( AlignedTileCount > 0 )
 		{
@@ -2186,7 +2186,7 @@ public:
 	/**
 	 * Release RHI resources.
 	 */
-	virtual void ReleaseRHI() OVERRIDE
+	virtual void ReleaseRHI() override
 	{
 		TileCount = 0;
 		AlignedTileCount = 0;
@@ -2226,7 +2226,7 @@ public:
 	}
 
 	/** Initialize RHI resources. */
-	virtual void InitRHI() OVERRIDE
+	virtual void InitRHI() override
 	{
 		if ( ParticleCount > 0 && CurrentRHISupportsGPUParticles() )
 		{
@@ -2584,7 +2584,7 @@ struct FGPUSpriteDynamicEmitterData : FDynamicEmitterDataBase
 	 * Called once the emitter has been determined to be visible. The emitter
 	 * will generate per-view render data.
 	 */
-	virtual void PreRenderView(FParticleSystemSceneProxy* Proxy, const FSceneViewFamily* ViewFamily, const uint32 VisibilityMap, int32 FrameNumber) OVERRIDE
+	virtual void PreRenderView(FParticleSystemSceneProxy* Proxy, const FSceneViewFamily* ViewFamily, const uint32 VisibilityMap, int32 FrameNumber) override
 	{
 		auto FeatureLevel = ViewFamily->Scene->GetFeatureLevel();
 
@@ -3010,7 +3010,7 @@ public:
 	/**
 	 * Initializes parameters for this emitter instance.
 	 */
-	virtual void InitParameters(UParticleEmitter* InTemplate, UParticleSystemComponent* InComponent, bool bClearResources) OVERRIDE
+	virtual void InitParameters(UParticleEmitter* InTemplate, UParticleSystemComponent* InComponent, bool bClearResources) override
 	{
 		FParticleEmitterInstance::InitParameters( InTemplate, InComponent, bClearResources );
 		SetupEmitterDuration();
@@ -3236,7 +3236,7 @@ public:
 	 *	Force kill all particles in the emitter.
 	 *	@param	bFireEvents		If true, fire events for the particles being killed.
 	 */
-	virtual void KillParticlesForced(bool bFireEvents) OVERRIDE
+	virtual void KillParticlesForced(bool bFireEvents) override
 	{
 		// Clear all active tiles. This will effectively kill all particles.
 		ClearAllocatedTiles();
@@ -3245,11 +3245,11 @@ public:
 	/**
 	 *	Called when the particle system is deactivating...
 	 */
-	virtual void OnDeactivateSystem() OVERRIDE
+	virtual void OnDeactivateSystem() override
 	{
 	}
 
-	virtual void Rewind() OVERRIDE
+	virtual void Rewind() override
 	{
 		FParticleEmitterInstance::Rewind();
 		InitLocalVectorField();
diff --git a/Engine/Source/Runtime/Engine/Private/Particles/ParticleResources.h b/Engine/Source/Runtime/Engine/Private/Particles/ParticleResources.h
--- a/Engine/Source/Runtime/Engine/Private/Particles/ParticleResources.h
+++ b/Engine/Source/Runtime/Engine/Private/Particles/ParticleResources.h
@@ -19,7 +19,7 @@ extern const int32 GParticleScratchVertexBufferSize;
 class FParticleTexCoordVertexBuffer : public FVertexBuffer
 {
 public:
-	virtual void InitRHI() OVERRIDE;
+	virtual void InitRHI() override;
 };
 
 /** Global particle texture coordinate vertex buffer. */
@@ -31,7 +31,7 @@ extern TGlobalResource<FParticleTexCoordVertexBuffer> GParticleTexCoordVertexBuf
 class FParticleIndexBuffer : public FIndexBuffer
 {
 public:
-	virtual void InitRHI() OVERRIDE;
+	virtual void InitRHI() override;
 };
 
 /** Global particle index buffer. */
@@ -52,8 +52,8 @@ public:
 	FParticleShaderParamRef GetShaderParam();
 	FParticleBufferParamRef GetBufferParam();
 
-	virtual void InitRHI() OVERRIDE;
-	virtual void ReleaseRHI() OVERRIDE;
+	virtual void InitRHI() override;
+	virtual void ReleaseRHI() override;
 };
 
 /** The global scratch vertex buffer. */
diff --git a/Engine/Source/Runtime/Engine/Private/Particles/ParticleSortingGPU.cpp b/Engine/Source/Runtime/Engine/Private/Particles/ParticleSortingGPU.cpp
--- a/Engine/Source/Runtime/Engine/Private/Particles/ParticleSortingGPU.cpp
+++ b/Engine/Source/Runtime/Engine/Private/Particles/ParticleSortingGPU.cpp
@@ -78,7 +78,7 @@ public:
 	}
 
 	/** Serialization. */
-	virtual bool Serialize( FArchive& Ar ) OVERRIDE
+	virtual bool Serialize( FArchive& Ar ) override
 	{
 		bool bShaderHasOutdatedParameters = FGlobalShader::Serialize( Ar );
 		Ar << InParticleIndices;
diff --git a/Engine/Source/Runtime/Engine/Private/Particles/ParticleSortingGPU.h b/Engine/Source/Runtime/Engine/Private/Particles/ParticleSortingGPU.h
--- a/Engine/Source/Runtime/Engine/Private/Particles/ParticleSortingGPU.h
+++ b/Engine/Source/Runtime/Engine/Private/Particles/ParticleSortingGPU.h
@@ -26,12 +26,12 @@ public:
 	/**
 	 * Initialize RHI resources.
 	 */
-	virtual void InitRHI() OVERRIDE;
+	virtual void InitRHI() override;
 
 	/**
 	 * Release RHI resources.
 	 */
-	virtual void ReleaseRHI() OVERRIDE;
+	virtual void ReleaseRHI() override;
 
 	/**
 	 * Retrieve the UAV for writing particle sort keys.
diff --git a/Engine/Source/Runtime/Engine/Private/Particles/ParticleVertexFactory.cpp b/Engine/Source/Runtime/Engine/Private/Particles/ParticleVertexFactory.cpp
--- a/Engine/Source/Runtime/Engine/Private/Particles/ParticleVertexFactory.cpp
+++ b/Engine/Source/Runtime/Engine/Private/Particles/ParticleVertexFactory.cpp
@@ -20,11 +20,11 @@ class FParticleSpriteVertexFactoryShaderParameters : public FVertexFactoryShader
 {
 public:
 
-	virtual void Bind(const FShaderParameterMap& ParameterMap) OVERRIDE
+	virtual void Bind(const FShaderParameterMap& ParameterMap) override
 	{
 	}
 
-	virtual void Serialize(FArchive& Ar) OVERRIDE
+	virtual void Serialize(FArchive& Ar) override
 	{
 	}
 };
@@ -33,7 +33,7 @@ class FParticleSpriteVertexFactoryShaderParametersVS : public FParticleSpriteVer
 {
 public:
 
-	virtual void SetMesh(FRHICommandList& RHICmdList, FShader* Shader,const FVertexFactory* VertexFactory,const FSceneView& View,const FMeshBatchElement& BatchElement,uint32 DataFlags) const OVERRIDE
+	virtual void SetMesh(FRHICommandList& RHICmdList, FShader* Shader,const FVertexFactory* VertexFactory,const FSceneView& View,const FMeshBatchElement& BatchElement,uint32 DataFlags) const override
 	{
 		FParticleSpriteVertexFactory* SpriteVF = (FParticleSpriteVertexFactory*)VertexFactory;
 		FVertexShaderRHIParamRef VertexShaderRHI = Shader->GetVertexShader();
@@ -45,7 +45,7 @@ class FParticleSpriteVertexFactoryShaderParametersPS : public FParticleSpriteVer
 {
 public:
 
-	virtual void SetMesh(FRHICommandList& RHICmdList, FShader* Shader,const FVertexFactory* VertexFactory,const FSceneView& View,const FMeshBatchElement& BatchElement,uint32 DataFlags) const OVERRIDE
+	virtual void SetMesh(FRHICommandList& RHICmdList, FShader* Shader,const FVertexFactory* VertexFactory,const FSceneView& View,const FMeshBatchElement& BatchElement,uint32 DataFlags) const override
 	{
 		FParticleSpriteVertexFactory* SpriteVF = (FParticleSpriteVertexFactory*)VertexFactory;
 		FPixelShaderRHIParamRef PixelShaderRHI = Shader->GetPixelShader();
diff --git a/Engine/Source/Runtime/Engine/Private/PhysicsEngine/PhysDerivedData.h b/Engine/Source/Runtime/Engine/Private/PhysicsEngine/PhysDerivedData.h
--- a/Engine/Source/Runtime/Engine/Private/PhysicsEngine/PhysDerivedData.h
+++ b/Engine/Source/Runtime/Engine/Private/PhysicsEngine/PhysDerivedData.h
@@ -32,12 +32,12 @@ public:
 	// This constructor only used by ULandscapeMeshCollisionComponent, which always only build TriMesh, not Convex...
 	FDerivedDataPhysXCooker( FName InFormat, ULandscapeMeshCollisionComponent* InMeshCollision, bool bMirrored );
 
-	virtual const TCHAR* GetPluginName() const OVERRIDE
+	virtual const TCHAR* GetPluginName() const override
 	{
 		return TEXT("PhysX");
 	}
 
-	virtual const TCHAR* GetVersionString() const OVERRIDE
+	virtual const TCHAR* GetVersionString() const override
 	{
 		// This is a version string that mimics the old versioning scheme. If you
 		// want to bump this version, generate a new guid using VS->Tools->Create GUID and
@@ -45,7 +45,7 @@ public:
 		return TEXT("{1F0627AE-ABEB-4206-8D78-E16BEB5DDC7E}");
 	}
 
-	virtual FString GetPluginSpecificCacheKeySuffix() const OVERRIDE
+	virtual FString GetPluginSpecificCacheKeySuffix() const override
 	{
 		enum { UE_PHYSX_DERIVEDDATA_VER = 1 };
 
@@ -66,12 +66,12 @@ public:
 	}
 
 
-	virtual bool IsBuildThreadsafe() const OVERRIDE
+	virtual bool IsBuildThreadsafe() const override
 	{
 		return false;
 	}
 
-	virtual bool Build( TArray<uint8>& OutData ) OVERRIDE;
+	virtual bool Build( TArray<uint8>& OutData ) override;
 
 	/** Return true if we can build **/
 	bool CanBuild()
diff --git a/Engine/Source/Runtime/Engine/Private/PhysicsEngine/PhysXSupport.h b/Engine/Source/Runtime/Engine/Private/PhysicsEngine/PhysXSupport.h
--- a/Engine/Source/Runtime/Engine/Private/PhysicsEngine/PhysXSupport.h
+++ b/Engine/Source/Runtime/Engine/Private/PhysicsEngine/PhysXSupport.h
@@ -258,7 +258,7 @@ public:
 		, ReadPos(0)
 	{}
 
-	virtual PxU32 read(void* Dest, PxU32 Count) OVERRIDE
+	virtual PxU32 read(void* Dest, PxU32 Count) override
 	{
 		check(Data);
 		check(Dest);
@@ -320,7 +320,7 @@ public:
 	virtual ~FPhysXAllocator() 
 	{}
 
-	virtual void* allocate(size_t size, const char* typeName, const char* filename, int line) OVERRIDE
+	virtual void* allocate(size_t size, const char* typeName, const char* filename, int line) override
 	{
 #if PHYSX_MEMORY_STATS
 		checkAtCompileTime(sizeof(FPhysXAllocationHeader) <= 16,FPhysXAllocationHeaderMustBeLessThan16Bytes);
@@ -352,7 +352,7 @@ public:
 #endif
 	}
 	 
-	virtual void deallocate(void* ptr) OVERRIDE
+	virtual void deallocate(void* ptr) override
 	{
 #if PHYSX_MEMORY_STATS
 		if( ptr )
@@ -398,7 +398,7 @@ class FPhysXBroadcastingAllocator : public PxBroadcastingAllocator
 class FPhysXErrorCallback : public PxErrorCallback
 {
 public:
-	virtual void reportError(PxErrorCode::Enum e, const char* message, const char* file, int line) OVERRIDE
+	virtual void reportError(PxErrorCode::Enum e, const char* message, const char* file, int line) override
 	{
 		// if not in game, ignore Perf warnings - i.e. Moving Static actor in editor will produce this warning
 		if (GIsEditor && e == PxErrorCode::ePERF_WARNING)
@@ -447,25 +447,25 @@ PxFilterFlags PhysXSimFilterShader(	PxFilterObjectAttributes attributes0, PxFilt
 /** Event callback used to notify engine about various collision events */
 class FPhysXSimEventCallback : public PxSimulationEventCallback
 {
-	virtual void onConstraintBreak(PxConstraintInfo* constraints, PxU32 count) OVERRIDE;
-	virtual void onWake(PxActor** actors, PxU32 count) OVERRIDE {}
-	virtual void onSleep(PxActor** actors, PxU32 count) OVERRIDE {}
-	virtual void onTrigger(PxTriggerPair* pairs, PxU32 count) OVERRIDE {}
-	virtual void onContact(const PxContactPairHeader& PairHeader, const PxContactPair* Pairs, PxU32 NumPairs) OVERRIDE;
+	virtual void onConstraintBreak(PxConstraintInfo* constraints, PxU32 count) override;
+	virtual void onWake(PxActor** actors, PxU32 count) override {}
+	virtual void onSleep(PxActor** actors, PxU32 count) override {}
+	virtual void onTrigger(PxTriggerPair* pairs, PxU32 count) override {}
+	virtual void onContact(const PxContactPairHeader& PairHeader, const PxContactPair* Pairs, PxU32 NumPairs) override;
 };
 
 /** Used to dispatch physx tasks to task graph */
 class FPhysXCPUDispatcher : public PxCpuDispatcher
 {
-	virtual void submitTask(PxBaseTask& task ) OVERRIDE;
-	virtual PxU32 getWorkerCount() const OVERRIDE;
+	virtual void submitTask(PxBaseTask& task ) override;
+	virtual PxU32 getWorkerCount() const override;
 };
 
 /** Used to dispatch physx tasks to the game thread */
 class FPhysXCPUDispatcherSingleThread : public PxCpuDispatcher
 {
-	virtual void submitTask( PxBaseTask& task ) OVERRIDE;
-	virtual PxU32 getWorkerCount() const OVERRIDE;
+	virtual void submitTask( PxBaseTask& task ) override;
+	virtual PxU32 getWorkerCount() const override;
 };
 
 #if WITH_APEX
@@ -478,53 +478,53 @@ class FApexNullRenderResourceManager : public NxUserRenderResourceManager
 public:
 	// NxUserRenderResourceManager interface.
 
-	virtual NxUserRenderVertexBuffer*	createVertexBuffer(const NxUserRenderVertexBufferDesc&) OVERRIDE
+	virtual NxUserRenderVertexBuffer*	createVertexBuffer(const NxUserRenderVertexBufferDesc&) override
 	{
 		return NULL;
 	}
-	virtual NxUserRenderIndexBuffer*	createIndexBuffer(const NxUserRenderIndexBufferDesc&) OVERRIDE
+	virtual NxUserRenderIndexBuffer*	createIndexBuffer(const NxUserRenderIndexBufferDesc&) override
 	{
 		return NULL;
 	}
-	virtual NxUserRenderBoneBuffer*		createBoneBuffer(const NxUserRenderBoneBufferDesc&) OVERRIDE
+	virtual NxUserRenderBoneBuffer*		createBoneBuffer(const NxUserRenderBoneBufferDesc&) override
 	{
 		return NULL;
 	}
-	virtual NxUserRenderInstanceBuffer*	createInstanceBuffer(const NxUserRenderInstanceBufferDesc&) OVERRIDE
+	virtual NxUserRenderInstanceBuffer*	createInstanceBuffer(const NxUserRenderInstanceBufferDesc&) override
 	{
 		return NULL;
 	}
-	virtual NxUserRenderSpriteBuffer*   createSpriteBuffer(const NxUserRenderSpriteBufferDesc&) OVERRIDE
+	virtual NxUserRenderSpriteBuffer*   createSpriteBuffer(const NxUserRenderSpriteBufferDesc&) override
 	{
 		return NULL;
 	}
 	
-	virtual NxUserRenderSurfaceBuffer*  createSurfaceBuffer(const NxUserRenderSurfaceBufferDesc& desc)   OVERRIDE
+	virtual NxUserRenderSurfaceBuffer*  createSurfaceBuffer(const NxUserRenderSurfaceBufferDesc& desc)   override
 	{
 		return NULL;
 	}
 	
-	virtual NxUserRenderResource*		createResource(const NxUserRenderResourceDesc&) OVERRIDE
+	virtual NxUserRenderResource*		createResource(const NxUserRenderResourceDesc&) override
 	{
 		return NULL;
 	}
-	virtual void						releaseVertexBuffer(NxUserRenderVertexBuffer&) OVERRIDE {}
-	virtual void						releaseIndexBuffer(NxUserRenderIndexBuffer&) OVERRIDE {}
-	virtual void						releaseBoneBuffer(NxUserRenderBoneBuffer&) OVERRIDE {}
-	virtual void						releaseInstanceBuffer(NxUserRenderInstanceBuffer&) OVERRIDE {}
-	virtual void						releaseSpriteBuffer(NxUserRenderSpriteBuffer&) OVERRIDE {}
-	virtual void                        releaseSurfaceBuffer(NxUserRenderSurfaceBuffer& buffer) OVERRIDE{}
-	virtual void						releaseResource(NxUserRenderResource&) OVERRIDE {}
-
-	virtual physx::PxU32				getMaxBonesForMaterial(void*) OVERRIDE
+	virtual void						releaseVertexBuffer(NxUserRenderVertexBuffer&) override {}
+	virtual void						releaseIndexBuffer(NxUserRenderIndexBuffer&) override {}
+	virtual void						releaseBoneBuffer(NxUserRenderBoneBuffer&) override {}
+	virtual void						releaseInstanceBuffer(NxUserRenderInstanceBuffer&) override {}
+	virtual void						releaseSpriteBuffer(NxUserRenderSpriteBuffer&) override {}
+	virtual void                        releaseSurfaceBuffer(NxUserRenderSurfaceBuffer& buffer) override{}
+	virtual void						releaseResource(NxUserRenderResource&) override {}
+
+	virtual physx::PxU32				getMaxBonesForMaterial(void*) override
 	{
 		return 0;
 	}
-	virtual bool						getSpriteLayoutData(physx::PxU32 , physx::PxU32 , NxUserRenderSpriteBufferDesc* ) OVERRIDE
+	virtual bool						getSpriteLayoutData(physx::PxU32 , physx::PxU32 , NxUserRenderSpriteBufferDesc* ) override
 	{
 		return false;
 	}
-	virtual bool						getInstanceLayoutData(physx::PxU32 , physx::PxU32 , NxUserRenderInstanceBufferDesc* ) OVERRIDE
+	virtual bool						getInstanceLayoutData(physx::PxU32 , physx::PxU32 , NxUserRenderInstanceBufferDesc* ) override
 	{
 		return false;
 	}
@@ -541,7 +541,7 @@ class FApexResourceCallback : public NxResourceCallback
 public:
 	// NxResourceCallback interface.
 
-	virtual void* requestResource(const char* NameSpace, const char* Name) OVERRIDE
+	virtual void* requestResource(const char* NameSpace, const char* Name) override
 	{
 		// Here a pointer is looked up by name and returned
 		(void)NameSpace;
@@ -550,7 +550,7 @@ public:
 		return NULL;
 	}
 
-	virtual void  releaseResource(const char* NameSpace, const char* Name, void* Resource) OVERRIDE
+	virtual void  releaseResource(const char* NameSpace, const char* Name, void* Resource) override
 	{
 		// Here we release a named resource
 		(void)NameSpace;
@@ -569,9 +569,9 @@ class FApexPhysX3Interface : public NxApexPhysX3Interface
 public:
 	// NxApexPhysX3Interface interface.
 
-	virtual void				setContactReportFlags(physx::PxShape* PShape, physx::PxPairFlags PFlags, NxDestructibleActor* actor, PxU16 actorChunkIndex) OVERRIDE;
+	virtual void				setContactReportFlags(physx::PxShape* PShape, physx::PxPairFlags PFlags, NxDestructibleActor* actor, PxU16 actorChunkIndex) override;
 
-	virtual physx::PxPairFlags	getContactReportFlags(const physx::PxShape* PShape) const OVERRIDE;
+	virtual physx::PxPairFlags	getContactReportFlags(const physx::PxShape* PShape) const override;
 };
 extern FApexPhysX3Interface GApexPhysX3Interface;
 
@@ -585,8 +585,8 @@ class FApexChunkReport : public NxUserChunkReport
 public:
 	// NxUserChunkReport interface.
 
-	virtual void	onDamageNotify(const NxApexDamageEventReportData& damageEvent) OVERRIDE;
-	virtual void	onStateChangeNotify(const NxApexChunkStateEventData& visibilityEvent) OVERRIDE;
+	virtual void	onDamageNotify(const NxApexDamageEventReportData& damageEvent) override;
+	virtual void	onStateChangeNotify(const NxApexChunkStateEventData& visibilityEvent) override;
 };
 extern FApexChunkReport GApexChunkReport;
 #endif // #if WITH_APEX
@@ -605,7 +605,7 @@ public:
 		: UsedMemory(0)
 	{}
 
-	virtual PxU32 write(const void* Src, PxU32 Count) OVERRIDE
+	virtual PxU32 write(const void* Src, PxU32 Count) override
 	{
 		UsedMemory += Count;
 		return Count;
diff --git a/Engine/Source/Runtime/Engine/Private/PhysicsEngine2D/Box2DIntegration.cpp b/Engine/Source/Runtime/Engine/Private/PhysicsEngine2D/Box2DIntegration.cpp
--- a/Engine/Source/Runtime/Engine/Private/PhysicsEngine2D/Box2DIntegration.cpp
+++ b/Engine/Source/Runtime/Engine/Private/PhysicsEngine2D/Box2DIntegration.cpp
@@ -12,12 +12,12 @@ class FBox2DVisualizer : public b2Draw
 {
 public:
 	// b2Draw interface
-	virtual void DrawPolygon(const b2Vec2* vertices, int32 vertexCount, const b2Color& color) OVERRIDE;
-	virtual void DrawSolidPolygon(const b2Vec2* vertices, int32 vertexCount, const b2Color& color) OVERRIDE;
-	virtual void DrawCircle(const b2Vec2& center, float32 radius, const b2Color& color) OVERRIDE;
-	virtual void DrawSolidCircle(const b2Vec2& center, float32 radius, const b2Vec2& axis, const b2Color& color) OVERRIDE;
-	virtual void DrawSegment(const b2Vec2& p1, const b2Vec2& p2, const b2Color& color) OVERRIDE;
-	virtual void DrawTransform(const b2Transform& xf) OVERRIDE;
+	virtual void DrawPolygon(const b2Vec2* vertices, int32 vertexCount, const b2Color& color) override;
+	virtual void DrawSolidPolygon(const b2Vec2* vertices, int32 vertexCount, const b2Color& color) override;
+	virtual void DrawCircle(const b2Vec2& center, float32 radius, const b2Color& color) override;
+	virtual void DrawSolidCircle(const b2Vec2& center, float32 radius, const b2Vec2& axis, const b2Color& color) override;
+	virtual void DrawSegment(const b2Vec2& p1, const b2Vec2& p2, const b2Color& color) override;
+	virtual void DrawTransform(const b2Transform& xf) override;
 	// End of b2Draw interface
 
 	FBox2DVisualizer()
diff --git a/Engine/Source/Runtime/Engine/Private/PhysicsEngine2D/Box2DIntegration.h b/Engine/Source/Runtime/Engine/Private/PhysicsEngine2D/Box2DIntegration.h
--- a/Engine/Source/Runtime/Engine/Private/PhysicsEngine2D/Box2DIntegration.h
+++ b/Engine/Source/Runtime/Engine/Private/PhysicsEngine2D/Box2DIntegration.h
@@ -14,8 +14,8 @@ struct FStartPhysics2DTickFunction : public FTickFunction
 	struct FPhysicsScene2D* Target;
 
 	// FTickFunction interface
-	virtual void ExecuteTick(float DeltaTime, enum ELevelTick TickType, ENamedThreads::Type CurrentThread, const FGraphEventRef& MyCompletionGraphEvent) OVERRIDE;
-	virtual FString DiagnosticMessage() OVERRIDE;
+	virtual void ExecuteTick(float DeltaTime, enum ELevelTick TickType, ENamedThreads::Type CurrentThread, const FGraphEventRef& MyCompletionGraphEvent) override;
+	virtual FString DiagnosticMessage() override;
 	// End of FTickFunction interface
 };
 
@@ -25,8 +25,8 @@ struct FEndPhysics2DTickFunction : public FTickFunction
 	struct FPhysicsScene2D* Target;
 
 	// FTickFunction interface
-	virtual void ExecuteTick(float DeltaTime, enum ELevelTick TickType, ENamedThreads::Type CurrentThread, const FGraphEventRef& MyCompletionGraphEvent) OVERRIDE;
-	virtual FString DiagnosticMessage() OVERRIDE;
+	virtual void ExecuteTick(float DeltaTime, enum ELevelTick TickType, ENamedThreads::Type CurrentThread, const FGraphEventRef& MyCompletionGraphEvent) override;
+	virtual FString DiagnosticMessage() override;
 	// End of FTickFunction interface
 };
 
@@ -49,7 +49,7 @@ public:
 	virtual ~FPhysicsScene2D();
 
 	// b2ContactFilter interface
-	virtual bool ShouldCollide(b2Fixture* FixtureA, b2Fixture* FixtureB) OVERRIDE;
+	virtual bool ShouldCollide(b2Fixture* FixtureA, b2Fixture* FixtureB) override;
 	// End of b2ContactFilter interface
 };
 
@@ -67,7 +67,7 @@ struct FPhysicsIntegration2D : public FGCObject
 	static void ShutdownPhysics();
 
 	// FSerializableObject interface
-	virtual void AddReferencedObjects(FReferenceCollector& Collector) OVERRIDE;
+	virtual void AddReferencedObjects(FReferenceCollector& Collector) override;
 	// End of FSerializableObject
 
 	static void DrawDebugPhysics(UWorld* World, FPrimitiveDrawInterface* PDI, const FSceneView* View);
diff --git a/Engine/Source/Runtime/Engine/Private/ShaderCompiler/ShaderCompiler.cpp b/Engine/Source/Runtime/Engine/Private/ShaderCompiler/ShaderCompiler.cpp
--- a/Engine/Source/Runtime/Engine/Private/ShaderCompiler/ShaderCompiler.cpp
+++ b/Engine/Source/Runtime/Engine/Private/ShaderCompiler/ShaderCompiler.cpp
@@ -2075,7 +2075,7 @@ protected:
 class FRecompileShaderMessageHandler : public IPlatformFile::IFileServerMessageHandler
 {
 	/** Subclass fills out an archive to send to the server */
-	virtual void FillPayload(FArchive& Payload) OVERRIDE
+	virtual void FillPayload(FArchive& Payload) override
 	{
 		// tell other side all the materials to load, by pathname
 		for (TObjectIterator<UMaterialInterface> It; It; ++It)
@@ -2097,7 +2097,7 @@ class FRecompileShaderMessageHandler : public IPlatformFile::IFileServerMessageH
 	}
 
 	/** Subclass pulls data response from the server */
-	virtual void ProcessResponse(FArchive& Response) OVERRIDE
+	virtual void ProcessResponse(FArchive& Response) override
 	{
 		// pull back the compiled mesh material data (if any)
 		TArray<uint8> MeshMaterialMaps;
diff --git a/Engine/Source/Runtime/Engine/Private/SkeletalRenderCPUSkin.h b/Engine/Source/Runtime/Engine/Private/SkeletalRenderCPUSkin.h
--- a/Engine/Source/Runtime/Engine/Private/SkeletalRenderCPUSkin.h
+++ b/Engine/Source/Runtime/Engine/Private/SkeletalRenderCPUSkin.h
@@ -139,16 +139,16 @@ public:
 	virtual ~FSkeletalMeshObjectCPUSkin();
 
 	// Begin FSkeletalMeshObject interface
-	virtual void InitResources() OVERRIDE;
-	virtual void ReleaseResources() OVERRIDE;
-	virtual void Update(int32 LODIndex,USkinnedMeshComponent* InMeshComponent,const TArray<FActiveVertexAnim>& ActiveVertexAnims) OVERRIDE;
-	virtual void UpdateDynamicData_RenderThread(FDynamicSkelMeshObjectData* InDynamicData) OVERRIDE;
-	virtual void EnableBlendWeightRendering(bool bEnabled, const TArray<int32>& InBonesOfInterest) OVERRIDE;
-	virtual void CacheVertices(int32 LODIndex, bool bForce) const OVERRIDE;
-	virtual bool IsCPUSkinned() const OVERRIDE { return true; }
-	virtual const FVertexFactory* GetVertexFactory(int32 LODIndex,int32 ChunkIdx) const OVERRIDE;
-	virtual TArray<FTransform>* GetSpaceBases() const OVERRIDE;
-	virtual int32 GetLOD() const OVERRIDE
+	virtual void InitResources() override;
+	virtual void ReleaseResources() override;
+	virtual void Update(int32 LODIndex,USkinnedMeshComponent* InMeshComponent,const TArray<FActiveVertexAnim>& ActiveVertexAnims) override;
+	virtual void UpdateDynamicData_RenderThread(FDynamicSkelMeshObjectData* InDynamicData) override;
+	virtual void EnableBlendWeightRendering(bool bEnabled, const TArray<int32>& InBonesOfInterest) override;
+	virtual void CacheVertices(int32 LODIndex, bool bForce) const override;
+	virtual bool IsCPUSkinned() const override { return true; }
+	virtual const FVertexFactory* GetVertexFactory(int32 LODIndex,int32 ChunkIdx) const override;
+	virtual TArray<FTransform>* GetSpaceBases() const override;
+	virtual int32 GetLOD() const override
 	{
 		if(DynamicData)
 		{
@@ -159,14 +159,14 @@ public:
 			return 0;
 		}
 	}
-	virtual const FTwoVectors& GetCustomLeftRightVectors(int32 SectionIndex) const OVERRIDE;
+	virtual const FTwoVectors& GetCustomLeftRightVectors(int32 SectionIndex) const override;
 
-	virtual bool HaveValidDynamicData() OVERRIDE
+	virtual bool HaveValidDynamicData() override
 	{ 
 		return ( DynamicData!=NULL ); 
 	}
 
-	virtual SIZE_T GetResourceSize() OVERRIDE
+	virtual SIZE_T GetResourceSize() override
 	{
 		SIZE_T ResourceSize=sizeof(*this);
 
@@ -189,7 +189,7 @@ public:
 		return ResourceSize;
 	}
 
-	virtual void DrawVertexElements(FPrimitiveDrawInterface* PDI, const FTransform& ToWorldSpace, bool bDrawNormals, bool bDrawTangents, bool bDrawBinormals) const OVERRIDE;
+	virtual void DrawVertexElements(FPrimitiveDrawInterface* PDI, const FTransform& ToWorldSpace, bool bDrawNormals, bool bDrawTangents, bool bDrawBinormals) const override;
 	// End FSkeletalMeshObject interface
 
 private:
diff --git a/Engine/Source/Runtime/Engine/Private/SkeletalRenderGPUSkin.h b/Engine/Source/Runtime/Engine/Private/SkeletalRenderGPUSkin.h
--- a/Engine/Source/Runtime/Engine/Private/SkeletalRenderGPUSkin.h
+++ b/Engine/Source/Runtime/Engine/Private/SkeletalRenderGPUSkin.h
@@ -245,16 +245,16 @@ public:
 	virtual ~FSkeletalMeshObjectGPUSkin();
 
 	// Begin FSkeletalMeshObject interface
-	virtual void InitResources() OVERRIDE;
-	virtual void ReleaseResources() OVERRIDE;
-	virtual void Update(int32 LODIndex,USkinnedMeshComponent* InMeshComponent,const TArray<FActiveVertexAnim>& ActiveVertexAnims) OVERRIDE;
-	virtual void UpdateDynamicData_RenderThread(FDynamicSkelMeshObjectData* InDynamicData) OVERRIDE;
-	virtual const FVertexFactory* GetVertexFactory(int32 LODIndex,int32 ChunkIdx) const OVERRIDE;
-	virtual void CacheVertices(int32 LODIndex, bool bForce) const OVERRIDE {}
-	virtual bool IsCPUSkinned() const OVERRIDE { return false; }
-	virtual TArray<FTransform>* GetSpaceBases() const OVERRIDE;
-
-	virtual int32 GetLOD() const OVERRIDE
+	virtual void InitResources() override;
+	virtual void ReleaseResources() override;
+	virtual void Update(int32 LODIndex,USkinnedMeshComponent* InMeshComponent,const TArray<FActiveVertexAnim>& ActiveVertexAnims) override;
+	virtual void UpdateDynamicData_RenderThread(FDynamicSkelMeshObjectData* InDynamicData) override;
+	virtual const FVertexFactory* GetVertexFactory(int32 LODIndex,int32 ChunkIdx) const override;
+	virtual void CacheVertices(int32 LODIndex, bool bForce) const override {}
+	virtual bool IsCPUSkinned() const override { return false; }
+	virtual TArray<FTransform>* GetSpaceBases() const override;
+
+	virtual int32 GetLOD() const override
 	{
 		if(DynamicData)
 		{
@@ -266,13 +266,13 @@ public:
 		}
 	}
 
-	virtual const FTwoVectors& GetCustomLeftRightVectors(int32 SectionIndex) const OVERRIDE;
-	virtual bool HaveValidDynamicData() OVERRIDE
+	virtual const FTwoVectors& GetCustomLeftRightVectors(int32 SectionIndex) const override;
+	virtual bool HaveValidDynamicData() override
 	{ 
 		return ( DynamicData!=NULL ); 
 	}
 
-	virtual SIZE_T GetResourceSize() OVERRIDE
+	virtual SIZE_T GetResourceSize() override
 	{
 		SIZE_T ResourceSize = sizeof(*this);
 		
diff --git a/Engine/Source/Runtime/Engine/Private/Slate/DebugCanvas.h b/Engine/Source/Runtime/Engine/Private/Slate/DebugCanvas.h
--- a/Engine/Source/Runtime/Engine/Private/Slate/DebugCanvas.h
+++ b/Engine/Source/Runtime/Engine/Private/Slate/DebugCanvas.h
@@ -30,7 +30,7 @@ private:
 	/**
 	 * ICustomSlateElement interface 
 	 */
-	virtual void DrawRenderThread( const void* InWindowBackBuffer ) OVERRIDE;
+	virtual void DrawRenderThread( const void* InWindowBackBuffer ) override;
 
 	/**
 	 * Deletes the rendering thread canvas 
diff --git a/Engine/Source/Runtime/Engine/Private/Slate/SlateSoundDevice.h b/Engine/Source/Runtime/Engine/Private/Slate/SlateSoundDevice.h
--- a/Engine/Source/Runtime/Engine/Private/Slate/SlateSoundDevice.h
+++ b/Engine/Source/Runtime/Engine/Private/Slate/SlateSoundDevice.h
@@ -10,6 +10,6 @@ public:
 	virtual ~FSlateSoundDevice(){}
 
 private:
-	virtual void PlaySound(const FSlateSound& Sound, int32 UserIndex = 0) const OVERRIDE;	
-	virtual float GetSoundDuration(const FSlateSound& Sound) const OVERRIDE;
+	virtual void PlaySound(const FSlateSound& Sound, int32 UserIndex = 0) const override;	
+	virtual float GetSoundDuration(const FSlateSound& Sound) const override;
 };
diff --git a/Engine/Source/Runtime/Engine/Private/SphereComponent.cpp b/Engine/Source/Runtime/Engine/Private/SphereComponent.cpp
--- a/Engine/Source/Runtime/Engine/Private/SphereComponent.cpp
+++ b/Engine/Source/Runtime/Engine/Private/SphereComponent.cpp
@@ -97,7 +97,7 @@ FPrimitiveSceneProxy* USphereComponent::CreateSceneProxy()
 		* @param	PDI - draw interface to render to
 		* @param	View - current view
 		*/
-		virtual void DrawDynamicElements(FPrimitiveDrawInterface* PDI,const FSceneView* View) OVERRIDE
+		virtual void DrawDynamicElements(FPrimitiveDrawInterface* PDI,const FSceneView* View) override
 		{
 			QUICK_SCOPE_CYCLE_COUNTER( STAT_SphereSceneProxy_DrawDynamicElements );
 
@@ -126,7 +126,7 @@ FPrimitiveSceneProxy* USphereComponent::CreateSceneProxy()
 			}
 		}
 
-		virtual FPrimitiveViewRelevance GetViewRelevance(const FSceneView* View)  OVERRIDE
+		virtual FPrimitiveViewRelevance GetViewRelevance(const FSceneView* View)  override
 		{
 			const bool bVisibleForSelection = !bDrawOnlyIfSelected || IsSelected();
 			const bool bVisibleForShowFlags = true; // @TODO
@@ -139,7 +139,7 @@ FPrimitiveSceneProxy* USphereComponent::CreateSceneProxy()
 			return Result;
 		}
 
-		virtual uint32 GetMemoryFootprint( void ) const OVERRIDE { return( sizeof( *this ) + GetAllocatedSize() ); }
+		virtual uint32 GetMemoryFootprint( void ) const override { return( sizeof( *this ) + GetAllocatedSize() ); }
 		uint32 GetAllocatedSize( void ) const { return( FPrimitiveSceneProxy::GetAllocatedSize() ); }
 
 	private:
diff --git a/Engine/Source/Runtime/Engine/Private/SplineMeshSceneProxy.h b/Engine/Source/Runtime/Engine/Private/SplineMeshSceneProxy.h
--- a/Engine/Source/Runtime/Engine/Private/SplineMeshSceneProxy.h
+++ b/Engine/Source/Runtime/Engine/Private/SplineMeshSceneProxy.h
@@ -67,11 +67,11 @@ private:
 /** Factory specific params */
 class FSplineMeshVertexFactoryShaderParameters : public FVertexFactoryShaderParameters
 {
-	void Bind(const FShaderParameterMap& ParameterMap) OVERRIDE;
+	void Bind(const FShaderParameterMap& ParameterMap) override;
 
-	void SetMesh(FRHICommandList& RHICmdList, FShader* Shader, const FVertexFactory* VertexFactory, const FSceneView& View, const FMeshBatchElement& BatchElement, uint32 DataFlags) const OVERRIDE;
+	void SetMesh(FRHICommandList& RHICmdList, FShader* Shader, const FVertexFactory* VertexFactory, const FSceneView& View, const FMeshBatchElement& BatchElement, uint32 DataFlags) const override;
 
-	void Serialize(FArchive& Ar) OVERRIDE
+	void Serialize(FArchive& Ar) override
 	{
 		Ar << SplineStartPosParam;
 		Ar << SplineStartTangentParam;
@@ -96,7 +96,7 @@ class FSplineMeshVertexFactoryShaderParameters : public FVertexFactoryShaderPara
 		Ar << SplineMeshYParam;
 	}
 
-	virtual uint32 GetSize() const OVERRIDE
+	virtual uint32 GetSize() const override
 	{
 		return sizeof(*this);
 	}
@@ -182,7 +182,7 @@ public:
 		}
 	}
 
-	virtual ~FSplineMeshSceneProxy() OVERRIDE
+	virtual ~FSplineMeshSceneProxy() override
 	{
 		ReleaseResources();
 
@@ -198,7 +198,7 @@ public:
 	void ReleaseResources();
 
 	/** Sets up a shadow FMeshBatch for a specific LOD. */
-	virtual bool GetShadowMeshElement(int32 LODIndex, uint8 InDepthPriorityGroup, FMeshBatch& OutMeshElement) const OVERRIDE
+	virtual bool GetShadowMeshElement(int32 LODIndex, uint8 InDepthPriorityGroup, FMeshBatch& OutMeshElement) const override
 	{
 		//checkf(LODIndex == 0, TEXT("Getting spline static mesh element with invalid LOD [%d]"), LODIndex);
 
@@ -212,7 +212,7 @@ public:
 	}
 
 	/** Sets up a FMeshBatch for a specific LOD and element. */
-	virtual bool GetMeshElement(int32 LODIndex, int32 SectionIndex, uint8 InDepthPriorityGroup, FMeshBatch& OutMeshElement, const bool bUseSelectedMaterial, const bool bUseHoveredMaterial) const OVERRIDE
+	virtual bool GetMeshElement(int32 LODIndex, int32 SectionIndex, uint8 InDepthPriorityGroup, FMeshBatch& OutMeshElement, const bool bUseSelectedMaterial, const bool bUseHoveredMaterial) const override
 	{
 		//checkf(LODIndex == 0 /*&& SectionIndex == 0*/, TEXT("Getting spline static mesh element with invalid params [%d, %d]"), LODIndex, SectionIndex);
 
@@ -226,7 +226,7 @@ public:
 	}
 
 	/** Sets up a wireframe FMeshBatch for a specific LOD. */
-	virtual bool GetWireframeMeshElement(int32 LODIndex, const FMaterialRenderProxy* WireframeRenderProxy, uint8 InDepthPriorityGroup, FMeshBatch& OutMeshElement) const OVERRIDE
+	virtual bool GetWireframeMeshElement(int32 LODIndex, const FMaterialRenderProxy* WireframeRenderProxy, uint8 InDepthPriorityGroup, FMeshBatch& OutMeshElement) const override
 	{
 		//checkf(LODIndex == 0, TEXT("Getting spline static mesh element with invalid LOD [%d]"), LODIndex);
 
@@ -239,7 +239,7 @@ public:
 		return false;
 	}
 
-	virtual FPrimitiveViewRelevance GetViewRelevance(const FSceneView* View) OVERRIDE
+	virtual FPrimitiveViewRelevance GetViewRelevance(const FSceneView* View) override
 	{
 		return FStaticMeshSceneProxy::GetViewRelevance(View);
 	}
diff --git a/Engine/Source/Runtime/Engine/Private/SpotLightComponent.cpp b/Engine/Source/Runtime/Engine/Private/SpotLightComponent.cpp
--- a/Engine/Source/Runtime/Engine/Private/SpotLightComponent.cpp
+++ b/Engine/Source/Runtime/Engine/Private/SpotLightComponent.cpp
@@ -141,7 +141,7 @@ public:
 
 	virtual float GetOuterConeAngle() const { return OuterConeAngle; }
 
-	virtual FVector2D GetLightShaftConeParams() const OVERRIDE { return FVector2D(CosLightShaftConeAngle, InvCosLightShaftConeDifference); }
+	virtual FVector2D GetLightShaftConeParams() const override { return FVector2D(CosLightShaftConeAngle, InvCosLightShaftConeDifference); }
 
 	virtual FSphere GetBoundingSphere() const
 	{
diff --git a/Engine/Source/Runtime/Engine/Private/StaticMeshComponent.cpp b/Engine/Source/Runtime/Engine/Private/StaticMeshComponent.cpp
--- a/Engine/Source/Runtime/Engine/Private/StaticMeshComponent.cpp
+++ b/Engine/Source/Runtime/Engine/Private/StaticMeshComponent.cpp
@@ -21,7 +21,7 @@ public:
 	{
 	}
 
-	virtual bool MatchesComponent(const UActorComponent* Component) const OVERRIDE
+	virtual bool MatchesComponent(const UActorComponent* Component) const override
 	{
 		return (CastChecked<UStaticMeshComponent>(Component)->StaticMesh == StaticMesh && FComponentInstanceDataBase::MatchesComponent(Component));
 	}
diff --git a/Engine/Source/Runtime/Engine/Private/TextRenderComponent.cpp b/Engine/Source/Runtime/Engine/Private/TextRenderComponent.cpp
--- a/Engine/Source/Runtime/Engine/Private/TextRenderComponent.cpp
+++ b/Engine/Source/Runtime/Engine/Private/TextRenderComponent.cpp
@@ -354,17 +354,17 @@ public:
 	virtual ~FTextRenderSceneProxy();
 
 	// Begin FPrimitiveSceneProxy interface
-	virtual void DrawDynamicElements(FPrimitiveDrawInterface* PDI,const FSceneView* View) OVERRIDE;
-	virtual void DrawStaticElements(FStaticPrimitiveDrawInterface* PDI) OVERRIDE;
-	virtual FPrimitiveViewRelevance GetViewRelevance(const FSceneView* View) OVERRIDE;
-	virtual bool CanBeOccluded() const OVERRIDE;
-	virtual uint32 GetMemoryFootprint() const OVERRIDE;
+	virtual void DrawDynamicElements(FPrimitiveDrawInterface* PDI,const FSceneView* View) override;
+	virtual void DrawStaticElements(FStaticPrimitiveDrawInterface* PDI) override;
+	virtual FPrimitiveViewRelevance GetViewRelevance(const FSceneView* View) override;
+	virtual bool CanBeOccluded() const override;
+	virtual uint32 GetMemoryFootprint() const override;
 	uint32 GetAllocatedSize() const;
 	// End FPrimitiveSceneProxy interface
 
 private:
 	// Begin FPrimitiveSceneProxy interface
-	virtual void CreateRenderThreadResources() OVERRIDE;
+	virtual void CreateRenderThreadResources() override;
 	// End FPrimitiveSceneProxy interface
 
 	void ReleaseRenderThreadResources();
diff --git a/Engine/Source/Runtime/Engine/Private/TickTaskManager.cpp b/Engine/Source/Runtime/Engine/Private/TickTaskManager.cpp
--- a/Engine/Source/Runtime/Engine/Private/TickTaskManager.cpp
+++ b/Engine/Source/Runtime/Engine/Private/TickTaskManager.cpp
@@ -549,13 +549,13 @@ public:
 	}
 
 	/** Allocate a new ticking structure for a ULevel **/
-	virtual FTickTaskLevel* AllocateTickTaskLevel() OVERRIDE
+	virtual FTickTaskLevel* AllocateTickTaskLevel() override
 	{
 		return new FTickTaskLevel;
 	}
 
 	/** Free a ticking structure for a ULevel **/
-	virtual void FreeTickTaskLevel(FTickTaskLevel* TickTaskLevel) OVERRIDE
+	virtual void FreeTickTaskLevel(FTickTaskLevel* TickTaskLevel) override
 	{
 		delete TickTaskLevel;
 	}
@@ -568,7 +568,7 @@ public:
 	 * @param DeltaSeconds - time in seconds since last tick
 	 * @param TickType - type of tick (viewports only, time only, etc)
 	 */
-	virtual void StartFrame(UWorld* InWorld, float InDeltaSeconds, ELevelTick InTickType) OVERRIDE
+	virtual void StartFrame(UWorld* InWorld, float InDeltaSeconds, ELevelTick InTickType) override
 	{
 		SCOPE_CYCLE_COUNTER(STAT_QueueTicks);
 		World = InWorld;
@@ -673,7 +673,7 @@ public:
 	 * @param DeltaSeconds - time in seconds since last tick
 	 * @param TickType - type of tick (viewports only, time only, etc)
 	 */
-	virtual void RunPauseFrame(UWorld* InWorld, float InDeltaSeconds, ELevelTick InTickType) OVERRIDE
+	virtual void RunPauseFrame(UWorld* InWorld, float InDeltaSeconds, ELevelTick InTickType) override
 	{
 		bTickNewlySpawned = true; // we don't support new spawns, but lets at least catch them.
 		Context.TickGroup = TG_PrePhysics; // reset this to the start tick group
@@ -695,7 +695,7 @@ public:
 		* @param Group - Ticking group to run
 		* @param bBlockTillComplete - if true, do not return until all ticks are complete
 	*/
-	virtual void RunTickGroup(ETickingGroup Group, bool bBlockTillComplete ) OVERRIDE
+	virtual void RunTickGroup(ETickingGroup Group, bool bBlockTillComplete ) override
 	{
 		if (QueueTickTasks.Num())
 		{
@@ -741,7 +741,7 @@ public:
 	}
 
 	/** Finish a frame of ticks **/
-	virtual void EndFrame() OVERRIDE
+	virtual void EndFrame() override
 	{
 		TickTaskSequencer.EndFrame();
 		bTickNewlySpawned = false;
diff --git a/Engine/Source/Runtime/Engine/Private/TileRendering.cpp b/Engine/Source/Runtime/Engine/Private/TileRendering.cpp
--- a/Engine/Source/Runtime/Engine/Private/TileRendering.cpp
+++ b/Engine/Source/Runtime/Engine/Private/TileRendering.cpp
@@ -108,7 +108,7 @@ public:
 	/** The mesh element. */
 	FMeshBatch MeshElement;
 
-	virtual void InitRHI() OVERRIDE
+	virtual void InitRHI() override
 	{
 		FMeshBatchElement& BatchElement = MeshElement.Elements[0];
 		MeshElement.VertexFactory = &GTileVertexFactory;
@@ -124,7 +124,7 @@ public:
 		BatchElement.PrimitiveUniformBufferResource = &GIdentityPrimitiveUniformBuffer;
 	}
 
-	virtual void ReleaseRHI() OVERRIDE
+	virtual void ReleaseRHI() override
 	{
 		MeshElement.Elements[0].PrimitiveUniformBuffer.SafeRelease();
 	}
diff --git a/Engine/Source/Runtime/Engine/Private/UnrealEngine.cpp b/Engine/Source/Runtime/Engine/Private/UnrealEngine.cpp
--- a/Engine/Source/Runtime/Engine/Private/UnrealEngine.cpp
+++ b/Engine/Source/Runtime/Engine/Private/UnrealEngine.cpp
@@ -6957,7 +6957,7 @@ static class FCDODump : private FSelfRegisteringExec
 	}
 
 	/** Console commands, see embeded usage statement **/
-	virtual bool Exec( UWorld* InWorld, const TCHAR* Cmd, FOutputDevice& Ar ) OVERRIDE
+	virtual bool Exec( UWorld* InWorld, const TCHAR* Cmd, FOutputDevice& Ar ) override
 	{
 		if(FParse::Command(&Cmd,TEXT("CDODump")))
 		{
@@ -9032,7 +9032,7 @@ struct FPendingStreamingLevelHolder : public FGCObject
 public:
 	TArray<ULevel*> Levels;
 
-	virtual void AddReferencedObjects( FReferenceCollector& Collector ) OVERRIDE
+	virtual void AddReferencedObjects( FReferenceCollector& Collector ) override
 	{
 		for( int32 LevelIndex = 0; LevelIndex < Levels.Num(); LevelIndex++ )
 		{
diff --git a/Engine/Source/Runtime/Engine/Private/VectorField.cpp b/Engine/Source/Runtime/Engine/Private/VectorField.cpp
--- a/Engine/Source/Runtime/Engine/Private/VectorField.cpp
+++ b/Engine/Source/Runtime/Engine/Private/VectorField.cpp
@@ -118,7 +118,7 @@ public:
 	/** 
 	 * @return ptr to the resource memory which has been preallocated
 	 */
-	virtual const void* GetResourceBulkData() const OVERRIDE
+	virtual const void* GetResourceBulkData() const override
 	{
 		check(BulkData != NULL);
 		return BulkData;
@@ -127,7 +127,7 @@ public:
 	/** 
 	 * @return size of resource memory
 	 */
-	virtual uint32 GetResourceBulkDataSize() const OVERRIDE
+	virtual uint32 GetResourceBulkDataSize() const override
 	{
 		check(BulkDataSize > 0);
 		return BulkDataSize;
@@ -136,7 +136,7 @@ public:
 	/**
 	 * Free memory after it has been used to initialize RHI resource 
 	 */
-	virtual void Discard() OVERRIDE
+	virtual void Discard() override
 	{
 	}
 
@@ -180,7 +180,7 @@ public:
 	/**
 	 * Initialize RHI resources.
 	 */
-	virtual void InitRHI() OVERRIDE
+	virtual void InitRHI() override
 	{
 		if (VolumeData && GRHIFeatureLevel >= ERHIFeatureLevel::SM4)
 		{
@@ -349,7 +349,7 @@ public:
 	 *	This function allows for generating renderer-side resources.
 	 *	Called in the rendering thread.
 	 */
-	virtual void CreateRenderThreadResources() OVERRIDE
+	virtual void CreateRenderThreadResources() override
 	{
 		VisualizationVertexFactory.InitResource();
 	}
@@ -357,7 +357,7 @@ public:
 	/**
 	 * Draws the bounding box for the vector field.
 	 */
-	virtual void DrawDynamicElements(FPrimitiveDrawInterface* PDI,const FSceneView* View) OVERRIDE
+	virtual void DrawDynamicElements(FPrimitiveDrawInterface* PDI,const FSceneView* View) override
 	{	
 		QUICK_SCOPE_CYCLE_COUNTER( STAT_VectorFieldSceneProxy_DrawDynamicElements );
 
@@ -373,7 +373,7 @@ public:
 	/**
 	 * Computes view relevance for this scene proxy.
 	 */
-	virtual FPrimitiveViewRelevance GetViewRelevance(const FSceneView* View) OVERRIDE
+	virtual FPrimitiveViewRelevance GetViewRelevance(const FSceneView* View) override
 	{
 		FPrimitiveViewRelevance Result;
 		Result.bDrawRelevance = IsShown(View); 
@@ -385,7 +385,7 @@ public:
 	/**
 	 * Computes the memory footprint of this scene proxy.
 	 */
-	virtual uint32 GetMemoryFootprint() const OVERRIDE
+	virtual uint32 GetMemoryFootprint() const override
 	{
 		return sizeof(*this) + GetAllocatedSize();
 	}
@@ -604,7 +604,7 @@ public:
 	}
 
 	/** Serialization. */
-	virtual bool Serialize( FArchive& Ar ) OVERRIDE
+	virtual bool Serialize( FArchive& Ar ) override
 	{
 		bool bShaderHasOutdatedParameters = FGlobalShader::Serialize( Ar );
 		Ar << AtlasTexture;
@@ -712,7 +712,7 @@ public:
 	/**
 	 * Initialize RHI resources.
 	 */
-	virtual void InitRHI() OVERRIDE
+	virtual void InitRHI() override
 	{
 		if (GRHIFeatureLevel >= ERHIFeatureLevel::SM4)
 		{
@@ -745,7 +745,7 @@ public:
 	/**
 	 * Release RHI resources.
 	 */
-	virtual void ReleaseRHI() OVERRIDE
+	virtual void ReleaseRHI() override
 	{
 		VolumeTextureUAV.SafeRelease();
 		FVectorFieldResource::ReleaseRHI();
@@ -832,7 +832,7 @@ public:
 	/**
 	 * Resets the vector field simulation.
 	 */
-	virtual void ResetVectorField() OVERRIDE
+	virtual void ResetVectorField() override
 	{
 		FrameTime = 0.0f;
 	}
diff --git a/Engine/Source/Runtime/Engine/Private/VectorField.h b/Engine/Source/Runtime/Engine/Private/VectorField.h
--- a/Engine/Source/Runtime/Engine/Private/VectorField.h
+++ b/Engine/Source/Runtime/Engine/Private/VectorField.h
@@ -50,7 +50,7 @@ public:
 	/**
 	 * Release RHI resources.
 	 */
-	virtual void ReleaseRHI() OVERRIDE;
+	virtual void ReleaseRHI() override;
 
 	/**
 	 * Updates the vector field.
diff --git a/Engine/Source/Runtime/Engine/Private/VectorFieldVisualization.cpp b/Engine/Source/Runtime/Engine/Private/VectorFieldVisualization.cpp
--- a/Engine/Source/Runtime/Engine/Private/VectorFieldVisualization.cpp
+++ b/Engine/Source/Runtime/Engine/Private/VectorFieldVisualization.cpp
@@ -28,19 +28,19 @@ IMPLEMENT_UNIFORM_BUFFER_STRUCT(FVectorFieldVisualizationParameters,TEXT("Vector
 class FVectorFieldVisualizationVertexFactoryShaderParameters : public FVertexFactoryShaderParameters
 {
 public:
-	virtual void Bind( const FShaderParameterMap& ParameterMap ) OVERRIDE
+	virtual void Bind( const FShaderParameterMap& ParameterMap ) override
 	{
 		VectorFieldTexture.Bind(ParameterMap, TEXT("VectorFieldTexture"));
 		VectorFieldTextureSampler.Bind(ParameterMap, TEXT("VectorFieldTextureSampler"));
 	}
 
-	virtual void Serialize(FArchive& Ar) OVERRIDE
+	virtual void Serialize(FArchive& Ar) override
 	{
 		Ar << VectorFieldTexture;
 		Ar << VectorFieldTextureSampler;
 	}
 
-	virtual void SetMesh(FRHICommandList& RHICmdList, FShader* Shader,const FVertexFactory* VertexFactory,const FSceneView& View,const FMeshBatchElement& BatchElement,uint32 DataFlags) const OVERRIDE;
+	virtual void SetMesh(FRHICommandList& RHICmdList, FShader* Shader,const FVertexFactory* VertexFactory,const FSceneView& View,const FMeshBatchElement& BatchElement,uint32 DataFlags) const override;
 
 	virtual uint32 GetSize() const { return sizeof(*this); }
 
diff --git a/Engine/Source/Runtime/Engine/Private/WorldComposition.cpp b/Engine/Source/Runtime/Engine/Private/WorldComposition.cpp
--- a/Engine/Source/Runtime/Engine/Private/WorldComposition.cpp
+++ b/Engine/Source/Runtime/Engine/Private/WorldComposition.cpp
@@ -116,7 +116,7 @@ struct FWorldTilesGatherer
 	// Tile short pkg name -> Tiles LOD
 	TMultiMap<FString, FPackageNameAndLODIndex> TilesLODCollection;
 	
-	virtual bool Visit(const TCHAR* FilenameOrDirectory, bool bIsDirectory) OVERRIDE
+	virtual bool Visit(const TCHAR* FilenameOrDirectory, bool bIsDirectory) override
 	{
 		FString FullPath = FilenameOrDirectory;
 	
diff --git a/Engine/Source/Runtime/Engine/Public/AI/NavLinkRenderingProxy.h b/Engine/Source/Runtime/Engine/Public/AI/NavLinkRenderingProxy.h
--- a/Engine/Source/Runtime/Engine/Public/AI/NavLinkRenderingProxy.h
+++ b/Engine/Source/Runtime/Engine/Public/AI/NavLinkRenderingProxy.h
@@ -35,9 +35,9 @@ private:
 public:
 	/** Initialization constructor. */
 	FNavLinkRenderingProxy(const UPrimitiveComponent* InComponent);
-	virtual void DrawDynamicElements(FPrimitiveDrawInterface* PDI,const FSceneView* View) OVERRIDE;
-	virtual FPrimitiveViewRelevance GetViewRelevance(const FSceneView* View) OVERRIDE;
-	virtual uint32 GetMemoryFootprint( void ) const OVERRIDE;
+	virtual void DrawDynamicElements(FPrimitiveDrawInterface* PDI,const FSceneView* View) override;
+	virtual FPrimitiveViewRelevance GetViewRelevance(const FSceneView* View) override;
+	virtual uint32 GetMemoryFootprint( void ) const override;
 	uint32 GetAllocatedSize( void ) const;
 	void StorePointLinks(const FTransform& LocalToWorld, const TArray<FNavigationLink>& LinksArray);
 	void StoreSegmentLinks(const FTransform& LocalToWorld, const TArray<FNavigationSegmentLink>& LinksArray);
diff --git a/Engine/Source/Runtime/Engine/Public/AudioDerivedData.h b/Engine/Source/Runtime/Engine/Public/AudioDerivedData.h
--- a/Engine/Source/Runtime/Engine/Public/AudioDerivedData.h
+++ b/Engine/Source/Runtime/Engine/Public/AudioDerivedData.h
@@ -18,12 +18,12 @@ public:
 
 	FDerivedAudioDataCompressor(USoundWave* InSoundNode, FName InFormat);
 
-	virtual const TCHAR* GetPluginName() const OVERRIDE
+	virtual const TCHAR* GetPluginName() const override
 	{
 		return TEXT("Audio");
 	}
 
-	virtual const TCHAR* GetVersionString() const OVERRIDE
+	virtual const TCHAR* GetVersionString() const override
 	{
 		// This is a version string that mimics the old versioning scheme. If you
 		// want to bump this version, generate a new guid using VS->Tools->Create GUID and
@@ -32,7 +32,7 @@ public:
 		return TEXT("0002_0000");
 	}
 
-	virtual FString GetPluginSpecificCacheKeySuffix() const OVERRIDE
+	virtual FString GetPluginSpecificCacheKeySuffix() const override
 	{
 		int32 FormatVersion = 0xffff; // if the compressor is NULL, this will be used as the version...and in that case we expect everything to fail anyway
 		if (Compressor)
@@ -44,10 +44,10 @@ public:
 		return FString::Printf( TEXT("%s_%04X_%s"), *FormatString, FormatVersion, *SoundNode->CompressedDataGuid.ToString());
 	}
 	
-	virtual bool IsBuildThreadsafe() const OVERRIDE
+	virtual bool IsBuildThreadsafe() const override
 	{
 		return false;
 	}
 
-	virtual bool Build(TArray<uint8>& OutData) OVERRIDE;
+	virtual bool Build(TArray<uint8>& OutData) override;
 };
diff --git a/Engine/Source/Runtime/Engine/Public/AudioDevice.h b/Engine/Source/Runtime/Engine/Public/AudioDevice.h
--- a/Engine/Source/Runtime/Engine/Public/AudioDevice.h
+++ b/Engine/Source/Runtime/Engine/Public/AudioDevice.h
@@ -165,7 +165,7 @@ class ENGINE_API FAudioDevice : public FExec
 public:
 
 	//Begin FExec Interface
-	virtual bool Exec( UWorld* InWorld, const TCHAR* Cmd, FOutputDevice& Ar = *GLog ) OVERRIDE;
+	virtual bool Exec( UWorld* InWorld, const TCHAR* Cmd, FOutputDevice& Ar = *GLog ) override;
 	//End FExec Interface
 
 #if !UE_BUILD_SHIPPING
