UE4 Refactoring. Changed OVERRIDE and FINAL macros to keywords override and final.

diff --git a/Engine/Source/Runtime/Engine/Public/CanvasItem.h b/Engine/Source/Runtime/Engine/Public/CanvasItem.h
--- a/Engine/Source/Runtime/Engine/Public/CanvasItem.h
+++ b/Engine/Source/Runtime/Engine/Public/CanvasItem.h
@@ -234,7 +234,7 @@ public:
 	 *
 	 * @param	InPosition		Draw position.
 	 */
-	virtual void Draw( class FCanvas* InCanvas ) OVERRIDE;
+	virtual void Draw( class FCanvas* InCanvas ) override;
 
 	/* Expose the functions defined in the base class. */
 	using FCanvasItem::Draw;
@@ -312,7 +312,7 @@ public:
 	 *
 	 * @param	InPosition		Draw position.
 	 */
-	virtual void Draw( class FCanvas* InCanvas ) OVERRIDE;
+	virtual void Draw( class FCanvas* InCanvas ) override;
 
 	/* Expose the functions defined in the base class. */
 	using FCanvasItem::Draw;
@@ -421,7 +421,7 @@ public:
 	 *
 	 * @param	InCanvas		Canvas on which to draw
 	 */
-	virtual void Draw( class FCanvas* InCanvas ) OVERRIDE;
+	virtual void Draw( class FCanvas* InCanvas ) override;
 
 	/* Expose the functions defined in the base class. */
 	using FCanvasItem::Draw;
@@ -558,7 +558,7 @@ public:
 	 *
 	 * @param	InCanvas		Canvas on which to draw
 	 */
-	virtual void Draw( class FCanvas* InCanvas ) OVERRIDE;
+	virtual void Draw( class FCanvas* InCanvas ) override;
 	
 	/* 
 	 * Draw line at the given coordinates.
@@ -566,7 +566,7 @@ public:
 	 * @param	InCanvas		Canvas on which to draw
 	 * @param	InPosition		Draw Start position
 	 */
-	virtual void Draw( class FCanvas* InCanvas, const FVector2D& InPosition ) OVERRIDE
+	virtual void Draw( class FCanvas* InCanvas, const FVector2D& InPosition ) override
 	{
 		Origin.X = InPosition.X;
 		Origin.Y = InPosition.Y;
@@ -594,7 +594,7 @@ public:
 	 * @param	X				X Draw position 
 	 * @param	Y				Y Draw position
 	 */
-	virtual void Draw( class FCanvas* InCanvas, float InX, float InY ) OVERRIDE
+	virtual void Draw( class FCanvas* InCanvas, float InX, float InY ) override
 	{
 		Origin.X = InX;
 		Origin.Y = InY;
@@ -658,7 +658,7 @@ public:
 		, Size( InSize )
 		, LineThickness( 0.0f )	{};
 
-	virtual void Draw( class FCanvas* InCanvas ) OVERRIDE;
+	virtual void Draw( class FCanvas* InCanvas ) override;
 
 	/* Expose the functions defined in the base class. */
 	using FCanvasItem::Draw;
@@ -778,13 +778,13 @@ public:
 		TriangleList[0].V2_Pos = InPointC;
 	}
 
-	virtual void Draw( class FCanvas* InCanvas ) OVERRIDE;
+	virtual void Draw( class FCanvas* InCanvas ) override;
 		
 	/* Expose the functions defined in the base class. */
 	using FCanvasItem::Draw;
 
 	/* Set the Color of the item. */
-	virtual void SetColor( const FLinearColor& InColor ) OVERRIDE;
+	virtual void SetColor( const FLinearColor& InColor ) override;
 
 	/* texture to use for triangle(s). */
 	const FTexture* Texture;
@@ -833,7 +833,7 @@ public:
 	 {
 		 delete TriListItem;
 	 }
-	 virtual void Draw( class FCanvas* InCanvas ) OVERRIDE;
+	 virtual void Draw( class FCanvas* InCanvas ) override;
 
 	 /* 	 
 	 * Regenerates the tri list for the object with a new central point and radius
@@ -866,7 +866,7 @@ public:
 	 }
 
 	 /* Set the Color of the item. */
-	 virtual void SetColor( const FLinearColor& InColor ) OVERRIDE;
+	 virtual void SetColor( const FLinearColor& InColor ) override;
 
 private:
 	
diff --git a/Engine/Source/Runtime/Engine/Public/ContentStreaming.h b/Engine/Source/Runtime/Engine/Public/ContentStreaming.h
--- a/Engine/Source/Runtime/Engine/Public/ContentStreaming.h
+++ b/Engine/Source/Runtime/Engine/Public/ContentStreaming.h
@@ -596,7 +596,7 @@ struct FStreamingManagerCollection : public IStreamingManager
 	 * @param Ar	Output device for feedback
 	 * @return		true if the command was handled
 	 */
-	virtual bool Exec( UWorld* InWorld, const TCHAR* Cmd, FOutputDevice& Ar ) OVERRIDE;
+	virtual bool Exec( UWorld* InWorld, const TCHAR* Cmd, FOutputDevice& Ar ) override;
 
 	/** Adds a ULevel to the streaming manager. */
 	virtual void AddLevel( class ULevel* Level );
@@ -852,7 +852,7 @@ struct FStreamingManagerTexture : public ITextureStreamingManager
 	 * @param Ar	Output device for feedback
 	 * @return		true if the command was handled
 	 */
-	virtual bool Exec( UWorld* InWorld, const TCHAR* Cmd, FOutputDevice& Ar ) OVERRIDE;
+	virtual bool Exec( UWorld* InWorld, const TCHAR* Cmd, FOutputDevice& Ar ) override;
 
 	/**
 	 * Exec command handlers
diff --git a/Engine/Source/Runtime/Engine/Public/DVRStreaming.h b/Engine/Source/Runtime/Engine/Public/DVRStreaming.h
--- a/Engine/Source/Runtime/Engine/Public/DVRStreaming.h
+++ b/Engine/Source/Runtime/Engine/Public/DVRStreaming.h
@@ -49,7 +49,7 @@ public:
 	{
 	}
 
-	virtual void GetStreamingStatus(FDVRStreamingStatus &StreamingStatus) OVERRIDE
+	virtual void GetStreamingStatus(FDVRStreamingStatus &StreamingStatus) override
 	{
 		StreamingStatus.bIsStreaming = false;
 		StreamingStatus.bIsStreamingEnabled = bIsStreamingEnabled;
@@ -60,7 +60,7 @@ public:
 		StreamingStatus.ProviderUrl = TEXT("");
 	}
 
-	virtual void EnableStreaming(bool Enable) OVERRIDE
+	virtual void EnableStreaming(bool Enable) override
 	{
 		bIsStreamingEnabled = Enable;
 	}
diff --git a/Engine/Source/Runtime/Engine/Public/DelayAction.h b/Engine/Source/Runtime/Engine/Public/DelayAction.h
--- a/Engine/Source/Runtime/Engine/Public/DelayAction.h
+++ b/Engine/Source/Runtime/Engine/Public/DelayAction.h
@@ -29,7 +29,7 @@ public:
 
 #if WITH_EDITOR
 	// Returns a human readable description of the latent operation's current state
-	virtual FString GetDescription() const OVERRIDE
+	virtual FString GetDescription() const override
 	{
 		return FString::Printf( *NSLOCTEXT("DelayAction", "DelayActionTime", "Delay (%.3f seconds left)").ToString(), TimeRemaining);
 	}
diff --git a/Engine/Source/Runtime/Engine/Public/EngineUtils.h b/Engine/Source/Runtime/Engine/Public/EngineUtils.h
--- a/Engine/Source/Runtime/Engine/Public/EngineUtils.h
+++ b/Engine/Source/Runtime/Engine/Public/EngineUtils.h
@@ -49,7 +49,7 @@ struct HActor : public HHitProxy
 		, MaterialIndex(InMaterialIndex)
 		{}
 
-	virtual void AddReferencedObjects( FReferenceCollector& Collector ) OVERRIDE
+	virtual void AddReferencedObjects( FReferenceCollector& Collector ) override
 	{
 		Collector.AddReferencedObject( Actor );
 
@@ -77,7 +77,7 @@ struct HBSPBrushVert : public HHitProxy
 		Brush(InBrush),
 		Vertex(InVertex)
 	{}
-	virtual void AddReferencedObjects( FReferenceCollector& Collector ) OVERRIDE
+	virtual void AddReferencedObjects( FReferenceCollector& Collector ) override
 	{
 		//@todo: Brush Hit proxies are currently referencing between UWorld's (undesired), 
 		// once this issue is resolved remove the TWeakObjectPtr and replace with Standard ABrush*.
@@ -100,7 +100,7 @@ struct HStaticMeshVert : public HHitProxy
 		Actor(InActor),
 		Vertex(InVertex)
 	{}
-	virtual void AddReferencedObjects( FReferenceCollector& Collector ) OVERRIDE
+	virtual void AddReferencedObjects( FReferenceCollector& Collector ) override
 	{
 		Collector.AddReferencedObject( Actor );
 	}
@@ -551,7 +551,7 @@ public:
 		Console(InConsole)
 	{}
 
-	virtual void Serialize(const TCHAR* Text, ELogVerbosity::Type Verbosity, const class FName& Category) OVERRIDE;
+	virtual void Serialize(const TCHAR* Text, ELogVerbosity::Type Verbosity, const class FName& Category) override;
 
 private:
 
diff --git a/Engine/Source/Runtime/Engine/Public/GPUSkinVertexFactory.h b/Engine/Source/Runtime/Engine/Public/GPUSkinVertexFactory.h
--- a/Engine/Source/Runtime/Engine/Public/GPUSkinVertexFactory.h
+++ b/Engine/Source/Runtime/Engine/Public/GPUSkinVertexFactory.h
@@ -207,7 +207,7 @@ public:
 	virtual ~FBoneBufferPool();
 	
 public: // From FTickableObjectRenderThread
-	virtual TStatId GetStatId() const OVERRIDE;
+	virtual TStatId GetStatId() const override;
 };
 
 /** The type for the buffer pool */
@@ -248,7 +248,7 @@ public:
 	virtual ~FBoneTexturePool();
 	
 public: // From FTickableObjectRenderThread
-	virtual TStatId GetStatId() const OVERRIDE;
+	virtual TStatId GetStatId() const override;
 };
 
 /** The type for bone buffers */
@@ -284,14 +284,14 @@ public:
 
 	// interface FRenderResource ------------------------------------------
 
-	virtual void ReleaseRHI() OVERRIDE
+	virtual void ReleaseRHI() override
 	{
 		DEC_DWORD_STAT_BY( STAT_SkeletalMeshMotionBlurSkinningMemory, ComputeMemorySize());
 		BoneBuffer.SafeRelease();
 		FRenderResource::ReleaseRHI();
 	}
 
-	virtual void InitDynamicRHI() OVERRIDE
+	virtual void InitDynamicRHI() override
 	{
 		if(SizeX)
 		{
@@ -575,9 +575,9 @@ public:
 	}
 
 	// FRenderResource interface.
-	virtual void InitRHI() OVERRIDE;
-	virtual void InitDynamicRHI() OVERRIDE;
-	virtual void ReleaseDynamicRHI() OVERRIDE;
+	virtual void InitRHI() override;
+	virtual void InitDynamicRHI() override;
+	virtual void ReleaseDynamicRHI() override;
 
 	static FVertexFactoryShaderParameters* ConstructShaderParameters(EShaderFrequency ShaderFrequency);
 
@@ -665,7 +665,7 @@ public:
 	* Creates declarations for each of the vertex stream components and
 	* initializes the device resource
 	*/
-	virtual void InitRHI() OVERRIDE;
+	virtual void InitRHI() override;
 
 	static FVertexFactoryShaderParameters* ConstructShaderParameters(EShaderFrequency ShaderFrequency);
 
@@ -813,12 +813,12 @@ public:
 		FGPUBaseSkinVertexFactory::UpdateRHI();
 	}
 
-	virtual FGPUBaseSkinVertexFactory* GetVertexFactory() OVERRIDE
+	virtual FGPUBaseSkinVertexFactory* GetVertexFactory() override
 	{
 		return this;
 	}
 
-	virtual const FGPUBaseSkinVertexFactory* GetVertexFactory() const OVERRIDE
+	virtual const FGPUBaseSkinVertexFactory* GetVertexFactory() const override
 	{
 		return this;
 	}
@@ -829,8 +829,8 @@ public:
 	* Creates declarations for each of the vertex stream components and
 	* initializes the device resource
 	*/
-	virtual void InitRHI() OVERRIDE;
-	virtual void ReleaseDynamicRHI() OVERRIDE;
+	virtual void InitRHI() override;
+	virtual void ReleaseDynamicRHI() override;
 
 	static FVertexFactoryShaderParameters* ConstructShaderParameters(EShaderFrequency ShaderFrequency);
 
diff --git a/Engine/Source/Runtime/Engine/Public/Landscape/LandscapeRender.h b/Engine/Source/Runtime/Engine/Public/Landscape/LandscapeRender.h
--- a/Engine/Source/Runtime/Engine/Public/Landscape/LandscapeRender.h
+++ b/Engine/Source/Runtime/Engine/Public/Landscape/LandscapeRender.h
@@ -112,9 +112,9 @@ public:
 	/**
 	* Set any shader data specific to this vertex factory
 	*/
-	virtual void SetMesh(FRHICommandList& RHICmdList, FShader* PixelShader,const class FVertexFactory* VertexFactory,const class FSceneView& View,const struct FMeshBatchElement& BatchElement,uint32 DataFlags) const OVERRIDE;
+	virtual void SetMesh(FRHICommandList& RHICmdList, FShader* PixelShader,const class FVertexFactory* VertexFactory,const class FSceneView& View,const struct FMeshBatchElement& BatchElement,uint32 DataFlags) const override;
 	
-	virtual uint32 GetSize() const OVERRIDE
+	virtual uint32 GetSize() const override
 	{
 		return sizeof(*this);
 	}
@@ -185,7 +185,7 @@ public:
 		UpdateRHI();
 	}
 
-	virtual uint64 GetStaticBatchElementVisibility( const class FSceneView& View, const struct FMeshBatch* Batch ) const OVERRIDE;
+	virtual uint64 GetStaticBatchElementVisibility( const class FSceneView& View, const struct FMeshBatch* Batch ) const override;
 
 	/** stream component data bound to this vertex factory */
 	DataType Data;  
@@ -476,14 +476,14 @@ public:
 	FLandscapeComponentSceneProxy(ULandscapeComponent* InComponent, FLandscapeEditToolRenderData* InEditToolRenderData);
 
 	// FPrimitiveSceneProxy interface.
-	virtual void DrawStaticElements(FStaticPrimitiveDrawInterface* PDI) OVERRIDE;
-	virtual void DrawDynamicElements(FPrimitiveDrawInterface* PDI,const FSceneView* View) OVERRIDE;
-	virtual uint32 GetMemoryFootprint( void ) const OVERRIDE { return( sizeof( *this ) + GetAllocatedSize() ); }
-	virtual FPrimitiveViewRelevance GetViewRelevance(const FSceneView* View) OVERRIDE;
-	virtual bool CanBeOccluded() const OVERRIDE;
-	virtual void GetLightRelevance(const FLightSceneProxy* LightSceneProxy, bool& bDynamic, bool& bRelevant, bool& bLightMapped, bool& bShadowMapped) const OVERRIDE;
-	virtual void OnTransformChanged() OVERRIDE;
-	virtual void CreateRenderThreadResources() OVERRIDE;
+	virtual void DrawStaticElements(FStaticPrimitiveDrawInterface* PDI) override;
+	virtual void DrawDynamicElements(FPrimitiveDrawInterface* PDI,const FSceneView* View) override;
+	virtual uint32 GetMemoryFootprint( void ) const override { return( sizeof( *this ) + GetAllocatedSize() ); }
+	virtual FPrimitiveViewRelevance GetViewRelevance(const FSceneView* View) override;
+	virtual bool CanBeOccluded() const override;
+	virtual void GetLightRelevance(const FLightSceneProxy* LightSceneProxy, bool& bDynamic, bool& bRelevant, bool& bLightMapped, bool& bShadowMapped) const override;
+	virtual void OnTransformChanged() override;
+	virtual void CreateRenderThreadResources() override;
 
 	friend class ULandscapeComponent;
 	friend class FLandscapeVertexFactoryVertexShaderParameters;
diff --git a/Engine/Source/Runtime/Engine/Public/Landscape/LandscapeRenderMobile.h b/Engine/Source/Runtime/Engine/Public/Landscape/LandscapeRenderMobile.h
--- a/Engine/Source/Runtime/Engine/Public/Landscape/LandscapeRenderMobile.h
+++ b/Engine/Source/Runtime/Engine/Public/Landscape/LandscapeRenderMobile.h
@@ -123,7 +123,7 @@ class FLandscapeComponentSceneProxyMobile : public FLandscapeComponentSceneProxy
 public:
 	FLandscapeComponentSceneProxyMobile(ULandscapeComponent* InComponent, FLandscapeEditToolRenderData* InEditToolRenderData);
 
-	virtual void CreateRenderThreadResources() OVERRIDE;
+	virtual void CreateRenderThreadResources() override;
 
 	friend class FLandscapeVertexBufferMobile;
 };
diff --git a/Engine/Source/Runtime/Engine/Public/Landscape/LandscapeSplineProxies.h b/Engine/Source/Runtime/Engine/Public/Landscape/LandscapeSplineProxies.h
--- a/Engine/Source/Runtime/Engine/Public/Landscape/LandscapeSplineProxies.h
+++ b/Engine/Source/Runtime/Engine/Public/Landscape/LandscapeSplineProxies.h
@@ -32,7 +32,7 @@ struct HLandscapeSplineProxy_Segment : public HLandscapeSplineProxy
 		SplineSegment(InSplineSegment)
 	{
 	}
-	virtual void AddReferencedObjects( FReferenceCollector& Collector ) OVERRIDE
+	virtual void AddReferencedObjects( FReferenceCollector& Collector ) override
 	{
 		Collector.AddReferencedObject( SplineSegment );
 	}
@@ -49,7 +49,7 @@ struct HLandscapeSplineProxy_ControlPoint : public HLandscapeSplineProxy
 		ControlPoint(InControlPoint)
 	{
 	}
-	virtual void AddReferencedObjects( FReferenceCollector& Collector ) OVERRIDE
+	virtual void AddReferencedObjects( FReferenceCollector& Collector ) override
 	{
 		Collector.AddReferencedObject( ControlPoint );
 	}
diff --git a/Engine/Source/Runtime/Engine/Public/MallocProfilerEx.h b/Engine/Source/Runtime/Engine/Public/MallocProfilerEx.h
--- a/Engine/Source/Runtime/Engine/Public/MallocProfilerEx.h
+++ b/Engine/Source/Runtime/Engine/Public/MallocProfilerEx.h
@@ -30,12 +30,12 @@ public:
 	 *
 	 * @param	InWorld		World Context.
 	 */
-	virtual void WriteLoadedLevels( UWorld* InWorld ) OVERRIDE;
+	virtual void WriteLoadedLevels( UWorld* InWorld ) override;
 
 	/** 
 	 * Gather texture memory stats. 
 	 */
-	virtual void GetTexturePoolSize( FGenericMemoryStats& out_Stats ) OVERRIDE;
+	virtual void GetTexturePoolSize( FGenericMemoryStats& out_Stats ) override;
 };
 
 #endif // USE_MALLOC_PROFILER
diff --git a/Engine/Source/Runtime/Engine/Public/MaterialCompiler.h b/Engine/Source/Runtime/Engine/Public/MaterialCompiler.h
--- a/Engine/Source/Runtime/Engine/Public/MaterialCompiler.h
+++ b/Engine/Source/Runtime/Engine/Public/MaterialCompiler.h
@@ -188,161 +188,161 @@ public:
 
 	// Simple pass through all other material operations unmodified.
 
-	virtual void SetMaterialProperty(EMaterialProperty InProperty, EShaderFrequency InShaderFrequency) OVERRIDE { Compiler->SetMaterialProperty(InProperty, InShaderFrequency); }
-	virtual int32 Error(const TCHAR* Text) OVERRIDE { return Compiler->Error(Text); }
+	virtual void SetMaterialProperty(EMaterialProperty InProperty, EShaderFrequency InShaderFrequency) override { Compiler->SetMaterialProperty(InProperty, InShaderFrequency); }
+	virtual int32 Error(const TCHAR* Text) override { return Compiler->Error(Text); }
 
-	virtual int32 CallExpression(FMaterialExpressionKey ExpressionKey,FMaterialCompiler* InCompiler) OVERRIDE { return Compiler->CallExpression(ExpressionKey,InCompiler); }
+	virtual int32 CallExpression(FMaterialExpressionKey ExpressionKey,FMaterialCompiler* InCompiler) override { return Compiler->CallExpression(ExpressionKey,InCompiler); }
 
-	virtual void PushFunction(const FMaterialFunctionCompileState& FunctionState) OVERRIDE { Compiler->PushFunction(FunctionState); }
-	virtual FMaterialFunctionCompileState PopFunction() OVERRIDE { return Compiler->PopFunction(); }
+	virtual void PushFunction(const FMaterialFunctionCompileState& FunctionState) override { Compiler->PushFunction(FunctionState); }
+	virtual FMaterialFunctionCompileState PopFunction() override { return Compiler->PopFunction(); }
 
-	virtual EMaterialValueType GetType(int32 Code) OVERRIDE { return Compiler->GetType(Code); }
-	virtual EMaterialQualityLevel::Type GetQualityLevel() OVERRIDE { return Compiler->GetQualityLevel(); }
-	virtual ERHIFeatureLevel::Type GetFeatureLevel() OVERRIDE { return Compiler->GetFeatureLevel(); }
-	virtual float GetRefractionDepthBiasValue() OVERRIDE { return Compiler->GetRefractionDepthBiasValue(); }
-	virtual int32 ValidCast(int32 Code,EMaterialValueType DestType) OVERRIDE { return Compiler->ValidCast(Code, DestType); }
-	virtual int32 ForceCast(int32 Code,EMaterialValueType DestType,bool bExactMatch=false,bool bReplicateValue=false) OVERRIDE
+	virtual EMaterialValueType GetType(int32 Code) override { return Compiler->GetType(Code); }
+	virtual EMaterialQualityLevel::Type GetQualityLevel() override { return Compiler->GetQualityLevel(); }
+	virtual ERHIFeatureLevel::Type GetFeatureLevel() override { return Compiler->GetFeatureLevel(); }
+	virtual float GetRefractionDepthBiasValue() override { return Compiler->GetRefractionDepthBiasValue(); }
+	virtual int32 ValidCast(int32 Code,EMaterialValueType DestType) override { return Compiler->ValidCast(Code, DestType); }
+	virtual int32 ForceCast(int32 Code,EMaterialValueType DestType,bool bExactMatch=false,bool bReplicateValue=false) override
 	{ return Compiler->ForceCast(Code,DestType,bExactMatch,bReplicateValue); }
 
-	virtual int32 AccessCollectionParameter(UMaterialParameterCollection* ParameterCollection, int32 ParameterIndex, int32 ComponentIndex) OVERRIDE { return Compiler->AccessCollectionParameter(ParameterCollection, ParameterIndex, ComponentIndex); }
-	virtual int32 VectorParameter(FName ParameterName,const FLinearColor& DefaultValue) OVERRIDE { return Compiler->VectorParameter(ParameterName,DefaultValue); }
-	virtual int32 ScalarParameter(FName ParameterName,float DefaultValue) OVERRIDE { return Compiler->ScalarParameter(ParameterName,DefaultValue); }
-
-	virtual int32 Constant(float X) OVERRIDE { return Compiler->Constant(X); }
-	virtual int32 Constant2(float X,float Y) OVERRIDE { return Compiler->Constant2(X,Y); }
-	virtual int32 Constant3(float X,float Y,float Z) OVERRIDE { return Compiler->Constant3(X,Y,Z); }
-	virtual int32 Constant4(float X,float Y,float Z,float W) OVERRIDE { return Compiler->Constant4(X,Y,Z,W); }
-
-	virtual int32 GameTime() OVERRIDE { return Compiler->GameTime(); }
-	virtual int32 RealTime() OVERRIDE { return Compiler->RealTime(); }
-
-	virtual int32 PeriodicHint(int32 PeriodicCode) OVERRIDE { return Compiler->PeriodicHint(PeriodicCode); }
-
-	virtual int32 Sine(int32 X) OVERRIDE { return Compiler->Sine(X); }
-	virtual int32 Cosine(int32 X) OVERRIDE { return Compiler->Cosine(X); }
-
-	virtual int32 Floor(int32 X) OVERRIDE { return Compiler->Floor(X); }
-	virtual int32 Ceil(int32 X) OVERRIDE { return Compiler->Ceil(X); }
-	virtual int32 Frac(int32 X) OVERRIDE { return Compiler->Frac(X); }
-	virtual int32 Fmod(int32 A, int32 B) OVERRIDE { return Compiler->Fmod(A,B); }
-	virtual int32 Abs(int32 X) OVERRIDE { return Compiler->Abs(X); }
-
-	virtual int32 ReflectionVector() OVERRIDE { return Compiler->ReflectionVector(); }
-	virtual int32 CameraVector() OVERRIDE { return Compiler->CameraVector(); }
-	virtual int32 CameraWorldPosition() OVERRIDE { return Compiler->CameraWorldPosition(); }
-	virtual int32 LightVector() OVERRIDE { return Compiler->LightVector(); }
-
-	virtual int32 ScreenPosition() OVERRIDE { return Compiler->ScreenPosition(); }
-	virtual int32 ViewSize() OVERRIDE { return Compiler->ViewSize(); }
-	virtual int32 SceneTexelSize() OVERRIDE { return Compiler->SceneTexelSize(); }
-	virtual int32 WorldPosition(EWorldPositionIncludedOffsets WorldPositionIncludedOffsets) OVERRIDE { return Compiler->WorldPosition(WorldPositionIncludedOffsets); }
-	virtual int32 ObjectWorldPosition() OVERRIDE { return Compiler->ObjectWorldPosition(); }
-	virtual int32 ObjectRadius() OVERRIDE { return Compiler->ObjectRadius(); }
-	virtual int32 ObjectBounds() OVERRIDE { return Compiler->ObjectBounds(); }
-	virtual int32 DistanceCullFade() OVERRIDE { return Compiler->DistanceCullFade(); }
-	virtual int32 ActorWorldPosition() OVERRIDE { return Compiler->ActorWorldPosition(); }
-	virtual int32 ParticleMacroUV() OVERRIDE { return Compiler->ParticleMacroUV(); }
-	virtual int32 ParticleSubUV(int32 TextureIndex, EMaterialSamplerType SamplerType, bool bBlend) OVERRIDE { return Compiler->ParticleSubUV(TextureIndex, SamplerType, bBlend); }
-	virtual int32 ParticleColor() OVERRIDE { return Compiler->ParticleColor(); }
-	virtual int32 ParticlePosition() OVERRIDE { return Compiler->ParticlePosition(); }
-	virtual int32 ParticleRadius() OVERRIDE { return Compiler->ParticleRadius(); }
-	virtual int32 SphericalParticleOpacity(int32 Density) OVERRIDE { return Compiler->SphericalParticleOpacity(Density); }
-
-	virtual int32 If(int32 A,int32 B,int32 AGreaterThanB,int32 AEqualsB,int32 ALessThanB,int32 Threshold) OVERRIDE { return Compiler->If(A,B,AGreaterThanB,AEqualsB,ALessThanB,Threshold); }
-
-	virtual int32 TextureSample(int32 InTexture,int32 Coordinate,EMaterialSamplerType SamplerType,int32 MipValueIndex=INDEX_NONE,ETextureMipValueMode MipValueMode=TMVM_None) OVERRIDE { return Compiler->TextureSample(InTexture,Coordinate,SamplerType,MipValueIndex,MipValueMode); }
-	virtual int32 TextureCoordinate(uint32 CoordinateIndex, bool UnMirrorU, bool UnMirrorV) OVERRIDE { return Compiler->TextureCoordinate(CoordinateIndex, UnMirrorU, UnMirrorV); }
-
-	virtual int32 Texture(UTexture* InTexture) OVERRIDE { return Compiler->Texture(InTexture); }
-	virtual int32 TextureParameter(FName ParameterName,UTexture* DefaultValue) OVERRIDE { return Compiler->TextureParameter(ParameterName,DefaultValue); }
-
-	virtual	int32 PixelDepth() OVERRIDE { return Compiler->PixelDepth();	}
-	virtual int32 SceneDepth(int32 Offset, int32 UV, bool bUseOffset) OVERRIDE { return Compiler->SceneDepth(Offset, UV, bUseOffset); }
-	virtual int32 SceneColor(int32 Offset, int32 UV, bool bUseOffset) OVERRIDE { return Compiler->SceneColor(Offset, UV, bUseOffset); }
-	virtual int32 SceneTextureLookup(int32 UV, uint32 InSceneTextureId) OVERRIDE { return Compiler->SceneTextureLookup(UV, InSceneTextureId); }
-	virtual int32 SceneTextureSize(uint32 InSceneTextureId, bool bInvert) OVERRIDE { return Compiler->SceneTextureSize(InSceneTextureId, bInvert); }
-	virtual int32 SceneTextureMax(uint32 InSceneTextureId) OVERRIDE { return Compiler->SceneTextureMax(InSceneTextureId); }
-	virtual int32 SceneTextureMin(uint32 InSceneTextureId) OVERRIDE { return Compiler->SceneTextureMin(InSceneTextureId); }
-
-	virtual int32 StaticBool(bool Value) OVERRIDE { return Compiler->StaticBool(Value); }
-	virtual int32 StaticBoolParameter(FName ParameterName,bool bDefaultValue) OVERRIDE { return Compiler->StaticBoolParameter(ParameterName,bDefaultValue); }
-	virtual int32 StaticComponentMask(int32 Vector,FName ParameterName,bool bDefaultR,bool bDefaultG,bool bDefaultB,bool bDefaultA) OVERRIDE { return Compiler->StaticComponentMask(Vector,ParameterName,bDefaultR,bDefaultG,bDefaultB,bDefaultA); }
-	virtual bool GetStaticBoolValue(int32 BoolIndex, bool& bSucceeded) OVERRIDE { return Compiler->GetStaticBoolValue(BoolIndex, bSucceeded); }
-	virtual int32 StaticTerrainLayerWeight(FName ParameterName,int32 Default) OVERRIDE { return Compiler->StaticTerrainLayerWeight(ParameterName,Default); }
-
-	virtual int32 VertexColor() OVERRIDE { return Compiler->VertexColor(); }
-
-	virtual int32 Add(int32 A,int32 B) OVERRIDE { return Compiler->Add(A,B); }
-	virtual int32 Sub(int32 A,int32 B) OVERRIDE { return Compiler->Sub(A,B); }
-	virtual int32 Mul(int32 A,int32 B) OVERRIDE { return Compiler->Mul(A,B); }
-	virtual int32 Div(int32 A,int32 B) OVERRIDE { return Compiler->Div(A,B); }
-	virtual int32 Dot(int32 A,int32 B) OVERRIDE { return Compiler->Dot(A,B); }
-	virtual int32 Cross(int32 A,int32 B) OVERRIDE { return Compiler->Cross(A,B); }
-
-	virtual int32 Power(int32 Base,int32 Exponent) OVERRIDE { return Compiler->Power(Base,Exponent); }
-	virtual int32 SquareRoot(int32 X) OVERRIDE { return Compiler->SquareRoot(X); }
-	virtual int32 Length(int32 X) OVERRIDE { return Compiler->Length(X); }
-
-	virtual int32 Lerp(int32 X,int32 Y,int32 A) OVERRIDE { return Compiler->Lerp(X,Y,A); }
-	virtual int32 Min(int32 A,int32 B) OVERRIDE { return Compiler->Min(A,B); }
-	virtual int32 Max(int32 A,int32 B) OVERRIDE { return Compiler->Max(A,B); }
-	virtual int32 Clamp(int32 X,int32 A,int32 B) OVERRIDE { return Compiler->Clamp(X,A,B); }
-
-	virtual int32 ComponentMask(int32 Vector,bool R,bool G,bool B,bool A) OVERRIDE { return Compiler->ComponentMask(Vector,R,G,B,A); }
-	virtual int32 AppendVector(int32 A,int32 B) OVERRIDE { return Compiler->AppendVector(A,B); }
-	virtual int32 TransformVector(uint8 SourceCoordType,uint8 DestCoordType,int32 A) OVERRIDE { return Compiler->TransformVector(SourceCoordType,DestCoordType,A); }
-	virtual int32 TransformPosition(uint8 SourceCoordType,uint8 DestCoordType,int32 A) OVERRIDE { return Compiler->TransformPosition(SourceCoordType,DestCoordType,A); }
-
-	virtual int32 DynamicParameter() OVERRIDE { return Compiler->DynamicParameter(); }
-	virtual int32 LightmapUVs() OVERRIDE { return Compiler->LightmapUVs(); }
-
-	virtual int32 LightmassReplace(int32 Realtime, int32 Lightmass) OVERRIDE { return Realtime; }
-	virtual int32 GIReplace(int32 Direct, int32 StaticIndirect, int32 DynamicIndirect) OVERRIDE { return Compiler->GIReplace(Direct, StaticIndirect, DynamicIndirect); }
-	virtual int32 ObjectOrientation() OVERRIDE { return Compiler->ObjectOrientation(); }
-	virtual int32 RotateAboutAxis(int32 NormalizedRotationAxisAndAngleIndex, int32 PositionOnAxisIndex, int32 PositionIndex) OVERRIDE
+	virtual int32 AccessCollectionParameter(UMaterialParameterCollection* ParameterCollection, int32 ParameterIndex, int32 ComponentIndex) override { return Compiler->AccessCollectionParameter(ParameterCollection, ParameterIndex, ComponentIndex); }
+	virtual int32 VectorParameter(FName ParameterName,const FLinearColor& DefaultValue) override { return Compiler->VectorParameter(ParameterName,DefaultValue); }
+	virtual int32 ScalarParameter(FName ParameterName,float DefaultValue) override { return Compiler->ScalarParameter(ParameterName,DefaultValue); }
+
+	virtual int32 Constant(float X) override { return Compiler->Constant(X); }
+	virtual int32 Constant2(float X,float Y) override { return Compiler->Constant2(X,Y); }
+	virtual int32 Constant3(float X,float Y,float Z) override { return Compiler->Constant3(X,Y,Z); }
+	virtual int32 Constant4(float X,float Y,float Z,float W) override { return Compiler->Constant4(X,Y,Z,W); }
+
+	virtual int32 GameTime() override { return Compiler->GameTime(); }
+	virtual int32 RealTime() override { return Compiler->RealTime(); }
+
+	virtual int32 PeriodicHint(int32 PeriodicCode) override { return Compiler->PeriodicHint(PeriodicCode); }
+
+	virtual int32 Sine(int32 X) override { return Compiler->Sine(X); }
+	virtual int32 Cosine(int32 X) override { return Compiler->Cosine(X); }
+
+	virtual int32 Floor(int32 X) override { return Compiler->Floor(X); }
+	virtual int32 Ceil(int32 X) override { return Compiler->Ceil(X); }
+	virtual int32 Frac(int32 X) override { return Compiler->Frac(X); }
+	virtual int32 Fmod(int32 A, int32 B) override { return Compiler->Fmod(A,B); }
+	virtual int32 Abs(int32 X) override { return Compiler->Abs(X); }
+
+	virtual int32 ReflectionVector() override { return Compiler->ReflectionVector(); }
+	virtual int32 CameraVector() override { return Compiler->CameraVector(); }
+	virtual int32 CameraWorldPosition() override { return Compiler->CameraWorldPosition(); }
+	virtual int32 LightVector() override { return Compiler->LightVector(); }
+
+	virtual int32 ScreenPosition() override { return Compiler->ScreenPosition(); }
+	virtual int32 ViewSize() override { return Compiler->ViewSize(); }
+	virtual int32 SceneTexelSize() override { return Compiler->SceneTexelSize(); }
+	virtual int32 WorldPosition(EWorldPositionIncludedOffsets WorldPositionIncludedOffsets) override { return Compiler->WorldPosition(WorldPositionIncludedOffsets); }
+	virtual int32 ObjectWorldPosition() override { return Compiler->ObjectWorldPosition(); }
+	virtual int32 ObjectRadius() override { return Compiler->ObjectRadius(); }
+	virtual int32 ObjectBounds() override { return Compiler->ObjectBounds(); }
+	virtual int32 DistanceCullFade() override { return Compiler->DistanceCullFade(); }
+	virtual int32 ActorWorldPosition() override { return Compiler->ActorWorldPosition(); }
+	virtual int32 ParticleMacroUV() override { return Compiler->ParticleMacroUV(); }
+	virtual int32 ParticleSubUV(int32 TextureIndex, EMaterialSamplerType SamplerType, bool bBlend) override { return Compiler->ParticleSubUV(TextureIndex, SamplerType, bBlend); }
+	virtual int32 ParticleColor() override { return Compiler->ParticleColor(); }
+	virtual int32 ParticlePosition() override { return Compiler->ParticlePosition(); }
+	virtual int32 ParticleRadius() override { return Compiler->ParticleRadius(); }
+	virtual int32 SphericalParticleOpacity(int32 Density) override { return Compiler->SphericalParticleOpacity(Density); }
+
+	virtual int32 If(int32 A,int32 B,int32 AGreaterThanB,int32 AEqualsB,int32 ALessThanB,int32 Threshold) override { return Compiler->If(A,B,AGreaterThanB,AEqualsB,ALessThanB,Threshold); }
+
+	virtual int32 TextureSample(int32 InTexture,int32 Coordinate,EMaterialSamplerType SamplerType,int32 MipValueIndex=INDEX_NONE,ETextureMipValueMode MipValueMode=TMVM_None) override { return Compiler->TextureSample(InTexture,Coordinate,SamplerType,MipValueIndex,MipValueMode); }
+	virtual int32 TextureCoordinate(uint32 CoordinateIndex, bool UnMirrorU, bool UnMirrorV) override { return Compiler->TextureCoordinate(CoordinateIndex, UnMirrorU, UnMirrorV); }
+
+	virtual int32 Texture(UTexture* InTexture) override { return Compiler->Texture(InTexture); }
+	virtual int32 TextureParameter(FName ParameterName,UTexture* DefaultValue) override { return Compiler->TextureParameter(ParameterName,DefaultValue); }
+
+	virtual	int32 PixelDepth() override { return Compiler->PixelDepth();	}
+	virtual int32 SceneDepth(int32 Offset, int32 UV, bool bUseOffset) override { return Compiler->SceneDepth(Offset, UV, bUseOffset); }
+	virtual int32 SceneColor(int32 Offset, int32 UV, bool bUseOffset) override { return Compiler->SceneColor(Offset, UV, bUseOffset); }
+	virtual int32 SceneTextureLookup(int32 UV, uint32 InSceneTextureId) override { return Compiler->SceneTextureLookup(UV, InSceneTextureId); }
+	virtual int32 SceneTextureSize(uint32 InSceneTextureId, bool bInvert) override { return Compiler->SceneTextureSize(InSceneTextureId, bInvert); }
+	virtual int32 SceneTextureMax(uint32 InSceneTextureId) override { return Compiler->SceneTextureMax(InSceneTextureId); }
+	virtual int32 SceneTextureMin(uint32 InSceneTextureId) override { return Compiler->SceneTextureMin(InSceneTextureId); }
+
+	virtual int32 StaticBool(bool Value) override { return Compiler->StaticBool(Value); }
+	virtual int32 StaticBoolParameter(FName ParameterName,bool bDefaultValue) override { return Compiler->StaticBoolParameter(ParameterName,bDefaultValue); }
+	virtual int32 StaticComponentMask(int32 Vector,FName ParameterName,bool bDefaultR,bool bDefaultG,bool bDefaultB,bool bDefaultA) override { return Compiler->StaticComponentMask(Vector,ParameterName,bDefaultR,bDefaultG,bDefaultB,bDefaultA); }
+	virtual bool GetStaticBoolValue(int32 BoolIndex, bool& bSucceeded) override { return Compiler->GetStaticBoolValue(BoolIndex, bSucceeded); }
+	virtual int32 StaticTerrainLayerWeight(FName ParameterName,int32 Default) override { return Compiler->StaticTerrainLayerWeight(ParameterName,Default); }
+
+	virtual int32 VertexColor() override { return Compiler->VertexColor(); }
+
+	virtual int32 Add(int32 A,int32 B) override { return Compiler->Add(A,B); }
+	virtual int32 Sub(int32 A,int32 B) override { return Compiler->Sub(A,B); }
+	virtual int32 Mul(int32 A,int32 B) override { return Compiler->Mul(A,B); }
+	virtual int32 Div(int32 A,int32 B) override { return Compiler->Div(A,B); }
+	virtual int32 Dot(int32 A,int32 B) override { return Compiler->Dot(A,B); }
+	virtual int32 Cross(int32 A,int32 B) override { return Compiler->Cross(A,B); }
+
+	virtual int32 Power(int32 Base,int32 Exponent) override { return Compiler->Power(Base,Exponent); }
+	virtual int32 SquareRoot(int32 X) override { return Compiler->SquareRoot(X); }
+	virtual int32 Length(int32 X) override { return Compiler->Length(X); }
+
+	virtual int32 Lerp(int32 X,int32 Y,int32 A) override { return Compiler->Lerp(X,Y,A); }
+	virtual int32 Min(int32 A,int32 B) override { return Compiler->Min(A,B); }
+	virtual int32 Max(int32 A,int32 B) override { return Compiler->Max(A,B); }
+	virtual int32 Clamp(int32 X,int32 A,int32 B) override { return Compiler->Clamp(X,A,B); }
+
+	virtual int32 ComponentMask(int32 Vector,bool R,bool G,bool B,bool A) override { return Compiler->ComponentMask(Vector,R,G,B,A); }
+	virtual int32 AppendVector(int32 A,int32 B) override { return Compiler->AppendVector(A,B); }
+	virtual int32 TransformVector(uint8 SourceCoordType,uint8 DestCoordType,int32 A) override { return Compiler->TransformVector(SourceCoordType,DestCoordType,A); }
+	virtual int32 TransformPosition(uint8 SourceCoordType,uint8 DestCoordType,int32 A) override { return Compiler->TransformPosition(SourceCoordType,DestCoordType,A); }
+
+	virtual int32 DynamicParameter() override { return Compiler->DynamicParameter(); }
+	virtual int32 LightmapUVs() override { return Compiler->LightmapUVs(); }
+
+	virtual int32 LightmassReplace(int32 Realtime, int32 Lightmass) override { return Realtime; }
+	virtual int32 GIReplace(int32 Direct, int32 StaticIndirect, int32 DynamicIndirect) override { return Compiler->GIReplace(Direct, StaticIndirect, DynamicIndirect); }
+	virtual int32 ObjectOrientation() override { return Compiler->ObjectOrientation(); }
+	virtual int32 RotateAboutAxis(int32 NormalizedRotationAxisAndAngleIndex, int32 PositionOnAxisIndex, int32 PositionIndex) override
 	{
 		return Compiler->RotateAboutAxis(NormalizedRotationAxisAndAngleIndex, PositionOnAxisIndex, PositionIndex);
 	}
-	virtual int32 TwoSidedSign() OVERRIDE { return Compiler->TwoSidedSign(); }
-	virtual int32 VertexNormal() OVERRIDE { return Compiler->VertexNormal(); }
-	virtual int32 PixelNormalWS() OVERRIDE { return Compiler->PixelNormalWS(); }
+	virtual int32 TwoSidedSign() override { return Compiler->TwoSidedSign(); }
+	virtual int32 VertexNormal() override { return Compiler->VertexNormal(); }
+	virtual int32 PixelNormalWS() override { return Compiler->PixelNormalWS(); }
 
-	virtual int32 CustomExpression( class UMaterialExpressionCustom* Custom, TArray<int32>& CompiledInputs ) OVERRIDE { return Compiler->CustomExpression(Custom,CompiledInputs); }
-	virtual int32 DDX(int32 X) OVERRIDE { return Compiler->DDX(X); }
-	virtual int32 DDY(int32 X) OVERRIDE { return Compiler->DDY(X); }
+	virtual int32 CustomExpression( class UMaterialExpressionCustom* Custom, TArray<int32>& CompiledInputs ) override { return Compiler->CustomExpression(Custom,CompiledInputs); }
+	virtual int32 DDX(int32 X) override { return Compiler->DDX(X); }
+	virtual int32 DDY(int32 X) override { return Compiler->DDY(X); }
 
-	virtual int32 AntialiasedTextureMask(int32 Tex, int32 UV, float Threshold, uint8 Channel) OVERRIDE
+	virtual int32 AntialiasedTextureMask(int32 Tex, int32 UV, float Threshold, uint8 Channel) override
 	{
 		return Compiler->AntialiasedTextureMask(Tex, UV, Threshold, Channel);
 	}
-	virtual int32 Noise(int32 Position, float Scale, int32 Quality, uint8 NoiseFunction, bool bTurbulence, int32 Levels, float OutputMin, float OutputMax, float LevelScale, int32 FilterWidth) OVERRIDE
+	virtual int32 Noise(int32 Position, float Scale, int32 Quality, uint8 NoiseFunction, bool bTurbulence, int32 Levels, float OutputMin, float OutputMax, float LevelScale, int32 FilterWidth) override
 	{
 		return Compiler->Noise(Position, Scale, Quality, NoiseFunction, bTurbulence, Levels, OutputMin, OutputMax, LevelScale, FilterWidth);
 	}
-	virtual int32 BlackBody( int32 Temp ) OVERRIDE { return Compiler->BlackBody(Temp); }
-	virtual int32 PerInstanceRandom() OVERRIDE { return Compiler->PerInstanceRandom(); }
-	virtual int32 PerInstanceFadeAmount() OVERRIDE { return Compiler->PerInstanceFadeAmount(); }
-	virtual int32 DepthOfFieldFunction(int32 Depth, int32 FunctionValueIndex) OVERRIDE
+	virtual int32 BlackBody( int32 Temp ) override { return Compiler->BlackBody(Temp); }
+	virtual int32 PerInstanceRandom() override { return Compiler->PerInstanceRandom(); }
+	virtual int32 PerInstanceFadeAmount() override { return Compiler->PerInstanceFadeAmount(); }
+	virtual int32 DepthOfFieldFunction(int32 Depth, int32 FunctionValueIndex) override
 	{
 		return Compiler->DepthOfFieldFunction(Depth, FunctionValueIndex);
 	}
 
-	virtual int32 SpeedTree(ESpeedTreeGeometryType GeometryType, ESpeedTreeWindType WindType, ESpeedTreeLODType LODType, float BillboardThreshold) OVERRIDE 
+	virtual int32 SpeedTree(ESpeedTreeGeometryType GeometryType, ESpeedTreeWindType WindType, ESpeedTreeLODType LODType, float BillboardThreshold) override 
 	{ 
 		return Compiler->SpeedTree(GeometryType, WindType, LODType, BillboardThreshold); 
 	}
 
-	virtual int32 AtmosphericFogColor(int32 WorldPosition) OVERRIDE
+	virtual int32 AtmosphericFogColor(int32 WorldPosition) override
 	{
 		return Compiler->AtmosphericFogColor(WorldPosition);
 	}
 
-	virtual int32 TextureCoordinateOffset() OVERRIDE
+	virtual int32 TextureCoordinateOffset() override
 	{
 		return Compiler->TextureCoordinateOffset();
 	}
 
-	virtual int32 EyeAdaptation() OVERRIDE
+	virtual int32 EyeAdaptation() override
 	{
 		return Compiler->EyeAdaptation();
 	}
diff --git a/Engine/Source/Runtime/Engine/Public/MaterialShared.h b/Engine/Source/Runtime/Engine/Public/MaterialShared.h
--- a/Engine/Source/Runtime/Engine/Public/MaterialShared.h
+++ b/Engine/Source/Runtime/Engine/Public/MaterialShared.h
@@ -1522,7 +1522,7 @@ public:
 
 	// FMaterial interface.
 	ENGINE_API virtual void GetShaderMapId(EShaderPlatform Platform, FMaterialShaderMapId& OutId) const;
-	ENGINE_API virtual int32 GetMaterialDomain() const OVERRIDE;
+	ENGINE_API virtual int32 GetMaterialDomain() const override;
 	ENGINE_API virtual bool IsTwoSided() const;
 	ENGINE_API virtual bool IsTangentSpaceNormal() const;
 	ENGINE_API virtual bool ShouldInjectEmissiveIntoLPV() const;
@@ -1577,9 +1577,9 @@ public:
 	 * Should shaders compiled for this material be saved to disk?
 	 */
 	ENGINE_API virtual bool IsPersistent() const;
-	ENGINE_API virtual FGuid GetMaterialId() const OVERRIDE;
+	ENGINE_API virtual FGuid GetMaterialId() const override;
 
-	ENGINE_API virtual void NotifyCompilationFinished() OVERRIDE;
+	ENGINE_API virtual void NotifyCompilationFinished() override;
 
 	/**
 	 * Gets instruction counts that best represent the likely usage of this material based on shading model and other factors.
@@ -1595,9 +1595,9 @@ public:
 
 	ENGINE_API virtual void LegacySerialize(FArchive& Ar);
 
-	ENGINE_API virtual const TArray<UTexture*>& GetReferencedTextures() const OVERRIDE;
+	ENGINE_API virtual const TArray<UTexture*>& GetReferencedTextures() const override;
 
-	ENGINE_API virtual bool GetAllowDevelopmentShaderCompile() const OVERRIDE;
+	ENGINE_API virtual bool GetAllowDevelopmentShaderCompile() const override;
 protected:
 	UMaterial* Material;
 	UMaterialInstance* MaterialInstance;
diff --git a/Engine/Source/Runtime/Engine/Public/Model.h b/Engine/Source/Runtime/Engine/Public/Model.h
--- a/Engine/Source/Runtime/Engine/Public/Model.h
+++ b/Engine/Source/Runtime/Engine/Public/Model.h
@@ -436,9 +436,9 @@ public:
 	virtual void Serialize( FArchive& Ar );	
 	virtual void PostLoad();
 #if WITH_EDITOR
-	virtual void PostEditUndo() OVERRIDE;
+	virtual void PostEditUndo() override;
 #endif // WITH_EDITOR
-	virtual bool Modify( bool bAlwaysMarkDirty=false ) OVERRIDE;
+	virtual bool Modify( bool bAlwaysMarkDirty=false ) override;
 	virtual bool Rename( const TCHAR* InName=NULL, UObject* NewOuter=NULL, ERenameFlags Flags=REN_None );
 	static void AddReferencedObjects(UObject* InThis, FReferenceCollector& Collector);
 
@@ -462,12 +462,12 @@ public:
 	virtual void BeginDestroy();
 	virtual bool IsReadyForFinishDestroy();
 	
-	virtual bool IsAsset() const OVERRIDE { return false; }
+	virtual bool IsAsset() const override { return false; }
 
 	/**
 	* @return		Sum of the size of textures referenced by this material.
 	*/
-	virtual SIZE_T GetResourceSize(EResourceSizeMode::Type Mode) OVERRIDE;
+	virtual SIZE_T GetResourceSize(EResourceSizeMode::Type Mode) override;
 
 	// UModel interface.
 	ENGINE_API void EmptyModel( int32 EmptySurfInfo, int32 EmptyPolys );
diff --git a/Engine/Source/Runtime/Engine/Public/ModelLight.h b/Engine/Source/Runtime/Engine/Public/ModelLight.h
--- a/Engine/Source/Runtime/Engine/Public/ModelLight.h
+++ b/Engine/Source/Runtime/Engine/Public/ModelLight.h
@@ -41,7 +41,7 @@ public:
 	virtual FLightRayIntersection IntersectLightRay(const FVector& Start,const FVector& End,bool bFindNearestIntersection) const;
 	//FStaticLightingTextureMapping interface.
 
-	virtual bool IsValidMapping() const OVERRIDE
+	virtual bool IsValidMapping() const override
 	{
 		return Model.IsValid() && !Model->bInvalidForStaticLighting;
 	}
diff --git a/Engine/Source/Runtime/Engine/Public/Net/RepLayout.h b/Engine/Source/Runtime/Engine/Public/Net/RepLayout.h
--- a/Engine/Source/Runtime/Engine/Public/Net/RepLayout.h
+++ b/Engine/Source/Runtime/Engine/Public/Net/RepLayout.h
@@ -28,7 +28,7 @@ public:
 	FRepChangedPropertyTracker() : LastReplicationGroupFrame( 0 ), LastReplicationFrame( 0 ), ActiveStatusChanged( false ), UnconditionalPropChanged( false ) { }
 	virtual ~FRepChangedPropertyTracker() { }
 
-	virtual void SetCustomIsActiveOverride( const uint16 RepIndex, const bool bIsActive ) OVERRIDE
+	virtual void SetCustomIsActiveOverride( const uint16 RepIndex, const bool bIsActive ) override
 	{
 		FRepChangedParent & Parent = Parents[RepIndex];
 
diff --git a/Engine/Source/Runtime/Engine/Public/ParticleEmitterInstances.h b/Engine/Source/Runtime/Engine/Public/ParticleEmitterInstances.h
--- a/Engine/Source/Runtime/Engine/Public/ParticleEmitterInstances.h
+++ b/Engine/Source/Runtime/Engine/Public/ParticleEmitterInstances.h
@@ -738,7 +738,7 @@ struct FParticleSpriteEmitterInstance : public FParticleEmitterInstance
 	 * @param	Mode	Specifies which resource size should be displayed. ( see EResourceSizeMode::Type )
 	 * @return  Size of resource as to be displayed to artists/ LDs in the Editor.
 	 */
-	virtual SIZE_T GetResourceSize(EResourceSizeMode::Type Mode) OVERRIDE;
+	virtual SIZE_T GetResourceSize(EResourceSizeMode::Type Mode) override;
 
 protected:
 
@@ -776,7 +776,7 @@ struct ENGINE_API FParticleMeshEmitterInstance : public FParticleEmitterInstance
 	virtual uint32 RequiredBytes();
 	virtual void PostSpawn(FBaseParticle* Particle, float InterpolationPercentage, float SpawnTime);
 	virtual FDynamicEmitterDataBase* GetDynamicData(bool bSelected);
-	virtual bool IsDynamicDataRequired(UParticleLODLevel* CurrentLODLevel) OVERRIDE;
+	virtual bool IsDynamicDataRequired(UParticleLODLevel* CurrentLODLevel) override;
 
 	/**
 	 *	Updates the dynamic data for the instance
@@ -807,12 +807,12 @@ struct ENGINE_API FParticleMeshEmitterInstance : public FParticleEmitterInstance
 	 * @param	Mode	Specifies which resource size should be displayed. ( see EResourceSizeMode::Type )
 	 * @return  Size of resource as to be displayed to artists/ LDs in the Editor.
 	 */
-	virtual SIZE_T GetResourceSize(EResourceSizeMode::Type Mode) OVERRIDE;
+	virtual SIZE_T GetResourceSize(EResourceSizeMode::Type Mode) override;
 
 	/**
 	 * Returns the offset to the mesh rotation payload, if any.
 	 */
-	virtual int32 GetMeshRotationOffset() const OVERRIDE
+	virtual int32 GetMeshRotationOffset() const override
 	{
 		return MeshRotationOffset;
 	}
@@ -820,7 +820,7 @@ struct ENGINE_API FParticleMeshEmitterInstance : public FParticleEmitterInstance
 	/**
 	 * Returns true if mesh rotation is active.
 	 */
-	virtual bool IsMeshRotationActive() const OVERRIDE
+	virtual bool IsMeshRotationActive() const override
 	{
 		return MeshRotationActive;
 	}
@@ -829,14 +829,14 @@ struct ENGINE_API FParticleMeshEmitterInstance : public FParticleEmitterInstance
 	 * Sets the materials with which mesh particles should be rendered.
 	 * @param InMaterials - The materials.
 	 */
-	virtual void SetMeshMaterials( const TArray<UMaterialInterface*>& InMaterials ) OVERRIDE;
+	virtual void SetMeshMaterials( const TArray<UMaterialInterface*>& InMaterials ) override;
 
 	/**
 	 * Gathers material relevance flags for this emitter instance.
 	 * @param OutMaterialRelevance - Pointer to where material relevance flags will be stored.
 	 * @param LODLevel - The LOD level for which to compute material relevance flags.
 	 */
-	virtual void GatherMaterialRelevance( FMaterialRelevance* OutMaterialRelevance, const UParticleLODLevel* LODLevel ) const OVERRIDE;
+	virtual void GatherMaterialRelevance( FMaterialRelevance* OutMaterialRelevance, const UParticleLODLevel* LODLevel ) const override;
 
 	/**
 	 * Gets the materials applied to each section of a mesh.
@@ -926,14 +926,14 @@ struct FParticleBeam2EmitterInstance : public FParticleEmitterInstance
 	virtual ~FParticleBeam2EmitterInstance();
 
 	// Beam interface.
-	virtual void SetBeamEndPoint(FVector NewEndPoint) OVERRIDE;
-	virtual void SetBeamSourcePoint(FVector NewSourcePoint,int32 SourceIndex) OVERRIDE;
-	virtual void SetBeamSourceTangent(FVector NewTangentPoint,int32 SourceIndex) OVERRIDE;
-	virtual void SetBeamSourceStrength(float NewSourceStrength,int32 SourceIndex) OVERRIDE;
-	virtual void SetBeamTargetPoint(FVector NewTargetPoint,int32 TargetIndex) OVERRIDE;
-	virtual void SetBeamTargetTangent(FVector NewTangentPoint,int32 TargetIndex) OVERRIDE;
-	virtual void SetBeamTargetStrength(float NewTargetStrength,int32 TargetIndex) OVERRIDE;
-	virtual void ApplyWorldOffset(FVector InOffset, bool bWorldShift) OVERRIDE;
+	virtual void SetBeamEndPoint(FVector NewEndPoint) override;
+	virtual void SetBeamSourcePoint(FVector NewSourcePoint,int32 SourceIndex) override;
+	virtual void SetBeamSourceTangent(FVector NewTangentPoint,int32 SourceIndex) override;
+	virtual void SetBeamSourceStrength(float NewSourceStrength,int32 SourceIndex) override;
+	virtual void SetBeamTargetPoint(FVector NewTargetPoint,int32 TargetIndex) override;
+	virtual void SetBeamTargetTangent(FVector NewTangentPoint,int32 TargetIndex) override;
+	virtual void SetBeamTargetStrength(float NewTargetStrength,int32 TargetIndex) override;
+	virtual void ApplyWorldOffset(FVector InOffset, bool bWorldShift) override;
 	
 	//
 	virtual void InitParameters(UParticleEmitter* InTemplate, UParticleSystemComponent* InComponent, bool bClearResources = true);
@@ -962,10 +962,10 @@ struct FParticleBeam2EmitterInstance : public FParticleEmitterInstance
 	 * @param	InterpolationPercentage		The percentage of the time slice it was spawned at
 	 * @param	SpawnTIme					The time it was spawned at
 	 */
-	virtual void PostSpawn(FBaseParticle* Particle, float InterpolationPercentage, float SpawnTime) OVERRIDE;
+	virtual void PostSpawn(FBaseParticle* Particle, float InterpolationPercentage, float SpawnTime) override;
 
 	virtual void UpdateBoundingBox(float DeltaTime);
-	virtual void ForceUpdateBoundingBox() OVERRIDE;
+	virtual void ForceUpdateBoundingBox() override;
 	virtual uint32 RequiredBytes();
 	float SpawnBeamParticles(float OldLeftover, float Rate, float DeltaTime, int32 Burst = 0, float BurstTime = 0.0f);
 	virtual void KillParticles();
@@ -1014,12 +1014,12 @@ struct FParticleBeam2EmitterInstance : public FParticleEmitterInstance
 	 * @param	Mode	Specifies which resource size should be displayed. ( see EResourceSizeMode::Type )
 	 * @return	Size of resource as to be displayed to artists/ LDs in the Editor.
 	 */
-	virtual SIZE_T GetResourceSize(EResourceSizeMode::Type Mode) OVERRIDE;
+	virtual SIZE_T GetResourceSize(EResourceSizeMode::Type Mode) override;
 
 	/**
 	 * When an emitter is killed, this will check other emitters and clean up anything pointing to this one
 	 */
-	virtual void OnEmitterInstanceKilled(FParticleEmitterInstance* Instance) OVERRIDE
+	virtual void OnEmitterInstanceKilled(FParticleEmitterInstance* Instance) override
 	{
 		if (SourceEmitter == Instance)
 		{
@@ -1118,7 +1118,7 @@ struct FParticleTrailsEmitterInstance_Base : public FParticleEmitterInstance
 	 */
 	virtual void Tick_RecalculateTangents(float DeltaTime, UParticleLODLevel* CurrentLODLevel);
 	virtual void UpdateBoundingBox(float DeltaTime);
-	virtual void ForceUpdateBoundingBox() OVERRIDE;
+	virtual void ForceUpdateBoundingBox() override;
 	virtual void KillParticles();
 
 	/**
@@ -1352,12 +1352,12 @@ struct FParticleRibbonEmitterInstance : public FParticleTrailsEmitterInstance_Ba
 	 * @param	Mode	Specifies which resource size should be displayed. ( see EResourceSizeMode::Type )
 	 * @return  Size of resource as to be displayed to artists/ LDs in the Editor.
 	 */
-	virtual SIZE_T GetResourceSize(EResourceSizeMode::Type Mode) OVERRIDE;
+	virtual SIZE_T GetResourceSize(EResourceSizeMode::Type Mode) override;
 
 	/**
 	 * When an emitter is killed, this will check other emitters and clean up anything pointing to this one
 	 */
-	virtual void OnEmitterInstanceKilled(FParticleEmitterInstance* Instance) OVERRIDE
+	virtual void OnEmitterInstanceKilled(FParticleEmitterInstance* Instance) override
 	{
 		if (SourceEmitter == Instance)
 		{
@@ -1514,7 +1514,7 @@ struct FParticleAnimTrailEmitterInstance : public FParticleTrailsEmitterInstance
 	 * @param	Mode	Specifies which resource size should be displayed. ( see EResourceSizeMode::Type )
 	 * @return  Size of resource as to be displayed to artists/ LDs in the Editor.
 	 */
-	virtual SIZE_T GetResourceSize(EResourceSizeMode::Type Mode) OVERRIDE;
+	virtual SIZE_T GetResourceSize(EResourceSizeMode::Type Mode) override;
 
 	virtual bool IsTrailEmitter()const{ return true; }
 
diff --git a/Engine/Source/Runtime/Engine/Public/ParticleHelper.h b/Engine/Source/Runtime/Engine/Public/ParticleHelper.h
--- a/Engine/Source/Runtime/Engine/Public/ParticleHelper.h
+++ b/Engine/Source/Runtime/Engine/Public/ParticleHelper.h
@@ -1697,7 +1697,7 @@ struct FDynamicSpriteEmitterData : public FDynamicSpriteEmitterDataBase
 	/**
 	 *	Get the vertex stride for the dynamic rendering data
 	 */
-	virtual int32 GetDynamicVertexStride(ERHIFeatureLevel::Type InFeatureLevel) const OVERRIDE
+	virtual int32 GetDynamicVertexStride(ERHIFeatureLevel::Type InFeatureLevel) const override
 	{
 		const bool bInstanced = InFeatureLevel >= ERHIFeatureLevel::SM3;
 		return bInstanced ? sizeof(FParticleSpriteVertex) : sizeof(FParticleSpriteVertexNonInstanced);
@@ -1706,7 +1706,7 @@ struct FDynamicSpriteEmitterData : public FDynamicSpriteEmitterDataBase
 	/**
 	 *	Get the vertex stride for the dynamic parameter rendering data
 	 */
-	virtual int32 GetDynamicParameterVertexStride() const OVERRIDE
+	virtual int32 GetDynamicParameterVertexStride() const override
 	{
 		return sizeof(FParticleVertexDynamicParameter);
 	}
@@ -1757,10 +1757,10 @@ struct FDynamicSpriteEmitterData : public FDynamicSpriteEmitterDataBase
 	 *	@param	VisibilityMap	A BitArray that indicates whether the primitive was visible in that view (index)
 	 *	@param	FrameNumber		The frame number of this pre-render
 	 */
-	virtual void PreRenderView(FParticleSystemSceneProxy* Proxy, const FSceneViewFamily* ViewFamily, const uint32 VisibilityMap, int32 FrameNumber) OVERRIDE;
+	virtual void PreRenderView(FParticleSystemSceneProxy* Proxy, const FSceneViewFamily* ViewFamily, const uint32 VisibilityMap, int32 FrameNumber) override;
 
 	/** Gathers simple lights for this emitter. */
-	virtual void GatherSimpleLights(const FParticleSystemSceneProxy* Proxy, const FSceneViewFamily& ViewFamily, FSimpleLightArray& OutParticleLights) const OVERRIDE;
+	virtual void GatherSimpleLights(const FParticleSystemSceneProxy* Proxy, const FSceneViewFamily& ViewFamily, FSimpleLightArray& OutParticleLights) const override;
 
 	/**
 	 *	Render thread only draw call
@@ -1769,7 +1769,7 @@ struct FDynamicSpriteEmitterData : public FDynamicSpriteEmitterDataBase
 	 *	@param	PDI			The primitive draw interface to render with
 	 *	@param	View		The scene view being rendered
 	 */
-	virtual int32 Render(FParticleSystemSceneProxy* Proxy, FPrimitiveDrawInterface* PDI,const FSceneView* View) OVERRIDE;
+	virtual int32 Render(FParticleSystemSceneProxy* Proxy, FPrimitiveDrawInterface* PDI,const FSceneView* View) override;
 
 	/**
 	 *	Create the render thread resources for this emitter data
@@ -1949,17 +1949,17 @@ struct FDynamicMeshEmitterData : public FDynamicSpriteEmitterDataBase
 	virtual void PreRenderView(FParticleSystemSceneProxy* Proxy, const FSceneViewFamily* ViewFamily, const uint32 VisibilityMap, int32 FrameNumber); 
 
 	/** Gathers simple lights for this emitter. */
-	virtual void GatherSimpleLights(const FParticleSystemSceneProxy* Proxy, const FSceneViewFamily& ViewFamily, FSimpleLightArray& OutParticleLights) const OVERRIDE;
+	virtual void GatherSimpleLights(const FParticleSystemSceneProxy* Proxy, const FSceneViewFamily& ViewFamily, FSimpleLightArray& OutParticleLights) const override;
 
 	/**
 	 *	Get the vertex stride for the dynamic rendering data
 	 */
-	virtual int32 GetDynamicVertexStride(ERHIFeatureLevel::Type /*InFeatureLevel*/) const OVERRIDE
+	virtual int32 GetDynamicVertexStride(ERHIFeatureLevel::Type /*InFeatureLevel*/) const override
 	{
 		return sizeof(FMeshParticleInstanceVertex);
 	}
 
-	virtual int32 GetDynamicParameterVertexStride() const OVERRIDE 
+	virtual int32 GetDynamicParameterVertexStride() const override 
 	{
 		return sizeof(FMeshParticleInstanceVertexDynamicParameter);
 	}
@@ -1967,7 +1967,7 @@ struct FDynamicMeshEmitterData : public FDynamicSpriteEmitterDataBase
 	/**
 	 *	Get the source replay data for this emitter
 	 */
-	virtual const FDynamicSpriteEmitterReplayDataBase* GetSourceData() const OVERRIDE
+	virtual const FDynamicSpriteEmitterReplayDataBase* GetSourceData() const override
 	{
 		return &Source;
 	}
@@ -2249,7 +2249,7 @@ struct FDynamicBeam2EmitterData : public FDynamicSpriteEmitterDataBase
 	/**
 	 *	Get the vertex stride for the dynamic rendering data
 	 */
-	virtual int32 GetDynamicVertexStride(ERHIFeatureLevel::Type /*InFeatureLevel*/) const OVERRIDE
+	virtual int32 GetDynamicVertexStride(ERHIFeatureLevel::Type /*InFeatureLevel*/) const override
 	{
 		return sizeof(FParticleBeamTrailVertex);
 	}
@@ -2257,7 +2257,7 @@ struct FDynamicBeam2EmitterData : public FDynamicSpriteEmitterDataBase
 	/**
 	 *	Get the vertex stride for the dynamic parameter rendering data
 	 */
-	virtual int32 GetDynamicParameterVertexStride() const OVERRIDE
+	virtual int32 GetDynamicParameterVertexStride() const override
 	{
 		return sizeof(FParticleBeamTrailVertexDynamicParameter);
 	}
@@ -2432,7 +2432,7 @@ struct FDynamicTrailsEmitterData : public FDynamicSpriteEmitterDataBase
 	/**
 	 *	Get the vertex stride for the dynamic rendering data
 	 */
-	virtual int32 GetDynamicVertexStride(ERHIFeatureLevel::Type /*InFeatureLevel*/) const OVERRIDE
+	virtual int32 GetDynamicVertexStride(ERHIFeatureLevel::Type /*InFeatureLevel*/) const override
 	{
 		return sizeof(FParticleBeamTrailVertex);
 	}
@@ -2440,7 +2440,7 @@ struct FDynamicTrailsEmitterData : public FDynamicSpriteEmitterDataBase
 	/**
 	 *	Get the vertex stride for the dynamic parameter rendering data
 	 */
-	virtual int32 GetDynamicParameterVertexStride() const OVERRIDE
+	virtual int32 GetDynamicParameterVertexStride() const override
 	{
 		return sizeof(FParticleBeamTrailVertexDynamicParameter);
 	}
@@ -2590,14 +2590,14 @@ public:
 	virtual ~FParticleSystemSceneProxy();
 
 	// FPrimitiveSceneProxy interface.
-	virtual bool CanBeOccluded() const OVERRIDE
+	virtual bool CanBeOccluded() const override
 	{
 		return false;
 	}
-	virtual void DrawDynamicElements(FPrimitiveDrawInterface* PDI,const FSceneView* View) OVERRIDE;
-	virtual FPrimitiveViewRelevance GetViewRelevance(const FSceneView* View) OVERRIDE;
-	virtual void OnActorPositionChanged() OVERRIDE;
-	virtual void OnTransformChanged() OVERRIDE;
+	virtual void DrawDynamicElements(FPrimitiveDrawInterface* PDI,const FSceneView* View) override;
+	virtual FPrimitiveViewRelevance GetViewRelevance(const FSceneView* View) override;
+	virtual void OnActorPositionChanged() override;
+	virtual void OnTransformChanged() override;
 
 	/**
 	 *	Helper function for determining the LOD distance for a given view.
@@ -2618,7 +2618,7 @@ public:
 	virtual void PreRenderView(const FSceneViewFamily* ViewFamily, const uint32 VisibilityMap, int32 FrameNumber);
 
 	/** Gathers simple lights for this emitter. */
-	virtual void GatherSimpleLights(const FSceneViewFamily& ViewFamily, FSimpleLightArray& OutParticleLights) const OVERRIDE;
+	virtual void GatherSimpleLights(const FSceneViewFamily& ViewFamily, FSimpleLightArray& OutParticleLights) const override;
 
 	/**
 	 *	Called when the rendering thread adds the proxy to the scene.
@@ -2739,7 +2739,7 @@ public:
 
 	// FPrimitiveSceneProxy interface.
 	/** @return true if the proxy requires occlusion queries */
-	virtual bool CanBeOccluded() const OVERRIDE
+	virtual bool CanBeOccluded() const override
 	{
 		return !MaterialRelevance.bDisableDepthTest;
 	}
diff --git a/Engine/Source/Runtime/Engine/Public/PreviewScene.h b/Engine/Source/Runtime/Engine/Public/PreviewScene.h
--- a/Engine/Source/Runtime/Engine/Public/PreviewScene.h
+++ b/Engine/Source/Runtime/Engine/Public/PreviewScene.h
@@ -61,7 +61,7 @@ public:
 	void RemoveComponent(class UActorComponent* Component);
 
 	// Serializer.
-	void AddReferencedObjects( FReferenceCollector& Collector ) OVERRIDE;
+	void AddReferencedObjects( FReferenceCollector& Collector ) override;
 
 	// Accessors.
 	UWorld* GetWorld() const { return PreviewWorld; }
diff --git a/Engine/Source/Runtime/Engine/Public/RawIndexBuffer.h b/Engine/Source/Runtime/Engine/Public/RawIndexBuffer.h
--- a/Engine/Source/Runtime/Engine/Public/RawIndexBuffer.h
+++ b/Engine/Source/Runtime/Engine/Public/RawIndexBuffer.h
@@ -167,7 +167,7 @@ public:
 	void Serialize(FArchive& Ar, bool bNeedsCPUAccess);
 
 	// FRenderResource interface.
-	virtual void InitRHI() OVERRIDE;
+	virtual void InitRHI() override;
 
 	inline bool Is32Bit() const { return b32Bit; }
 
diff --git a/Engine/Source/Runtime/Engine/Public/ResourcePool.h b/Engine/Source/Runtime/Engine/Public/ResourcePool.h
--- a/Engine/Source/Runtime/Engine/Public/ResourcePool.h
+++ b/Engine/Source/Runtime/Engine/Public/ResourcePool.h
@@ -217,28 +217,28 @@ public:
 	}
 	
 public: // From FTickableObjectRenderThread
-	virtual void Tick( float DeltaTime ) OVERRIDE
+	virtual void Tick( float DeltaTime ) override
 	{
 		TResourcePool<ResourceType, ResourcePoolPolicy, ResourceCreationArguments>::DrainPool(false);
 	}
 	
-	virtual bool IsTickable() const OVERRIDE
+	virtual bool IsTickable() const override
 	{
 		return true;
 	}
 	
-	virtual bool NeedsRenderingResumedForRenderingThreadTick() const OVERRIDE
+	virtual bool NeedsRenderingResumedForRenderingThreadTick() const override
 	{
 		return true;
 	}
 	
 public: // From FRenderResource
-	virtual void InitRHI() OVERRIDE
+	virtual void InitRHI() override
 	{
 		FTickableObjectRenderThread::Register();
 	}
 
-	virtual void ReleaseRHI() OVERRIDE
+	virtual void ReleaseRHI() override
 	{
 		FTickableObjectRenderThread::Unregister();
 		TResourcePool<ResourceType, ResourcePoolPolicy, ResourceCreationArguments>::DrainPool(true);
diff --git a/Engine/Source/Runtime/Engine/Public/SaveGameSystem.h b/Engine/Source/Runtime/Engine/Public/SaveGameSystem.h
--- a/Engine/Source/Runtime/Engine/Public/SaveGameSystem.h
+++ b/Engine/Source/Runtime/Engine/Public/SaveGameSystem.h
@@ -34,12 +34,12 @@ public:
 class FGenericSaveGameSystem : public ISaveGameSystem
 {
 public:
-	virtual bool PlatformHasNativeUI() OVERRIDE
+	virtual bool PlatformHasNativeUI() override
 	{
 		return false;
 	}
 
-	virtual bool DoesSaveGameExist(const TCHAR* Name, const int32 UserIndex) OVERRIDE
+	virtual bool DoesSaveGameExist(const TCHAR* Name, const int32 UserIndex) override
 	{
 #if PLATFORM_HTML5_BROWSER
 		return UE_DoesSaveGameExist(TCHAR_TO_ANSI(Name),UserIndex);
@@ -54,7 +54,7 @@ public:
 #endif
 	}
 
-	virtual bool SaveGame(bool bAttemptToUseUI, const TCHAR* Name, const int32 UserIndex, const TArray<uint8>& Data) OVERRIDE
+	virtual bool SaveGame(bool bAttemptToUseUI, const TCHAR* Name, const int32 UserIndex, const TArray<uint8>& Data) override
 	{
 #if PLATFORM_HTML5_BROWSER
 		return UE_SaveGame(TCHAR_TO_ANSI(Name),UserIndex,(char*)Data.GetData(),Data.Num());
@@ -70,7 +70,7 @@ public:
 
 	}
 
-	virtual bool LoadGame(bool bAttemptToUseUI, const TCHAR* Name, const int32 UserIndex, TArray<uint8>& Data) OVERRIDE
+	virtual bool LoadGame(bool bAttemptToUseUI, const TCHAR* Name, const int32 UserIndex, TArray<uint8>& Data) override
 	{
 #if PLATFORM_HTML5_BROWSER
 		char*	OutData;
diff --git a/Engine/Source/Runtime/Engine/Public/SkeletalMeshTypes.h b/Engine/Source/Runtime/Engine/Public/SkeletalMeshTypes.h
--- a/Engine/Source/Runtime/Engine/Public/SkeletalMeshTypes.h
+++ b/Engine/Source/Runtime/Engine/Public/SkeletalMeshTypes.h
@@ -2147,12 +2147,12 @@ public:
 
 	// FPrimitiveSceneProxy interface.
 #if WITH_EDITOR
-	virtual HHitProxy* CreateHitProxies(UPrimitiveComponent* Component, TArray<TRefCountPtr<HHitProxy> >& OutHitProxies) OVERRIDE;
+	virtual HHitProxy* CreateHitProxies(UPrimitiveComponent* Component, TArray<TRefCountPtr<HHitProxy> >& OutHitProxies) override;
 #endif
-	virtual void DrawDynamicElements(FPrimitiveDrawInterface* PDI, const FSceneView* View) OVERRIDE;
-	virtual FPrimitiveViewRelevance GetViewRelevance(const FSceneView* View) OVERRIDE;
-	virtual bool CanBeOccluded() const OVERRIDE;
-	virtual void PreRenderView(const FSceneViewFamily* ViewFamily, const uint32 VisibilityMap, int32 FrameNumber) OVERRIDE;
+	virtual void DrawDynamicElements(FPrimitiveDrawInterface* PDI, const FSceneView* View) override;
+	virtual FPrimitiveViewRelevance GetViewRelevance(const FSceneView* View) override;
+	virtual bool CanBeOccluded() const override;
+	virtual void PreRenderView(const FSceneViewFamily* ViewFamily, const uint32 VisibilityMap, int32 FrameNumber) override;
 	
 	/**
 	* Draw only the section of the material ID given of the scene proxy as a dynamic element
diff --git a/Engine/Source/Runtime/Engine/Public/Slate/SceneViewport.h b/Engine/Source/Runtime/Engine/Public/Slate/SceneViewport.h
--- a/Engine/Source/Runtime/Engine/Public/Slate/SceneViewport.h
+++ b/Engine/Source/Runtime/Engine/Public/Slate/SceneViewport.h
@@ -29,23 +29,23 @@ public:
 	virtual void CaptureMouse( bool bCapture );
 	virtual void LockMouseToViewport( bool bLock );
 	virtual void ShowCursor( bool bVisible );
-	virtual void SetPreCaptureMousePosFromSlateCursor() OVERRIDE;
-	virtual bool IsCursorVisible() const OVERRIDE { return bIsCursorVisible; }
-	virtual void ShowSoftwareCursor( bool bVisible ) OVERRIDE { bIsSoftwareCursorVisible = bVisible; }
-	virtual void SetSoftwareCursorPosition( FVector2D Position ) OVERRIDE { SoftwareCursorPosition = Position; }
-	virtual bool IsSoftwareCursorVisible() const OVERRIDE { return bIsSoftwareCursorVisible; }
-	virtual FVector2D GetSoftwareCursorPosition() const OVERRIDE
+	virtual void SetPreCaptureMousePosFromSlateCursor() override;
+	virtual bool IsCursorVisible() const override { return bIsCursorVisible; }
+	virtual void ShowSoftwareCursor( bool bVisible ) override { bIsSoftwareCursorVisible = bVisible; }
+	virtual void SetSoftwareCursorPosition( FVector2D Position ) override { SoftwareCursorPosition = Position; }
+	virtual bool IsSoftwareCursorVisible() const override { return bIsSoftwareCursorVisible; }
+	virtual FVector2D GetSoftwareCursorPosition() const override
 	{
 		return SoftwareCursorPosition;
 	}
-	virtual FCanvas* GetDebugCanvas() OVERRIDE;
+	virtual FCanvas* GetDebugCanvas() override;
 
 	/**
 	 * Captures or uncaptures the joystick
 	 *
 	 * @param Capture	true if we should capture, false if we should uncapture
 	 */
-	virtual bool CaptureJoystickInput(bool Capture) OVERRIDE;
+	virtual bool CaptureJoystickInput(bool Capture) override;
 
 	/**
 	 * Returns the state of the provided key. 
@@ -54,17 +54,17 @@ public:
 	 *
 	 * @return true if the key is pressed, false otherwise
 	 */
-	virtual bool KeyState(FKey Key) const OVERRIDE;
+	virtual bool KeyState(FKey Key) const override;
 
 	/**
 	 * @return The current X position of the mouse (in local space, relative to the viewports geometry)                 
 	 */
-	virtual int32 GetMouseX() const OVERRIDE;
+	virtual int32 GetMouseX() const override;
 
 	/**
 	 * @return The current Y position of the mouse (in local space, relative to the viewports geometry)                 
 	 */
-	virtual int32 GetMouseY() const OVERRIDE;
+	virtual int32 GetMouseY() const override;
 
 	/**
 	 * Sets MousePosition to the current mouse position 
@@ -72,40 +72,40 @@ public:
 	 * @param MousePosition	Populated with the current mouse position     
 	 * @param bLocalPosition Indicates whether the mouse position returned should be in local or absolute space
 	 */
-	virtual void GetMousePos( FIntPoint& MousePosition, const bool bLocalPosition = true) OVERRIDE;
+	virtual void GetMousePos( FIntPoint& MousePosition, const bool bLocalPosition = true) override;
 
 	/**
 	 * Not implemented                   
 	 */
-	virtual void SetMouse( int32 X, int32 Y ) OVERRIDE;
+	virtual void SetMouse( int32 X, int32 Y ) override;
 
 	/**
 	 * Additional input processing that happens every frame                   
 	 */
-	virtual void ProcessInput( float DeltaTime ) OVERRIDE;
+	virtual void ProcessInput( float DeltaTime ) override;
 
 	/**
 	 * Called when the viewport should be invalidated and redrawn                   
 	 */
-	virtual void InvalidateDisplay() OVERRIDE;
+	virtual void InvalidateDisplay() override;
 
 	/**
 	 * Invalidates the viewport's cached hit proxies at the end of the frame.
 	 */
-	virtual void DeferInvalidateHitProxy() OVERRIDE;
+	virtual void DeferInvalidateHitProxy() override;
 
 	/** FViewportFrame interface */
-	virtual FViewport* GetViewport() OVERRIDE { return this; }
-	virtual FViewportFrame* GetViewportFrame() OVERRIDE { return this; }
+	virtual FViewport* GetViewport() override { return this; }
+	virtual FViewportFrame* GetViewportFrame() override { return this; }
 
 	/** @return The viewport widget being used */
 	TWeakPtr<SViewport> GetViewportWidget() const { return ViewportWidget; }
 
 	/** Called before BeginRenderFrame is enqueued */
-	virtual void EnqueueBeginRenderFrame() OVERRIDE;
+	virtual void EnqueueBeginRenderFrame() override;
 
 	/** Called when a frame starts to render */
-	virtual void BeginRenderFrame() OVERRIDE;
+	virtual void BeginRenderFrame() override;
 
 	/** 
 	 * Called when a frame is done rendering
@@ -113,7 +113,7 @@ public:
 	 * @param bPresent	Not used in Slate viewports
 	 * @param bLockToVsync	Not used in Slate viewports
 	 */
-	virtual void EndRenderFrame( bool bPresent, bool bLockToVsync ) OVERRIDE;
+	virtual void EndRenderFrame( bool bPresent, bool bLockToVsync ) override;
 
 	/**
 	 * Ticks the viewport
@@ -137,12 +137,12 @@ public:
 	/**
 	 * Indicate that the viewport should be block for vsync.
 	 */
-	virtual void SetRequiresVsync(bool bShouldVsync) OVERRIDE { bRequiresVsync = bShouldVsync; }
+	virtual void SetRequiresVsync(bool bShouldVsync) override { bRequiresVsync = bShouldVsync; }
 
 	/**
 	 * Returns true if the viewport should be vsynced.
 	 */
-	virtual bool RequiresVsync() const OVERRIDE { return bRequiresVsync; }
+	virtual bool RequiresVsync() const override { return bRequiresVsync; }
 
 	/**
 	 * Called to resize the actual window where this viewport resides
@@ -151,7 +151,7 @@ public:
 	 * @param NewSizeY		The new height of the viewport
 	 * @param NewWindowMode	 What window mode should the viewport be resized to
 	 */
-	virtual void ResizeFrame(uint32 NewSizeX,uint32 NewSizeY,EWindowMode::Type NewWindowMode,int32 InPosX, int32 InPosY ) OVERRIDE;
+	virtual void ResizeFrame(uint32 NewSizeX,uint32 NewSizeY,EWindowMode::Type NewWindowMode,int32 InPosX, int32 InPosY ) override;
 
 	/**
 	 *	Sets the Viewport resize delegate.
@@ -180,34 +180,34 @@ public:
 	}
 
 	/** Updates the viewport RHI with a new size and fullscreen flag */
-	virtual void UpdateViewportRHI(bool bDestroyed,uint32 NewSizeX,uint32 NewSizeY,EWindowMode::Type NewWindowMode) OVERRIDE;
+	virtual void UpdateViewportRHI(bool bDestroyed,uint32 NewSizeX,uint32 NewSizeY,EWindowMode::Type NewWindowMode) override;
 
 	/** ISlateViewport interface */
-	virtual FSlateShaderResource* GetViewportRenderTargetTexture() const OVERRIDE;
-	virtual void OnDrawViewport( const FGeometry& AllottedGeometry, const FSlateRect& MyClippingRect, FSlateWindowElementList& OutDrawElements, int32 LayerId, const FWidgetStyle& InWidgetStyle, bool bParentEnabled ) OVERRIDE;
-	virtual FCursorReply OnCursorQuery( const FGeometry& MyGeometry, const FPointerEvent& CursorEvent ) OVERRIDE;
-	virtual FReply OnMouseButtonDown( const FGeometry& InGeometry, const FPointerEvent& MouseEvent ) OVERRIDE;
-	virtual FReply OnMouseButtonUp( const FGeometry& InGeometry, const FPointerEvent& MouseEvent ) OVERRIDE;
-	virtual void OnMouseEnter( const FGeometry& MyGeometry, const FPointerEvent& MouseEvent ) OVERRIDE;
-	virtual void OnMouseLeave( const FPointerEvent& MouseEvent ) OVERRIDE;
-	virtual FReply OnMouseMove( const FGeometry& InGeometry, const FPointerEvent& MouseEvent ) OVERRIDE;
-	virtual FReply OnMouseWheel( const FGeometry& InGeometry, const FPointerEvent& MouseEvent ) OVERRIDE;
-	virtual FReply OnMouseButtonDoubleClick( const FGeometry& InGeometry, const FPointerEvent& InMouseEvent ) OVERRIDE;
-	virtual FReply OnControllerButtonPressed( const FGeometry& MyGeometry, const FControllerEvent& ControllerEvent ) OVERRIDE;
-	virtual FReply OnControllerButtonReleased( const FGeometry& MyGeometry, const FControllerEvent& ControllerEvent ) OVERRIDE;
-	virtual FReply OnControllerAnalogValueChanged( const FGeometry& MyGeometry, const FControllerEvent& ControllerEvent ) OVERRIDE;
-	virtual FReply OnTouchStarted( const FGeometry& MyGeometry, const FPointerEvent& InTouchEvent ) OVERRIDE;
-	virtual FReply OnTouchMoved( const FGeometry& MyGeometry, const FPointerEvent& InTouchEvent ) OVERRIDE;
-	virtual FReply OnTouchEnded( const FGeometry& MyGeometry, const FPointerEvent& InTouchEvent ) OVERRIDE;
-	virtual FReply OnTouchGesture( const FGeometry& MyGeometry, const FPointerEvent& InGestureEvent ) OVERRIDE;
-	virtual FReply OnMotionDetected( const FGeometry& MyGeometry, const FMotionEvent& InMotionEvent ) OVERRIDE;
-	virtual FReply OnKeyDown( const FGeometry& InGeometry, const FKeyboardEvent& InKeyboardEvent ) OVERRIDE;
-	virtual FReply OnKeyUp( const FGeometry& InGeometry, const FKeyboardEvent& InKeyboardEvent ) OVERRIDE;
-	virtual FReply OnKeyChar( const FGeometry& InGeometry, const FCharacterEvent& InCharacterEvent ) OVERRIDE;
-	virtual FReply OnKeyboardFocusReceived( const FKeyboardFocusEvent& InKeyboardFocusEvent ) OVERRIDE;
-	virtual void OnKeyboardFocusLost( const FKeyboardFocusEvent& InKeyboardFocusEvent ) OVERRIDE;
-	virtual void OnViewportClosed() OVERRIDE;
-	virtual FIntPoint GetSize() const OVERRIDE { return GetSizeXY(); }
+	virtual FSlateShaderResource* GetViewportRenderTargetTexture() const override;
+	virtual void OnDrawViewport( const FGeometry& AllottedGeometry, const FSlateRect& MyClippingRect, FSlateWindowElementList& OutDrawElements, int32 LayerId, const FWidgetStyle& InWidgetStyle, bool bParentEnabled ) override;
+	virtual FCursorReply OnCursorQuery( const FGeometry& MyGeometry, const FPointerEvent& CursorEvent ) override;
+	virtual FReply OnMouseButtonDown( const FGeometry& InGeometry, const FPointerEvent& MouseEvent ) override;
+	virtual FReply OnMouseButtonUp( const FGeometry& InGeometry, const FPointerEvent& MouseEvent ) override;
+	virtual void OnMouseEnter( const FGeometry& MyGeometry, const FPointerEvent& MouseEvent ) override;
+	virtual void OnMouseLeave( const FPointerEvent& MouseEvent ) override;
+	virtual FReply OnMouseMove( const FGeometry& InGeometry, const FPointerEvent& MouseEvent ) override;
+	virtual FReply OnMouseWheel( const FGeometry& InGeometry, const FPointerEvent& MouseEvent ) override;
+	virtual FReply OnMouseButtonDoubleClick( const FGeometry& InGeometry, const FPointerEvent& InMouseEvent ) override;
+	virtual FReply OnControllerButtonPressed( const FGeometry& MyGeometry, const FControllerEvent& ControllerEvent ) override;
+	virtual FReply OnControllerButtonReleased( const FGeometry& MyGeometry, const FControllerEvent& ControllerEvent ) override;
+	virtual FReply OnControllerAnalogValueChanged( const FGeometry& MyGeometry, const FControllerEvent& ControllerEvent ) override;
+	virtual FReply OnTouchStarted( const FGeometry& MyGeometry, const FPointerEvent& InTouchEvent ) override;
+	virtual FReply OnTouchMoved( const FGeometry& MyGeometry, const FPointerEvent& InTouchEvent ) override;
+	virtual FReply OnTouchEnded( const FGeometry& MyGeometry, const FPointerEvent& InTouchEvent ) override;
+	virtual FReply OnTouchGesture( const FGeometry& MyGeometry, const FPointerEvent& InGestureEvent ) override;
+	virtual FReply OnMotionDetected( const FGeometry& MyGeometry, const FMotionEvent& InMotionEvent ) override;
+	virtual FReply OnKeyDown( const FGeometry& InGeometry, const FKeyboardEvent& InKeyboardEvent ) override;
+	virtual FReply OnKeyUp( const FGeometry& InGeometry, const FKeyboardEvent& InKeyboardEvent ) override;
+	virtual FReply OnKeyChar( const FGeometry& InGeometry, const FCharacterEvent& InCharacterEvent ) override;
+	virtual FReply OnKeyboardFocusReceived( const FKeyboardFocusEvent& InKeyboardFocusEvent ) override;
+	virtual void OnKeyboardFocusLost( const FKeyboardFocusEvent& InKeyboardFocusEvent ) override;
+	virtual void OnViewportClosed() override;
+	virtual FIntPoint GetSize() const override { return GetSizeXY(); }
 	
 private:
 	/**
@@ -216,15 +216,15 @@ private:
 	void Destroy();
 
 	// FRenderResource interface.
-	virtual void InitDynamicRHI() OVERRIDE;
-	virtual void ReleaseDynamicRHI() OVERRIDE;
+	virtual void InitDynamicRHI() override;
+	virtual void ReleaseDynamicRHI() override;
 
 	// @todo UE4 DLL: Without these functions we get unresolved linker errors with FRenderResource
-	virtual void InitRHI() OVERRIDE {}
-	virtual void ReleaseRHI() OVERRIDE {}
-	virtual void InitResource() OVERRIDE { FViewport::InitResource(); }
-	virtual void ReleaseResource() OVERRIDE { FViewport::ReleaseResource(); }
-	virtual FString GetFriendlyName() const OVERRIDE { return FString(TEXT("FSlateSceneViewport"));}
+	virtual void InitRHI() override {}
+	virtual void ReleaseRHI() override {}
+	virtual void InitResource() override { FViewport::InitResource(); }
+	virtual void ReleaseResource() override { FViewport::ReleaseResource(); }
+	virtual FString GetFriendlyName() const override { return FString(TEXT("FSlateSceneViewport"));}
 
 	/**
 	 * Called from Slate when the viewport should be resized
diff --git a/Engine/Source/Runtime/Engine/Public/Slate/SlateGameResources.h b/Engine/Source/Runtime/Engine/Public/Slate/SlateGameResources.h
--- a/Engine/Source/Runtime/Engine/Public/Slate/SlateGameResources.h
+++ b/Engine/Source/Runtime/Engine/Public/Slate/SlateGameResources.h
@@ -38,12 +38,12 @@ public:
 	 * Populate an array of FSlateBrush with resources consumed by this style chunk.
 	 * @param OutResources - the array to populate.
 	 */
-	virtual void GetResources( TArray< const FSlateBrush* >& OutResources ) const OVERRIDE;
+	virtual void GetResources( TArray< const FSlateBrush* >& OutResources ) const override;
 
-	virtual void SetContentRoot( const FString& InContentRootDir ) OVERRIDE;
+	virtual void SetContentRoot( const FString& InContentRootDir ) override;
 
-	virtual const FSlateBrush* GetBrush( const FName PropertyName, const ANSICHAR* Specifier = NULL ) const OVERRIDE;
-	virtual const FSlateBrush* GetOptionalBrush(const FName PropertyName, const ANSICHAR* Specifier = NULL, const FSlateBrush* const DefaultBrush = FStyleDefaults::GetNoBrush()) const OVERRIDE;
+	virtual const FSlateBrush* GetBrush( const FName PropertyName, const ANSICHAR* Specifier = NULL ) const override;
+	virtual const FSlateBrush* GetOptionalBrush(const FName PropertyName, const ANSICHAR* Specifier = NULL, const FSlateBrush* const DefaultBrush = FStyleDefaults::GetNoBrush()) const override;
 
 	UCurveFloat* GetCurveFloat( const FName AssetName ) const;
 	UCurveVector* GetCurveVector( const FName AssetName ) const;
@@ -51,9 +51,9 @@ public:
 
 protected:
 
-	virtual const FSlateWidgetStyle* GetWidgetStyleInternal( const FName DesiredTypeName, const FName StyleName ) const OVERRIDE;
+	virtual const FSlateWidgetStyle* GetWidgetStyleInternal( const FName DesiredTypeName, const FName StyleName ) const override;
 
-	virtual void Log( ISlateStyle::EStyleMessageSeverity Severity, const FText& Message ) const OVERRIDE;
+	virtual void Log( ISlateStyle::EStyleMessageSeverity Severity, const FText& Message ) const override;
 
 	virtual void Log( const TSharedRef< class FTokenizedMessage >& Message ) const;
 
@@ -73,7 +73,7 @@ protected:
 
 	FName GenerateMapName( UObject* StyleObject );
 
-	virtual void AddReferencedObjects( FReferenceCollector& Collector ) OVERRIDE;
+	virtual void AddReferencedObjects( FReferenceCollector& Collector ) override;
 
 private:
 
diff --git a/Engine/Source/Runtime/Engine/Public/Slate/SlateTextures.h b/Engine/Source/Runtime/Engine/Public/Slate/SlateTextures.h
--- a/Engine/Source/Runtime/Engine/Public/Slate/SlateTextures.h
+++ b/Engine/Source/Runtime/Engine/Public/Slate/SlateTextures.h
@@ -132,18 +132,18 @@ public:
 
 public:
 	// FTextureRenderTargetResource implementation
-	virtual void ClampSize(int32 SizeX,int32 SizeY) OVERRIDE;
+	virtual void ClampSize(int32 SizeX,int32 SizeY) override;
 
 	// FRenderResource implementation
-	virtual void InitDynamicRHI() OVERRIDE;
-	virtual void ReleaseDynamicRHI() OVERRIDE;
+	virtual void InitDynamicRHI() override;
+	virtual void ReleaseDynamicRHI() override;
 
 	// FDeferredUpdateResource implementation
-	virtual void UpdateResource() OVERRIDE;
+	virtual void UpdateResource() override;
 
 	// FRenderTarget interface
-	virtual FIntPoint GetSizeXY() const OVERRIDE;
-	virtual float GetDisplayGamma() const OVERRIDE;
+	virtual FIntPoint GetSizeXY() const override;
+	virtual float GetDisplayGamma() const override;
 
 private:
 	FTexture2DRHIRef Texture2DRHI;
diff --git a/Engine/Source/Runtime/Engine/Public/StaticMeshLight.h b/Engine/Source/Runtime/Engine/Public/StaticMeshLight.h
--- a/Engine/Source/Runtime/Engine/Public/StaticMeshLight.h
+++ b/Engine/Source/Runtime/Engine/Public/StaticMeshLight.h
@@ -114,7 +114,7 @@ public:
 	}
 
 	/** Whether or not this mapping should be processed or imported */
-	virtual bool IsValidMapping() const OVERRIDE {return Primitive.IsValid();}
+	virtual bool IsValidMapping() const override {return Primitive.IsValid();}
 
 protected:
 
diff --git a/Engine/Source/Runtime/Engine/Public/StaticMeshResources.h b/Engine/Source/Runtime/Engine/Public/StaticMeshResources.h
--- a/Engine/Source/Runtime/Engine/Public/StaticMeshResources.h
+++ b/Engine/Source/Runtime/Engine/Public/StaticMeshResources.h
@@ -530,7 +530,7 @@ class FDistanceFieldVolumeTextureAtlas : public FRenderResource
 public:
 	FDistanceFieldVolumeTextureAtlas(EPixelFormat InFormat);
 
-	virtual void ReleaseRHI() OVERRIDE
+	virtual void ReleaseRHI() override
 	{
 		VolumeTextureRHI.SafeRelease();
 	}
@@ -875,19 +875,19 @@ protected:
 public:
 	// FPrimitiveSceneProxy interface.
 #if WITH_EDITOR
-	virtual HHitProxy* CreateHitProxies(UPrimitiveComponent* Component, TArray<TRefCountPtr<HHitProxy> >& OutHitProxies) OVERRIDE;
+	virtual HHitProxy* CreateHitProxies(UPrimitiveComponent* Component, TArray<TRefCountPtr<HHitProxy> >& OutHitProxies) override;
 #endif
-	virtual void DrawStaticElements(FStaticPrimitiveDrawInterface* PDI) OVERRIDE;
-	virtual void DrawDynamicElements(FPrimitiveDrawInterface* PDI,const FSceneView* View) OVERRIDE { DrawDynamicElements( PDI, View, 0 ); }
-	virtual void DrawDynamicElements(FPrimitiveDrawInterface* PDI,const FSceneView* View, uint32 DrawDynamicFlags ) OVERRIDE;
-	virtual void OnTransformChanged() OVERRIDE;
-	virtual int32 GetLOD(const FSceneView* View) const OVERRIDE;
-	virtual FPrimitiveViewRelevance GetViewRelevance(const FSceneView* View) OVERRIDE;
-	virtual bool CanBeOccluded() const OVERRIDE;
-	virtual void GetLightRelevance(const FLightSceneProxy* LightSceneProxy, bool& bDynamic, bool& bRelevant, bool& bLightMapped, bool& bShadowMapped) const OVERRIDE;
-	virtual void GetDistancefieldAtlasData(FBox& LocalVolumeBounds, FIntVector& OutBlockMin, FIntVector& OutBlockSize) const OVERRIDE;
-	virtual bool HasDistanceFieldRepresentation() const OVERRIDE;
-	virtual uint32 GetMemoryFootprint( void ) const OVERRIDE { return( sizeof( *this ) + GetAllocatedSize() ); }
+	virtual void DrawStaticElements(FStaticPrimitiveDrawInterface* PDI) override;
+	virtual void DrawDynamicElements(FPrimitiveDrawInterface* PDI,const FSceneView* View) override { DrawDynamicElements( PDI, View, 0 ); }
+	virtual void DrawDynamicElements(FPrimitiveDrawInterface* PDI,const FSceneView* View, uint32 DrawDynamicFlags ) override;
+	virtual void OnTransformChanged() override;
+	virtual int32 GetLOD(const FSceneView* View) const override;
+	virtual FPrimitiveViewRelevance GetViewRelevance(const FSceneView* View) override;
+	virtual bool CanBeOccluded() const override;
+	virtual void GetLightRelevance(const FLightSceneProxy* LightSceneProxy, bool& bDynamic, bool& bRelevant, bool& bLightMapped, bool& bShadowMapped) const override;
+	virtual void GetDistancefieldAtlasData(FBox& LocalVolumeBounds, FIntVector& OutBlockMin, FIntVector& OutBlockSize) const override;
+	virtual bool HasDistanceFieldRepresentation() const override;
+	virtual uint32 GetMemoryFootprint( void ) const override { return( sizeof( *this ) + GetAllocatedSize() ); }
 	uint32 GetAllocatedSize( void ) const { return( FPrimitiveSceneProxy::GetAllocatedSize() + LODs.GetAllocatedSize() ); }
 
 protected:
diff --git a/Engine/Source/Runtime/Engine/Public/SystemSettings.h b/Engine/Source/Runtime/Engine/Public/SystemSettings.h
--- a/Engine/Source/Runtime/Engine/Public/SystemSettings.h
+++ b/Engine/Source/Runtime/Engine/Public/SystemSettings.h
@@ -112,7 +112,7 @@ public:
 	void ApplyOverrides();
 	
 	// Begin Exec Interface
-	virtual bool Exec( UWorld* InWorld, const TCHAR* Cmd, FOutputDevice& Ar ) OVERRIDE;
+	virtual bool Exec( UWorld* InWorld, const TCHAR* Cmd, FOutputDevice& Ar ) override;
 	// End Exec Interface
 
 	/** 
diff --git a/Engine/Source/Runtime/Engine/Public/VisualLog.h b/Engine/Source/Runtime/Engine/Public/VisualLog.h
--- a/Engine/Source/Runtime/Engine/Public/VisualLog.h
+++ b/Engine/Source/Runtime/Engine/Public/VisualLog.h
@@ -342,7 +342,7 @@ public:
 
 protected:
 
-	virtual void Serialize( const TCHAR* V, ELogVerbosity::Type Verbosity, const class FName& Category ) OVERRIDE;
+	virtual void Serialize( const TCHAR* V, ELogVerbosity::Type Verbosity, const class FName& Category ) override;
 
 	void LogElementImpl(FVisLogSelfDrawingElement* Element);
 
diff --git a/Engine/Source/Runtime/EngineMessages/Private/EngineMessagesModule.cpp b/Engine/Source/Runtime/EngineMessages/Private/EngineMessagesModule.cpp
--- a/Engine/Source/Runtime/EngineMessages/Private/EngineMessagesModule.cpp
+++ b/Engine/Source/Runtime/EngineMessages/Private/EngineMessagesModule.cpp
@@ -13,10 +13,10 @@ public:
 
 	// IModuleInterface interface
 
-	virtual void StartupModule( ) OVERRIDE { }
-	virtual void ShutdownModule( ) OVERRIDE { }
+	virtual void StartupModule( ) override { }
+	virtual void ShutdownModule( ) override { }
 
-	virtual bool SupportsDynamicReloading( ) OVERRIDE
+	virtual bool SupportsDynamicReloading( ) override
 	{
 		return true;
 	}
diff --git a/Engine/Source/Runtime/EngineSettings/Private/EngineSettingsModule.cpp b/Engine/Source/Runtime/EngineSettings/Private/EngineSettingsModule.cpp
--- a/Engine/Source/Runtime/EngineSettings/Private/EngineSettingsModule.cpp
+++ b/Engine/Source/Runtime/EngineSettings/Private/EngineSettingsModule.cpp
@@ -13,10 +13,10 @@ public:
 
 	// IModuleInterface interface
 
-	virtual void StartupModule( ) OVERRIDE { }
-	virtual void ShutdownModule( ) OVERRIDE { }
+	virtual void StartupModule( ) override { }
+	virtual void ShutdownModule( ) override { }
 
-	virtual bool SupportsDynamicReloading( ) OVERRIDE
+	virtual bool SupportsDynamicReloading( ) override
 	{
 		return true;
 	}
diff --git a/Engine/Source/Runtime/GameplayAbilities/Public/Abilities/GameplayAbility.h b/Engine/Source/Runtime/GameplayAbilities/Public/Abilities/GameplayAbility.h
--- a/Engine/Source/Runtime/GameplayAbilities/Public/Abilities/GameplayAbility.h
+++ b/Engine/Source/Runtime/GameplayAbilities/Public/Abilities/GameplayAbility.h
@@ -303,7 +303,7 @@ protected:
 
 	virtual void ClientActivateAbilitySucceed_Internal(int32 PredictionKey);
 
-	bool IsSupportedForNetworking() const OVERRIDE;
+	bool IsSupportedForNetworking() const override;
 
 	/** Checks cooldown. returns true if we can be used again. False if not */
 	bool	CheckCooldown(const FGameplayAbilityActorInfo* ActorInfo) const;
diff --git a/Engine/Source/Runtime/GameplayAbilities/Public/Abilities/GameplayAbility_CharacterJump.h b/Engine/Source/Runtime/GameplayAbilities/Public/Abilities/GameplayAbility_CharacterJump.h
--- a/Engine/Source/Runtime/GameplayAbilities/Public/Abilities/GameplayAbility_CharacterJump.h
+++ b/Engine/Source/Runtime/GameplayAbilities/Public/Abilities/GameplayAbility_CharacterJump.h
@@ -15,13 +15,13 @@ class GAMEPLAYABILITIES_API UGameplayAbility_CharacterJump : public UGameplayAbi
 
 public:
 
-	virtual bool CanActivateAbility(const FGameplayAbilityActorInfo* ActorInfo) const OVERRIDE;
+	virtual bool CanActivateAbility(const FGameplayAbilityActorInfo* ActorInfo) const override;
 	
-	virtual void ActivateAbility(const FGameplayAbilityActorInfo* OwnerInfo, const FGameplayAbilityActivationInfo ActivationInfo) OVERRIDE;
+	virtual void ActivateAbility(const FGameplayAbilityActorInfo* OwnerInfo, const FGameplayAbilityActivationInfo ActivationInfo) override;
 
-	virtual void InputReleased(int32 InputID, const FGameplayAbilityActorInfo* ActorInfo) OVERRIDE;
+	virtual void InputReleased(int32 InputID, const FGameplayAbilityActorInfo* ActorInfo) override;
 
-	virtual void PredictiveActivateAbility(const FGameplayAbilityActorInfo* ActorInfo, const FGameplayAbilityActivationInfo ActivationInfo) OVERRIDE;
+	virtual void PredictiveActivateAbility(const FGameplayAbilityActorInfo* ActorInfo, const FGameplayAbilityActivationInfo ActivationInfo) override;
 
-	virtual void CancelAbility(const FGameplayAbilityActorInfo* ActorInfo, const FGameplayAbilityActivationInfo ActivationInfo) OVERRIDE;
+	virtual void CancelAbility(const FGameplayAbilityActorInfo* ActorInfo, const FGameplayAbilityActivationInfo ActivationInfo) override;
 };
\ No newline at end of file
diff --git a/Engine/Source/Runtime/GameplayAbilities/Public/Abilities/GameplayAbility_Montage.h b/Engine/Source/Runtime/GameplayAbilities/Public/Abilities/GameplayAbility_Montage.h
--- a/Engine/Source/Runtime/GameplayAbilities/Public/Abilities/GameplayAbility_Montage.h
+++ b/Engine/Source/Runtime/GameplayAbilities/Public/Abilities/GameplayAbility_Montage.h
@@ -15,7 +15,7 @@ class GAMEPLAYABILITIES_API UGameplayAbility_Montage : public UGameplayAbility
 
 public:
 	
-	virtual void ActivateAbility(const FGameplayAbilityActorInfo* OwnerInfo, const FGameplayAbilityActivationInfo ActivationInfo) OVERRIDE;
+	virtual void ActivateAbility(const FGameplayAbilityActorInfo* OwnerInfo, const FGameplayAbilityActivationInfo ActivationInfo) override;
 
 	UPROPERTY(EditDefaultsOnly, Category = MontageAbility)
 	UAnimMontage *	MontageToPlay;
diff --git a/Engine/Source/Runtime/GameplayAbilities/Public/Abilities/Tasks/AbilityTask_WaitConfirmCancel.h b/Engine/Source/Runtime/GameplayAbilities/Public/Abilities/Tasks/AbilityTask_WaitConfirmCancel.h
--- a/Engine/Source/Runtime/GameplayAbilities/Public/Abilities/Tasks/AbilityTask_WaitConfirmCancel.h
+++ b/Engine/Source/Runtime/GameplayAbilities/Public/Abilities/Tasks/AbilityTask_WaitConfirmCancel.h
@@ -27,7 +27,7 @@ class UAbilityTask_WaitConfirmCancel : public UAbilityTask
 	UFUNCTION(BlueprintCallable, meta=(HidePin = "WorldContextObject", DefaultToSelf = "WorldContextObject", BlueprintInternalUseOnly = "true", FriendlyName="Wait for Confirm Input"), Category="Abilities")
 	static UAbilityTask_WaitConfirmCancel* WaitConfirmCancel(UObject* WorldContextObject);	
 
-	virtual void Activate() OVERRIDE;
+	virtual void Activate() override;
 
 	UPROPERTY()
 	UAbilitySystemComponent* ASC;
diff --git a/Engine/Source/Runtime/GameplayAbilities/Public/Abilities/Tasks/AbilityTask_WaitTargetData.h b/Engine/Source/Runtime/GameplayAbilities/Public/Abilities/Tasks/AbilityTask_WaitTargetData.h
--- a/Engine/Source/Runtime/GameplayAbilities/Public/Abilities/Tasks/AbilityTask_WaitTargetData.h
+++ b/Engine/Source/Runtime/GameplayAbilities/Public/Abilities/Tasks/AbilityTask_WaitTargetData.h
@@ -38,7 +38,7 @@ class UAbilityTask_WaitTargetData: public UAbilityTask
 	UFUNCTION(BlueprintCallable, meta = (HidePin = "WorldContextObject", DefaultToSelf = "WorldContextObject", BlueprintInternalUseOnly = "true"), Category = "Abilities")
 	static UAbilityTask_WaitTargetData* CreateWaitTargetData(UObject* WorldContextObject, TSubclassOf<AGameplayAbilityTargetActor> TargetClass);
 
-	virtual void Activate() OVERRIDE;
+	virtual void Activate() override;
 
 	TSubclassOf<AGameplayAbilityTargetActor> TargetClass;
 
diff --git a/Engine/Source/Runtime/GameplayAbilities/Public/AbilitySystemComponent.h b/Engine/Source/Runtime/GameplayAbilities/Public/AbilitySystemComponent.h
--- a/Engine/Source/Runtime/GameplayAbilities/Public/AbilitySystemComponent.h
+++ b/Engine/Source/Runtime/GameplayAbilities/Public/AbilitySystemComponent.h
@@ -67,7 +67,7 @@ class GAMEPLAYABILITIES_API UAbilitySystemComponent : public UActorComponent
 
 	virtual ~UAbilitySystemComponent();
 
-	virtual void InitializeComponent() OVERRIDE;
+	virtual void InitializeComponent() override;
 
 	/** Finds existing AttributeSet */
 	template <class T >
@@ -115,11 +115,11 @@ class GAMEPLAYABILITIES_API UAbilitySystemComponent : public UActorComponent
 
 	// -- Replication -------------------------------------------------------------------------------------------------
 
-	virtual bool ReplicateSubobjects(class UActorChannel *Channel, class FOutBunch *Bunch, FReplicationFlags *RepFlags) OVERRIDE;
+	virtual bool ReplicateSubobjects(class UActorChannel *Channel, class FOutBunch *Bunch, FReplicationFlags *RepFlags) override;
 	
-	virtual void GetSubobjectsWithStableNamesForNetworking(TArray<UObject*>& Objs) OVERRIDE;
+	virtual void GetSubobjectsWithStableNamesForNetworking(TArray<UObject*>& Objs) override;
 
-	virtual void PostNetReceive() OVERRIDE;
+	virtual void PostNetReceive() override;
 
 	/**
 	 *	Prediction Keys
@@ -456,7 +456,7 @@ private:
 
 protected:
 
-	virtual void OnRegister() OVERRIDE;
+	virtual void OnRegister() override;
 
 	UAttributeSet*	GetAttributeSubobject(const TSubclassOf<UAttributeSet> AttributeClass) const;
 	UAttributeSet*	GetAttributeSubobjectChecked(const TSubclassOf<UAttributeSet> AttributeClass) const;
diff --git a/Engine/Source/Runtime/GameplayAbilities/Public/AbilitySystemDebugHUD.h b/Engine/Source/Runtime/GameplayAbilities/Public/AbilitySystemDebugHUD.h
--- a/Engine/Source/Runtime/GameplayAbilities/Public/AbilitySystemDebugHUD.h
+++ b/Engine/Source/Runtime/GameplayAbilities/Public/AbilitySystemDebugHUD.h
@@ -30,7 +30,7 @@ class AAbilitySystemDebugHUD : public AHUD
 	GENERATED_UCLASS_BODY()
 
 	/** main HUD update loop */
-	virtual void DrawHUD() OVERRIDE;
+	virtual void DrawHUD() override;
 
 	void DrawWithBackground(UFont* InFont, const FString& Text, const FColor& TextColor, EAlignHorizontal::Type HAlign, float& OffsetX, EAlignVertical::Type VAlign, float& OffsetY, float Alpha=1.f);
 
diff --git a/Engine/Source/Runtime/GameplayAbilities/Public/AbilitySystemTestAttributeSet.h b/Engine/Source/Runtime/GameplayAbilities/Public/AbilitySystemTestAttributeSet.h
--- a/Engine/Source/Runtime/GameplayAbilities/Public/AbilitySystemTestAttributeSet.h
+++ b/Engine/Source/Runtime/GameplayAbilities/Public/AbilitySystemTestAttributeSet.h
@@ -60,6 +60,6 @@ class GAMEPLAYABILITIES_API UAbilitySystemTestAttributeSet : public UAttributeSe
 	UPROPERTY(EditAnywhere, BlueprintReadWrite, Category = "AttributeTest")
 	float	NoStackAttribute;
 
-	virtual void PreAttributeModify(struct FGameplayEffectModCallbackData &Data) OVERRIDE;
-	virtual void PostAttributeModify(const struct FGameplayEffectModCallbackData &Data) OVERRIDE;
+	virtual void PreAttributeModify(struct FGameplayEffectModCallbackData &Data) override;
+	virtual void PostAttributeModify(const struct FGameplayEffectModCallbackData &Data) override;
 };
diff --git a/Engine/Source/Runtime/GameplayAbilities/Public/AbilitySystemTestPawn.h b/Engine/Source/Runtime/GameplayAbilities/Public/AbilitySystemTestPawn.h
--- a/Engine/Source/Runtime/GameplayAbilities/Public/AbilitySystemTestPawn.h
+++ b/Engine/Source/Runtime/GameplayAbilities/Public/AbilitySystemTestPawn.h
@@ -10,19 +10,19 @@ class GAMEPLAYABILITIES_API AAbilitySystemTestPawn : public ADefaultPawn, public
 {
 	GENERATED_UCLASS_BODY()
 
-	virtual void PostInitializeComponents() OVERRIDE;
+	virtual void PostInitializeComponents() override;
 
 	UFUNCTION()
-	virtual void GameplayCueActivated(const FGameplayTagContainer & GameplayCueTags, float NormalizedMagnitude, const FGameplayEffectInstigatorContext InstigatorContext) OVERRIDE;
+	virtual void GameplayCueActivated(const FGameplayTagContainer & GameplayCueTags, float NormalizedMagnitude, const FGameplayEffectInstigatorContext InstigatorContext) override;
 
 	UFUNCTION()
-	virtual void GameplayCueExecuted(const FGameplayTagContainer & GameplayCueTags, float NormalizedMagnitude, const FGameplayEffectInstigatorContext InstigatorContext) OVERRIDE;
+	virtual void GameplayCueExecuted(const FGameplayTagContainer & GameplayCueTags, float NormalizedMagnitude, const FGameplayEffectInstigatorContext InstigatorContext) override;
 
 	UFUNCTION()
-	virtual void GameplayCueAdded(const FGameplayTagContainer & GameplayCueTags, float NormalizedMagnitude, const FGameplayEffectInstigatorContext InstigatorContext) OVERRIDE;
+	virtual void GameplayCueAdded(const FGameplayTagContainer & GameplayCueTags, float NormalizedMagnitude, const FGameplayEffectInstigatorContext InstigatorContext) override;
 
 	UFUNCTION()
-	virtual void GameplayCueRemoved(const FGameplayTagContainer & GameplayCueTags, float NormalizedMagnitude, const FGameplayEffectInstigatorContext InstigatorContext) OVERRIDE;
+	virtual void GameplayCueRemoved(const FGameplayTagContainer & GameplayCueTags, float NormalizedMagnitude, const FGameplayEffectInstigatorContext InstigatorContext) override;
 
 
 	/** DefaultPawn collision component */
diff --git a/Engine/Source/Runtime/GameplayAbilities/Public/AttributeSet.h b/Engine/Source/Runtime/GameplayAbilities/Public/AttributeSet.h
--- a/Engine/Source/Runtime/GameplayAbilities/Public/AttributeSet.h
+++ b/Engine/Source/Runtime/GameplayAbilities/Public/AttributeSet.h
@@ -68,7 +68,7 @@ class GAMEPLAYABILITIES_API UAttributeSet : public UObject
 
 public:
 
-	bool IsSupportedForNetworking() const OVERRIDE
+	bool IsSupportedForNetworking() const override
 	{
 		return true;
 	}
diff --git a/Engine/Source/Runtime/GameplayAbilities/Public/GameplayEffect.h b/Engine/Source/Runtime/GameplayAbilities/Public/GameplayEffect.h
--- a/Engine/Source/Runtime/GameplayAbilities/Public/GameplayEffect.h
+++ b/Engine/Source/Runtime/GameplayAbilities/Public/GameplayEffect.h
@@ -408,7 +408,7 @@ public:
 	// ------------------------------------------------
 
 	/** Overridden to return requirements tags */
-	virtual void GetOwnedGameplayTags(FGameplayTagContainer& TagContainer) const OVERRIDE;
+	virtual void GetOwnedGameplayTags(FGameplayTagContainer& TagContainer) const override;
 
 	/** Get the "clear tags" for the effect */
 	virtual void GetClearGameplayTags(FGameplayTagContainer& TagContainer) const;
@@ -429,7 +429,7 @@ public:
 		return false;
 	}
 
-	virtual void PostLoad() OVERRIDE;
+	virtual void PostLoad() override;
 
 	// -----------------------------------------------
 	
diff --git a/Engine/Source/Runtime/GameplayAbilities/Public/GameplayEffectExtension_LifestealTest.h b/Engine/Source/Runtime/GameplayAbilities/Public/GameplayEffectExtension_LifestealTest.h
--- a/Engine/Source/Runtime/GameplayAbilities/Public/GameplayEffectExtension_LifestealTest.h
+++ b/Engine/Source/Runtime/GameplayAbilities/Public/GameplayEffectExtension_LifestealTest.h
@@ -15,8 +15,8 @@ class GAMEPLAYABILITIES_API UGameplayEffectExtension_LifestealTest : public UGam
 
 public:
 
-	void PreGameplayEffectExecute(const FGameplayModifierEvaluatedData &SelfData, FGameplayEffectModCallbackData &Data) const OVERRIDE;
-	void PostGameplayEffectExecute(const FGameplayModifierEvaluatedData &SelfData, const FGameplayEffectModCallbackData &Data) const OVERRIDE;
+	void PreGameplayEffectExecute(const FGameplayModifierEvaluatedData &SelfData, FGameplayEffectModCallbackData &Data) const override;
+	void PostGameplayEffectExecute(const FGameplayModifierEvaluatedData &SelfData, const FGameplayEffectModCallbackData &Data) const override;
 
 	/** The GameplayEffect to apply when restoring health to the instigator */
 	UPROPERTY(EditDefaultsOnly, Category = Lifesteal)
