Removed all animation classes from EnginePrivate.h and edited their files to include only their dependancies.

diff --git a/Engine/Source/Runtime/Engine/Classes/Animation/AimOffsetBlendSpace.h b/Engine/Source/Runtime/Engine/Classes/Animation/AimOffsetBlendSpace.h
--- a/Engine/Source/Runtime/Engine/Classes/Animation/AimOffsetBlendSpace.h
+++ b/Engine/Source/Runtime/Engine/Classes/Animation/AimOffsetBlendSpace.h
@@ -7,8 +7,10 @@
 
 #pragma once 
 
+#include "BlendSpace.h"
 #include "AimOffsetBlendSpace.generated.h"
 
+struct FBlendSample;
 
 UCLASS(config=Engine, hidecategories=Object, dependson=UAnimInstance, MinimalAPI, BlueprintType)
 class UAimOffsetBlendSpace : public UBlendSpace
diff --git a/Engine/Source/Runtime/Engine/Classes/Animation/AimOffsetBlendSpace1D.h b/Engine/Source/Runtime/Engine/Classes/Animation/AimOffsetBlendSpace1D.h
--- a/Engine/Source/Runtime/Engine/Classes/Animation/AimOffsetBlendSpace1D.h
+++ b/Engine/Source/Runtime/Engine/Classes/Animation/AimOffsetBlendSpace1D.h
@@ -7,6 +7,7 @@
 
 #pragma once 
 
+#include "BlendSpace1D.h"
 #include "AimOffsetBlendSpace1D.generated.h"
 
 
diff --git a/Engine/Source/Runtime/Engine/Classes/Animation/AnimBlueprintGeneratedClass.h b/Engine/Source/Runtime/Engine/Classes/Animation/AnimBlueprintGeneratedClass.h
--- a/Engine/Source/Runtime/Engine/Classes/Animation/AnimBlueprintGeneratedClass.h
+++ b/Engine/Source/Runtime/Engine/Classes/Animation/AnimBlueprintGeneratedClass.h
@@ -2,11 +2,12 @@
 
 #pragma once
 
+#include "AnimStateMachineTypes.h"
+#include "AnimSequenceBase.h"
 #include "Engine/BlueprintGeneratedClass.h"
 #include "AnimBlueprintGeneratedClass.generated.h"
 
 struct FAnimNotifyEvent;
-struct FBakedAnimationStateMachine;
 class  UAnimInstance;
 class  UEdGraph;
 class  UEdGraphNode;
diff --git a/Engine/Source/Runtime/Engine/Classes/Animation/AnimComposite.h b/Engine/Source/Runtime/Engine/Classes/Animation/AnimComposite.h
--- a/Engine/Source/Runtime/Engine/Classes/Animation/AnimComposite.h
+++ b/Engine/Source/Runtime/Engine/Classes/Animation/AnimComposite.h
@@ -6,6 +6,7 @@
  */
 
 #pragma once
+#include "AnimCompositeBase.h"
 #include "AnimComposite.generated.h"
 
 UCLASS(config=Engine, hidecategories=UObject, MinimalAPI, BlueprintType)
diff --git a/Engine/Source/Runtime/Engine/Classes/Animation/AnimInstance.h b/Engine/Source/Runtime/Engine/Classes/Animation/AnimInstance.h
--- a/Engine/Source/Runtime/Engine/Classes/Animation/AnimInstance.h
+++ b/Engine/Source/Runtime/Engine/Classes/Animation/AnimInstance.h
@@ -2,14 +2,18 @@
 
 #pragma once
 
-#include "AnimMontage.h"
-#include "AnimSequence.h"
 #include "AnimationAsset.h"
 #include "Components/SkeletalMeshComponent.h"
+#include "AnimSequence.h"
 #include "AnimInstance.generated.h"
 
 struct FBoneContainer;
+struct FAnimMontageInstance;
 class  USkeletalMesh;
+class UAnimMontage;
+
+DECLARE_DELEGATE_TwoParams(FOnMontageEnded, class UAnimMontage*, bool /*bInterrupted*/)
+DECLARE_DELEGATE_TwoParams(FOnMontageBlendingOutStarted, class UAnimMontage*, bool /*bInterrupted*/)
 
 /**
 * Delegate for when Montage is completed, whether interrupted or finished
diff --git a/Engine/Source/Runtime/Engine/Classes/Animation/AnimMontage.h b/Engine/Source/Runtime/Engine/Classes/Animation/AnimMontage.h
--- a/Engine/Source/Runtime/Engine/Classes/Animation/AnimMontage.h
+++ b/Engine/Source/Runtime/Engine/Classes/Animation/AnimMontage.h
@@ -9,6 +9,7 @@
 
 #include "Delegate.h"
 #include "AnimCompositeBase.h"
+#include "AnimInstance.h"
 #include "AnimMontage.generated.h"
 
 /**
diff --git a/Engine/Source/Runtime/Engine/Classes/Animation/AnimNotifies/AnimNotify.h b/Engine/Source/Runtime/Engine/Classes/Animation/AnimNotifies/AnimNotify.h
--- a/Engine/Source/Runtime/Engine/Classes/Animation/AnimNotifies/AnimNotify.h
+++ b/Engine/Source/Runtime/Engine/Classes/Animation/AnimNotifies/AnimNotify.h
@@ -4,6 +4,8 @@
 #pragma once
 #include "AnimNotify.generated.h"
 
+struct FAnimNotifyEvent;
+
 UCLASS(abstract, Blueprintable, const, hidecategories=Object, collapsecategories)
 class ENGINE_API UAnimNotify : public UObject
 {
diff --git a/Engine/Source/Runtime/Engine/Classes/Animation/AnimNotifies/AnimNotifyState.h b/Engine/Source/Runtime/Engine/Classes/Animation/AnimNotifies/AnimNotifyState.h
--- a/Engine/Source/Runtime/Engine/Classes/Animation/AnimNotifies/AnimNotifyState.h
+++ b/Engine/Source/Runtime/Engine/Classes/Animation/AnimNotifies/AnimNotifyState.h
@@ -4,6 +4,8 @@
 #pragma once
 #include "AnimNotifyState.generated.h"
 
+struct FAnimNotifyEvent;
+
 UCLASS(abstract, editinlinenew, Blueprintable, const, hidecategories=Object, collapsecategories, meta=(ShowHiddenSelfPins))
 class ENGINE_API UAnimNotifyState : public UObject
 {
diff --git a/Engine/Source/Runtime/Engine/Classes/Animation/AnimNotifies/AnimNotifyState_TimedParticleEffect.h b/Engine/Source/Runtime/Engine/Classes/Animation/AnimNotifies/AnimNotifyState_TimedParticleEffect.h
--- a/Engine/Source/Runtime/Engine/Classes/Animation/AnimNotifies/AnimNotifyState_TimedParticleEffect.h
+++ b/Engine/Source/Runtime/Engine/Classes/Animation/AnimNotifies/AnimNotifyState_TimedParticleEffect.h
@@ -1,6 +1,8 @@
 // Copyright 1998-2014 Epic Games, Inc. All Rights Reserved.
 
 #pragma once
+
+#include "AnimNotifyState.h"
 #include "AnimNotifyState_TimedParticleEffect.generated.h"
 
 // Timed Particle Effect Notify
diff --git a/Engine/Source/Runtime/Engine/Classes/Animation/AnimNotifies/AnimNotifyState_Trail.h b/Engine/Source/Runtime/Engine/Classes/Animation/AnimNotifies/AnimNotifyState_Trail.h
--- a/Engine/Source/Runtime/Engine/Classes/Animation/AnimNotifies/AnimNotifyState_Trail.h
+++ b/Engine/Source/Runtime/Engine/Classes/Animation/AnimNotifies/AnimNotifyState_Trail.h
@@ -2,6 +2,7 @@
 
 #pragma once
 #include "ParticleHelper.h"
+#include "AnimNotifyState.h"
 #include "AnimNotifyState_Trail.generated.h"
 
 DECLARE_LOG_CATEGORY_EXTERN(LogAnimTrails, Log, All);
diff --git a/Engine/Source/Runtime/Engine/Classes/Animation/BlendSpace.h b/Engine/Source/Runtime/Engine/Classes/Animation/BlendSpace.h
--- a/Engine/Source/Runtime/Engine/Classes/Animation/BlendSpace.h
+++ b/Engine/Source/Runtime/Engine/Classes/Animation/BlendSpace.h
@@ -2,6 +2,7 @@
 
 #pragma once
 
+#include "BlendSpaceBase.h"
 #include "AnimInterpFilter.h"
 #include "BlendSpace.generated.h"
 
diff --git a/Engine/Source/Runtime/Engine/Classes/Animation/BlendSpace1D.h b/Engine/Source/Runtime/Engine/Classes/Animation/BlendSpace1D.h
--- a/Engine/Source/Runtime/Engine/Classes/Animation/BlendSpace1D.h
+++ b/Engine/Source/Runtime/Engine/Classes/Animation/BlendSpace1D.h
@@ -8,6 +8,7 @@
 #pragma once
 
 #include "AnimInterpFilter.h"
+#include "BlendSpaceBase.h"
 #include "BlendSpace1D.generated.h"
 
 UCLASS(config=Engine, hidecategories=Object, dependson=UVimInstance, MinimalAPI, BlueprintType)
diff --git a/Engine/Source/Runtime/Engine/Classes/Animation/BlendSpaceBase.h b/Engine/Source/Runtime/Engine/Classes/Animation/BlendSpaceBase.h
--- a/Engine/Source/Runtime/Engine/Classes/Animation/BlendSpaceBase.h
+++ b/Engine/Source/Runtime/Engine/Classes/Animation/BlendSpaceBase.h
@@ -7,6 +7,7 @@
 
 #pragma once
 
+#include "AnimSequence.h"
 #include "BlendSpaceBase.generated.h"
 
 // Interpolation data types.
diff --git a/Engine/Source/Runtime/Engine/Classes/Vehicles/VehicleAnimInstance.h b/Engine/Source/Runtime/Engine/Classes/Vehicles/VehicleAnimInstance.h
--- a/Engine/Source/Runtime/Engine/Classes/Vehicles/VehicleAnimInstance.h
+++ b/Engine/Source/Runtime/Engine/Classes/Vehicles/VehicleAnimInstance.h
@@ -4,6 +4,7 @@
 /**
  *
  */
+#include "Animation/AnimInstance.h"
 #include "VehicleAnimInstance.generated.h"
 
 UCLASS(transient)
diff --git a/Engine/Source/Runtime/Engine/Private/Actor.cpp b/Engine/Source/Runtime/Engine/Private/Actor.cpp
--- a/Engine/Source/Runtime/Engine/Private/Actor.cpp
+++ b/Engine/Source/Runtime/Engine/Private/Actor.cpp
@@ -15,6 +15,7 @@
 #include "Particles/ParticleSystemComponent.h"
 #include "AI/Navigation/NavRelevantComponent.h"
 #include "VisualLog.h"
+#include "Animation/AnimInstance.h"
 
 //DEFINE_LOG_CATEGORY_STATIC(LogActor, Log, All);
 DEFINE_LOG_CATEGORY(LogActor);
diff --git a/Engine/Source/Runtime/Engine/Private/AnimBlueprint.cpp b/Engine/Source/Runtime/Engine/Private/AnimBlueprint.cpp
--- a/Engine/Source/Runtime/Engine/Private/AnimBlueprint.cpp
+++ b/Engine/Source/Runtime/Engine/Private/AnimBlueprint.cpp
@@ -3,6 +3,8 @@
 #include "EnginePrivate.h"
 #include "BlueprintUtilities.h"
 #include "LatentActions.h"
+#include "Animation/AnimBlueprint.h"
+#include "Animation/AnimBlueprintGeneratedClass.h"
 
 //////////////////////////////////////////////////////////////////////////
 // UAnimBlueprint
diff --git a/Engine/Source/Runtime/Engine/Private/Animation/AimOffsetBlendSpace.cpp b/Engine/Source/Runtime/Engine/Private/Animation/AimOffsetBlendSpace.cpp
--- a/Engine/Source/Runtime/Engine/Private/Animation/AimOffsetBlendSpace.cpp
+++ b/Engine/Source/Runtime/Engine/Private/Animation/AimOffsetBlendSpace.cpp
@@ -5,6 +5,9 @@
 =============================================================================*/ 
 
 #include "EnginePrivate.h"
+#include "Animation/BlendSpaceBase.h"
+#include "Animation/AimOffsetBlendSpace.h"
+#include "Animation/AnimSequence.h"
 
 UAimOffsetBlendSpace::UAimOffsetBlendSpace(const class FPostConstructInitializeProperties& PCIP)
 	: Super(PCIP)
diff --git a/Engine/Source/Runtime/Engine/Private/Animation/AimOffsetBlendSpace1D.cpp b/Engine/Source/Runtime/Engine/Private/Animation/AimOffsetBlendSpace1D.cpp
--- a/Engine/Source/Runtime/Engine/Private/Animation/AimOffsetBlendSpace1D.cpp
+++ b/Engine/Source/Runtime/Engine/Private/Animation/AimOffsetBlendSpace1D.cpp
@@ -5,6 +5,9 @@
 =============================================================================*/ 
 
 #include "EnginePrivate.h"
+#include "Animation/BlendSpaceBase.h"
+#include "Animation/AimOffsetBlendSpace1D.h"
+#include "Animation/AnimSequence.h"
 
 UAimOffsetBlendSpace1D::UAimOffsetBlendSpace1D(const class FPostConstructInitializeProperties& PCIP)
 	: Super(PCIP)
diff --git a/Engine/Source/Runtime/Engine/Private/Animation/AnimComposite.cpp b/Engine/Source/Runtime/Engine/Private/Animation/AnimComposite.cpp
--- a/Engine/Source/Runtime/Engine/Private/Animation/AnimComposite.cpp
+++ b/Engine/Source/Runtime/Engine/Private/Animation/AnimComposite.cpp
@@ -5,6 +5,7 @@
 =============================================================================*/ 
 
 #include "EnginePrivate.h"
+#include "Animation/AnimComposite.h"
 #include "AnimationUtils.h"
 #include "AnimationRuntime.h"
 
diff --git a/Engine/Source/Runtime/Engine/Private/Animation/AnimCompositeBase.cpp b/Engine/Source/Runtime/Engine/Private/Animation/AnimCompositeBase.cpp
--- a/Engine/Source/Runtime/Engine/Private/Animation/AnimCompositeBase.cpp
+++ b/Engine/Source/Runtime/Engine/Private/Animation/AnimCompositeBase.cpp
@@ -6,6 +6,7 @@
 
 #include "EnginePrivate.h"
 #include "AnimationUtils.h"
+#include "Animation/AnimCompositeBase.h"
 
 ///////////////////////////////////////////////////////
 // FAnimSegment
diff --git a/Engine/Source/Runtime/Engine/Private/Animation/AnimInstance.cpp b/Engine/Source/Runtime/Engine/Private/Animation/AnimInstance.cpp
--- a/Engine/Source/Runtime/Engine/Private/Animation/AnimInstance.cpp
+++ b/Engine/Source/Runtime/Engine/Private/Animation/AnimInstance.cpp
@@ -12,8 +12,12 @@
 #include "GameFramework/Character.h"
 #include "ParticleDefinitions.h"
 #include "DisplayDebugHelpers.h"
-
 #include "MessageLog.h"
+#include "Animation/BlendSpaceBase.h"
+#include "Animation/AnimComposite.h"
+#include "Animation/AnimNotifies/AnimNotifyState.h"
+#include "Animation/AnimNotifies/AnimNotify.h"
+#include "Animation/AnimMontage.h"
 
 /** Anim stats */
 
diff --git a/Engine/Source/Runtime/Engine/Private/Animation/AnimMontage.cpp b/Engine/Source/Runtime/Engine/Private/Animation/AnimMontage.cpp
--- a/Engine/Source/Runtime/Engine/Private/Animation/AnimMontage.cpp
+++ b/Engine/Source/Runtime/Engine/Private/Animation/AnimMontage.cpp
@@ -5,6 +5,7 @@
 =============================================================================*/ 
 
 #include "EnginePrivate.h"
+#include "Animation/AnimMontage.h"
 #include "AnimationUtils.h"
 #include "AnimationRuntime.h"
 #include "GameFramework/CharacterMovementComponent.h"
diff --git a/Engine/Source/Runtime/Engine/Private/Animation/AnimNode_BlendSpaceEvaluator.cpp b/Engine/Source/Runtime/Engine/Private/Animation/AnimNode_BlendSpaceEvaluator.cpp
--- a/Engine/Source/Runtime/Engine/Private/Animation/AnimNode_BlendSpaceEvaluator.cpp
+++ b/Engine/Source/Runtime/Engine/Private/Animation/AnimNode_BlendSpaceEvaluator.cpp
@@ -2,6 +2,7 @@
 
 #include "EnginePrivate.h"
 #include "Animation/AnimNode_BlendSpaceEvaluator.h"
+#include "Animation/BlendSpaceBase.h"
 
 /////////////////////////////////////////////////////
 // FAnimNode_BlendSpaceEvaluator
@@ -24,7 +25,7 @@ void FAnimNode_BlendSpaceEvaluator::Update(const FAnimationUpdateContext& Contex
 void FAnimNode_BlendSpaceEvaluator::GatherDebugData(FNodeDebugData& DebugData)
 {
 	FString DebugLine = DebugData.GetNodeName(this);
-	
+
 	DebugLine += FString::Printf(TEXT("('%s' Play Time: %.3f)"), *BlendSpace->GetName(), InternalTimeAccumulator);
 	DebugData.AddDebugItem(DebugLine, true);
 }
diff --git a/Engine/Source/Runtime/Engine/Private/Animation/AnimNode_BlendSpacePlayer.cpp b/Engine/Source/Runtime/Engine/Private/Animation/AnimNode_BlendSpacePlayer.cpp
--- a/Engine/Source/Runtime/Engine/Private/Animation/AnimNode_BlendSpacePlayer.cpp
+++ b/Engine/Source/Runtime/Engine/Private/Animation/AnimNode_BlendSpacePlayer.cpp
@@ -2,6 +2,7 @@
 
 #include "EnginePrivate.h"
 #include "Animation/AnimNode_BlendSpacePlayer.h"
+#include "Animation/BlendSpaceBase.h"
 
 /////////////////////////////////////////////////////
 // FAnimNode_BlendSpacePlayer
diff --git a/Engine/Source/Runtime/Engine/Private/Animation/AnimNode_Slot.cpp b/Engine/Source/Runtime/Engine/Private/Animation/AnimNode_Slot.cpp
--- a/Engine/Source/Runtime/Engine/Private/Animation/AnimNode_Slot.cpp
+++ b/Engine/Source/Runtime/Engine/Private/Animation/AnimNode_Slot.cpp
@@ -2,6 +2,7 @@
 
 #include "EnginePrivate.h"
 #include "Animation/AnimNode_Slot.h"
+#include "Animation/AnimMontage.h"
 
 /////////////////////////////////////////////////////
 // FAnimNode_Slot
diff --git a/Engine/Source/Runtime/Engine/Private/Animation/AnimNotify.cpp b/Engine/Source/Runtime/Engine/Private/Animation/AnimNotify.cpp
--- a/Engine/Source/Runtime/Engine/Private/Animation/AnimNotify.cpp
+++ b/Engine/Source/Runtime/Engine/Private/Animation/AnimNotify.cpp
@@ -1,6 +1,7 @@
 // Copyright 1998-2014 Epic Games, Inc. All Rights Reserved.
 
 #include "EnginePrivate.h"
+#include "Animation/AnimNotifies/AnimNotify.h"
 
 /////////////////////////////////////////////////////
 // UAnimNotify
diff --git a/Engine/Source/Runtime/Engine/Private/Animation/AnimNotifyState.cpp b/Engine/Source/Runtime/Engine/Private/Animation/AnimNotifyState.cpp
--- a/Engine/Source/Runtime/Engine/Private/Animation/AnimNotifyState.cpp
+++ b/Engine/Source/Runtime/Engine/Private/Animation/AnimNotifyState.cpp
@@ -1,6 +1,7 @@
 // Copyright 1998-2014 Epic Games, Inc. All Rights Reserved.
 
 #include "EnginePrivate.h"
+#include "Animation/AnimNotifies/AnimNotifyState.h"
 
 /////////////////////////////////////////////////////
 // UAnimNotifyState
diff --git a/Engine/Source/Runtime/Engine/Private/Animation/AnimNotifyState_TimedParticleEffect.cpp b/Engine/Source/Runtime/Engine/Private/Animation/AnimNotifyState_TimedParticleEffect.cpp
--- a/Engine/Source/Runtime/Engine/Private/Animation/AnimNotifyState_TimedParticleEffect.cpp
+++ b/Engine/Source/Runtime/Engine/Private/Animation/AnimNotifyState_TimedParticleEffect.cpp
@@ -1,6 +1,7 @@
 // Copyright 1998-2014 Epic Games, Inc. All Rights Reserved.
 
 #include "EnginePrivate.h"
+#include "Animation/AnimNotifies/AnimNotifyState_TimedParticleEffect.h"
 #include "Particles/ParticleSystemComponent.h"
 
 UAnimNotifyState_TimedParticleEffect::UAnimNotifyState_TimedParticleEffect(const class FPostConstructInitializeProperties& PCIP)
diff --git a/Engine/Source/Runtime/Engine/Private/Animation/AnimNotifyState_Trail.cpp b/Engine/Source/Runtime/Engine/Private/Animation/AnimNotifyState_Trail.cpp
--- a/Engine/Source/Runtime/Engine/Private/Animation/AnimNotifyState_Trail.cpp
+++ b/Engine/Source/Runtime/Engine/Private/Animation/AnimNotifyState_Trail.cpp
@@ -1,11 +1,15 @@
 // Copyright 1998-2014 Epic Games, Inc. All Rights Reserved.
 
 #include "EnginePrivate.h"
-#include "MessageLog.h"
+#include "Animation/AnimNotifies/AnimNotifyState_Trail.h"
 #include "ParticleDefinitions.h"
 #include "Particles/TypeData/ParticleModuleTypeDataAnimTrail.h"
 #include "Particles/ParticleSystem.h"
 #include "Particles/ParticleSystemComponent.h"
+#include "MessageLog.h"
+#include "AssertionMacros.h"
+#include "Animation/AnimInstance.h"
+
 #define LOCTEXT_NAMESPACE "AnimNotifyState_Trail"
 
 DEFINE_LOG_CATEGORY(LogAnimTrails);
diff --git a/Engine/Source/Runtime/Engine/Private/Animation/AnimSequence.cpp b/Engine/Source/Runtime/Engine/Private/Animation/AnimSequence.cpp
--- a/Engine/Source/Runtime/Engine/Private/Animation/AnimSequence.cpp
+++ b/Engine/Source/Runtime/Engine/Private/Animation/AnimSequence.cpp
@@ -11,6 +11,8 @@
 #include "AnimationRuntime.h"
 #include "TargetPlatform.h"
 #include "Animation/AnimCompress_RevertToRaw.h"
+#include "Animation/AnimNotifies/AnimNotify.h"
+#include "Animation/AnimNotifies/AnimNotifyState.h"
 
 #define USE_SLERP 0
 
diff --git a/Engine/Source/Runtime/Engine/Private/Animation/AnimSequenceBase.cpp b/Engine/Source/Runtime/Engine/Private/Animation/AnimSequenceBase.cpp
--- a/Engine/Source/Runtime/Engine/Private/Animation/AnimSequenceBase.cpp
+++ b/Engine/Source/Runtime/Engine/Private/Animation/AnimSequenceBase.cpp
@@ -3,6 +3,7 @@
 #include "EnginePrivate.h"
 #include "AnimationUtils.h"
 #include "AnimationRuntime.h"
+#include "Animation/AnimNotifies/AnimNotify.h"
 
 #define NOTIFY_TRIGGER_OFFSET KINDA_SMALL_NUMBER;
 
diff --git a/Engine/Source/Runtime/Engine/Private/Animation/AnimSingleNodeInstance.cpp b/Engine/Source/Runtime/Engine/Private/Animation/AnimSingleNodeInstance.cpp
--- a/Engine/Source/Runtime/Engine/Private/Animation/AnimSingleNodeInstance.cpp
+++ b/Engine/Source/Runtime/Engine/Private/Animation/AnimSingleNodeInstance.cpp
@@ -7,8 +7,15 @@
 
 #include "EnginePrivate.h"
 #include "Animation/AnimNodeBase.h"
+#include "Animation/AnimSingleNodeInstance.h"
 #include "Animation/VertexAnim/VertexAnimation.h"
 #include "AnimationRuntime.h"
+#include "Animation/BlendSpace.h"
+#include "Animation/BlendSpaceBase.h"
+#include "Animation/AnimComposite.h"
+#include "Animation/AimOffsetBlendSpace.h"
+#include "Animation/AimOffsetBlendSpace1D.h"
+#include "Animation/AnimMontage.h"
 
 /////////////////////////////////////////////////////
 // UAnimSingleNodeInstance
diff --git a/Engine/Source/Runtime/Engine/Private/Animation/AnimStateMachineTypes.cpp b/Engine/Source/Runtime/Engine/Private/Animation/AnimStateMachineTypes.cpp
--- a/Engine/Source/Runtime/Engine/Private/Animation/AnimStateMachineTypes.cpp
+++ b/Engine/Source/Runtime/Engine/Private/Animation/AnimStateMachineTypes.cpp
@@ -3,6 +3,7 @@
 #include "EnginePrivate.h"
 #include "AnimationRuntime.h"
 #include "AnimationUtils.h"
+#include "Animation/AnimStateMachineTypes.h"
 
 /////////////////////////////////////////////////////
 // UAnimStateMachineTypes
diff --git a/Engine/Source/Runtime/Engine/Private/Animation/AnimationRuntime.cpp b/Engine/Source/Runtime/Engine/Private/Animation/AnimationRuntime.cpp
--- a/Engine/Source/Runtime/Engine/Private/Animation/AnimationRuntime.cpp
+++ b/Engine/Source/Runtime/Engine/Private/Animation/AnimationRuntime.cpp
@@ -10,6 +10,8 @@
 #include "AnimationRuntime.h"
 #include "AnimationUtils.h"
 #include "Animation/AnimData/BoneMaskFilter.h"
+#include "Animation/BlendSpaceBase.h"
+#include "Animation/AnimCompositeBase.h"
 
 DEFINE_LOG_CATEGORY(LogAnimation);
 DEFINE_LOG_CATEGORY(LogRootMotion);
diff --git a/Engine/Source/Runtime/Engine/Private/Animation/BlendSpace.cpp b/Engine/Source/Runtime/Engine/Private/Animation/BlendSpace.cpp
--- a/Engine/Source/Runtime/Engine/Private/Animation/BlendSpace.cpp
+++ b/Engine/Source/Runtime/Engine/Private/Animation/BlendSpace.cpp
@@ -5,6 +5,8 @@
 =============================================================================*/ 
 
 #include "EnginePrivate.h"
+#include "Animation/BlendSpaceBase.h"
+#include "Animation/BlendSpace.h"
 #include "AnimationUtils.h"
 
 UBlendSpace::UBlendSpace(const class FPostConstructInitializeProperties& PCIP)
diff --git a/Engine/Source/Runtime/Engine/Private/Animation/BlendSpace1D.cpp b/Engine/Source/Runtime/Engine/Private/Animation/BlendSpace1D.cpp
--- a/Engine/Source/Runtime/Engine/Private/Animation/BlendSpace1D.cpp
+++ b/Engine/Source/Runtime/Engine/Private/Animation/BlendSpace1D.cpp
@@ -7,6 +7,7 @@
 #include "EnginePrivate.h"
 #include "AnimationRuntime.h"
 #include "AnimationUtils.h"
+#include "Animation/BlendSpace1D.h"
 
 UBlendSpace1D::UBlendSpace1D(const class FPostConstructInitializeProperties& PCIP)
 	: Super(PCIP)
diff --git a/Engine/Source/Runtime/Engine/Private/Animation/BlendSpaceBase.cpp b/Engine/Source/Runtime/Engine/Private/Animation/BlendSpaceBase.cpp
--- a/Engine/Source/Runtime/Engine/Private/Animation/BlendSpaceBase.cpp
+++ b/Engine/Source/Runtime/Engine/Private/Animation/BlendSpaceBase.cpp
@@ -6,6 +6,7 @@
 
 #include "EnginePrivate.h"
 #include "AnimationRuntime.h"
+#include "Animation/BlendSpaceBase.h"
 
 // global variable that's used in editor when you change the property only
 // I can't make this member since I  need to change this in const function
diff --git a/Engine/Source/Runtime/Engine/Private/AnimationAsset.cpp b/Engine/Source/Runtime/Engine/Private/AnimationAsset.cpp
--- a/Engine/Source/Runtime/Engine/Private/AnimationAsset.cpp
+++ b/Engine/Source/Runtime/Engine/Private/AnimationAsset.cpp
@@ -1,6 +1,7 @@
 // Copyright 1998-2014 Epic Games, Inc. All Rights Reserved.
 
 #include "EnginePrivate.h"
+#include "Animation/AnimSequence.h"
 
 UAnimationAsset::UAnimationAsset(const class FPostConstructInitializeProperties& PCIP)
 	: Super(PCIP)
diff --git a/Engine/Source/Runtime/Engine/Private/Character.cpp b/Engine/Source/Runtime/Engine/Private/Character.cpp
--- a/Engine/Source/Runtime/Engine/Private/Character.cpp
+++ b/Engine/Source/Runtime/Engine/Private/Character.cpp
@@ -9,6 +9,8 @@
 #include "GameFramework/Character.h"
 #include "Net/UnrealNetwork.h"
 #include "DisplayDebugHelpers.h"
+#include "Animation/AnimInstance.h"
+#include "Animation/AnimMontage.h"
 
 DEFINE_LOG_CATEGORY_STATIC(LogCharacter, Log, All);
 DEFINE_LOG_CATEGORY_STATIC(LogAvatar, Log, All);
diff --git a/Engine/Source/Runtime/Engine/Private/CharacterMovementComponent.cpp b/Engine/Source/Runtime/Engine/Private/CharacterMovementComponent.cpp
--- a/Engine/Source/Runtime/Engine/Private/CharacterMovementComponent.cpp
+++ b/Engine/Source/Runtime/Engine/Private/CharacterMovementComponent.cpp
@@ -9,6 +9,7 @@
 #include "GameFramework/PhysicsVolume.h"
 #include "GameFramework/GameNetworkManager.h"
 #include "GameFramework/Character.h"
+#include "Animation/AnimMontage.h"
 
 // @todo this is here only due to circular dependency to AIModule. To be removed
 #include "Navigation/PathFollowingComponent.h"
diff --git a/Engine/Source/Runtime/Engine/Private/EnginePrivate.h b/Engine/Source/Runtime/Engine/Private/EnginePrivate.h
--- a/Engine/Source/Runtime/Engine/Private/EnginePrivate.h
+++ b/Engine/Source/Runtime/Engine/Private/EnginePrivate.h
@@ -156,23 +156,6 @@
 #include "GameFramework/SpringArmComponent.h"
 #include "Components/WindDirectionalSourceComponent.h"
 #include "Components/TimelineComponent.h"
-#include "Animation/AnimSequenceBase.h"
-#include "Animation/AnimCompositeBase.h"
-#include "Animation/AnimComposite.h"
-#include "Animation/AnimMontage.h"
-#include "Animation/AnimSequence.h"
-#include "Animation/BlendSpaceBase.h"
-#include "Animation/AnimStateMachineTypes.h"
-#include "Animation/AnimInstance.h"
-#include "Animation/BlendSpace.h"
-#include "Animation/AimOffsetBlendSpace.h"
-#include "Animation/BlendSpace1D.h"
-#include "Animation/AimOffsetBlendSpace1D.h"
-#include "Animation/AnimSingleNodeInstance.h"
-#include "Animation/AnimNotifies/AnimNotify.h"
-#include "Animation/AnimNotifies/AnimNotifyState.h"
-#include "Animation/AnimNotifies/AnimNotifyState_TimedParticleEffect.h"
-#include "Animation/AnimNotifies/AnimNotifyState_Trail.h"
 #include "EditorFramework/AssetImportData.h"
 #include "Engine/AssetUserData.h"
 #include "Sound/AudioSettings.h"
@@ -180,7 +163,6 @@
 #include "Engine/BlendableInterface.h"
 #include "Engine/BlueprintCore.h"
 #include "Engine/Blueprint.h"
-#include "Animation/AnimBlueprint.h"
 #include "Sound/DialogueTypes.h"
 #include "Kismet/GameplayStatics.h"
 #include "Kismet/HeadMountedDisplayFunctionLibrary.h"
@@ -193,7 +175,6 @@
 #include "Kismet/KismetSystemLibrary.h"
 #include "Kismet/KismetTextLibrary.h"
 #include "Engine/BlueprintGeneratedClass.h"
-#include "Animation/AnimBlueprintGeneratedClass.h"
 #include "PhysicsEngine/BodySetup.h"
 #include "Camera/CameraAnim.h"
 #include "Camera/CameraAnimInst.h"
@@ -294,11 +275,9 @@
 #include "Engine/SCS_Node.h"
 #include "Engine/Selection.h"
 #include "Engine/SimpleConstructionScript.h"
-#include "Animation/PreviewAssetAttachComponent.h"
 #include "Engine/SkeletalMesh.h"
 #include "Engine/DestructibleMesh.h"
 #include "Engine/SkeletalMeshSocket.h"
-#include "Animation/Skeleton.h"
 #include "Sound/SoundBase.h"
 #include "Sound/DialogueSoundWaveProxy.h"
 #include "Sound/SoundClass.h"
@@ -333,7 +312,6 @@
 #include "VectorField/VectorField.h"
 #include "VectorField/VectorFieldAnimated.h"
 #include "VectorField/VectorFieldStatic.h"
-#include "Animation/VertexAnim/MorphTarget.h"
 #include "Engine/World.h"
 
 #include "SystemSettings.h"					// Scalability options.
diff --git a/Engine/Source/Runtime/Engine/Private/MorphMesh.cpp b/Engine/Source/Runtime/Engine/Private/MorphMesh.cpp
--- a/Engine/Source/Runtime/Engine/Private/MorphMesh.cpp
+++ b/Engine/Source/Runtime/Engine/Private/MorphMesh.cpp
@@ -5,6 +5,8 @@
 =============================================================================*/
 
 #include "EnginePrivate.h"
+#include "Animation/VertexAnim/VertexAnimBase.h"
+#include "Animation/VertexAnim/MorphTarget.h"
 
 UVertexAnimBase::UVertexAnimBase(const class FPostConstructInitializeProperties& PCIP)
 	: Super(PCIP)
diff --git a/Engine/Source/Runtime/Engine/Private/MorphTools.cpp b/Engine/Source/Runtime/Engine/Private/MorphTools.cpp
--- a/Engine/Source/Runtime/Engine/Private/MorphTools.cpp
+++ b/Engine/Source/Runtime/Engine/Private/MorphTools.cpp
@@ -8,6 +8,7 @@
 #include "RawIndexBuffer.h"
 #include "MeshBuild.h"
 #include "PhysicsEngine/PhysXSupport.h"
+#include "Animation/VertexAnim/MorphTarget.h"
 
 
 /** compare based on base mesh source vertex indices */
diff --git a/Engine/Source/Runtime/Engine/Private/PoseableMeshComponent.cpp b/Engine/Source/Runtime/Engine/Private/PoseableMeshComponent.cpp
--- a/Engine/Source/Runtime/Engine/Private/PoseableMeshComponent.cpp
+++ b/Engine/Source/Runtime/Engine/Private/PoseableMeshComponent.cpp
@@ -6,6 +6,7 @@
 
 #include "EnginePrivate.h"
 #include "AnimTree.h"
+#include "Animation/AnimInstance.h"
 
 UPoseableMeshComponent::UPoseableMeshComponent(const class FPostConstructInitializeProperties& PCIP)
 	: Super(PCIP)
diff --git a/Engine/Source/Runtime/Engine/Private/SkeletalMesh.cpp b/Engine/Source/Runtime/Engine/Private/SkeletalMesh.cpp
--- a/Engine/Source/Runtime/Engine/Private/SkeletalMesh.cpp
+++ b/Engine/Source/Runtime/Engine/Private/SkeletalMesh.cpp
@@ -20,6 +20,10 @@
 #include "SkeletalRenderGPUSkin.h"
 #include "RawIndexBuffer.h"
 #include "PhysicsPublic.h"
+#include "Animation/VertexAnim/MorphTarget.h"
+#include "Animation/AnimSingleNodeInstance.h"
+#include "Animation/AnimBlueprint.h"
+#include "Animation/AnimBlueprintGeneratedClass.h"
 
 #if WITH_EDITOR
 #include "MeshUtilities.h"
diff --git a/Engine/Source/Runtime/Engine/Private/SkeletalMeshComponent.cpp b/Engine/Source/Runtime/Engine/Private/SkeletalMeshComponent.cpp
--- a/Engine/Source/Runtime/Engine/Private/SkeletalMeshComponent.cpp
+++ b/Engine/Source/Runtime/Engine/Private/SkeletalMeshComponent.cpp
@@ -18,6 +18,7 @@
 #include "Animation/VertexAnim/VertexAnimation.h"
 #include "GameFramework/Character.h"
 #include "Particles/ParticleSystemComponent.h"
+#include "Animation/AnimSingleNodeInstance.h"
 
 #ifndef EXPERIMENTAL_PARALLEL_CODE  
 	#error EXPERIMENTAL_PARALLEL_CODE must be defined as either zero or one
diff --git a/Engine/Source/Runtime/Engine/Private/SkeletalRenderCPUSkin.cpp b/Engine/Source/Runtime/Engine/Private/SkeletalRenderCPUSkin.cpp
--- a/Engine/Source/Runtime/Engine/Private/SkeletalRenderCPUSkin.cpp
+++ b/Engine/Source/Runtime/Engine/Private/SkeletalRenderCPUSkin.cpp
@@ -22,6 +22,10 @@
 
 #include "EnginePrivate.h"
 #include "SkeletalRenderCPUSkin.h"
+#include "Animation/VertexAnim/VertexAnimBase.h"
+
+struct FVertexAnimDelta;
+struct FVertexAnimEvalStateBase;
 
 template<typename BaseVertexType, typename VertexType>
 static void SkinVertices( FFinalSkinVertex* DestVertex, FMatrix* ReferenceToLocal, int32 LODIndex, FStaticLODModel& LOD, TArray<FActiveVertexAnim>& ActiveVertexAnims  );
diff --git a/Engine/Source/Runtime/Engine/Private/SkeletalRenderGPUSkin.cpp b/Engine/Source/Runtime/Engine/Private/SkeletalRenderGPUSkin.cpp
--- a/Engine/Source/Runtime/Engine/Private/SkeletalRenderGPUSkin.cpp
+++ b/Engine/Source/Runtime/Engine/Private/SkeletalRenderGPUSkin.cpp
@@ -25,6 +25,7 @@
 #include "SkeletalRenderGPUSkin.h"
 #include "SkeletalRenderCPUSkin.h"
 #include "GPUSkinCache.h"
+#include "Animation/VertexAnim/VertexAnimBase.h"
 
 DEFINE_LOG_CATEGORY_STATIC(LogSkeletalGPUSkinMesh, Warning, All);
 
@@ -364,7 +365,7 @@ void FSkeletalMeshObjectGPUSkin::FSkeletalMeshObjectLOD::UpdateMorphVertexBuffer
 
 				// Allocate temp state
 				FVertexAnimEvalStateBase* AnimState = VertAnim.VertAnim->InitEval();
-
+				
 				// Get deltas
 				int32 NumDeltas;
 				FVertexAnimDelta* Deltas = VertAnim.VertAnim->GetDeltasAtTime(VertAnim.Time, LODIndex, AnimState, NumDeltas);
diff --git a/Engine/Source/Runtime/Engine/Private/SkinnedMeshComponent.cpp b/Engine/Source/Runtime/Engine/Private/SkinnedMeshComponent.cpp
--- a/Engine/Source/Runtime/Engine/Private/SkinnedMeshComponent.cpp
+++ b/Engine/Source/Runtime/Engine/Private/SkinnedMeshComponent.cpp
@@ -11,6 +11,7 @@
 #include "SkeletalRenderGPUSkin.h"
 #include "AnimationUtils.h"
 #include "AnimTree.h"
+#include "Animation/VertexAnim/MorphTarget.h"
 
 DEFINE_LOG_CATEGORY_STATIC(LogSkinnedMeshComp, Log, All);
 
diff --git a/Engine/Source/Runtime/Engine/Public/AnimEncoding.h b/Engine/Source/Runtime/Engine/Public/AnimEncoding.h
--- a/Engine/Source/Runtime/Engine/Public/AnimEncoding.h
+++ b/Engine/Source/Runtime/Engine/Public/AnimEncoding.h
@@ -6,6 +6,8 @@
 
 #pragma once
 
+#include "Animation/AnimSequence.h"
+
 // switches to toggle subsets of the new animation codec system
 #define USE_ANIMATION_CODEC_BATCH_SOLVER 1
 
diff --git a/Engine/Source/Runtime/Engine/Public/AnimationCompression.h b/Engine/Source/Runtime/Engine/Public/AnimationCompression.h
--- a/Engine/Source/Runtime/Engine/Public/AnimationCompression.h
+++ b/Engine/Source/Runtime/Engine/Public/AnimationCompression.h
@@ -8,6 +8,7 @@
 #define __ANIMATIONCOMPRESSION_H__
 
 #include "FloatPacker.h"
+#include "Animation/AnimSequence.h"
 
 DECLARE_LOG_CATEGORY_EXTERN(LogAnimationCompression, Log, All);
 
diff --git a/Engine/Source/Runtime/Engine/Public/AnimationRuntime.h b/Engine/Source/Runtime/Engine/Public/AnimationRuntime.h
--- a/Engine/Source/Runtime/Engine/Public/AnimationRuntime.h
+++ b/Engine/Source/Runtime/Engine/Public/AnimationRuntime.h
@@ -6,8 +6,12 @@
 =============================================================================*/ 
 
 #pragma once
+#include "Animation/AnimInstance.h"
 
 struct FInputBlendPose;
+struct FA2Pose;
+struct FAnimTrack;
+class UBlendSpaceBase;
 
 /** In AnimationRunTime Library, we extract animation data based on Skeleton hierarchy, not ref pose hierarchy. 
 	Ref pose will need to be re-mapped later **/
diff --git a/Engine/Source/Runtime/Engine/Public/AnimationUtils.h b/Engine/Source/Runtime/Engine/Public/AnimationUtils.h
--- a/Engine/Source/Runtime/Engine/Public/AnimationUtils.h
+++ b/Engine/Source/Runtime/Engine/Public/AnimationUtils.h
@@ -10,9 +10,10 @@
 ENGINE_API DECLARE_LOG_CATEGORY_EXTERN(LogMatinee, Warning, All);
 ENGINE_API DECLARE_LOG_CATEGORY_EXTERN(LogAnimNotify, Warning, All);
 
+#include "Animation/AnimSequence.h"
+
 // Forward declarations.
 class UAnimCompress;
-class UAnimSequence;
 class UAnimSet;
 class USkeletalMesh;
 struct FAnimSetMeshLinkup;
