fix NULL crash in EnsureCollisionTreeIsBuilt

diff --git a/Engine/Source/Runtime/Engine/Private/World.cpp b/Engine/Source/Runtime/Engine/Private/World.cpp
--- a/Engine/Source/Runtime/Engine/Private/World.cpp
+++ b/Engine/Source/Runtime/Engine/Private/World.cpp
@@ -1310,20 +1310,18 @@ void UWorld::EnsureCollisionTreeIsBuilt()
 	if (PhysicsScene)
 	{
 		PhysicsScene->SetIsStaticLoading(true);
-	}
 
-	for (int Iteration = 0; Iteration < 6; ++Iteration)
-	{
-		SetupPhysicsTickFunctions(0.1f);
-		PhysicsScene->StartFrame();
-        PhysicsScene->WaitPhysScenes();
-		PhysicsScene->EndFrame(NULL);
-	}
+		for (int Iteration = 0; Iteration < 6; ++Iteration)
+		{
+			SetupPhysicsTickFunctions(0.1f);
+			PhysicsScene->StartFrame();
+			PhysicsScene->WaitPhysScenes();
+			PhysicsScene->EndFrame(NULL);
+		}
 
-	if (PhysicsScene)
-	{
 		PhysicsScene->SetIsStaticLoading(false);
 	}
+
     bIsBuilt = true;
 }
 
