Remove sound headers from  EnginePrivate.h

diff --git a/Engine/Source/Developer/AssetTools/Private/AssetTypeActions/AssetTypeActions_SoundClass.cpp b/Engine/Source/Developer/AssetTools/Private/AssetTypeActions/AssetTypeActions_SoundClass.cpp
--- a/Engine/Source/Developer/AssetTools/Private/AssetTypeActions/AssetTypeActions_SoundClass.cpp
+++ b/Engine/Source/Developer/AssetTools/Private/AssetTypeActions/AssetTypeActions_SoundClass.cpp
@@ -1,7 +1,7 @@
 // Copyright 1998-2014 Epic Games, Inc. All Rights Reserved.
 
 #include "AssetToolsPrivatePCH.h"
-#include "SoundDefinitions.h"
+#include "Sound/SoundClass.h"
 #include "Editor/SoundClassEditor/Public/SoundClassEditorModule.h"
 
 #define LOCTEXT_NAMESPACE "AssetTypeActions"
diff --git a/Engine/Source/Runtime/Engine/Classes/Components/AudioComponent.h b/Engine/Source/Runtime/Engine/Classes/Components/AudioComponent.h
--- a/Engine/Source/Runtime/Engine/Classes/Components/AudioComponent.h
+++ b/Engine/Source/Runtime/Engine/Classes/Components/AudioComponent.h
@@ -3,6 +3,7 @@
 
 #pragma once
 #include "Audio.h"
+#include "Sound/SoundAttenuation.h"
 #include "AudioComponent.generated.h"
 
 /** called when we finish playing audio, either because it played to completion or because a Stop() call turned it off early */
diff --git a/Engine/Source/Runtime/Engine/Classes/Sound/DialogueSoundWaveProxy.h b/Engine/Source/Runtime/Engine/Classes/Sound/DialogueSoundWaveProxy.h
--- a/Engine/Source/Runtime/Engine/Classes/Sound/DialogueSoundWaveProxy.h
+++ b/Engine/Source/Runtime/Engine/Classes/Sound/DialogueSoundWaveProxy.h
@@ -2,8 +2,10 @@
 
 #pragma once
 
+#include "Sound/SoundBase.h"
 #include "DialogueSoundWaveProxy.generated.h"
 
+
 UCLASS()
 class UDialogueSoundWaveProxy : public USoundBase
 {
diff --git a/Engine/Source/Runtime/Engine/Classes/Sound/SoundNode.h b/Engine/Source/Runtime/Engine/Classes/Sound/SoundNode.h
--- a/Engine/Source/Runtime/Engine/Classes/Sound/SoundNode.h
+++ b/Engine/Source/Runtime/Engine/Classes/Sound/SoundNode.h
@@ -2,6 +2,7 @@
 
 
 #pragma once
+#include "ActiveSound.h"
 #include "SoundNode.generated.h"
 
 /*-----------------------------------------------------------------------------
@@ -163,16 +164,7 @@ public:
 	 * Used to create a unique string to identify unique nodes
 	 */
 	virtual FString GetUniqueString() const;
-	static FORCEINLINE UPTRINT GetNodeWaveInstanceHash(const UPTRINT ParentWaveInstanceHash, const USoundNode* ChildNode, const uint32 ChildIndex)
-	{
-		checkf(ChildIndex < MAX_ALLOWED_CHILD_NODES, TEXT("Too many children (%d) in SoundCue '%s'"), ChildIndex, *CastChecked<USoundCue>(ChildNode->GetOuter())->GetFullName());
-		return ((ParentWaveInstanceHash << ChildIndex) ^ (UPTRINT)ChildNode);
-	}
-
-	static FORCEINLINE UPTRINT GetNodeWaveInstanceHash(const UPTRINT ParentWaveInstanceHash, const UPTRINT ChildNodeHash, const uint32 ChildIndex)
-	{
-		checkf(ChildIndex < MAX_ALLOWED_CHILD_NODES, TEXT("Too many children (%d) in SoundCue"), ChildIndex);
-		return ((ParentWaveInstanceHash << ChildIndex) ^ ChildNodeHash);
-	}
+	static UPTRINT GetNodeWaveInstanceHash(const UPTRINT ParentWaveInstanceHash, const USoundNode* ChildNode, const uint32 ChildIndex);
+	static UPTRINT GetNodeWaveInstanceHash(const UPTRINT ParentWaveInstanceHash, const UPTRINT ChildNodeHash, const uint32 ChildIndex);
 };
 
diff --git a/Engine/Source/Runtime/Engine/Classes/Sound/SoundWave.h b/Engine/Source/Runtime/Engine/Classes/Sound/SoundWave.h
--- a/Engine/Source/Runtime/Engine/Classes/Sound/SoundWave.h
+++ b/Engine/Source/Runtime/Engine/Classes/Sound/SoundWave.h
@@ -6,6 +6,8 @@
  * Playable sound object for raw wave files
  */
 
+#include "Sound/SoundBase.h"
+#include "Sound/SoundGroups.h"
 #include "SoundWave.generated.h"
 
 struct FActiveSound;
diff --git a/Engine/Source/Runtime/Engine/Classes/Sound/SoundWaveStreaming.h b/Engine/Source/Runtime/Engine/Classes/Sound/SoundWaveStreaming.h
--- a/Engine/Source/Runtime/Engine/Classes/Sound/SoundWaveStreaming.h
+++ b/Engine/Source/Runtime/Engine/Classes/Sound/SoundWaveStreaming.h
@@ -6,6 +6,7 @@
  */
 
 #pragma once
+#include "Sound/SoundWave.h"
 #include "SoundWaveStreaming.generated.h"
 
 DECLARE_DELEGATE_TwoParams( FOnSoundWaveStreamingUnderflow, class USoundWaveStreaming*, int32 );
diff --git a/Engine/Source/Runtime/Engine/Private/ActiveSound.cpp b/Engine/Source/Runtime/Engine/Private/ActiveSound.cpp
--- a/Engine/Source/Runtime/Engine/Private/ActiveSound.cpp
+++ b/Engine/Source/Runtime/Engine/Private/ActiveSound.cpp
@@ -1,7 +1,9 @@
 // Copyright 1998-2014 Epic Games, Inc. All Rights Reserved.
 
 #include "EnginePrivate.h"
-#include "SoundDefinitions.h"
+#include "AudioDevice.h"
+#include "Sound/SoundCue.h"
+#include "Sound/SoundWave.h"
 #include "Sound/SoundNodeAttenuation.h"
 #include "SubtitleManager.h"
 
@@ -105,7 +107,7 @@ USoundClass* FActiveSound::GetSoundClass() const
 	return NULL;
 }
 
-int32 FActiveSound::FindClosestListener( const TArray<struct FListener>& InListeners ) const
+int32 FActiveSound::FindClosestListener( const TArray<FListener>& InListeners ) const
 {
 	int32 ClosestListenerIndex = 0;
 	float ClosestDistSq = FVector::DistSquared(Transform.GetTranslation(), InListeners[0].Transform.GetTranslation() );
diff --git a/Engine/Source/Runtime/Engine/Private/Actor.cpp b/Engine/Source/Runtime/Engine/Private/Actor.cpp
--- a/Engine/Source/Runtime/Engine/Private/Actor.cpp
+++ b/Engine/Source/Runtime/Engine/Private/Actor.cpp
@@ -4,7 +4,7 @@
 #include "EnginePrivate.h"
 #include "PhysicsPublic.h"
 #include "ParticleDefinitions.h"
-#include "SoundDefinitions.h"
+#include "Sound/SoundCue.h"
 #include "LatentActions.h"
 #include "MessageLog.h"
 #include "Net/UnrealNetwork.h"
@@ -2691,12 +2691,12 @@ AWorldSettings * AActor::GetWorldSettings() const
 	return GetWorld()->GetWorldSettings();
 }
 
-void AActor::PlaySoundOnActor(class USoundCue* InSoundCue, float VolumeMultiplier/*=1.f*/, float PitchMultiplier/*=1.f*/)
+void AActor::PlaySoundOnActor(USoundCue* InSoundCue, float VolumeMultiplier/*=1.f*/, float PitchMultiplier/*=1.f*/)
 {
 	UGameplayStatics::PlaySoundAtLocation( this, InSoundCue, GetActorLocation(), VolumeMultiplier, PitchMultiplier );
 }
 
-void AActor::PlaySoundAtLocation(class USoundCue* InSoundCue, FVector SoundLocation, float VolumeMultiplier/*=1.f*/, float PitchMultiplier/*=1.f*/)
+void AActor::PlaySoundAtLocation(USoundCue* InSoundCue, FVector SoundLocation, float VolumeMultiplier/*=1.f*/, float PitchMultiplier/*=1.f*/)
 {
 	UGameplayStatics::PlaySoundAtLocation( this, InSoundCue, (SoundLocation.IsZero() ? GetActorLocation() : SoundLocation), VolumeMultiplier, PitchMultiplier );
 }
diff --git a/Engine/Source/Runtime/Engine/Private/AmbientSound.cpp b/Engine/Source/Runtime/Engine/Private/AmbientSound.cpp
--- a/Engine/Source/Runtime/Engine/Private/AmbientSound.cpp
+++ b/Engine/Source/Runtime/Engine/Private/AmbientSound.cpp
@@ -5,6 +5,9 @@
 =============================================================================*/
 
 #include "EnginePrivate.h"
+#include "Sound/AmbientSound.h"
+#include "Sound/SoundBase.h"
+#include "Sound/SoundCue.h"
 #include "Sound/SoundNodeAmbient.h" // deprecated
 #include "Sound/SoundNodeAmbientNonLoop.h" // deprecated
 #include "Sound/SoundNodeMixer.h"
@@ -13,7 +16,6 @@
 #include "Sound/SoundNodeDelay.h"
 #include "Sound/SoundNodeWavePlayer.h"
 #include "Sound/SoundNodeModulator.h"
-#include "SoundDefinitions.h"
 #include "MessageLog.h"
 #include "UObjectToken.h"
 #include "MapErrors.h"
diff --git a/Engine/Source/Runtime/Engine/Private/Audio.cpp b/Engine/Source/Runtime/Engine/Private/Audio.cpp
--- a/Engine/Source/Runtime/Engine/Private/Audio.cpp
+++ b/Engine/Source/Runtime/Engine/Private/Audio.cpp
@@ -5,8 +5,11 @@
 =============================================================================*/
 
 #include "EnginePrivate.h"
+#include "Audio.h"
+#include "AudioDevice.h"
+#include "Sound/SoundCue.h"
+#include "Sound/SoundWave.h"
 #include "Sound/SoundNodeWavePlayer.h"
-#include "SoundDefinitions.h"
 #include "AudioEffect.h"
 #include "Net/UnrealNetwork.h"
 #include "TargetPlatform.h"
diff --git a/Engine/Source/Runtime/Engine/Private/AudioComponent.cpp b/Engine/Source/Runtime/Engine/Private/AudioComponent.cpp
--- a/Engine/Source/Runtime/Engine/Private/AudioComponent.cpp
+++ b/Engine/Source/Runtime/Engine/Private/AudioComponent.cpp
@@ -3,6 +3,7 @@
 #include "EnginePrivate.h"
 #include "SoundDefinitions.h"
 #include "Sound/SoundNodeAttenuation.h"
+#include "Sound/SoundCue.h"
 #include "SubtitleManager.h"
 
 /*-----------------------------------------------------------------------------
diff --git a/Engine/Source/Runtime/Engine/Private/AudioDecompress.cpp b/Engine/Source/Runtime/Engine/Private/AudioDecompress.cpp
--- a/Engine/Source/Runtime/Engine/Private/AudioDecompress.cpp
+++ b/Engine/Source/Runtime/Engine/Private/AudioDecompress.cpp
@@ -2,9 +2,9 @@
 
 
 #include "EnginePrivate.h"
-
-#include "SoundDefinitions.h"
 #include "AudioDecompress.h"
+#include "AudioDevice.h"
+#include "Sound/SoundWave.h"
 
 #if WITH_OGGVORBIS
 
diff --git a/Engine/Source/Runtime/Engine/Private/AudioDerivedData.cpp b/Engine/Source/Runtime/Engine/Private/AudioDerivedData.cpp
--- a/Engine/Source/Runtime/Engine/Private/AudioDerivedData.cpp
+++ b/Engine/Source/Runtime/Engine/Private/AudioDerivedData.cpp
@@ -1,8 +1,8 @@
 // Copyright 1998-2014 Epic Games, Inc. All Rights Reserved.
 #include "EnginePrivate.h"
-#include "SoundDefinitions.h"
 #include "AudioDerivedData.h"
 #include "TargetPlatform.h"
+#include "Sound/SoundWave.h"
 
 DEFINE_LOG_CATEGORY_STATIC(LogAudioDerivedData, Log, All);
 
@@ -771,6 +771,19 @@ FDerivedAudioDataCompressor::FDerivedAudioDataCompressor(USoundWave* InSoundNode
 	}
 }
 
+FString FDerivedAudioDataCompressor::GetPluginSpecificCacheKeySuffix() const
+{
+	int32 FormatVersion = 0xffff; // if the compressor is NULL, this will be used as the version...and in that case we expect everything to fail anyway
+	if (Compressor)
+	{
+		FormatVersion = (int32)Compressor->GetVersion(Format);
+	}
+	FString FormatString = Format.ToString().ToUpper();
+	check(SoundNode->CompressedDataGuid.IsValid());
+	return FString::Printf(TEXT("%s_%04X_%s"), *FormatString, FormatVersion, *SoundNode->CompressedDataGuid.ToString());
+}
+
+
 bool FDerivedAudioDataCompressor::Build(TArray<uint8>& OutData) 
 {
 #if WITH_EDITORONLY_DATA
diff --git a/Engine/Source/Runtime/Engine/Private/AudioDevice.cpp b/Engine/Source/Runtime/Engine/Private/AudioDevice.cpp
--- a/Engine/Source/Runtime/Engine/Private/AudioDevice.cpp
+++ b/Engine/Source/Runtime/Engine/Private/AudioDevice.cpp
@@ -1,10 +1,14 @@
 // Copyright 1998-2014 Epic Games, Inc. All Rights Reserved.
 
 #include "EnginePrivate.h"
-#include "Sound/SoundNodeWavePlayer.h"
-#include "SoundDefinitions.h"
+#include "Audio.h"
+#include "AudioDevice.h"
 #include "AudioEffect.h"
 #include "AudioDecompress.h"
+#include "Sound/AudioSettings.h"
+#include "Sound/SoundCue.h"
+#include "Sound/SoundNodeWavePlayer.h"
+#include "Sound/SoundWave.h"
 
 /*-----------------------------------------------------------------------------
 	FAudioDevice implementation.
diff --git a/Engine/Source/Runtime/Engine/Private/AudioEffect.cpp b/Engine/Source/Runtime/Engine/Private/AudioEffect.cpp
--- a/Engine/Source/Runtime/Engine/Private/AudioEffect.cpp
+++ b/Engine/Source/Runtime/Engine/Private/AudioEffect.cpp
@@ -5,8 +5,8 @@
 =============================================================================*/
 
 #include "EnginePrivate.h" 
-#include "SoundDefinitions.h"
 #include "AudioEffect.h"
+#include "Sound/ReverbEffect.h"
 
 /** 
  * Default settings for a null reverb effect
diff --git a/Engine/Source/Runtime/Engine/Private/AudioSettings.cpp b/Engine/Source/Runtime/Engine/Private/AudioSettings.cpp
--- a/Engine/Source/Runtime/Engine/Private/AudioSettings.cpp
+++ b/Engine/Source/Runtime/Engine/Private/AudioSettings.cpp
@@ -5,6 +5,7 @@
 =============================================================================*/
 
 #include "EnginePrivate.h"
+#include "Sound/AudioSettings.h"
 
 
 UAudioSettings::UAudioSettings(const class FPostConstructInitializeProperties& PCIP)
diff --git a/Engine/Source/Runtime/Engine/Private/DialogueWave.cpp b/Engine/Source/Runtime/Engine/Private/DialogueWave.cpp
--- a/Engine/Source/Runtime/Engine/Private/DialogueWave.cpp
+++ b/Engine/Source/Runtime/Engine/Private/DialogueWave.cpp
@@ -1,7 +1,9 @@
 // Copyright 1998-2014 Epic Games, Inc. All Rights Reserved.
 
 #include "EnginePrivate.h"
-#include "SoundDefinitions.h"
+#include "ActiveSound.h"
+#include "Sound/DialogueSoundWaveProxy.h"
+#include "Sound/SoundWave.h"
 #include "SubtitleManager.h"
 
 bool operator==(const FDialogueContextMapping& LHS, const FDialogueContextMapping& RHS)
diff --git a/Engine/Source/Runtime/Engine/Private/EnginePrivate.h b/Engine/Source/Runtime/Engine/Private/EnginePrivate.h
--- a/Engine/Source/Runtime/Engine/Private/EnginePrivate.h
+++ b/Engine/Source/Runtime/Engine/Private/EnginePrivate.h
@@ -21,7 +21,6 @@
 #include "Engine/EngineTypes.h"
 #include "Components/ActorComponent.h"
 #include "GameFramework/Actor.h"
-#include "Sound/AmbientSound.h"
 #include "Engine/Brush.h"
 #include "Camera/CameraActor.h"
 #include "AI/Navigation/NavAgentInterface.h"
@@ -108,7 +107,6 @@
 #include "Components/PawnNoiseEmitterComponent.h"
 #include "PhysicsEngine/PhysicsHandleComponent.h"
 #include "Atmosphere/AtmosphericFogComponent.h"
-#include "Sound/SoundAttenuation.h"
 #include "Components/AudioComponent.h"
 #include "Camera/CameraComponent.h"
 #include "Components/ChildActorComponent.h"
@@ -158,7 +156,6 @@
 #include "Components/TimelineComponent.h"
 #include "EditorFramework/AssetImportData.h"
 #include "Engine/AssetUserData.h"
-#include "Sound/AudioSettings.h"
 #include "AI/Navigation/AvoidanceManager.h"
 #include "Engine/BlendableInterface.h"
 #include "Engine/BlueprintCore.h"
@@ -268,7 +265,6 @@
 #include "Engine/ChildConnection.h"
 #include "Engine/Polys.h"
 #include "Engine/RendererSettings.h"
-#include "Sound/ReverbEffect.h"
 #include "MovieScene/RuntimeMovieScenePlayerInterface.h"
 #include "GameFramework/SaveGame.h"
 #include "Engine/SaveGameSummary.h"
@@ -278,14 +274,6 @@
 #include "Engine/SkeletalMesh.h"
 #include "Engine/DestructibleMesh.h"
 #include "Engine/SkeletalMeshSocket.h"
-#include "Sound/SoundBase.h"
-#include "Sound/DialogueSoundWaveProxy.h"
-#include "Sound/SoundClass.h"
-#include "Sound/SoundCue.h"
-#include "Sound/SoundGroups.h"
-#include "Sound/SoundWave.h"
-#include "Sound/SoundWaveStreaming.h"
-#include "Sound/SoundMix.h"
 #include "Engine/StaticMesh.h"
 #include "Engine/StaticMeshSocket.h"
 #include "Camera/CameraStackTypes.h"
diff --git a/Engine/Source/Runtime/Engine/Private/GameViewportClient.cpp b/Engine/Source/Runtime/Engine/Private/GameViewportClient.cpp
--- a/Engine/Source/Runtime/Engine/Private/GameViewportClient.cpp
+++ b/Engine/Source/Runtime/Engine/Private/GameViewportClient.cpp
@@ -4,7 +4,6 @@
 #include "Engine/Console.h"
 #include "GameFramework/HUD.h"
 #include "ParticleDefinitions.h"
-#include "SoundDefinitions.h"
 #include "FXSystem.h"
 #include "SubtitleManager.h"
 #include "ImageUtils.h"
@@ -16,6 +15,8 @@
 #include "SVirtualJoystick.h"
 #include "SceneViewport.h"
 #include "EngineModule.h"
+#include "AudioDevice.h"
+#include "Sound/SoundWave.h"
 
 #include "HighResScreenshot.h"
 
diff --git a/Engine/Source/Runtime/Engine/Private/GameplayStatics.cpp b/Engine/Source/Runtime/Engine/Private/GameplayStatics.cpp
--- a/Engine/Source/Runtime/Engine/Private/GameplayStatics.cpp
+++ b/Engine/Source/Runtime/Engine/Private/GameplayStatics.cpp
@@ -9,6 +9,8 @@
 #include "Slate.h"
 #include "Particles/ParticleSystemComponent.h"
 #include "GameFramework/Character.h"
+#include "Sound/SoundBase.h"
+#include "Sound/SoundCue.h"
 
 //////////////////////////////////////////////////////////////////////////
 // UGameplayStatics
@@ -753,7 +755,7 @@ class UAudioComponent* UGameplayStatics::PlayDialogueAttached(class UDialogueWav
 	return AudioComponent;
 }
 
-void UGameplayStatics::PlaySound(UObject* WorldContextObject, class USoundCue* InSoundCue, class USceneComponent* AttachComponent, FName AttachName, bool bFollow, float VolumeMultiplier, float PitchMultiplier)
+void UGameplayStatics::PlaySound(UObject* WorldContextObject, USoundCue* InSoundCue, USceneComponent* AttachComponent, FName AttachName, bool bFollow, float VolumeMultiplier, float PitchMultiplier)
 {
 	if (bFollow)
 	{
diff --git a/Engine/Source/Runtime/Engine/Private/Interpolation.cpp b/Engine/Source/Runtime/Engine/Private/Interpolation.cpp
--- a/Engine/Source/Runtime/Engine/Private/Interpolation.cpp
+++ b/Engine/Source/Runtime/Engine/Private/Interpolation.cpp
@@ -61,7 +61,8 @@
 #include "Materials/MaterialInstanceActor.h"
 #include "Animation/SkeletalMeshActor.h"
 #include "ParticleDefinitions.h"
-#include "SoundDefinitions.h"
+#include "AudioDevice.h"
+#include "Sound/SoundBase.h"
 #include "InterpolationHitProxy.h"
 #include "AVIWriter.h"
 #include "AnimationUtils.h"
diff --git a/Engine/Source/Runtime/Engine/Private/InterpolationDraw.cpp b/Engine/Source/Runtime/Engine/Private/InterpolationDraw.cpp
--- a/Engine/Source/Runtime/Engine/Private/InterpolationDraw.cpp
+++ b/Engine/Source/Runtime/Engine/Private/InterpolationDraw.cpp
@@ -28,7 +28,7 @@
 #include "Matinee/InterpTrackAudioMaster.h"
 #include "Matinee/InterpTrackVisibility.h"
 #include "Matinee/InterpTrackParticleReplay.h"
-#include "SoundDefinitions.h"
+#include "Sound/SoundBase.h"
 #include "InterpolationHitProxy.h"
 #include "AnimationUtils.h"
 #include "InterpolationHitProxy.h"
diff --git a/Engine/Source/Runtime/Engine/Private/ReverbEffect.cpp b/Engine/Source/Runtime/Engine/Private/ReverbEffect.cpp
--- a/Engine/Source/Runtime/Engine/Private/ReverbEffect.cpp
+++ b/Engine/Source/Runtime/Engine/Private/ReverbEffect.cpp
@@ -1,7 +1,7 @@
 // Copyright 1998-2014 Epic Games, Inc. All Rights Reserved.
 
 #include "EnginePrivate.h"
-#include "SoundDefinitions.h"
+#include "Sound/ReverbEffect.h"
 
 UReverbEffect::UReverbEffect(const class FPostConstructInitializeProperties& PCIP)
 	: Super(PCIP)
diff --git a/Engine/Source/Runtime/Engine/Private/ReverbVolume.cpp b/Engine/Source/Runtime/Engine/Private/ReverbVolume.cpp
--- a/Engine/Source/Runtime/Engine/Private/ReverbVolume.cpp
+++ b/Engine/Source/Runtime/Engine/Private/ReverbVolume.cpp
@@ -5,8 +5,9 @@
 =============================================================================*/
 
 #include "EnginePrivate.h"
+#include "AudioDevice.h"
 #include "Net/UnrealNetwork.h"
-#include "SoundDefinitions.h"
+#include "Sound/ReverbEffect.h"
 
 void FReverbSettings::PostSerialize(const FArchive& Ar)
 {
diff --git a/Engine/Source/Runtime/Engine/Private/Slate/SlateSoundDevice.cpp b/Engine/Source/Runtime/Engine/Private/Slate/SlateSoundDevice.cpp
--- a/Engine/Source/Runtime/Engine/Private/Slate/SlateSoundDevice.cpp
+++ b/Engine/Source/Runtime/Engine/Private/Slate/SlateSoundDevice.cpp
@@ -1,8 +1,10 @@
 // Copyright 1998-2014 Epic Games, Inc. All Rights Reserved.
 
 #include "EnginePrivate.h"
-#include "SlateSoundDevice.h"
-#include "SoundDefinitions.h"
+#include "ActiveSound.h"
+#include "AudioDevice.h"
+#include "Slate/SlateSoundDevice.h"
+#include "Sound/SoundBase.h"
 
 DEFINE_LOG_CATEGORY_STATIC(LogSlateSoundDevice, Log, All);
 
diff --git a/Engine/Source/Runtime/Engine/Private/SoundBase.cpp b/Engine/Source/Runtime/Engine/Private/SoundBase.cpp
--- a/Engine/Source/Runtime/Engine/Private/SoundBase.cpp
+++ b/Engine/Source/Runtime/Engine/Private/SoundBase.cpp
@@ -1,7 +1,9 @@
 // Copyright 1998-2014 Epic Games, Inc. All Rights Reserved.
 
 #include "EnginePrivate.h"
-#include "SoundDefinitions.h"
+#include "Sound/SoundBase.h"
+#include "Sound/AudioSettings.h"
+#include "AudioDevice.h"
 
 USoundBase::USoundBase(const class FPostConstructInitializeProperties& PCIP)
 	: Super(PCIP)
diff --git a/Engine/Source/Runtime/Engine/Private/SoundCue.cpp b/Engine/Source/Runtime/Engine/Private/SoundCue.cpp
--- a/Engine/Source/Runtime/Engine/Private/SoundCue.cpp
+++ b/Engine/Source/Runtime/Engine/Private/SoundCue.cpp
@@ -1,10 +1,11 @@
 // Copyright 1998-2014 Epic Games, Inc. All Rights Reserved.
 
 #include "EnginePrivate.h"
-#include "SoundDefinitions.h"
+#include "Sound/SoundCue.h"
 #include "Sound/SoundNodeMixer.h"
 #include "Sound/SoundNodeWavePlayer.h"
 #include "Sound/SoundNodeAttenuation.h"
+#include "Sound/SoundWave.h"
 #if WITH_EDITORONLY_DATA
 #include "Kismet2/BlueprintEditorUtils.h"
 #include "UnrealEd.h"
diff --git a/Engine/Source/Runtime/Engine/Private/SoundMix.cpp b/Engine/Source/Runtime/Engine/Private/SoundMix.cpp
--- a/Engine/Source/Runtime/Engine/Private/SoundMix.cpp
+++ b/Engine/Source/Runtime/Engine/Private/SoundMix.cpp
@@ -2,7 +2,8 @@
 
 
 #include "EnginePrivate.h"
-#include "SoundDefinitions.h"
+#include "Sound/SoundMix.h"
+#include "AudioDevice.h"
 #if WITH_EDITOR
 #include "Slate.h"
 #endif
diff --git a/Engine/Source/Runtime/Engine/Private/SoundNode.cpp b/Engine/Source/Runtime/Engine/Private/SoundNode.cpp
--- a/Engine/Source/Runtime/Engine/Private/SoundNode.cpp
+++ b/Engine/Source/Runtime/Engine/Private/SoundNode.cpp
@@ -2,8 +2,8 @@
 
 
 #include "EnginePrivate.h"
-#include "SoundDefinitions.h"
 #include "Sound/SoundNode.h"
+#include "Sound/SoundCue.h"
 
 #if WITH_EDITORONLY_DATA
 #include "UnrealEd.h"
@@ -38,6 +38,18 @@ void USoundNode::AddReferencedObjects(UObject* InThis, FReferenceCollector& Coll
 }
 #endif //WITH_EDITORONLY_DATA
 
+UPTRINT USoundNode::GetNodeWaveInstanceHash(const UPTRINT ParentWaveInstanceHash, const USoundNode* ChildNode, const uint32 ChildIndex)
+{
+	checkf(ChildIndex < MAX_ALLOWED_CHILD_NODES, TEXT("Too many children (%d) in SoundCue '%s'"), ChildIndex, *CastChecked<USoundCue>(ChildNode->GetOuter())->GetFullName());
+	return ((ParentWaveInstanceHash << ChildIndex) ^ (UPTRINT)ChildNode);
+}
+
+UPTRINT USoundNode::GetNodeWaveInstanceHash(const UPTRINT ParentWaveInstanceHash, const UPTRINT ChildNodeHash, const uint32 ChildIndex)
+{
+	checkf(ChildIndex < MAX_ALLOWED_CHILD_NODES, TEXT("Too many children (%d) in SoundCue"), ChildIndex);
+	return ((ParentWaveInstanceHash << ChildIndex) ^ ChildNodeHash);
+}
+
 void USoundNode::ParseNodes( FAudioDevice* AudioDevice, const UPTRINT NodeWaveInstanceHash, FActiveSound& ActiveSound, const FSoundParseParameters& ParseParams, TArray<FWaveInstance*>& WaveInstances )
 {
 	for( int32 i = 0; i < ChildNodes.Num() && i < GetMaxChildNodes(); ++i )
diff --git a/Engine/Source/Runtime/Engine/Private/SoundNodeDialoguePlayer.cpp b/Engine/Source/Runtime/Engine/Private/SoundNodeDialoguePlayer.cpp
--- a/Engine/Source/Runtime/Engine/Private/SoundNodeDialoguePlayer.cpp
+++ b/Engine/Source/Runtime/Engine/Private/SoundNodeDialoguePlayer.cpp
@@ -1,7 +1,7 @@
 // Copyright 1998-2014 Epic Games, Inc. All Rights Reserved.
 
 #include "EnginePrivate.h"
-#include "SoundDefinitions.h"
+#include "Sound/SoundBase.h"
 #include "Sound/SoundNodeDialoguePlayer.h"
 
 USoundNodeDialoguePlayer::USoundNodeDialoguePlayer(const class FPostConstructInitializeProperties& PCIP)
diff --git a/Engine/Source/Runtime/Engine/Private/SoundNodeMature.cpp b/Engine/Source/Runtime/Engine/Private/SoundNodeMature.cpp
--- a/Engine/Source/Runtime/Engine/Private/SoundNodeMature.cpp
+++ b/Engine/Source/Runtime/Engine/Private/SoundNodeMature.cpp
@@ -2,9 +2,9 @@
 
 
 #include "EnginePrivate.h"
-#include "SoundDefinitions.h"
 #include "Sound/SoundNodeMature.h"
 #include "Sound/SoundNodeWavePlayer.h"
+#include "Sound/SoundWave.h"
 
 /*-----------------------------------------------------------------------------
        USoundNodeMature implementation.
diff --git a/Engine/Source/Runtime/Engine/Private/SoundNodeWaveParam.cpp b/Engine/Source/Runtime/Engine/Private/SoundNodeWaveParam.cpp
--- a/Engine/Source/Runtime/Engine/Private/SoundNodeWaveParam.cpp
+++ b/Engine/Source/Runtime/Engine/Private/SoundNodeWaveParam.cpp
@@ -2,8 +2,8 @@
 
 
 #include "EnginePrivate.h"
-#include "SoundDefinitions.h"
 #include "Sound/SoundNodeWaveParam.h"
+#include "Sound/SoundWave.h"
 
 /*-----------------------------------------------------------------------------
 	USoundNodeWaveParam implementation
diff --git a/Engine/Source/Runtime/Engine/Private/SoundNodeWavePlayer.cpp b/Engine/Source/Runtime/Engine/Private/SoundNodeWavePlayer.cpp
--- a/Engine/Source/Runtime/Engine/Private/SoundNodeWavePlayer.cpp
+++ b/Engine/Source/Runtime/Engine/Private/SoundNodeWavePlayer.cpp
@@ -1,8 +1,8 @@
 // Copyright 1998-2014 Epic Games, Inc. All Rights Reserved.
 
 #include "EnginePrivate.h"
-#include "SoundDefinitions.h"
 #include "Sound/SoundNodeWavePlayer.h"
+#include "Sound/SoundWave.h"
 
 #define LOCTEXT_NAMESPACE "SoundNodeWavePlayer"
 
diff --git a/Engine/Source/Runtime/Engine/Private/SoundWave.cpp b/Engine/Source/Runtime/Engine/Private/SoundWave.cpp
--- a/Engine/Source/Runtime/Engine/Private/SoundWave.cpp
+++ b/Engine/Source/Runtime/Engine/Private/SoundWave.cpp
@@ -1,7 +1,10 @@
 // Copyright 1998-2014 Epic Games, Inc. All Rights Reserved.
 
 #include "EnginePrivate.h"
-#include "SoundDefinitions.h"
+#include "Sound/SoundWave.h"
+#include "ActiveSound.h"
+#include "Audio.h"
+#include "AudioDevice.h"
 #include "AudioDecompress.h"
 #include "TargetPlatform.h"
 #include "AudioDerivedData.h"
diff --git a/Engine/Source/Runtime/Engine/Private/SoundWaveStreaming.cpp b/Engine/Source/Runtime/Engine/Private/SoundWaveStreaming.cpp
--- a/Engine/Source/Runtime/Engine/Private/SoundWaveStreaming.cpp
+++ b/Engine/Source/Runtime/Engine/Private/SoundWaveStreaming.cpp
@@ -1,7 +1,7 @@
 // Copyright 1998-2014 Epic Games, Inc. All Rights Reserved.
 
 #include "EnginePrivate.h"
-#include "SoundDefinitions.h"
+#include "Sound/SoundWaveStreaming.h"
 
 USoundWaveStreaming::USoundWaveStreaming(const class FPostConstructInitializeProperties& PCIP)
 	: Super(PCIP)
diff --git a/Engine/Source/Runtime/Engine/Private/UnrealEngine.cpp b/Engine/Source/Runtime/Engine/Private/UnrealEngine.cpp
--- a/Engine/Source/Runtime/Engine/Private/UnrealEngine.cpp
+++ b/Engine/Source/Runtime/Engine/Private/UnrealEngine.cpp
@@ -57,6 +57,9 @@
 #include "Particles/ParticleSystem.h"
 #include "Particles/ParticleSystemComponent.h"
 
+#include "Sound/ReverbEffect.h"
+#include "Sound/SoundWave.h"
+
 // @todo this is here only due to circular dependency to AIModule. To be removed
 #include "BehaviorTree/BehaviorTreeManager.h"
 #include "EnvironmentQuery/EnvQueryManager.h"
diff --git a/Engine/Source/Runtime/Engine/Public/ActiveSound.h b/Engine/Source/Runtime/Engine/Public/ActiveSound.h
--- a/Engine/Source/Runtime/Engine/Public/ActiveSound.h
+++ b/Engine/Source/Runtime/Engine/Public/ActiveSound.h
@@ -3,6 +3,7 @@
 #pragma once
 
 #include "Sound/ReverbVolume.h"
+#include "Sound/SoundAttenuation.h"
 
 struct ENGINE_API FActiveSound
 {
diff --git a/Engine/Source/Runtime/Engine/Public/AudioEffect.h b/Engine/Source/Runtime/Engine/Public/AudioEffect.h
--- a/Engine/Source/Runtime/Engine/Public/AudioEffect.h
+++ b/Engine/Source/Runtime/Engine/Public/AudioEffect.h
@@ -6,6 +6,8 @@
 
 #pragma once
 
+#include "Sound/SoundMix.h"
+
 class FAudioReverbEffect
 {
 public:
diff --git a/Engine/Source/Runtime/Engine/Public/SoundDefinitions.h b/Engine/Source/Runtime/Engine/Public/SoundDefinitions.h
--- a/Engine/Source/Runtime/Engine/Public/SoundDefinitions.h
+++ b/Engine/Source/Runtime/Engine/Public/SoundDefinitions.h
@@ -2,8 +2,5 @@
 
 #pragma once
 #include "Audio.h"							// Audio code.
-#include "../../../../Source/Runtime/Engine/Classes/Sound/SoundClass.h"
-#include "../../../../Source/Runtime/Engine/Classes/Sound/SoundAttenuation.h"
-#include "../../../../Source/Runtime/Engine/Classes/Sound/SoundGroups.h"
 #include "ActiveSound.h"
 #include "AudioDevice.h"
