Change assert to warning, and recover

diff --git a/Engine/Source/Runtime/Engine/Private/PackageMapClient.cpp b/Engine/Source/Runtime/Engine/Private/PackageMapClient.cpp
--- a/Engine/Source/Runtime/Engine/Private/PackageMapClient.cpp
+++ b/Engine/Source/Runtime/Engine/Private/PackageMapClient.cpp
@@ -320,7 +320,13 @@ bool UPackageMapClient::SerializeNewActor(FArchive& Ar, class UActorChannel *Cha
 				FNetworkGUID SubobjNetGUID = GuidCache->NetGUIDLookup.FindRef( SubObj );
 
 				// Make sure these sub objects have net guids that are relative to this owning actor
-				check( SubobjNetGUID.Value == COMPOSE_RELATIVE_NET_GUID( NetGUID, i + 1 ) );
+				if ( SubobjNetGUID.Value != COMPOSE_RELATIVE_NET_GUID( NetGUID, i + 1 ) )
+				{
+					// If this happens, someone most likely added a stably name sub-object during gameplay (bad)
+					UE_LOG( LogNetPackageMap, Error, TEXT( "SerializeNewActor: Sub-object net guid out of sync. Actor: %s, SubObj: %s" ), *Actor->GetName(), *SubObj->GetName() );
+					Checksum = 0;		 // Force the checksum on the client to fail, so the client won't have any guids for these sub-objects
+					break;
+				}
 
 				// Evolve checksum so we can sanity check
 				Checksum = FCrc::MemCrc32( &SubobjNetGUID.Value, sizeof( SubobjNetGUID.Value ), Checksum );
