Fix bad merge

diff --git a/Engine/Source/Runtime/Engine/Public/AudioDerivedData.h b/Engine/Source/Runtime/Engine/Public/AudioDerivedData.h
--- a/Engine/Source/Runtime/Engine/Public/AudioDerivedData.h
+++ b/Engine/Source/Runtime/Engine/Public/AudioDerivedData.h
@@ -32,17 +32,7 @@ public:
 		return TEXT("0002_0000");
 	}
 
-	virtual FString GetPluginSpecificCacheKeySuffix() const override
-	{
-		int32 FormatVersion = 0xffff; // if the compressor is NULL, this will be used as the version...and in that case we expect everything to fail anyway
-		if (Compressor)
-		{
-			FormatVersion = (int32)Compressor->GetVersion(Format);
-		}
-		FString FormatString = Format.ToString().ToUpper();
-		check(SoundNode->CompressedDataGuid.IsValid());
-		return FString::Printf( TEXT("%s_%04X_%s"), *FormatString, FormatVersion, *SoundNode->CompressedDataGuid.ToString());
-	}
+	virtual FString GetPluginSpecificCacheKeySuffix() const override;
 	
 	virtual bool IsBuildThreadsafe() const override
 	{
