TTP 334627; added PSA_AwayFromCenter to align particles with the emitter center; only works when simulating in local space

diff --git a/Engine/Source/Runtime/Engine/Classes/Particles/ParticleSpriteEmitter.h b/Engine/Source/Runtime/Engine/Classes/Particles/ParticleSpriteEmitter.h
--- a/Engine/Source/Runtime/Engine/Classes/Particles/ParticleSpriteEmitter.h
+++ b/Engine/Source/Runtime/Engine/Classes/Particles/ParticleSpriteEmitter.h
@@ -12,6 +12,7 @@ enum EParticleScreenAlignment
 	PSA_Square,
 	PSA_Rectangle,
 	PSA_Velocity,
+	PSA_AwayFromCenter,
 	PSA_TypeSpecific,
 	PSA_MAX,
 };
diff --git a/Engine/Source/Runtime/Engine/Private/Particles/ParticleEmitterInstances.cpp b/Engine/Source/Runtime/Engine/Private/Particles/ParticleEmitterInstances.cpp
--- a/Engine/Source/Runtime/Engine/Private/Particles/ParticleEmitterInstances.cpp
+++ b/Engine/Source/Runtime/Engine/Private/Particles/ParticleEmitterInstances.cpp
@@ -2856,37 +2856,46 @@ void FParticleMeshEmitterInstance::Tick(float DeltaTime, bool bSuppressSpawning)
 			DECLARE_PARTICLE(Particle, ParticleData + ParticleStride * ParticleIndices[i]);
 			FMeshRotationPayloadData* PayloadData	= (FMeshRotationPayloadData*)((uint8*)&Particle + MeshRotationOffset);
 			PayloadData->RotationRate				= PayloadData->RotationRateBase;
-			if (LODLevel->RequiredModule->ScreenAlignment == PSA_Velocity)
+			if (LODLevel->RequiredModule->ScreenAlignment == PSA_Velocity
+				|| LODLevel->RequiredModule->ScreenAlignment == PSA_AwayFromCenter)
 			{
 				// Determine the rotation to the velocity vector and apply it to the mesh
 				FVector	NewDirection	= Particle.Velocity;
 				
-				//check if an orbit module should affect the velocity...		
-				if (LODLevel->RequiredModule->bOrbitModuleAffectsVelocityAlignment &&
-					LODLevel->OrbitModules.Num() > 0)
+				if (LODLevel->RequiredModule->ScreenAlignment == PSA_Velocity)
 				{
-					UParticleModuleOrbit* LastOrbit = SpriteTemplate->LODLevels[0]->OrbitModules[LODLevel->OrbitModules.Num() - 1];
-					check(LastOrbit);
-					
-					uint32 OrbitModuleOffset = *ModuleOffsetMap.Find(LastOrbit);
-					if (OrbitModuleOffset != 0)
+					//check if an orbit module should affect the velocity...		
+					if (LODLevel->RequiredModule->bOrbitModuleAffectsVelocityAlignment &&
+						LODLevel->OrbitModules.Num() > 0)
 					{
-						FOrbitChainModuleInstancePayload &OrbitPayload = *(FOrbitChainModuleInstancePayload*)((uint8*)&Particle + OrbitModuleOffset);
-
-						FVector OrbitOffset = OrbitPayload.Offset;
-						FVector PrevOrbitOffset = OrbitPayload.PreviousOffset;
-						FVector Location = Particle.Location;
-						FVector OldLocation = Particle.OldLocation;
-
-						//this should be our current position
-						FVector NewPos = Location + OrbitOffset;	
-						//this should be our previous position
-						FVector OldPos = OldLocation + PrevOrbitOffset;
-
-						NewDirection = NewPos - OldPos;
-					}	
-				}			              
-               
+						UParticleModuleOrbit* LastOrbit = SpriteTemplate->LODLevels[0]->OrbitModules[LODLevel->OrbitModules.Num() - 1];
+						check(LastOrbit);
+					
+						uint32 OrbitModuleOffset = *ModuleOffsetMap.Find(LastOrbit);
+						if (OrbitModuleOffset != 0)
+						{
+							FOrbitChainModuleInstancePayload &OrbitPayload = *(FOrbitChainModuleInstancePayload*)((uint8*)&Particle + OrbitModuleOffset);
+
+							FVector OrbitOffset = OrbitPayload.Offset;
+							FVector PrevOrbitOffset = OrbitPayload.PreviousOffset;
+							FVector Location = Particle.Location;
+							FVector OldLocation = Particle.OldLocation;
+
+							//this should be our current position
+							FVector NewPos = Location + OrbitOffset;	
+							//this should be our previous position
+							FVector OldPos = OldLocation + PrevOrbitOffset;
+
+							NewDirection = NewPos - OldPos;
+						}	
+					}			              
+				}
+				else if (LODLevel->RequiredModule->ScreenAlignment == PSA_AwayFromCenter)
+				{
+
+					NewDirection = Particle.Location;
+				}
+
 				NewDirection.Normalize();
 				FVector	OldDirection(1.0f, 0.0f, 0.0f);
 
