Fixed enum forward declare for clang builds, related to CL2104492

diff --git a/Engine/Source/Runtime/Engine/Classes/Animation/AnimCompress.h b/Engine/Source/Runtime/Engine/Classes/Animation/AnimCompress.h
--- a/Engine/Source/Runtime/Engine/Classes/Animation/AnimCompress.h
+++ b/Engine/Source/Runtime/Engine/Classes/Animation/AnimCompress.h
@@ -7,6 +7,7 @@
  *
  */
 
+#include "AnimSequence.h"
 #include "AnimCompress.generated.h"
 
 UCLASS(abstract, DependsOn=UAnimSequence, hidecategories=Object, MinimalAPI)
@@ -24,15 +25,15 @@ class UAnimCompress : public UObject
 
 	/** Format for bitwise compression of translation data. */
 	UPROPERTY()
-	TEnumAsByte<enum AnimationCompressionFormat> TranslationCompressionFormat;
+	TEnumAsByte<AnimationCompressionFormat> TranslationCompressionFormat;
 
 	/** Format for bitwise compression of rotation data. */
 	UPROPERTY()
-	TEnumAsByte<enum AnimationCompressionFormat> RotationCompressionFormat;
+	TEnumAsByte<AnimationCompressionFormat> RotationCompressionFormat;
 
 	/** Format for bitwise compression of scale data. */
 	UPROPERTY()
-	TEnumAsByte<enum AnimationCompressionFormat> ScaleCompressionFormat;
+	TEnumAsByte<AnimationCompressionFormat> ScaleCompressionFormat;
 
 
 public:
