#ue4 Fix for crash when using GetCharacterArrayFronString

diff --git a/Engine/Source/Runtime/Engine/Private/KismetStringLibrary.cpp b/Engine/Source/Runtime/Engine/Private/KismetStringLibrary.cpp
--- a/Engine/Source/Runtime/Engine/Private/KismetStringLibrary.cpp
+++ b/Engine/Source/Runtime/Engine/Private/KismetStringLibrary.cpp
@@ -288,14 +288,17 @@ TArray<FString> UKismetStringLibrary::GetCharacterArrayFromString(const FString&
 {
 	TArray<FString> SeparatedChars;
 
-	for (auto CharIt(SourceString.CreateConstIterator()); CharIt; ++CharIt)
+	if (!SourceString.IsEmpty())
 	{
-		TCHAR Char = *CharIt;
-		SeparatedChars.Add(FString(1, &Char));
+		for (auto CharIt(SourceString.CreateConstIterator()); CharIt; ++CharIt)
+		{
+			TCHAR Char = *CharIt;
+			SeparatedChars.Add(FString(1, &Char));
+		}
+
+		// Remove the null terminator on the end
+		SeparatedChars.RemoveAt(SeparatedChars.Num() - 1, 1);
 	}
 
-	// Remove the null terminator on the end
-	SeparatedChars.RemoveAt(SeparatedChars.Num() - 1, 1);
-
 	return SeparatedChars;
 }
\ No newline at end of file
