Fixed FNavMeshPath::GetPathCorridorLength returning 0 for valid single-poly paths #UE4

diff --git a/Engine/Source/Runtime/Engine/Private/AI/Navigation/NavigationPath.cpp b/Engine/Source/Runtime/Engine/Private/AI/Navigation/NavigationPath.cpp
--- a/Engine/Source/Runtime/Engine/Private/AI/Navigation/NavigationPath.cpp
+++ b/Engine/Source/Runtime/Engine/Private/AI/Navigation/NavigationPath.cpp
@@ -228,10 +228,15 @@ float FNavMeshPath::GetStringPulledLength(const int32 StartingPoint) const
 
 float FNavMeshPath::GetPathCorridorLength(const int32 StartingEdge) const
 {
-	if (bCorridorEdgesGenerated == false || StartingEdge >= PathCorridorEdges.Num())
+	if (bCorridorEdgesGenerated == false)
 	{
 		return 0.f;
 	}
+	else if (StartingEdge >= PathCorridorEdges.Num())
+	{
+		return StartingEdge == 0 && PathPoints.Num() > 1 ? FVector::Dist(PathPoints[0].Location, PathPoints[PathPoints.Num()-1].Location) : 0.f;
+	}
+
 	
 	const FNavigationPortalEdge* PrevEdge = PathCorridorEdges.GetTypedData() + StartingEdge;
 	const FNavigationPortalEdge* CorridorEdge = PrevEdge + 1;
