Fix WITH_EDITOR 0 compile error

diff --git a/Engine/Source/Runtime/Engine/Private/EdGraph/EdGraphPin.cpp b/Engine/Source/Runtime/Engine/Private/EdGraph/EdGraphPin.cpp
--- a/Engine/Source/Runtime/Engine/Private/EdGraph/EdGraphPin.cpp
+++ b/Engine/Source/Runtime/Engine/Private/EdGraph/EdGraphPin.cpp
@@ -172,7 +172,7 @@ void UEdGraphPin::CopyPersistentDataFromOldPin(const UEdGraphPin& SourcePin)
 	// If the source pin is split, then split the new one
 	if (SourcePin.SubPins.Num() > 0)
 	{
-		GetOwningNode()->GetSchema()->SplitPin(this);
+		GetSchema()->SplitPin(this);
 	}
 
 #if WITH_EDITORONLY_DATA
