Expose box sweep to blueprint
TTP#336627

diff --git a/Engine/Source/Runtime/Engine/Classes/Kismet/KismetSystemLibrary.h b/Engine/Source/Runtime/Engine/Classes/Kismet/KismetSystemLibrary.h
--- a/Engine/Source/Runtime/Engine/Classes/Kismet/KismetSystemLibrary.h
+++ b/Engine/Source/Runtime/Engine/Classes/Kismet/KismetSystemLibrary.h
@@ -574,6 +574,39 @@ class UKismetSystemLibrary : public UBlueprintFunctionLibrary
 	 */
 	UFUNCTION(BlueprintCallable, Category="Collision", meta=(bIgnoreSelf="true", HidePin="WorldContextObject", DefaultToSelf="WorldContextObject", AutoCreateRefTerm="ActorsToIgnore", FriendlyName = "MultiSphereTraceByChannel"))
 	static ENGINE_API bool SphereTraceMulti_NEW(UObject* WorldContextObject, const FVector& Start, const FVector& End, float Radius, ETraceTypeQuery TraceChannel, bool bTraceComplex, const TArray<AActor*>& ActorsToIgnore, EDrawDebugTrace::Type DrawDebugType, TArray<FHitResult>& OutHits, bool bIgnoreSelf);
+
+	/**
+	* Sweeps a box along the given line and returns the first blocking hit encountered.
+	* This trace finds the objects that RESPONDS to the given TraceChannel
+	*
+	* @param Start			Start of line segment.
+	* @param End			End of line segment.
+	* @param HalfSize	    Distance from the center of box along each axis
+	* @param Orientation	Orientation of the box
+	* @param TraceChannel
+	* @param bTraceComplex	True to test against complex collision, false to test against simplified collision.
+	* @param OutHit			Properties of the trace hit.
+	* @return				True if there was a hit, false otherwise.
+	*/
+	UFUNCTION(BlueprintCallable, Category = "Collision", meta = (bIgnoreSelf = "true", HidePin = "WorldContextObject", DefaultToSelf = "WorldContextObject", AutoCreateRefTerm = "ActorsToIgnore", FriendlyName = "SingleBoxTraceByChannel"))
+	static ENGINE_API bool BoxTraceSingle(UObject* WorldContextObject, const FVector& Start, const FVector& End, const FVector & HalfSize, const FRotator & Orientation, ETraceTypeQuery TraceChannel, bool bTraceComplex, const TArray<AActor*>& ActorsToIgnore, EDrawDebugTrace::Type DrawDebugType, FHitResult& OutHit, bool bIgnoreSelf);
+
+	/**
+	* Sweeps a box along the given line and returns all hits encountered.
+	* This trace finds the objects that RESPONDS to the given TraceChannel
+	*
+	* @param Start			Start of line segment.
+	* @param End			End of line segment.
+	* @param HalfSize	    Distance from the center of box along each axis
+	* @param Orientation	Orientation of the box
+	* @param TraceChannel
+	* @param bTraceComplex	True to test against complex collision, false to test against simplified collision.
+	* @param OutHits		A list of hits, sorted along the trace from start to finish. The blocking hit will be the last hit, if there was one.
+	* @return				True if there was a hit, false otherwise.
+	*/
+	UFUNCTION(BlueprintCallable, Category = "Collision", meta = (bIgnoreSelf = "true", HidePin = "WorldContextObject", DefaultToSelf = "WorldContextObject", AutoCreateRefTerm = "ActorsToIgnore", FriendlyName = "MultiBoxTraceByChannel"))
+	static ENGINE_API bool BoxTraceMulti(UObject* WorldContextObject, const FVector& Start, const FVector& End, const FVector & HalfSize, const FRotator & Orientation, ETraceTypeQuery TraceChannel, bool bTraceComplex, const TArray<AActor*>& ActorsToIgnore, EDrawDebugTrace::Type DrawDebugType, TArray<FHitResult>& OutHits, bool bIgnoreSelf);
+
 	
 	/**
 	 * Sweeps a capsule along the given line and returns the first blocking hit encountered.
@@ -670,6 +703,40 @@ class UKismetSystemLibrary : public UBlueprintFunctionLibrary
 	UFUNCTION(BlueprintCallable, Category="Collision", meta=(bIgnoreSelf="true", HidePin="WorldContextObject", DefaultToSelf="WorldContextObject", AutoCreateRefTerm="ActorsToIgnore", FriendlyName = "MultiSphereTraceForObjects"))
 	static ENGINE_API bool SphereTraceMultiForObjects(UObject* WorldContextObject, const FVector& Start, const FVector& End, float Radius, const TArray<TEnumAsByte<EObjectTypeQuery> > & ObjectTypes, bool bTraceComplex, const TArray<AActor*>& ActorsToIgnore, EDrawDebugTrace::Type DrawDebugType, TArray<FHitResult>& OutHits, bool bIgnoreSelf);
 	
+
+	/**
+	* Sweeps a box along the given line and returns the first hit encountered.
+	* This only finds objects that are of a type specified by ObjectTypes.
+	*
+	* @param Start			Start of line segment.
+	* @param End			End of line segment.
+	* @param Orientation	
+	* @param HalfSize		Radius of the sphere to sweep
+	* @param ObjectTypes	Array of Object Types to trace
+	* @param bTraceComplex	True to test against complex collision, false to test against simplified collision.
+	* @param OutHit			Properties of the trace hit.
+	* @return				True if there was a hit, false otherwise.
+	*/
+	UFUNCTION(BlueprintCallable, Category = "Collision", meta = (bIgnoreSelf = "true", HidePin = "WorldContextObject", DefaultToSelf = "WorldContextObject", AutoCreateRefTerm = "ActorsToIgnore", FriendlyName = "SingleBoxTraceForObjects"))
+	static ENGINE_API bool BoxTraceSingleForObjects(UObject* WorldContextObject, const FVector& Start, const FVector& End, const FVector & HalfSize, const FRotator & Orientation, const TArray<TEnumAsByte<EObjectTypeQuery> > & ObjectTypes, bool bTraceComplex, const TArray<AActor*>& ActorsToIgnore, EDrawDebugTrace::Type DrawDebugType, FHitResult& OutHit, bool bIgnoreSelf);
+
+
+	/**
+	* Sweeps a box along the given line and returns all hits encountered.
+	* This only finds objects that are of a type specified by ObjectTypes.
+	*
+	* @param Start			Start of line segment.
+	* @param End			End of line segment.
+	* @param Orientation
+	* @param HalfSize		Radius of the sphere to sweep
+	* @param ObjectTypes	Array of Object Types to trace
+	* @param bTraceComplex	True to test against complex collision, false to test against simplified collision.
+	* @param OutHits		A list of hits, sorted along the trace from start to finish.  The blocking hit will be the last hit, if there was one.
+	* @return				True if there was a hit, false otherwise.
+	*/
+	UFUNCTION(BlueprintCallable, Category = "Collision", meta = (bIgnoreSelf = "true", HidePin = "WorldContextObject", DefaultToSelf = "WorldContextObject", AutoCreateRefTerm = "ActorsToIgnore", FriendlyName = "MultiBoxTraceForObjects"))
+	static ENGINE_API bool BoxTraceMultiForObjects(UObject* WorldContextObject, const FVector& Start, const FVector& End, const FVector & HalfSize, const FRotator & Orientation, const TArray<TEnumAsByte<EObjectTypeQuery> > & ObjectTypes, bool bTraceComplex, const TArray<AActor*>& ActorsToIgnore, EDrawDebugTrace::Type DrawDebugType, TArray<FHitResult>& OutHits, bool bIgnoreSelf);
+
 	/**
 	 * Sweeps a capsule along the given line and returns the first hit encountered.
 	 * This only finds objects that are of a type specified by ObjectTypes.
diff --git a/Engine/Source/Runtime/Engine/Private/KismetSystemLibrary.cpp b/Engine/Source/Runtime/Engine/Private/KismetSystemLibrary.cpp
--- a/Engine/Source/Runtime/Engine/Private/KismetSystemLibrary.cpp
+++ b/Engine/Source/Runtime/Engine/Private/KismetSystemLibrary.cpp
@@ -1095,6 +1095,151 @@ void DrawDebugSweptSphere(const UWorld* InWorld, FVector const& Start, FVector c
 	::DrawDebugCapsule(InWorld, Center, HalfHeight, Radius, CapsuleRot, Color, bPersistentLines, LifeTime, DepthPriority );
 }
 
+void DrawDebugSweptBox(const UWorld* InWorld, FVector const& Start, FVector const& End, FRotator const & Orientation, FVector const & HalfSize, FColor const& Color, bool bPersistentLines = false, float LifeTime = -1.f, uint8 DepthPriority = 0)
+{
+	FVector const TraceVec = End - Start;
+	float const Dist = TraceVec.Size();
+
+	FVector const Center = Start + TraceVec * 0.5f;
+
+	FQuat const CapsuleRot = Orientation.Quaternion();
+	::DrawDebugBox(InWorld, Start, HalfSize, CapsuleRot, Color, bPersistentLines, LifeTime, DepthPriority);
+	
+	//now draw lines from vertices
+	FVector Vertices[8];
+	Vertices[0] = Start + CapsuleRot.RotateVector(FVector(-HalfSize.X, -HalfSize.Y, -HalfSize.Z));	//flt
+	Vertices[1] = Start + CapsuleRot.RotateVector(FVector(-HalfSize.X, HalfSize.Y, -HalfSize.Z));	//frt
+	Vertices[2] = Start + CapsuleRot.RotateVector(FVector(-HalfSize.X, -HalfSize.Y, HalfSize.Z));	//flb
+	Vertices[3] = Start + CapsuleRot.RotateVector(FVector(-HalfSize.X, HalfSize.Y, HalfSize.Z));	//frb
+	Vertices[4] = Start + CapsuleRot.RotateVector(FVector(HalfSize.X, -HalfSize.Y, -HalfSize.Z));	//blt
+	Vertices[5] = Start + CapsuleRot.RotateVector(FVector(HalfSize.X, HalfSize.Y, -HalfSize.Z));	//brt
+	Vertices[6] = Start + CapsuleRot.RotateVector(FVector(HalfSize.X, -HalfSize.Y, HalfSize.Z));	//blb
+	Vertices[7] = Start + CapsuleRot.RotateVector(FVector(HalfSize.X, HalfSize.Y, HalfSize.Z));		//brb
+	for (int32 VertexIdx = 0; VertexIdx < 8; ++VertexIdx)
+	{
+		::DrawDebugLine(InWorld, Vertices[VertexIdx], Vertices[VertexIdx] + TraceVec, Color, bPersistentLines, LifeTime, DepthPriority);
+	}
+
+	::DrawDebugBox(InWorld, End, HalfSize, CapsuleRot, Color, bPersistentLines, LifeTime, DepthPriority);
+}
+
+
+bool UKismetSystemLibrary::BoxTraceSingle(UObject* WorldContextObject, const FVector& Start, const FVector& End, const FVector & HalfSize, const FRotator & Orientation, ETraceTypeQuery TraceChannel, bool bTraceComplex, const TArray<AActor*>& ActorsToIgnore, EDrawDebugTrace::Type DrawDebugType, FHitResult& OutHit, bool bIgnoreSelf)
+{
+	static const FName BoxTraceSingleName(TEXT("BoxTraceSingle"));
+
+	FCollisionQueryParams Params(BoxTraceSingleName, bTraceComplex);
+	Params.bReturnPhysicalMaterial = true;
+	Params.bTraceAsyncScene = true;
+	Params.AddIgnoredActors(ActorsToIgnore);
+	if (bIgnoreSelf)
+	{
+		if (AActor* IgnoreActor = Cast<AActor>(WorldContextObject))
+		{
+			Params.AddIgnoredActor(IgnoreActor);
+		}
+		else
+		{
+			// find owner
+			UObject * CurrentObject = WorldContextObject;
+			while (CurrentObject)
+			{
+				CurrentObject = CurrentObject->GetOuter();
+				if (AActor * IgnoreActor = Cast<AActor>(CurrentObject))
+				{
+					Params.AddIgnoredActor(IgnoreActor);
+					break;
+				}
+			}
+		}
+	}
+
+	UWorld* World = GEngine->GetWorldFromContextObject(WorldContextObject);
+	bool const bHit = World->SweepSingle(OutHit, Start, End, Orientation.Quaternion(), UEngineTypes::ConvertToCollisionChannel(TraceChannel), FCollisionShape::MakeBox(HalfSize), Params);
+
+	if (DrawDebugType != EDrawDebugTrace::None)
+	{
+		bool bPersistent = DrawDebugType == EDrawDebugTrace::Persistent;
+		float LifeTime = (DrawDebugType == EDrawDebugTrace::ForDuration) ? KISMET_TRACE_DEBUG_DRAW_DURATION : 0.f;
+
+		if (bHit && OutHit.bBlockingHit)
+		{
+			// Red up to the blocking hit, green thereafter
+			::DrawDebugSweptBox(World, Start, OutHit.Location, Orientation, HalfSize, FColor::Red, bPersistent, LifeTime);
+			::DrawDebugSweptBox(World, OutHit.Location, End, Orientation, HalfSize, FColor::Green, bPersistent, LifeTime);
+			::DrawDebugPoint(World, OutHit.ImpactPoint, KISMET_TRACE_DEBUG_IMPACTPOINT_SIZE, FColor::Red, bPersistent, LifeTime);
+		}
+		else
+		{
+			// no hit means all red
+			::DrawDebugSweptBox(World, Start, End, Orientation, HalfSize, FColor::Red, bPersistent, LifeTime);
+		}
+	}
+
+	return bHit;
+}
+
+bool UKismetSystemLibrary::BoxTraceMulti(UObject* WorldContextObject, const FVector& Start, const FVector& End, const FVector & HalfSize, const FRotator & Orientation, ETraceTypeQuery TraceChannel, bool bTraceComplex, const TArray<AActor*>& ActorsToIgnore, EDrawDebugTrace::Type DrawDebugType, TArray<FHitResult>& OutHits, bool bIgnoreSelf)
+{
+	static const FName BoxTraceMultiName(TEXT("BoxTraceMulti"));
+
+	FCollisionQueryParams Params(BoxTraceMultiName, bTraceComplex);
+	Params.bReturnPhysicalMaterial = true;
+	Params.bTraceAsyncScene = true;
+	Params.AddIgnoredActors(ActorsToIgnore);
+	if (bIgnoreSelf)
+	{
+		if (AActor* IgnoreActor = Cast<AActor>(WorldContextObject))
+		{
+			Params.AddIgnoredActor(IgnoreActor);
+		}
+		else
+		{
+			// find owner
+			UObject * CurrentObject = WorldContextObject;
+			while (CurrentObject)
+			{
+				CurrentObject = CurrentObject->GetOuter();
+				if (AActor * IgnoreActor = Cast<AActor>(CurrentObject))
+				{
+					Params.AddIgnoredActor(IgnoreActor);
+					break;
+				}
+			}
+		}
+	}
+
+	UWorld* World = GEngine->GetWorldFromContextObject(WorldContextObject);
+	bool const bHit = World->SweepMulti(OutHits, Start, End, Orientation.Quaternion(), UEngineTypes::ConvertToCollisionChannel(TraceChannel), FCollisionShape::MakeBox(HalfSize), Params);
+
+	if (DrawDebugType != EDrawDebugTrace::None)
+	{
+		bool bPersistent = DrawDebugType == EDrawDebugTrace::Persistent;
+		float LifeTime = (DrawDebugType == EDrawDebugTrace::ForDuration) ? KISMET_TRACE_DEBUG_DRAW_DURATION : 0.f;
+
+		if (bHit && OutHits.Last().bBlockingHit)
+		{
+			// Red up to the blocking hit, green thereafter
+			FVector const BlockingHitPoint = OutHits.Last().Location;
+			::DrawDebugSweptBox(World, Start, BlockingHitPoint, Orientation, HalfSize, FColor::Red, bPersistent, LifeTime);
+			::DrawDebugSweptBox(World, BlockingHitPoint, End, Orientation, HalfSize, FColor::Green, bPersistent, LifeTime);
+		}
+		else
+		{
+			// no hit means all red
+			::DrawDebugSweptBox(World, Start, End, Orientation, HalfSize, FColor::Red, bPersistent, LifeTime);
+		}
+
+		// draw hits
+		for (int32 HitIdx = 0; HitIdx < OutHits.Num(); ++HitIdx)
+		{
+			FHitResult const& Hit = OutHits[HitIdx];
+			::DrawDebugPoint(World, Hit.ImpactPoint, KISMET_TRACE_DEBUG_IMPACTPOINT_SIZE, (Hit.bBlockingHit ? FColor::Red : FColor::Green), bPersistent, LifeTime);
+		}
+	}
+
+	return bHit;
+}
 
 bool UKismetSystemLibrary::SphereTraceSingle_NEW(UObject* WorldContextObject, const FVector& Start, const FVector& End, float Radius, ETraceTypeQuery TraceChannel, bool bTraceComplex, const TArray<AActor*>& ActorsToIgnore, EDrawDebugTrace::Type DrawDebugType, FHitResult& OutHit, bool bIgnoreSelf)
 {
@@ -1578,12 +1723,16 @@ bool UKismetSystemLibrary::SphereTraceSingleByObject_DEPRECATED(UObject* WorldCo
 {
 	static const FName SphereTraceSingleName(TEXT("SphereTraceSingle"));
 
+
+
+
 	FCollisionQueryParams Params(SphereTraceSingleName, bTraceComplex);
 	Params.bReturnPhysicalMaterial = true;
 	Params.bTraceAsyncScene = true;
 	Params.AddIgnoredActors(ActorsToIgnore);
 	if (bIgnoreSelf)
 	{
+		
 		AActor* IgnoreActor = Cast<AActor>(WorldContextObject);
 		if (IgnoreActor)
 		{
@@ -1645,7 +1794,7 @@ bool UKismetSystemLibrary::SphereTraceSingleByObject_DEPRECATED(UObject* WorldCo
 		{
 			// no hit means all red
 			::DrawDebugSweptSphere(World, Start, End, Radius, FColor::Red, bPersistent, LifeTime);
-		}
+        }
 	}
 
 	return bHit;
@@ -1748,6 +1897,179 @@ bool UKismetSystemLibrary::SphereTraceMultiByObject_DEPRECATED(UObject* WorldCon
 	return bHit;
 }
 
+bool UKismetSystemLibrary::BoxTraceSingleForObjects(UObject* WorldContextObject, const FVector& Start, const FVector& End, const FVector & HalfSize, const FRotator & Orientation, const TArray<TEnumAsByte<EObjectTypeQuery> > & ObjectTypes, bool bTraceComplex, const TArray<AActor*>& ActorsToIgnore, EDrawDebugTrace::Type DrawDebugType, FHitResult& OutHit, bool bIgnoreSelf)
+{
+	static const FName BoxTraceSingleName(TEXT("BoxTraceSingle"));
+
+	TArray<TEnumAsByte<ECollisionChannel>> CollisionObjectTraces;
+	CollisionObjectTraces.AddUninitialized(ObjectTypes.Num());
+
+	for (auto Iter = ObjectTypes.CreateConstIterator(); Iter; ++Iter)
+	{
+		CollisionObjectTraces[Iter.GetIndex()] = UEngineTypes::ConvertToCollisionChannel(*Iter);
+	}
+
+	FCollisionQueryParams Params(BoxTraceSingleName, bTraceComplex);
+	Params.bReturnPhysicalMaterial = true;
+	Params.bTraceAsyncScene = true;
+	Params.AddIgnoredActors(ActorsToIgnore);
+	if (bIgnoreSelf)
+	{
+		if (AActor* IgnoreActor = Cast<AActor>(WorldContextObject))
+		{
+			Params.AddIgnoredActor(IgnoreActor);
+		}
+		else
+		{
+			// find owner
+			UObject * CurrentObject = WorldContextObject;
+			while (CurrentObject)
+			{
+				CurrentObject = CurrentObject->GetOuter();
+				if (AActor * IgnoreActor = Cast<AActor>(CurrentObject))
+				{
+					Params.AddIgnoredActor(IgnoreActor);
+					break;
+				}
+			}
+		}
+	}
+
+	FCollisionObjectQueryParams ObjectParams;
+	for (auto Iter = CollisionObjectTraces.CreateConstIterator(); Iter; ++Iter)
+	{
+		const ECollisionChannel & Channel = (*Iter);
+		if (FCollisionObjectQueryParams::IsValidObjectQuery(Channel))
+		{
+			ObjectParams.AddObjectTypesToQuery(Channel);
+		}
+		else
+		{
+			UE_LOG(LogBlueprintUserMessages, Warning, TEXT("%d isn't valid object type"), (int32)Channel);
+		}
+	}
+
+	if (ObjectParams.IsValid() == false)
+	{
+		UE_LOG(LogBlueprintUserMessages, Warning, TEXT("Invalid object types"));
+		return false;
+	}
+
+	UWorld* World = GEngine->GetWorldFromContextObject(WorldContextObject);
+	bool const bHit = World->SweepSingle(OutHit, Start, End, Orientation.Quaternion(), FCollisionShape::MakeBox(HalfSize), Params, ObjectParams);
+
+	if (DrawDebugType != EDrawDebugTrace::None)
+	{
+		bool bPersistent = DrawDebugType == EDrawDebugTrace::Persistent;
+		float LifeTime = (DrawDebugType == EDrawDebugTrace::ForDuration) ? KISMET_TRACE_DEBUG_DRAW_DURATION : 0.f;
+
+		if (bHit && OutHit.bBlockingHit)
+		{
+			// Red up to the blocking hit, green thereafter
+			FVector const BlockingHitPoint = OutHit.Location;
+			::DrawDebugSweptBox(World, Start, BlockingHitPoint, Orientation, HalfSize, FColor::Red, bPersistent, LifeTime);
+			::DrawDebugSweptBox(World, BlockingHitPoint, End, Orientation, HalfSize, FColor::Green, bPersistent, LifeTime);
+		}
+		else
+		{
+			// no hit means all red
+			::DrawDebugSweptBox(World, Start, End, Orientation, HalfSize, FColor::Red, bPersistent, LifeTime);
+		}
+	}
+
+	return bHit;
+}
+
+bool UKismetSystemLibrary::BoxTraceMultiForObjects(UObject* WorldContextObject, const FVector& Start, const FVector& End, const FVector & HalfSize, const FRotator & Orientation, const TArray<TEnumAsByte<	EObjectTypeQuery> > & ObjectTypes, bool bTraceComplex, const TArray<AActor*>& ActorsToIgnore, EDrawDebugTrace::Type DrawDebugType, TArray<FHitResult>& OutHits, bool bIgnoreSelf)
+{
+	static const FName BoxTraceMultiName(TEXT("BoxTraceMulti"));
+
+	TArray<TEnumAsByte<ECollisionChannel>> CollisionObjectTraces;
+	CollisionObjectTraces.AddUninitialized(ObjectTypes.Num());
+
+	for (auto Iter = ObjectTypes.CreateConstIterator(); Iter; ++Iter)
+	{
+		CollisionObjectTraces[Iter.GetIndex()] = UEngineTypes::ConvertToCollisionChannel(*Iter);
+	}
+
+	FCollisionQueryParams Params(BoxTraceMultiName, bTraceComplex);
+	Params.bReturnPhysicalMaterial = true;
+	Params.bTraceAsyncScene = true;
+	Params.AddIgnoredActors(ActorsToIgnore);
+	if (bIgnoreSelf)
+	{
+		if (AActor* IgnoreActor = Cast<AActor>(WorldContextObject))
+		{
+			Params.AddIgnoredActor(IgnoreActor);
+		}
+		else
+		{
+			// find owner
+			UObject * CurrentObject = WorldContextObject;
+			while (CurrentObject)
+			{
+				CurrentObject = CurrentObject->GetOuter();
+				if (AActor * IgnoreActor = Cast<AActor>(CurrentObject))
+				{
+					Params.AddIgnoredActor(IgnoreActor);
+					break;
+				}
+			}
+		}
+	}
+
+	FCollisionObjectQueryParams ObjectParams;
+	for (auto Iter = CollisionObjectTraces.CreateConstIterator(); Iter; ++Iter)
+	{
+		const ECollisionChannel & Channel = (*Iter);
+		if (FCollisionObjectQueryParams::IsValidObjectQuery(Channel))
+		{
+			ObjectParams.AddObjectTypesToQuery(Channel);
+		}
+		else
+		{
+			UE_LOG(LogBlueprintUserMessages, Warning, TEXT("%d isn't valid object type"), (int32)Channel);
+		}
+	}
+
+	if (ObjectParams.IsValid() == false)
+	{
+		UE_LOG(LogBlueprintUserMessages, Warning, TEXT("Invalid object types"));
+		return false;
+	}
+
+	UWorld* World = GEngine->GetWorldFromContextObject(WorldContextObject);
+	bool const bHit = World->SweepMulti(OutHits, Start, End, Orientation.Quaternion(), FCollisionShape::MakeBox(HalfSize), Params, ObjectParams);
+
+	if (DrawDebugType != EDrawDebugTrace::None)
+	{
+		bool bPersistent = DrawDebugType == EDrawDebugTrace::Persistent;
+		float LifeTime = (DrawDebugType == EDrawDebugTrace::ForDuration) ? KISMET_TRACE_DEBUG_DRAW_DURATION : 0.f;
+
+		if (bHit && OutHits.Last().bBlockingHit)
+		{
+			// Red up to the blocking hit, green thereafter
+			FVector const BlockingHitPoint = OutHits.Last().Location;
+			::DrawDebugSweptBox(World, Start, BlockingHitPoint, Orientation, HalfSize, FColor::Red, bPersistent, LifeTime);
+			::DrawDebugSweptBox(World, BlockingHitPoint, End, Orientation, HalfSize, FColor::Green, bPersistent, LifeTime);
+		}
+		else
+		{
+			// no hit means all red
+			::DrawDebugSweptBox(World, Start, End, Orientation, HalfSize, FColor::Red, bPersistent, LifeTime);
+		}
+
+		// draw hits
+		for (int32 HitIdx = 0; HitIdx < OutHits.Num(); ++HitIdx)
+		{
+			FHitResult const& Hit = OutHits[HitIdx];
+			::DrawDebugPoint(World, Hit.ImpactPoint, KISMET_TRACE_DEBUG_IMPACTPOINT_SIZE, (Hit.bBlockingHit ? FColor::Red : FColor::Green), bPersistent, LifeTime);
+		}
+	}
+
+	return bHit;
+}
+
 bool UKismetSystemLibrary::CapsuleTraceSingleForObjects(UObject* WorldContextObject, const FVector& Start, const FVector& End, float Radius, float HalfHeight, const TArray<TEnumAsByte<EObjectTypeQuery> > & ObjectTypes, bool bTraceComplex, const TArray<AActor*>& ActorsToIgnore, EDrawDebugTrace::Type DrawDebugType, FHitResult& OutHit, bool bIgnoreSelf)
 {
 	TArray<TEnumAsByte<ECollisionChannel>> CollisionObjectTraces;
