Fix crash when trying to update body scale but BodySetup is pending kill
TTP#338453

diff --git a/Engine/Source/Runtime/Engine/Private/PhysicsEngine/BodyInstance.cpp b/Engine/Source/Runtime/Engine/Private/PhysicsEngine/BodyInstance.cpp
--- a/Engine/Source/Runtime/Engine/Private/PhysicsEngine/BodyInstance.cpp
+++ b/Engine/Source/Runtime/Engine/Private/PhysicsEngine/BodyInstance.cpp
@@ -1628,35 +1628,39 @@ bool FBodyInstance::UpdateBodyScale(const FVector& InScale3D)
 				// find which convex elems it is
 				// it would be nice to know if the order of PShapes array index is in the order of createShape
 				// Create convex shapes
-				for (int32 i = 0; i < BodySetup->AggGeom.ConvexElems.Num(); i++)
+				if (BodySetup.IsValid())
 				{
-					FKConvexElem* ConvexElem = &(BodySetup->AggGeom.ConvexElems[i]);
-
-					// found it
-					if (ConvexElem->ConvexMesh == PConvexGeom.convexMesh)
+					for (int32 i = 0; i < BodySetup->AggGeom.ConvexElems.Num(); i++)
 					{
-						// Please note that this one we don't inverse old scale, but just set new one (but we still follow scale mode restriction)
-						FVector NewScale3D = RelativeScale3D * OldScale3D;
-						FVector Scale3DAbs(FMath::Abs(NewScale3D.X), FMath::Abs(NewScale3D.Y), FMath::Abs(NewScale3D.Z)); // magnitude of scale (sign removed)
+						FKConvexElem* ConvexElem = &(BodySetup->AggGeom.ConvexElems[i]);
 
-						PxTransform PNewLocalPose;
-						bool bUseNegX = CalcMeshNegScaleCompensation(NewScale3D, PNewLocalPose);
+						// found it
+						if (ConvexElem->ConvexMesh == PConvexGeom.convexMesh)
+						{
+							// Please note that this one we don't inverse old scale, but just set new one (but we still follow scale mode restriction)
+							FVector NewScale3D = RelativeScale3D * OldScale3D;
+							FVector Scale3DAbs(FMath::Abs(NewScale3D.X), FMath::Abs(NewScale3D.Y), FMath::Abs(NewScale3D.Z)); // magnitude of scale (sign removed)
 
-						PxTransform PElementTransform = U2PTransform(ConvexElem->GetTransform());
-						PNewLocalPose.q *= PElementTransform.q;
-						PNewLocalPose.p += PElementTransform.p;
+							PxTransform PNewLocalPose;
+							bool bUseNegX = CalcMeshNegScaleCompensation(NewScale3D, PNewLocalPose);
 
-						PConvexGeom.convexMesh = bUseNegX ? ConvexElem->ConvexMeshNegX : ConvexElem->ConvexMesh;
-						PConvexGeom.scale.scale = U2PVector(Scale3DAbs);
+							PxTransform PElementTransform = U2PTransform(ConvexElem->GetTransform());
+							PNewLocalPose.q *= PElementTransform.q;
+							PNewLocalPose.p += PElementTransform.p;
 
-						if (PConvexGeom.isValid())
-						{
-							UpdatedGeometry = &PConvexGeom;
-							bSuccess = true;
+							PConvexGeom.convexMesh = bUseNegX ? ConvexElem->ConvexMeshNegX : ConvexElem->ConvexMesh;
+							PConvexGeom.scale.scale = U2PVector(Scale3DAbs);
+
+							if (PConvexGeom.isValid())
+							{
+								UpdatedGeometry = &PConvexGeom;
+								bSuccess = true;
+							}
+							break;
 						}
-						break;
 					}
 				}
+				
 				break;
 			}
 			case PxGeometryType::eTRIANGLEMESH:
@@ -1665,7 +1669,7 @@ bool FBodyInstance::UpdateBodyScale(const FVector& InScale3D)
 				SCENE_UNLOCK_READ(PScene);
 
 				// Create tri-mesh shape
-				if (BodySetup->TriMesh != NULL || BodySetup->TriMeshNegX != NULL)
+				if (BodySetup.IsValid() && (BodySetup->TriMesh != NULL || BodySetup->TriMeshNegX != NULL))
 				{
 					// Please note that this one we don't inverse old scale, but just set new one (but still adjust for scale mode)
 					FVector NewScale3D = RelativeScale3D * OldScale3D;
