Disable PhysX cleaning of mesh for per poly skeletal mesh case.

diff --git a/Engine/Source/Developer/PhysXFormats/Private/PhysXFormats.cpp b/Engine/Source/Developer/PhysXFormats/Private/PhysXFormats.cpp
--- a/Engine/Source/Developer/PhysXFormats/Private/PhysXFormats.cpp
+++ b/Engine/Source/Developer/PhysXFormats/Private/PhysXFormats.cpp
@@ -115,7 +115,7 @@ public:
 		return false;
 	}
 
-	virtual bool CookTriMesh(FName Format, const TArray<FVector>& SrcVertices, const TArray<FTriIndices>& SrcIndices, const TArray<uint16>& SrcMaterialIndices, const bool FlipNormals, TArray<uint8>& OutBuffer) const override
+	virtual bool CookTriMesh(FName Format, const TArray<FVector>& SrcVertices, const TArray<FTriIndices>& SrcIndices, const TArray<uint16>& SrcMaterialIndices, const bool FlipNormals, TArray<uint8>& OutBuffer, bool bPerPolySkeletalMesh = false) const override
 	{
 #if WITH_PHYSX
 		PxPlatform::Enum PhysXFormat = PxPlatform::ePC;
@@ -137,11 +137,26 @@ public:
 		const PxCookingParams& Params = PhysXCooking->getParams();
 		PxCookingParams NewParams = Params;
 		NewParams.targetPlatform = PhysXFormat;
+		PxMeshPreprocessingFlags OldCookingFlags = NewParams.meshPreprocessParams;
+
+		if (bPerPolySkeletalMesh)
+		{
+			//per poly skeletal mesh requires deforming mesh. This is a very special case so we have to change the cook params
+			NewParams.meshPreprocessParams = PxMeshPreprocessingFlag::eDISABLE_CLEAN_MESH;
+		}
+
 		PhysXCooking->setParams(NewParams);
 
+
 		// Cook TriMesh Data
 		FPhysXOutputStream Buffer(&OutBuffer);
 		bool Result = PhysXCooking->cookTriangleMesh(PTriMeshDesc, Buffer);
+		
+		if (bPerPolySkeletalMesh)	//restore old params
+		{
+			NewParams.meshPreprocessParams = OldCookingFlags;
+			PhysXCooking->setParams(NewParams);
+		}
 		return Result;
 #else
 		return false;
diff --git a/Engine/Source/Developer/TargetPlatform/Public/Interfaces/IPhysXFormat.h b/Engine/Source/Developer/TargetPlatform/Public/Interfaces/IPhysXFormat.h
--- a/Engine/Source/Developer/TargetPlatform/Public/Interfaces/IPhysXFormat.h
+++ b/Engine/Source/Developer/TargetPlatform/Public/Interfaces/IPhysXFormat.h
@@ -42,9 +42,10 @@ public:
 	 * @param Format The desired format.
 	 * @param SrcBuffer The source buffer.
 	 * @param OutBuffer The resulting cooked data.
+	 * @param bPerPolySkeletalMesh This is a very special case that requires different cooking parameters set.
 	 * @return true on success, false otherwise.
 	 */
-	virtual bool CookTriMesh( FName Format, const TArray<FVector>& SrcVertices, const TArray<struct FTriIndices>& SrcIndices, const TArray<uint16>& SrcMaterialIndices, const bool FlipNormals, TArray<uint8>& OutBuffer ) const = 0;
+	virtual bool CookTriMesh( FName Format, const TArray<FVector>& SrcVertices, const TArray<struct FTriIndices>& SrcIndices, const TArray<uint16>& SrcMaterialIndices, const bool FlipNormals, TArray<uint8>& OutBuffer, bool bPerPolySkeletalMesh = false ) const = 0;
 		
 	/**
 	 * Cooks the source height field data for the platform and stores the cooked data internally.
diff --git a/Engine/Source/Runtime/Engine/Classes/PhysicsEngine/BodyInstance.h b/Engine/Source/Runtime/Engine/Classes/PhysicsEngine/BodyInstance.h
--- a/Engine/Source/Runtime/Engine/Classes/PhysicsEngine/BodyInstance.h
+++ b/Engine/Source/Runtime/Engine/Classes/PhysicsEngine/BodyInstance.h
@@ -281,7 +281,7 @@ public:
 	 */
 	bool UpdateBodyScale(const FVector& InScale3D);
 
-	void UpdateTriMeshVertices(const TArray<FVector> & NewPositions, const TArray<uint32> & Indices);
+	void UpdateTriMeshVertices(const TArray<FVector> & NewPositions);
 
 	/** Returns the center of mass of this body (in world space) */
 	FVector GetCOMPosition() const;
diff --git a/Engine/Source/Runtime/Engine/Classes/PhysicsEngine/BodySetup.h b/Engine/Source/Runtime/Engine/Classes/PhysicsEngine/BodySetup.h
--- a/Engine/Source/Runtime/Engine/Classes/PhysicsEngine/BodySetup.h
+++ b/Engine/Source/Runtime/Engine/Classes/PhysicsEngine/BodySetup.h
@@ -67,49 +67,49 @@ struct FKConvexElem
 	UPROPERTY()
 	FTransform Transform;
 
-	/** Convex mesh for this body, created from cooked data in CreatePhysicsMeshes */
-	physx::PxConvexMesh*   ConvexMesh;	
+		/** Convex mesh for this body, created from cooked data in CreatePhysicsMeshes */
+		physx::PxConvexMesh*   ConvexMesh;	
 
-	/** Convex mesh for this body, flipped across X, created from cooked data in CreatePhysicsMeshes */
-	physx::PxConvexMesh*   ConvexMeshNegX;
+		/** Convex mesh for this body, flipped across X, created from cooked data in CreatePhysicsMeshes */
+		physx::PxConvexMesh*   ConvexMeshNegX;
 
-	FKConvexElem() 
-	: ElemBox(0)
-	, Transform(FTransform::Identity)
-	, ConvexMesh(NULL)
-	, ConvexMeshNegX(NULL)
-	{}
+		FKConvexElem() 
+		: ElemBox(0)
+		, Transform(FTransform::Identity)
+		, ConvexMesh(NULL)
+		, ConvexMeshNegX(NULL)
+		{}
 
 	ENGINE_API void	DrawElemWire(class FPrimitiveDrawInterface* PDI, const FTransform& ElemTM, const FColor Color) const;
-	void	AddCachedSolidConvexGeom(TArray<FDynamicMeshVertex>& VertexBuffer, TArray<int32>& IndexBuffer, const FColor VertexColor);
+		void	AddCachedSolidConvexGeom(TArray<FDynamicMeshVertex>& VertexBuffer, TArray<int32>& IndexBuffer, const FColor VertexColor);
 
-	/** Reset the hull to empty all arrays */
-	ENGINE_API void	Reset();
+		/** Reset the hull to empty all arrays */
+		ENGINE_API void	Reset();
 
-	/** Updates internal ElemBox based on current value of VertexData */
-	ENGINE_API void	UpdateElemBox();
+		/** Updates internal ElemBox based on current value of VertexData */
+		ENGINE_API void	UpdateElemBox();
 
-	/** Calculate a bounding box for this convex element with the specified transform and scale */
+		/** Calculate a bounding box for this convex element with the specified transform and scale */
 	ENGINE_API FBox	CalcAABB(const FTransform& BoneTM, const FVector& Scale3D) const;
 
-	/** Utility for creating a convex hull from a set of planes. Will reset current state of this elem. */
-	bool	HullFromPlanes(const TArray<FPlane>& InPlanes, const TArray<FVector>& SnapVerts);
+		/** Utility for creating a convex hull from a set of planes. Will reset current state of this elem. */
+		bool	HullFromPlanes(const TArray<FPlane>& InPlanes, const TArray<FVector>& SnapVerts);
 
-	/** Returns the volume of this element */
-	float GetVolume(const FVector& Scale) const;
+		/** Returns the volume of this element */
+		float GetVolume(const FVector& Scale) const;
 
-	FTransform GetTransform() const
-	{
-		return Transform;
-	};
+		FTransform GetTransform() const
+		{
+			return Transform;
+		};
 
-	void SetTransform( const FTransform& InTransform )
-	{
-		ensure(InTransform.IsValid());
-		Transform = InTransform;
-	}
+		void SetTransform( const FTransform& InTransform )
+		{
+			ensure(InTransform.IsValid());
+			Transform = InTransform;
+		}
 
-	friend FArchive& operator<<(FArchive& Ar,FKConvexElem& Elem);
+		friend FArchive& operator<<(FArchive& Ar,FKConvexElem& Elem);
 	
 	ENGINE_API void ScaleElem(FVector DeltaSize, float MinSize);
 };
@@ -380,12 +380,12 @@ struct ENGINE_API FKAggregateGeom
 	FBox CalcAABB(const FTransform& Transform) const;
 
 	/**
-	 * Calculates a tight box-sphere bounds for the aggregate geometry; this is more expensive than CalcAABB
-	 * (tight meaning the sphere may be smaller than would be required to encompass the AABB, but all individual components lie within both the box and the sphere)
-	 *
-	 *  @param Output The output box-sphere bounds calculated for this set of aggregate geometry
-	 *  @param LocalToWorld Transform
-	 */
+		* Calculates a tight box-sphere bounds for the aggregate geometry; this is more expensive than CalcAABB
+		* (tight meaning the sphere may be smaller than would be required to encompass the AABB, but all individual components lie within both the box and the sphere)
+		*
+		* @param Output The output box-sphere bounds calculated for this set of aggregate geometry
+		*	@param LocalToWorld Transform
+		*/
 	void CalcBoxSphereBounds(FBoxSphereBounds& Output, const FTransform& LocalToWorld) const;
 
 	/** Returns the volume of this element */
@@ -560,7 +560,7 @@ public:
 	 * In 99% of cases you should be fine using a physics asset created for the skeletal mesh
 	 * @param	InSkeletalMeshComponent		The skeletal mesh component we'll be grabbing the skinning information from
 	 */
-	ENGINE_API void UpdateTriMeshVertices(const TArray<FVector> & NewPositions, const TArray<uint32> & Indices);
+	ENGINE_API void UpdateTriMeshVertices(const TArray<FVector> & NewPositions);
 
 	/**
 	 * Given a format name returns its cooked data.
diff --git a/Engine/Source/Runtime/Engine/Private/PhysicsEngine/BodyInstance.cpp b/Engine/Source/Runtime/Engine/Private/PhysicsEngine/BodyInstance.cpp
--- a/Engine/Source/Runtime/Engine/Private/PhysicsEngine/BodyInstance.cpp
+++ b/Engine/Source/Runtime/Engine/Private/PhysicsEngine/BodyInstance.cpp
@@ -309,7 +309,7 @@ TArray<PxShape*> FBodyInstance::GetAllShapes(int32& OutNumSyncShapes) const
 }
 #endif
 
-void FBodyInstance::UpdateTriMeshVertices(const TArray<FVector> & NewPositions, const TArray<uint32> & Indices)
+void FBodyInstance::UpdateTriMeshVertices(const TArray<FVector> & NewPositions)
 {
 #if WITH_PHYSX
 	if (BodySetup.IsValid())
@@ -317,7 +317,7 @@ void FBodyInstance::UpdateTriMeshVertices(const TArray<FVector> & NewPositions,
 		SCOPED_SCENE_WRITE_LOCK((RigidActorSync ? RigidActorSync->getScene() : NULL));
 		SCOPED_SCENE_WRITE_LOCK((RigidActorAsync ? RigidActorAsync->getScene() : NULL));
 
-		BodySetup->UpdateTriMeshVertices(NewPositions, Indices);
+		BodySetup->UpdateTriMeshVertices(NewPositions);
 
 		//after updating the vertices we must call setGeometry again to update any shapes referencing the mesh
 
diff --git a/Engine/Source/Runtime/Engine/Private/PhysicsEngine/BodySetup.cpp b/Engine/Source/Runtime/Engine/Private/PhysicsEngine/BodySetup.cpp
--- a/Engine/Source/Runtime/Engine/Private/PhysicsEngine/BodySetup.cpp
+++ b/Engine/Source/Runtime/Engine/Private/PhysicsEngine/BodySetup.cpp
@@ -734,44 +734,17 @@ void UBodySetup::PostLoad()
 	}
 }
 
-void UBodySetup::UpdateTriMeshVertices(const TArray<FVector> & NewPositions, const TArray<uint32> & Indices)
+void UBodySetup::UpdateTriMeshVertices(const TArray<FVector> & NewPositions)
 {
 	SCOPE_CYCLE_COUNTER(STAT_UpdateTriMeshVertices);
 #if WITH_PHYSX
 	if (TriMesh)
 	{
-		const bool bNeedsRemap = true;	//depending on if we clean verts during cooking there can be a 1-1 map or not. Need to benchmark this
-		if (bNeedsRemap)
+		
+		PxVec3 * PNewPositions = TriMesh->getVerticesForModification();
+		for (int32 i = 0; i < NewPositions.Num(); ++i)
 		{
-			const PxU32 * RuntimeToOriginal = TriMesh->getTrianglesRemap();
-			bool b16BitIndices = TriMesh->getTriangleMeshFlags() & PxTriangleMeshFlag::e16_BIT_INDICES;
-			const void * PIndices = TriMesh->getTriangles();
-
-			uint32 PNumTris = TriMesh->getNbTriangles();
-			PxVec3 * PNewPositions = TriMesh->getVerticesForModification();
-
-			for (uint32 TriIdx = 0; TriIdx < PNumTris; ++TriIdx)
-			{
-				PxU32 OriginalStartIndex = RuntimeToOriginal[TriIdx] * 3;
-				for (int i = 0; i < 3; ++i)
-				{
-					PxU32 PIndex = TriIdx * 3 + i;
-					PxU32 PVertexIndex = b16BitIndices ? ((PxU16*)PIndices)[PIndex] : ((PxU32*)PIndices)[PIndex];
-					PxVec3 & PRuntimeVertex = PNewPositions[PVertexIndex];
-
-					PxU32 OriginalVertexIndex = Indices[OriginalStartIndex + i];
-					const FVector & OriginalVertex = NewPositions[OriginalVertexIndex];
-					PRuntimeVertex = U2PVector(OriginalVertex);
-				}
-			}
-		}
-		else
-		{
-			PxVec3 * PNewPositions = TriMesh->getVerticesForModification();
-			for (int32 i = 0; i < NewPositions.Num(); ++i)
-			{
-				PNewPositions[i] = U2PVector(NewPositions[i]);
-			}
+			PNewPositions[i] = U2PVector(NewPositions[i]);
 		}
 
 		TriMesh->refitBVH();
diff --git a/Engine/Source/Runtime/Engine/Private/PhysicsEngine/PhysAnim.cpp b/Engine/Source/Runtime/Engine/Private/PhysicsEngine/PhysAnim.cpp
--- a/Engine/Source/Runtime/Engine/Private/PhysicsEngine/PhysAnim.cpp
+++ b/Engine/Source/Runtime/Engine/Private/PhysicsEngine/PhysAnim.cpp
@@ -387,10 +387,7 @@ void USkeletalMeshComponent::UpdateKinematicBonesToPhysics(bool bTeleport, bool
 						}
 					}
 				}
-
-				TArray<uint32> Indices;
-				Model.MultiSizeIndexContainer.GetIndexBuffer(Indices);
-				BodyInstance.UpdateTriMeshVertices(NewPositions, Indices);
+				BodyInstance.UpdateTriMeshVertices(NewPositions);
 			}
 			
 			BodyInstance.SetBodyTransform(CurrentLocalToWorld, bTeleport);
diff --git a/Engine/Source/Runtime/Engine/Private/PhysicsEngine/PhysDerivedData.cpp b/Engine/Source/Runtime/Engine/Private/PhysicsEngine/PhysDerivedData.cpp
--- a/Engine/Source/Runtime/Engine/Private/PhysicsEngine/PhysDerivedData.cpp
+++ b/Engine/Source/Runtime/Engine/Private/PhysicsEngine/PhysDerivedData.cpp
@@ -200,7 +200,13 @@ bool FDerivedDataPhysXCooker::BuildTriMesh( TArray<uint8>& OutData, bool bInMirr
 		}
 
 		UE_LOG(LogPhysics, Log, TEXT("Cook TriMesh: %s (FlipX: %d)"), *CollisionDataProvider->GetPathName(), bInMirrored);
-		bResult = Cooker->CookTriMesh( Format, *MeshVertices, TriangleMeshDesc.Indices, TriangleMeshDesc.MaterialIndices, bInMirrored ? !TriangleMeshDesc.bFlipNormals : TriangleMeshDesc.bFlipNormals, OutData );
+		bool bPerPolySkeletalMesh = false;
+		if (USkeletalMeshComponent * SkeletalMeshComponent = Cast<USkeletalMeshComponent>(CollisionDataProvider))
+		{
+			ensure(SkeletalMeshComponent->bEnablePerPolyCollision);
+			bPerPolySkeletalMesh = true;
+		}
+		bResult = Cooker->CookTriMesh( Format, *MeshVertices, TriangleMeshDesc.Indices, TriangleMeshDesc.MaterialIndices, bInMirrored ? !TriangleMeshDesc.bFlipNormals : TriangleMeshDesc.bFlipNormals, OutData, bPerPolySkeletalMesh );
 		if( !bResult )
 		{
 			UE_LOG(LogPhysics, Warning, TEXT("Failed to cook TriMesh: %s (FlipX:%d)."), *CollisionDataProvider->GetPathName(), bInMirrored );
