Added safe guard (ensure) to make sure it doesn't access invalid index when we initialize the AI.

336031	CRASH: Assertion failed: (Index >= 0) Function: FPoseLinkBase::AttemptRelink() USkeletalMeshComponent::InitializeAnimScriptInsà?

diff --git a/Engine/Source/Runtime/Engine/Private/Animation/AnimNodeBase.cpp b/Engine/Source/Runtime/Engine/Private/Animation/AnimNodeBase.cpp
--- a/Engine/Source/Runtime/Engine/Private/Animation/AnimNodeBase.cpp
+++ b/Engine/Source/Runtime/Engine/Private/Animation/AnimNodeBase.cpp
@@ -14,9 +14,13 @@ void FPoseLinkBase::AttemptRelink(const FAnimationBaseContext& Context)
 		UAnimBlueprintGeneratedClass* AnimBlueprintClass = Context.GetAnimBlueprintClass();
 		check(AnimBlueprintClass);
 
-		UProperty* LinkedProperty = AnimBlueprintClass->AnimNodeProperties[LinkID];
-		void* LinkedNodePtr = LinkedProperty->ContainerPtrToValuePtr<void>(Context.AnimInstance);
-		LinkedNode = (FAnimNode_Base*)LinkedNodePtr;
+		// adding ensure. We had a crash on here
+		if ( ensure(AnimBlueprintClass->AnimNodeProperties.IsValidIndex(LinkID)) )
+		{
+			UProperty* LinkedProperty = AnimBlueprintClass->AnimNodeProperties[LinkID];
+			void* LinkedNodePtr = LinkedProperty->ContainerPtrToValuePtr<void>(Context.AnimInstance);
+			LinkedNode = (FAnimNode_Base*)LinkedNodePtr;
+		}
 	}
 }
 
