Prevent crash caused by a negative landscape streaming distance multiplier

diff --git a/Engine/Source/Runtime/Engine/Classes/Landscape/LandscapeProxy.h b/Engine/Source/Runtime/Engine/Classes/Landscape/LandscapeProxy.h
--- a/Engine/Source/Runtime/Engine/Classes/Landscape/LandscapeProxy.h
+++ b/Engine/Source/Runtime/Engine/Classes/Landscape/LandscapeProxy.h
@@ -221,7 +221,7 @@ public:
 	 * Allows artists to adjust the distance where textures using UV 0 are streamed in/out.
 	 * 1.0 is the default, whereas a higher value increases the streamed-in resolution.
 	 */
-	UPROPERTY(EditAnywhere, Category=Landscape)
+	UPROPERTY(EditAnywhere, Category=Landscape, meta=(ClampMin=0))
 	float StreamingDistanceMultiplier;
 
 	/** Combined material used to render the landscape */
