fixed navoctree updates on moving geometry brush

diff --git a/Engine/Source/Runtime/Engine/Private/AI/Navigation/NavigationSystem.cpp b/Engine/Source/Runtime/Engine/Private/AI/Navigation/NavigationSystem.cpp
--- a/Engine/Source/Runtime/Engine/Private/AI/Navigation/NavigationSystem.cpp
+++ b/Engine/Source/Runtime/Engine/Private/AI/Navigation/NavigationSystem.cpp
@@ -2530,12 +2530,6 @@ void UNavigationSystem::OnLevelAddedToWorld(ULevel* InLevel, UWorld* InWorld)
 	if (InWorld == GetWorld())
 	{
 		AddLevelCollisionToOctree(InLevel);
-
-		FBox NavigableLevelBounds = GetLevelBounds(InLevel);
-		if (NavigableLevelBounds.IsValid)
-		{
-			AddDirtyArea(NavigableLevelBounds, ENavigationDirtyFlag::All);
-		}
 	}
 }
 
@@ -2544,12 +2538,6 @@ void UNavigationSystem::OnLevelRemovedFromWorld(ULevel* InLevel, UWorld* InWorld
 	if (InWorld == GetWorld())
 	{
 		RemoveLevelCollisionFromOctree(InLevel);
-
-		FBox NavigableLevelBounds = GetLevelBounds(InLevel);
-		if (NavigableLevelBounds.IsValid)
-		{
-			AddDirtyArea(NavigableLevelBounds, ENavigationDirtyFlag::All);
-		}
 	}
 }
 
@@ -2566,6 +2554,8 @@ void UNavigationSystem::AddLevelCollisionToOctree(ULevel* Level)
 		BSPElem.Owner = Level;
 		
 		NavOctree->AddNode(BSPElem);
+
+		AddDirtyArea(BSPElem.Bounds.GetBox(), ENavigationDirtyFlag::All);
 	}
 #endif // NAVOCTREE_CONTAINS_COLLISION_DATA
 }
@@ -2576,6 +2566,13 @@ void UNavigationSystem::RemoveLevelCollisionFromOctree(ULevel* Level)
 	const FOctreeElementId* ElementId = GetObjectsNavOctreeId(Level);
 	if (ElementId != NULL)
 	{
+		if (ElementId->IsValidId() && NavOctree->IsValidElementId(*ElementId))
+		{
+			// mark area occupied by given actor as dirty
+			FNavigationOctreeElement& ElementData = NavOctree->GetElementById(*ElementId);
+			AddDirtyArea(ElementData.Bounds.GetBox(), ENavigationDirtyFlag::All);
+		}
+
 		NavOctree->RemoveNode(ElementId);
 		RemoveObjectsNavOctreeId(Level);
 	}
