Added code to ensure that a spotlight's inner cone angle cannot be set greater than its outer cone angle

This addresses TTP#338313 - EDITOR: LIGHTING: Unexpected behavior. Inner Cone Angle can be wider than the Outer Cone Angle

diff --git a/Engine/Source/Runtime/Engine/Classes/Components/SpotLightComponent.h b/Engine/Source/Runtime/Engine/Classes/Components/SpotLightComponent.h
--- a/Engine/Source/Runtime/Engine/Classes/Components/SpotLightComponent.h
+++ b/Engine/Source/Runtime/Engine/Classes/Components/SpotLightComponent.h
@@ -39,6 +39,7 @@ class ENGINE_API USpotLightComponent : public UPointLightComponent
 	virtual bool AffectsBounds(const FBoxSphereBounds& Bounds) const;
 	virtual ELightComponentType GetLightType() const;
 	virtual FLightSceneProxy* CreateSceneProxy() const override;
+	virtual void PostEditChangeProperty( struct FPropertyChangedEvent& PropertyChangedEvent) override;
 };
 
 
diff --git a/Engine/Source/Runtime/Engine/Private/SpotLightComponent.cpp b/Engine/Source/Runtime/Engine/Private/SpotLightComponent.cpp
--- a/Engine/Source/Runtime/Engine/Private/SpotLightComponent.cpp
+++ b/Engine/Source/Runtime/Engine/Private/SpotLightComponent.cpp
@@ -256,3 +256,20 @@ ELightComponentType USpotLightComponent::GetLightType() const
 {
 	return LightType_Spot;
 }
+
+void USpotLightComponent::PostEditChangeProperty( FPropertyChangedEvent& PropertyChangedEvent)
+{
+	if (PropertyChangedEvent.Property)
+	{
+		if (PropertyChangedEvent.Property->GetFName() == TEXT("InnerConeAngle"))
+		{
+			OuterConeAngle = FMath::Max(InnerConeAngle, OuterConeAngle);
+		}
+		else if (PropertyChangedEvent.Property->GetFName() == TEXT("OuterConeAngle"))
+		{
+			InnerConeAngle = FMath::Min(InnerConeAngle, OuterConeAngle);
+		}
+	}
+
+	UPointLightComponent::PostEditChangeProperty(PropertyChangedEvent);
+}
