CIS fix for non-editor builds

diff --git a/Engine/Source/Runtime/Engine/Classes/Components/SpotLightComponent.h b/Engine/Source/Runtime/Engine/Classes/Components/SpotLightComponent.h
--- a/Engine/Source/Runtime/Engine/Classes/Components/SpotLightComponent.h
+++ b/Engine/Source/Runtime/Engine/Classes/Components/SpotLightComponent.h
@@ -39,7 +39,9 @@ class ENGINE_API USpotLightComponent : public UPointLightComponent
 	virtual bool AffectsBounds(const FBoxSphereBounds& Bounds) const;
 	virtual ELightComponentType GetLightType() const;
 	virtual FLightSceneProxy* CreateSceneProxy() const override;
+#if WITH_EDITOR
 	virtual void PostEditChangeProperty( struct FPropertyChangedEvent& PropertyChangedEvent) override;
+#endif
 };
 
 
diff --git a/Engine/Source/Runtime/Engine/Private/SpotLightComponent.cpp b/Engine/Source/Runtime/Engine/Private/SpotLightComponent.cpp
--- a/Engine/Source/Runtime/Engine/Private/SpotLightComponent.cpp
+++ b/Engine/Source/Runtime/Engine/Private/SpotLightComponent.cpp
@@ -257,6 +257,8 @@ ELightComponentType USpotLightComponent::GetLightType() const
 	return LightType_Spot;
 }
 
+#if WITH_EDITOR
+
 void USpotLightComponent::PostEditChangeProperty( FPropertyChangedEvent& PropertyChangedEvent)
 {
 	if (PropertyChangedEvent.Property)
@@ -273,3 +275,5 @@ void USpotLightComponent::PostEditChangeProperty( FPropertyChangedEvent& Propert
 
 	UPointLightComponent::PostEditChangeProperty(PropertyChangedEvent);
 }
+
+#endif	// WITH_EDITOR
\ No newline at end of file
