More packagemap cleanup

* Remove unused macros
* Better sanity checking on missing assets, and more sanity checking in general in other areas
* Removed redundant code
* Improved comments

diff --git a/Engine/Source/Runtime/Engine/Private/PackageMapClient.cpp b/Engine/Source/Runtime/Engine/Private/PackageMapClient.cpp
--- a/Engine/Source/Runtime/Engine/Private/PackageMapClient.cpp
+++ b/Engine/Source/Runtime/Engine/Private/PackageMapClient.cpp
@@ -25,8 +25,6 @@ static const int INTERNAL_LOAD_OBJECT_RECURSION_LIMIT = 16;
 #define COMPOSE_NET_GUID( Index, IsStatic )			( ( ( Index ) << 1 ) | ( IsStatic ) )
 #define COMPOSE_RELATIVE_NET_GUID( GUID, Index )	( COMPOSE_NET_GUID( GET_NET_GUID_INDEX( GUID ) + Index, GET_NET_GUID_STATIC_MOD( GUID ) ) )
 
-#define STRESS_PENDING_GUIDS ( 0 )
-
 /*-----------------------------------------------------------------------------
 	UPackageMapClient implementation.
 -----------------------------------------------------------------------------*/
@@ -654,9 +652,18 @@ FNetworkGUID UPackageMapClient::InternalLoadObject( FArchive & Ar, UObject *& Ob
 			UE_LOG( LogNetPackageMap, Warning, TEXT( "InternalLoadObject: Unable to resolve object from path. Path: %s, Outer: %s, NetGUID: %s" ), *PathName, ObjOuter ? *ObjOuter->GetPathName() : TEXT( "NULL" ), *NetGUID.ToString() );
 		}
 	}
-	else if ( Object == NULL && NetGUID.IsStatic() && !GuidCache->ShouldIgnoreWhenMissing( NetGUID ) )
+	else if ( Object == NULL && !GuidCache->ShouldIgnoreWhenMissing( NetGUID ) )
 	{
-		UE_LOG( LogNetPackageMap, Warning, TEXT( "InternalLoadObject: Unable to resolve object. NetGUID: %s, HasPath: %s" ), *NetGUID.ToString(), ExportFlags.bHasPath ? TEXT( "TRUE" ) : TEXT( "FALSE" ) );
+		const FNetGuidCacheObject * CacheObject = GuidCache->ObjectLookup.Find( NetGUID );
+
+		if ( CacheObject == NULL )
+		{
+			UE_LOG( LogNetPackageMap, Warning, TEXT( "InternalLoadObject: Unable to resolve object (not in cache). NetGUID: %s" ), *NetGUID.ToString() );
+		}
+		else
+		{
+			UE_LOG( LogNetPackageMap, Warning, TEXT( "InternalLoadObject: Unable to resolve object. Cached Path: %s, NetGUID: %s" ), *CacheObject->PathName.ToString(), *NetGUID.ToString() );
+		}
 	}
 
 	return NetGUID;
@@ -706,14 +713,6 @@ bool UPackageMapClient::ExportNetGUID( FNetworkGUID NetGUID, const UObject * Obj
 	check( !NetGUID.IsDefault() );
 	check( Object == NULL || ShouldSendFullPath( Object, NetGUID ) );
 
-#if !(UE_BUILD_SHIPPING || UE_BUILD_TEST)
-	// Warn if dynamic, non actor component, NetGUID is being exported.
-	if ( NetGUID.IsDynamic() && ( !PathName.IsEmpty() || ( Object && Cast<const UActorComponent>( Object ) == NULL ) ) )
-	{
-		UE_LOG( LogNetPackageMap, Warning, TEXT( "ExportNetGUID: Exporting Dynamic NetGUID. NetGUID: <%s> Object: %s" ), *NetGUID.ToString(), Object ? *Object->GetName() : *PathName );
-	}
-#endif
-
 	// Two passes are used to export this net guid:
 	// 1. Attempt to append this net guid to current bunch
 	// 2. If step 1 fails, append to fresh new bunch
@@ -1067,7 +1066,7 @@ void UPackageMapClient::HandleUnAssignedObject(const UObject* Obj)
 //--------------------------------------------------------------------
 
 /** Do we need to include the full path of this object for the client to resolve it? */
-bool UPackageMapClient::ShouldSendFullPath(const UObject* Object, const FNetworkGUID &NetGUID)
+bool UPackageMapClient::ShouldSendFullPath( const UObject* Object, const FNetworkGUID &NetGUID )
 {
 	if ( !Connection )
 	{
@@ -1092,6 +1091,8 @@ bool UPackageMapClient::ShouldSendFullPath(const UObject* Object, const FNetwork
 
 	if ( NetGUID.IsDefault() )
 	{
+		check( !IsNetGUIDAuthority() );
+		check( Object->IsNameStableForNetworking() );
 		return true;
 	}
 
@@ -1282,22 +1283,23 @@ public:
 
 void FNetGUIDCache::CleanReferences()
 {
-	// NOTE - Quick fix for some package map issues where client was cleaning up guids, which was causing issues.
-	// For now, we no longer clean these up.  The client can sometimes cleanup guids that the server doesn't which causes issues.
-	// The guid memory will leak for now until we find a better way.
-
-	// Free any dynamic guids that refer to objects that are dead
+	// Cleanup any guids that are dynamic
+	// We keep static guids since the client most likely hasn't GC'd the static objects the exact same way we have
+	// Since we can't rely on the the client clearing static guids in the same we we would, we need to just keep them
+	// The only issue with this, is this list can grow forever
+	// To fix this, we can ultimately re-use static guids, to at least bound the list size by the number of packages on disk
 	for ( auto It = ObjectLookup.CreateIterator(); It; ++It )
 	{
-		if ( It.Value().Object.IsValid() )
+		check( !It.Key().IsDefault() );
+		check( It.Key().IsStatic() != It.Key().IsDynamic() );
+
+		if ( It.Value().Object.IsValid() || It.Key().IsStatic() )
 		{
+			// Either the object is still around, or it's a fully stably named object (which we always keep those registered)
 			continue;
 		}
 
-		if ( It.Key().IsDynamic() )
-		{
-			It.RemoveCurrent();
-		}
+		It.RemoveCurrent();
 	}
 
 	for ( auto It = NetGUIDLookup.CreateIterator(); It; ++It )
@@ -1377,7 +1379,7 @@ bool FNetGUIDCache::SupportsObject( const UObject * Object )
 }
 
 /**
- *	Dynamic objects are actors that were spawned in the world at run time, and therefor cannot be
+ *	Dynamic objects are actors or sub-objects that were spawned in the world at run time, and therefor cannot be
  *	referenced with a path name to the client.
  */
 bool FNetGUIDCache::IsDynamicObject( const UObject * Object )
@@ -1547,6 +1549,7 @@ void FNetGUIDCache::RegisterNetGUID_Internal( const FNetworkGUID & NetGUID, cons
  */
 void FNetGUIDCache::RegisterNetGUID_Server( const FNetworkGUID & NetGUID, const UObject * Object )
 {
+	check( Object != NULL );
 	check( IsNetGUIDAuthority() );				// Only the server should call this
 	check( !Object->IsPendingKill() );
 	check( !NetGUID.IsDefault() );
@@ -1752,7 +1755,7 @@ UObject * FNetGUIDCache::GetObjectFromNetGUID( const FNetworkGUID & NetGUID )
 	if ( Object == NULL )
 	{
 		// If we can't find an object, and it's not pending, warn
-		if ( !CacheObjectPtr->bIgnoreWhenMissing )
+		if ( !CacheObjectPtr->bIgnoreWhenMissing && !CacheObjectPtr->bIsPending )
 		{
 			UE_LOG( LogNetPackageMap, Warning, TEXT( "GetObjectFromNetGUID: Failed to resolve path. Path: %s, NetGUID: %s" ), *CacheObjectPtr->PathName.ToString(), *NetGUID.ToString() );
 		}
@@ -1783,15 +1786,21 @@ bool FNetGUIDCache::ShouldIgnoreWhenMissing( const FNetworkGUID & NetGUID )
 		return false;		// Server never ignores when missing, always warns
 	}
 
+	if ( NetGUID.IsDynamic() )
+	{
+		return true;		// Ignore dynamic guids. Things like actors and sub-objects will come in when they are relevant to us, and we don't know when that is
+	}
+
 	const FNetGuidCacheObject * CacheObject = ObjectLookup.Find( NetGUID );
 
-	if ( CacheObject != NULL )
+	if ( CacheObject == NULL )
 	{
-		// Ignore warnings when we explicitly are told to, or if we're pending load (which could also be asynchronously loading)
-		return CacheObject->bIgnoreWhenMissing || CacheObject->bIsPending;
+		return false;		// If we haven't been told about this static guid before, we need to warn
 	}
 
-	return false;
+	// Ignore warnings when we explicitly are told to, or if we're pending load (which could also be asynchronously loading)
+	// We will explicitly be told to ignore missing guids for levels for example, since the client is in control of when it will actually load levels
+	return CacheObject->bIgnoreWhenMissing || CacheObject->bIsPending;
 }
 
 UObject * FNetGUIDCache::ResolvePath( const FString & PathName, UObject * ObjOuter, const bool bNoLoad )
