Fix for anim trail alignment crash.
Also fixed up Ribbons and changed the offending FVector4 to an array of floats so that this can't catch anyone out again.

diff --git a/Engine/Source/Runtime/Engine/Private/Particles/ParticleSystemRender.cpp b/Engine/Source/Runtime/Engine/Private/Particles/ParticleSystemRender.cpp
--- a/Engine/Source/Runtime/Engine/Private/Particles/ParticleSystemRender.cpp
+++ b/Engine/Source/Runtime/Engine/Private/Particles/ParticleSystemRender.cpp
@@ -6006,11 +6006,13 @@ int32 FDynamicRibbonEmitterData::FillVertexData(struct FAsyncBufferFillData& Dat
 
 				float InvCount = 1.0f / InterpCount;
 				float Diff = PrevTrailPayload->SpawnTime - TrailPayload->SpawnTime;
-
-				if (bFillDynamic == true)
+				
+				FVector4 CurrDynParam;
+				FVector4 PrevDynParam;
+				if (bFillDynamic)
 				{
-					CurrDynPayload = ((FEmitterDynamicParameterPayload*)((uint8*)(PackingParticle) + Source.DynamicParameterDataOffset));
-					PrevDynPayload = ((FEmitterDynamicParameterPayload*)((uint8*)(PrevParticle) + Source.DynamicParameterDataOffset));
+					GetDynamicValueFromPayload(Source.DynamicParameterDataOffset, *PackingParticle, CurrDynParam);
+					GetDynamicValueFromPayload(Source.DynamicParameterDataOffset, *PrevParticle, PrevDynParam);
 				}
 
 				FVector4 InterpDynamic(1.0f, 1.0f, 1.0f, 1.0f);
@@ -6021,9 +6023,9 @@ int32 FDynamicRibbonEmitterData::FillVertexData(struct FAsyncBufferFillData& Dat
 					FVector InterpUp = FMath::Lerp<FVector>(CurrUp, PrevUp, TimeStep);
 					FLinearColor InterpColor = FMath::Lerp<FLinearColor>(CurrColor, PrevColor, TimeStep);
 					float InterpSize = FMath::Lerp<float>(CurrSize, PrevSize, TimeStep);
-					if (CurrDynPayload && PrevDynPayload)
+					if (bFillDynamic)
 					{
-						InterpDynamic = FMath::Lerp<FVector4>(CurrDynPayload->DynamicParameterValue, PrevDynPayload->DynamicParameterValue, TimeStep);
+						InterpDynamic = FMath::Lerp<FVector4>(CurrDynParam, PrevDynParam, TimeStep);
 					}
 
 					if (bTextureTileDistance == true)	
@@ -6119,10 +6121,10 @@ int32 FDynamicRibbonEmitterData::FillVertexData(struct FAsyncBufferFillData& Dat
 					DynParamVertex = (FParticleBeamTrailVertexDynamicParameter*)(TempDynamicParamData);
 					if (CurrDynPayload != NULL)
 					{
-						DynParamVertex->DynamicValue[0] = CurrDynPayload->DynamicParameterValue.X;
-						DynParamVertex->DynamicValue[1] = CurrDynPayload->DynamicParameterValue.Y;
-						DynParamVertex->DynamicValue[2] = CurrDynPayload->DynamicParameterValue.Z;
-						DynParamVertex->DynamicValue[3] = CurrDynPayload->DynamicParameterValue.W;
+						DynParamVertex->DynamicValue[0] = CurrDynPayload->DynamicParameterValue[0];
+						DynParamVertex->DynamicValue[1] = CurrDynPayload->DynamicParameterValue[1];
+						DynParamVertex->DynamicValue[2] = CurrDynPayload->DynamicParameterValue[2];
+						DynParamVertex->DynamicValue[3] = CurrDynPayload->DynamicParameterValue[3];
 					}
 					else
 					{
@@ -6153,10 +6155,10 @@ int32 FDynamicRibbonEmitterData::FillVertexData(struct FAsyncBufferFillData& Dat
 					DynParamVertex = (FParticleBeamTrailVertexDynamicParameter*)(TempDynamicParamData);
 					if (CurrDynPayload != NULL)
 					{
-						DynParamVertex->DynamicValue[0] = CurrDynPayload->DynamicParameterValue.X;
-						DynParamVertex->DynamicValue[1] = CurrDynPayload->DynamicParameterValue.Y;
-						DynParamVertex->DynamicValue[2] = CurrDynPayload->DynamicParameterValue.Z;
-						DynParamVertex->DynamicValue[3] = CurrDynPayload->DynamicParameterValue.W;
+						DynParamVertex->DynamicValue[0] = CurrDynPayload->DynamicParameterValue[0];
+						DynParamVertex->DynamicValue[1] = CurrDynPayload->DynamicParameterValue[1];
+						DynParamVertex->DynamicValue[2] = CurrDynPayload->DynamicParameterValue[2];
+						DynParamVertex->DynamicValue[3] = CurrDynPayload->DynamicParameterValue[3];
 					}
 					else
 					{
@@ -6332,8 +6334,7 @@ struct FAnimTrailParticleRenderData
 			OutColor = Particle->Color;
 			if( OutDynamicParameters )
 			{
-				FEmitterDynamicParameterPayload* DynPayload = ((FEmitterDynamicParameterPayload*)((uint8*)(Particle) + Source.DynamicParameterDataOffset));
-				*OutDynamicParameters = DynPayload->DynamicParameterValue;
+				GetDynamicValueFromPayload(Source.DynamicParameterDataOffset, *Particle, *OutDynamicParameters);
 			}
 			return;
 		}
@@ -6348,8 +6349,7 @@ struct FAnimTrailParticleRenderData
 			OutColor = PrevParticle->Color;
 			if( OutDynamicParameters )
 			{
-				FEmitterDynamicParameterPayload* DynPayload = ((FEmitterDynamicParameterPayload*)((uint8*)(PrevParticle) + Source.DynamicParameterDataOffset));
-				*OutDynamicParameters = DynPayload->DynamicParameterValue;
+				GetDynamicValueFromPayload(Source.DynamicParameterDataOffset, *PrevParticle, *OutDynamicParameters);
 			}
 			return;
 		}
@@ -6366,7 +6366,7 @@ struct FAnimTrailParticleRenderData
 			float PrevPrevTiledU;
 			float PrevPrevSize;
 			FLinearColor PrevPrevColor;
-			FEmitterDynamicParameterPayload* PrevPrevDynPayload;
+			FBaseParticle* PrevPrevDynParamParticle;
 			if( PrevPrevParticle )
 			{
 				PrevPrevLocation = PrevPrevParticle->Location;
@@ -6375,7 +6375,7 @@ struct FAnimTrailParticleRenderData
 				PrevPrevTiledU = PrevPrevPayload->TiledU;
 				PrevPrevSize = PrevPrevParticle->Size.X * Source.Scale.X;
 				PrevPrevColor = PrevPrevParticle->Color;
-				PrevPrevDynPayload = ((FEmitterDynamicParameterPayload*)((uint8*)(PrevPrevParticle) + Source.DynamicParameterDataOffset));
+				PrevPrevDynParamParticle = PrevPrevParticle;
 			}
 			else
 			{
@@ -6385,7 +6385,7 @@ struct FAnimTrailParticleRenderData
 				PrevPrevTiledU = PrevPayload->TiledU;
 				PrevPrevSize = PrevParticle->Size.X * Source.Scale.X;
 				PrevPrevColor = PrevParticle->Color;
-				PrevPrevDynPayload = ((FEmitterDynamicParameterPayload*)((uint8*)(PrevParticle) + Source.DynamicParameterDataOffset));
+				PrevPrevDynParamParticle = PrevParticle;
 			}
 
 			FVector NextLocation;
@@ -6394,7 +6394,7 @@ struct FAnimTrailParticleRenderData
 			float NextTiledU;
 			float NextSize;
 			FLinearColor NextColor;
-			FEmitterDynamicParameterPayload* NextDynPayload;
+			FBaseParticle* NextDynParamParticle;
 			if( NextParticle )
 			{
 				NextLocation = NextParticle->Location;
@@ -6403,7 +6403,7 @@ struct FAnimTrailParticleRenderData
 				NextTiledU = NextPayload->TiledU;
 				NextSize = NextParticle->Size.X * Source.Scale.X;
 				NextColor = NextParticle->Color;
-				NextDynPayload = ((FEmitterDynamicParameterPayload*)((uint8*)(NextParticle) + Source.DynamicParameterDataOffset));
+				NextDynParamParticle = NextParticle;
 			}
 			else
 			{
@@ -6413,13 +6413,9 @@ struct FAnimTrailParticleRenderData
 				NextTiledU = Payload->TiledU;
 				NextSize = Particle->Size.X * Source.Scale.X;
 				NextColor = Particle->Color;
-				NextDynPayload = ((FEmitterDynamicParameterPayload*)((uint8*)(Particle) + Source.DynamicParameterDataOffset));
+				NextDynParamParticle = Particle;
 			}
 
-			const FEmitterDynamicParameterPayload* PrevDynPayload = ((FEmitterDynamicParameterPayload*)((uint8*)(PrevParticle) + Source.DynamicParameterDataOffset));
-			const FEmitterDynamicParameterPayload* CurrDynPayload = ((FEmitterDynamicParameterPayload*)((uint8*)(Particle) + Source.DynamicParameterDataOffset));
-
-
 			float NextT = 0.0f;
 			float CurrT = NextParticle ? Payload->InterpolationParameter : GCatmullRomEndParamOffset;
 			float PrevT = CurrT + PrevPayload->InterpolationParameter;
@@ -6438,10 +6434,20 @@ struct FAnimTrailParticleRenderData
 
 			if( OutDynamicParameters )
 			{
-				*OutDynamicParameters = FMath::CubicCRSplineInterpSafe(PrevPrevDynPayload->DynamicParameterValue,
-					PrevDynPayload->DynamicParameterValue,
-					CurrDynPayload->DynamicParameterValue,
-					NextDynPayload->DynamicParameterValue, PrevPrevT, PrevT, CurrT, NextT, T);
+				FVector4 PrevPrevDynamicParam;
+				FVector4 PrevDynamicParam;
+				FVector4 CurrDynamicParam;
+				FVector4 NextDynamicParam;
+				
+				GetDynamicValueFromPayload(Source.DynamicParameterDataOffset, *PrevPrevDynParamParticle, PrevPrevDynamicParam);
+				GetDynamicValueFromPayload(Source.DynamicParameterDataOffset, *PrevParticle, PrevDynamicParam);
+				GetDynamicValueFromPayload(Source.DynamicParameterDataOffset, *Particle, CurrDynamicParam);
+				GetDynamicValueFromPayload(Source.DynamicParameterDataOffset, *NextDynParamParticle, NextDynamicParam);
+
+				*OutDynamicParameters = FMath::CubicCRSplineInterpSafe(PrevPrevDynamicParam,
+					PrevDynamicParam,
+					CurrDynamicParam,
+					NextDynamicParam, PrevPrevT, PrevT, CurrT, NextT, T);
 			}
 
 			FVector Offset = (InterpDir * InterpLength);
diff --git a/Engine/Source/Runtime/Engine/Public/ParticleHelper.h b/Engine/Source/Runtime/Engine/Public/ParticleHelper.h
--- a/Engine/Source/Runtime/Engine/Public/ParticleHelper.h
+++ b/Engine/Source/Runtime/Engine/Public/ParticleHelper.h
@@ -976,9 +976,26 @@ struct FParticleEventInstancePayload
 struct FEmitterDynamicParameterPayload
 {
 	/** The float4 value to assign to the dynamic parameter. */
-	FVector4 DynamicParameterValue;
+	float DynamicParameterValue[4];
 };
 
+/**
+ *	Helper function for retrieving the dynamic payload of a particle.
+ *
+ *	@param	InDynamicPayloadOffset		The offset to the payload
+ *	@param	InParticle					The particle being processed
+ *	@param	OutDynamicData				The dynamic data from the particle
+ */
+FORCEINLINE void GetDynamicValueFromPayload(int32 InDynamicPayloadOffset, FBaseParticle& InParticle, FVector4& OutDynamicData)
+{
+	checkSlow(InDynamicPayloadOffset > 0);
+	FEmitterDynamicParameterPayload* DynPayload = ((FEmitterDynamicParameterPayload*)((uint8*)(&InParticle) + InDynamicPayloadOffset));
+	OutDynamicData.X = DynPayload->DynamicParameterValue[0];
+	OutDynamicData.Y = DynPayload->DynamicParameterValue[1];
+	OutDynamicData.Z = DynPayload->DynamicParameterValue[2];
+	OutDynamicData.W = DynPayload->DynamicParameterValue[3];
+}
+
 /** Camera offset particle payload */
 struct FCameraOffsetParticlePayload
 {
@@ -1540,23 +1557,6 @@ struct FDynamicSpriteEmitterDataBase : public FDynamicEmitterDataBase
 	 */
 	virtual void RenderDebug(FParticleSystemSceneProxy* Proxy, FPrimitiveDrawInterface* PDI, const FSceneView* View, bool bCrosses);
 
-	/**
-	 *	Helper function for retrieving the dynamic payload of a particle.
-	 *
-	 *	@param	InDynamicPayloadOffset		The offset to the payload
-	 *	@param	InParticle					The particle being processed
-	 *	@param	OutDynamicData				The dynamic data from the particle
-	 */
-	FORCEINLINE void GetDynamicValueFromPayload(int32 InDynamicPayloadOffset, FBaseParticle& InParticle, FVector4& OutDynamicData)
-	{
-		checkSlow(InDynamicPayloadOffset > 0);
-		FEmitterDynamicParameterPayload* DynPayload = ((FEmitterDynamicParameterPayload*)((uint8*)(&InParticle) + InDynamicPayloadOffset));
-		OutDynamicData.X = DynPayload->DynamicParameterValue.X;
-		OutDynamicData.Y = DynPayload->DynamicParameterValue.Y;
-		OutDynamicData.Z = DynPayload->DynamicParameterValue.Z;
-		OutDynamicData.W = DynPayload->DynamicParameterValue.W;
-	}
-
 	/**
 	 *	Fill index and vertex buffers. Often called from a different thread
 	 *
@@ -1905,23 +1905,6 @@ struct FDynamicMeshEmitterData : public FDynamicSpriteEmitterDataBase
 	 */
 	void GetInstanceData(void* InstanceData, void* DynamicParameterData, FParticleSystemSceneProxy* Proxy, const FSceneView* View);
 
-	/**
-	 *	Helper function for retrieving the dynamic payload of a particle.
-	 *
-	 *	@param	InDynamicPayloadOffset		The offset to the payload
-	 *	@param	InParticle					The particle being processed
-	 *	@param	OutDynamicData				The dynamic data from the particle
-	 */
-	FORCEINLINE void GetDynamicValueFromPayload(int32 InDynamicPayloadOffset, FBaseParticle& InParticle, FVector4& OutDynamicData)
-	{
-		checkSlow(InDynamicPayloadOffset > 0);
-		FEmitterDynamicParameterPayload* DynPayload = ((FEmitterDynamicParameterPayload*)((uint8*)(&InParticle) + InDynamicPayloadOffset));
-		OutDynamicData.X = DynPayload->DynamicParameterValue.X;
-		OutDynamicData.Y = DynPayload->DynamicParameterValue.Y;
-		OutDynamicData.Z = DynPayload->DynamicParameterValue.Z;
-		OutDynamicData.W = DynPayload->DynamicParameterValue.W;
-	}
-
 	/**
 	 *	Helper function for retrieving the particle transform.
 	 *
