Add OnComponentFracture and OnActorFracture events for destructible component and destructible actor

diff --git a/Engine/Source/Runtime/Engine/Classes/Components/DestructibleComponent.h b/Engine/Source/Runtime/Engine/Classes/Components/DestructibleComponent.h
--- a/Engine/Source/Runtime/Engine/Classes/Components/DestructibleComponent.h
+++ b/Engine/Source/Runtime/Engine/Classes/Components/DestructibleComponent.h
@@ -32,6 +32,9 @@ struct FDestructibleChunkInfo
 };
 #endif // WITH_PHYSX 
 
+/** Delegate for notification when fracture occurs */
+DECLARE_DYNAMIC_MULTICAST_DELEGATE_TwoParams(FComponentFractureSignature, const FVector &, HitPoint, const FVector &, HitDirection);
+
 /**
  *	This component holds the physics data for a DestructibleActor
  *
@@ -94,6 +97,11 @@ class ENGINE_API UDestructibleComponent : public USkinnedMeshComponent
 
 	UFUNCTION(BlueprintCallable, Category="Components|Destructible")
 	class UDestructibleMesh * GetDestructibleMesh();
+
+	/** Called when a component is touched */
+	UPROPERTY(BlueprintAssignable, Category = "Components|Destructible")
+	FComponentFractureSignature OnComponentFracture;
+
 public:
 #if WITH_APEX
 	/** The NxDestructibleActor instantated from an NxDestructibleAsset, which contains the runtime physical state. */
diff --git a/Engine/Source/Runtime/Engine/Classes/PhysicsEngine/DestructibleActor.h b/Engine/Source/Runtime/Engine/Classes/PhysicsEngine/DestructibleActor.h
--- a/Engine/Source/Runtime/Engine/Classes/PhysicsEngine/DestructibleActor.h
+++ b/Engine/Source/Runtime/Engine/Classes/PhysicsEngine/DestructibleActor.h
@@ -8,6 +8,9 @@
 
 class UDestructibleComponent;
 
+/** Delegate for notification when fracture occurs */
+DECLARE_DYNAMIC_MULTICAST_DELEGATE_TwoParams(FActorFractureSignature, const FVector &, HitPoint, const FVector &, HitDirection);
+
 UCLASS(MinimalAPI, hidecategories=(Input), showcategories=("Input|MouseInput", "Input|TouchInput"))
 class ADestructibleActor : public AActor
 {
@@ -22,6 +25,9 @@ class ADestructibleActor : public AActor
 	UPROPERTY(config, EditAnywhere, BlueprintReadWrite, Category=Navigation)
 	uint32 bAffectNavigation : 1;
 
+	UPROPERTY(BlueprintAssignable, Category = "Components|Destructible")
+	FActorFractureSignature OnActorFracture;
+
 	// Begin AActor interface.
 	virtual bool UpdateNavigationRelevancy() override { SetNavigationRelevancy(!!bAffectNavigation); return !!bAffectNavigation; }
 #if WITH_EDITOR
diff --git a/Engine/Source/Runtime/Engine/Private/DestructibleComponent.cpp b/Engine/Source/Runtime/Engine/Private/DestructibleComponent.cpp
--- a/Engine/Source/Runtime/Engine/Private/DestructibleComponent.cpp
+++ b/Engine/Source/Runtime/Engine/Private/DestructibleComponent.cpp
@@ -626,6 +626,15 @@ void UDestructibleComponent::SpawnFractureEffectsFromDamageEvent(const NxApexDam
 
 void UDestructibleComponent::OnDamageEvent(const NxApexDamageEventReportData& InDamageEvent)
 {
+	FVector HitPosition = P2UVector(InDamageEvent.hitPosition);
+	FVector HitDirection = P2UVector(InDamageEvent.hitDirection);
+
+	OnComponentFracture.Broadcast(HitPosition, HitDirection);
+	if (ADestructibleActor * DestructibleActor = Cast<ADestructibleActor>(GetOwner()))
+	{
+		DestructibleActor->OnActorFracture.Broadcast(HitPosition, HitDirection);
+	}
+
 	SpawnFractureEffectsFromDamageEvent(InDamageEvent);
 
 	// After receiving damage, no longer receive decals.
