Make it optional to preallocate net guids for sub-objects

diff --git a/Engine/Source/Runtime/Engine/Private/PackageMapClient.cpp b/Engine/Source/Runtime/Engine/Private/PackageMapClient.cpp
--- a/Engine/Source/Runtime/Engine/Private/PackageMapClient.cpp
+++ b/Engine/Source/Runtime/Engine/Private/PackageMapClient.cpp
@@ -20,10 +20,15 @@ static const int GUID_PACKET_ACKED		= -1;
  */
 static const int INTERNAL_LOAD_OBJECT_RECURSION_LIMIT = 16;
 
+// Comment this out to disable preallocation of sub-object netguids (will will export them as stable names rather than deterministically produce same list on both sides, which can be fragile)
+#define PRE_ALLOC_SUB_OBJECT_GUIDS
+
+#ifdef PRE_ALLOC_SUB_OBJECT_GUIDS
 #define GET_NET_GUID_INDEX( GUID )					( GUID.Value >> 1 )
 #define GET_NET_GUID_STATIC_MOD( GUID )				( GUID.Value & 1 )
 #define COMPOSE_NET_GUID( Index, IsStatic )			( ( ( Index ) << 1 ) | ( IsStatic ) )
 #define COMPOSE_RELATIVE_NET_GUID( GUID, Index )	( COMPOSE_NET_GUID( GET_NET_GUID_INDEX( GUID ) + Index, GET_NET_GUID_STATIC_MOD( GUID ) ) )
+#endif
 
 /*-----------------------------------------------------------------------------
 	UPackageMapClient implementation.
@@ -292,6 +297,7 @@ bool UPackageMapClient::SerializeNewActor(FArchive& Ar, class UActorChannel *Cha
 			}
 		}
 
+#ifdef PRE_ALLOC_SUB_OBJECT_GUIDS
 		//
 		// Serialize Initially replicating sub-objects here
 		//
@@ -381,6 +387,7 @@ bool UPackageMapClient::SerializeNewActor(FArchive& Ar, class UActorChannel *Cha
 				}
 			}
 		}
+#endif
 	}
 
 	UE_LOG(LogNetPackageMap, Log, TEXT("SerializeNewActor END: Finished Serializing Actor %s <%s> on Channel %d"), Actor ? *Actor->GetName() : TEXT("NULL"), *NetGUID.ToString(), Channel->ChIndex );
@@ -1084,10 +1091,19 @@ bool UPackageMapClient::ShouldSendFullPath( const UObject* Object, const FNetwor
 		return false;
 	}
 
+#ifdef PRE_ALLOC_SUB_OBJECT_GUIDS
 	if ( NetGUID.IsDynamic() )
 	{
 		return false;		// We only export fully stably named objects
 	}
+#else
+	if ( !Object->IsNameStableForNetworking() )
+	{
+		check( !NetGUID.IsDefault() );
+		check( NetGUID.IsDynamic() );
+		return false;		// We only export objects that have stable names
+	}
+#endif
 
 	if ( NetGUID.IsDefault() )
 	{
@@ -1426,6 +1442,7 @@ FNetworkGUID FNetGUIDCache::GetOrAssignNetGUID( const UObject * Object )
 	return AssignNewNetGUID_Server( Object );
 }
 
+#ifdef PRE_ALLOC_SUB_OBJECT_GUIDS
 static const AActor * GetOuterActor( const UObject * Object )
 {
 	Object = Object->GetOuter();
@@ -1444,6 +1461,7 @@ static const AActor * GetOuterActor( const UObject * Object )
 
 	return NULL;
 }
+#endif
 
 /**
  *	Generate a new NetGUID for this object and assign it.
@@ -1452,6 +1470,7 @@ FNetworkGUID FNetGUIDCache::AssignNewNetGUID_Server( const UObject * Object )
 {
 	check( IsNetGUIDAuthority() );
 
+#ifdef PRE_ALLOC_SUB_OBJECT_GUIDS
 	const AActor * Actor = Cast< const AActor >( Object );
 
 	// So what is going on here, is we want to make sure the network guid's of stably named sub-objects get 
@@ -1513,8 +1532,20 @@ FNetworkGUID FNetGUIDCache::AssignNewNetGUID_Server( const UObject * Object )
 			RegisterNetGUID_Server( SubobjNetGUID, SubObj );
 		}
 	}
+	return NewNetGuid;
+#else
+#define COMPOSE_NET_GUID( Index, IsStatic )	( ( ( Index ) << 1 ) | ( IsStatic ) )
+#define ALLOC_NEW_NET_GUID( IsStatic )		( COMPOSE_NET_GUID( ++UniqueNetIDs[ IsStatic ], IsStatic ) )
+
+	// Generate new NetGUID and assign it
+	const int32 IsStatic = IsDynamicObject( Object ) ? 0 : 1;
+
+	const FNetworkGUID NewNetGuid( ALLOC_NEW_NET_GUID( IsStatic ) );
+
+	RegisterNetGUID_Server( NewNetGuid, Object );
 
 	return NewNetGuid;
+#endif
 }
 
 void FNetGUIDCache::RegisterNetGUID_Internal( const FNetworkGUID & NetGUID, const FNetGuidCacheObject & CacheObject )
@@ -1640,7 +1671,9 @@ void FNetGUIDCache::RegisterNetGUIDFromPath_Client( const FNetworkGUID & NetGUID
 {
 	check( !IsNetGUIDAuthority() );		// Server never calls this locally
 	check( !NetGUID.IsDefault() );
+#ifdef PRE_ALLOC_SUB_OBJECT_GUIDS
 	check( !NetGUID.IsDynamic() );		// It only makes sense to do this for static guids
+#endif
 
 	const FNetGuidCacheObject * ExistingCacheObjectPtr = ObjectLookup.Find( NetGUID );
 
@@ -1716,6 +1749,7 @@ UObject * FNetGUIDCache::GetObjectFromNetGUID( const FNetworkGUID & NetGUID )
 		return NULL;
 	}
 
+#ifdef PRE_ALLOC_SUB_OBJECT_GUIDS
 	if ( NetGUID.IsDynamic() )
 	{
 		// Dynamic objects can never load from their names (name is not stable)
@@ -1728,6 +1762,14 @@ UObject * FNetGUIDCache::GetObjectFromNetGUID( const FNetworkGUID & NetGUID )
 		UE_LOG( LogNetPackageMap, Warning, TEXT( "GetObjectFromNetGUID: Static guid on client with no path. NetGUID: %s" ), *NetGUID.ToString() );
 		return NULL;
 	}
+#else
+	if ( CacheObjectPtr->PathName == NAME_None )
+	{
+		// If we don't have a path, assume this is a non stably named guid
+		check( NetGUID.IsDynamic() );
+		return NULL;
+	}
+#endif
 
 	// First, resolve the outer
 	UObject * ObjOuter = NULL;
@@ -1757,7 +1799,7 @@ UObject * FNetGUIDCache::GetObjectFromNetGUID( const FNetworkGUID & NetGUID )
 		// If we can't find an object, and it's not pending, warn
 		if ( !CacheObjectPtr->bIgnoreWhenMissing && !CacheObjectPtr->bIsPending )
 		{
-			UE_LOG( LogNetPackageMap, Warning, TEXT( "GetObjectFromNetGUID: Failed to resolve path. Path: %s, NetGUID: %s" ), *CacheObjectPtr->PathName.ToString(), *NetGUID.ToString() );
+			UE_LOG( LogNetPackageMap, Warning, TEXT( "GetObjectFromNetGUID: Failed to resolve path. Path: %s, Outer: %s, NetGUID: %s" ), *CacheObjectPtr->PathName.ToString(), ObjOuter != NULL ? *ObjOuter->GetPathName() : TEXT( "NULL" ), *NetGUID.ToString() );
 		}
 
 		return NULL;
