Fixed a case in level streaming when active level could be in queue for unloading

diff --git a/Engine/Source/Runtime/Engine/Classes/Engine/LevelStreaming.h b/Engine/Source/Runtime/Engine/Classes/Engine/LevelStreaming.h
--- a/Engine/Source/Runtime/Engine/Classes/Engine/LevelStreaming.h
+++ b/Engine/Source/Runtime/Engine/Classes/Engine/LevelStreaming.h
@@ -289,14 +289,9 @@ private:
 	bool RequestLevel(UWorld* PersistentWorld, bool bAllowLevelLoadRequests, bool bBlockOnLoad);
 	
 	/** Sets the value of LoadedLevel */
-	void SetLoadedLevel(class ULevel* Level)
-	{ 
-		// Pending level should be unloaded at this point
-		check(PendingUnloadLevel == nullptr);
-		PendingUnloadLevel = LoadedLevel;
-		LoadedLevel = Level;
-	}
-
+	void SetLoadedLevel(class ULevel* Level);
+	
+	/** Hide and queue for unloading previously used level */
 	void DiscardPendingUnloadLevel(UWorld* PersistentWorld);
 
 	/** 
@@ -310,7 +305,7 @@ private:
 	UPROPERTY(transient)
 	class ULevel* LoadedLevel;
 
-	/** Pointer to replaced loaded Level object */
+	/** Pointer to a Level object that was previously active and was replaced with a new LoadedLevel (for LOD switching) */
 	UPROPERTY(transient)
 	class ULevel* PendingUnloadLevel;
 	
diff --git a/Engine/Source/Runtime/Engine/Private/LevelStreaming.cpp b/Engine/Source/Runtime/Engine/Private/LevelStreaming.cpp
--- a/Engine/Source/Runtime/Engine/Private/LevelStreaming.cpp
+++ b/Engine/Source/Runtime/Engine/Private/LevelStreaming.cpp
@@ -265,6 +265,17 @@ FName ULevelStreaming::GetLODPackageNameToLoad() const
 	}
 }
 
+void ULevelStreaming::SetLoadedLevel(class ULevel* Level)
+{ 
+	// Pending level should be unloaded at this point
+	check(PendingUnloadLevel == nullptr);
+	PendingUnloadLevel = LoadedLevel;
+	LoadedLevel = Level;
+
+	// Cancel unloading for this level, in case it was queued for it
+	FLevelStreamingGCHelper::CancelUnloadRequest(LoadedLevel);
+}
+
 void ULevelStreaming::DiscardPendingUnloadLevel(UWorld* PersistentWorld)
 {
 	if (PendingUnloadLevel)
