Fixed case when Foliage painted on top of landscape was not updating during landscape sculptiing

diff --git a/Engine/Source/Runtime/Engine/Private/InstancedFoliage.cpp b/Engine/Source/Runtime/Engine/Private/InstancedFoliage.cpp
--- a/Engine/Source/Runtime/Engine/Private/InstancedFoliage.cpp
+++ b/Engine/Source/Runtime/Engine/Private/InstancedFoliage.cpp
@@ -827,7 +827,8 @@ void AInstancedFoliageActor::SnapInstancesForLandscape( class ULandscapeHeightfi
 					TArray<FHitResult> Results;
 					UWorld* World = InLandscapeComponent->GetWorld();
 					check(World);
-					World->LineTraceMulti(Results, Start, End, FCollisionQueryParams(TraceTag, true), FCollisionObjectQueryParams(FCollisionObjectQueryParams::InitType::AllObjects));
+					// Editor specific landscape heightfield uses ECC_Visibility collision channel
+					World->LineTraceMulti(Results, Start, End, FCollisionQueryParams(TraceTag, true), FCollisionObjectQueryParams(ECollisionChannel::ECC_Visibility));
 				
 					bool bFoundHit = false;
 					for( int32 HitIdx=0; HitIdx<Results.Num(); HitIdx++ )
