Foliage stores relative to component per-instance transformations in UInstancedStaticMeshComponent

diff --git a/Engine/Source/Editor/FoliageEdit/Private/FoliageEdMode.cpp b/Engine/Source/Editor/FoliageEdit/Private/FoliageEdMode.cpp
--- a/Engine/Source/Editor/FoliageEdit/Private/FoliageEdMode.cpp
+++ b/Engine/Source/Editor/FoliageEdit/Private/FoliageEdMode.cpp
@@ -378,7 +378,7 @@ void FEdModeFoliage::GetRandomVectorInBrush( FVector& OutStart, FVector& OutEnd
 
 bool CheckCollisionWithWorld(UInstancedFoliageSettings* MeshSettings, const FFoliageInstance& Inst, const FVector& HitNormal, const FVector& HitLocation,UWorld* InWorld)
 {
-	FMatrix InstTransform = Inst.GetInstanceTransform();
+	FMatrix InstTransform = Inst.GetInstanceWorldTransform().ToMatrixWithScale();
 	FVector LocalHit = InstTransform.InverseTransformPosition(HitLocation);
 
 	if ( MeshSettings->CollisionWithWorld )
@@ -480,7 +480,7 @@ struct FPotentialInstance
 		// Apply the Z offset in local space
 		if( FMath::Abs(Inst.ZOffset) > KINDA_SMALL_NUMBER )
 		{
-			Inst.Location = Inst.GetInstanceTransform().TransformPosition(FVector(0,0,Inst.ZOffset));
+			Inst.Location = Inst.GetInstanceWorldTransform().TransformPosition(FVector(0,0,Inst.ZOffset));
 		}
 
 		Inst.Base = HitComponent;
@@ -808,7 +808,7 @@ void FEdModeFoliage::ReapplyInstancesForBrush( UWorld* InWorld, AInstancedFoliag
 			if( FMath::Abs(Instance.ZOffset) > KINDA_SMALL_NUMBER )
 			{
 				MeshInfo.InstanceHash->RemoveInstance(Instance.Location, InstanceIndex);
-				Instance.Location = Instance.GetInstanceTransform().TransformPosition(FVector(0,0,-Instance.ZOffset));
+				Instance.Location = Instance.GetInstanceWorldTransform().TransformPosition(FVector(0,0,-Instance.ZOffset));
 				bReapplyLocation = true;
 			}
 
@@ -1042,7 +1042,7 @@ void FEdModeFoliage::ReapplyInstancesForBrush( UWorld* InWorld, AInstancedFoliag
 			if( bUpdated && FMath::Abs(Instance.ZOffset) > KINDA_SMALL_NUMBER )
 			{
 				// Reapply the Z offset in new local space
-				Instance.Location = Instance.GetInstanceTransform().TransformPosition(FVector(0,0,Instance.ZOffset));
+				Instance.Location = Instance.GetInstanceWorldTransform().TransformPosition(FVector(0,0,Instance.ZOffset));
 			}
 
 			// Re-add to the hash
diff --git a/Engine/Source/Runtime/Engine/Private/InstancedFoliage.cpp b/Engine/Source/Runtime/Engine/Private/InstancedFoliage.cpp
--- a/Engine/Source/Runtime/Engine/Private/InstancedFoliage.cpp
+++ b/Engine/Source/Runtime/Engine/Private/InstancedFoliage.cpp
@@ -225,12 +225,12 @@ void FFoliageMeshInfo::CheckValid()
 		{
 			int32 InstanceIdx = Indices(InstIdx);
 
-			FMatrix InstanceTransform = Instances(InstanceIdx).GetInstanceTransform();
-			FMatrix& CompTransform = Comp->PerInstanceSMData(InstIdx).Transform;
-			
-			if( InstanceTransform != CompTransform )
+			FTransform InstanceToWorldEd = Instances(InstanceIdx).GetInstanceTransform();
+			FTransform InstanceToWorldCluster = Comp->PerInstanceSMData(InstIdx).Transform * Comp->GetComponentToWorld();
+	
+			if( !InstanceToWorldEd.Equals(InstanceToWorldCluster) )
 			{
-				CompTransform = InstanceTransform;
+				Comp->PerInstanceSMData(InstIdx).Transform = InstanceToWorldEd.ToMatrixWithScale();
 				MismatchCount++;
 			}
 		}
@@ -288,14 +288,14 @@ void FFoliageMeshInfo::AddInstance( AInstancedFoliageActor* InIFA, UStaticMesh*
 	}
 
 	// Calculate transform for the instance
-	FMatrix InstanceTransform = InNewInstance.GetInstanceTransform();
+	FTransform InstanceToWorld = InNewInstance.GetInstanceWorldTransform();
 
 	if( BestCluster == NULL )
 	{
 		BestClusterIndex = InstanceClusters.Num();
 		BestCluster = new(InstanceClusters) FFoliageInstanceCluster(
 			ConstructObject<UInstancedStaticMeshComponent>(UInstancedStaticMeshComponent::StaticClass(),InIFA,NAME_None,RF_Transactional),
-			InMesh->GetBounds().TransformBy(InstanceTransform)
+			InMesh->GetBounds().TransformBy(InstanceToWorld)
 			);
 
 		// Make the instanced static mesh component movable so it doesn't get statically lit; see the comment below about CastShadow for more details
@@ -320,7 +320,7 @@ void FFoliageMeshInfo::AddInstance( AInstancedFoliageActor* InIFA, UStaticMesh*
 
 		BestCluster->ClusterComponent->BodyInstance.CopyBodyInstancePropertiesFrom(&Settings->BodyInstance);
 
-		BestCluster->ClusterComponent->SetRelativeTransform(FTransform::Identity);
+		BestCluster->ClusterComponent->SetWorldTransform(InstanceToWorld);
 		BestCluster->ClusterComponent->RegisterComponent();
 	}
 	else
@@ -335,7 +335,7 @@ void FFoliageMeshInfo::AddInstance( AInstancedFoliageActor* InIFA, UStaticMesh*
 	AddedInstance.ClusterIndex = BestClusterIndex;
 	
 	// Add the instance to the component
-	BestCluster->ClusterComponent->AddInstance(FTransform(InstanceTransform));
+	BestCluster->ClusterComponent->AddInstanceWorldSpace(InstanceToWorld);
 
 	if( BestCluster->ClusterComponent->SelectedInstances.Num() > 0 )
 	{
@@ -343,7 +343,7 @@ void FFoliageMeshInfo::AddInstance( AInstancedFoliageActor* InIFA, UStaticMesh*
 	}
 
 	// Update the bounds for the cluster
-	BestCluster->Bounds = BestCluster->ClusterComponent->CalcBounds(FTransform::Identity);
+	BestCluster->Bounds = BestCluster->ClusterComponent->CalcBounds(BestCluster->ClusterComponent->ComponentToWorld);
 
 	// Update PrimitiveComponent's culling distance taking into account the radius of the bounds, as
 	// it is based on the center of the component's bounds.
@@ -512,10 +512,12 @@ void FFoliageMeshInfo::PostUpdateInstances( class AInstancedFoliageActor* InIFA,
 			Cluster.ClusterComponent->MarkRenderStateDirty();
 
 			// Update bounds
-			FMatrix InstanceTransform = Instance.GetInstanceTransform();
-			Cluster.Bounds = Cluster.Bounds + Cluster.ClusterComponent->StaticMesh->GetBounds().TransformBy(InstanceTransform);
-			// Update transform in InstancedStaticMeshComponent
-			Cluster.ClusterComponent->PerInstanceSMData[ClusterInstanceDataIndex].Transform = InstanceTransform;
+			FTransform InstanceToWorld = Instance.GetInstanceWorldTransform();
+			Cluster.Bounds = Cluster.Bounds + Cluster.ClusterComponent->StaticMesh->GetBounds().TransformBy(InstanceToWorld);
+			// Update transform in InstancedStaticMeshComponent (in component local space)
+			// Transform from world space to local space
+			FTransform InstanceToLocal = InstanceToWorld.GetRelativeTransform(Cluster.ClusterComponent->ComponentToWorld);
+			Cluster.ClusterComponent->PerInstanceSMData[ClusterInstanceDataIndex].Transform = InstanceToLocal.ToMatrixWithScale();
 			Cluster.ClusterComponent->InvalidateLightingCache();
 		}
 		
@@ -809,7 +811,7 @@ void AInstancedFoliageActor::SnapInstancesForLandscape( class ULandscapeHeightfi
 
 				// Test location should remove any Z offset
 				FVector TestLocation = FMath::Abs(Instance.ZOffset) > KINDA_SMALL_NUMBER
-					?	(FVector)Instance.GetInstanceTransform().TransformPosition(FVector(0,0,-Instance.ZOffset))
+					?	(FVector)Instance.GetInstanceWorldTransform().TransformPosition(FVector(0,0,-Instance.ZOffset))
 					:	Instance.Location;
 
 				if( InInstanceBox.IsInside(TestLocation) )
@@ -855,7 +857,7 @@ void AInstancedFoliageActor::SnapInstancesForLandscape( class ULandscapeHeightfi
 								// Reapply the Z offset in local space
 								if( FMath::Abs(Instance.ZOffset) > KINDA_SMALL_NUMBER )
 								{
-									Instance.Location = Instance.GetInstanceTransform().TransformPosition(FVector(0,0,Instance.ZOffset));
+									Instance.Location = Instance.GetInstanceWorldTransform().TransformPosition(FVector(0,0,Instance.ZOffset));
 								}
 
 								// Todo: add do validation with other parameters such as max/min height etc.
@@ -874,11 +876,12 @@ void AInstancedFoliageActor::SnapInstancesForLandscape( class ULandscapeHeightfi
 									Cluster.ClusterComponent->MarkRenderStateDirty();
 									
 									// Update bounds
-									FMatrix InstanceTransform = Instance.GetInstanceTransform();
-									Cluster.Bounds = Cluster.Bounds + Cluster.ClusterComponent->StaticMesh->GetBounds().TransformBy(InstanceTransform);
+									FTransform InstanceToWorld = Instance.GetInstanceWorldTransform();
+									Cluster.Bounds = Cluster.Bounds + Cluster.ClusterComponent->StaticMesh->GetBounds().TransformBy(InstanceToWorld);
 									
 									// Update transform in InstancedStaticMeshComponent
-									Cluster.ClusterComponent->PerInstanceSMData[ClusterInstanceDataIndex].Transform = InstanceTransform;
+									FTransform InstanceToLocal = InstanceToWorld.GetRelativeTransform(Cluster.ClusterComponent->ComponentToWorld);
+									Cluster.ClusterComponent->PerInstanceSMData[ClusterInstanceDataIndex].Transform = InstanceToLocal.ToMatrixWithScale();
 									Cluster.ClusterComponent->InvalidateLightingCache();
 								}
 
@@ -957,10 +960,11 @@ void AInstancedFoliageActor::MoveInstancesForMovedComponent( class UActorCompone
 					Cluster.ClusterComponent->MarkRenderStateDirty();
 					bUpdatedInstances = true;
 					// Update bounds
-					FMatrix InstanceTransform = Instance.GetInstanceTransform();
-					Cluster.Bounds = Cluster.Bounds + Cluster.ClusterComponent->StaticMesh->GetBounds().TransformBy(InstanceTransform);
+					FTransform InstanceToWorld = Instance.GetInstanceWorldTransform();
+					Cluster.Bounds = Cluster.Bounds + Cluster.ClusterComponent->StaticMesh->GetBounds().TransformBy(InstanceToWorld);
 					// Update transform in InstancedStaticMeshComponent
-					Cluster.ClusterComponent->PerInstanceSMData[ClusterInstanceDataIndex].Transform = InstanceTransform;
+					FTransform InstanceToLocal = InstanceToWorld.GetRelativeTransform(Cluster.ClusterComponent->ComponentToWorld);
+					Cluster.ClusterComponent->PerInstanceSMData[ClusterInstanceDataIndex].Transform = InstanceToLocal.ToMatrixWithScale();
 					Cluster.ClusterComponent->InvalidateLightingCache();
 				}
 
@@ -1422,10 +1426,10 @@ void AInstancedFoliageActor::MapRebuild()
 					FFoliageInstance NewInstance = Instance;
 					NewInstance.ClusterIndex = -1;
 					
-					FMatrix InstanceTransform(Instance.GetInstanceTransform());
+					FTransform InstanceToWorld = Instance.GetInstanceWorldTransform();
 					FVector Down(-FVector::UpVector);
-					FVector Start(InstanceTransform.TransformPosition(FVector::UpVector));
-					FVector End(InstanceTransform.TransformPosition(Down));
+					FVector Start(InstanceToWorld.TransformPosition(FVector::UpVector));
+					FVector End(InstanceToWorld.TransformPosition(Down));
 
 					FHitResult Result;
 					bool bHit = World->LineTraceSingle(Result, Start, End, FCollisionQueryParams(true), FCollisionObjectQueryParams(ECC_WorldStatic));
diff --git a/Engine/Source/Runtime/Engine/Private/InstancedStaticMesh.cpp b/Engine/Source/Runtime/Engine/Private/InstancedStaticMesh.cpp
--- a/Engine/Source/Runtime/Engine/Private/InstancedStaticMesh.cpp
+++ b/Engine/Source/Runtime/Engine/Private/InstancedStaticMesh.cpp
@@ -886,22 +886,6 @@ public:
 			SetSelection_GameThread(true);
 		}
 #endif
-		FMatrix LocalToWorld = InComponent->GetComponentToWorld().ToMatrixWithScale();
-
-		// Copy and cache any per-instance data that we need
-		if( InComponent->PerInstanceSMData.Num() > 0 )
-		{
-			PerInstanceSMData.AddUninitialized( InComponent->PerInstanceSMData.Num() );
-			for( int32 CurInstanceIndex = 0; CurInstanceIndex < InComponent->PerInstanceSMData.Num(); ++CurInstanceIndex )
-			{
-				const FInstancedStaticMeshInstanceData& CurInstanceComponentData = InComponent->PerInstanceSMData[ CurInstanceIndex ];
-				FInstancedStaticMeshSceneProxyInstanceData& CurInstanceSceneProxyData = PerInstanceSMData[ CurInstanceIndex ];
-
-				// Cache instance -> world transform
-				CurInstanceSceneProxyData.InstanceToWorld = CurInstanceComponentData.Transform * LocalToWorld;
-			}
-		}
-
 		// Make sure all the materials are okay to be rendered as an instanced mesh.
 		for (int32 LODIndex = 0; LODIndex < LODs.Num(); LODIndex++)
 		{
@@ -987,10 +971,6 @@ public:
 	}
 
 private:
-
-	/** Array of per-instance static mesh rendering data for the scene proxy */
-	TArray< FInstancedStaticMeshSceneProxyInstanceData > PerInstanceSMData;
-
 	/** Per component render data */
 	FInstancedStaticMeshRenderData InstancedRenderData;
 
@@ -1478,14 +1458,8 @@ void UInstancedStaticMeshComponent::SetupNewInstanceData(FInstancedStaticMeshIns
 
 void UInstancedStaticMeshComponent::ApplyWorldOffset(const FVector& InOffset, bool bWorldShift)
 {
-	// Here we assume that per instance data holds absolute positions
-	// so we need to shift only instances locations but not the component location
-	for (int32 InstanceIdx = 0; InstanceIdx < PerInstanceSMData.Num(); ++InstanceIdx)
-	{
-		FInstancedStaticMeshInstanceData& Data = PerInstanceSMData[InstanceIdx];
-		Data.Transform = Data.Transform.ConcatTranslation(InOffset);
-	}
-
+	Super::ApplyWorldOffset(InOffset, bWorldShift);
+	// We still need to recreate buffers with instance data, as they holds absolute transformations right now
 	MarkRenderStateDirty();
 }
 
diff --git a/Engine/Source/Runtime/Engine/Public/InstancedFoliage.h b/Engine/Source/Runtime/Engine/Public/InstancedFoliage.h
--- a/Engine/Source/Runtime/Engine/Public/InstancedFoliage.h
+++ b/Engine/Source/Runtime/Engine/Public/InstancedFoliage.h
@@ -60,13 +60,9 @@ struct FFoliageInstance : public FFoliageInstancePlacementInfo
 
 	friend FArchive& operator<<( FArchive& Ar, FFoliageInstance& Instance );
 
-	FMatrix GetInstanceTransform() const
+	FTransform GetInstanceWorldTransform() const
 	{
-		FMatrix InstanceTransform = FScaleMatrix(DrawScale3D);
-		InstanceTransform *= FRotationMatrix(Rotation);
-		InstanceTransform *= FTranslationMatrix(Location);
-		check( !InstanceTransform.ContainsNaN() );
-		return InstanceTransform;
+		return FTransform(Rotation, Location, DrawScale3D);
 	}
 
 	void AlignToNormal( const FVector& InNormal, float AlignMaxAngle=0.f )
