Fixed the build paths option in the editor

When you turn off update navigation automatically in the editor and use
build paths, navigation generation always seems to fail.
In UNavigationSystem::Build() returns early because if NavDataClasses.Num() always returns 0
as it has not initiated yet. This is done in DoInitialSetup() which is called from SpawnMissingNavigationData()
All that needs to be done is reorder the logic slightly and
navigation generation will work again.

GitHub pull request #221
https://github.com/EpicGames/UnrealEngine/pull/221

diff --git a/Engine/Source/Runtime/Engine/Private/AI/Navigation/NavigationSystem.cpp b/Engine/Source/Runtime/Engine/Private/AI/Navigation/NavigationSystem.cpp
--- a/Engine/Source/Runtime/Engine/Private/AI/Navigation/NavigationSystem.cpp
+++ b/Engine/Source/Runtime/Engine/Private/AI/Navigation/NavigationSystem.cpp
@@ -2269,7 +2269,7 @@ void UNavigationSystem::OnNavigationBoundsUpdated(AVolume* NavVolume)
 
 void UNavigationSystem::Build()
 {
-	if (NavDataClasses.Num() == 0 || IsThereAnywhereToBuildNavigation() == false)
+	if (IsThereAnywhereToBuildNavigation() == false)
 	{
 		return;
 	}
@@ -2278,6 +2278,11 @@ void UNavigationSystem::Build()
 
 	SpawnMissingNavigationData();
 
+	if (NavDataClasses.Num() == 0)
+	{
+		return;
+	}
+
 	// make sure freshly created navigation instances are registered before we try to build them
 	ProcessRegistrationCandidates();
 
