Fixed an issue with the Distance CrossFade where no sound would play during fade out

Code was checking the wrong distance value to see whether there was any point in doing the fade out interpolation

diff --git a/Engine/Source/Runtime/Engine/Private/SoundNodeDistanceCrossFade.cpp b/Engine/Source/Runtime/Engine/Private/SoundNodeDistanceCrossFade.cpp
--- a/Engine/Source/Runtime/Engine/Private/SoundNodeDistanceCrossFade.cpp
+++ b/Engine/Source/Runtime/Engine/Private/SoundNodeDistanceCrossFade.cpp
@@ -71,7 +71,7 @@ void USoundNodeDistanceCrossFade::ParseNodes( FAudioDevice* AudioDevice, const U
 			// else if we are inside the FadeOut edge
 			else if( ( Distance >= FadeOutDistanceMin ) && ( Distance <= FadeOutDistanceMax ) )
 			{
-				VolumeToSet = (FadeInDistanceMax > 0.f ? CrossFadeInput[ ChildNodeIndex ].Volume * ( 1.0f - ( Distance - FadeOutDistanceMin ) / ( FadeOutDistanceMax - FadeOutDistanceMin ) ) : 0.f);
+				VolumeToSet = (FadeOutDistanceMax > 0.f ? CrossFadeInput[ ChildNodeIndex ].Volume * ( 1.0f - ( Distance - FadeOutDistanceMin ) / ( FadeOutDistanceMax - FadeOutDistanceMin ) ) : 0.f);
 				//UE_LOG(LogAudio, Log,  TEXT("     FadeOut.  Distance: %f,  VolumeToSet: %f"), Distance, VolumeToSet );
 			}
 			// else we are in between the fading edges of the CrossFaded sound and we should play the
