Exposed Controller::LineOfSightTo() to Blueprint

LineOfSightTo is now BlueprintCallable.
Update: also made it const.

GitHub pull request #211
https://github.com/EpicGames/UnrealEngine/pull/211

diff --git a/Engine/Source/Runtime/AIModule/Classes/AIController.h b/Engine/Source/Runtime/AIModule/Classes/AIController.h
--- a/Engine/Source/Runtime/AIModule/Classes/AIController.h
+++ b/Engine/Source/Runtime/AIModule/Classes/AIController.h
@@ -253,7 +253,7 @@ public:
 	 * @param bAlternateChecks used only in AIController implementation
 	 * @return true if controller's pawn can see Other actor.
 	 */
-	virtual bool LineOfSightTo(const class AActor* Other, FVector ViewPoint = FVector(ForceInit), bool bAlternateChecks = false) override;
+	virtual bool LineOfSightTo(const class AActor* Other, FVector ViewPoint = FVector(ForceInit), bool bAlternateChecks = false) const override;
 	// End AController Interface
 
 	/** Update direction AI is looking based on FocalPoint */
diff --git a/Engine/Source/Runtime/AIModule/Private/AIController.cpp b/Engine/Source/Runtime/AIModule/Private/AIController.cpp
--- a/Engine/Source/Runtime/AIModule/Private/AIController.cpp
+++ b/Engine/Source/Runtime/AIModule/Private/AIController.cpp
@@ -254,7 +254,7 @@ void AAIController::ClearFocus(uint8 InPriority)
 	}
 }
 
-bool AAIController::LineOfSightTo(const AActor* Other, FVector ViewPoint, bool bAlternateChecks)
+bool AAIController::LineOfSightTo(const AActor* Other, FVector ViewPoint, bool bAlternateChecks) const
 {
 	if( !Other )
 	{
diff --git a/Engine/Source/Runtime/Engine/Classes/GameFramework/Controller.h b/Engine/Source/Runtime/Engine/Classes/GameFramework/Controller.h
--- a/Engine/Source/Runtime/Engine/Classes/GameFramework/Controller.h
+++ b/Engine/Source/Runtime/Engine/Classes/GameFramework/Controller.h
@@ -138,7 +138,8 @@ public:
 	 * @param bAlternateChecks used only in AIController implementation
 	 * @return true if controller's pawn can see Other actor.
 	 */
-	virtual bool LineOfSightTo(const class AActor* Other, FVector ViewPoint = FVector(ForceInit), bool bAlternateChecks = false);
+	UFUNCTION(BlueprintCallable, Category="Controller")
+	virtual bool LineOfSightTo(const class AActor* Other, FVector ViewPoint = FVector(ForceInit), bool bAlternateChecks = false) const;
 
 	/** Replication Notification Callbacks */
 	UFUNCTION()
diff --git a/Engine/Source/Runtime/Engine/Private/Controller.cpp b/Engine/Source/Runtime/Engine/Private/Controller.cpp
--- a/Engine/Source/Runtime/Engine/Private/Controller.cpp
+++ b/Engine/Source/Runtime/Engine/Private/Controller.cpp
@@ -150,7 +150,7 @@ void AController::GetPlayerViewPoint( FVector& out_Location, FRotator& out_Rotat
 	GetActorEyesViewPoint( out_Location, out_Rotation);
 }
 
-bool AController::LineOfSightTo(const AActor* Other, FVector ViewPoint, bool bAlternateChecks)
+bool AController::LineOfSightTo(const AActor* Other, FVector ViewPoint, bool bAlternateChecks) const
 {
 	if( !Other )
 	{
