Remove redundant headers from Engine.h.
#proj core
#branch UE4
#summary Remove redundant headers from Engine.h. (NavDataGenerator.h NavLinkRenderingProxy.h NavigationModifier.h NavigationOctree.h ContentStreaming.h ComponentReregisterContext.h InstancedFoliage.h UnrealExporter.h EngineService.h NavigationSystemHelpers.h HardwareInfo.h EngineModule.h)

diff --git a/Engine/Plugins/Editor/SpeedTreeImporter/Source/SpeedTreeImporter/Private/SpeedTreeImportFactory.cpp b/Engine/Plugins/Editor/SpeedTreeImporter/Source/SpeedTreeImporter/Private/SpeedTreeImportFactory.cpp
--- a/Engine/Plugins/Editor/SpeedTreeImporter/Source/SpeedTreeImporter/Private/SpeedTreeImportFactory.cpp
+++ b/Engine/Plugins/Editor/SpeedTreeImporter/Source/SpeedTreeImporter/Private/SpeedTreeImportFactory.cpp
@@ -27,6 +27,7 @@
 #include "../Private/GeomFitUtils.h"
 #include "SpeedTreeWind.h"
 #include "RawMesh.h"
+#include "ComponentReregisterContext.h"
 
 #if WITH_SPEEDTREE
 
diff --git a/Engine/Plugins/Experimental/Paper2D/Source/Paper2D/Private/PaperSprite.cpp b/Engine/Plugins/Experimental/Paper2D/Source/Paper2D/Private/PaperSprite.cpp
--- a/Engine/Plugins/Experimental/Paper2D/Source/Paper2D/Private/PaperSprite.cpp
+++ b/Engine/Plugins/Experimental/Paper2D/Source/Paper2D/Private/PaperSprite.cpp
@@ -11,6 +11,7 @@
 #include "PaperSpriteAtlas.h"
 #include "GeomTools.h"
 #include "BitmapUtils.h"
+#include "ComponentReregisterContext.h"
 
 //////////////////////////////////////////////////////////////////////////
 // maf
diff --git a/Engine/Source/Developer/AssetTools/Private/AssetTools.cpp b/Engine/Source/Developer/AssetTools/Private/AssetTools.cpp
--- a/Engine/Source/Developer/AssetTools/Private/AssetTools.cpp
+++ b/Engine/Source/Developer/AssetTools/Private/AssetTools.cpp
@@ -13,6 +13,7 @@
 #include "EngineAnalytics.h"
 #include "IAnalyticsProvider.h"
 #include "MessageLog.h"
+#include "UnrealExporter.h"
 
 #define LOCTEXT_NAMESPACE "AssetTools"
 
diff --git a/Engine/Source/Developer/MeshBoneReduction/Private/MeshBoneReduction.cpp b/Engine/Source/Developer/MeshBoneReduction/Private/MeshBoneReduction.cpp
--- a/Engine/Source/Developer/MeshBoneReduction/Private/MeshBoneReduction.cpp
+++ b/Engine/Source/Developer/MeshBoneReduction/Private/MeshBoneReduction.cpp
@@ -3,6 +3,7 @@
 #include "Engine.h"
 #include "MeshBoneReduction.h"
 #include "Developer/MeshUtilities/Public/MeshUtilities.h"
+#include "ComponentReregisterContext.h"
 
 class FMeshBoneReductionModule : public IMeshBoneReductionModule
 {
diff --git a/Engine/Source/Developer/SimplygonMeshReduction/Private/SimplygonMeshReduction.cpp b/Engine/Source/Developer/SimplygonMeshReduction/Private/SimplygonMeshReduction.cpp
--- a/Engine/Source/Developer/SimplygonMeshReduction/Private/SimplygonMeshReduction.cpp
+++ b/Engine/Source/Developer/SimplygonMeshReduction/Private/SimplygonMeshReduction.cpp
@@ -6,6 +6,7 @@
 #include "SimplygonSDK.h"
 #include "MeshBoneReduction.h"
 #include "MaterialExportUtils.h"
+#include "ComponentReregisterContext.h"
 
 #define LOCTEXT_NAMESPACE "SimplygonMeshReduction"
 
diff --git a/Engine/Source/Developer/iOS/MetalShaderFormat/Private/MetalShaderFormat.cpp b/Engine/Source/Developer/iOS/MetalShaderFormat/Private/MetalShaderFormat.cpp
--- a/Engine/Source/Developer/iOS/MetalShaderFormat/Private/MetalShaderFormat.cpp
+++ b/Engine/Source/Developer/iOS/MetalShaderFormat/Private/MetalShaderFormat.cpp
@@ -6,7 +6,7 @@
 #include "ModuleManager.h"
 #include "TargetPlatform.h"   
 #include "ShaderCore.h"
-
+#include "hlslcc.h"
 
 static FName NAME_SF_METAL(TEXT("SF_METAL"));
 
diff --git a/Engine/Source/Editor/DestructibleMeshEditor/Private/SDestructibleMeshEditorViewport.cpp b/Engine/Source/Editor/DestructibleMeshEditor/Private/SDestructibleMeshEditorViewport.cpp
--- a/Engine/Source/Editor/DestructibleMeshEditor/Private/SDestructibleMeshEditorViewport.cpp
+++ b/Engine/Source/Editor/DestructibleMeshEditor/Private/SDestructibleMeshEditorViewport.cpp
@@ -9,6 +9,7 @@
 #include "ApexDestructibleAssetImport.h"
 #include "DesktopPlatformModule.h"
 #include "FbxImporter.h"
+#include "ComponentReregisterContext.h"
 
 DEFINE_LOG_CATEGORY_STATIC(LogDestructibleMeshEditor, Log, All);
 
diff --git a/Engine/Source/Editor/FoliageEdit/Private/FoliageEdMode.h b/Engine/Source/Editor/FoliageEdit/Private/FoliageEdMode.h
--- a/Engine/Source/Editor/FoliageEdit/Private/FoliageEdMode.h
+++ b/Engine/Source/Editor/FoliageEdit/Private/FoliageEdMode.h
@@ -10,6 +10,13 @@
 
 #pragma once
 
+#include "InstancedFoliage.h"
+
+//
+// Forward declarations.
+//
+class UStaticMesh;
+
 // Current user settings in Foliage UI
 struct FFoliageUISettings
 {
diff --git a/Engine/Source/Editor/Kismet/Private/InstancedStaticMeshSCSEditorCustomization.cpp b/Engine/Source/Editor/Kismet/Private/InstancedStaticMeshSCSEditorCustomization.cpp
--- a/Engine/Source/Editor/Kismet/Private/InstancedStaticMeshSCSEditorCustomization.cpp
+++ b/Engine/Source/Editor/Kismet/Private/InstancedStaticMeshSCSEditorCustomization.cpp
@@ -3,6 +3,7 @@
 #include "BlueprintEditorPrivatePCH.h"
 #include "InstancedStaticMeshSCSEditorCustomization.h"
 #include "BlueprintEditorModule.h"
+#include "InstancedFoliage.h"
 
 TSharedRef<ISCSEditorCustomization> FInstancedStaticMeshSCSEditorCustomization::MakeInstance(TSharedRef< IBlueprintEditor > InBlueprintEditor)
 {
diff --git a/Engine/Source/Editor/Kismet/Private/SCSEditorViewportClient.cpp b/Engine/Source/Editor/Kismet/Private/SCSEditorViewportClient.cpp
--- a/Engine/Source/Editor/Kismet/Private/SCSEditorViewportClient.cpp
+++ b/Engine/Source/Editor/Kismet/Private/SCSEditorViewportClient.cpp
@@ -15,6 +15,7 @@
 #include "Editor/UnrealEd/Public/Kismet2/ComponentEditorUtils.h"
 #include "ISCSEditorCustomization.h"
 #include "ComponentVisualizer.h"
+#include "InstancedFoliage.h"
 
 DEFINE_LOG_CATEGORY_STATIC(LogSCSEditorViewport, Log, All);
 
diff --git a/Engine/Source/Editor/Kismet/Private/SSCSEditor.cpp b/Engine/Source/Editor/Kismet/Private/SSCSEditor.cpp
--- a/Engine/Source/Editor/Kismet/Private/SSCSEditor.cpp
+++ b/Engine/Source/Editor/Kismet/Private/SSCSEditor.cpp
@@ -24,6 +24,7 @@
 
 #include "IDocumentation.h"
 #include "Kismet2NameValidators.h"
+#include "UnrealExporter.h"
 
 #define LOCTEXT_NAMESPACE "SSCSEditor"
 
diff --git a/Engine/Source/Editor/LandscapeEditor/Private/LandscapeEdMode.cpp b/Engine/Source/Editor/LandscapeEditor/Private/LandscapeEdMode.cpp
--- a/Engine/Source/Editor/LandscapeEditor/Private/LandscapeEdMode.cpp
+++ b/Engine/Source/Editor/LandscapeEditor/Private/LandscapeEdMode.cpp
@@ -30,6 +30,7 @@
 #include "Landscape/LandscapeMaterialInstanceConstant.h"
 #include "Landscape/LandscapeSplinesComponent.h"
 #include "Foliage/InstancedFoliageActor.h"
+#include "ComponentReregisterContext.h"
 
 #define LOCTEXT_NAMESPACE "Landscape"
 
diff --git a/Engine/Source/Editor/MaterialEditor/Private/SMaterialEditorViewport.cpp b/Engine/Source/Editor/MaterialEditor/Private/SMaterialEditorViewport.cpp
--- a/Engine/Source/Editor/MaterialEditor/Private/SMaterialEditorViewport.cpp
+++ b/Engine/Source/Editor/MaterialEditor/Private/SMaterialEditorViewport.cpp
@@ -10,6 +10,7 @@
 #include "AssetToolsModule.h"
 #include "MaterialEditor.h"
 #include "SMaterialEditorViewportToolBar.h"
+#include "ComponentReregisterContext.h"
 
 /** Viewport Client for the preview viewport */
 class FMaterialEditorViewportClient : public FEditorViewportClient
diff --git a/Engine/Source/Editor/MeshPaint/Private/MeshPaintEdMode.cpp b/Engine/Source/Editor/MeshPaint/Private/MeshPaintEdMode.cpp
--- a/Engine/Source/Editor/MeshPaint/Private/MeshPaintEdMode.cpp
+++ b/Engine/Source/Editor/MeshPaint/Private/MeshPaintEdMode.cpp
@@ -27,7 +27,7 @@
 
 #include "Runtime/Engine/Classes/PhysicsEngine/BodySetup.h"
 #include "SMeshPaint.h"
-
+#include "ComponentReregisterContext.h"
 
 #define LOCTEXT_NAMESPACE "MeshPaint_Mode"
 
diff --git a/Engine/Source/Editor/NiagaraEditor/Private/NiagaraCompiler.cpp b/Engine/Source/Editor/NiagaraEditor/Private/NiagaraCompiler.cpp
--- a/Engine/Source/Editor/NiagaraEditor/Private/NiagaraCompiler.cpp
+++ b/Engine/Source/Editor/NiagaraEditor/Private/NiagaraCompiler.cpp
@@ -6,6 +6,7 @@
 #include "CompilerResultsLog.h"
 #include "EdGraphUtilities.h"
 #include "VectorVM.h"
+#include "ComponentReregisterContext.h"
 
 DEFINE_LOG_CATEGORY_STATIC(LogNiagaraCompiler,All,All);
 
diff --git a/Engine/Source/Editor/Persona/Private/SSkeletonTree.cpp b/Engine/Source/Editor/Persona/Private/SSkeletonTree.cpp
--- a/Engine/Source/Editor/Persona/Private/SSkeletonTree.cpp
+++ b/Engine/Source/Editor/Persona/Private/SSkeletonTree.cpp
@@ -27,6 +27,7 @@
 
 #include "Factories.h"
 #include "Developer/MeshUtilities/Public/MeshUtilities.h"
+#include "UnrealExporter.h"
 
 #define LOCTEXT_NAMESPACE "SSkeletonTree"
 
diff --git a/Engine/Source/Editor/StaticMeshEditor/Private/SStaticMeshEditorViewport.cpp b/Engine/Source/Editor/StaticMeshEditor/Private/SStaticMeshEditorViewport.cpp
--- a/Engine/Source/Editor/StaticMeshEditor/Private/SStaticMeshEditorViewport.cpp
+++ b/Engine/Source/Editor/StaticMeshEditor/Private/SStaticMeshEditorViewport.cpp
@@ -21,6 +21,7 @@
 #include "StaticMeshEditorViewportClient.h"
 
 #include "../Private/GeomFitUtils.h"
+#include "ComponentReregisterContext.h"
 
 #if WITH_PHYSX
 #include "Editor/UnrealEd/Private/EditorPhysXSupport.h"
diff --git a/Engine/Source/Editor/UnrealEd/Private/Commandlets/DiffAssetsCommandlet.cpp b/Engine/Source/Editor/UnrealEd/Private/Commandlets/DiffAssetsCommandlet.cpp
--- a/Engine/Source/Editor/UnrealEd/Private/Commandlets/DiffAssetsCommandlet.cpp
+++ b/Engine/Source/Editor/UnrealEd/Private/Commandlets/DiffAssetsCommandlet.cpp
@@ -5,6 +5,7 @@
  */
 
 #include "UnrealEd.h"
+#include "UnrealExporter.h"
 
 DEFINE_LOG_CATEGORY_STATIC(LogDiffAssetsCommandlet, Log, All);
 
diff --git a/Engine/Source/Editor/UnrealEd/Private/EdGraphUtilities.cpp b/Engine/Source/Editor/UnrealEd/Private/EdGraphUtilities.cpp
--- a/Engine/Source/Editor/UnrealEd/Private/EdGraphUtilities.cpp
+++ b/Engine/Source/Editor/UnrealEd/Private/EdGraphUtilities.cpp
@@ -5,6 +5,7 @@
 #include "EdGraphUtilities.h"
 #include "Factories.h"
 #include "Editor/UnrealEd/Public/Kismet2/BlueprintEditorUtils.h"
+#include "UnrealExporter.h"
 
 /////////////////////////////////////////////////////
 // FGraphObjectTextFactory
diff --git a/Engine/Source/Editor/UnrealEd/Private/Editor.cpp b/Engine/Source/Editor/UnrealEd/Private/Editor.cpp
--- a/Engine/Source/Editor/UnrealEd/Private/Editor.cpp
+++ b/Engine/Source/Editor/UnrealEd/Private/Editor.cpp
@@ -48,6 +48,9 @@
 
 #include "Runtime/Engine/Public/Slate/SceneViewport.h"
 
+#include "ComponentReregisterContext.h"
+#include "EngineModule.h"
+
 #if PLATFORM_WINDOWS
 // For WAVEFORMATEXTENSIBLE
 	#include "AllowWindowsPlatformTypes.h"
diff --git a/Engine/Source/Editor/UnrealEd/Private/EditorActor.cpp b/Engine/Source/Editor/UnrealEd/Private/EditorActor.cpp
--- a/Engine/Source/Editor/UnrealEd/Private/EditorActor.cpp
+++ b/Engine/Source/Editor/UnrealEd/Private/EditorActor.cpp
@@ -17,6 +17,7 @@
 #include "Animation/SkeletalMeshActor.h"
 #include "Particles/Emitter.h"
 #include "Particles/ParticleSystemComponent.h"
+#include "UnrealExporter.h"
 
 #define LOCTEXT_NAMESPACE "UnrealEd.EditorActor"
 
diff --git a/Engine/Source/Editor/UnrealEd/Private/EditorExporters.cpp b/Engine/Source/Editor/UnrealEd/Private/EditorExporters.cpp
--- a/Engine/Source/Editor/UnrealEd/Private/EditorExporters.cpp
+++ b/Engine/Source/Editor/UnrealEd/Private/EditorExporters.cpp
@@ -19,6 +19,8 @@
 #include "ImageUtils.h"
 #include "Foliage/InstancedFoliageActor.h"
 #include "Landscape/Landscape.h"
+#include "UnrealExporter.h"
+#include "EngineModule.h"
 
 DEFINE_LOG_CATEGORY_STATIC(LogEditorExporters, Log, All);
 
diff --git a/Engine/Source/Editor/UnrealEd/Private/EditorLevelUtils.cpp b/Engine/Source/Editor/UnrealEd/Private/EditorLevelUtils.cpp
--- a/Engine/Source/Editor/UnrealEd/Private/EditorLevelUtils.cpp
+++ b/Engine/Source/Editor/UnrealEd/Private/EditorLevelUtils.cpp
@@ -16,6 +16,7 @@
 #include "ScopedTransaction.h"
 #include "ActorEditorUtils.h"
 #include "Editor/Levels/Public/LevelEdMode.h"
+#include "ContentStreaming.h"
 
 #include "Runtime/AssetRegistry/Public/AssetRegistryModule.h"
 
diff --git a/Engine/Source/Editor/UnrealEd/Private/EditorServer.cpp b/Engine/Source/Editor/UnrealEd/Private/EditorServer.cpp
--- a/Engine/Source/Editor/UnrealEd/Private/EditorServer.cpp
+++ b/Engine/Source/Editor/UnrealEd/Private/EditorServer.cpp
@@ -48,6 +48,8 @@
 #include "Particles/Emitter.h"
 #include "Particles/ParticleSystemComponent.h"
 
+#include "ComponentReregisterContext.h"
+
 DEFINE_LOG_CATEGORY_STATIC(LogEditorServer, Log, All);
 
 /** Used for the "tagsounds" and "checksounds" commands only			*/
diff --git a/Engine/Source/Editor/UnrealEd/Private/EditorViewportClient.cpp b/Engine/Source/Editor/UnrealEd/Private/EditorViewportClient.cpp
--- a/Engine/Source/Editor/UnrealEd/Private/EditorViewportClient.cpp
+++ b/Engine/Source/Editor/UnrealEd/Private/EditorViewportClient.cpp
@@ -14,7 +14,7 @@
 #include "AnalyticsEventAttribute.h"
 #include "IAnalyticsProvider.h"
 #include "Matinee/MatineeActor.h"
-
+#include "EngineModule.h"
 
 #define LOCTEXT_NAMESPACE "EditorViewportClient"
 
diff --git a/Engine/Source/Editor/UnrealEd/Private/Factories/ApexDestructibleAssetImport.cpp b/Engine/Source/Editor/UnrealEd/Private/Factories/ApexDestructibleAssetImport.cpp
--- a/Engine/Source/Editor/UnrealEd/Private/Factories/ApexDestructibleAssetImport.cpp
+++ b/Engine/Source/Editor/UnrealEd/Private/Factories/ApexDestructibleAssetImport.cpp
@@ -22,6 +22,7 @@ DEFINE_LOG_CATEGORY_STATIC(LogApexDestructibleAssetImport, Log, All);
 #include "EditorPhysXSupport.h"
 #include "ApexDestructibleAssetImport.h"
 #include "Developer/MeshUtilities/Public/MeshUtilities.h"
+#include "ComponentReregisterContext.h"
 
 #if WITH_APEX
 
diff --git a/Engine/Source/Editor/UnrealEd/Private/Factories/VectorFieldFactory.cpp b/Engine/Source/Editor/UnrealEd/Private/Factories/VectorFieldFactory.cpp
--- a/Engine/Source/Editor/UnrealEd/Private/Factories/VectorFieldFactory.cpp
+++ b/Engine/Source/Editor/UnrealEd/Private/Factories/VectorFieldFactory.cpp
@@ -7,6 +7,7 @@
 #include "UnrealEd.h"
 
 #include "Particles/ParticleSystemComponent.h"
+#include "ComponentReregisterContext.h"
 
 DEFINE_LOG_CATEGORY_STATIC(LogVectorFieldFactory, Log, All);
 
diff --git a/Engine/Source/Editor/UnrealEd/Private/Fbx/FbxMaterialImport.cpp b/Engine/Source/Editor/UnrealEd/Private/Fbx/FbxMaterialImport.cpp
--- a/Engine/Source/Editor/UnrealEd/Private/Fbx/FbxMaterialImport.cpp
+++ b/Engine/Source/Editor/UnrealEd/Private/Fbx/FbxMaterialImport.cpp
@@ -40,6 +40,7 @@
 #include "PackageTools.h"
 #include "AssetRegistryModule.h"
 #include "AssetToolsModule.h"
+#include "ComponentReregisterContext.h"
 
 DEFINE_LOG_CATEGORY_STATIC(LogFbxMaterialImport, Log, All);
 
diff --git a/Engine/Source/Editor/UnrealEd/Private/Fbx/FbxSkeletalMeshImport.cpp b/Engine/Source/Editor/UnrealEd/Private/Fbx/FbxSkeletalMeshImport.cpp
--- a/Engine/Source/Editor/UnrealEd/Private/Fbx/FbxSkeletalMeshImport.cpp
+++ b/Engine/Source/Editor/UnrealEd/Private/Fbx/FbxSkeletalMeshImport.cpp
@@ -51,6 +51,7 @@
 #include "Developer/MeshUtilities/Public/MeshUtilities.h"
 
 #include "MessageLogModule.h"
+#include "ComponentReregisterContext.h"
 
 #define LOCTEXT_NAMESPACE "FBXImpoter"
 
diff --git a/Engine/Source/Editor/UnrealEd/Private/FbxMeshUtils.cpp b/Engine/Source/Editor/UnrealEd/Private/FbxMeshUtils.cpp
--- a/Engine/Source/Editor/UnrealEd/Private/FbxMeshUtils.cpp
+++ b/Engine/Source/Editor/UnrealEd/Private/FbxMeshUtils.cpp
@@ -14,6 +14,8 @@
 	#include "ApexClothingUtils.h"
 #endif // #if WITH_APEX_CLOTHING
 
+#include "ComponentReregisterContext.h"
+
 DEFINE_LOG_CATEGORY_STATIC(LogExportMeshUtils, Log, All);
 
 #define LOCTEXT_NAMESPACE "FbxMeshUtil"
diff --git a/Engine/Source/Editor/UnrealEd/Private/Kismet2/BlueprintEditorUtils.cpp b/Engine/Source/Editor/UnrealEd/Private/Kismet2/BlueprintEditorUtils.cpp
--- a/Engine/Source/Editor/UnrealEd/Private/Kismet2/BlueprintEditorUtils.cpp
+++ b/Engine/Source/Editor/UnrealEd/Private/Kismet2/BlueprintEditorUtils.cpp
@@ -33,6 +33,7 @@
 #include "ObjectEditorUtils.h"
 #include "ToolkitManager.h"
 #include "Runtime/Engine/Classes/Engine/UserDefinedStruct.h"
+#include "UnrealExporter.h"
 
 #define LOCTEXT_NAMESPACE "Blueprint"
 
diff --git a/Engine/Source/Editor/UnrealEd/Private/LODUtilities.cpp b/Engine/Source/Editor/UnrealEd/Private/LODUtilities.cpp
--- a/Engine/Source/Editor/UnrealEd/Private/LODUtilities.cpp
+++ b/Engine/Source/Editor/UnrealEd/Private/LODUtilities.cpp
@@ -17,6 +17,8 @@
 	#include "ApexClothingUtils.h"
 #endif // #if WITH_APEX_CLOTHING
 
+#include "ComponentReregisterContext.h"
+
 void FLODUtilities::RemoveLOD(FSkeletalMeshUpdateContext& UpdateContext, int32 DesiredLOD )
 {
 	USkeletalMesh* SkeletalMesh = UpdateContext.SkeletalMesh;
diff --git a/Engine/Source/Editor/UnrealEd/Private/Lightmass/Lightmass.cpp b/Engine/Source/Editor/UnrealEd/Private/Lightmass/Lightmass.cpp
--- a/Engine/Source/Editor/UnrealEd/Private/Lightmass/Lightmass.cpp
+++ b/Engine/Source/Editor/UnrealEd/Private/Lightmass/Lightmass.cpp
@@ -17,6 +17,7 @@
 #include "Runtime/Engine/Classes/Components/SplineMeshComponent.h"
 #include "Lightmass/PrecomputedVisibilityVolume.h"
 #include "Lightmass/PrecomputedVisibilityOverrideVolume.h"
+#include "ComponentReregisterContext.h"
 
 extern FSwarmDebugOptions GSwarmDebugOptions;
 
diff --git a/Engine/Source/Editor/UnrealEd/Private/ObjectTools.cpp b/Engine/Source/Editor/UnrealEd/Private/ObjectTools.cpp
--- a/Engine/Source/Editor/UnrealEd/Private/ObjectTools.cpp
+++ b/Engine/Source/Editor/UnrealEd/Private/ObjectTools.cpp
@@ -28,6 +28,8 @@
 #include "DesktopPlatformModule.h"
 #include "LevelUtils.h"
 #include "ConsolidateWindow.h"
+#include "ComponentReregisterContext.h"
+
 DEFINE_LOG_CATEGORY_STATIC(LogObjectTools, Log, All);
 
 namespace ObjectTools
diff --git a/Engine/Source/Editor/UnrealEd/Private/PackageTools.cpp b/Engine/Source/Editor/UnrealEd/Private/PackageTools.cpp
--- a/Engine/Source/Editor/UnrealEd/Private/PackageTools.cpp
+++ b/Engine/Source/Editor/UnrealEd/Private/PackageTools.cpp
@@ -14,6 +14,7 @@
 #include "DesktopPlatformModule.h"
 #include "MainFrame.h"
 #include "MessageLog.h"
+#include "ComponentReregisterContext.h"
 
 #define LOCTEXT_NAMESPACE "PackageTools"
 
diff --git a/Engine/Source/Editor/UnrealEd/Private/PreviewMaterial.cpp b/Engine/Source/Editor/UnrealEd/Private/PreviewMaterial.cpp
--- a/Engine/Source/Editor/UnrealEd/Private/PreviewMaterial.cpp
+++ b/Engine/Source/Editor/UnrealEd/Private/PreviewMaterial.cpp
@@ -2,6 +2,7 @@
 
 #include "UnrealEd.h"
 #include "Editor/MaterialEditor/Public/MaterialEditorModule.h"
+#include "ComponentReregisterContext.h"
 
 /**
  * Class for rendering the material on the preview mesh in the Material Editor
diff --git a/Engine/Source/Editor/UnrealEd/Private/StaticLightingSystem/StaticLightingSystem.cpp b/Engine/Source/Editor/UnrealEd/Private/StaticLightingSystem/StaticLightingSystem.cpp
--- a/Engine/Source/Editor/UnrealEd/Private/StaticLightingSystem/StaticLightingSystem.cpp
+++ b/Engine/Source/Editor/UnrealEd/Private/StaticLightingSystem/StaticLightingSystem.cpp
@@ -17,6 +17,7 @@ FSwarmDebugOptions GSwarmDebugOptions;
 #include "PrecomputedLightVolume.h"
 #include "LevelUtils.h"
 #include "CrashTracker.h"
+#include "EngineModule.h"
 
 DEFINE_LOG_CATEGORY(LogStaticLightingSystem);
 
diff --git a/Engine/Source/Editor/UnrealEd/Private/ThumbnailHelpers.cpp b/Engine/Source/Editor/UnrealEd/Private/ThumbnailHelpers.cpp
--- a/Engine/Source/Editor/UnrealEd/Private/ThumbnailHelpers.cpp
+++ b/Engine/Source/Editor/UnrealEd/Private/ThumbnailHelpers.cpp
@@ -6,6 +6,8 @@
 #include "Animation/SkeletalMeshActor.h"
 #include "Particles/ParticleSystem.h"
 #include "Particles/ParticleSystemComponent.h"
+#include "ContentStreaming.h"
+
 /*
 ***************************************************************
   FThumbnailPreviewScene
diff --git a/Engine/Source/Editor/UnrealEd/Private/ThumbnailRendering/BlueprintThumbnailRenderer.cpp b/Engine/Source/Editor/UnrealEd/Private/ThumbnailRendering/BlueprintThumbnailRenderer.cpp
--- a/Engine/Source/Editor/UnrealEd/Private/ThumbnailRendering/BlueprintThumbnailRenderer.cpp
+++ b/Engine/Source/Editor/UnrealEd/Private/ThumbnailRendering/BlueprintThumbnailRenderer.cpp
@@ -4,6 +4,7 @@
 
 // FPreviewScene derived helpers for rendering
 #include "ThumbnailHelpers.h"
+#include "EngineModule.h"
 
 UBlueprintThumbnailRenderer::UBlueprintThumbnailRenderer(const class FPostConstructInitializeProperties& PCIP)
 	: Super(PCIP)
diff --git a/Engine/Source/Editor/UnrealEd/Private/ThumbnailRendering/LevelThumbnailRenderer.cpp b/Engine/Source/Editor/UnrealEd/Private/ThumbnailRendering/LevelThumbnailRenderer.cpp
--- a/Engine/Source/Editor/UnrealEd/Private/ThumbnailRendering/LevelThumbnailRenderer.cpp
+++ b/Engine/Source/Editor/UnrealEd/Private/ThumbnailRendering/LevelThumbnailRenderer.cpp
@@ -1,6 +1,7 @@
 // Copyright 1998-2014 Epic Games, Inc. All Rights Reserved.
 
 #include "UnrealEd.h"
+#include "EngineModule.h"
 
 ULevelThumbnailRenderer::ULevelThumbnailRenderer(const class FPostConstructInitializeProperties& PCIP)
 	: Super(PCIP)
diff --git a/Engine/Source/Editor/UnrealEd/Private/ThumbnailRendering/MaterialFunctionThumbnailRenderer.cpp b/Engine/Source/Editor/UnrealEd/Private/ThumbnailRendering/MaterialFunctionThumbnailRenderer.cpp
--- a/Engine/Source/Editor/UnrealEd/Private/ThumbnailRendering/MaterialFunctionThumbnailRenderer.cpp
+++ b/Engine/Source/Editor/UnrealEd/Private/ThumbnailRendering/MaterialFunctionThumbnailRenderer.cpp
@@ -6,6 +6,7 @@
 
 // FPreviewScene derived helpers for rendering
 #include "ThumbnailHelpers.h"
+#include "EngineModule.h"
 
 UMaterialFunctionThumbnailRenderer::UMaterialFunctionThumbnailRenderer(const class FPostConstructInitializeProperties& PCIP)
 	: Super(PCIP)
diff --git a/Engine/Source/Editor/UnrealEd/Private/ThumbnailRendering/MaterialInstanceThumbnailRenderer.cpp b/Engine/Source/Editor/UnrealEd/Private/ThumbnailRendering/MaterialInstanceThumbnailRenderer.cpp
--- a/Engine/Source/Editor/UnrealEd/Private/ThumbnailRendering/MaterialInstanceThumbnailRenderer.cpp
+++ b/Engine/Source/Editor/UnrealEd/Private/ThumbnailRendering/MaterialInstanceThumbnailRenderer.cpp
@@ -4,6 +4,7 @@
 
 // FPreviewScene derived helpers for rendering
 #include "ThumbnailHelpers.h"
+#include "EngineModule.h"
 
 UMaterialInstanceThumbnailRenderer::UMaterialInstanceThumbnailRenderer(const class FPostConstructInitializeProperties& PCIP)
 	: Super(PCIP)
diff --git a/Engine/Source/Editor/UnrealEd/Private/ThumbnailRendering/ParticleSystemThumbnailRenderer.cpp b/Engine/Source/Editor/UnrealEd/Private/ThumbnailRendering/ParticleSystemThumbnailRenderer.cpp
--- a/Engine/Source/Editor/UnrealEd/Private/ThumbnailRendering/ParticleSystemThumbnailRenderer.cpp
+++ b/Engine/Source/Editor/UnrealEd/Private/ThumbnailRendering/ParticleSystemThumbnailRenderer.cpp
@@ -6,6 +6,7 @@
 #include "ThumbnailHelpers.h"
 
 #include "Particles/ParticleSystem.h"
+#include "EngineModule.h"
 
 UParticleSystemThumbnailRenderer::UParticleSystemThumbnailRenderer(const class FPostConstructInitializeProperties& PCIP)
 	: Super(PCIP)
diff --git a/Engine/Source/Editor/UnrealEd/Private/ThumbnailRendering/SkeletalMeshThumbnailRenderer.cpp b/Engine/Source/Editor/UnrealEd/Private/ThumbnailRendering/SkeletalMeshThumbnailRenderer.cpp
--- a/Engine/Source/Editor/UnrealEd/Private/ThumbnailRendering/SkeletalMeshThumbnailRenderer.cpp
+++ b/Engine/Source/Editor/UnrealEd/Private/ThumbnailRendering/SkeletalMeshThumbnailRenderer.cpp
@@ -4,6 +4,7 @@
 
 // FPreviewScene derived helpers for rendering
 #include "ThumbnailHelpers.h"
+#include "EngineModule.h"
 
 USkeletalMeshThumbnailRenderer::USkeletalMeshThumbnailRenderer(const class FPostConstructInitializeProperties& PCIP)
 	: Super(PCIP)
diff --git a/Engine/Source/Editor/UnrealEd/Private/ThumbnailRendering/StaticMeshThumbnailRenderer.cpp b/Engine/Source/Editor/UnrealEd/Private/ThumbnailRendering/StaticMeshThumbnailRenderer.cpp
--- a/Engine/Source/Editor/UnrealEd/Private/ThumbnailRendering/StaticMeshThumbnailRenderer.cpp
+++ b/Engine/Source/Editor/UnrealEd/Private/ThumbnailRendering/StaticMeshThumbnailRenderer.cpp
@@ -4,6 +4,7 @@
 
 // FPreviewScene derived helpers for rendering
 #include "ThumbnailHelpers.h"
+#include "EngineModule.h"
 
 UStaticMeshThumbnailRenderer::UStaticMeshThumbnailRenderer(const class FPostConstructInitializeProperties& PCIP)
 	: Super(PCIP)
diff --git a/Engine/Source/Editor/UnrealEd/Private/UnrealEdSrv.cpp b/Engine/Source/Editor/UnrealEd/Private/UnrealEdSrv.cpp
--- a/Engine/Source/Editor/UnrealEd/Private/UnrealEdSrv.cpp
+++ b/Engine/Source/Editor/UnrealEd/Private/UnrealEdSrv.cpp
@@ -31,6 +31,7 @@
 #include "Landscape/LandscapeGizmoActiveActor.h"
 #include "Landscape/LandscapeComponent.h"
 #include "Landscape/LandscapeHeightfieldCollisionComponent.h"
+#include "ComponentReregisterContext.h"
 
 DEFINE_LOG_CATEGORY_STATIC(LogUnrealEdSrv, Log, All);
 
diff --git a/Engine/Source/Editor/UnrealEd/Public/DragAndDrop/ExportTextDragDropOp.h b/Engine/Source/Editor/UnrealEd/Public/DragAndDrop/ExportTextDragDropOp.h
--- a/Engine/Source/Editor/UnrealEd/Public/DragAndDrop/ExportTextDragDropOp.h
+++ b/Engine/Source/Editor/UnrealEd/Public/DragAndDrop/ExportTextDragDropOp.h
@@ -1,7 +1,7 @@
 // Copyright 1998-2014 Epic Games, Inc. All Rights Reserved.
 
-
 #pragma once
+#include "UnrealExporter.h"
 
 class FSelectedActorExportObjectInnerContext : public FExportObjectInnerContext
 {
diff --git a/Engine/Source/Runtime/Engine/Private/AI/Navigation/NavCollision.cpp b/Engine/Source/Runtime/Engine/Private/AI/Navigation/NavCollision.cpp
--- a/Engine/Source/Runtime/Engine/Private/AI/Navigation/NavCollision.cpp
+++ b/Engine/Source/Runtime/Engine/Private/AI/Navigation/NavCollision.cpp
@@ -2,6 +2,7 @@
 
 #include "EnginePrivate.h"
 #include "AI/Navigation/NavCollision.h"
+#include "AI/NavigationSystemHelpers.h"
 #include "NavigationModifier.h"
 #include "../Private/AI/Navigation/RecastNavMeshGenerator.h"
 #include "DerivedDataPluginInterface.h"
diff --git a/Engine/Source/Runtime/Engine/Private/AI/Navigation/RecastNavMeshGenerator.h b/Engine/Source/Runtime/Engine/Private/AI/Navigation/RecastNavMeshGenerator.h
--- a/Engine/Source/Runtime/Engine/Private/AI/Navigation/RecastNavMeshGenerator.h
+++ b/Engine/Source/Runtime/Engine/Private/AI/Navigation/RecastNavMeshGenerator.h
@@ -7,6 +7,7 @@
 #include "DetourCommon.h"
 #include "DetourNavMesh.h"
 #include "DetourNavMeshBuilder.h"
+#include "AI/NavDataGenerator.h"
 
 #define MAX_VERTS_PER_POLY	6
 
diff --git a/Engine/Source/Runtime/Engine/Private/Audio.cpp b/Engine/Source/Runtime/Engine/Private/Audio.cpp
--- a/Engine/Source/Runtime/Engine/Private/Audio.cpp
+++ b/Engine/Source/Runtime/Engine/Private/Audio.cpp
@@ -15,7 +15,7 @@
 #include "TargetPlatform.h"
 #include "EngineAnalytics.h"
 #include "Runtime/Analytics/Analytics/Public/Interfaces/IAnalyticsProvider.h"
-
+#include "ContentStreaming.h"
 
 DEFINE_LOG_CATEGORY(LogAudio);
 
diff --git a/Engine/Source/Runtime/Engine/Private/InstancedStaticMesh.cpp b/Engine/Source/Runtime/Engine/Private/InstancedStaticMesh.cpp
--- a/Engine/Source/Runtime/Engine/Private/InstancedStaticMesh.cpp
+++ b/Engine/Source/Runtime/Engine/Private/InstancedStaticMesh.cpp
@@ -33,6 +33,7 @@
 #include "StaticMeshLight.h"
 #include "SpeedTreeWind.h"
 #include "ComponentInstanceDataCache.h"
+#include "InstancedFoliage.h"
 
 #if WITH_PHYSX
 #include "PhysicsEngine/PhysXSupport.h"
diff --git a/Engine/Source/Runtime/Engine/Private/Interpolation.cpp b/Engine/Source/Runtime/Engine/Private/Interpolation.cpp
--- a/Engine/Source/Runtime/Engine/Private/Interpolation.cpp
+++ b/Engine/Source/Runtime/Engine/Private/Interpolation.cpp
@@ -70,6 +70,7 @@
 #include "Particles/Emitter.h"
 #include "Particles/ParticleSystemComponent.h"
 #include "Particles/ParticleSystemReplay.h"
+#include "ContentStreaming.h"
 
 #if WITH_EDITOR
 #include "UnrealEd.h"
diff --git a/Engine/Source/Runtime/Engine/Private/Landscape/LandscapeEdit.cpp b/Engine/Source/Runtime/Engine/Private/Landscape/LandscapeEdit.cpp
--- a/Engine/Source/Runtime/Engine/Private/Landscape/LandscapeEdit.cpp
+++ b/Engine/Source/Runtime/Engine/Private/Landscape/LandscapeEdit.cpp
@@ -30,6 +30,7 @@ LandscapeEdit.cpp: Landscape editing
 #include "ScopedTransaction.h"
 #include "ImageWrapper.h"
 #endif
+#include "ComponentReregisterContext.h"
 
 DEFINE_LOG_CATEGORY(LogLandscape);
 
diff --git a/Engine/Source/Runtime/Engine/Private/Level.cpp b/Engine/Source/Runtime/Engine/Private/Level.cpp
--- a/Engine/Source/Runtime/Engine/Private/Level.cpp
+++ b/Engine/Source/Runtime/Engine/Private/Level.cpp
@@ -23,6 +23,7 @@ Level.cpp: Level-related functions
 #include "Runtime/Engine/Classes/MovieScene/RuntimeMovieScenePlayerInterface.h"
 #include "LevelUtils.h"
 #include "TargetPlatform.h"
+#include "ContentStreaming.h"
 
 DEFINE_LOG_CATEGORY(LogLevel);
 
diff --git a/Engine/Source/Runtime/Engine/Private/Materials/Material.cpp b/Engine/Source/Runtime/Engine/Private/Materials/Material.cpp
--- a/Engine/Source/Runtime/Engine/Private/Materials/Material.cpp
+++ b/Engine/Source/Runtime/Engine/Private/Materials/Material.cpp
@@ -29,6 +29,7 @@
 #include "UObjectAnnotation.h"
 #include "MaterialCompiler.h"
 #include "TargetPlatform.h"
+#include "ComponentReregisterContext.h"
 
 #if WITH_EDITOR
 #include "UnrealEd.h"
diff --git a/Engine/Source/Runtime/Engine/Private/Materials/MaterialShared.cpp b/Engine/Source/Runtime/Engine/Private/Materials/MaterialShared.cpp
--- a/Engine/Source/Runtime/Engine/Private/Materials/MaterialShared.cpp
+++ b/Engine/Source/Runtime/Engine/Private/Materials/MaterialShared.cpp
@@ -15,6 +15,7 @@
 #include "HLSLMaterialTranslator.h"
 #include "MaterialUniformExpressions.h"
 #include "Developer/TargetPlatform/Public/TargetPlatform.h"
+#include "ComponentReregisterContext.h"
 #include "EngineModule.h"
 
 DEFINE_LOG_CATEGORY(LogMaterial);
diff --git a/Engine/Source/Runtime/Engine/Private/PlayerController.cpp b/Engine/Source/Runtime/Engine/Private/PlayerController.cpp
--- a/Engine/Source/Runtime/Engine/Private/PlayerController.cpp
+++ b/Engine/Source/Runtime/Engine/Private/PlayerController.cpp
@@ -19,6 +19,7 @@
 #include "Engine/ActorChannel.h"
 #include "GameFramework/SpectatorPawn.h"
 #include "GameFramework/HUD.h"
+#include "ContentStreaming.h"
 
 DEFINE_LOG_CATEGORY_STATIC(LogPlayerController, Log, All);
 
diff --git a/Engine/Source/Runtime/Engine/Private/SkeletalMesh.cpp b/Engine/Source/Runtime/Engine/Private/SkeletalMesh.cpp
--- a/Engine/Source/Runtime/Engine/Private/SkeletalMesh.cpp
+++ b/Engine/Source/Runtime/Engine/Private/SkeletalMesh.cpp
@@ -24,6 +24,7 @@
 #include "Animation/AnimSingleNodeInstance.h"
 #include "Animation/AnimBlueprint.h"
 #include "Animation/AnimBlueprintGeneratedClass.h"
+#include "ComponentReregisterContext.h"
 
 #if WITH_EDITOR
 #include "MeshUtilities.h"
diff --git a/Engine/Source/Runtime/Engine/Private/Texture.cpp b/Engine/Source/Runtime/Engine/Private/Texture.cpp
--- a/Engine/Source/Runtime/Engine/Private/Texture.cpp
+++ b/Engine/Source/Runtime/Engine/Private/Texture.cpp
@@ -8,6 +8,7 @@
 #include "EditorSupportDelegates.h"
 #include "TargetPlatform.h"
 #include "ImageWrapper.h"
+#include "ContentStreaming.h"
 
 #if WITH_EDITORONLY_DATA
 #include "DDSLoader.h"
diff --git a/Engine/Source/Runtime/Engine/Private/Texture2D.cpp b/Engine/Source/Runtime/Engine/Private/Texture2D.cpp
--- a/Engine/Source/Runtime/Engine/Private/Texture2D.cpp
+++ b/Engine/Source/Runtime/Engine/Private/Texture2D.cpp
@@ -16,6 +16,7 @@
 #endif
 
 #include "TargetPlatform.h"
+#include "ContentStreaming.h"
 
 UTexture2D::UTexture2D(const class FPostConstructInitializeProperties& PCIP)
 	: Super(PCIP)
diff --git a/Engine/Source/Runtime/Engine/Private/UnrealClient.cpp b/Engine/Source/Runtime/Engine/Private/UnrealClient.cpp
--- a/Engine/Source/Runtime/Engine/Private/UnrealClient.cpp
+++ b/Engine/Source/Runtime/Engine/Private/UnrealClient.cpp
@@ -14,6 +14,8 @@
 #include "SNotificationList.h"
 #include "Engine/PostProcessVolume.h"
 #include "EngineModule.h"
+#include "EngineModule.h"
+#include "ContentStreaming.h"
 
 DEFINE_LOG_CATEGORY_STATIC(LogClient, Log, All);
 
diff --git a/Engine/Source/Runtime/Engine/Private/UnrealEngine.cpp b/Engine/Source/Runtime/Engine/Private/UnrealEngine.cpp
--- a/Engine/Source/Runtime/Engine/Private/UnrealEngine.cpp
+++ b/Engine/Source/Runtime/Engine/Private/UnrealEngine.cpp
@@ -73,6 +73,9 @@
 
 #include "HardwareInfo.h"
 #include "EngineModule.h"
+#include "UnrealExporter.h"
+#include "ComponentReregisterContext.h"
+#include "ContentStreaming.h"
 
 DEFINE_LOG_CATEGORY_STATIC(LogEngine, Log, All);
 
diff --git a/Engine/Source/Runtime/Engine/Private/World.cpp b/Engine/Source/Runtime/Engine/Private/World.cpp
--- a/Engine/Source/Runtime/Engine/Private/World.cpp
+++ b/Engine/Source/Runtime/Engine/Private/World.cpp
@@ -43,6 +43,9 @@
 #include "Landscape/LandscapeInfo.h"
 #include "Particles/ParticleEventManager.h"
 
+#include "EngineModule.h"
+#include "ContentStreaming.h"
+
 DEFINE_LOG_CATEGORY_STATIC(LogWorld, Log, All);
 DEFINE_LOG_CATEGORY(LogSpawn);
 
diff --git a/Engine/Source/Runtime/Engine/Public/ComponentReregisterContext.h b/Engine/Source/Runtime/Engine/Public/ComponentReregisterContext.h
--- a/Engine/Source/Runtime/Engine/Public/ComponentReregisterContext.h
+++ b/Engine/Source/Runtime/Engine/Public/ComponentReregisterContext.h
@@ -2,6 +2,8 @@
 
 #pragma once
 
+#include "ContentStreaming.h"
+
 ENGINE_API DECLARE_LOG_CATEGORY_EXTERN(LogActorComponent, Log, All);
 
 /**
diff --git a/Engine/Source/Runtime/Engine/Public/ContentStreaming.h b/Engine/Source/Runtime/Engine/Public/ContentStreaming.h
--- a/Engine/Source/Runtime/Engine/Public/ContentStreaming.h
+++ b/Engine/Source/Runtime/Engine/Public/ContentStreaming.h
@@ -62,6 +62,7 @@ class FAsyncTextureStreaming;
 class UActorComponent;
 class UPrimitiveComponent;
 class AActor;
+class UTexture2D;
 template<typename T>
 class FAsyncTask;
 struct FStreamingManagerTexture;
diff --git a/Engine/Source/Runtime/Engine/Public/Engine.h b/Engine/Source/Runtime/Engine/Public/Engine.h
--- a/Engine/Source/Runtime/Engine/Public/Engine.h
+++ b/Engine/Source/Runtime/Engine/Public/Engine.h
@@ -377,24 +377,8 @@
 #include "SystemSettings.h"					// Scalability options.
 #include "SceneManagement.h"				// Scene management.
 
-#include "NavDataGenerator.h"
-#include "NavLinkRenderingProxy.h"
-#include "NavigationModifier.h"
-#include "NavigationOctree.h"
-#include "ContentStreaming.h"				// Content streaming class definitions.
-#include "ComponentReregisterContext.h"	
 #include "DrawDebugHelpers.h"
 #include "UnrealEngine.h"					// Unreal engine helpers.
 #include "Canvas.h"							// Canvas.
 #include "EngineUtils.h"
-#include "InstancedFoliage.h"				// Instanced foliage.
-#include "UnrealExporter.h"					// Exporter definition.
 #include "TimerManager.h"					// Game play timers
-#include "EngineService.h"
-#include "NavigationSystemHelpers.h"
-#include "HardwareInfo.h"
-
-#include "EngineModule.h"
-
-
-
diff --git a/Engine/Source/Runtime/Engine/Public/Tests/AutomationCommon.h b/Engine/Source/Runtime/Engine/Public/Tests/AutomationCommon.h
--- a/Engine/Source/Runtime/Engine/Public/Tests/AutomationCommon.h
+++ b/Engine/Source/Runtime/Engine/Public/Tests/AutomationCommon.h
@@ -2,6 +2,8 @@
 
 #pragma once
 
+#include "HardwareInfo.h"
+
 ///////////////////////////////////////////////////////////////////////
 // Common Latent commands which are used across test type. I.e. Engine, Network, etc...
 
diff --git a/Engine/Source/Runtime/GameplayAbilities/Private/AttributeSet.cpp b/Engine/Source/Runtime/GameplayAbilities/Private/AttributeSet.cpp
--- a/Engine/Source/Runtime/GameplayAbilities/Private/AttributeSet.cpp
+++ b/Engine/Source/Runtime/GameplayAbilities/Private/AttributeSet.cpp
@@ -2,6 +2,9 @@
 
 #include "AbilitySystemPrivatePCH.h"
 #include "AttributeSet.h"
+#include "HAL/OutputDevices.h"
+
+#include "ComponentReregisterContext.h"	
 
 void FGameplayAttribute::SetNumericValueChecked(const float NewValue, class UAttributeSet* Dest) const
 {
diff --git a/Engine/Source/Runtime/Launch/Private/LaunchEngineLoop.cpp b/Engine/Source/Runtime/Launch/Private/LaunchEngineLoop.cpp
--- a/Engine/Source/Runtime/Launch/Private/LaunchEngineLoop.cpp
+++ b/Engine/Source/Runtime/Launch/Private/LaunchEngineLoop.cpp
@@ -42,6 +42,8 @@
 	#include "PlatformFeatures.h"
 	#include "DeviceProfiles/DeviceProfileManager.h"
 	#include "Commandlets/Commandlet.h"
+	#include "EngineService.h"
+	#include "ContentStreaming.h"
 
 #if !UE_SERVER
 	#include "ISlateRHIRendererModule.h"
diff --git a/Engine/Source/Runtime/Launch/Public/LaunchEngineLoop.h b/Engine/Source/Runtime/Launch/Public/LaunchEngineLoop.h
--- a/Engine/Source/Runtime/Launch/Public/LaunchEngineLoop.h
+++ b/Engine/Source/Runtime/Launch/Public/LaunchEngineLoop.h
@@ -9,6 +9,7 @@
 #if WITH_ENGINE
 #include "ISessionService.h"
 #endif
+class FEngineService;
 
 /**
  * Implements the main engine loop.	
diff --git a/Engine/Source/Runtime/Renderer/Private/RendererScene.cpp b/Engine/Source/Runtime/Renderer/Private/RendererScene.cpp
--- a/Engine/Source/Runtime/Renderer/Private/RendererScene.cpp
+++ b/Engine/Source/Runtime/Renderer/Private/RendererScene.cpp
@@ -9,7 +9,7 @@
 #include "ShaderCompiler.h"
 #include "StaticMeshResources.h"
 #include "DistanceFieldSurfaceCacheLighting.h"
-
+#include "EngineModule.h"
 
 // Enable this define to do slow checks for components being added to the wrong
 // world's scene, when using PIE. This can happen if a PIE component is reattached
diff --git a/Engine/Source/Runtime/Windows/D3D11RHI/Private/D3D11Util.cpp b/Engine/Source/Runtime/Windows/D3D11RHI/Private/D3D11Util.cpp
--- a/Engine/Source/Runtime/Windows/D3D11RHI/Private/D3D11Util.cpp
+++ b/Engine/Source/Runtime/Windows/D3D11RHI/Private/D3D11Util.cpp
@@ -5,6 +5,7 @@
 =============================================================================*/
 
 #include "D3D11RHIPrivate.h"
+#include "EngineModule.h"
 
 #define D3DERR(x) case x: ErrorCodeText = TEXT(#x); break;
 #define LOCTEXT_NAMESPACE "Developer.MessageLog"
