Fix crash from existing foliage that's based on blueprint components when running an uncooked game

diff --git a/Engine/Source/Runtime/Core/Public/UObject/ObjectVersion.h b/Engine/Source/Runtime/Core/Public/UObject/ObjectVersion.h
--- a/Engine/Source/Runtime/Core/Public/UObject/ObjectVersion.h
+++ b/Engine/Source/Runtime/Core/Public/UObject/ObjectVersion.h
@@ -556,6 +556,8 @@ enum EUnrealEngineObjectUE4Version
 	VER_UE4_STORE_BONE_EXPORT_NAMES,
 	// changed mesh emitter initial orientation to distribution
 	VER_UE4_MESH_EMITTER_INITIAL_ORIENTATION_DISTRIBUTION,
+	// Foliage on blueprints causes crashes
+	VER_UE4_DISALLOW_FOLIAGE_ON_BLUEPRINTS,
 
 	// -----<new versions can be added before this line>-------------------------------------------------
 	// - this needs to be the last line (see note below)
diff --git a/Engine/Source/Runtime/Engine/Private/InstancedFoliage.cpp b/Engine/Source/Runtime/Engine/Private/InstancedFoliage.cpp
--- a/Engine/Source/Runtime/Engine/Private/InstancedFoliage.cpp
+++ b/Engine/Source/Runtime/Engine/Private/InstancedFoliage.cpp
@@ -1494,6 +1494,24 @@ void AInstancedFoliageActor::Serialize(FArchive& Ar)
 void AInstancedFoliageActor::PostLoad()
 {
 	Super::PostLoad();
+
+	if (GetLinkerUE4Version() < VER_UE4_DISALLOW_FOLIAGE_ON_BLUEPRINTS)
+	{
+		for (auto& MeshPair : FoliageMeshes)
+		{
+			for (FFoliageInstance& Instance : MeshPair.Value.Instances)
+			{
+				// Clear out the Base for any instances based on blueprint-created components,
+				// as those components will be destroyed when the construction scripts are
+				// re-run, leaving dangling references and causing crashes (woo!)
+				if (Instance.Base && Instance.Base->bCreatedByConstructionScript)
+				{
+					Instance.Base = NULL;
+				}
+			}
+		}
+	}
+
 	if( GIsEditor )
 	{
 		for( TMap<class UStaticMesh*, struct FFoliageMeshInfo>::TIterator MeshIt(FoliageMeshes); MeshIt; ++MeshIt )
@@ -1524,12 +1542,6 @@ void AInstancedFoliageActor::PostLoad()
 				}
 				else
 				{
-					// Clear out the Base for any instances based on blueprint components, as they will be destroyed when the construction scripts are re-run.
-					if (Instance.Base && Instance.Base->bCreatedByConstructionScript)
-					{
-						Instance.Base = NULL;
-					}
-
 					// Add valid instances to the hash.
 					MeshInfo.InstanceHash->InsertInstance(Instance.Location, InstanceIdx);
 					FFoliageComponentHashInfo* ComponentHashInfo = MeshInfo.ComponentHash.Find(Instance.Base);
