Flush GPU Resources when loading maps.
Fixes TTP: 337252

diff --git a/Engine/Source/Runtime/Engine/Private/UnrealEngine.cpp b/Engine/Source/Runtime/Engine/Private/UnrealEngine.cpp
--- a/Engine/Source/Runtime/Engine/Private/UnrealEngine.cpp
+++ b/Engine/Source/Runtime/Engine/Private/UnrealEngine.cpp
@@ -8088,8 +8088,8 @@ bool UEngine::LoadMap( FWorldContext& WorldContext, FURL URL, class UPendingNetG
 			else
 			{
 				TransitionGameMode = TEXT("");
-			}
-			LoadMapRedrawViewports();
+			}			
+			LoadMapRedrawViewports();			
 			TransitionType = TT_None;
 		}
 
@@ -8162,6 +8162,16 @@ bool UEngine::LoadMap( FWorldContext& WorldContext, FURL URL, class UPendingNetG
 
 	// Clean up the previous level out of memory.
 	CollectGarbage( GARBAGE_COLLECTION_KEEPFLAGS, true );
+	
+	// For platforms which manage GPU memory directly we must Enqueue a flush, and wait for it to be processed
+	// so that any pending frees that depend on the GPU will be processed.  Otherwise a whole map's worth of GPU memory
+	// may be unavailable to load the next one.
+	ENQUEUE_UNIQUE_RENDER_COMMAND(FlushCommand, 
+		{
+			RHIFlushResources();
+		}
+	);
+	FlushRenderingCommands();	  
 
 	// Cancels the Forced StreamType for textures using a timer.
 	if (!IStreamingManager::HasShutdown())
