Expose Twist and Swing mode for constraints

diff --git a/Engine/Source/Runtime/Engine/Classes/PhysicsEngine/ConstraintInstance.h b/Engine/Source/Runtime/Engine/Classes/PhysicsEngine/ConstraintInstance.h
--- a/Engine/Source/Runtime/Engine/Classes/PhysicsEngine/ConstraintInstance.h
+++ b/Engine/Source/Runtime/Engine/Classes/PhysicsEngine/ConstraintInstance.h
@@ -34,6 +34,18 @@ namespace EConstraintFrame
 	};
 }
 
+UENUM()
+namespace EAngularDriveMode
+{
+	enum Type
+	{
+		//Follows the shortest arc between a pair of anuglar configurations (Ignored if any angular limits/locks are used)
+		SLERP,
+		//Path is decomposed into twist and swing. Doesn't follow shortest arc and may have gimbal lock. (Works with angular limits/locks)
+		TwistAndSwing
+	};
+}
+
 /** Container for a physics representation of an object */
 USTRUCT()
 struct ENGINE_API FConstraintInstance
@@ -312,6 +324,10 @@ struct ENGINE_API FConstraintInstance
 	UPROPERTY()
 	FQuat AngularPositionTarget_DEPRECATED;
 	
+	/** The way rotation paths are estimated */
+	UPROPERTY(EditAnywhere, BlueprintReadOnly, Category = AngularMotor, meta = (editcondition = "bAngularVelocityDrive"))
+	TEnumAsByte<enum EAngularDriveMode::Type> AngularDriveMode;
+
 	/** Target orientation for the angular drive. */
 	UPROPERTY(EditAnywhere, BlueprintReadOnly, Category=AngularMotor)
 	FRotator AngularOrientationTarget;
@@ -369,18 +385,19 @@ public:
 	/** Get the position of this constraint in world space. */
 	FVector GetConstraintLocation();
 
-	void	SetLinearPositionDrive(bool bEnableXDrive, bool bEnableYDrive, bool bEnableZDrive);
-	void	SetLinearVelocityDrive(bool bEnableXDrive, bool bEnableYDrive, bool bEnableZDrive);
-	void	SetAngularPositionDrive(bool bEnableSwingDrive, bool bEnableTwistDrive);
-	void	SetAngularVelocityDrive(bool bEnableSwingDrive, bool bEnableTwistDrive);
+	void SetLinearPositionDrive(bool bEnableXDrive, bool bEnableYDrive, bool bEnableZDrive);
+	void SetLinearVelocityDrive(bool bEnableXDrive, bool bEnableYDrive, bool bEnableZDrive);
+	void SetAngularPositionDrive(bool bEnableSwingDrive, bool bEnableTwistDrive);
+	void SetAngularVelocityDrive(bool bEnableSwingDrive, bool bEnableTwistDrive);
+
+	void SetLinearPositionTarget(const FVector& InPosTarget);
+	void SetLinearVelocityTarget(const FVector& InVelTarget);
+	void SetLinearDriveParams(float InSpring, float InDamping, float InForceLimit);
 
-	void	SetLinearPositionTarget(const FVector& InPosTarget);
-	void	SetLinearVelocityTarget(const FVector& InVelTarget);
-	void	SetLinearDriveParams(float InSpring, float InDamping, float InForceLimit);
+	void SetAngularOrientationTarget(const FQuat& InPosTarget);
+	void SetAngularVelocityTarget(const FVector& InVelTarget);
+	void SetAngularDriveParams(float InSpring, float InDamping, float InForceLimit);
 
-	void	SetAngularOrientationTarget(const FQuat& InPosTarget);
-	void	SetAngularVelocityTarget(const FVector& InVelTarget);
-	void	SetAngularDriveParams(float InSpring, float InDamping, float InForceLimit);
 
 	/** Scale Angular Limit Constraints (as defined in RB_ConstraintSetup) */
 	void	SetAngularDOFLimitScale(float InSwing1LimitScale, float InSwing2LimitScale, float InTwistLimitScale);
diff --git a/Engine/Source/Runtime/Engine/Private/PhysicsEngine/ConstraintInstance.cpp b/Engine/Source/Runtime/Engine/Private/PhysicsEngine/ConstraintInstance.cpp
--- a/Engine/Source/Runtime/Engine/Private/PhysicsEngine/ConstraintInstance.cpp
+++ b/Engine/Source/Runtime/Engine/Private/PhysicsEngine/ConstraintInstance.cpp
@@ -374,7 +374,15 @@ void FConstraintInstance::InitConstraint(USceneComponent* Owner, FBodyInstance*
 		float DriveSpring = bAngularOrientationDrive ? AngularDriveSpring : 0.0f;
 		float DriveDamping = bAngularVelocityDrive  ? AngularDriveDamping : 0.0f;
 
-		PD6Joint->setDrive(PxD6Drive::eSLERP, PxD6JointDrive(DriveSpring, DriveDamping, AngularForceLimit, bAccelerationDrive));
+		if (AngularDriveMode == EAngularDriveMode::SLERP)
+		{
+			PD6Joint->setDrive(PxD6Drive::eSLERP, PxD6JointDrive(DriveSpring, DriveDamping, AngularForceLimit, bAccelerationDrive));
+		}
+		else
+		{
+			PD6Joint->setDrive(PxD6Drive::eTWIST, PxD6JointDrive(DriveSpring, DriveDamping, AngularForceLimit, bAccelerationDrive));
+			PD6Joint->setDrive(PxD6Drive::eSWING, PxD6JointDrive(DriveSpring, DriveDamping, AngularForceLimit, bAccelerationDrive));
+		}
 	}
 
 	FQuat OrientationTargetQuat(AngularOrientationTarget);
@@ -786,6 +794,10 @@ void FConstraintInstance::SetAngularDriveParams(float InSpring, float InDamping,
 }
 
 
+
+
+
+
 /** Scale Angular Limit Constraints (as defined in RB_ConstraintSetup) */
 void FConstraintInstance::SetAngularDOFLimitScale(float InSwing1LimitScale, float InSwing2LimitScale, float InTwistLimitScale)
 {
