Fix part of constraint motors using rad/sec and other parts using rev/sec

diff --git a/Engine/Source/Runtime/Core/Public/UObject/ObjectVersion.h b/Engine/Source/Runtime/Core/Public/UObject/ObjectVersion.h
--- a/Engine/Source/Runtime/Core/Public/UObject/ObjectVersion.h
+++ b/Engine/Source/Runtime/Core/Public/UObject/ObjectVersion.h
@@ -558,6 +558,8 @@ enum EUnrealEngineObjectUE4Version
 	VER_UE4_MESH_EMITTER_INITIAL_ORIENTATION_DISTRIBUTION,
 	// Foliage on blueprints causes crashes
 	VER_UE4_DISALLOW_FOLIAGE_ON_BLUEPRINTS,
+	// change motors to use revolutions per second instead of rads/second
+	VER_UE4_FIXUP_MOTOR_UNITS,
 
 	// -----<new versions can be added before this line>-------------------------------------------------
 	// - this needs to be the last line (see note below)
diff --git a/Engine/Source/Runtime/Engine/Private/PhysicsEngine/ConstraintInstance.cpp b/Engine/Source/Runtime/Engine/Private/PhysicsEngine/ConstraintInstance.cpp
--- a/Engine/Source/Runtime/Engine/Private/PhysicsEngine/ConstraintInstance.cpp
+++ b/Engine/Source/Runtime/Engine/Private/PhysicsEngine/ConstraintInstance.cpp
@@ -388,7 +388,7 @@ void FConstraintInstance::InitConstraint(USceneComponent* Owner, FBodyInstance*
 	FQuat OrientationTargetQuat(AngularOrientationTarget);
 
 	PD6Joint->setDrivePosition(PxTransform(U2PVector(LinearPositionTarget), U2PQuat(OrientationTargetQuat)));
-	PD6Joint->setDriveVelocity(U2PVector(LinearVelocityTarget), U2PVector(AngularVelocityTarget));
+	PD6Joint->setDriveVelocity(U2PVector(LinearVelocityTarget), U2PVector(AngularVelocityTarget * 2 * PI));
 	
 	// creation of joints wakes up rigid bodies, so we put them to sleep again if both were initially asleep
 	if(bActor1Asleep && bActor2Asleep)
@@ -995,6 +995,11 @@ void FConstraintInstance::PostSerialize(const FArchive& Ar)
 		SwingLimitDamping		/=  CVarConstraintDamplingScale->GetValueOnGameThread();
 		TwistLimitDamping		/=  CVarConstraintDamplingScale->GetValueOnGameThread();
 	}
+
+	if (Ar.IsLoading() && Ar.UE4Ver() < VER_UE4_FIXUP_MOTOR_UNITS)
+	{
+		AngularVelocityTarget *= 1.f / (2.f * PI);	//we want to use revolutions per second - old system was using radians directly
+	}
 }
 
 void FConstraintInstance::OnConstraintBroken()
