Change motor bp to use rotator instead of quaternion

diff --git a/Engine/Source/Runtime/Engine/Classes/PhysicsEngine/PhysicsConstraintComponent.h b/Engine/Source/Runtime/Engine/Classes/PhysicsEngine/PhysicsConstraintComponent.h
--- a/Engine/Source/Runtime/Engine/Classes/PhysicsEngine/PhysicsConstraintComponent.h
+++ b/Engine/Source/Runtime/Engine/Classes/PhysicsEngine/PhysicsConstraintComponent.h
@@ -158,7 +158,7 @@ public:
 	 *	@param InPosTarget		Target orientation
 	 */
 	UFUNCTION(BlueprintCallable, Category="Physics|Components|PhysicsConstraint")
-	ENGINE_API void	SetAngularOrientationTarget(const FQuat& InPosTarget);
+	ENGINE_API void	SetAngularOrientationTarget(const FRotator& InPosTarget);
 
 	
 	/** Sets the target velocity for the angular drive. 
diff --git a/Engine/Source/Runtime/Engine/Private/PhysicsEngine/PhysicsConstraintComponent.cpp b/Engine/Source/Runtime/Engine/Private/PhysicsEngine/PhysicsConstraintComponent.cpp
--- a/Engine/Source/Runtime/Engine/Private/PhysicsEngine/PhysicsConstraintComponent.cpp
+++ b/Engine/Source/Runtime/Engine/Private/PhysicsEngine/PhysicsConstraintComponent.cpp
@@ -461,9 +461,9 @@ void UPhysicsConstraintComponent::SetLinearDriveParams( float InSpring, float In
 	ConstraintInstance.SetLinearDriveParams(InSpring, InDamping, InForceLimit);
 }
 
-void UPhysicsConstraintComponent::SetAngularOrientationTarget( const FQuat& InPosTarget )
+void UPhysicsConstraintComponent::SetAngularOrientationTarget( const FRotator& InPosTarget )
 {
-	ConstraintInstance.SetAngularOrientationTarget(InPosTarget);
+	ConstraintInstance.SetAngularOrientationTarget(InPosTarget.Quaternion());
 }
 
 void UPhysicsConstraintComponent::SetAngularVelocityTarget( const FVector& InVelTarget )
