Don't keep spamming when network asset is missing

diff --git a/Engine/Source/Runtime/Engine/Private/PackageMapClient.cpp b/Engine/Source/Runtime/Engine/Private/PackageMapClient.cpp
--- a/Engine/Source/Runtime/Engine/Private/PackageMapClient.cpp
+++ b/Engine/Source/Runtime/Engine/Private/PackageMapClient.cpp
@@ -1799,6 +1799,7 @@ UObject * FNetGUIDCache::GetObjectFromNetGUID( const FNetworkGUID & NetGUID )
 		// If we can't find an object, and it's not pending, warn
 		if ( !CacheObjectPtr->bIgnoreWhenMissing && !CacheObjectPtr->bIsPending )
 		{
+			CacheObjectPtr->bIgnoreWhenMissing = 1;	// Set this to 1 so that we don't keep spamming
 			UE_LOG( LogNetPackageMap, Warning, TEXT( "GetObjectFromNetGUID: Failed to resolve path. Path: %s, Outer: %s, NetGUID: %s" ), *CacheObjectPtr->PathName.ToString(), ObjOuter != NULL ? *ObjOuter->GetPathName() : TEXT( "NULL" ), *NetGUID.ToString() );
 		}
 
