Mark assets as broken when they fail to load, so we completely stop trying

diff --git a/Engine/Source/Runtime/Engine/Classes/Engine/PackageMapClient.h b/Engine/Source/Runtime/Engine/Classes/Engine/PackageMapClient.h
--- a/Engine/Source/Runtime/Engine/Classes/Engine/PackageMapClient.h
+++ b/Engine/Source/Runtime/Engine/Classes/Engine/PackageMapClient.h
@@ -42,7 +42,7 @@ public:
 class FNetGuidCacheObject
 {
 public:
-	FNetGuidCacheObject() : bNoLoad( 0 ), bIgnoreWhenMissing( 0 ), bIsPending( 0 ) {}
+	FNetGuidCacheObject() : bNoLoad( 0 ), bIgnoreWhenMissing( 0 ), bIsPending( 0 ), bIsBroken( 0 ) {}
 
 	TWeakObjectPtr< UObject >	Object;
 
@@ -53,6 +53,7 @@ public:
 	uint8						bNoLoad				: 1;	// Don't load this, only do a find
 	uint8						bIgnoreWhenMissing	: 1;	// Don't warn when this asset can't be found or loaded
 	uint8						bIsPending			: 1;	// This object is waiting to be fully loaded
+	uint8						bIsBroken			: 1;	// If this object failed to load, then we set this to signify that we should stop trying
 };
 
 class ENGINE_API FNetGUIDCache
diff --git a/Engine/Source/Runtime/Engine/Private/PackageMapClient.cpp b/Engine/Source/Runtime/Engine/Private/PackageMapClient.cpp
--- a/Engine/Source/Runtime/Engine/Private/PackageMapClient.cpp
+++ b/Engine/Source/Runtime/Engine/Private/PackageMapClient.cpp
@@ -1743,6 +1743,12 @@ UObject * FNetGUIDCache::GetObjectFromNetGUID( const FNetworkGUID & NetGUID )
 		return Object;
 	}
 
+	if ( CacheObjectPtr->bIsBroken )
+	{
+		// Don't warn, assume we already have
+		return NULL;
+	}
+
 	if ( IsNetGUIDAuthority() )
 	{
 		UE_LOG( LogNetPackageMap, Warning, TEXT( "GetObjectFromNetGUID: Guid with no object on server. Path: %s, NetGUID: %s" ), *CacheObjectPtr->PathName.ToString(), *NetGUID.ToString() );
@@ -1799,7 +1805,7 @@ UObject * FNetGUIDCache::GetObjectFromNetGUID( const FNetworkGUID & NetGUID )
 		// If we can't find an object, and it's not pending, warn
 		if ( !CacheObjectPtr->bIgnoreWhenMissing && !CacheObjectPtr->bIsPending )
 		{
-			CacheObjectPtr->bIgnoreWhenMissing = 1;	// Set this to 1 so that we don't keep spamming
+			CacheObjectPtr->bIsBroken = 1;	// Set this to 1 so that we don't keep spamming
 			UE_LOG( LogNetPackageMap, Warning, TEXT( "GetObjectFromNetGUID: Failed to resolve path. Path: %s, Outer: %s, NetGUID: %s" ), *CacheObjectPtr->PathName.ToString(), ObjOuter != NULL ? *ObjOuter->GetPathName() : TEXT( "NULL" ), *NetGUID.ToString() );
 		}
 
