#ue4 - Deprecate MovementComponent functions and variables related to speed modifiers.

Old implementations and uses of these functions should use GetMaxSpeed() and GetMaxAcceleration() instead, and apply the multiplier internally in their implementations.
The old functions are no longer used by engine code and only remain to maintain compilation and provide the deprecation message.

- Deprecated function UMovementComponent::GetMaxSpeedModifier()
- Deprecated function UMovementComponent::GetModifiedMaxSpeed()
- Deprecated function UCharacterMovementComponent::GetModifiedMaxAcceleration()
- Deprecated variable UCharacterMovementComponent::CrouchedSpeedMultiplier

- Added UCharacterMovementComponent::GetMaxAcceleration()
- Added MaxWalkSpeedCrouched - max walk speed when crouched. For older content, the initial value will be MaxWalkSpeed * CrouchedSpeedMultiplier (now deprecated).
- Added MaxCustomSpeed - max speed used in custom movement mode (previously used MaxWalkSpeed, and will default to that value for existing content).

diff --git a/Engine/Config/BaseEngine.ini b/Engine/Config/BaseEngine.ini
--- a/Engine/Config/BaseEngine.ini
+++ b/Engine/Config/BaseEngine.ini
@@ -951,6 +951,11 @@ bUseStreamingPause=false
 +K2ParamRedirects=(NodeName="/Script/BlueprintGraph.K2Node_CallFunction", OldParamName="CharacterMovementComponent.AddImpulse.InMomentum", NewParamName="Impulse")
 +K2ParamRedirects=(NodeName="/Script/BlueprintGraph.K2Node_CallFunction", OldParamName="CharacterMovementComponent.AddImpulse.bMassIndependent", NewParamName="bVelocityChange")
 
++K2FieldRedirects=(OldFieldName="MovementComponent.GetModifiedMaxSpeed",NewFieldName="MovementComponent.K2_GetModifiedMaxSpeed")
++K2FieldRedirects=(OldFieldName="MovementComponent.GetMaxSpeedModifier",NewFieldName="MovementComponent.K2_GetMaxSpeedModifier")
++K2FieldRedirects=(OldFieldName="CharacterMovementComponent.GetModifiedMaxAcceleration",NewFieldName="CharacterMovementComponent.K2_GetModifiedMaxAcceleration")
+
+
 +K2FieldRedirects=(OldFieldName="PlayerController.LookRightScale",NewFieldName="PlayerController.InputYawScale")
 +K2FieldRedirects=(OldFieldName="PlayerController.LookUpScale",NewFieldName="PlayerController.InputPitchScale")
 
diff --git a/Engine/Source/Runtime/Core/Public/UObject/ObjectVersion.h b/Engine/Source/Runtime/Core/Public/UObject/ObjectVersion.h
--- a/Engine/Source/Runtime/Core/Public/UObject/ObjectVersion.h
+++ b/Engine/Source/Runtime/Core/Public/UObject/ObjectVersion.h
@@ -560,6 +560,8 @@ enum EUnrealEngineObjectUE4Version
 	VER_UE4_DISALLOW_FOLIAGE_ON_BLUEPRINTS,
 	// change motors to use revolutions per second instead of rads/second
 	VER_UE4_FIXUP_MOTOR_UNITS,
+	// deprecated MovementComponent functions including "ModifiedMaxSpeed" et al
+	VER_UE4_DEPRECATED_MOVEMENTCOMPONENT_MODIFIED_SPEEDS,
 
 	// -----<new versions can be added before this line>-------------------------------------------------
 	// - this needs to be the last line (see note below)
diff --git a/Engine/Source/Runtime/Engine/Classes/GameFramework/CharacterMovementComponent.h b/Engine/Source/Runtime/Engine/Classes/GameFramework/CharacterMovementComponent.h
--- a/Engine/Source/Runtime/Engine/Classes/GameFramework/CharacterMovementComponent.h
+++ b/Engine/Source/Runtime/Engine/Classes/GameFramework/CharacterMovementComponent.h
@@ -349,13 +349,17 @@ public:
 	UPROPERTY(Category="Character Movement", EditAnywhere, BlueprintReadWrite)
 	float GroundFriction;
 
-	/** The maximum ground speed. Also determines maximum lateral speed when falling. */
+	/** The maximum ground speed when walking. Also determines maximum lateral speed when falling. */
 	UPROPERTY(Category="Character Movement", EditAnywhere, BlueprintReadWrite)
 	float MaxWalkSpeed;
 
-	/** Multiplier to max ground speed to use when crouched */
+	/** The maximum ground speed when walking and crouched. */
+	UPROPERTY(Category="Character Movement", EditAnywhere, BlueprintReadWrite)
+	float MaxWalkSpeedCrouched;
+
+	/** The maximum speed when using Custom movement mode. */
 	UPROPERTY(Category="Character Movement", EditAnywhere, BlueprintReadWrite)
-	float CrouchedSpeedMultiplier;
+	float MaxCustomMovementSpeed;
 
 	/** Collision half-height when crouching (component scale is applied separately) */
 	UPROPERTY(Category="Character Movement", EditAnywhere, BlueprintReadOnly)
@@ -511,6 +515,10 @@ public:
 	UPROPERTY()
 	uint32 bForceBraking_DEPRECATED:1;
 
+	/** Multiplier to max ground speed to use when crouched */
+	UPROPERTY()
+	float CrouchedSpeedMultiplier_DEPRECATED;
+
 protected:
 
 	/**
@@ -815,7 +823,6 @@ public:
 
 	//BEGIN UMovementComponent Interface
 	virtual float GetMaxSpeed() const override;
-	virtual float GetMaxSpeedModifier() const override;
 	virtual void StopActiveMovement() override;
 	virtual bool IsCrouching() const override;
 	virtual bool IsFalling() const override;
@@ -967,14 +974,22 @@ public:
 	virtual float GetMaxJumpHeight() const;
 	
 	/** @return Maximum acceleration for the current state, based on MaxAcceleration and any additional modifiers. */
-	UFUNCTION(BlueprintCallable, Category="Pawn|Components|CharacterMovement")
+	DEPRECATED(4.3, "GetModifiedMaxAcceleration() is deprecated, apply your own modifiers to GetMaxAcceleration() if desired.")
 	virtual float GetModifiedMaxAcceleration() const;
+	
+	/** @return Maximum acceleration for the current state, based on MaxAcceleration and any additional modifiers. */
+	UFUNCTION(BlueprintCallable, Category="Pawn|Components|CharacterMovement", meta=(DeprecatedFunction, DisplayName="GetModifiedMaxAcceleration()", DeprecationMessage="GetModifiedMaxAcceleration() is deprecated, apply your own modifiers to GetMaxAcceleration() if desired."))
+	virtual float K2_GetModifiedMaxAcceleration() const;
+
+	/** @return Maximum acceleration for the current state. */
+	UFUNCTION(BlueprintCallable, Category="Pawn|Components|CharacterMovement")
+	virtual float GetMaxAcceleration() const;
 
 	/** @return Current acceleration, computed from input vector each update. */
 	UFUNCTION(BlueprintCallable, Category="Pawn|Components|CharacterMovement", meta=(Keywords="Acceleration GetAcceleration"))
 	FVector GetCurrentAcceleration() const;
 
-	/** @return Modifier [0..1] which affects max speed, based on the magnitude of the input vector. */
+	/** @return Modifier [0..1] based on the magnitude of the last input vector, which is used to modify the acceleration and max speed during movement. */
 	UFUNCTION(BlueprintCallable, Category="Pawn|Components|CharacterMovement")
 	float GetAnalogInputModifier() const;
 	
diff --git a/Engine/Source/Runtime/Engine/Classes/GameFramework/MovementComponent.h b/Engine/Source/Runtime/Engine/Classes/GameFramework/MovementComponent.h
--- a/Engine/Source/Runtime/Engine/Classes/GameFramework/MovementComponent.h
+++ b/Engine/Source/Runtime/Engine/Classes/GameFramework/MovementComponent.h
@@ -80,20 +80,28 @@ public:
 	UFUNCTION(BlueprintCallable, Category="Components|Movement")
 	virtual float GetGravityZ() const;
 
-	/** @return maximum speed of component in current movement mode.  Multiply by MaxSpeedModifier for limitations based on posture, etc. */
+	/** @return Maximum speed of component in current movement mode. */
 	UFUNCTION(BlueprintCallable, Category="Components|Movement")
 	virtual float GetMaxSpeed() const;
 
 	/** @return a scalar applied to the maximum velocity that the component can currently move. */
-	UFUNCTION(BlueprintCallable, Category="Components|Movement")
+	DEPRECATED(4.3, "GetMaxSpeedModifier() is deprecated, apply your own modifiers to GetMaxSpeed() if desired.")
 	virtual float GetMaxSpeedModifier() const;
+	
+	/** @return a scalar applied to the maximum velocity that the component can currently move. */
+	UFUNCTION(BlueprintCallable, Category="Components|Movement", meta=(DeprecatedFunction, DisplayName="GetMaxSpeedModifier", DeprecationMessage="GetMaxSpeedModifier() is deprecated, apply your own modifiers to GetMaxSpeed() if desired."))
+	virtual float K2_GetMaxSpeedModifier() const;
 
 	/** @return the result of GetMaxSpeed() * GetMaxSpeedModifier(). */
-	UFUNCTION(BlueprintCallable, Category="Components|Movement")
+	DEPRECATED(4.3, "GetModifiedMaxSpeed() is deprecated, use GetMaxSpeed() instead.")
 	virtual float GetModifiedMaxSpeed() const;
 
+	/** @return the result of GetMaxSpeed() * GetMaxSpeedModifier(). */
+	UFUNCTION(BlueprintCallable, Category="Components|Movement", meta=(DeprecatedFunction, DisplayName="GetModifiedMaxSpeed", DeprecationMessage="GetModifiedMaxSpeed() is deprecated, apply your own modifiers to GetMaxSpeed() if desired."))
+	virtual float K2_GetModifiedMaxSpeed() const;
+
 	/**
-	 * Returns true if the current velocity is exceeding the given max speed (usually the result of GetModifiedMaxSpeed()), within a small error tolerance.
+	 * Returns true if the current velocity is exceeding the given max speed (usually the result of GetMaxSpeed()), within a small error tolerance.
 	 * Note that under normal circumstances updates cause by acceleration will not cause this to be true, however external forces or changes in the max speed limit
 	 * can cause the max speed to be violated.
 	 */
@@ -285,16 +293,6 @@ inline float UMovementComponent::GetMaxSpeed() const
 	return 0.f;
 }
 
-inline float UMovementComponent::GetMaxSpeedModifier() const
-{
-	return 1.0f;
-}
-
-inline float UMovementComponent::GetModifiedMaxSpeed() const
-{
-	return GetMaxSpeed() * GetMaxSpeedModifier();
-}
-
 inline void UMovementComponent::StopMovementImmediately()
 {
 	Velocity = FVector::ZeroVector;
diff --git a/Engine/Source/Runtime/Engine/Private/CharacterMovementComponent.cpp b/Engine/Source/Runtime/Engine/Private/CharacterMovementComponent.cpp
--- a/Engine/Source/Runtime/Engine/Private/CharacterMovementComponent.cpp
+++ b/Engine/Source/Runtime/Engine/Private/CharacterMovementComponent.cpp
@@ -130,11 +130,13 @@ UCharacterMovementComponent::UCharacterMovementComponent(const class FPostConstr
 	MaxFlySpeed = 600.0f;
 	MaxWalkSpeed = 600.0f;
 	MaxSwimSpeed = 300.0f;
+	MaxCustomMovementSpeed = MaxWalkSpeed;
 	
 	MaxSimulationTimeStep = 0.05f;
 	MaxSimulationIterations = 8;
 
-	CrouchedSpeedMultiplier = 0.5f;
+	CrouchedSpeedMultiplier_DEPRECATED = 0.5f;
+	MaxWalkSpeedCrouched = MaxWalkSpeed * CrouchedSpeedMultiplier_DEPRECATED;
 	MaxOutOfWaterStepHeight = 40.0f;
 	OutofWaterZ = 420.0f;
 	AirControl = 0.05f;
@@ -230,6 +232,12 @@ void UCharacterMovementComponent::PostLoad()
 		// Compute the walkable floor angle, since we have never done so yet.
 		UCharacterMovementComponent::SetWalkableFloorZ(WalkableFloorZ);
 	}
+
+	if (LinkerUE4Ver < VER_UE4_DEPRECATED_MOVEMENTCOMPONENT_MODIFIED_SPEEDS)
+	{
+		MaxWalkSpeedCrouched = MaxWalkSpeed * CrouchedSpeedMultiplier_DEPRECATED;
+		MaxCustomMovementSpeed = MaxWalkSpeed;
+	}
 }
 
 
@@ -656,7 +664,7 @@ void UCharacterMovementComponent::PerformAirControl(FVector Direction, float ZDi
 			{
 				float Dist2D = Direction.Size2D();
 				//Direction.Z = 0.f;
-				Acceleration = Direction.SafeNormal() * GetModifiedMaxAcceleration();
+				Acceleration = Direction.SafeNormal() * GetMaxAcceleration();
 
 				if ( (Dist2D < 0.5f * FMath::Abs(Direction.Z)) && ((Velocity | Direction) > 0.5f*FMath::Square(Dist2D)) )
 				{
@@ -1706,16 +1714,22 @@ float UCharacterMovementComponent::GetGravityZ() const
 
 float UCharacterMovementComponent::GetMaxSpeed() const
 {
-	float Result = MaxWalkSpeed;
-	if ( MovementMode == MOVE_Flying )
-	{
-		Result = MaxFlySpeed;
-	}
-	else if ( MovementMode == MOVE_Swimming )
+	switch(MovementMode)
 	{
-		Result = MaxSwimSpeed;
+	case MOVE_Walking:
+		return IsCrouching() ? MaxWalkSpeedCrouched : MaxWalkSpeed;
+	case MOVE_Falling:
+		return MaxWalkSpeed;
+	case MOVE_Swimming:
+		return MaxSwimSpeed;
+	case MOVE_Flying:
+		return MaxFlySpeed;
+	case MOVE_Custom:
+		return MaxCustomMovementSpeed;
+	case MOVE_None:
+	default:
+		return 0.f;
 	}
-	return Result;
 }
 
 
@@ -1974,8 +1988,8 @@ void UCharacterMovementComponent::CalcVelocity(float DeltaTime, float Friction,
 	}
 
 	Friction = FMath::Max(0.f, Friction);
-	const float MaxAccel = GetModifiedMaxAcceleration();
-	float MaxSpeed = GetModifiedMaxSpeed();
+	const float MaxAccel = GetMaxAcceleration();
+	float MaxSpeed = GetMaxSpeed();
 	
 	// Check if path following requested movement
 	FVector RequestedAcceleration = FVector::ZeroVector;
@@ -2081,7 +2095,7 @@ void UCharacterMovementComponent::RequestDirectMove(const FVector& MoveVelocity,
 
 	if (IsFalling())
 	{
-		const FVector FallVelocity = MoveVelocity.ClampMaxSize(GetModifiedMaxSpeed());
+		const FVector FallVelocity = MoveVelocity.ClampMaxSize(GetMaxSpeed());
 		PerformAirControl(FallVelocity, FallVelocity.Z);
 		return;
 	}
@@ -2261,9 +2275,32 @@ float UCharacterMovementComponent::GetMaxJumpHeight() const
 	}
 }
 
+// TODO: deprecated, remove.
 float UCharacterMovementComponent::GetModifiedMaxAcceleration() const
 {
+	// Allow calling old deprecated function to maintain old behavior until it is removed.
+#pragma warning(push)
+#pragma warning(disable:4995)
+#pragma warning(disable:4996)
 	return CharacterOwner ? MaxAcceleration * GetMaxSpeedModifier() : 0.f;
+#pragma warning(pop)
+}
+
+// TODO: deprecated, remove.
+float UCharacterMovementComponent::K2_GetModifiedMaxAcceleration() const
+{
+	// Allow calling old deprecated function to maintain old behavior until it is removed.
+#pragma warning(push)
+#pragma warning(disable:4995)
+#pragma warning(disable:4996)
+	return GetModifiedMaxAcceleration();
+#pragma warning(pop)
+}
+
+
+float UCharacterMovementComponent::GetMaxAcceleration() const
+{
+	return MaxAcceleration;
 }
 
 FVector UCharacterMovementComponent::GetCurrentAcceleration() const
@@ -2552,7 +2589,7 @@ void UCharacterMovementComponent::PhysFalling(float deltaTime, int32 Iterations)
 
 	// Bound final 2d portion of velocity
 	const float Speed2d = Velocity.Size2D();
-	const float BoundSpeed = FMath::Max(Speed2d, GetModifiedMaxSpeed() * AnalogInputModifier);
+	const float BoundSpeed = FMath::Max(Speed2d, GetMaxSpeed() * AnalogInputModifier);
 
 	//bound acceleration, falling object has minimal ability to impact acceleration
 	FVector FallAcceleration = Acceleration;
@@ -2566,7 +2603,7 @@ void UCharacterMovementComponent::PhysFalling(float deltaTime, int32 Iterations)
 		if ( TickAirControl > 0.0f && FallAcceleration.SizeSquared() > 0.f )
 		{
 			const float TestWalkTime = FMath::Max(deltaTime, 0.05f);
-			const FVector TestWalk = ((TickAirControl * GetModifiedMaxAcceleration() * FallAcceleration.SafeNormal() + FVector(0.f,0.f,GetGravityZ())) * TestWalkTime + Velocity) * TestWalkTime;
+			const FVector TestWalk = ((TickAirControl * GetMaxAcceleration() * FallAcceleration.SafeNormal() + FVector(0.f,0.f,GetGravityZ())) * TestWalkTime + Velocity) * TestWalkTime;
 			if(!TestWalk.IsZero())
 			{
 				static const FName FallingTraceParamsTag = FName(TEXT("PhysFalling"));
@@ -2600,7 +2637,7 @@ void UCharacterMovementComponent::PhysFalling(float deltaTime, int32 Iterations)
 			}
 		}
 
-		float MaxAccel = GetModifiedMaxAcceleration() * TickAirControl;
+		float MaxAccel = GetMaxAcceleration() * TickAirControl;
 
 		// Boost maxAccel to increase player's control when falling
 		if( (Speed2d < 10.f) && (TickAirControl > 0.f) && (TickAirControl <= 0.05f) ) //allow initial burst
@@ -3323,19 +3360,6 @@ void UCharacterMovementComponent::AdjustFloorHeight()
 	}
 }
 
-float UCharacterMovementComponent::GetMaxSpeedModifier() const
-{
-	float Result = 1.0f;
-
-	// Apply crouch modifier
-	if (CharacterOwner && IsCrouching() && IsMovingOnGround())
-	{
-		Result *= CrouchedSpeedMultiplier;
-	}
-
-	return Result;
-}
-
 
 void UCharacterMovementComponent::StopActiveMovement() 
 { 
@@ -4561,16 +4585,16 @@ FVector UCharacterMovementComponent::ConstrainInputAcceleration(const FVector& I
 
 FVector UCharacterMovementComponent::ScaleInputAcceleration(const FVector& InputAcceleration) const
 {
-	return GetModifiedMaxAcceleration() * ClampMaxSizePrecise(InputAcceleration, 1.0f);
+	return GetMaxAcceleration() * ClampMaxSizePrecise(InputAcceleration, 1.0f);
 }
 
 
 float UCharacterMovementComponent::ComputeAnalogInputModifier() const
 {
-	const float ModifiedMaxAccel = GetModifiedMaxAcceleration();
-	if (Acceleration.SizeSquared() > 0.f && ModifiedMaxAccel > SMALL_NUMBER)
+	const float MaxAccel = GetMaxAcceleration();
+	if (Acceleration.SizeSquared() > 0.f && MaxAccel > SMALL_NUMBER)
 	{
-		return FMath::Clamp(Acceleration.Size() / ModifiedMaxAccel, 0.f, 1.f);
+		return FMath::Clamp(Acceleration.Size() / MaxAccel, 0.f, 1.f);
 	}
 
 	return 0.f;
diff --git a/Engine/Source/Runtime/Engine/Private/FloatingPawnMovement.cpp b/Engine/Source/Runtime/Engine/Private/FloatingPawnMovement.cpp
--- a/Engine/Source/Runtime/Engine/Private/FloatingPawnMovement.cpp
+++ b/Engine/Source/Runtime/Engine/Private/FloatingPawnMovement.cpp
@@ -89,7 +89,7 @@ void UFloatingPawnMovement::ApplyControlInputToVelocity(float DeltaTime)
 	const FVector ControlAcceleration = GetInputVector().ClampMaxSize(1.f);
 
 	const float AnalogInputModifier = (ControlAcceleration.SizeSquared() > 0.f ? ControlAcceleration.Size() : 0.f);
-	const float MaxPawnSpeed = GetModifiedMaxSpeed() * AnalogInputModifier;
+	const float MaxPawnSpeed = GetMaxSpeed() * AnalogInputModifier;
 	const bool bExceedingMaxSpeed = IsExceedingMaxSpeed(MaxPawnSpeed);
 
 	if (AnalogInputModifier > 0.f && !bExceedingMaxSpeed)
diff --git a/Engine/Source/Runtime/Engine/Private/MovementComponent.cpp b/Engine/Source/Runtime/Engine/Private/MovementComponent.cpp
--- a/Engine/Source/Runtime/Engine/Private/MovementComponent.cpp
+++ b/Engine/Source/Runtime/Engine/Private/MovementComponent.cpp
@@ -464,4 +464,38 @@ void UMovementComponent::AddRadialImpulse(const FVector& Origin, float Radius, f
 	// Default implementation does nothing
 }
 
+// TODO: Deprecated, remove.
+float UMovementComponent::GetMaxSpeedModifier() const
+{
+	return 1.0f;
+}
+
+// TODO: Deprecated, remove.
+float UMovementComponent::K2_GetMaxSpeedModifier() const
+{
+	// Allow calling old deprecated function to maintain old behavior until it is removed.
+#pragma warning(push)
+#pragma warning(disable:4995)
+#pragma warning(disable:4996)
+	return GetMaxSpeedModifier();
+#pragma warning(pop)
+}
+
+// TODO: Deprecated, remove.
+float UMovementComponent::GetModifiedMaxSpeed() const
+{
+	return GetMaxSpeed();
+}
+
+// TODO: Deprecated, remove.
+float UMovementComponent::K2_GetModifiedMaxSpeed() const
+{
+	// Allow calling old deprecated function to maintain old behavior until it is removed.
+#pragma warning(push)
+#pragma warning(disable:4995)
+#pragma warning(disable:4996)
+	return GetModifiedMaxSpeed();
+#pragma warning(pop)
+}
+
 #undef LOCTEXT_NAMESPACE
diff --git a/Engine/Source/Runtime/Engine/Private/ProjectileMovementComponent.cpp b/Engine/Source/Runtime/Engine/Private/ProjectileMovementComponent.cpp
--- a/Engine/Source/Runtime/Engine/Private/ProjectileMovementComponent.cpp
+++ b/Engine/Source/Runtime/Engine/Private/ProjectileMovementComponent.cpp
@@ -204,7 +204,7 @@ FVector UProjectileMovementComponent::LimitVelocity(FVector NewVelocity) const
 {
 	if (GetMaxSpeed() > 0.f)
 	{
-		NewVelocity = NewVelocity.ClampMaxSize(GetModifiedMaxSpeed());
+		NewVelocity = NewVelocity.ClampMaxSize(GetMaxSpeed());
 	}
 
 	return NewVelocity;
