Fix the issue where morphtarget imports more verts because of normal is not imported for morphtarget

- if you choose import normal option, and import morph targets, morphtargets don't import normal, so when it diffs normal, it will consider different, causing it to import as diff verts.
- Fixed it so that it won't compare normal if it's importing normal

diff --git a/Engine/Source/Editor/UnrealEd/Private/Factories/EditorMorphFactory.cpp b/Engine/Source/Editor/UnrealEd/Private/Factories/EditorMorphFactory.cpp
--- a/Engine/Source/Editor/UnrealEd/Private/Factories/EditorMorphFactory.cpp
+++ b/Engine/Source/Editor/UnrealEd/Private/Factories/EditorMorphFactory.cpp
@@ -90,7 +90,8 @@ void FMorphTargetBinaryImport::ImportMorphLODModel( UMorphTarget* MorphTarget, c
 			{
 				// populate the vertex data for the morph target mesh using its base mesh
 				// and the newly imported mesh				
-				MorphTarget->PostProcess( TargetSkelMesh, BaseMeshRawData, TargetMeshRawData, LODIndex );
+				// @todo : bCompareNormal is not used in this function
+				MorphTarget->PostProcess( TargetSkelMesh, BaseMeshRawData, TargetMeshRawData, LODIndex, false );
 			}
 
 			Warn->EndSlowTask();
diff --git a/Engine/Source/Editor/UnrealEd/Private/Fbx/FbxSkeletalMeshImport.cpp b/Engine/Source/Editor/UnrealEd/Private/Fbx/FbxSkeletalMeshImport.cpp
--- a/Engine/Source/Editor/UnrealEd/Private/Fbx/FbxSkeletalMeshImport.cpp
+++ b/Engine/Source/Editor/UnrealEd/Private/Fbx/FbxSkeletalMeshImport.cpp
@@ -2788,7 +2788,7 @@ void UnFbx::FFbxImporter::ImportMorphTargetsInternal( TArray<FbxNode*>& SkelMesh
 		FMorphMeshRawSource TargetMeshRawData( TmpSkeletalMesh );
 		FMorphMeshRawSource BaseMeshRawData( BaseSkelMesh, LODIndex );
 
-		MorphTarget->PostProcess( BaseSkelMesh, BaseMeshRawData, TargetMeshRawData, LODIndex );
+		MorphTarget->PostProcess( BaseSkelMesh, BaseMeshRawData, TargetMeshRawData, LODIndex, ImportOptions->ShouldImportNormals() == false );
 	}
 
 	GWarn->EndSlowTask();
diff --git a/Engine/Source/Runtime/Engine/Classes/Animation/VertexAnim/MorphTarget.h b/Engine/Source/Runtime/Engine/Classes/Animation/VertexAnim/MorphTarget.h
--- a/Engine/Source/Runtime/Engine/Classes/Animation/VertexAnim/MorphTarget.h
+++ b/Engine/Source/Runtime/Engine/Classes/Animation/VertexAnim/MorphTarget.h
@@ -79,7 +79,7 @@ public:
 	// Begin UObject interface.
 
 	/** Post process after importing **/
-	ENGINE_API void PostProcess( USkeletalMesh * NewMesh, const FMorphMeshRawSource& BaseSource, const FMorphMeshRawSource& TargetSource, int32 LODIndex );
+	ENGINE_API void PostProcess( USkeletalMesh * NewMesh, const FMorphMeshRawSource& BaseSource, const FMorphMeshRawSource& TargetSource, int32 LODIndex, bool bCompareNormal );
 
 	/** Remap vertex indices with base mesh. */
 	void RemapVertexIndices( USkeletalMesh* InBaseMesh, const TArray< TArray<uint32> > & BasedWedgePointIndices );
@@ -99,6 +99,6 @@ private:
 	* @param	TargetSource - final target vertex positions/attributes 
 	* @param	LODIndex - level of detail to use for the geometry
 	*/
-	void CreateMorphMeshStreams( const FMorphMeshRawSource& BaseSource, const FMorphMeshRawSource& TargetSource, int32 LODIndex );
+	void CreateMorphMeshStreams( const FMorphMeshRawSource& BaseSource, const FMorphMeshRawSource& TargetSource, int32 LODIndex, bool bCompareNormal );
 };
 
diff --git a/Engine/Source/Runtime/Engine/Private/MorphTools.cpp b/Engine/Source/Runtime/Engine/Private/MorphTools.cpp
--- a/Engine/Source/Runtime/Engine/Private/MorphTools.cpp
+++ b/Engine/Source/Runtime/Engine/Private/MorphTools.cpp
@@ -39,18 +39,18 @@ bool UMorphTarget::HasDataForLOD(int32 LODIndex)
 }
 
 
-void UMorphTarget::PostProcess( USkeletalMesh * NewMesh, const FMorphMeshRawSource& BaseSource, const FMorphMeshRawSource& TargetSource, int32 LODIndex )
+void UMorphTarget::PostProcess( USkeletalMesh * NewMesh, const FMorphMeshRawSource& BaseSource, const FMorphMeshRawSource& TargetSource, int32 LODIndex, bool bCompareNormal )
 {
 	// @todo anim: update BaseSkelMesh with my information
 	NewMesh->RegisterMorphTarget(this);
 
-	CreateMorphMeshStreams( BaseSource, TargetSource, LODIndex );
+	CreateMorphMeshStreams( BaseSource, TargetSource, LODIndex, bCompareNormal );
 
 	MarkPackageDirty();
 }
 
 
-void UMorphTarget::CreateMorphMeshStreams( const FMorphMeshRawSource& BaseSource, const FMorphMeshRawSource& TargetSource, int32 LODIndex )
+void UMorphTarget::CreateMorphMeshStreams( const FMorphMeshRawSource& BaseSource, const FMorphMeshRawSource& TargetSource, int32 LODIndex, bool bCompareNormal )
 {
 	check(BaseSource.IsValidTarget(TargetSource));
 
@@ -119,7 +119,11 @@ void UMorphTarget::CreateMorphMeshStreams( const FMorphMeshRawSource& BaseSource
 					FVector NormalDeltaZ (VTarget.TanZ - VBase.TanZ);
 
 					// check if position actually changed much
-					if( PositionDelta.Size() > THRESH_POINTS_ARE_NEAR || NormalDeltaZ.Size() > 0.1f )
+					if( PositionDelta.Size() > THRESH_POINTS_ARE_NEAR || 
+						// since we can't get imported morphtarget normal from FBX
+						// we can't compare normal unless it's calculated
+						// this is special flag to ignore normal diff
+						( bCompareNormal && NormalDeltaZ.Size() > 0.1f) )
 					{
 						// create a new entry
 						FVertexAnimDelta NewVertex;
