Exposed GetActorScale3D and GetActorRelativeScale3D functions to Blueprints
#ttp 333950 Blueprints: We do not have GetActorScale3D or GetActorRelativeScale3D functions
reviewer Andrew.Brown

diff --git a/Engine/Source/Runtime/Engine/Classes/GameFramework/Actor.h b/Engine/Source/Runtime/Engine/Classes/GameFramework/Actor.h
--- a/Engine/Source/Runtime/Engine/Classes/GameFramework/Actor.h
+++ b/Engine/Source/Runtime/Engine/Classes/GameFramework/Actor.h
@@ -704,6 +704,10 @@ public:
 	UFUNCTION(BlueprintCallable, Category="Utilities|Transformation")
 	void SetActorScale3D(const FVector& NewScale3D);
 
+	/** Returns the Actor's world-space scale. */
+	UFUNCTION(BlueprintCallable, Category="Utilities|Orientation")
+	FVector GetActorScale3D() const;
+
 	/** Returns the distance from this Actor to OtherActor. */
 	UFUNCTION(BlueprintCallable, Category="Utilities|Transformation")
 	float GetDistanceTo(AActor* OtherActor);
@@ -774,6 +778,12 @@ public:
 	UFUNCTION(BlueprintCallable, Category="Utilities|Transformation")
 	void SetActorRelativeScale3D(FVector NewRelativeScale);
 
+	/**
+	 * Return the actor's relative scale 3d
+	 */
+	UFUNCTION(BlueprintCallable, Category="Utilities|Orientation")
+	FVector GetActorRelativeScale3D() const;
+
 	/**
 	 *	Sets the actor to be hidden in the game
 	 *	@param	bNewHidden	Whether or not to hide the actor and all its components
diff --git a/Engine/Source/Runtime/Engine/Private/Actor.cpp b/Engine/Source/Runtime/Engine/Private/Actor.cpp
--- a/Engine/Source/Runtime/Engine/Private/Actor.cpp
+++ b/Engine/Source/Runtime/Engine/Private/Actor.cpp
@@ -2481,6 +2481,15 @@ void AActor::SetActorScale3D(const FVector& NewScale3D)
 }
 
 
+FVector AActor::GetActorScale3D() const
+{
+	if (RootComponent)
+	{
+		return RootComponent->GetComponentScale();
+	}
+	return FVector(1,1,1);
+}
+
 bool AActor::SetActorTransform(const FTransform& NewTransform, bool bSweep)
 {
 	// we have seen this gets NAN from kismet, and would like to see if this
@@ -2563,6 +2572,15 @@ void AActor::SetActorRelativeScale3D(FVector NewRelativeScale)
 	}
 }
 
+FVector AActor::GetActorRelativeScale3D() const
+{
+	if (RootComponent)
+	{
+		return RootComponent->RelativeScale3D;
+	}
+	return FVector(1,1,1);
+}
+
 void AActor::SetActorHiddenInGame( bool bNewHidden )
 {
 	if (bHidden != bNewHidden)
