Added code to vehicle controls setup menu to restrict certain bindings for certain controls.
Added function in UInputSettings to force an update of the keymaps after a mapping has been added/removed.
(Updated the default keymaps to stop the config file getting confused)

Fixes TTP 336944 - Vehicle Bindings easy to break

diff --git a/Engine/Source/Runtime/Engine/Classes/GameFramework/InputSettings.h b/Engine/Source/Runtime/Engine/Classes/GameFramework/InputSettings.h
--- a/Engine/Source/Runtime/Engine/Classes/GameFramework/InputSettings.h
+++ b/Engine/Source/Runtime/Engine/Classes/GameFramework/InputSettings.h
@@ -71,4 +71,6 @@ class ENGINE_API UInputSettings
 
 	void GetActionNames(TArray<FName>& ActionNames) const;
 	void GetAxisNames(TArray<FName>& AxisNames) const;
+private:
+	void ForceRebuildKeymaps();
 };
diff --git a/Engine/Source/Runtime/Engine/Private/UserInterface/InputSettings.cpp b/Engine/Source/Runtime/Engine/Private/UserInterface/InputSettings.cpp
--- a/Engine/Source/Runtime/Engine/Private/UserInterface/InputSettings.cpp
+++ b/Engine/Source/Runtime/Engine/Private/UserInterface/InputSettings.cpp
@@ -20,10 +20,7 @@ void UInputSettings::PostEditChangeChainProperty(FPropertyChangedChainEvent& Pro
 
 	if (MemberPropertyName == "ActionMappings" || MemberPropertyName == "AxisMappings" || MemberPropertyName == "AxisConfig")
 	{
-		for (TObjectIterator<UPlayerInput> It; It; ++It)
-		{
-			It->ForceRebuildingKeyMaps(true);
-		}
+		ForceRebuildKeymaps();
 	}
 }
 
@@ -39,6 +36,7 @@ void UInputSettings::SaveKeyMappings()
 void UInputSettings::AddActionMapping(const FInputActionKeyMapping& KeyMapping)
 {
 	ActionMappings.AddUnique(KeyMapping);
+	ForceRebuildKeymaps();
 }
 
 void UInputSettings::RemoveActionMapping(const FInputActionKeyMapping& KeyMapping)
@@ -51,11 +49,13 @@ void UInputSettings::RemoveActionMapping(const FInputActionKeyMapping& KeyMappin
 			// we don't break because the mapping may have been in the array twice
 		}
 	}
+	ForceRebuildKeymaps();
 }
 
 void UInputSettings::AddAxisMapping(const FInputAxisKeyMapping& KeyMapping)
 {
 	AxisMappings.AddUnique(KeyMapping);
+	ForceRebuildKeymaps();
 }
 
 void UInputSettings::RemoveAxisMapping(const FInputAxisKeyMapping& InKeyMapping)
@@ -70,6 +70,7 @@ void UInputSettings::RemoveAxisMapping(const FInputAxisKeyMapping& InKeyMapping)
 			// we don't break because the mapping may have been in the array twice
 		}
 	}
+	ForceRebuildKeymaps();
 }
 
 void UInputSettings::GetActionNames(TArray<FName>& ActionNames) const
@@ -89,5 +90,13 @@ void UInputSettings::GetAxisNames(TArray<FName>& AxisNames) const
 	for (const FInputAxisKeyMapping& AxisMapping : AxisMappings)
 	{
 		AxisNames.AddUnique(AxisMapping.AxisName);
+	}	
+}
+
+void UInputSettings::ForceRebuildKeymaps()
+{
+	for (TObjectIterator<UPlayerInput> It; It; ++It)
+	{
+		It->ForceRebuildingKeyMaps(true);
 	}
 }
\ No newline at end of file
