Fix for clang builds
Deprecation warnings werent compatible.

diff --git a/Engine/Source/Runtime/Engine/Private/CharacterMovementComponent.cpp b/Engine/Source/Runtime/Engine/Private/CharacterMovementComponent.cpp
--- a/Engine/Source/Runtime/Engine/Private/CharacterMovementComponent.cpp
+++ b/Engine/Source/Runtime/Engine/Private/CharacterMovementComponent.cpp
@@ -2279,22 +2279,44 @@ float UCharacterMovementComponent::GetMaxJumpHeight() const
 float UCharacterMovementComponent::GetModifiedMaxAcceleration() const
 {
 	// Allow calling old deprecated function to maintain old behavior until it is removed.
-#pragma warning(push)
-#pragma warning(disable:4995)
-#pragma warning(disable:4996)
+#ifdef __clang__
+	#pragma clang diagnostic push
+	#pragma clang diagnostic ignored "-Wdeprecated-declarations"
+#else
+	#pragma warning(push)
+	#pragma warning(disable:4995)
+	#pragma warning(disable:4996)
+#endif
+
 	return CharacterOwner ? MaxAcceleration * GetMaxSpeedModifier() : 0.f;
-#pragma warning(pop)
+
+#ifdef __clang__
+	#pragma clang diagnostic pop
+#else
+	#pragma warning(pop)
+#endif
 }
 
 // TODO: deprecated, remove.
 float UCharacterMovementComponent::K2_GetModifiedMaxAcceleration() const
 {
 	// Allow calling old deprecated function to maintain old behavior until it is removed.
-#pragma warning(push)
-#pragma warning(disable:4995)
-#pragma warning(disable:4996)
+#ifdef __clang__
+	#pragma clang diagnostic push
+	#pragma clang diagnostic ignored "-Wdeprecated-declarations"
+#else
+	#pragma warning(push)
+	#pragma warning(disable:4995)
+	#pragma warning(disable:4996)
+#endif
+
 	return GetModifiedMaxAcceleration();
-#pragma warning(pop)
+
+#ifdef __clang__
+	#pragma clang diagnostic pop
+#else
+	#pragma warning(pop)
+#endif
 }
 
 
diff --git a/Engine/Source/Runtime/Engine/Private/MovementComponent.cpp b/Engine/Source/Runtime/Engine/Private/MovementComponent.cpp
--- a/Engine/Source/Runtime/Engine/Private/MovementComponent.cpp
+++ b/Engine/Source/Runtime/Engine/Private/MovementComponent.cpp
@@ -474,11 +474,22 @@ float UMovementComponent::GetMaxSpeedModifier() const
 float UMovementComponent::K2_GetMaxSpeedModifier() const
 {
 	// Allow calling old deprecated function to maintain old behavior until it is removed.
-#pragma warning(push)
-#pragma warning(disable:4995)
-#pragma warning(disable:4996)
+#ifdef __clang__
+	#pragma clang diagnostic push
+	#pragma clang diagnostic ignored "-Wdeprecated-declarations"
+#else
+	#pragma warning(push)
+	#pragma warning(disable:4995)
+	#pragma warning(disable:4996)
+#endif
+
 	return GetMaxSpeedModifier();
-#pragma warning(pop)
+
+#ifdef __clang__
+	#pragma clang diagnostic pop
+#else
+	#pragma warning(pop)
+#endif
 }
 
 // TODO: Deprecated, remove.
@@ -491,11 +502,22 @@ float UMovementComponent::GetModifiedMaxSpeed() const
 float UMovementComponent::K2_GetModifiedMaxSpeed() const
 {
 	// Allow calling old deprecated function to maintain old behavior until it is removed.
-#pragma warning(push)
-#pragma warning(disable:4995)
-#pragma warning(disable:4996)
+#ifdef __clang__
+	#pragma clang diagnostic push
+	#pragma clang diagnostic ignored "-Wdeprecated-declarations"
+#else
+	#pragma warning(push)
+	#pragma warning(disable:4995)
+	#pragma warning(disable:4996)
+#endif
+
 	return GetModifiedMaxSpeed();
-#pragma warning(pop)
+
+#ifdef __clang__
+	#pragma clang diagnostic pop
+#else
+	#pragma warning(pop)
+#endif
 }
 
 #undef LOCTEXT_NAMESPACE
