Make MustSpectate const

diff --git a/Engine/Source/Runtime/Engine/Classes/GameFramework/GameMode.h b/Engine/Source/Runtime/Engine/Classes/GameFramework/GameMode.h
--- a/Engine/Source/Runtime/Engine/Classes/GameFramework/GameMode.h
+++ b/Engine/Source/Runtime/Engine/Classes/GameFramework/GameMode.h
@@ -405,7 +405,7 @@ public:
 	virtual APlayerController* SpawnPlayerController(FVector const& SpawnLocation, FRotator const& SpawnRotation);
 
 	/** @Returns true if NewPlayer may only join the server as a spectator. */
-	virtual bool MustSpectate(APlayerController* NewPlayer);
+	virtual bool MustSpectate(APlayerController* NewPlayer) const;
 
 	/** returns default pawn class for given controller */
 	virtual UClass* GetDefaultPawnClassForController(AController* InController);
diff --git a/Engine/Source/Runtime/Engine/Private/GameMode.cpp b/Engine/Source/Runtime/Engine/Private/GameMode.cpp
--- a/Engine/Source/Runtime/Engine/Private/GameMode.cpp
+++ b/Engine/Source/Runtime/Engine/Private/GameMode.cpp
@@ -960,7 +960,7 @@ void AGameMode::InitNewPlayer(AController* NewPlayer, const TSharedPtr<FUniqueNe
 {
 }
 
-bool AGameMode::MustSpectate(APlayerController* NewPlayer)
+bool AGameMode::MustSpectate(APlayerController* NewPlayer) const
 {
 	return NewPlayer->PlayerState->bOnlySpectator;
 }
