Add utility function in Platform Math to get the min or max out of an array
Expose BP functions to find the Min/Max int, float, or byte from an array

diff --git a/Engine/Source/Runtime/Core/Public/GenericPlatform/GenericPlatformMath.h b/Engine/Source/Runtime/Core/Public/GenericPlatform/GenericPlatformMath.h
--- a/Engine/Source/Runtime/Core/Public/GenericPlatform/GenericPlatformMath.h
+++ b/Engine/Source/Runtime/Core/Public/GenericPlatform/GenericPlatformMath.h
@@ -454,6 +454,80 @@ struct FGenericPlatformMath
 		return (A<=B) ? A : B;
 	}
 
+	/**
+	* Min of Array
+	* @param	Array of templated type
+	* @param	Optional pointer for returning the index of the minimum element, if multiple minimum elements the first index is returned
+	* @return	The min value found in the array or default value if the array was empty
+	*/
+	template< class T >
+	static FORCEINLINE T Min(const TArray<T>& Values, int32* MinIndex = NULL)
+	{
+		if (Values.Num() == 0)
+		{
+			if (MinIndex)
+			{
+				*MinIndex = INDEX_NONE;
+			}
+			return T();
+		}
+
+		T CurMin = Values[0];
+		int32 CurMinIndex = 0;
+		for (int32 v = 1; v < Values.Num(); ++v)
+		{
+			const T Value = Values[v];
+			if (Value < CurMin)
+			{
+				CurMin = Value;
+				CurMinIndex = v;
+			}
+		}
+
+		if (MinIndex)
+		{
+			*MinIndex = CurMinIndex;
+		}
+		return CurMin;
+	}
+
+	/**
+	* Max of Array
+	* @param	Array of templated type
+	* @param	Optional pointer for returning the index of the maximum element, if multiple maximum elements the first index is returned
+	* @return	The max value found in the array or default value if the array was empty
+	*/
+	template< class T >
+	static FORCEINLINE T Max(const TArray<T>& Values, int32* MaxIndex = NULL)
+	{
+		if (Values.Num() == 0)
+		{
+			if (MaxIndex)
+			{
+				*MaxIndex = INDEX_NONE;
+			}
+			return T();
+		}
+
+		T CurMax = Values[0];
+		int32 CurMaxIndex = 0;
+		for (int32 v = 1; v < Values.Num(); ++v)
+		{
+			const T Value = Values[v];
+			if (CurMax < Value)
+			{
+				CurMax = Value;
+				CurMaxIndex = v;
+			}
+		}
+
+		if (MaxIndex)
+		{
+			*MaxIndex = CurMaxIndex;
+		}
+		return CurMax;
+	}
+
 	/** Test some of the tricky functions above **/
 	static void AutoTest();
 };
diff --git a/Engine/Source/Runtime/Core/Public/Windows/WindowsPlatformMath.h b/Engine/Source/Runtime/Core/Public/Windows/WindowsPlatformMath.h
--- a/Engine/Source/Runtime/Core/Public/Windows/WindowsPlatformMath.h
+++ b/Engine/Source/Runtime/Core/Public/Windows/WindowsPlatformMath.h
@@ -120,13 +120,6 @@ struct FWindowsPlatformMath : public FGenericPlatformMath
 	{
 		return 1 << CeilLogTwo(Arg);
 	}
-
-	/** Here for specialization below */
-	template< class T > 
-	static FORCEINLINE T Min( const T A, const T B )
-	{
-		return (A<=B) ? A : B;
-	}
 #endif
 };
 
diff --git a/Engine/Source/Runtime/Engine/Classes/Kismet/KismetMathLibrary.h b/Engine/Source/Runtime/Engine/Classes/Kismet/KismetMathLibrary.h
--- a/Engine/Source/Runtime/Engine/Classes/Kismet/KismetMathLibrary.h
+++ b/Engine/Source/Runtime/Engine/Classes/Kismet/KismetMathLibrary.h
@@ -343,6 +343,30 @@ class UKismetMathLibrary : public UBlueprintFunctionLibrary
 	UFUNCTION(BlueprintPure, meta=(FriendlyName = "Clamp (float)", Min="0.0", Max="1.0"), Category="Math|Float")
 	static float FClamp(float Value, float Min, float Max);
 
+	/** Returns max of all array entries and the index at which it was found. Returns value of 0 and index of -1 if the supplied array is empty. */
+	UFUNCTION(BlueprintPure, Category="Math|Integer")
+	static void MaxOfIntArray(const TArray<int32>& IntArray, int32& IndexOfMaxValue, int32& MaxValue);
+
+	/** Returns min of all array entries and the index at which it was found. Returns value of 0 and index of -1 if the supplied array is empty. */
+	UFUNCTION(BlueprintPure, Category="Math|Integer")
+	static void MinOfIntArray(const TArray<int32>& IntArray, int32& IndexOfMinValue, int32& MinValue);
+
+	/** Returns max of all array entries and the index at which it was found. Returns value of 0 and index of -1 if the supplied array is empty. */
+	UFUNCTION(BlueprintPure, Category="Math|Float")
+	static void MaxOfFloatArray(const TArray<float>& FloatArray, int32& IndexOfMaxValue, float& MaxValue);
+
+	/** Returns min of all array entries and the index at which it was found. Returns value of 0 and index of -1 if the supplied array is empty. */
+	UFUNCTION(BlueprintPure, Category="Math|Float")
+	static void MinOfFloatArray(const TArray<float>& FloatArray, int32& IndexOfMinValue, float& MinValue);
+
+	/** Returns max of all array entries and the index at which it was found. Returns value of 0 and index of -1 if the supplied array is empty. */
+	UFUNCTION(BlueprintPure, Category="Math|Byte")
+	static void MaxOfByteArray(const TArray<uint8>& ByteArray, int32& IndexOfMaxValue, uint8& MaxValue);
+
+	/** Returns min of all array entries and the index at which it was found. Returns value of 0 and index of -1 if the supplied array is empty. */
+	UFUNCTION(BlueprintPure, Category="Math|Byte")
+	static void MinOfByteArray(const TArray<uint8>& ByteArray, int32& IndexOfMinValue, uint8& MinValue);
+
 	/* Linearly interpolates between A and B based on Alpha (100% of A when Alpha=0 and 100% of B when Alpha=1) */
 	UFUNCTION(BlueprintPure, Category="Math|Float")
 	static float Lerp(float A, float B, float Alpha);
diff --git a/Engine/Source/Runtime/Engine/Private/KismetMathLibrary.cpp b/Engine/Source/Runtime/Engine/Private/KismetMathLibrary.cpp
--- a/Engine/Source/Runtime/Engine/Private/KismetMathLibrary.cpp
+++ b/Engine/Source/Runtime/Engine/Private/KismetMathLibrary.cpp
@@ -524,6 +524,36 @@ float UKismetMathLibrary::FClamp(float V, float A, float B)
 	return FMath::Clamp(V, A, B);
 }	
 
+void UKismetMathLibrary::MaxOfIntArray(const TArray<int32>& IntArray, int32& IndexOfMaxValue, int32& MaxValue)
+{
+	MaxValue = FMath::Max(IntArray, &IndexOfMaxValue);
+}
+
+void UKismetMathLibrary::MinOfIntArray(const TArray<int32>& IntArray, int32& IndexOfMinValue, int32& MinValue)
+{
+	MinValue = FMath::Min<int32>(IntArray, &IndexOfMinValue);
+}
+
+void UKismetMathLibrary::MaxOfFloatArray(const TArray<float>& FloatArray, int32& IndexOfMaxValue, float& MaxValue)
+{
+	MaxValue = FMath::Max(FloatArray, &IndexOfMaxValue);
+}
+
+void UKismetMathLibrary::MinOfFloatArray(const TArray<float>& FloatArray, int32& IndexOfMinValue, float& MinValue)
+{
+	MinValue = FMath::Min(FloatArray, &IndexOfMinValue);
+}
+
+void UKismetMathLibrary::MaxOfByteArray(const TArray<uint8>& ByteArray, int32& IndexOfMaxValue, uint8& MaxValue)
+{
+	MaxValue = FMath::Max(ByteArray, &IndexOfMaxValue);
+}
+
+void UKismetMathLibrary::MinOfByteArray(const TArray<uint8>& ByteArray, int32& IndexOfMinValue, uint8& MinValue)
+{
+	MinValue = FMath::Min(ByteArray, &IndexOfMinValue);
+}
+
 float UKismetMathLibrary::Lerp(float A, float B, float V)
 {
 	return A + V*(B-A);
