Fix body instance simulating getting unset

diff --git a/Engine/Source/Runtime/Engine/Private/PhysicsEngine/BodyInstance.cpp b/Engine/Source/Runtime/Engine/Private/PhysicsEngine/BodyInstance.cpp
--- a/Engine/Source/Runtime/Engine/Private/PhysicsEngine/BodyInstance.cpp
+++ b/Engine/Source/Runtime/Engine/Private/PhysicsEngine/BodyInstance.cpp
@@ -358,7 +358,7 @@ void FBodyInstance::UpdatePhysicalMaterials()
 
 	int32 NumSyncShapes = 0;
 	TArray<PxShape*> AllShapes = GetAllShapes(NumSyncShapes);
-
+	
 	for(int32 ShapeIdx=0; ShapeIdx<AllShapes.Num(); ShapeIdx++)
 	{
 		PxShape* PShape = AllShapes[ShapeIdx];
@@ -394,6 +394,7 @@ void FBodyInstance::UpdatePhysicalMaterials()
 		}
 	}
 
+
 	if (RigidActorSync != NULL)
 	{
 		SCENE_UNLOCK_WRITE(RigidActorSync->getScene());
@@ -1747,11 +1748,11 @@ void FBodyInstance::UpdateInstanceSimulatePhysics(bool bIgnoreParent)
 	// In skeletal case, we need both our bone and skelcomponent flag to be true.
 	// This might be 'and'ing us with ourself, but thats fine.
 	const bool bUseSimulate = IsInstanceSimulatingPhysics(bIgnoreParent);
-
+	bool bInitialized = false;
 #if WITH_PHYSX
-	PxRigidDynamic* PRigidDynamic = GetPxRigidDynamic();
-	if(PRigidDynamic != NULL)
+	if (PxRigidDynamic* PRigidDynamic = GetPxRigidDynamic())
 	{
+		bInitialized = true;
 		// If we want it fixed, and it is currently not kinematic
 		bool bNewKinematic = (bUseSimulate == false);
 		{
@@ -1764,20 +1765,26 @@ void FBodyInstance::UpdateInstanceSimulatePhysics(bool bIgnoreParent)
 #if WITH_BOX2D
 	if (BodyInstancePtr != NULL)
 	{
+		bInitialized = true;
 		BodyInstancePtr->SetType(bUseSimulate ? b2_dynamicBody : b2_kinematicBody);
 	}
 #endif
 
-	if (bUseSimulate)
+	//In the original physx only implementation this was wrapped in a PRigidDynamic != NULL check.
+	//We use bInitialized to check rigid actor has been created in either engine because if we haven't even initialized yet, we don't want to undo our settings
+	if (bInitialized)
 	{
-		PhysicsBlendWeight = 1.f;
-	}
-	else
-	{
-		PhysicsBlendWeight = 0.f;
-	}
+		if (bUseSimulate)
+		{
+			PhysicsBlendWeight = 1.f;
+		}
+		else
+		{
+			PhysicsBlendWeight = 0.f;
+		}
 
-	bSimulatePhysics = bUseSimulate;
+		bSimulatePhysics = bUseSimulate;
+	}
 }
 
 bool FBodyInstance::IsDynamic() const
@@ -2353,7 +2360,6 @@ int32 GetNumSimShapes(PxRigidDynamic* PRigidDynamic)
 }
 #endif // WITH_PHYSX
 
-
 void FBodyInstance::UpdateMassProperties()
 {
 	UPhysicalMaterial* PhysMat = GetSimplePhysicalMaterial();
