#ue4 - Fix various issues with swimming in CharacterMovementComponent.

- Fix lack of SlideAlongSurface() code when StepUpFails, or when it isn't going to help (now very similar to Flying implementation).
- CanEverSwim() wasn't checked for various workflows (initially starting in a water volume, or when landing after falling).

diff --git a/Engine/Source/Runtime/Engine/Private/CharacterMovementComponent.cpp b/Engine/Source/Runtime/Engine/Private/CharacterMovementComponent.cpp
--- a/Engine/Source/Runtime/Engine/Private/CharacterMovementComponent.cpp
+++ b/Engine/Source/Runtime/Engine/Private/CharacterMovementComponent.cpp
@@ -511,7 +511,7 @@ float UCharacterMovementComponent::GetNetworkSafeRandomAngleDegrees() const
 void UCharacterMovementComponent::SetDefaultMovementMode()
 {
 	// check for water volume
-	if ( IsInWater() )
+	if ( IsInWater() && CanEverSwim() )
 	{
 		SetMovementMode(DefaultWaterMovementMode);
 	}
@@ -2399,9 +2399,9 @@ void UCharacterMovementComponent::PhysFlying(float deltaTime, int32 Iterations)
 
 	if( Hit.Time < 1.f && CharacterOwner )
 	{
-		FVector GravDir = FVector(0.f,0.f,-1.f);
-		FVector VelDir = Velocity.SafeNormal();
-		float UpDown = GravDir | VelDir;
+		const FVector GravDir = FVector(0.f,0.f,-1.f);
+		const FVector VelDir = Velocity.SafeNormal();
+		const float UpDown = GravDir | VelDir;
 
 		bool bSteppedUp = false;
 		if( (FMath::Abs(Hit.ImpactNormal.Z) < 0.2f) && (UpDown < 0.5f) && (UpDown > -0.2f) && CanStepUp(Hit))
@@ -2467,18 +2467,36 @@ void UCharacterMovementComponent::PhysSwimming(float deltaTime, int32 Iterations
 
 	if ( Hit.Time < 1.f && CharacterOwner)
 	{
-		float stepZ = CharacterOwner->GetActorLocation().Z;
-		FVector RealVelocity = Velocity;
-		Velocity.Z = 1.f;	// HACK: since will be moving up, in case pawn leaves the water
-		StepUp(-1.f*Hit.ImpactNormal, Adjusted * (1.f - Hit.Time), Hit);
-		//may have left water - if so, script might have set new physics mode
-		if ( !IsSwimming() )
+		const FVector GravDir = FVector(0.f,0.f,-1.f);
+		const FVector VelDir = Velocity.SafeNormal();
+		const float UpDown = GravDir | VelDir;
+
+		bool bSteppedUp = false;
+		if( (FMath::Abs(Hit.ImpactNormal.Z) < 0.2f) && (UpDown < 0.5f) && (UpDown > -0.2f) && CanStepUp(Hit))
 		{
-			StartNewPhysics(remainingTime, Iterations);
-			return;
+			float stepZ = CharacterOwner->GetActorLocation().Z;
+			const FVector RealVelocity = Velocity;
+			Velocity.Z = 1.f;	// HACK: since will be moving up, in case pawn leaves the water
+			bSteppedUp = StepUp(GravDir, Adjusted * (1.f - Hit.Time), Hit);
+			if (bSteppedUp)
+			{
+				//may have left water - if so, script might have set new physics mode
+				if (!IsSwimming())
+				{
+					StartNewPhysics(remainingTime, Iterations);
+					return;
+				}
+				OldLocation.Z = CharacterOwner->GetActorLocation().Z + (OldLocation.Z - stepZ);
+			}
+			Velocity = RealVelocity;
+		}
+
+		if (!bSteppedUp)
+		{
+			//adjust and try again
+			HandleImpact(Hit, deltaTime, Adjusted);
+			SlideAlongSurface(Adjusted, (1.f - Hit.Time), Hit.Normal, Hit, true);
 		}
-		Velocity = RealVelocity;
-		OldLocation.Z = CharacterOwner->GetActorLocation().Z + (OldLocation.Z - stepZ);
 	}
 
 	if( !HasRootMotion() && !bJustTeleported && ((deltaTime - remainingTime) > KINDA_SMALL_NUMBER) && CharacterOwner )
@@ -3412,7 +3430,7 @@ void UCharacterMovementComponent::SetPostLandedPhysics(const FHitResult& Hit)
 {
 	if( CharacterOwner )
 	{
-		if ( GetPhysicsVolume()->bWaterVolume )
+		if ( GetPhysicsVolume()->bWaterVolume && CanEverSwim() )
 		{
 			SetMovementMode(MOVE_Swimming);
 		}
