#ue4 - Characters can swim by default (bCanSwim = true on NavAgentProps).

diff --git a/Engine/Source/Runtime/Engine/Private/Character.cpp b/Engine/Source/Runtime/Engine/Private/Character.cpp
--- a/Engine/Source/Runtime/Engine/Private/Character.cpp
+++ b/Engine/Source/Runtime/Engine/Private/Character.cpp
@@ -78,6 +78,7 @@ ACharacter::ACharacter(const class FPostConstructInitializeProperties& PCIP)
 
 		CharacterMovement->GetNavAgentProperties()->bCanJump = true;
 		CharacterMovement->GetNavAgentProperties()->bCanWalk = true;
+		CharacterMovement->GetNavAgentProperties()->bCanSwim = true;
 		CharacterMovement->SetJumpAllowed(true);
 	}
 
