GetComponentsByClass for Blueprints

diff --git a/Engine/Source/Runtime/Engine/Classes/GameFramework/Actor.h b/Engine/Source/Runtime/Engine/Classes/GameFramework/Actor.h
--- a/Engine/Source/Runtime/Engine/Classes/GameFramework/Actor.h
+++ b/Engine/Source/Runtime/Engine/Classes/GameFramework/Actor.h
@@ -2007,6 +2007,11 @@ public:
 	UFUNCTION()
 	virtual UActorComponent* GetComponentByClass(TSubclassOf<UActorComponent> ComponentClass);
 
+	/* Gets all the components that inherit from the given class.
+		Currently returns an array of UActorComponent which must be cast to the correct type. */
+	UFUNCTION(BlueprintCallable, Category="Actor", meta=(ComponentClass="ActorComponent"))
+	TArray<UActorComponent*> GetComponentsByClass(TSubclassOf<UActorComponent> ComponentClass) const;
+
 	/** Templatized version for syntactic nicety. */
 	template<class T>
 	T* FindComponentByClass() const
diff --git a/Engine/Source/Runtime/Engine/Private/Actor.cpp b/Engine/Source/Runtime/Engine/Private/Actor.cpp
--- a/Engine/Source/Runtime/Engine/Private/Actor.cpp
+++ b/Engine/Source/Runtime/Engine/Private/Actor.cpp
@@ -2058,14 +2058,36 @@ UActorComponent* AActor::GetComponentByClass(TSubclassOf<UActorComponent> Compon
 	return FoundComponent;
 }
 
+TArray<UActorComponent*> AActor::GetComponentsByClass(TSubclassOf<UActorComponent> ComponentClass) const
+{
+	if (ComponentClass == UActorComponent::StaticClass())
+	{
+		return OwnedComponents;
+	}
+
+	TArray<UActorComponent*> ValidComponents;
+	if (*ComponentClass)
+	{
+		for (UActorComponent* Component : OwnedComponents)
+		{
+			if (Component && Component->IsA(ComponentClass))
+			{
+				ValidComponents.Add(Component);
+			}
+		}
+	}
+	
+	return ValidComponents;
+}
+
 void AActor::DisableComponentsSimulatePhysics()
 {
 	TArray<UPrimitiveComponent*> Components;
 	GetComponents(Components);
 
-	for (int32 i = 0; i < Components.Num(); i++)
+	for (UPrimitiveComponent* Component : Components)
 	{
-		Components[i]->SetSimulatePhysics(false);
+		Component->SetSimulatePhysics(false);
 	}
 }
 
