Use guide sequence counter to clean package maps across seamless travel

* Removed code that would force exporting of guids during seamless travel, no longer needed with guid sequence logic
* Added ability to reject old guid packets
* No longer even attempt to append guid packets to final fragmented bunch, removes the worry of them not being processed if they were previously rejected due to out of order
* We no longer need to byte align guid packetsguid sequence

diff --git a/Engine/Source/Runtime/CoreUObject/Public/UObject/CoreNet.h b/Engine/Source/Runtime/CoreUObject/Public/UObject/CoreNet.h
--- a/Engine/Source/Runtime/CoreUObject/Public/UObject/CoreNet.h
+++ b/Engine/Source/Runtime/CoreUObject/Public/UObject/CoreNet.h
@@ -120,10 +120,6 @@ class COREUOBJECT_API UPackageMap : public UObject
 	// @todo document
 	virtual bool		SerializeName( FArchive& Ar, FName& Name );
 
-	virtual void		ResetPackageMap() { }
-
-	virtual void		SetLocked( bool L ) { }
-
 	virtual UObject *	ResolvePathAndAssignNetGUID( const FNetworkGUID & NetGUID, const FString & PathName ) { return NULL; }
 
 	virtual bool		SerializeNewActor(FArchive & Ar, class UActorChannel * Channel, class AActor *& Actor) { return false; }
diff --git a/Engine/Source/Runtime/Engine/Classes/Engine/NetDriver.h b/Engine/Source/Runtime/Engine/Classes/Engine/NetDriver.h
--- a/Engine/Source/Runtime/Engine/Classes/Engine/NetDriver.h
+++ b/Engine/Source/Runtime/Engine/Classes/Engine/NetDriver.h
@@ -495,11 +495,6 @@ public:
 	 */
 	bool VerifyPackageInfo(FPackageInfo& Info);
 
-	/** Flushes and clears all packagemaps on this driver and driver's connections */
-	ENGINE_API virtual void ResetPackageMaps();
-
-	ENGINE_API virtual void LockPackageMaps();
-
 	ENGINE_API virtual void CleanPackageMaps();
 
 	ENGINE_API void PreSeamlessTravelGarbageCollect();
diff --git a/Engine/Source/Runtime/Engine/Classes/Engine/PackageMapClient.h b/Engine/Source/Runtime/Engine/Classes/Engine/PackageMapClient.h
--- a/Engine/Source/Runtime/Engine/Classes/Engine/PackageMapClient.h
+++ b/Engine/Source/Runtime/Engine/Classes/Engine/PackageMapClient.h
@@ -42,13 +42,16 @@ public:
 class FNetGuidCacheObject
 {
 public:
-	FNetGuidCacheObject() : bNoLoad( 0 ), bIgnoreWhenMissing( 0 ), bIsPending( 0 ), bIsBroken( 0 ) {}
+	FNetGuidCacheObject() : GuidSequence( 0 ), bNoLoad( 0 ), bIgnoreWhenMissing( 0 ), bIsPending( 0 ), bIsBroken( 0 )
+	{
+	}
 
 	TWeakObjectPtr< UObject >	Object;
 
 	// These fields are set when this guid is static
 	FNetworkGUID				OuterGUID;
 	FName						PathName;
+	int32						GuidSequence;				// We remember the guid sequence this net guid was generated on so we can tell how old it is
 
 	uint8						bNoLoad				: 1;	// Don't load this, only do a find
 	uint8						bIgnoreWhenMissing	: 1;	// Don't warn when this asset can't be found or loaded
@@ -61,7 +64,6 @@ class ENGINE_API FNetGUIDCache
 public:
 	FNetGUIDCache( UNetDriver * InDriver );
 
-	void			Reset();
 	void			CleanReferences();
 	bool			SupportsObject( const UObject * Object );
 	bool			IsDynamicObject( const UObject * Object );
@@ -80,6 +82,8 @@ public:
 	TMap< TWeakObjectPtr< UObject >, FNetworkGUID >	NetGUIDLookup;
 	int32											UniqueNetIDs[2];
 
+	int32											GuidSequence;
+
 	UNetDriver *									Driver;
 
 #if !(UE_BUILD_SHIPPING || UE_BUILD_TEST)
@@ -98,7 +102,6 @@ class UPackageMapClient : public UPackageMap
 	,	Connection(InConnection)
 	{
 		GuidCache				= InNetGUIDCache;
-		Locked					= false;
 		ExportNetGUIDCount		= 0;
 		IsExportingNetGUIDBunch = false;
 		IsSerializingNewActor	= false;
@@ -120,10 +123,6 @@ class UPackageMapClient : public UPackageMap
 	
 	virtual bool WriteObject( FArchive& Ar, UObject* Outer, FNetworkGUID NetGUID, FString ObjName ) override;
 
-	virtual void ResetPackageMap() override;
-
-	virtual void SetLocked(bool L) { Locked = L; }
-
 	// UPackageMapClient Connection specific methods
 
 	bool NetGUIDHasBeenAckd(FNetworkGUID NetGUID);
@@ -163,7 +162,6 @@ protected:
 	
 	bool IsNetGUIDAuthority();
 
-	bool Locked;
 	bool IsExportingNetGUIDBunch;
 
 	class UNetConnection* Connection;
diff --git a/Engine/Source/Runtime/Engine/Private/DataChannel.cpp b/Engine/Source/Runtime/Engine/Private/DataChannel.cpp
--- a/Engine/Source/Runtime/Engine/Private/DataChannel.cpp
+++ b/Engine/Source/Runtime/Engine/Private/DataChannel.cpp
@@ -373,8 +373,10 @@ bool UChannel::ReceivedNextBunch( FInBunch & Bunch, bool & bOutSkipAck )
 	// If its not a partial bunch, of it completes a partial bunch, we can call ReceivedSequencedBunch to actually handle it
 	
 	// Note this bunch's retirement.
-	if( Bunch.bReliable )
+	if ( Bunch.bReliable )
+	{
 		Connection->InReliable[Bunch.ChIndex] = Bunch.ChSequence;
+	}
 
 	FInBunch* HandleBunch = &Bunch;
 	if (Bunch.bPartial)
@@ -405,7 +407,7 @@ bool UChannel::ReceivedNextBunch( FInBunch & Bunch, bool & bOutSkipAck )
 			}
 
 			InPartialBunch = new FInBunch(Bunch, false);
-			if (Bunch.GetBitsLeft() > 0)
+			if ( !Bunch.bHasGUIDs && Bunch.GetBitsLeft() > 0 )
 			{
 				check( Bunch.GetBitsLeft() % 8 == 0); // Starting partial bunches should always be byte aligned.
 
@@ -430,14 +432,14 @@ bool UChannel::ReceivedNextBunch( FInBunch & Bunch, bool & bOutSkipAck )
 				// Merge.
 				UE_LOG(LogNetPartialBunch, Verbose, TEXT("Merging Partial Bunch: %d Bytes"), Bunch.GetBytesLeft() );
 
-				if (Bunch.GetBitsLeft() > 0)
+				if ( !Bunch.bHasGUIDs && Bunch.GetBitsLeft() > 0 )
 				{
 					InPartialBunch->AppendDataFromChecked( Bunch.GetDataPosChecked(), Bunch.GetBitsLeft() );
 				}
 
 				// Only the final partial bunch should ever be non byte aligned. This is enforced during partial bunch creation
 				// This is to ensure fast copies/appending of partial bunches. The final partial bunch may be non byte aligned.
-				check( Bunch.bPartialFinal || Bunch.GetBitsLeft() % 8 == 0);
+				check( Bunch.bHasGUIDs || Bunch.bPartialFinal || Bunch.GetBitsLeft() % 8 == 0 );
 
 				InPartialBunch->ChSequence++;
 
diff --git a/Engine/Source/Runtime/Engine/Private/NetworkDriver.cpp b/Engine/Source/Runtime/Engine/Private/NetworkDriver.cpp
--- a/Engine/Source/Runtime/Engine/Private/NetworkDriver.cpp
+++ b/Engine/Source/Runtime/Engine/Private/NetworkDriver.cpp
@@ -2832,40 +2832,6 @@ bool UNetDriver::VerifyPackageInfo(FPackageInfo& Info)
 	return true;
 }
 
-void UNetDriver::ResetPackageMaps()
-{
-	if ( GuidCache.IsValid() )
-	{
-		GuidCache->Reset();
-	}
-
-	if (ServerConnection && ServerConnection->PackageMap)
-	{
-		ServerConnection->PackageMap->ResetPackageMap();
-	}
-
-	for (auto It = ClientConnections.CreateIterator(); It; ++It)
-	{
-		UNetConnection *connection = *It;
-		if (connection->PackageMap)
-		{
-			connection->PackageMap->ResetPackageMap();
-		}
-	}
-}
-
-void UNetDriver::LockPackageMaps()
-{
-	for (auto It = ClientConnections.CreateIterator(); It; ++It)
-	{
-		UNetConnection *connection = *It;
-		if (connection->PackageMap)
-		{
-			connection->PackageMap->SetLocked(true);
-		}
-	}
-}
-
 void UNetDriver::SetWorld(class UWorld* InWorld)
 {
 	if (World)
@@ -2906,8 +2872,9 @@ void UNetDriver::ResetGameWorldState()
 
 void UNetDriver::CleanPackageMaps()
 {
-	if ( GuidCache.IsValid() )
-	{
+	if ( IsServer() && GuidCache.IsValid()  )
+	{ 
+		// Only the server makes this call initially, client will call when the GuidSequence increments on their end
 		GuidCache->CleanReferences();
 	}
 }
diff --git a/Engine/Source/Runtime/Engine/Private/PackageMapClient.cpp b/Engine/Source/Runtime/Engine/Private/PackageMapClient.cpp
--- a/Engine/Source/Runtime/Engine/Private/PackageMapClient.cpp
+++ b/Engine/Source/Runtime/Engine/Private/PackageMapClient.cpp
@@ -738,7 +738,8 @@ bool UPackageMapClient::ExportNetGUID( FNetworkGUID NetGUID, const UObject * Obj
 
 			ExportNetGUIDCount = 0;
 			*CurrentExportBunch << ExportNetGUIDCount;
-			NET_CHECKSUM(*CurrentExportBunch);
+			*CurrentExportBunch << GuidCache->GuidSequence;
+			NET_CHECKSUM( *CurrentExportBunch );
 		}
 
 		if ( CurrentExportNetGUIDs.Num() != 0 )
@@ -816,6 +817,7 @@ void UPackageMapClient::ExportNetGUIDHeader()
 	FBitWriterMark Restore(Out);
 	Reset.PopWithoutClear(Out);
 	Out << ExportNetGUIDCount;
+	Out << GuidCache->GuidSequence;
 	Restore.PopWithoutClear(Out);
 
 	// If we've written new NetGUIDs to the 'bunch' set (current+1)
@@ -849,7 +851,25 @@ void UPackageMapClient::ReceiveNetGUIDBunch( FInBunch &InBunch )
 
 	int32 NumGUIDsInBunch = 0;
 	InBunch << NumGUIDsInBunch;
-	
+
+	int32 NewGuidSequence = -1;
+
+	InBunch << NewGuidSequence;
+
+	if ( NewGuidSequence < GuidCache->GuidSequence )
+	{
+		// Older sequence, ignore
+		UE_LOG( LogNetPackageMap, Warning, TEXT( "UPackageMapClient::ReceiveNetGUIDBunch: NewGuidSequence < GuidSequence (%i / %i)" ), GuidCache->GuidSequence, NewGuidSequence );
+		return;
+	}
+	else if ( NewGuidSequence > GuidCache->GuidSequence )
+	{
+		// Newer sequence, this is our signal to clean the package to stay in sync with the server
+		UE_LOG( LogNetPackageMap, Verbose, TEXT( "UPackageMapClient::ReceiveNetGUIDBunch: NewGuidSequence > GuidSequence (%i / %i)" ), GuidCache->GuidSequence, NewGuidSequence );
+		GuidCache->GuidSequence = NewGuidSequence;
+		GuidCache->CleanReferences();
+	}
+
 	static const int32 MAX_GUID_COUNT = 2048;
 
 	if ( NumGUIDsInBunch > MAX_GUID_COUNT )
@@ -877,8 +897,6 @@ void UPackageMapClient::ReceiveNetGUIDBunch( FInBunch &InBunch )
 		NumGUIDsRead++;
 	}
 
-	InBunch.EatByteAlign();
-
 	UE_LOG(LogNetPackageMap, Log, TEXT("UPackageMapClient::ReceiveNetGUIDBunch end. BitPos: %d"), InBunch.GetPosBits() );
 	IsExportingNetGUIDBunch = false;
 }
@@ -932,16 +950,6 @@ void UPackageMapClient::NotifyBunchCommit( const int32 OutPacketId, const TArray
 	check( OutPacketId > GUID_PACKET_ACKED );	// Assumptions break if this isn't true ( We assume ( OutPacketId > GUID_PACKET_ACKED ) == PENDING )
 	check( ExportNetGUIDs.Num() > 0 );			// We should have never created this bunch if there are not exports
 
-#if 1
-	// Not sure we actually need to do this.  
-	// Assuming we are acking/nacking property, and those systems are working correctly, I don't see the need.
-	// Keeping it for now though, maybe I'm missing something.
-	if ( Locked )
-	{
-		return;
-	}
-#endif
-
 	for ( int32 i = 0; i < ExportNetGUIDs.Num(); i++ )
 	{
 		if ( !NetGUIDAckStatus.Contains( ExportNetGUIDs[i] ) )
@@ -1112,7 +1120,7 @@ bool UPackageMapClient::ShouldSendFullPath( const UObject* Object, const FNetwor
 		return true;
 	}
 
-	return Locked || !NetGUIDHasBeenAckd( NetGUID );
+	return !NetGUIDHasBeenAckd( NetGUID );
 }
 
 /**
@@ -1186,15 +1194,6 @@ bool UPackageMapClient::IsNetGUIDAuthority()
 	return GuidCache->IsNetGUIDAuthority();
 }
 
-/** Resets packagemap state for server travel */
-void UPackageMapClient::ResetPackageMap()
-{
-	NetGUIDExportCountMap.Empty();
-	CurrentExportNetGUIDs.Empty();
-	NetGUIDAckStatus.Empty();
-	PendingAckGUIDs.Empty();
-}
-
 /**	
  *	Returns stats for NetGUID usage
  */
@@ -1280,13 +1279,7 @@ UObject * UPackageMapClient::GetObjectFromNetGUID( const FNetworkGUID & NetGUID
 FNetGUIDCache::FNetGUIDCache( UNetDriver * InDriver ) : Driver( InDriver )
 {
 	UniqueNetIDs[0] = UniqueNetIDs[1] = 0;
-}
-
-void FNetGUIDCache::Reset()
-{
-	ObjectLookup.Empty();
-	NetGUIDLookup.Empty();
-	UniqueNetIDs[0] = UniqueNetIDs[1] = 0;
+	GuidSequence = 0;
 }
 
 class FArchiveCountMemGUID : public FArchive
@@ -1299,51 +1292,68 @@ public:
 
 void FNetGUIDCache::CleanReferences()
 {
-	// Cleanup any guids that are dynamic
-	// We keep static guids since the client most likely hasn't GC'd the static objects the exact same way we have
-	// Since we can't rely on the the client clearing static guids in the same we we would, we need to just keep them
-	// The only issue with this, is this list can grow forever
-	// To fix this, we can ultimately re-use static guids, to at least bound the list size by the number of packages on disk
+	if ( IsNetGUIDAuthority() )
+	{
+		// The server increments this to signify to clients to also call CleanReferences when they detect this change
+		GuidSequence++;
+	}
+
+	// Server:
+	//	Clean all static or non valid guids
+	//	This means static objects will have new guids assigned, we're assuming clients are ok with this
+	// Client:
+	//	Clean all guids that are old enough
+	//	This means guids that are older than SEAMLESS_TRAVEL_COUNT_FOR_CLEAN seamless travels
+	//		We use a GuidSequence counter to determine how often the server has seamlessly traveled
+	// 
+	// It's possible for guids that have been cleaned to get lost if they were sent from a level that was 
+	// was older than (SEAMLESS_TRAVEL_COUNT_FOR_CLEAN) seamless travels, but we're assuming the game code 
+	// is resilient to this, and only needs information from the previous level at most
 	for ( auto It = ObjectLookup.CreateIterator(); It; ++It )
 	{
-		check( !It.Key().IsDefault() );
-		check( It.Key().IsStatic() != It.Key().IsDynamic() );
-
-		if ( It.Value().Object.IsValid() || It.Key().IsStatic() )
+		if ( IsNetGUIDAuthority() )
 		{
-			// Either the object is still around, or it's a fully stably named object (which we always keep those registered)
-			continue;
+			// Server cleans references that are NULL and static
+			// (server will assign a new guid to static references after a seamless travel)
+			if ( !It.Value().Object.IsValid() || It.Key().IsStatic() )
+			{
+				It.RemoveCurrent();
+			}
 		}
+		else
+		{
+			static const int32 SEAMLESS_TRAVEL_COUNT_FOR_CLEAN = 2;
 
-		It.RemoveCurrent();
+			// Clients only clean references that are older than SEAMLESS_TRAVEL_COUNT_FOR_CLEAN 
+			//	(this guid has seamless traveled this many times, so we can safely punt it once it's old enough)
+			if ( ( GuidSequence - It.Value().GuidSequence ) > SEAMLESS_TRAVEL_COUNT_FOR_CLEAN )
+			{
+				It.RemoveCurrent();
+			}
+		}
 	}
 
 	for ( auto It = NetGUIDLookup.CreateIterator(); It; ++It )
 	{
-		if ( !It.Key().IsValid() )
+		if ( !It.Key().IsValid() || !ObjectLookup.Contains( It.Value() ) )
 		{
 			It.RemoveCurrent();
 		}
 	}
 
 	// Sanity check
+	// (make sure look-ups are reciprocal)
 	for ( auto It = ObjectLookup.CreateIterator(); It; ++It )
 	{
-		check( It.Value().Object.IsValid() || It.Key().IsStatic() );
+		check( !It.Key().IsDefault() );
+		check( It.Key().IsStatic() != It.Key().IsDynamic() );
 
-		if ( It.Value().Object.IsValid() )
-		{
-			checkf( !It.Value().Object.IsValid() || NetGUIDLookup.FindRef( It.Value().Object ) == It.Key(), TEXT("Failed to validate ObjectLookup map in UPackageMap. Object '%s' was not in the NetGUIDLookup map with with value '%s'."), *It.Value().Object.Get()->GetPathName(), *It.Key().ToString());
-		}
-		else
-		{
-			for ( auto It2 = NetGUIDLookup.CreateIterator(); It2; ++It2 )
-			{
-				checkf( It2.Value() != It.Key(), TEXT("Failed to validate ObjectLookup map in UPackageMap. Object '%s' was in the NetGUIDLookup map with with value '%s'."), *It.Value().Object.Get()->GetPathName(), *It.Key().ToString());
-			}
-		}
+		checkf( !It.Value().Object.IsValid() || NetGUIDLookup.FindRef( It.Value().Object ) == It.Key(), TEXT( "Failed to validate ObjectLookup map in UPackageMap. Object '%s' was not in the NetGUIDLookup map with with value '%s'." ), *It.Value().Object.Get()->GetPathName(), *It.Key().ToString() );
 	}
 
+	// Server lists should be completely in sync
+	check( ObjectLookup.Num() == NetGUIDLookup.Num() || !IsNetGUIDAuthority() );
+
 	for ( auto It = NetGUIDLookup.CreateIterator(); It; ++It )
 	{
 		check( It.Key().IsValid() );
@@ -1552,6 +1562,7 @@ void FNetGUIDCache::RegisterNetGUID_Internal( const FNetworkGUID & NetGUID, cons
 {
 	// We're pretty strict in this function, we expect everything to have been handled before we get here
 	check( !ObjectLookup.Contains( NetGUID ) );
+	check( CacheObject.GuidSequence == GuidSequence );
 
 	ObjectLookup.Add( NetGUID, CacheObject );
 
@@ -1588,8 +1599,9 @@ void FNetGUIDCache::RegisterNetGUID_Server( const FNetworkGUID & NetGUID, const
 
 	FNetGuidCacheObject CacheObject;
 
-	CacheObject.Object = Object;
-
+	CacheObject.Object			= Object;
+	CacheObject.GuidSequence	= GuidSequence;
+	
 	RegisterNetGUID_Internal( NetGUID, CacheObject );
 }
 
@@ -1658,7 +1670,8 @@ void FNetGUIDCache::RegisterNetGUID_Client( const FNetworkGUID & NetGUID, const
 
 	FNetGuidCacheObject CacheObject;
 
-	CacheObject.Object = Object;
+	CacheObject.Object			= Object;
+	CacheObject.GuidSequence	= GuidSequence;
 
 	RegisterNetGUID_Internal( NetGUID, CacheObject );
 }
@@ -1710,6 +1723,7 @@ void FNetGUIDCache::RegisterNetGUIDFromPath_Client( const FNetworkGUID & NetGUID
 	CacheObject.OuterGUID			= OuterGUID;
 	CacheObject.bNoLoad				= bNoLoad;
 	CacheObject.bIgnoreWhenMissing	= bIgnoreWhenMissing;
+	CacheObject.GuidSequence		= GuidSequence;
 
 	RegisterNetGUID_Internal( NetGUID, CacheObject );
 }
@@ -1815,15 +1829,17 @@ UObject * FNetGUIDCache::GetObjectFromNetGUID( const FNetworkGUID & NetGUID )
 	// Assign the resolved object to this guid
 	CacheObjectPtr->Object = Object;		
 
-	if ( NetGUIDLookup.Contains( Object ) )
+	// Assign the guid to the object 
+	// Don't allow the assignment if this guid is from an older sequence
+	if ( CacheObjectPtr->GuidSequence == GuidSequence )
 	{
-		// This can happen to static guids on the client. The server deleted this object, but for some reason the client did not
-		// When this happens, the server will give the object a different guid, and we will just conform
-		ObjectLookup.Remove( NetGUIDLookup.FindRef( Object ) );
+		// We're good, assign the guid
+		NetGUIDLookup.Add( Object, NetGUID );
+	}
+	else
+	{
+		UE_LOG( LogNetPackageMap, Warning, TEXT( "GetObjectFromNetGUID: Attempt to reassign guid from older sequence (%i / %i). Path: %s, Outer: %s, NetGUID: %s" ), GuidSequence, CacheObjectPtr->GuidSequence, *CacheObjectPtr->PathName.ToString(), ObjOuter != NULL ? *ObjOuter->GetPathName() : TEXT( "NULL" ), *NetGUID.ToString() );
 	}
-
-	// Make sure the object is in the GUIDToObjectLookup.
-	NetGUIDLookup.Add( Object, NetGUID );
 
 	return Object;
 }
diff --git a/Engine/Source/Runtime/Engine/Private/PlayerController.cpp b/Engine/Source/Runtime/Engine/Private/PlayerController.cpp
--- a/Engine/Source/Runtime/Engine/Private/PlayerController.cpp
+++ b/Engine/Source/Runtime/Engine/Private/PlayerController.cpp
@@ -426,7 +426,6 @@ void APlayerController::ServerNotifyLoadedWorld_Implementation(FName WorldPackag
 		if (Connection != NULL)
 		{
 			Connection->ClientWorldPackageName = WorldPackageName;
-			Connection->PackageMap->SetLocked(false);
 		}
 
 		// if both the server and this client have completed the transition, handle it
diff --git a/Engine/Source/Runtime/Engine/Private/World.cpp b/Engine/Source/Runtime/Engine/Private/World.cpp
--- a/Engine/Source/Runtime/Engine/Private/World.cpp
+++ b/Engine/Source/Runtime/Engine/Private/World.cpp
@@ -4072,15 +4072,6 @@ bool FSeamlessTravelHandler::StartTravel(UWorld* InCurrentWorld, const FURL& InU
 			bPauseAtMidpoint = false;
 			bNeedCancelCleanUp = false;
 			
-			if (CurrentWorld->NetDriver)
-			{
-				// Lock the packagemap from acking any new NetGUIDs during server travel.
-				// While locked, packagemap will serialize everything as full name/path pair due to unreliability of 
-				// when exactly GC will be executed. Once clients are fully connected, packagemap is unlocked and returns to
-				// 'optimized' mode.
-				CurrentWorld->NetDriver->LockPackageMaps();
-			}
-
 			FString TransitionMap = GetDefault<UGameMapsSettings>()->TransitionMap;
 			FName DefaultMapFinalName(*TransitionMap);
 
