#ue4 - Rename variable UPrimitiveComponent::CanBeCharacterBase to CanCharacterStepUpOn, and UPrimitiveComponent::CanBeBaseForCharacter() to CanCharacterStepUp().

- Made UCharacterMovementComponent::StepUp() check the surface it lands on with CanStepUp(), and reject if false and trying to step higher.

TTP 337602

diff --git a/Engine/Config/BaseEngine.ini b/Engine/Config/BaseEngine.ini
--- a/Engine/Config/BaseEngine.ini
+++ b/Engine/Config/BaseEngine.ini
@@ -605,6 +605,7 @@ bUseStreamingPause=false
 +TaggedPropertyRedirects=(ClassName="FloatingPawnMovement",OldPropertyName="DecelRate",NewPropertyName="Deceleration")
 +TaggedPropertyRedirects=(ClassName="PassiveSoundMixModifier",OldPropertyName="VolumeThreshold",NewPropertyName="MinVolumeThreshold")
 +TaggedPropertyRedirects=(ClassName="ProjectileMovementComponent",OldPropertyName="Speed",NewPropertyName="InitialSpeed")
++TaggedPropertyRedirects=(ClassName="PrimitiveComponent",OldPropertyName="CanBeCharacterBase",NewPropertyName="CanCharacterStepUpOn")
 
 +TaggedPropertyRedirects=(ClassName="MaterialEditorInstanceConstant",OldPropertyName="RefractionBias",NewPropertyName="RefractionDepthBias")
 +TaggedPropertyRedirects=(ClassName="Material",OldPropertyName="RefractionBias",NewPropertyName="RefractionDepthBias")
diff --git a/Engine/Source/Runtime/Core/Public/UObject/ObjectVersion.h b/Engine/Source/Runtime/Core/Public/UObject/ObjectVersion.h
--- a/Engine/Source/Runtime/Core/Public/UObject/ObjectVersion.h
+++ b/Engine/Source/Runtime/Core/Public/UObject/ObjectVersion.h
@@ -562,6 +562,8 @@ enum EUnrealEngineObjectUE4Version
 	VER_UE4_FIXUP_MOTOR_UNITS,
 	// deprecated MovementComponent functions including "ModifiedMaxSpeed" et al
 	VER_UE4_DEPRECATED_MOVEMENTCOMPONENT_MODIFIED_SPEEDS,
+	// rename CanBeCharacterBase
+	VER_UE4_RENAME_CANBECHARACTERBASE,
 
 	// -----<new versions can be added before this line>-------------------------------------------------
 	// - this needs to be the last line (see note below)
diff --git a/Engine/Source/Runtime/Engine/Classes/Components/PrimitiveComponent.h b/Engine/Source/Runtime/Engine/Classes/Components/PrimitiveComponent.h
--- a/Engine/Source/Runtime/Engine/Classes/Components/PrimitiveComponent.h
+++ b/Engine/Source/Runtime/Engine/Classes/Components/PrimitiveComponent.h
@@ -53,11 +53,11 @@ enum ERadialImpulseFalloff
 UENUM()
 enum ECanBeCharacterBase
 {
-	// can never be base for character to stand on
+	// Character cannot step up onto this Component.
 	ECB_No,
-	// can always be character base
+	// Character can step up onto this Component.
 	ECB_Yes,
-	// owning acter determines whether can be character base
+	// Owning actor determines whether character can step up onto this Component (default true unless overridden in code).
 	ECB_Owner,
 	ECB_MAX,
 };
@@ -523,8 +523,16 @@ public:
 	UPROPERTY(transient)
 	float LastRenderTime;
 
+private:
+	UPROPERTY()
+	TEnumAsByte<enum ECanBeCharacterBase> CanBeCharacterBase_DEPRECATED;
+
+public:
+	/**
+	 * Determine whether a Character can step up onto this component.
+	 */
 	UPROPERTY(EditDefaultsOnly, BlueprintReadOnly, Category=Base)
-	TEnumAsByte<enum ECanBeCharacterBase> CanBeCharacterBase;
+	TEnumAsByte<enum ECanBeCharacterBase> CanCharacterStepUpOn;
 
 	/** Set of actors to ignore during component sweeps in MoveComponent() */
 	TArray<TWeakObjectPtr<AActor> > MoveIgnoreActors;
@@ -1447,11 +1455,20 @@ public:
 	virtual bool ComputePenetration(FMTDResult & OutMTD, const FCollisionShape & CollisionShape, const FVector & Pos, const FQuat & Rot);
 	
 	/**
-	 * Called from the Pawn's BaseChanged() event.
-	 * @param APawn is the pawn that wants to be based on this actor
-	 * @return true if we don't want the pawn to bounce off
+	 * Return true if the given Pawn can step up onto this component.
+	 * @param APawn is the pawn that wants to step onto this component.
+	 */
+	DEPRECATED(4.3, "UPrimitiveComponent::CanBeBaseForCharacter() is deprecated, use CanCharacterStepUp() instead.")
+	virtual bool CanBeBaseForCharacter(class APawn* Pawn) const
+	{
+		return CanCharacterStepUp(Pawn);
+	}
+
+	/**
+	 * Return true if the given Pawn can step up onto this component.
+	 * @param Pawn is the Pawn that wants to step onto this component.
 	 */
-	virtual bool CanBeBaseForCharacter(class APawn* APawn) const;
+	virtual bool CanCharacterStepUp(class APawn* Pawn) const;
 
 	/** 
 	 *	Indicates whether this actor is to be considered by navigation system a valid actor
diff --git a/Engine/Source/Runtime/Engine/Classes/GameFramework/Actor.h b/Engine/Source/Runtime/Engine/Classes/GameFramework/Actor.h
--- a/Engine/Source/Runtime/Engine/Classes/GameFramework/Actor.h
+++ b/Engine/Source/Runtime/Engine/Classes/GameFramework/Actor.h
@@ -1900,11 +1900,10 @@ public:
 	class AWorldSettings* GetWorldSettings() const;
 
 	/**
-	 * Called from the Component's CanBeBaseForCharacter().
-	 * @param APawn - The pawn that wants to be based on this actor
-	 * @return bool - True if we don't want the pawn to bounce off
+	 * Return true if the given Pawn can be "based" on this actor (ie walk on it).
+	 * @param Pawn - The pawn that wants to be based on this actor
 	 */
-	virtual bool CanBeBaseForCharacter(class APawn* APawn) const;
+	virtual bool CanBeBaseForCharacter(class APawn* Pawn) const;
 
 	/** Apply damage to this actor.
 	 * @param DamageAmount		How much damage to apply
diff --git a/Engine/Source/Runtime/Engine/Classes/GameFramework/Pawn.h b/Engine/Source/Runtime/Engine/Classes/GameFramework/Pawn.h
--- a/Engine/Source/Runtime/Engine/Classes/GameFramework/Pawn.h
+++ b/Engine/Source/Runtime/Engine/Classes/GameFramework/Pawn.h
@@ -189,7 +189,7 @@ public:
 	virtual void EnableInput(class APlayerController* PlayerController) override;
 	virtual void DisableInput(class APlayerController* PlayerController) override;
 
-	/** Overridden to defer to the RootComponent's CanBeCharacterBase setting if it is explicitly Yes or No. If set to Owner, will return Super::CanBeBaseForCharacter(). */
+	/** Overridden to defer to the RootComponent's CanCharacterStepUpOn setting if it is explicitly Yes or No. If set to Owner, will return Super::CanBeBaseForCharacter(). */
 	virtual bool CanBeBaseForCharacter(class APawn* APawn) const override;
 	// End AActor Interface
 
diff --git a/Engine/Source/Runtime/Engine/Private/AI/Navigation/NavigationTestingActor.cpp b/Engine/Source/Runtime/Engine/Private/AI/Navigation/NavigationTestingActor.cpp
--- a/Engine/Source/Runtime/Engine/Private/AI/Navigation/NavigationTestingActor.cpp
+++ b/Engine/Source/Runtime/Engine/Private/AI/Navigation/NavigationTestingActor.cpp
@@ -51,7 +51,7 @@ ANavigationTestingActor::ANavigationTestingActor(const class FPostConstructIniti
 	CapsuleComponent = PCIP.CreateDefaultSubobject<UCapsuleComponent>(this, TEXT("CollisionCylinder"));
 	CapsuleComponent->InitCapsuleSize(NavAgentProps.AgentRadius, NavAgentProps.AgentHeight / 2);
 	CapsuleComponent->SetCollisionProfileName(CollisionProfileName);
-	CapsuleComponent->CanBeCharacterBase = ECB_No;
+	CapsuleComponent->CanCharacterStepUpOn = ECB_No;
 	CapsuleComponent->bShouldUpdatePhysicsVolume = true;
 
 	RootComponent = CapsuleComponent;
diff --git a/Engine/Source/Runtime/Engine/Private/Character.cpp b/Engine/Source/Runtime/Engine/Private/Character.cpp
--- a/Engine/Source/Runtime/Engine/Private/Character.cpp
+++ b/Engine/Source/Runtime/Engine/Private/Character.cpp
@@ -48,7 +48,7 @@ ACharacter::ACharacter(const class FPostConstructInitializeProperties& PCIP)
 	static FName CollisionProfileName(TEXT("Pawn"));
 	CapsuleComponent->SetCollisionProfileName(CollisionProfileName);
 
-	CapsuleComponent->CanBeCharacterBase = ECB_No;
+	CapsuleComponent->CanCharacterStepUpOn = ECB_No;
 	CapsuleComponent->bShouldUpdatePhysicsVolume = true;
 	CapsuleComponent->bCheckAsyncSceneOnMove = false;	
 	RootComponent = CapsuleComponent;
diff --git a/Engine/Source/Runtime/Engine/Private/CharacterMovementComponent.cpp b/Engine/Source/Runtime/Engine/Private/CharacterMovementComponent.cpp
--- a/Engine/Source/Runtime/Engine/Private/CharacterMovementComponent.cpp
+++ b/Engine/Source/Runtime/Engine/Private/CharacterMovementComponent.cpp
@@ -3100,7 +3100,7 @@ void UCharacterMovementComponent::MoveAlongFloor(const FVector& InVelocity, cons
 					bJustTeleported |= !bMaintainHorizontalGroundVelocity;
 				}
 			}
-			else if ( Hit.Component.IsValid() && !Hit.Component.Get()->CanBeBaseForCharacter(CharacterOwner) )
+			else if ( Hit.Component.IsValid() && !Hit.Component.Get()->CanCharacterStepUp(CharacterOwner) )
 			{
 				HandleImpact(Hit, DeltaSeconds, Delta);
 				SlideAlongSurface(Delta, 1.f - TimeApplied, Hit.Normal, Hit, true);
@@ -3585,19 +3585,6 @@ void UCharacterMovementComponent::PhysicsRotation(float DeltaTime)
 	}
 }
 
-bool UPrimitiveComponent::CanBeBaseForCharacter(APawn* Pawn) const
-{
-	if ( CanBeCharacterBase != ECB_Owner )
-	{
-		return CanBeCharacterBase == ECB_Yes;
-	}
-	else
-	{	
-		const AActor* Owner = GetOwner();
-		return Owner && Owner->CanBeBaseForCharacter(Pawn);
-	}
-}
-
 
 void UCharacterMovementComponent::PhysicsVolumeChanged( APhysicsVolume* NewVolume )
 {
@@ -4252,7 +4239,7 @@ bool UCharacterMovementComponent::CanStepUp(const FHitResult& Hit) const
 		return true;
 	}
 
-	if (!HitComponent->CanBeBaseForCharacter(CharacterOwner))
+	if (!HitComponent->CanCharacterStepUp(CharacterOwner))
 	{
 		return false;
 	}
@@ -4424,6 +4411,14 @@ bool UCharacterMovementComponent::StepUp(const FVector& GravDir, const FVector&
 			return false;
 		}
 
+		// Don't step up onto invalid surfaces if traveling higher.
+		if (DeltaZ > 0.f && !CanStepUp(Hit))
+		{
+			//UE_LOG(LogCharacterMovement, VeryVerbose, TEXT("- Reject StepUp (up onto surface with !CanStepUp())"));
+			ScopedStepUpMovement.RevertMove();
+			return false;
+		}
+
 		// See if we can validate the floor as a result of this step down. In almost all cases this should succeed, and we can avoid computing the floor outside this method.
 		if (OutStepDownResult != NULL)
 		{
diff --git a/Engine/Source/Runtime/Engine/Private/DefaultPawn.cpp b/Engine/Source/Runtime/Engine/Private/DefaultPawn.cpp
--- a/Engine/Source/Runtime/Engine/Private/DefaultPawn.cpp
+++ b/Engine/Source/Runtime/Engine/Private/DefaultPawn.cpp
@@ -28,7 +28,7 @@ ADefaultPawn::ADefaultPawn(const class FPostConstructInitializeProperties& PCIP)
 	static FName CollisionProfileName(TEXT("Pawn"));
 	CollisionComponent->SetCollisionProfileName(CollisionProfileName);
 
-	CollisionComponent->CanBeCharacterBase = ECB_No;
+	CollisionComponent->CanCharacterStepUpOn = ECB_No;
 	CollisionComponent->bShouldUpdatePhysicsVolume = true;
 
 	RootComponent = CollisionComponent;
diff --git a/Engine/Source/Runtime/Engine/Private/Pawn.cpp b/Engine/Source/Runtime/Engine/Private/Pawn.cpp
--- a/Engine/Source/Runtime/Engine/Private/Pawn.cpp
+++ b/Engine/Source/Runtime/Engine/Private/Pawn.cpp
@@ -137,9 +137,9 @@ AActor* APawn::GetMovementBaseActor(const APawn* Pawn)
 bool APawn::CanBeBaseForCharacter(class APawn* APawn) const
 {
 	UPrimitiveComponent* RootPrimitive = GetRootPrimitiveComponent();
-	if (RootPrimitive && RootPrimitive->CanBeCharacterBase != ECB_Owner)
+	if (RootPrimitive && RootPrimitive->CanCharacterStepUpOn != ECB_Owner)
 	{
-		return RootPrimitive->CanBeCharacterBase == ECB_Yes;
+		return RootPrimitive->CanCharacterStepUpOn == ECB_Yes;
 	}
 
 	return Super::CanBeBaseForCharacter(APawn);
diff --git a/Engine/Source/Runtime/Engine/Private/PrimitiveComponent.cpp b/Engine/Source/Runtime/Engine/Private/PrimitiveComponent.cpp
--- a/Engine/Source/Runtime/Engine/Private/PrimitiveComponent.cpp
+++ b/Engine/Source/Runtime/Engine/Private/PrimitiveComponent.cpp
@@ -85,7 +85,8 @@ UPrimitiveComponent::UPrimitiveComponent(const class FPostConstructInitializePro
 	bAlwaysCreatePhysicsState = false;
 	bRenderInMainPass = true;
 	VisibilityId = -1;
-	CanBeCharacterBase = ECB_Yes;
+	CanBeCharacterBase_DEPRECATED = ECB_Yes;
+	CanCharacterStepUpOn = ECB_Yes;
 	ComponentId.Value = NextComponentId;
 	check(IsInGameThread());
 	NextComponentId++;
@@ -649,6 +650,11 @@ void UPrimitiveComponent::PostLoad()
 		// else use defaults
 	}
 
+	if (UE4Version < VER_UE4_RENAME_CANBECHARACTERBASE)
+	{
+		CanCharacterStepUpOn = CanBeCharacterBase_DEPRECATED;
+	}
+
 	// Make sure cached cull distance is up-to-date.
 	if( LDMaxDrawDistance > 0 )
 	{
@@ -2450,4 +2456,20 @@ const int32 UPrimitiveComponent::GetNumUncachedStaticLightingInteractions() cons
 }
 #endif
 
+
+bool UPrimitiveComponent::CanCharacterStepUp(APawn* Pawn) const
+{
+	if ( CanCharacterStepUpOn != ECB_Owner )
+	{
+		return CanCharacterStepUpOn == ECB_Yes;
+	}
+	else
+	{	
+		const AActor* Owner = GetOwner();
+		return Owner && Owner->CanBeBaseForCharacter(Pawn);
+	}
+}
+
+
 #undef LOCTEXT_NAMESPACE
+
diff --git a/Engine/Source/Runtime/Engine/Private/SkeletalMesh.cpp b/Engine/Source/Runtime/Engine/Private/SkeletalMesh.cpp
--- a/Engine/Source/Runtime/Engine/Private/SkeletalMesh.cpp
+++ b/Engine/Source/Runtime/Engine/Private/SkeletalMesh.cpp
@@ -4505,7 +4505,7 @@ USkinnedMeshComponent::USkinnedMeshComponent(const class FPostConstructInitializ
 	StreamingDistanceMultiplier = 1.0f;
 	ProgressiveDrawingFraction = 1.0f;
 	bCanHighlightSelectedSections = false;
-	CanBeCharacterBase = ECB_Owner;
+	CanCharacterStepUpOn = ECB_Owner;
 #if WITH_EDITORONLY_DATA
 	ChunkIndexPreview = -1;
 	SectionIndexPreview = -1;
