Removed non-functional secondary viewport support

diff --git a/Engine/Source/Runtime/Engine/Classes/Engine/GameEngine.h b/Engine/Source/Runtime/Engine/Classes/Engine/GameEngine.h
--- a/Engine/Source/Runtime/Engine/Classes/Engine/GameEngine.h
+++ b/Engine/Source/Runtime/Engine/Classes/Engine/GameEngine.h
@@ -22,18 +22,7 @@ class ENGINE_API UGameEngine : public UEngine
 	UPROPERTY(config)
 	float MaxDeltaTime;
 
-	/** String name for any secondary viewport clients created for secondary screens */
-	UPROPERTY(config)
-	FString SecondaryViewportClientClassName;
-
-	/** Secondary viewport clients inside of secondary windows (not for split screen) */
-	UPROPERTY()
-	TArray<class UScriptViewportClient*> SecondaryViewportClients;
-
 public:
-	/** Array parallel to SecondaryViewportClients - these are the frames that render the SecondaryViewport clients */
-	TArray<FViewportFrame*> SecondaryViewportFrames;
-
 	/** The game viewport window */
 	TWeakPtr<class SWindow> GameViewportWindow;
 	/** The primary scene viewport */
@@ -126,21 +115,6 @@ public:
 	 */
 	virtual void PostLoadMap();
 
-	/**
-	 * @return true, if the GEngine is a game engine and has any secondary screens active
-	 */
-	static bool HasSecondaryScreenActive();
-
-	/**
-	 * Creates a new FViewportFrame with a viewport client of class SecondaryViewportClientClassName
-	 */
-	void CreateSecondaryViewport(uint32 SizeX, uint32 SizeY);
-
-	/**
-	 * Closes all secondary viewports opened with CreateSecondaryViewport
-	 */
-	void CloseSecondaryViewports();
-
 	/** Returns the GameViewport widget */
 	virtual TSharedPtr<SViewport> GetGameViewportWidget() const
 	{
diff --git a/Engine/Source/Runtime/Engine/Private/GameEngine.cpp b/Engine/Source/Runtime/Engine/Private/GameEngine.cpp
--- a/Engine/Source/Runtime/Engine/Private/GameEngine.cpp
+++ b/Engine/Source/Runtime/Engine/Private/GameEngine.cpp
@@ -356,14 +356,6 @@ void UGameEngine::RedrawViewports( bool bShouldPresent /*= true*/ )
 			GameViewport->Viewport->Draw(bShouldPresent);
 		}
 	}
-
-	// render the secondary viewports
-	checkSlow(SecondaryViewportClients.Num() == SecondaryViewportFrames.Num());
-	for (int32 SecondaryIndex = 0; SecondaryIndex < SecondaryViewportFrames.Num(); SecondaryIndex++)
-	{
-		SecondaryViewportFrames[SecondaryIndex]->GetViewport()->Draw(bShouldPresent);
-	}
-
 }
 
 /*-----------------------------------------------------------------------------
@@ -1114,72 +1106,6 @@ void UGameEngine::ProcessToggleFreezeStreamingCommand(UWorld* InWorld)
 	InWorld->bIsLevelStreamingFrozen = !InWorld->bIsLevelStreamingFrozen;
 }
 
-bool UGameEngine::HasSecondaryScreenActive()
-{
-	UGameEngine* GameEngine = Cast<UGameEngine>(GEngine);
-	if (GameEngine)
-{
-		check(GameEngine->SecondaryViewportClients.Num() == GameEngine->SecondaryViewportFrames.Num());
-		return (GameEngine->SecondaryViewportFrames.Num() > 0 ? true : false);
-	}
-	return false;
-}
-
-void UGameEngine::CreateSecondaryViewport(uint32 SizeX, uint32 SizeY)
-{
-/*
-	// this needs a client object to succeed (ie not be a dedicated server)
-	if (Client == NULL)
-	{
-		UE_LOG(LogEngine, Log, TEXT("Attempted to create a secondary viewport without a client object. This will not work"));
-		return;
-	}
-
-	// create a second window
-	UClass* SecondaryClass = LoadObject<UClass>(NULL, *SecondaryViewportClientClassName, NULL, LOAD_None, NULL);
-	if (SecondaryClass)
-	{
-		UScriptViewportClient* SecondClient = ConstructObject<UScriptViewportClient>(SecondaryClass);
-		FViewportFrame* SecondFrame = Client->CreateViewportFrame(SecondClient, TEXT("SecondScreen"), SizeX, SizeY, false);
-		if (SecondFrame != NULL)
-		{
-			// since nothing will directly point to this object from another object, we add it to the root
-			SecondClient->AddToRoot();
-
-			// add it to the engine for drawing, etc
-			SecondaryViewportClients.Add(SecondClient);
-			SecondaryViewportFrames.Add(SecondFrame);
-
-			// Update any mobile input zones after adding the secondary viewport
-			//extern void UpdateMobileInputZoneLayout();
-			//UpdateMobileInputZoneLayout();
-		}
-	}
-*/
-}
-
-
-void UGameEngine::CloseSecondaryViewports()
-{
-/*	if (Client == NULL)
-	{
-		return;
-	}
-
-	// Close the secondary viewports
-	for (int32 ViewportIndex = 0; ViewportIndex < SecondaryViewportFrames.Num(); ViewportIndex++)
-	{
-		Client->CloseViewport(SecondaryViewportFrames[ViewportIndex]->GetViewport());
-	}
-	SecondaryViewportFrames.Empty();
-	
-	// Remove the secondary clients 
-	for (int32 ClientIndex = 0; ClientIndex < SecondaryViewportClients.Num(); ClientIndex++)
-	{
-		SecondaryViewportClients[ClientIndex]->RemoveFromRoot();
-	}
-	SecondaryViewportClients.Empty();*/
-}
 
 UWorld* UGameEngine::GetGameWorld()
 {
