Update console variable for parallel animation evaluation and early out RefreshBoneTransforms in the case of missing skeletal mesh

diff --git a/Engine/Source/Runtime/Core/Private/HAL/ConsoleManager.cpp b/Engine/Source/Runtime/Core/Private/HAL/ConsoleManager.cpp
--- a/Engine/Source/Runtime/Core/Private/HAL/ConsoleManager.cpp
+++ b/Engine/Source/Runtime/Core/Private/HAL/ConsoleManager.cpp
@@ -1283,6 +1283,7 @@ void CreateConsoleVariables()
 	//
 	// r.      Renderer / 3D Engine / graphical feature
 	// RHI.    Low level RHI (rendering platform) specific
+	// a.	   Animation
 	// s. 	   Sound / Music
 	// n.      Network
 	// ai.     Artificial intelligence
diff --git a/Engine/Source/Runtime/Engine/Private/SkeletalMeshComponent.cpp b/Engine/Source/Runtime/Engine/Private/SkeletalMeshComponent.cpp
--- a/Engine/Source/Runtime/Engine/Private/SkeletalMeshComponent.cpp
+++ b/Engine/Source/Runtime/Engine/Private/SkeletalMeshComponent.cpp
@@ -24,7 +24,7 @@
 	#error EXPERIMENTAL_PARALLEL_CODE must be defined as either zero or one
 #endif
 
-TAutoConsoleVariable<int32> CVarUseParallelAnimationEvaluation(TEXT("ParallelAnimEvaluation"), 0, TEXT("If true, animation evaluation will be run across the task graph system as opposed to purely on the game thread"));
+TAutoConsoleVariable<int32> CVarUseParallelAnimationEvaluation(TEXT("a.ParallelAnimEvaluation"), 0, TEXT("If 1, animation evaluation will be run across the task graph system. If 0, evaluation will run purely on the game thread"));
 
 USkeletalMeshComponent::USkeletalMeshComponent(const class FPostConstructInitializeProperties& PCIP)
 	: Super(PCIP)
@@ -863,6 +863,11 @@ void USkeletalMeshComponent::RefreshBoneTransforms(FActorComponentTickFunction*
 {
 	SCOPE_CYCLE_COUNTER(STAT_RefreshBoneTransforms);
 
+	if (!SkeletalMesh || SpaceBases.Num() == 0)
+	{
+		return;
+	}
+
 	AActor * Owner = GetOwner();
 	UE_LOG(LogAnimation, Verbose, TEXT("RefreshBoneTransforms(%s)"), *GetNameSafe(Owner));
 
