Anim Trail Outer fix
Change 2107084 broke some anim trail stuff.
I had been setting the outer of the trails ParticleComponent as the AnimNotify that spawned it but this is illegal. To fix this they set the outer to the mesh component and left it more or less at that. I was using the outer connection to link the trail PSC with the notify as notify's cannot store state but are given a reference to a MeshComonent.
So I was searching the child components that had an outer == to this notify.

Now that is not possible so I've stored a pointer to the notify in the anim trail emitters.

I also took the oportunity to improve the cleaning up of the PSC from anim notifies by overriding HasCompleted() so that the bAutoDestroy will kick in properly.

diff --git a/Engine/Source/Runtime/Engine/Classes/Particles/ParticleSystemComponent.h b/Engine/Source/Runtime/Engine/Classes/Particles/ParticleSystemComponent.h
--- a/Engine/Source/Runtime/Engine/Classes/Particles/ParticleSystemComponent.h
+++ b/Engine/Source/Runtime/Engine/Classes/Particles/ParticleSystemComponent.h
@@ -624,12 +624,14 @@ public:
 	/** Get the current number of active particles in this system */
 	UFUNCTION(BlueprintCallable, Category="Effects|Components|ParticleSystem")
 	int32 GetNumActiveParticles() const;
-	
+
 	/**
-	* Fills the passed array with all trail emitters associated with a particular AnimNotifyState.
+	* Fills the passed array with all trail emitters associated with a particular object.
 	* @param OutTrailEmitters	The array to fill with pointers to the trail emitters.
+	* @param InOwner			The object that triggered this trail. Can be NULL if no assosiation was set by the owner. Not to be confused with the result of GetOwner().
+	* @param bSetOwner			If true, all trail emitters will be set as owned by InOwner.
 	*/
-	virtual void GetTrailEmitters(TArray< struct FParticleAnimTrailEmitterInstance* >& OutTrailEmitters);
+	virtual void GetOwnedTrailEmitters(TArray< struct FParticleAnimTrailEmitterInstance* >& OutTrailEmitters, const void* InOwner, bool bSetOwner = false);
 	
 	/**
 	* Begins all trail emitters in this component.
diff --git a/Engine/Source/Runtime/Engine/Private/Animation/AnimNotifyState_Trail.cpp b/Engine/Source/Runtime/Engine/Private/Animation/AnimNotifyState_Trail.cpp
--- a/Engine/Source/Runtime/Engine/Private/Animation/AnimNotifyState_Trail.cpp
+++ b/Engine/Source/Runtime/Engine/Private/Animation/AnimNotifyState_Trail.cpp
@@ -56,13 +56,14 @@ void UAnimNotifyState_Trail::NotifyBegin(class USkeletalMeshComponent * MeshComp
 	for (USceneComponent* Comp : Children)
 	{
 		UParticleSystemComponent* ParticleComp = Cast<UParticleSystemComponent>(Comp);
-		if (ParticleComp && ParticleComp->GetOuter() == this)
+		if (ParticleComp)
 		{
-			bFoundExistingTrail = true;
 			TArray< FParticleAnimTrailEmitterInstance* > TrailEmitters;
-			ParticleComp->GetTrailEmitters(TrailEmitters);
+			ParticleComp->GetOwnedTrailEmitters(TrailEmitters, this, false);
 
-			//If there are any trails, ensure the template hasn't been changed. Also destroy the component if there are erros.
+			bFoundExistingTrail = TrailEmitters.Num() > 0;
+
+			//If there are any trails, ensure the template hasn't been changed. Also destroy the component if there are errors.
 			if (TrailEmitters.Num() > 0 && (bError || PSTemplate != ParticleComp->Template))
 			{
 				//The PSTemplate was changed so we need to destroy this system and create it again with the new template. May be able to just change the template?
@@ -80,13 +81,18 @@ void UAnimNotifyState_Trail::NotifyBegin(class USkeletalMeshComponent * MeshComp
 #endif
 				}
 			}
+
+			if (bFoundExistingTrail)
+			{
+				break;//Found and re-used the existing trail so can exit now.
+			}
 		}
 	}
 
 	if (!bFoundExistingTrail && !bError)
 	{
 		//Spawn a new component from PSTemplate. This notify is made the outer so that the component can be identified later.
-		UParticleSystemComponent* NewParticleComp = ConstructObject<UParticleSystemComponent>(UParticleSystemComponent::StaticClass(), this);
+		UParticleSystemComponent* NewParticleComp = ConstructObject<UParticleSystemComponent>(UParticleSystemComponent::StaticClass(), MeshComp);
 		NewParticleComp->bAutoDestroy = true;
 		NewParticleComp->SecondsBeforeInactive = 0.0f;
 		NewParticleComp->bAutoActivate = false;
@@ -101,7 +107,7 @@ void UAnimNotifyState_Trail::NotifyBegin(class USkeletalMeshComponent * MeshComp
 		NewParticleComp->ActivateSystem(true);
 
 		TArray< FParticleAnimTrailEmitterInstance* > TrailEmitters;
-		NewParticleComp->GetTrailEmitters(TrailEmitters);
+		NewParticleComp->GetOwnedTrailEmitters(TrailEmitters, this, true);
 
 		for (FParticleAnimTrailEmitterInstance* Trail : TrailEmitters)
 		{
@@ -131,14 +137,13 @@ void UAnimNotifyState_Trail::NotifyTick(class USkeletalMeshComponent * MeshComp,
 		Width = AnimInst->GetCurveValue(WidthScaleCurve);
 	}
 
-	bool bFoundExistingTrail = false;
 	for (USceneComponent* Comp : Children)
 	{
 		UParticleSystemComponent* ParticleComp = Cast<UParticleSystemComponent>(Comp);
-		if ( ParticleComp && ParticleComp->GetOuter() == this )
+		if ( ParticleComp )
 		{
 			TArray< FParticleAnimTrailEmitterInstance* > TrailEmitters;
-			ParticleComp->GetTrailEmitters(TrailEmitters);
+			ParticleComp->GetOwnedTrailEmitters(TrailEmitters, this, false);
 			if (bError && TrailEmitters.Num() > 0)
 			{
 				Comp->DestroyComponent();
@@ -165,13 +170,17 @@ void UAnimNotifyState_Trail::NotifyEnd(class USkeletalMeshComponent * MeshComp,
 	TArray<USceneComponent*> Children;
 	MeshComp->GetChildrenComponents(false, Children);
 
-	bool bFoundExistingTrail = false;
 	for (USceneComponent* Comp : Children)
 	{
 		UParticleSystemComponent* ParticleComp = Cast<UParticleSystemComponent>(Comp);
-		if (ParticleComp && ParticleComp->GetOuter() == this)
+		if (ParticleComp)
 		{
-			ParticleComp->EndTrails();
+			TArray< FParticleAnimTrailEmitterInstance* > TrailEmitters;
+			ParticleComp->GetOwnedTrailEmitters(TrailEmitters, this, false);
+			for (FParticleAnimTrailEmitterInstance* Trail : TrailEmitters)
+			{
+				Trail->EndTrail();
+			}
 		}
 	}
 
diff --git a/Engine/Source/Runtime/Engine/Private/Particles/ParticleComponents.cpp b/Engine/Source/Runtime/Engine/Private/Particles/ParticleComponents.cpp
--- a/Engine/Source/Runtime/Engine/Private/Particles/ParticleComponents.cpp
+++ b/Engine/Source/Runtime/Engine/Private/Particles/ParticleComponents.cpp
@@ -4763,19 +4763,29 @@ int32 UParticleSystemComponent::GetNumActiveParticles() const
 	return NumParticles;
 }
 
-void UParticleSystemComponent::GetTrailEmitters(TArray< FParticleAnimTrailEmitterInstance* >& OutTrailEmitters)
+void UParticleSystemComponent::GetOwnedTrailEmitters(TArray< struct FParticleAnimTrailEmitterInstance* >& OutTrailEmitters, const void* InOwner, bool bSetOwner)
 {
 	for (FParticleEmitterInstance* Inst : EmitterInstances)
 	{
 		if (Inst && Inst->IsTrailEmitter())
 		{
-			OutTrailEmitters.Add((FParticleAnimTrailEmitterInstance*)(Inst));
+			FParticleAnimTrailEmitterInstance* TrailEmitter = (FParticleAnimTrailEmitterInstance*)Inst;
+			if (bSetOwner)
+			{
+				TrailEmitter->Owner = InOwner;
+				OutTrailEmitters.Add(TrailEmitter);
+			}
+			else if (TrailEmitter->Owner == InOwner)
+			{
+				OutTrailEmitters.Add(TrailEmitter);
+			}
 		}
 	}
 }
 
 void UParticleSystemComponent::BeginTrails(FName InFirstSocketName, FName InSecondSocketName, ETrailWidthMode InWidthMode, float InWidth)
 {
+	ActivateSystem(true);
 	for (FParticleEmitterInstance* Inst : EmitterInstances)
 	{
 		if (Inst)
@@ -4795,6 +4805,7 @@ void UParticleSystemComponent::EndTrails()
 			Inst->EndTrail();
 		}
 	}
+	DeactivateSystem();
 }
 
 void UParticleSystemComponent::SetTrailSourceData(FName InFirstSocketName, FName InSecondSocketName, ETrailWidthMode InWidthMode, float InWidth)
@@ -4831,7 +4842,7 @@ bool UParticleSystemComponent::HasCompleted()
 		{
 			if (Instance->CurrentLODLevel->bEnabled)
 			{
-				if (Instance->CurrentLODLevel->RequiredModule->EmitterLoops > 0)
+				if (Instance->CurrentLODLevel->RequiredModule->EmitterLoops > 0 || Instance->IsTrailEmitter())
 				{
 					if (bWasDeactivated && bSuppressSpawning)
 					{
diff --git a/Engine/Source/Runtime/Engine/Private/Particles/ParticleTrail2EmitterInstance.cpp b/Engine/Source/Runtime/Engine/Private/Particles/ParticleTrail2EmitterInstance.cpp
--- a/Engine/Source/Runtime/Engine/Private/Particles/ParticleTrail2EmitterInstance.cpp
+++ b/Engine/Source/Runtime/Engine/Private/Particles/ParticleTrail2EmitterInstance.cpp
@@ -2831,6 +2831,7 @@ FParticleAnimTrailEmitterInstance::FParticleAnimTrailEmitterInstance() :
 	, SecondSocketName(NAME_None)
 	, Width(1.0f)
 	, WidthMode( ETrailWidthMode_FromCentre )
+	, Owner(NULL)
 	, bTagTrailAsDead(false)
 	, bTrailEnabled(false)
 #if WITH_EDITORONLY_DATA
@@ -3878,6 +3879,11 @@ void FParticleAnimTrailEmitterInstance::DetermineVertexAndTriangleCount()
 	}
 }
 
+bool FParticleAnimTrailEmitterInstance::HasCompleted()
+{
+	return !IsTrailActive() && ActiveParticles == 0;
+}
+
 /**
  *	Retrieves the dynamic data for the emitter
  */
diff --git a/Engine/Source/Runtime/Engine/Public/ParticleEmitterInstances.h b/Engine/Source/Runtime/Engine/Public/ParticleEmitterInstances.h
--- a/Engine/Source/Runtime/Engine/Public/ParticleEmitterInstances.h
+++ b/Engine/Source/Runtime/Engine/Public/ParticleEmitterInstances.h
@@ -1409,6 +1409,11 @@ struct FParticleAnimTrailEmitterInstance : public FParticleTrailsEmitterInstance
 	*/
 	ETrailWidthMode WidthMode;
 
+	/**
+	*	The owner of this trail. Used only for assosiating trail with it's creator in some cases. E.g. AnimNotifyState_Trail. Do not use.
+	*/
+	const void* Owner;
+	
 	/**
 	*	When set, the current trail will be marked as dead in the next tick.
 	*/
@@ -1480,6 +1485,7 @@ struct FParticleAnimTrailEmitterInstance : public FParticleTrailsEmitterInstance
 	/** Determine the number of vertices and triangles in each trail */
 	void DetermineVertexAndTriangleCount();
 
+	virtual bool HasCompleted() override;
 	/**
 	 *	Retrieves the dynamic data for the emitter
 	 */
