Force older guid assignment if newer guid hasn't arrived yet

diff --git a/Engine/Source/Runtime/Engine/Private/PackageMapClient.cpp b/Engine/Source/Runtime/Engine/Private/PackageMapClient.cpp
--- a/Engine/Source/Runtime/Engine/Private/PackageMapClient.cpp
+++ b/Engine/Source/Runtime/Engine/Private/PackageMapClient.cpp
@@ -1833,12 +1833,19 @@ UObject * FNetGUIDCache::GetObjectFromNetGUID( const FNetworkGUID & NetGUID )
 	// Don't allow the assignment if this guid is from an older sequence
 	if ( CacheObjectPtr->GuidSequence == GuidSequence )
 	{
-		// We're good, assign the guid
+		// We're good, assign the guid (this may replace the old guid, but we're ok with that, since this guid is newer)
 		NetGUIDLookup.Add( Object, NetGUID );
 	}
 	else
 	{
 		UE_LOG( LogNetPackageMap, Warning, TEXT( "GetObjectFromNetGUID: Attempt to reassign guid from older sequence (%i / %i). Path: %s, Outer: %s, NetGUID: %s" ), GuidSequence, CacheObjectPtr->GuidSequence, *CacheObjectPtr->PathName.ToString(), ObjOuter != NULL ? *ObjOuter->GetPathName() : TEXT( "NULL" ), *NetGUID.ToString() );
+
+		// If we haven't set the guid yet for this object, go ahead and accept the older guid until the newer one comes in
+		if ( !NetGUIDLookup.Contains( Object ) )
+		{
+			UE_LOG( LogNetPackageMap, Warning, TEXT( "GetObjectFromNetGUID: Forced assign guid from older sequence (%i / %i). Path: %s, Outer: %s, NetGUID: %s" ), GuidSequence, CacheObjectPtr->GuidSequence, *CacheObjectPtr->PathName.ToString(), ObjOuter != NULL ? *ObjOuter->GetPathName() : TEXT( "NULL" ), *NetGUID.ToString() );
+			NetGUIDLookup.Add( Object, NetGUID );
+		}
 	}
 
 	return Object;
