Relax the check that will disconnect clients that are spamming the server

diff --git a/Engine/Source/Runtime/Engine/Private/NetConnection.cpp b/Engine/Source/Runtime/Engine/Private/NetConnection.cpp
--- a/Engine/Source/Runtime/Engine/Private/NetConnection.cpp
+++ b/Engine/Source/Runtime/Engine/Private/NetConnection.cpp
@@ -1805,9 +1805,9 @@ bool UNetConnection::TrackLogsPerSecond()
 	LogCallCount++;
 
 	static const double LOG_AVG_THRESHOLD				= 0.5;		// Frequency to check threshold
-	static const double	MAX_LOGS_PER_SECOND_INSTANT		= 30;		// If they hit this limit, they will instantly get disconnected
+	static const double	MAX_LOGS_PER_SECOND_INSTANT		= 60;		// If they hit this limit, they will instantly get disconnected
 	static const double	MAX_LOGS_PER_SECOND_SUSTAINED	= 5;		// If they sustain this logs/second for a certain count, they get disconnected
-	static const double	MAX_SUSTAINED_COUNT				= 5;		// If they sustain MAX_LOGS_PER_SECOND_SUSTAINED for this count, they get disconnected
+	static const double	MAX_SUSTAINED_COUNT				= 10;		// If they sustain MAX_LOGS_PER_SECOND_SUSTAINED for this count, they get disconnected (5 seconds currently)
 
 	if ( LogCallTotalTime > LOG_AVG_THRESHOLD )
 	{
