fixing PSA_Velocity with zero initial orientation: always apply payload rotation, even if initial rotation is zero

diff --git a/Engine/Source/Runtime/Engine/Private/Particles/ParticleSystemRender.cpp b/Engine/Source/Runtime/Engine/Private/Particles/ParticleSystemRender.cpp
--- a/Engine/Source/Runtime/Engine/Private/Particles/ParticleSystemRender.cpp
+++ b/Engine/Source/Runtime/Engine/Private/Particles/ParticleSystemRender.cpp
@@ -1619,7 +1619,7 @@ void FDynamicMeshEmitterData::Init( bool bInSelected,
 		MeshTD->RollPitchYawRange.Distribution->GetRange(Mins, Maxs);
 
 		// Enable/Disable pre-rotation
-		if ( Mins.SizeSquared() || Maxs.SizeSquared() )
+		if (Mins.SizeSquared() || Maxs.SizeSquared())
 		{
 			bApplyPreRotation = true;
 		}
@@ -2222,12 +2222,8 @@ void FDynamicMeshEmitterData::GetParticleTransform(FBaseParticle& InParticle, co
 		FVector kRotVec = FVector(fRot, fRot, fRot);
 		kRotator = FRotator::MakeFromEuler(kRotVec);
 
-		// if we're using initial orientation, apply here
-		if (bApplyPreRotation)
-		{
-			FMeshRotationPayloadData* PayloadData = (FMeshRotationPayloadData*)((uint8*)&InParticle + Source.MeshRotationOffset);
-			kRotator += FRotator::MakeFromEuler(PayloadData->Rotation);
-		}
+		FMeshRotationPayloadData* PayloadData = (FMeshRotationPayloadData*)((uint8*)&InParticle + Source.MeshRotationOffset);
+		kRotator += FRotator::MakeFromEuler(PayloadData->Rotation);
 	}
 
 	FRotationMatrix kRotMat(kRotator);
