Improvements to SplineComponent, and added a new way of handling UI Command Lists.
#add Added a new FUICommandList::ProcessCommandBindings overload, which can be called from a Viewport Client's InputKey() method (in the absence of an FKeyboardEvent object).
#add Fixed key binding handling in SplineComponentVisualizer so that it works through the FUICommandList, instead of being hardcoded.
#change USplineComponent - changed method by which spline reparam table is created (now uses the old method of splitting each segment into a fixed number of steps, instead of using fixed arc length reparameterization, which doesn't lead to such good results where the distance between control points varies largely throughout the spline).  The ReparamStepsPerSegment parameter remains, although I'd like to remove it when I can hit upon a better way to adaptively step through the spline.  Merged UpdateSplineReparamTable and AutoSetTangents into a single UpdateSpline method.
#add Added a parameter to FInterpCurve::AutoSetTangents - bStationaryEndpoints - which determines whether the endpoints should be assigned zero tangents, or whether the tangents should be set in order to extrapolate the curve beyond the endpoints, essentially providing a constant velocity at the start/end.
#add Added bStationaryEndpoints as an editable parameter to USplineComponent.  Use this if you only wish to traverse the spline at constant velocity, as it will lead to more accurate results.

diff --git a/Engine/Plugins/Experimental/Paper2D/Source/Paper2D/Private/Terrain/PaperTerrainComponent.cpp b/Engine/Plugins/Experimental/Paper2D/Source/Paper2D/Private/Terrain/PaperTerrainComponent.cpp
--- a/Engine/Plugins/Experimental/Paper2D/Source/Paper2D/Private/Terrain/PaperTerrainComponent.cpp
+++ b/Engine/Plugins/Experimental/Paper2D/Source/Paper2D/Private/Terrain/PaperTerrainComponent.cpp
@@ -91,7 +91,7 @@ void UPaperTerrainComponent::PostEditChangeProperty(FPropertyChangedEvent& Prope
 			if (AssociatedSpline != nullptr)
 			{
 				AssociatedSpline->ReparamStepsPerSegment = ReparamStepsPerSegment;
-				AssociatedSpline->UpdateSplineReparamTable();
+				AssociatedSpline->UpdateSpline();
 			}
 		}
 	}
@@ -226,7 +226,7 @@ void UPaperTerrainComponent::OnSplineEdited()
 		if (AssociatedSpline->ReparamStepsPerSegment != ReparamStepsPerSegment)
 		{
 			AssociatedSpline->ReparamStepsPerSegment = ReparamStepsPerSegment;
-			AssociatedSpline->UpdateSplineReparamTable();
+			AssociatedSpline->UpdateSpline();
 		}
 
 		FRandomStream RandomStream(RandomSeed);
diff --git a/Engine/Source/Editor/ComponentVisualizers/Private/SplineComponentVisualizer.cpp b/Engine/Source/Editor/ComponentVisualizers/Private/SplineComponentVisualizer.cpp
--- a/Engine/Source/Editor/ComponentVisualizers/Private/SplineComponentVisualizer.cpp
+++ b/Engine/Source/Editor/ComponentVisualizers/Private/SplineComponentVisualizer.cpp
@@ -94,8 +94,8 @@ void FSplineComponentVisualizer::DrawVisualization(const UActorComponent* Compon
 		for (const auto& Point : SplineReparamTable.Points)
 		{
 			const float Value = Point.OutVal;
-			FVector KeyPos = SplineComp->ComponentToWorld.TransformPosition(SplineInfo.Eval(Value, FVector::ZeroVector));
-			PDI->DrawPoint(KeyPos, FColor(255, 0, 255), 5.f, SDPG_World);
+			const FVector KeyPos = SplineComp->ComponentToWorld.TransformPosition(SplineInfo.Eval(Value, FVector::ZeroVector));
+			PDI->DrawPoint(KeyPos, FColor(255, 0, 255), 10.f, SDPG_World);
 		}
 		*/
 
@@ -284,23 +284,17 @@ bool FSplineComponentVisualizer::HandleInputDelta(FEditorViewportClient* Viewpor
 
 bool FSplineComponentVisualizer::HandleInputKey(FEditorViewportClient* ViewportClient, FViewport* Viewport, FKey Key, EInputEvent Event)
 {
-	// TODO: find a way to make this work with ProcessCommandBindings(FKeyboardEvent&)
-
 	bool bHandled = false;
 
-	if(Key == EKeys::LeftMouseButton && Event == IE_Released)
+	if (Key == EKeys::LeftMouseButton && Event == IE_Released)
 	{
 		// Reset duplication flag on LMB release
 		bAllowDuplication = true;
 	}
-	else if(Key == EKeys::Delete && Event == IE_Pressed)
+
+	if (Event == IE_Pressed)
 	{
-		// Delete selected key
-		if (IsSelectionValid())
-		{
-			OnDeleteKey();
-			bHandled = true; // consume key input
-		}
+		bHandled = SplineComponentVisualizerActions->ProcessCommandBindings(Key, FSlateApplication::Get().GetModifierKeys(), false);
 	}
 
 	return bHandled;
@@ -391,9 +385,7 @@ void FSplineComponentVisualizer::NotifyComponentModified()
 {
 	USplineComponent* SplineComp = GetEditedSplineComponent();
 
-	// Update tangents and reparam table
-	SplineComp->SplineInfo.AutoSetTangents();
-	SplineComp->UpdateSplineReparamTable();
+	SplineComp->UpdateSpline();
 
 	// Notify that the spline info has been modified
 	UProperty* SplineInfoProperty = FindField<UProperty>(USplineComponent::StaticClass(), GET_MEMBER_NAME_CHECKED(USplineComponent, SplineInfo));
diff --git a/Engine/Source/Runtime/Core/Public/Math/InterpCurve.h b/Engine/Source/Runtime/Core/Public/Math/InterpCurve.h
--- a/Engine/Source/Runtime/Core/Public/Math/InterpCurve.h
+++ b/Engine/Source/Runtime/Core/Public/Math/InterpCurve.h
@@ -90,7 +90,7 @@ public:
 	float InaccurateFindNearestOnSegment( const T &PointInSpace, int32 PtIdx, float& OutSquaredDistance ) const;
 
 	/** Automatically set the tangents on the curve based on surrounding points */
-	void AutoSetTangents(float Tension = 0.f);
+	void AutoSetTangents(float Tension = 0.f, bool bStationaryEndpoints = true);
 
 	/** Calculate the min/max out value that can be returned by this InterpCurve. */
 	void CalcBounds(T& OutMin, T& OutMax, const T& Default) const;
@@ -475,20 +475,22 @@ float FInterpCurve<T>::InaccurateFindNearestOnSegment( const T &PointInSpace, in
 
 
 template< class T > 
-void FInterpCurve<T>::AutoSetTangents(float Tension)
+void FInterpCurve<T>::AutoSetTangents(float Tension, bool bStationaryEndpoints)
 {
 	// Iterate over all points in this InterpCurve
 	for(int32 PointIndex=0; PointIndex<Points.Num(); PointIndex++)
 	{
-		T ArriveTangent = Points[PointIndex].ArriveTangent;
-		T LeaveTangent = Points[PointIndex].LeaveTangent;
+		auto& ThisPoint = Points[PointIndex];
 
-		if(PointIndex == 0)
+		T ArriveTangent = ThisPoint.ArriveTangent;
+		T LeaveTangent = ThisPoint.LeaveTangent;
+
+		if(PointIndex == 0 && bStationaryEndpoints)
 		{
 			if(PointIndex < Points.Num()-1) // Start point
 			{
 				// If first section is not a curve, or is a curve and first point has manual tangent setting.
-				if( Points[PointIndex].InterpMode == CIM_CurveAuto || Points[PointIndex].InterpMode == CIM_CurveAutoClamped )
+				if (ThisPoint.InterpMode == CIM_CurveAuto || ThisPoint.InterpMode == CIM_CurveAutoClamped)
 				{
 					FMemory::Memset( &LeaveTangent, 0, sizeof(T) );
 				}
@@ -500,29 +502,32 @@ void FInterpCurve<T>::AutoSetTangents(float Tension)
 		}
 		else
 		{
-			if(PointIndex < Points.Num()-1) // Inner point
+			if(PointIndex < Points.Num()-1 || !bStationaryEndpoints) // Inner point
 			{
-				if( Points[PointIndex].InterpMode == CIM_CurveAuto || Points[PointIndex].InterpMode == CIM_CurveAutoClamped )
+				const auto& PrevPoint = Points[FMath::Max(0, PointIndex - 1)];
+				const auto& NextPoint = Points[FMath::Min(Points.Num() - 1, PointIndex + 1)];
+
+				if( ThisPoint.InterpMode == CIM_CurveAuto || ThisPoint.InterpMode == CIM_CurveAutoClamped )
 				{
-					if( Points[PointIndex-1].IsCurveKey() && Points[PointIndex].IsCurveKey() )
+					if( PrevPoint.IsCurveKey() && ThisPoint.IsCurveKey() )
 					{
-						const bool bWantClamping = ( Points[ PointIndex ].InterpMode == CIM_CurveAutoClamped );
+						const bool bWantClamping = ( ThisPoint.InterpMode == CIM_CurveAutoClamped );
 
 						ComputeCurveTangent(
-							Points[ PointIndex - 1 ].InVal,		// Previous time
-							Points[ PointIndex - 1 ].OutVal,	// Previous point
-							Points[ PointIndex ].InVal,			// Current time
-							Points[ PointIndex ].OutVal,		// Current point
-							Points[ PointIndex + 1 ].InVal,		// Next time
-							Points[ PointIndex + 1 ].OutVal,	// Next point
-							Tension,							// Tension
-							bWantClamping,						// Want clamping?
-							ArriveTangent );					// Out
+							PrevPoint.InVal,	// Previous time
+							PrevPoint.OutVal,	// Previous point
+							ThisPoint.InVal,	// Current time
+							ThisPoint.OutVal,	// Current point
+							NextPoint.InVal,	// Next time
+							NextPoint.OutVal,	// Next point
+							Tension,			// Tension
+							bWantClamping,		// Want clamping?
+							ArriveTangent );	// Out
 
 						// In 'auto' mode, arrive and leave tangents are always the same
 						LeaveTangent = ArriveTangent;
 					}
-					else if( Points[PointIndex-1].InterpMode == CIM_Constant || Points[PointIndex].InterpMode == CIM_Constant )
+					else if( PrevPoint.InterpMode == CIM_Constant || ThisPoint.InterpMode == CIM_Constant )
 					{
 						FMemory::Memset( &ArriveTangent, 0, sizeof(T) );
 						FMemory::Memset( &LeaveTangent, 0, sizeof(T) );
@@ -532,15 +537,15 @@ void FInterpCurve<T>::AutoSetTangents(float Tension)
 			else // End point
 			{
 				// If last section is not a curve, or is a curve and final point has manual tangent setting.
-				if( Points[PointIndex].InterpMode == CIM_CurveAuto || Points[PointIndex].InterpMode == CIM_CurveAutoClamped )
+				if (ThisPoint.InterpMode == CIM_CurveAuto || ThisPoint.InterpMode == CIM_CurveAutoClamped)
 				{
 					FMemory::Memset( &ArriveTangent, 0, sizeof(T) );
 				}
 			}
 		}
 
-		Points[PointIndex].ArriveTangent = ArriveTangent;
-		Points[PointIndex].LeaveTangent = LeaveTangent;
+		ThisPoint.ArriveTangent = ArriveTangent;
+		ThisPoint.LeaveTangent = LeaveTangent;
 	}
 }
 
diff --git a/Engine/Source/Runtime/Engine/Classes/Components/SplineComponent.h b/Engine/Source/Runtime/Engine/Classes/Components/SplineComponent.h
--- a/Engine/Source/Runtime/Engine/Classes/Components/SplineComponent.h
+++ b/Engine/Source/Runtime/Engine/Classes/Components/SplineComponent.h
@@ -20,21 +20,25 @@ class ENGINE_API USplineComponent : public USceneComponent
 	FInterpCurveFloat SplineReparamTable;
 
 	/** Number of steps per spline segment to place in the reparameterization table */
-	UPROPERTY(EditAnywhere, AdvancedDisplay, Category = Spline, meta=(ClampMin=4, UIMin=4))
+	UPROPERTY(EditAnywhere, AdvancedDisplay, Category = Spline, meta=(ClampMin=4, UIMin=4, ClampMax=100, UIMax=100))
 	int32 ReparamStepsPerSegment;
 
 	/** Specifies the duration of the spline in seconds */
 	UPROPERTY(EditAnywhere, BlueprintReadWrite, Category=Spline)
 	float Duration;
 
+	/** Whether the endpoints of the spline are considered stationary when traversing the spline at non-constant velocity.  Essentially this sets the endpoints' tangents to zero vectors. */
+	UPROPERTY(EditAnywhere, BlueprintReadWrite, AdvancedDisplay, Category = Spline)
+	bool bStationaryEndpoints;
+
 	// Begin UActorComponent interface.
 	virtual TSharedPtr<FComponentInstanceDataBase> GetComponentInstanceData() const override;
 	virtual FName GetComponentInstanceDataType() const override;
 	virtual void ApplyComponentInstanceData(TSharedPtr<FComponentInstanceDataBase> ComponentInstanceData) override;
 	// End UActorComponent interface.
 
-	/** Update the SplineReparamTable */
-	void UpdateSplineReparamTable();
+	/** Update the spline tangents and SplineReparamTable */
+	void UpdateSpline();
 
 	/** Clears all the points in the spline */
 	UFUNCTION(BlueprintCallable, Category=Spline)
diff --git a/Engine/Source/Runtime/Engine/Private/SplineComponent.cpp b/Engine/Source/Runtime/Engine/Private/SplineComponent.cpp
--- a/Engine/Source/Runtime/Engine/Private/SplineComponent.cpp
+++ b/Engine/Source/Runtime/Engine/Private/SplineComponent.cpp
@@ -11,8 +11,9 @@
 
 USplineComponent::USplineComponent(const class FPostConstructInitializeProperties& PCIP)
 	: Super(PCIP)
-	, ReparamStepsPerSegment(4)
+	, ReparamStepsPerSegment(10)
 	, Duration(1.0f)
+	, bStationaryEndpoints(false)
 {
 	// Add 2 keys by default
 	int32 PointIndex = SplineInfo.AddPoint(0.f, FVector(0,0,0));
@@ -21,12 +22,14 @@ USplineComponent::USplineComponent(const class FPostConstructInitializePropertie
 	PointIndex = SplineInfo.AddPoint(1.f, FVector(100,0,0));
 	SplineInfo.Points[PointIndex].InterpMode = CIM_CurveAuto;
 
-	SplineInfo.AutoSetTangents();
-	UpdateSplineReparamTable();
+	UpdateSpline();
 }
 
-void USplineComponent::UpdateSplineReparamTable()
+void USplineComponent::UpdateSpline()
 {
+	// Automatically set the tangents on any CurveAuto keys
+	SplineInfo.AutoSetTangents(0.0f, bStationaryEndpoints);
+
 	// Start by clearing it
 	SplineReparamTable.Reset();
 	
@@ -37,43 +40,19 @@ void USplineComponent::UpdateSplineReparamTable()
 	}
 
 	const int32 NumSegments = SplineInfo.Points.Num() - 1;
-	const int32 NumSteps = FMath::Max<int32>(ReparamStepsPerSegment, 4) * NumSegments;
-
-	// Build array of segment lengths
-	TArray<float> SegmentLengths;
-	SegmentLengths.Reset(NumSegments);
 
-	float TotalLength = 0.0f;
-	for (int32 SegmentIndex = 0; SegmentIndex < NumSegments; ++SegmentIndex)
-	{
-		float SegmentLength = GetSegmentLength(SegmentIndex);
-		SegmentLengths.Add(SegmentLength);
-		TotalLength += SegmentLength;
-	}
-
-	// Now step through the spine at a constant distance interval, adding points to the reparam table
-	int32 CurrentSegment = 0;
 	float AccumulatedLength = 0.0f;
-	const float LengthPerStep = TotalLength / (float)NumSteps;
-	for (int32 Step = 0; Step < NumSteps; ++Step)
+	for (int32 SegmentIndex = 0; SegmentIndex < NumSegments; ++SegmentIndex)
 	{
-		const float TargetLength = Step * LengthPerStep;
-
-		while (TargetLength - AccumulatedLength > SegmentLengths[CurrentSegment])
+		for (int32 Step = 0; Step < ReparamStepsPerSegment; ++Step)
 		{
-			AccumulatedLength += SegmentLengths[CurrentSegment];
-			CurrentSegment++;
+			const float Param = static_cast<float>(Step) / ReparamStepsPerSegment;
+			const float SegmentLength = (Step == 0) ? 0.0f : GetSegmentLength(SegmentIndex, Param);
+			SplineReparamTable.AddPoint(SegmentLength + AccumulatedLength, SegmentIndex + Param);
 		}
-
-		const float Param = GetSegmentParamFromLength(CurrentSegment, TargetLength - AccumulatedLength, SegmentLengths[CurrentSegment]);
-		const int32 Index = SplineReparamTable.AddPoint(TargetLength, CurrentSegment + Param);
-		SplineReparamTable.Points[Index].InterpMode = CIM_CurveAuto;
+		AccumulatedLength += GetSegmentLength(SegmentIndex, 1.0f);
 	}
-
-	const int32 Index = SplineReparamTable.AddPoint(TotalLength, NumSegments);
-	SplineReparamTable.Points[Index].InterpMode = CIM_CurveAuto;
-
-	SplineReparamTable.AutoSetTangents();
+	SplineReparamTable.AddPoint(AccumulatedLength, static_cast<float>(NumSegments));
 }
 
 
@@ -191,8 +170,7 @@ void USplineComponent::AddSplineWorldPoint(const FVector& Position)
 	int32 PointIndex = SplineInfo.AddPoint(InputKey, ComponentToWorld.InverseTransformPosition(Position));
 	SplineInfo.Points[PointIndex].InterpMode = CIM_CurveAuto;
 
-	SplineInfo.AutoSetTangents();
-	UpdateSplineReparamTable();
+	UpdateSpline();
 }
 
 
@@ -207,8 +185,7 @@ void USplineComponent::SetSplineWorldPoints(const TArray<FVector>& Points)
 		InputKey += 1.0f;
 	}
 
-	SplineInfo.AutoSetTangents();
-	UpdateSplineReparamTable();
+	UpdateSpline();
 }
 
 void USplineComponent::SetWorldLocationAtSplinePoint(int32 PointIndex, const FVector& InLocation)
@@ -216,9 +193,7 @@ void USplineComponent::SetWorldLocationAtSplinePoint(int32 PointIndex, const FVe
 	if((PointIndex >= 0) && (PointIndex < SplineInfo.Points.Num()))
 	{
 		SplineInfo.Points[PointIndex].OutVal = InLocation;
-
-		SplineInfo.AutoSetTangents();
-		UpdateSplineReparamTable();
+		UpdateSpline();
 	}
 }
 
@@ -362,6 +337,7 @@ void USplineComponent::RefreshSplineInputs()
 	}
 }
 
+
 /** Used to store spline data during RerunConstructionScripts */
 class FSplineInstanceData : public FComponentInstanceDataBase
 {
@@ -372,7 +348,6 @@ public:
 	{
 	}
 
-	/** Lightmaps from LODData */
 	FInterpCurveVector SplineInfo;
 };
 
@@ -390,7 +365,7 @@ TSharedPtr<FComponentInstanceDataBase> USplineComponent::GetComponentInstanceDat
 void USplineComponent::ApplyComponentInstanceData(TSharedPtr<FComponentInstanceDataBase> ComponentInstanceData)
 {
 	SplineInfo = StaticCastSharedPtr<FSplineInstanceData>(ComponentInstanceData)->SplineInfo;
-	UpdateSplineReparamTable();
+	UpdateSpline();
 }
 
 #if WITH_EDITORONLY_DATA
@@ -399,9 +374,10 @@ void USplineComponent::PostEditChangeProperty(FPropertyChangedEvent& PropertyCha
 	if (PropertyChangedEvent.Property != nullptr)
 	{
 		const FName PropertyName(PropertyChangedEvent.Property->GetFName());
-		if (PropertyName == GET_MEMBER_NAME_CHECKED(USplineComponent, ReparamStepsPerSegment))
+		if (PropertyName == GET_MEMBER_NAME_CHECKED(USplineComponent, ReparamStepsPerSegment) ||
+			PropertyName == GET_MEMBER_NAME_CHECKED(USplineComponent, bStationaryEndpoints))
 		{
-			UpdateSplineReparamTable();
+			UpdateSpline();
 		}
 	}
 
diff --git a/Engine/Source/Runtime/Slate/Private/Framework/Commands/UICommandList.cpp b/Engine/Source/Runtime/Slate/Private/Framework/Commands/UICommandList.cpp
--- a/Engine/Source/Runtime/Slate/Private/Framework/Commands/UICommandList.cpp
+++ b/Engine/Source/Runtime/Slate/Private/Framework/Commands/UICommandList.cpp
@@ -204,6 +204,26 @@ bool FUICommandList::ProcessCommandBindings( const FPointerEvent& InMouseEvent )
 }
 
 
+/**
+ * Processes any UI commands which are activated by the specified key, modifier keys state and input event
+ *
+ * @param Key				The current key that is pressed
+ * @param ModifierKeysState	Pressed state of keys that are commonly used as modifiers
+ * @param bRepeat			True if the input is repeating (held)
+ *
+ * @return true if an action was processed
+ */
+bool FUICommandList::ProcessCommandBindings( const FKey Key, const FModifierKeysState& ModifierKeysState, const bool bRepeat ) const
+{
+	return ConditionalProcessCommandBindings(
+		Key,
+		ModifierKeysState.IsControlDown(),
+		ModifierKeysState.IsAltDown(),
+		ModifierKeysState.IsShiftDown(),
+		bRepeat );
+}
+
+
 /* FUICommandList implementation
  *****************************************************************************/
 
diff --git a/Engine/Source/Runtime/Slate/Public/Framework/Commands/UICommandList.h b/Engine/Source/Runtime/Slate/Public/Framework/Commands/UICommandList.h
--- a/Engine/Source/Runtime/Slate/Public/Framework/Commands/UICommandList.h
+++ b/Engine/Source/Runtime/Slate/Public/Framework/Commands/UICommandList.h
@@ -121,6 +121,17 @@ public:
 	 */
 	bool ProcessCommandBindings( const FPointerEvent& InMouseEvent ) const;
 
+	/**
+	 * Processes any UI commands which are activated by the specified key, modifier keys state and input event
+	 *
+	 * @param Key				The current key that is pressed
+	 * @param ModifierKeysState	Pressed state of keys that are commonly used as modifiers
+	 * @param bRepeat			True if input is repeating (held)
+	 *
+	 * @return true if an action was processed
+	 */
+	bool ProcessCommandBindings( const FKey Key, const FModifierKeysState& ModifierKeysState, const bool bRepeat ) const;
+
 	/** Sets the delegate that determines if this UICommandList is capable of producing an action for the supplied command */
 	void SetCanProduceActionForCommand( const FCanProduceActionForCommand& NewCanProduceActionForCommand ) { CanProduceActionForCommand = NewCanProduceActionForCommand; }
 
