Add salt to network profiles

* During seamless travel, we can write two back to back real fast, which overwrites the important one (they have the same time)
* Log improvements

diff --git a/Engine/Source/Runtime/Engine/Private/NetworkProfiler.cpp b/Engine/Source/Runtime/Engine/Private/NetworkProfiler.cpp
--- a/Engine/Source/Runtime/Engine/Private/NetworkProfiler.cpp
+++ b/Engine/Source/Runtime/Engine/Private/NetworkProfiler.cpp
@@ -357,8 +357,10 @@ void FNetworkProfiler::TrackEvent( const FString& EventName, const FString& Even
 void FNetworkProfiler::TrackSessionChange( bool bShouldContinueTracking, const FURL& InURL )
 {
 #if ALLOW_DEBUG_FILES
-	if( bIsTrackingEnabled )
+	if ( bIsTrackingEnabled )
 	{
+		UE_LOG( LogNet, Log, TEXT( "Network Profiler: TrackSessionChange.  InURL: %s" ), *InURL.ToString() );
+
 		// Session change might occur while other thread uses low level networking.
 		SCOPE_LOCK_REF(CriticalSection);
 
@@ -367,7 +369,7 @@ void FNetworkProfiler::TrackSessionChange( bool bShouldContinueTracking, const F
 		{	
 			if( bHasNoticeableNetworkTrafficOccured )
 			{
-				UE_LOG(LogNet, Log, TEXT("Netork Profiler: Writing out session file for '%s'"), *CurrentURL.ToString());
+				UE_LOG(LogNet, Log, TEXT("Network Profiler: Writing out session file for '%s'"), *CurrentURL.ToString());
 
 				// Write end of stream marker.
 				uint8 Type = NPTYPE_EndOfStreamMarker;
@@ -398,18 +400,26 @@ void FNetworkProfiler::TrackSessionChange( bool bShouldContinueTracking, const F
 				FileWriter->Close();
 			
 				// Rename/ move file.
-				const FString FinalFileName = FPaths::ProfilingDir() + GGameName + TEXT("-") + FDateTime::Now().ToString() + TEXT(".nprof");
+				static int32 Salt = 0;
+				Salt++;		// Use a salt to solve the issue where this function is called so fast it produces the same time (seems to happen during seamless travel)
+				const FString FinalFileName = FPaths::ProfilingDir() + GGameName + TEXT( "-" ) + FDateTime::Now().ToString() + FString::Printf( TEXT( "[%i]" ), Salt ) + TEXT( ".nprof" );
 				bool bWasMovedSuccessfully = IFileManager::Get().Move( *FinalFileName, *TempFileName );
 
 				// Send data to UnrealConsole to upload to DB.
 				if( bWasMovedSuccessfully )
 				{
-					SendDataToPCViaUnrealConsole( TEXT("UE_PROFILER!NETWORK:"), *FinalFileName );
+					UE_LOG( LogNet, Log, TEXT( "Network Profiler: Saved SUCCESS: %s" ), *FinalFileName );
+
+					SendDataToPCViaUnrealConsole( TEXT( "UE_PROFILER!NETWORK:" ), *FinalFileName );
+				}
+				else
+				{
+					UE_LOG( LogNet, Error, TEXT( "Network Profiler: Saved FAILED: %s" ), *FinalFileName );
 				}
 			}
 			else
 			{
-				UE_LOG(LogNet, Warning, TEXT("Netork Profiler: Nothing important happened"));
+				UE_LOG(LogNet, Warning, TEXT("Network Profiler: Nothing important happened"));
 				FileWriter->Close();
 			}
 
