#ue4 Fix for networked games with HMDs applying rotation to all PCs

diff --git a/Engine/Source/Runtime/Engine/Private/PlayerController.cpp b/Engine/Source/Runtime/Engine/Private/PlayerController.cpp
--- a/Engine/Source/Runtime/Engine/Private/PlayerController.cpp
+++ b/Engine/Source/Runtime/Engine/Private/PlayerController.cpp
@@ -786,7 +786,7 @@ void APlayerController::UpdateRotation( float DeltaTime )
 
 	if (!PlayerCameraManager || !PlayerCameraManager->bFollowHmdOrientation)
 	{
-		if (GEngine->HMDDevice.IsValid() && GEngine->HMDDevice->IsHeadTrackingAllowed())
+		if (IsLocalPlayerController() && GEngine->HMDDevice.IsValid() && GEngine->HMDDevice->IsHeadTrackingAllowed())
 		{
 			GEngine->HMDDevice->ApplyHmdRotation(this, ViewRotation);
 		}
