[GitHub] - Fix for CharacterMovementComponent:
1. Allow correct execution of PhysicsRotation() if bRunPhysicsWithNoController is enabled.
2. Removal of redundant checks.

Pull request: https://github.com/EpicGames/UnrealEngine/pull/240

diff --git a/Engine/Source/Runtime/Engine/Private/CharacterMovementComponent.cpp b/Engine/Source/Runtime/Engine/Private/CharacterMovementComponent.cpp
--- a/Engine/Source/Runtime/Engine/Private/CharacterMovementComponent.cpp
+++ b/Engine/Source/Runtime/Engine/Private/CharacterMovementComponent.cpp
@@ -1352,12 +1352,9 @@ void UCharacterMovementComponent::PerformMovement(float DeltaSeconds)
 			UnCrouch(false);
 		}
 
-		if( CharacterOwner->Controller || bRunPhysicsWithNoController )
+		if (!HasRootMotion() && !CharacterOwner->IsMatineeControlled())
 		{
-			if (!HasRootMotion() && !CharacterOwner->IsMatineeControlled())
-			{
-				PhysicsRotation(DeltaSeconds);
-			}
+			PhysicsRotation(DeltaSeconds);
 		}
 
 		// Apply Root Motion rotation after movement is complete.
@@ -3516,7 +3513,7 @@ FRotator UCharacterMovementComponent::ComputeOrientToMovementRotation(const FRot
 
 void UCharacterMovementComponent::PhysicsRotation(float DeltaTime)
 {
-	if (!HasValidData() || !CharacterOwner->Controller)
+	if (!HasValidData() || (!CharacterOwner->Controller && !bRunPhysicsWithNoController))
 	{
 		return;
 	}
@@ -3534,7 +3531,7 @@ void UCharacterMovementComponent::PhysicsRotation(float DeltaTime)
 	{
 		DesiredRotation = ComputeOrientToMovementRotation(CurrentRotation, DeltaTime, DeltaRot);
 	}
-	else if (bUseControllerDesiredRotation)
+	else if (CharacterOwner->Controller && bUseControllerDesiredRotation)
 	{
 		DesiredRotation = CharacterOwner->Controller->GetDesiredRotation();
 	}
