- Added support in the Input project settings to disable the four finger tap.

diff --git a/Engine/Config/BaseInput.ini b/Engine/Config/BaseInput.ini
--- a/Engine/Config/BaseInput.ini
+++ b/Engine/Config/BaseInput.ini
@@ -6,6 +6,7 @@ DoubleClickTime=0.2f
 
 DefaultTouchInterface=/Engine/MobileResources/HUD/DefaultVirtualJoysticks.DefaultVirtualJoysticks
 bAlwaysShowTouchInterface=false
+bShowConsoleOnFourFingerTap=true
 
 ConsoleKey=Tilde
 bRequireCtrlToNavigateAutoComplete=False
diff --git a/Engine/Source/Runtime/Core/Private/IOS/IOSView.cpp b/Engine/Source/Runtime/Core/Private/IOS/IOSView.cpp
--- a/Engine/Source/Runtime/Core/Private/IOS/IOSView.cpp
+++ b/Engine/Source/Runtime/Core/Private/IOS/IOSView.cpp
@@ -364,8 +364,14 @@ id<MTLDevice> GMetalDevice = nil;
 	if( NumActiveTouches >= 4 )
 #endif
 	{
-		// Route the command to the main iOS thread (all UI must go to the main thread)
-		[[IOSAppDelegate GetDelegate] performSelectorOnMainThread:@selector(ShowConsole) withObject:nil waitUntilDone:NO];
+		bool bShowConsole = true;
+		GConfig->GetBool(TEXT("/Script/Engine.InputSettings"), TEXT("bShowConsoleOnFourFingerTap"), bShowConsole, GInputIni);
+
+		if (bShowConsole)
+		{
+			// Route the command to the main iOS thread (all UI must go to the main thread)
+			[[IOSAppDelegate GetDelegate] performSelectorOnMainThread:@selector(ShowConsole) withObject:nil waitUntilDone:NO];
+		}
 	}
 #endif
 }
diff --git a/Engine/Source/Runtime/Engine/Classes/GameFramework/InputSettings.h b/Engine/Source/Runtime/Engine/Classes/GameFramework/InputSettings.h
--- a/Engine/Source/Runtime/Engine/Classes/GameFramework/InputSettings.h
+++ b/Engine/Source/Runtime/Engine/Classes/GameFramework/InputSettings.h
@@ -47,6 +47,10 @@ class ENGINE_API UInputSettings
 	UPROPERTY(config, EditAnywhere, Category="Mobile")
 	bool bAlwaysShowTouchInterface;
 
+	/** Whether or not to show the console on 4 finger tap, on mobile platforms */
+	UPROPERTY(config, EditAnywhere, Category="Mobile")
+	bool bShowConsoleOnFourFingerTap;
+
 	/** The default on-screen touch input interface for the game (can be null to disable the onscreen interface) */
 	UPROPERTY(config, EditAnywhere, Category="Mobile", meta=(AllowedClasses="TouchInterface"))
 	FStringAssetReference DefaultTouchInterface;
@@ -55,6 +59,7 @@ class ENGINE_API UInputSettings
 	UPROPERTY(config, EditAnywhere, Category="Console")
 	FKey ConsoleKey;
 
+
 	// UObject interface
 #if WITH_EDITOR
 	virtual void PostEditChangeChainProperty(struct FPropertyChangedChainEvent& PropertyChangedEvent) override;
diff --git a/Engine/Source/Runtime/Launch/Private/Android/LaunchAndroid.cpp b/Engine/Source/Runtime/Launch/Private/Android/LaunchAndroid.cpp
--- a/Engine/Source/Runtime/Launch/Private/Android/LaunchAndroid.cpp
+++ b/Engine/Source/Runtime/Launch/Private/Android/LaunchAndroid.cpp
@@ -634,7 +634,13 @@ static int32_t HandleInputCB(struct android_app* app, AInputEvent* event)
 #if !UE_BUILD_SHIPPING
 			if ((pointerCount >= 4) && (type == TouchBegan))
 			{
-				GShowConsoleWindowNextTick = true;
+				bool bShowConsole = true;
+				GConfig->GetBool(TEXT("/Script/Engine.InputSettings"), TEXT("bShowConsoleOnFourFingerTap"), bShowConsole, GInputIni);
+
+				if (bShowConsole)
+				{
+					GShowConsoleWindowNextTick = true;
+				}
 			}
 #endif
 		}
