Editor: Removed bogus MapCheck_Message_WorldTraceBlocked message

diff --git a/Engine/Source/Runtime/Engine/Private/ActorEditor.cpp b/Engine/Source/Runtime/Engine/Private/ActorEditor.cpp
--- a/Engine/Source/Runtime/Engine/Private/ActorEditor.cpp
+++ b/Engine/Source/Runtime/Engine/Private/ActorEditor.cpp
@@ -686,30 +686,6 @@ void AActor::CheckForErrors()
 		if (ActorComponent->IsRegistered())
 		{
 			ActorComponent->CheckForErrors();
-
-			// add extra error message if they have WorldTrace blocked and not following components
-			// Blocking WorldTrace means it will be considered to be world geometry, so will give out warning about it
-			// Only allow, BrushComponent, StaticMeshComponent, LandscapeComponent, LandscapeHeightfieldCollisionComponent
-			UPrimitiveComponent * PrimComponent = Cast<UPrimitiveComponent>(ActorComponent);
-			if ( PrimComponent && PrimComponent->IsCollisionEnabled() )
-			{
-				// this is a bit difficult to do in Component CheckForError
-				// since parenting, and very specific children will need this check to be disabled
-				// without adding any new flag. So doing this here. 
-				if (PrimComponent->IsWorldGeometry() && 
-					(!PrimComponent->IsA(UBrushComponent::StaticClass())
-					&& !PrimComponent->IsA(UStaticMeshComponent::StaticClass())
-					&& !PrimComponent->IsA(ULandscapeComponent::StaticClass())
-					&& !PrimComponent->IsA(ULandscapeHeightfieldCollisionComponent::StaticClass())))
-				{
-					FFormatNamedArguments Arguments;
-					Arguments.Add(TEXT("ActorName"), FText::FromString(GetName()));
-					FMessageLog("MapCheck").Error()
-						->AddToken(FUObjectToken::Create(this))
-						->AddToken(FTextToken::Create(FText::Format(LOCTEXT( "MapCheck_Message_WorldTraceBlocked", "{ActorName} has WorldTrace blocked. It will be considered to be world geometry." ), Arguments) ))
-						->AddToken(FMapErrorToken::Create(FMapErrors::InvalidTrace));
-				}
-			}
 		}
 	}
 }
