UE4: Fixed a level streaming problem related to redirectors

diff --git a/Engine/Source/Runtime/Engine/Private/LevelStreaming.cpp b/Engine/Source/Runtime/Engine/Private/LevelStreaming.cpp
--- a/Engine/Source/Runtime/Engine/Private/LevelStreaming.cpp
+++ b/Engine/Source/Runtime/Engine/Private/LevelStreaming.cpp
@@ -569,8 +569,8 @@ void ULevelStreaming::AsyncLevelLoadComplete( const FString& InPackageName, UPac
 					// Make sure the redirector is not in the way of the new world.
 					if (WorldRedirector->GetFName() == DestinationWorld->GetFName())
 					{
-						// Pass NULL as the name to make a new unique name and NULL for the outer to not change the outer.
-						WorldRedirector->Rename(NULL, NULL, REN_DoNotDirty | REN_DontCreateRedirectors | REN_ForceNoResetLoaders | REN_NonTransactional);
+						// Pass NULL as the name to make a new unique name and GetTransientPackage() for the outer to remove it from the package.
+						WorldRedirector->Rename(NULL, GetTransientPackage(), REN_DoNotDirty | REN_DontCreateRedirectors | REN_ForceNoResetLoaders | REN_NonTransactional);
 					}
 				}
 				else
