Fix build

diff --git a/Engine/Source/Runtime/Engine/Private/Slate/DebugCanvas.h b/Engine/Source/Runtime/Engine/Private/Slate/DebugCanvas.h
--- a/Engine/Source/Runtime/Engine/Private/Slate/DebugCanvas.h
+++ b/Engine/Source/Runtime/Engine/Private/Slate/DebugCanvas.h
@@ -53,5 +53,5 @@ private:
 	/** The canvas that can be used by the render thread */
 	FCanvasProxy* RenderThreadCanvas;
 	/** Render target that the canvas renders to */
-	class FSlateBackBufferTarget* RenderTarget;
+	class FSlateCanvasRenderTarget* RenderTarget;
 };
