Support casting static shadows *from* foliage

Note: We still don't support static lighting or shadowing *of* foliage

diff --git a/Engine/Source/Runtime/Engine/Classes/Components/InstancedStaticMeshComponent.h b/Engine/Source/Runtime/Engine/Classes/Components/InstancedStaticMeshComponent.h
--- a/Engine/Source/Runtime/Engine/Classes/Components/InstancedStaticMeshComponent.h
+++ b/Engine/Source/Runtime/Engine/Classes/Components/InstancedStaticMeshComponent.h
@@ -66,22 +66,6 @@ class ENGINE_API UInstancedStaticMeshComponent : public UStaticMeshComponent
 	UPROPERTY(EditAnywhere, Transient, DuplicateTransient, DisplayName="Instances", Category=Instances, meta=(MakeEditWidget=true))
 	TArray<struct FInstancedStaticMeshInstanceData> PerInstanceSMData;
 
-	/** Number of pending lightmaps still to be calculated (Apply()'d) */
-	UPROPERTY(transient)
-	int32 NumPendingLightmaps;
-
-	/**
-	 * A key for deciding which components are compatible when joining components together after a lighting build. 
-	 * Will default to the staticmesh pointer when SetStaticMesh is called, so this must be set after calling
-	 * SetStaticMesh on the component
-	 */
-	UPROPERTY()
-	int32 ComponentJoinKey;
-
-	/** The mappings for all the instances of this component */
-	UPROPERTY(transient)
-	TArray<struct FInstancedStaticMeshMappingInfo> CachedMappings;
-
 	/** Value used to seed the random number stream that generates random numbers for each of this mesh's instances.
 		The random number is stored in a buffer accessible to materials through the PerInstanceRandom expression.  If
 		this is set to zero (default), it will be populated automatically by the editor */
diff --git a/Engine/Source/Runtime/Engine/Classes/Foliage/InstancedFoliageSettings.h b/Engine/Source/Runtime/Engine/Classes/Foliage/InstancedFoliageSettings.h
--- a/Engine/Source/Runtime/Engine/Classes/Foliage/InstancedFoliageSettings.h
+++ b/Engine/Source/Runtime/Engine/Classes/Foliage/InstancedFoliageSettings.h
@@ -188,16 +188,12 @@ class UInstancedFoliageSettings : public UObject
 	UPROPERTY(EditAnywhere, BlueprintReadOnly, Category=Lighting)
 	uint32 bAffectDynamicIndirectLighting:1;
 
-	// Static lighting is not currently supported by instanced static meshes, so we treat CastShadow as controlling all shadow settings
-	// (enables dynamic shadows, disables static shadows) which makes these two settings (bCastDynamicShadow, bCastStaticShadow) irrelevant at this point in time
-	// See UInstancedStaticMeshComponent::GetStaticLightingInfo (it's empty, but it's what would add data to the Lightmass generation)
-
 	/** Controls whether the primitive should cast shadows in the case of non precomputed shadowing.  This flag is only used if CastShadow is true. **/
-	//UPROPERTY(EditAnywhere, BlueprintReadOnly, Category=Lighting)
-	//uint32 bCastDynamicShadow:1;
+	UPROPERTY(EditAnywhere, BlueprintReadOnly, Category=Lighting)
+	uint32 bCastDynamicShadow:1;
 	/** Whether the object should cast a static shadow from shadow casting lights.  This flag is only used if CastShadow is true. */
-	//UPROPERTY(EditAnywhere, BlueprintReadOnly, Category=Lighting)
-	//uint32 bCastStaticShadow:1;
+	UPROPERTY(EditAnywhere, BlueprintReadOnly, Category=Lighting)
+	uint32 bCastStaticShadow:1;
 
 	/** 
 	 *	If true, the primitive will cast shadows even if bHidden is true.
diff --git a/Engine/Source/Runtime/Engine/Private/InstancedFoliage.cpp b/Engine/Source/Runtime/Engine/Private/InstancedFoliage.cpp
--- a/Engine/Source/Runtime/Engine/Private/InstancedFoliage.cpp
+++ b/Engine/Source/Runtime/Engine/Private/InstancedFoliage.cpp
@@ -126,12 +126,9 @@ UInstancedFoliageSettings::UInstancedFoliageSettings(const class FPostConstructI
 	VertexColorMask = FOLIAGEVERTEXCOLORMASK_Disabled;
 	VertexColorMaskThreshold = 0.5f;
 
-	// Static lighting is not currently supported by instanced static meshes, so we treat CastShadow as controlling all shadow settings
-	// (enables dynamic shadows, disables static shadows) which makes these two settings (bCastDynamicShadow, bCastStaticShadow) irrelevant at this point in time
-	// See UInstancedStaticMeshComponent::GetStaticLightingInfo (it's empty, but it's what would add data to the Lightmass generation)
 	CastShadow = true;
-	//bCastDynamicShadow = true;
-	//bCastStaticShadow = true;
+	bCastDynamicShadow = true;
+	bCastStaticShadow = true;
 	bAffectDynamicIndirectLighting = false;
 	bCastHiddenShadow = false;
 	bCastShadowAsTwoSided = false;
@@ -308,12 +305,9 @@ void FFoliageMeshInfo::AddInstance( AInstancedFoliageActor* InIFA, UStaticMesh*
 		BestCluster->ClusterComponent->InstanceStartCullDistance = Settings->StartCullDistance;
 		BestCluster->ClusterComponent->InstanceEndCullDistance = Settings->EndCullDistance;
 
-		// Static lighting is not currently supported by instanced static meshes, so we treat CastShadow as controlling all shadow settings
-		// (enables dynamic shadows, disables static shadows) which makes these two settings (bCastDynamicShadow, bCastStaticShadow) irrelevant at this point in time
-		// See UInstancedStaticMeshComponent::GetStaticLightingInfo (it's empty, but it's what would add data to the Lightmass generation)
 		BestCluster->ClusterComponent->CastShadow = Settings->CastShadow;
-		BestCluster->ClusterComponent->bCastDynamicShadow = true; //Settings->bCastDynamicShadow;
-		BestCluster->ClusterComponent->bCastStaticShadow = false; //Settings->bCastStaticShadow;
+		BestCluster->ClusterComponent->bCastDynamicShadow = Settings->bCastDynamicShadow;
+		BestCluster->ClusterComponent->bCastStaticShadow = Settings->bCastStaticShadow;
 		BestCluster->ClusterComponent->bAffectDynamicIndirectLighting = Settings->bAffectDynamicIndirectLighting;
 		BestCluster->ClusterComponent->bCastHiddenShadow = Settings->bCastHiddenShadow;
 		BestCluster->ClusterComponent->bCastShadowAsTwoSided = Settings->bCastShadowAsTwoSided;
diff --git a/Engine/Source/Runtime/Engine/Private/InstancedStaticMesh.cpp b/Engine/Source/Runtime/Engine/Private/InstancedStaticMesh.cpp
--- a/Engine/Source/Runtime/Engine/Private/InstancedStaticMesh.cpp
+++ b/Engine/Source/Runtime/Engine/Private/InstancedStaticMesh.cpp
@@ -1052,41 +1052,6 @@ bool FInstancedStaticMeshSceneProxy::GetWireframeMeshElement(int32 LODIndex, con
 	return false;
 }
 
-/*-----------------------------------------------------------------------------
-	FInstancedStaticMeshStaticLightingTextureMapping
------------------------------------------------------------------------------*/
-
-
-/** Represents a static mesh primitive with texture mapped static lighting. */
-class FInstancedStaticMeshStaticLightingTextureMapping : public FStaticMeshStaticLightingTextureMapping
-{
-public:
-
-	/** Initialization constructor. */
-	FInstancedStaticMeshStaticLightingTextureMapping(UInstancedStaticMeshComponent* InPrimitive,int32 InInstanceIndex,FStaticLightingMesh* InMesh,int32 InSizeX,int32 InSizeY,int32 InTextureCoordinateIndex,bool bPerformFullQualityRebuild)
-		: FStaticMeshStaticLightingTextureMapping(InPrimitive, 0, InMesh, InSizeX, InSizeY, InTextureCoordinateIndex, bPerformFullQualityRebuild)
-		, InstanceIndex(InInstanceIndex), LightMapData(NULL), QuantizedData(NULL)
-		, bComplete(false)
-	{
-	}
-
-private:
-
-	friend class UInstancedStaticMeshComponent;
-
-	/** The instance of the primitive this mapping represents. */
-	const int32 InstanceIndex;
-
-	/** Static lighting data */
-	FLightMapData2D* LightMapData;
-
-	/** Quantized light map data */
-	FQuantizedLightmapData* QuantizedData;
-
-	/** Has this mapping already been completed? */
-	bool bComplete;
-};
-
 
 /*-----------------------------------------------------------------------------
 	UInstancedStaticMeshComponent
@@ -1097,10 +1062,6 @@ UInstancedStaticMeshComponent::UInstancedStaticMeshComponent(const class FPostCo
 {
 	Mobility = EComponentMobility::Movable;
 	BodyInstance.bSimulatePhysics = false;
-
-	// Static lighting is not currently supported by instanced static meshes, so disable it
-	// See UInstancedStaticMeshComponent::GetStaticLightingInfo (it's empty, but it's what would add data to the Lightmass generation)
-	bCastStaticShadow = false;
 }
 
 #if WITH_EDITOR
@@ -1294,9 +1255,78 @@ FBoxSphereBounds UInstancedStaticMeshComponent::CalcBounds(const FTransform& Bou
 }
 
 #if WITH_EDITOR
+/*-----------------------------------------------------------------------------
+	FInstancedStaticMeshStaticLightingMesh
+-----------------------------------------------------------------------------*/
+
+/**
+ * A static lighting mesh class that transforms the points by the per-instance transform of an 
+ * InstancedStaticMeshComponent
+ */
+class FStaticLightingMesh_InstancedStaticMesh : public FStaticMeshStaticLightingMesh
+{
+public:
+
+	/** Initialization constructor. */
+	FStaticLightingMesh_InstancedStaticMesh(const UInstancedStaticMeshComponent* InPrimitive, int32 InstanceIndex, const TArray<ULightComponent*>& InRelevantLights)
+		: FStaticMeshStaticLightingMesh(InPrimitive, 0, InRelevantLights)
+	{
+		// override the local to world to combine the per instance transform with the component's standard transform
+		SetLocalToWorld(InPrimitive->PerInstanceSMData[InstanceIndex].Transform * InPrimitive->ComponentToWorld.ToMatrixWithScale());
+	}
+};
+
+/*-----------------------------------------------------------------------------
+	FInstancedStaticMeshStaticLightingTextureMapping
+-----------------------------------------------------------------------------*/
+
+
+/** Represents a static mesh primitive with texture mapped static lighting. */
+class FStaticLightingTextureMapping_InstancedStaticMesh : public FStaticMeshStaticLightingTextureMapping
+{
+public:
+	/** Initialization constructor. */
+	FStaticLightingTextureMapping_InstancedStaticMesh(UInstancedStaticMeshComponent* InPrimitive, FStaticLightingMesh* InMesh)
+		: FStaticMeshStaticLightingTextureMapping(InPrimitive, 0, InMesh, 0, 0, 0, false)
+	{
+		// We don't actually support light/shadow mapping on instanced meshes, only *casting* shadows
+		bProcessMapping = false;
+	}
+
+	virtual void Apply(FQuantizedLightmapData* QuantizedData, const TMap<ULightComponent*, FShadowMapData2D*>& ShadowMapData) override
+	{
+		// Not supported
+	}
+
+#if WITH_EDITOR
+	virtual bool DebugThisMapping() const override
+	{
+		return false;
+	}
+#endif	//WITH_EDITOR
+
+	virtual FString GetDescription() const override
+	{
+		return FString(TEXT("InstancedSMLightingMapping"));
+	}
+};
+
 void UInstancedStaticMeshComponent::GetStaticLightingInfo(FStaticLightingPrimitiveInfo& OutPrimitiveInfo,const TArray<ULightComponent*>& InRelevantLights,const FLightingBuildOptions& Options)
 {
+	// We don't support light/shadow mapping for instanced meshes, only *casting* shadows
+	// we intentionally ignore the mobility setting here, as foliage is marked "Movable" to force
+	// dynamic lighting due to static lighting not being supported, but is actually static
+	if (StaticMesh && bCastStaticShadow)
+	{
+		for (int32 InstanceIndex = 0; InstanceIndex < PerInstanceSMData.Num(); InstanceIndex++)
+		{
+			FStaticLightingMesh_InstancedStaticMesh* StaticLightingMesh = new FStaticLightingMesh_InstancedStaticMesh(this, InstanceIndex, InRelevantLights);
+			OutPrimitiveInfo.Meshes.Add(StaticLightingMesh);
 
+			FStaticLightingTextureMapping_InstancedStaticMesh* InstancedMapping = new FStaticLightingTextureMapping_InstancedStaticMesh(this, StaticLightingMesh);
+			OutPrimitiveInfo.Mappings.Add(InstancedMapping);
+		}
+	}
 }
 #endif
 
