Fix possible crash due to Foliage UI holding pointers to data inside a TMap which can be reallocated.

diff --git a/Engine/Source/Editor/FoliageEdit/Private/FoliageEdMode.cpp b/Engine/Source/Editor/FoliageEdit/Private/FoliageEdMode.cpp
--- a/Engine/Source/Editor/FoliageEdit/Private/FoliageEdMode.cpp
+++ b/Engine/Source/Editor/FoliageEdit/Private/FoliageEdMode.cpp
@@ -125,14 +125,14 @@ void FEdModeFoliage::Enter()
 
 	// Fixup any broken clusters
 	AInstancedFoliageActor* IFA = AInstancedFoliageActor::GetInstancedFoliageActor(GetWorld());
-	for( TMap<class UStaticMesh*, struct FFoliageMeshInfo>::TIterator MeshIt(IFA->FoliageMeshes); MeshIt; ++MeshIt )
+	for (auto& MeshPair : IFA->FoliageMeshes)
 	{
-		FFoliageMeshInfo& MeshInfo = MeshIt.Value();
+		FFoliageMeshInfo& MeshInfo = *MeshPair.Value;
 		for( int32 ClusterIdx=0;ClusterIdx<MeshInfo.InstanceClusters.Num();ClusterIdx++ )
 		{
 			if( MeshInfo.InstanceClusters[ClusterIdx].ClusterComponent == NULL )
 			{
-				MeshInfo.ReallocateClusters(IFA, MeshIt.Key());
+				MeshInfo.ReallocateClusters(IFA, MeshPair.Key);
 				break;
 			}
 		}
@@ -1122,9 +1122,9 @@ void FEdModeFoliage::ApplyBrush( FEditorViewportClient* ViewportClient )
 	// Cache a copy of the world pointer
 	UWorld* World = ViewportClient->GetWorld();
 
-	for( TMap<class UStaticMesh*, struct FFoliageMeshInfo>::TIterator MeshIt(IFA->FoliageMeshes); MeshIt; ++MeshIt )
+	for (auto& MeshPair : IFA->FoliageMeshes)
 	{
-		FFoliageMeshInfo& MeshInfo = MeshIt.Value();
+		FFoliageMeshInfo& MeshInfo = *MeshPair.Value;
 		UInstancedFoliageSettings* MeshSettings = MeshInfo.Settings;
 
 		if( MeshSettings->IsSelected )
@@ -1145,14 +1145,14 @@ void FEdModeFoliage::ApplyBrush( FEditorViewportClient* ViewportClient )
 				if( MeshSettings->ReapplyDensity )
 				{
 					// Adjust instance density
-					FMeshInfoSnapshot* SnapShot = InstanceSnapshot.Find(MeshIt.Key());
+					FMeshInfoSnapshot* SnapShot = InstanceSnapshot.Find(MeshPair.Key);
 					if( SnapShot ) 
 					{
 						// Use snapshot to determine number of instances at the start of the brush stroke
 						int32 NewInstanceCount = FMath::RoundToInt( (float)SnapShot->CountInstancesInsideSphere(BrushSphere) * MeshSettings->ReapplyDensityAmount );
 						if( MeshSettings->ReapplyDensityAmount > 1.f && NewInstanceCount > Instances.Num() )
 						{
-							AddInstancesForBrush( World , IFA, MeshIt.Key(), MeshInfo, NewInstanceCount, Instances, Pressure );
+							AddInstancesForBrush( World , IFA, MeshPair.Key, MeshInfo, NewInstanceCount, Instances, Pressure );
 						}
 						else
 						if( MeshSettings->ReapplyDensityAmount < 1.f && NewInstanceCount < Instances.Num() )			
@@ -1186,7 +1186,7 @@ void FEdModeFoliage::ApplyBrush( FEditorViewportClient* ViewportClient )
 					// Allow a single instance with a random chance, if the brush is smaller than the density
 					int32 DesiredInstanceCount = DesiredInstanceCountFloat > 1.f ? FMath::RoundToInt(DesiredInstanceCountFloat) : FMath::FRand() < DesiredInstanceCountFloat ? 1 : 0;
 
-					AddInstancesForBrush( World, IFA, MeshIt.Key(), MeshInfo, DesiredInstanceCount, Instances, Pressure );
+					AddInstancesForBrush( World, IFA, MeshPair.Key, MeshInfo, DesiredInstanceCount, Instances, Pressure );
 				}
 			}
 		}
@@ -1263,9 +1263,9 @@ void FEdModeFoliage::ApplyPaintBucket(AActor* Actor, bool bRemove)
 
 		for( int32 ComponentIdx=0;ComponentIdx<Components.Num();ComponentIdx++ )
 		{
-			for( TMap<class UStaticMesh*, struct FFoliageMeshInfo>::TIterator MeshIt(IFA->FoliageMeshes); MeshIt; ++MeshIt )
+			for (auto& MeshPair : IFA->FoliageMeshes)
 			{
-				FFoliageMeshInfo& MeshInfo = MeshIt.Value();
+				FFoliageMeshInfo& MeshInfo = *MeshPair.Value;
 				UInstancedFoliageSettings* MeshSettings = MeshInfo.Settings;
 
 				if( MeshSettings->IsSelected )
@@ -1336,9 +1336,9 @@ void FEdModeFoliage::ApplyPaintBucket(AActor* Actor, bool bRemove)
 		// Place foliage
 		AInstancedFoliageActor* IFA = AInstancedFoliageActor::GetInstancedFoliageActor(GetWorld());
 
-		for( TMap<class UStaticMesh*, struct FFoliageMeshInfo>::TIterator MeshIt(IFA->FoliageMeshes); MeshIt; ++MeshIt )
+		for (auto& MeshPair : IFA->FoliageMeshes)
 		{
-			FFoliageMeshInfo& MeshInfo = MeshIt.Value();
+			FFoliageMeshInfo& MeshInfo = *MeshPair.Value;
 			UInstancedFoliageSettings* MeshSettings = MeshInfo.Settings;
 
 			if( MeshSettings->IsSelected )
@@ -1400,7 +1400,7 @@ void FEdModeFoliage::ApplyPaintBucket(AActor* Actor, bool bRemove)
 					FFoliageInstance Inst;
 					if (InstancesToPlace[Idx].PlaceInstance(MeshSettings, Inst, World ))
 					{
-						MeshInfo.AddInstance( IFA, MeshIt.Key(), Inst );
+						MeshInfo.AddInstance( IFA, MeshPair.Key, Inst );
 					}
 				}
 			}
@@ -1512,9 +1512,9 @@ void FEdModeFoliage::UpdateFoliageMeshList()
 
 	AInstancedFoliageActor* IFA = AInstancedFoliageActor::GetInstancedFoliageActor(GetWorld());
 
-	for( TMap<class UStaticMesh*, struct FFoliageMeshInfo>::TIterator MeshIt(IFA->FoliageMeshes); MeshIt; ++MeshIt )
+	for (auto& MeshPair : IFA->FoliageMeshes)
 	{
-		new(FoliageMeshList) FFoliageMeshUIInfo(MeshIt.Key(), &MeshIt.Value());
+		new(FoliageMeshList) FFoliageMeshUIInfo(MeshPair.Key, &*MeshPair.Value);
 	}
 
 	struct FCompareFFoliageMeshUIInfo
@@ -1531,7 +1531,7 @@ void FEdModeFoliage::UpdateFoliageMeshList()
 void FEdModeFoliage::AddFoliageMesh(UStaticMesh* StaticMesh)
 {
 	AInstancedFoliageActor* IFA = AInstancedFoliageActor::GetInstancedFoliageActor(GetWorld());
-	if( IFA->FoliageMeshes.Find(StaticMesh) == NULL )
+	if( IFA->FindMesh(StaticMesh) == NULL )
 	{
 		IFA->AddMesh(StaticMesh);
 
@@ -1546,7 +1546,7 @@ bool FEdModeFoliage::RemoveFoliageMesh(UStaticMesh* StaticMesh)
 	bool bMeshRemoved = false;
 
 	AInstancedFoliageActor* IFA = AInstancedFoliageActor::GetInstancedFoliageActor(GetWorld());
-	FFoliageMeshInfo* MeshInfo = IFA->FoliageMeshes.Find(StaticMesh);
+	FFoliageMeshInfo* MeshInfo = IFA->FindMesh(StaticMesh);
 	if( MeshInfo != NULL )
 	{
 		int32 InstancesNum = MeshInfo->Instances.Num() - MeshInfo->FreeInstanceIndices.Num();
@@ -1559,7 +1559,7 @@ bool FEdModeFoliage::RemoveFoliageMesh(UStaticMesh* StaticMesh)
 			GEditor->EndTransaction();
 
 			bMeshRemoved = true;
-		}		
+		}
 
 		// Update mesh list.
 		UpdateFoliageMeshList();
@@ -1572,7 +1572,7 @@ bool FEdModeFoliage::RemoveFoliageMesh(UStaticMesh* StaticMesh)
 void FEdModeFoliage::BakeFoliage(UStaticMesh* StaticMesh, bool bSelectedOnly)
 {
 	AInstancedFoliageActor* IFA = AInstancedFoliageActor::GetInstancedFoliageActor(GetWorld());
-	FFoliageMeshInfo* MeshInfo = IFA->FoliageMeshes.Find(StaticMesh);
+	FFoliageMeshInfo* MeshInfo = IFA->FindMesh(StaticMesh);
 	if( MeshInfo != NULL )
 	{
 		TArray<int32> InstancesToConvert;
@@ -1610,7 +1610,7 @@ void FEdModeFoliage::BakeFoliage(UStaticMesh* StaticMesh, bool bSelectedOnly)
 UInstancedFoliageSettings* FEdModeFoliage::CopySettingsObject(UStaticMesh* StaticMesh)
 {
 	AInstancedFoliageActor* IFA = AInstancedFoliageActor::GetInstancedFoliageActor(GetWorld());
-	FFoliageMeshInfo* MeshInfo = IFA->FoliageMeshes.Find(StaticMesh);
+	FFoliageMeshInfo* MeshInfo = IFA->FindMesh(StaticMesh);
 	if( MeshInfo )
 	{
 		GEditor->BeginTransaction( NSLOCTEXT("UnrealEd", "FoliageMode_SettingsObjectTransaction", "Foliage Editing: Settings Object") );
@@ -1629,7 +1629,7 @@ UInstancedFoliageSettings* FEdModeFoliage::CopySettingsObject(UStaticMesh* Stati
 void FEdModeFoliage::ReplaceSettingsObject(UStaticMesh* StaticMesh, UInstancedFoliageSettings* NewSettings)
 {
 	AInstancedFoliageActor* IFA = AInstancedFoliageActor::GetInstancedFoliageActor(GetWorld());
-	FFoliageMeshInfo* MeshInfo = IFA->FoliageMeshes.Find(StaticMesh);
+	FFoliageMeshInfo* MeshInfo = IFA->FindMesh(StaticMesh);
 
 	if( MeshInfo )
 	{
@@ -1644,7 +1644,7 @@ void FEdModeFoliage::ReplaceSettingsObject(UStaticMesh* StaticMesh, UInstancedFo
 UInstancedFoliageSettings* FEdModeFoliage::SaveSettingsObject(const FText& InSettingsPackageName, UStaticMesh* StaticMesh)
 {
 	AInstancedFoliageActor* IFA = AInstancedFoliageActor::GetInstancedFoliageActor(GetWorld());
-	FFoliageMeshInfo* MeshInfo = IFA->FoliageMeshes.Find(StaticMesh);
+	FFoliageMeshInfo* MeshInfo = IFA->FindMesh(StaticMesh);
 
 	if( MeshInfo )
 	{
@@ -1672,7 +1672,7 @@ UInstancedFoliageSettings* FEdModeFoliage::SaveSettingsObject(const FText& InSet
 void FEdModeFoliage::ReallocateClusters(UStaticMesh* StaticMesh)
 {
 	AInstancedFoliageActor* IFA = AInstancedFoliageActor::GetInstancedFoliageActor(GetWorld());
-	FFoliageMeshInfo* MeshInfo = IFA->FoliageMeshes.Find(StaticMesh);
+	FFoliageMeshInfo* MeshInfo = IFA->FindMesh(StaticMesh);
 	if( MeshInfo != NULL )
 	{
 		MeshInfo->ReallocateClusters(IFA, StaticMesh);
@@ -1685,14 +1685,14 @@ bool FEdModeFoliage::ReplaceStaticMesh(UStaticMesh* OldStaticMesh, UStaticMesh*
 	bOutMeshMerged = false;
 
 	AInstancedFoliageActor* IFA = AInstancedFoliageActor::GetInstancedFoliageActor(GetWorld());
-	FFoliageMeshInfo* OldMeshInfo = IFA->FoliageMeshes.Find(OldStaticMesh);
+	FFoliageMeshInfo* OldMeshInfo = IFA->FindMesh(OldStaticMesh);
 
 	if( OldMeshInfo != NULL && OldStaticMesh != NewStaticMesh )
 	{
 		int32 InstancesNum = OldMeshInfo->Instances.Num() - OldMeshInfo->FreeInstanceIndices.Num();
 
 		// Look for the new mesh in the mesh list, and either create a new mesh or merge the instances.
-		FFoliageMeshInfo* NewMeshInfo = IFA->FoliageMeshes.Find(NewStaticMesh);
+		FFoliageMeshInfo* NewMeshInfo = IFA->FindMesh(NewStaticMesh);
 		if( NewMeshInfo == NULL )
 		{
 			FText Message = FText::Format( NSLOCTEXT("UnrealEd", "FoliageMode_ReplaceMesh", "Are you sure you want to replace all {0} instances of {1} with {2}?"), FText::AsNumber( InstancesNum ), FText::FromString( OldStaticMesh->GetName() ), FText::FromString( NewStaticMesh->GetName() ) );
@@ -1767,13 +1767,13 @@ bool FEdModeFoliage::InputKey( FEditorViewportClient* ViewportClient, FViewport*
 				if( UISettings.GetReapplyToolSelected() )
 				{
 					AInstancedFoliageActor* IFA = AInstancedFoliageActor::GetInstancedFoliageActor(GetWorld());
-					for( TMap<class UStaticMesh*, struct FFoliageMeshInfo>::TIterator MeshIt(IFA->FoliageMeshes); MeshIt; ++MeshIt )
+					for (auto& MeshPair : IFA->FoliageMeshes)
 					{
-						FFoliageMeshInfo& MeshInfo = MeshIt.Value();
+						FFoliageMeshInfo& MeshInfo = *MeshPair.Value;
 						if( MeshInfo.Settings->IsSelected )
 						{
 							// Take a snapshot of all the locations
-							InstanceSnapshot.Add( MeshIt.Key(), FMeshInfoSnapshot(&MeshInfo) ) ;
+							InstanceSnapshot.Add( MeshPair.Key, FMeshInfoSnapshot(&MeshInfo) ) ;
 
 							// Clear the "FOLIAGE_Readjusted" flag
 							for( int32 Idx=0;Idx<MeshInfo.Instances.Num();Idx++ )
@@ -1810,9 +1810,9 @@ bool FEdModeFoliage::InputKey( FEditorViewportClient* ViewportClient, FViewport*
 				GEditor->BeginTransaction( NSLOCTEXT("UnrealEd", "FoliageMode_EditTransaction", "Foliage Editing") );
 				AInstancedFoliageActor* IFA = AInstancedFoliageActor::GetInstancedFoliageActor(GetWorld());
 				IFA->Modify();
-				for( TMap<class UStaticMesh*, struct FFoliageMeshInfo>::TIterator MeshIt(IFA->FoliageMeshes); MeshIt; ++MeshIt )
+				for (auto& MeshPair : IFA->FoliageMeshes)
 				{
-					FFoliageMeshInfo& Mesh = MeshIt.Value();
+					FFoliageMeshInfo& Mesh = *MeshPair.Value;
 					if( Mesh.SelectedIndices.Num() > 0 )
 					{
 						TArray<int32> InstancesToDelete = Mesh.SelectedIndices;
@@ -1830,9 +1830,9 @@ bool FEdModeFoliage::InputKey( FEditorViewportClient* ViewportClient, FViewport*
 				AInstancedFoliageActor* IFA = AInstancedFoliageActor::GetInstancedFoliageActor(GetWorld());
 				IFA->Modify();
 				bool bMovedInstance = false;
-				for( TMap<class UStaticMesh*, struct FFoliageMeshInfo>::TIterator MeshIt(IFA->FoliageMeshes); MeshIt; ++MeshIt )
+				for (auto& MeshPair : IFA->FoliageMeshes)
 				{
-					FFoliageMeshInfo& Mesh = MeshIt.Value();
+					FFoliageMeshInfo& Mesh = *MeshPair.Value;
 					
 					Mesh.PreMoveInstances(IFA, Mesh.SelectedIndices);
 
@@ -2038,14 +2038,14 @@ bool FEdModeFoliage::InputDelta( FEditorViewportClient* InViewportClient, FViewp
 	{
 		AInstancedFoliageActor* IFA = AInstancedFoliageActor::GetInstancedFoliageActor(GetWorld());
 		IFA->Modify();
-		for( TMap<class UStaticMesh*, struct FFoliageMeshInfo>::TIterator MeshIt(IFA->FoliageMeshes); MeshIt; ++MeshIt )
+		for (auto& MeshPair : IFA->FoliageMeshes)
 		{
-			FFoliageMeshInfo& MeshInfo = MeshIt.Value();
+			FFoliageMeshInfo& MeshInfo = *MeshPair.Value;
 			bFoundSelection |= MeshInfo.SelectedIndices.Num() > 0;
 
 			if( bAltDown && bCanAltDrag && (InViewportClient->GetCurrentWidgetAxis() & EAxisList::XYZ)  )
 			{
-				MeshInfo.DuplicateInstances(IFA, MeshIt.Key(), MeshInfo.SelectedIndices);
+				MeshInfo.DuplicateInstances(IFA, MeshPair.Key, MeshInfo.SelectedIndices);
 			}
 			
 			MeshInfo.PreMoveInstances(IFA, MeshInfo.SelectedIndices);
diff --git a/Engine/Source/Editor/UnrealEd/Private/EditorObject.cpp b/Engine/Source/Editor/UnrealEd/Private/EditorObject.cpp
--- a/Engine/Source/Editor/UnrealEd/Private/EditorObject.cpp
+++ b/Engine/Source/Editor/UnrealEd/Private/EditorObject.cpp
@@ -282,11 +282,7 @@ static const TCHAR* ImportProperties(
 							Instance.Base = ActorComponent;
 
 							// Add the instance
-							FFoliageMeshInfo* MeshInfo = IFA->FoliageMeshes.Find(StaticMesh);
-							if( MeshInfo == NULL ) 
-							{
-								MeshInfo = IFA->AddMesh(StaticMesh);
-							}
+							FFoliageMeshInfo* MeshInfo = IFA->FindOrAddMesh(StaticMesh);
 							MeshInfo->AddInstance(IFA, StaticMesh, Instance);
 						}
 					}
diff --git a/Engine/Source/Programs/UnrealHeaderTool/Private/ClassMaps.h b/Engine/Source/Programs/UnrealHeaderTool/Private/ClassMaps.h
--- a/Engine/Source/Programs/UnrealHeaderTool/Private/ClassMaps.h
+++ b/Engine/Source/Programs/UnrealHeaderTool/Private/ClassMaps.h
@@ -2,32 +2,7 @@
 
 #pragma once
 
-#include "UniquePtr.h"
-
-template <typename T>
-class TUniqueObj
-{
-public:
-	TUniqueObj()
-		: Obj_(MakeUnique<T>())
-	{
-	}
-
-	template <typename Arg>
-	explicit TUniqueObj(Arg&& arg)
-		: Obj_(MakeUnique<T>(Forward<Arg>(arg)))
-	{
-	}
-
-	      T* operator->()       { return Obj_.Get(); }
-	const T* operator->() const { return Obj_.Get(); }
-
-	      T& operator*()       { return *Obj_; }
-	const T& operator*() const { return *Obj_; }
-
-private:
-	TUniquePtr<T> Obj_;
-};
+#include "UniqueObj.h"
 
 /** Information about class source and generated headers */
 struct FClassHeaderInfo
diff --git a/Engine/Source/Runtime/Core/Public/Templates/UniqueObj.h b/Engine/Source/Runtime/Core/Public/Templates/UniqueObj.h
new file mode 100644
--- /dev/null
+++ b/Engine/Source/Runtime/Core/Public/Templates/UniqueObj.h
@@ -0,0 +1,70 @@
+// Copyright 1998-2014 Epic Games, Inc. All Rights Reserved.
+
+#pragma once
+
+#include "UniquePtr.h"
+
+// This is essentially a reference version of TUniquePtr
+// Useful to force heap allocation of a value, e.g. in a TMap to give similar behaviour to TIndirectArray
+
+template <typename T>
+class TUniqueObj
+{
+public:
+	TUniqueObj()
+		: Obj(MakeUnique<T>())
+	{
+	}
+
+	TUniqueObj(const TUniqueObj& other)
+		: Obj(MakeUnique<T>(*other))
+	{
+	}
+
+	template <typename Arg>
+	explicit TUniqueObj(Arg&& arg)
+		: Obj(MakeUnique<T>(Forward<Arg>(arg)))
+	{
+	}
+
+#if PLATFORM_COMPILER_HAS_VARIADIC_TEMPLATES
+	template <typename... Args>
+	explicit TUniqueObj(Args&&... args)
+		: Obj(MakeUnique<T>(Forward<Args>(args)...))
+	{
+	}
+#endif
+
+#if PLATFORM_COMPILER_HAS_DEFAULTED_FUNCTIONS
+	TUniqueObj& operator=(const TUniqueObj&) = delete;
+#else
+private:
+	TUniqueObj& operator=(const TUniqueObj&);
+public:
+#endif
+
+	template <typename Arg>
+	TUniqueObj& operator=(Arg&& other)
+	{
+		*Obj = Forward<Arg>(other);
+		return *this;
+	}
+
+	      T& Get()       { return *Obj; }
+	const T& Get() const { return *Obj; }
+
+	      T* operator->()       { return Obj.Get(); }
+	const T* operator->() const { return Obj.Get(); }
+
+	      T& operator*()       { return *Obj; }
+	const T& operator*() const { return *Obj; }
+
+	friend FArchive& operator<<(FArchive& Ar, TUniqueObj& P)
+	{
+		Ar << *P.Obj;
+		return Ar;
+	}
+
+private:
+	TUniquePtr<T> Obj;
+};
diff --git a/Engine/Source/Runtime/Engine/Classes/Foliage/InstancedFoliageActor.h b/Engine/Source/Runtime/Engine/Classes/Foliage/InstancedFoliageActor.h
--- a/Engine/Source/Runtime/Engine/Classes/Foliage/InstancedFoliageActor.h
+++ b/Engine/Source/Runtime/Engine/Classes/Foliage/InstancedFoliageActor.h
@@ -2,9 +2,12 @@
 
 
 #pragma once
+
+#include "UniqueObj.h"
 #include "InstancedFoliage.h"
 #include "InstancedFoliageActor.generated.h"
 
+
 UCLASS(notplaceable, hidecategories = Object, MinimalAPI, NotBlueprintable)
 class AInstancedFoliageActor : public AActor
 {
@@ -12,13 +15,13 @@ class AInstancedFoliageActor : public AActor
 
 	/** The static mesh type that will be used to show the widget */
 	UPROPERTY(transient)
-	class UStaticMesh* SelectedMesh;
-
+	UStaticMesh* SelectedMesh;
 
 public:
-	TMap<class UStaticMesh*, struct FFoliageMeshInfo> FoliageMeshes;
+	TMap<UStaticMesh*, TUniqueObj<FFoliageMeshInfo>> FoliageMeshes;
 
-	// Begin UObject interface. 	
+public:
+	// Begin UObject interface.
 	virtual void Serialize(FArchive& Ar) override;
 	virtual void PostLoad() override;
 	static void AddReferencedObjects(UObject* InThis, FReferenceCollector& Collector);
@@ -55,8 +58,14 @@ public:
 	// Deletes the instances attached to a component
 	ENGINE_API void DeleteInstancesForComponent( class UActorComponent* InComponent );
 
+	// Finds a mesh entry
+	ENGINE_API FFoliageMeshInfo* FindMesh(UStaticMesh* InMesh);
+
+	// Finds a mesh entry or adds it if it doesn't already exist
+	ENGINE_API FFoliageMeshInfo* FindOrAddMesh(UStaticMesh* InMesh);
+
 	// Add a new static mesh.
-	ENGINE_API struct FFoliageMeshInfo* AddMesh( class UStaticMesh* InMesh );
+	ENGINE_API FFoliageMeshInfo* AddMesh(UStaticMesh* InMesh);
 
 	// Remove the static mesh from the mesh list, and all its instances.
 	ENGINE_API void RemoveMesh( class UStaticMesh* InMesh );
diff --git a/Engine/Source/Runtime/Engine/Private/InstancedFoliage.cpp b/Engine/Source/Runtime/Engine/Private/InstancedFoliage.cpp
--- a/Engine/Source/Runtime/Engine/Private/InstancedFoliage.cpp
+++ b/Engine/Source/Runtime/Engine/Private/InstancedFoliage.cpp
@@ -21,26 +21,26 @@
 
 DEFINE_LOG_CATEGORY_STATIC(LogInstancedFoliage, Log, All);
 
-IMPLEMENT_HIT_PROXY(HInstancedStaticMeshInstance,HHitProxy);
+IMPLEMENT_HIT_PROXY(HInstancedStaticMeshInstance, HHitProxy);
 
 //
 // Serializers for struct data
 //
 
-FArchive& operator<<( FArchive& Ar, FFoliageInstance& Instance )
+FArchive& operator<<(FArchive& Ar, FFoliageInstance& Instance)
 {
 	Ar << Instance.Base;
 	Ar << Instance.Location;
 	Ar << Instance.Rotation;
 	Ar << Instance.DrawScale3D;
-		Ar << Instance.ClusterIndex;
-		Ar << Instance.PreAlignRotation;
-		Ar << Instance.Flags;
-		Ar << Instance.ZOffset;
+	Ar << Instance.ClusterIndex;
+	Ar << Instance.PreAlignRotation;
+	Ar << Instance.Flags;
+	Ar << Instance.ZOffset;
 	return Ar;
 }
 
-FArchive& operator<<( FArchive& Ar, FFoliageInstanceCluster& Cluster )
+FArchive& operator<<(FArchive& Ar, FFoliageInstanceCluster& Cluster)
 {
 	Ar << Cluster.Bounds;
 	Ar << Cluster.ClusterComponent;
@@ -51,14 +51,14 @@ FArchive& operator<<( FArchive& Ar, FFoliageInstanceCluster& Cluster )
 	return Ar;
 }
 
-FArchive& operator<<( FArchive& Ar, FFoliageMeshInfo& MeshInfo )
+FArchive& operator<<(FArchive& Ar, FFoliageMeshInfo& MeshInfo)
 {
 	Ar << MeshInfo.InstanceClusters;
 	//!! editor only
 	Ar << MeshInfo.Instances;
 
 	// Serialize the transient data for undo.
-	if( Ar.IsTransacting() )
+	if (Ar.IsTransacting())
 	{
 		Ar << *MeshInfo.InstanceHash;
 		Ar << MeshInfo.ComponentHash;
@@ -66,12 +66,12 @@ FArchive& operator<<( FArchive& Ar, FFoliageMeshInfo& MeshInfo )
 		Ar << MeshInfo.SelectedIndices;
 	}
 
-		Ar << MeshInfo.Settings;
+	Ar << MeshInfo.Settings;
 
 	return Ar;
 }
 
-FArchive& operator<<( FArchive& Ar, FFoliageComponentHashInfo& ComponentHashInfo )
+FArchive& operator<<(FArchive& Ar, FFoliageComponentHashInfo& ComponentHashInfo)
 {
 	return Ar << ComponentHashInfo.CachedLocation << ComponentHashInfo.CachedRotation << ComponentHashInfo.CachedDrawScale << ComponentHashInfo.Instances;
 }
@@ -140,28 +140,28 @@ UInstancedFoliageSettings::UInstancedFoliageSettings(const class FPostConstructI
 // FFoliageMeshInfo
 //
 FFoliageMeshInfo::FFoliageMeshInfo()
-:	InstanceHash(NULL)
-,	Settings(NULL)
+	: InstanceHash(NULL)
+	, Settings(NULL)
 {
-	if( GIsEditor )
+	if (GIsEditor)
 	{
 		InstanceHash = new FFoliageInstanceHash();
 	}
 }
 
 FFoliageMeshInfo::FFoliageMeshInfo(const FFoliageMeshInfo& Other)
-:	InstanceHash(NULL)
-,	Settings(NULL)
+	: InstanceHash(NULL)
+	, Settings(NULL)
 {
-	if( GIsEditor )
+	if (GIsEditor)
 	{
 		InstanceHash = new FFoliageInstanceHash();
-	}	
+	}
 }
 
 FFoliageMeshInfo::~FFoliageMeshInfo()
 {
-	if( GIsEditor )
+	if (GIsEditor)
 	{
 		delete InstanceHash;
 	}
@@ -224,7 +224,7 @@ void FFoliageMeshInfo::CheckValid()
 
 			FTransform InstanceToWorldEd = Instances(InstanceIdx).GetInstanceTransform();
 			FTransform InstanceToWorldCluster = Comp->PerInstanceSMData(InstIdx).Transform * Comp->GetComponentToWorld();
-	
+
 			if( !InstanceToWorldEd.Equals(InstanceToWorldCluster) )
 			{
 				Comp->PerInstanceSMData(InstIdx).Transform = InstanceToWorldEd.ToMatrixWithScale();
@@ -242,12 +242,12 @@ void FFoliageMeshInfo::CheckValid()
 #endif
 }
 
-void FFoliageMeshInfo::AddInstance( AInstancedFoliageActor* InIFA, UStaticMesh* InMesh, const FFoliageInstance& InNewInstance )
+void FFoliageMeshInfo::AddInstance(AInstancedFoliageActor* InIFA, UStaticMesh* InMesh, const FFoliageInstance& InNewInstance)
 {
 	InIFA->Modify();
 
 	// Add the instance taking either a free slot or adding a new item.
-	int32 InstanceIndex = FreeInstanceIndices.Num() > 0 ?  FreeInstanceIndices.Pop() : Instances.AddUninitialized();
+	int32 InstanceIndex = FreeInstanceIndices.Num() > 0 ? FreeInstanceIndices.Pop() : Instances.AddUninitialized();
 
 	FFoliageInstance& AddedInstance = Instances[InstanceIndex];
 	AddedInstance = InNewInstance;
@@ -255,7 +255,7 @@ void FFoliageMeshInfo::AddInstance( AInstancedFoliageActor* InIFA, UStaticMesh*
 	// Add the instance to the hash
 	InstanceHash->InsertInstance(InNewInstance.Location, InstanceIndex);
 	FFoliageComponentHashInfo* ComponentHashInfo = ComponentHash.Find(InNewInstance.Base);
-	if( ComponentHashInfo == NULL )
+	if (ComponentHashInfo == NULL)
 	{
 		ComponentHashInfo = &ComponentHash.Add(InNewInstance.Base, FFoliageComponentHashInfo(InNewInstance.Base));
 	}
@@ -269,13 +269,13 @@ void FFoliageMeshInfo::AddInstance( AInstancedFoliageActor* InIFA, UStaticMesh*
 	int32 MaxInstancesPerCluster = Settings->MaxInstancesPerCluster;
 	float MaxClusterRadiusSq = FMath::Square(Settings->MaxClusterRadius);
 
-	for( int32 ClusterIdx=0;ClusterIdx<InstanceClusters.Num();ClusterIdx++ )
+	for (int32 ClusterIdx = 0; ClusterIdx < InstanceClusters.Num(); ClusterIdx++)
 	{
 		FFoliageInstanceCluster& Cluster = InstanceClusters[ClusterIdx];
-		if( Cluster.InstanceIndices.Num() < MaxInstancesPerCluster )
+		if (Cluster.InstanceIndices.Num() < MaxInstancesPerCluster)
 		{
 			float DistSq = (Cluster.Bounds.Origin - InNewInstance.Location).SizeSquared();
-			if( DistSq < BestClusterDistSq && DistSq < MaxClusterRadiusSq )
+			if (DistSq < BestClusterDistSq && DistSq < MaxClusterRadiusSq)
 			{
 				BestCluster = &Cluster;
 				BestClusterIndex = ClusterIdx;
@@ -287,11 +287,11 @@ void FFoliageMeshInfo::AddInstance( AInstancedFoliageActor* InIFA, UStaticMesh*
 	// Calculate transform for the instance
 	FTransform InstanceToWorld = InNewInstance.GetInstanceWorldTransform();
 
-	if( BestCluster == NULL )
+	if (BestCluster == NULL)
 	{
 		BestClusterIndex = InstanceClusters.Num();
-		BestCluster = new(InstanceClusters) FFoliageInstanceCluster(
-			ConstructObject<UInstancedStaticMeshComponent>(UInstancedStaticMeshComponent::StaticClass(),InIFA,NAME_None,RF_Transactional),
+		BestCluster = new(InstanceClusters)FFoliageInstanceCluster(
+			ConstructObject<UInstancedStaticMeshComponent>(UInstancedStaticMeshComponent::StaticClass(), InIFA, NAME_None, RF_Transactional),
 			InMesh->GetBounds().TransformBy(InstanceToWorld)
 			);
 
@@ -327,11 +327,11 @@ void FFoliageMeshInfo::AddInstance( AInstancedFoliageActor* InIFA, UStaticMesh*
 
 	// Save the cluster index
 	AddedInstance.ClusterIndex = BestClusterIndex;
-	
+
 	// Add the instance to the component
 	BestCluster->ClusterComponent->AddInstanceWorldSpace(InstanceToWorld);
 
-	if( BestCluster->ClusterComponent->SelectedInstances.Num() > 0 )
+	if (BestCluster->ClusterComponent->SelectedInstances.Num() > 0)
 	{
 		BestCluster->ClusterComponent->SelectedInstances.Add(false);
 	}
@@ -348,7 +348,7 @@ void FFoliageMeshInfo::AddInstance( AInstancedFoliageActor* InIFA, UStaticMesh*
 	CheckValid();
 }
 
-void FFoliageMeshInfo::RemoveInstances( AInstancedFoliageActor* InIFA, const TArray<int32>& InInstancesToRemove )
+void FFoliageMeshInfo::RemoveInstances(AInstancedFoliageActor* InIFA, const TArray<int32>& InInstancesToRemove)
 {
 	bool bRemoved = false;
 
@@ -358,16 +358,16 @@ void FFoliageMeshInfo::RemoveInstances( AInstancedFoliageActor* InIFA, const TAr
 	TMap<int32, TSet<int32> > ClusterInstanceSetMap;
 
 	// Remove instances from the hash
-	for( TArray<int32>::TConstIterator It(InInstancesToRemove); It; ++It )
+	for (TArray<int32>::TConstIterator It(InInstancesToRemove); It; ++It)
 	{
 		int32 InstanceIndex = *It;
 		const FFoliageInstance& Instance = Instances[InstanceIndex];
 		InstanceHash->RemoveInstance(Instance.Location, InstanceIndex);
 		FFoliageComponentHashInfo* ComponentHashInfo = ComponentHash.Find(Instance.Base);
-		if( ComponentHashInfo )
+		if (ComponentHashInfo)
 		{
 			ComponentHashInfo->Instances.Remove(InstanceIndex);
-			if( ComponentHashInfo->Instances.Num() == 0 )
+			if (ComponentHashInfo->Instances.Num() == 0)
 			{
 				// Remove the component from the component hash if this is the last instance.
 				ComponentHash.Remove(Instance.Base);
@@ -381,7 +381,7 @@ void FFoliageMeshInfo::RemoveInstances( AInstancedFoliageActor* InIFA, const TAr
 	}
 
 	// Process each cluster in turn
-	for( TMap<int32, TSet<int32> >::TIterator It(ClusterInstanceSetMap); It; ++It )
+	for (TMap<int32, TSet<int32> >::TIterator It(ClusterInstanceSetMap); It; ++It)
 	{
 		int32 CurrentClusterIndex = It.Key();
 		FFoliageInstanceCluster& Cluster = InstanceClusters[CurrentClusterIndex];
@@ -391,21 +391,21 @@ void FFoliageMeshInfo::RemoveInstances( AInstancedFoliageActor* InIFA, const TAr
 		TSet<int32>& ClusterInstancesToRemove = It.Value();
 
 		// Look through all indices in this cluster
-		for( int32 ClusterInstanceIdx=0;ClusterInstanceIdx<Cluster.InstanceIndices.Num();ClusterInstanceIdx++ )
+		for (int32 ClusterInstanceIdx = 0; ClusterInstanceIdx < Cluster.InstanceIndices.Num(); ClusterInstanceIdx++)
 		{
 			int32 InstanceIndex = Cluster.InstanceIndices[ClusterInstanceIdx];
 
 			// Check if this instance is one we need to remove
-			if( ClusterInstancesToRemove.Contains(InstanceIndex) )
+			if (ClusterInstancesToRemove.Contains(InstanceIndex))
 			{
 				// Remove the instance data and the index array entry
-				if( ClusterComponent->IsRegistered() )
+				if (ClusterComponent->IsRegistered())
 				{
 					ClusterComponent->UnregisterComponent();
 				}
 				ClusterComponent->Modify();
 				ClusterComponent->PerInstanceSMData.RemoveAtSwap(ClusterInstanceIdx);
-				if( ClusterInstanceIdx < ClusterComponent->SelectedInstances.Num() )
+				if (ClusterInstanceIdx < ClusterComponent->SelectedInstances.Num())
 				{
 					ClusterComponent->SelectedInstances.RemoveAtSwap(ClusterInstanceIdx);
 				}
@@ -433,7 +433,7 @@ void FFoliageMeshInfo::RemoveInstances( AInstancedFoliageActor* InIFA, const TAr
 			}
 		}
 
-		if( bRemovedFromComponent )
+		if (bRemovedFromComponent)
 		{
 			//!! need to update bounds for this cluster
 			ClusterComponent->RegisterComponent();
@@ -441,13 +441,13 @@ void FFoliageMeshInfo::RemoveInstances( AInstancedFoliageActor* InIFA, const TAr
 		}
 	}
 
-	if( bRemoved )
+	if (bRemoved)
 	{
 		// See if we can delete any clusters.
-		for( int32 ClusterIdx=0;ClusterIdx<InstanceClusters.Num();ClusterIdx++ )
+		for (int32 ClusterIdx = 0; ClusterIdx < InstanceClusters.Num(); ClusterIdx++)
 		{
-			if( InstanceClusters[ClusterIdx].InstanceIndices.Num() == 0 )
-			{			
+			if (InstanceClusters[ClusterIdx].InstanceIndices.Num() == 0)
+			{
 				// Unregister the component to remove it from the scene
 				UInstancedStaticMeshComponent* Component = InstanceClusters[ClusterIdx].ClusterComponent;
 
@@ -458,9 +458,9 @@ void FFoliageMeshInfo::RemoveInstances( AInstancedFoliageActor* InIFA, const TAr
 				InstanceClusters.RemoveAt(ClusterIdx);
 
 				// update the ClusterIndex for remaining foliage instances
-				for( int32 InstanceIdx=0;InstanceIdx<Instances.Num();InstanceIdx++ )
+				for (int32 InstanceIdx = 0; InstanceIdx < Instances.Num(); InstanceIdx++)
 				{
-					if( Instances[InstanceIdx].ClusterIndex > ClusterIdx )
+					if (Instances[InstanceIdx].ClusterIndex > ClusterIdx)
 					{
 						Instances[InstanceIdx].ClusterIndex--;
 					}
@@ -476,10 +476,10 @@ void FFoliageMeshInfo::RemoveInstances( AInstancedFoliageActor* InIFA, const TAr
 	}
 }
 
-void FFoliageMeshInfo::PreMoveInstances( class AInstancedFoliageActor* InIFA, const TArray<int32>& InInstancesToMove )
+void FFoliageMeshInfo::PreMoveInstances(class AInstancedFoliageActor* InIFA, const TArray<int32>& InInstancesToMove)
 {
 	// Remove instances from the hash
-	for( TArray<int32>::TConstIterator It(InInstancesToMove); It; ++It )
+	for (TArray<int32>::TConstIterator It(InInstancesToMove); It; ++It)
 	{
 		int32 InstanceIndex = *It;
 		const FFoliageInstance& Instance = Instances[InstanceIndex];
@@ -488,20 +488,20 @@ void FFoliageMeshInfo::PreMoveInstances( class AInstancedFoliageActor* InIFA, co
 }
 
 
-void FFoliageMeshInfo::PostUpdateInstances( class AInstancedFoliageActor* InIFA, const TArray<int32>& InInstancesUpdated, bool bReAddToHash )
+void FFoliageMeshInfo::PostUpdateInstances(class AInstancedFoliageActor* InIFA, const TArray<int32>& InInstancesUpdated, bool bReAddToHash)
 {
-	for( TArray<int32>::TConstIterator It(InInstancesUpdated); It; ++It )
+	for (TArray<int32>::TConstIterator It(InInstancesUpdated); It; ++It)
 	{
 		int32 InstanceIndex = *It;
 		const FFoliageInstance& Instance = Instances[InstanceIndex];
 
 		// Update this instances' transform in the UInstancedStaticMeshComponent
-		if( InstanceClusters.IsValidIndex(Instance.ClusterIndex) )
+		if (InstanceClusters.IsValidIndex(Instance.ClusterIndex))
 		{
 			FFoliageInstanceCluster& Cluster = InstanceClusters[Instance.ClusterIndex];
 			Cluster.ClusterComponent->Modify();
 			int32 ClusterInstanceDataIndex = Cluster.InstanceIndices.Find(InstanceIndex);
-			check(ClusterInstanceDataIndex != INDEX_NONE);					
+			check(ClusterInstanceDataIndex != INDEX_NONE);
 			check(Cluster.ClusterComponent->PerInstanceSMData.IsValidIndex(ClusterInstanceDataIndex));
 			Cluster.ClusterComponent->MarkRenderStateDirty();
 
@@ -514,8 +514,8 @@ void FFoliageMeshInfo::PostUpdateInstances( class AInstancedFoliageActor* InIFA,
 			Cluster.ClusterComponent->PerInstanceSMData[ClusterInstanceDataIndex].Transform = InstanceToLocal.ToMatrixWithScale();
 			Cluster.ClusterComponent->InvalidateLightingCache();
 		}
-		
-		if( bReAddToHash )
+
+		if (bReAddToHash)
 		{
 			// Re-add instance to the hash
 			InstanceHash->InsertInstance(Instance.Location, InstanceIndex);
@@ -523,14 +523,14 @@ void FFoliageMeshInfo::PostUpdateInstances( class AInstancedFoliageActor* InIFA,
 	}
 }
 
-void FFoliageMeshInfo::PostMoveInstances( class AInstancedFoliageActor* InIFA, const TArray<int32>& InInstancesMoved )
+void FFoliageMeshInfo::PostMoveInstances(class AInstancedFoliageActor* InIFA, const TArray<int32>& InInstancesMoved)
 {
-	PostUpdateInstances( InIFA, InInstancesMoved, true );
+	PostUpdateInstances(InIFA, InInstancesMoved, true);
 }
 
-void FFoliageMeshInfo::DuplicateInstances( class AInstancedFoliageActor* InIFA, class UStaticMesh* InMesh, const TArray<int32>& InInstancesToDuplicate )
+void FFoliageMeshInfo::DuplicateInstances(class AInstancedFoliageActor* InIFA, class UStaticMesh* InMesh, const TArray<int32>& InInstancesToDuplicate)
 {
-	for( TArray<int32>::TConstIterator It(InInstancesToDuplicate); It; ++It )
+	for (TArray<int32>::TConstIterator It(InInstancesToDuplicate); It; ++It)
 	{
 		int32 InstanceIndex = *It;
 		const FFoliageInstance TempInstance = Instances[InstanceIndex];
@@ -542,9 +542,9 @@ void FFoliageMeshInfo::DuplicateInstances( class AInstancedFoliageActor* InIFA,
 int32 FFoliageMeshInfo::GetInstanceCount() const
 {
 	int32 Count = 0;
-	for( int32 Idx=0;Idx<Instances.Num();Idx++ )
+	for (int32 Idx = 0; Idx < Instances.Num(); Idx++)
 	{
-		if( Instances[Idx].ClusterIndex != -1 )
+		if (Instances[Idx].ClusterIndex != -1)
 		{
 			Count++;
 		}
@@ -553,25 +553,25 @@ int32 FFoliageMeshInfo::GetInstanceCount() const
 }
 
 // Destroy existing clusters and reassign all instances to new clusters
-void FFoliageMeshInfo::ReallocateClusters( AInstancedFoliageActor* InIFA, UStaticMesh* InMesh )
+void FFoliageMeshInfo::ReallocateClusters(AInstancedFoliageActor* InIFA, UStaticMesh* InMesh)
 {
 	// Detach all components
 	InIFA->UnregisterAllComponents();
 
 	// Delete any existing clusters from the components array
-	for( int32 ClusterIndex=0;ClusterIndex < InstanceClusters.Num(); ClusterIndex++ )
+	for (int32 ClusterIndex = 0; ClusterIndex < InstanceClusters.Num(); ClusterIndex++)
 	{
 		UInstancedStaticMeshComponent* Component = InstanceClusters[ClusterIndex].ClusterComponent;
-		if( Component != NULL )
+		if (Component != NULL)
 		{
 			Component->bAutoRegister = false;
 		}
 	}
 
 	TArray<FFoliageInstance> OldInstances = Instances;
-	for( int32 Idx=0;Idx<OldInstances.Num();Idx++ )
+	for (int32 Idx = 0; Idx < OldInstances.Num(); Idx++)
 	{
-		if( OldInstances[Idx].ClusterIndex == -1 )
+		if (OldInstances[Idx].ClusterIndex == -1)
 		{
 			OldInstances.RemoveAtSwap(Idx);
 			Idx--;
@@ -591,18 +591,18 @@ void FFoliageMeshInfo::ReallocateClusters( AInstancedFoliageActor* InIFA, UStati
 	SelectedIndices.Empty();
 
 	// Re-add
-	for( int32 Idx=0;Idx<OldInstances.Num();Idx++ )
+	for (int32 Idx = 0; Idx < OldInstances.Num(); Idx++)
 	{
-		AddInstance( InIFA, InMesh, OldInstances[Idx] );
+		AddInstance(InIFA, InMesh, OldInstances[Idx]);
 	}
 
 	InIFA->RegisterAllComponents();
 }
 
 // Update settings in the clusters based on the current settings (eg culling distance, collision, ...)
-void FFoliageMeshInfo::UpdateClusterSettings( AInstancedFoliageActor* InIFA )
+void FFoliageMeshInfo::UpdateClusterSettings(AInstancedFoliageActor* InIFA)
 {
-	for( int32 ClusterIdx=0;ClusterIdx<InstanceClusters.Num();ClusterIdx++ )
+	for (int32 ClusterIdx = 0; ClusterIdx < InstanceClusters.Num(); ClusterIdx++)
 	{
 		UInstancedStaticMeshComponent* ClusterComponent = InstanceClusters[ClusterIdx].ClusterComponent;
 		ClusterComponent->Modify();
@@ -610,7 +610,7 @@ void FFoliageMeshInfo::UpdateClusterSettings( AInstancedFoliageActor* InIFA )
 
 		// Copy settings
 		ClusterComponent->InstanceStartCullDistance = Settings->StartCullDistance;
-		ClusterComponent->InstanceEndCullDistance = Settings->EndCullDistance;	
+		ClusterComponent->InstanceEndCullDistance = Settings->EndCullDistance;
 
 		// Update PrimitiveComponent's culling distance taking into account the radius of the bounds, as
 		// it is based on the center of the component's bounds.
@@ -622,29 +622,29 @@ void FFoliageMeshInfo::UpdateClusterSettings( AInstancedFoliageActor* InIFA )
 	InIFA->MarkComponentsRenderStateDirty();
 }
 
-void FFoliageMeshInfo::GetInstancesInsideSphere( const FSphere& Sphere, TArray<int32>& OutInstances )
+void FFoliageMeshInfo::GetInstancesInsideSphere(const FSphere& Sphere, TArray<int32>& OutInstances)
 {
 	TSet<int32> TempInstances;
-	InstanceHash->GetInstancesOverlappingBox(FBox::BuildAABB(Sphere.Center, FVector(Sphere.W,Sphere.W,Sphere.W)), TempInstances );
+	InstanceHash->GetInstancesOverlappingBox(FBox::BuildAABB(Sphere.Center, FVector(Sphere.W, Sphere.W, Sphere.W)), TempInstances);
 
-	for( TSet<int32>::TConstIterator It(TempInstances); It; ++It )
+	for (TSet<int32>::TConstIterator It(TempInstances); It; ++It)
 	{
-		if( FSphere(Instances[*It].Location,0.f).IsInside(Sphere) )
+		if (FSphere(Instances[*It].Location, 0.f).IsInside(Sphere))
 		{
 			OutInstances.Add(*It);
 		}
-	}	
+	}
 }
 
 // Returns whether or not there is are any instances overlapping the sphere specified
-bool FFoliageMeshInfo::CheckForOverlappingSphere( const FSphere& Sphere )
+bool FFoliageMeshInfo::CheckForOverlappingSphere(const FSphere& Sphere)
 {
 	TSet<int32> TempInstances;
-	InstanceHash->GetInstancesOverlappingBox(FBox::BuildAABB(Sphere.Center, FVector(Sphere.W,Sphere.W,Sphere.W)), TempInstances );
+	InstanceHash->GetInstancesOverlappingBox(FBox::BuildAABB(Sphere.Center, FVector(Sphere.W, Sphere.W, Sphere.W)), TempInstances);
 
-	for( TSet<int32>::TConstIterator It(TempInstances); It; ++It )
+	for (TSet<int32>::TConstIterator It(TempInstances); It; ++It)
 	{
-		if( FSphere(Instances[*It].Location,0.f).IsInside(Sphere) )
+		if (FSphere(Instances[*It].Location, 0.f).IsInside(Sphere))
 		{
 			return true;
 		}
@@ -653,15 +653,15 @@ bool FFoliageMeshInfo::CheckForOverlappingSphere( const FSphere& Sphere )
 }
 
 // Returns whether or not there is are any instances overlapping the instance specified, excluding the set of instances provided
-bool FFoliageMeshInfo::CheckForOverlappingInstanceExcluding( int32 TestInstanceIdx, float Radius, TSet<int32>& ExcludeInstances )
+bool FFoliageMeshInfo::CheckForOverlappingInstanceExcluding(int32 TestInstanceIdx, float Radius, TSet<int32>& ExcludeInstances)
 {
-	FSphere Sphere( Instances[TestInstanceIdx].Location,Radius);
+	FSphere Sphere(Instances[TestInstanceIdx].Location, Radius);
 	TSet<int32> TempInstances;
-	InstanceHash->GetInstancesOverlappingBox(FBox::BuildAABB(Sphere.Center, FVector(Sphere.W,Sphere.W,Sphere.W)), TempInstances );
+	InstanceHash->GetInstancesOverlappingBox(FBox::BuildAABB(Sphere.Center, FVector(Sphere.W, Sphere.W, Sphere.W)), TempInstances);
 
-	for( TSet<int32>::TConstIterator It(TempInstances); It; ++It )
+	for (TSet<int32>::TConstIterator It(TempInstances); It; ++It)
 	{
-		if( *It != TestInstanceIdx && !ExcludeInstances.Contains(*It) && FSphere(Instances[*It].Location,0.f).IsInside(Sphere) )
+		if (*It != TestInstanceIdx && !ExcludeInstances.Contains(*It) && FSphere(Instances[*It].Location, 0.f).IsInside(Sphere))
 		{
 			return true;
 		}
@@ -669,33 +669,33 @@ bool FFoliageMeshInfo::CheckForOverlappingInstanceExcluding( int32 TestInstanceI
 	return false;
 }
 
-void FFoliageMeshInfo::SelectInstances( AInstancedFoliageActor* InIFA, bool bSelect, TArray<int32>& InInstances )
+void FFoliageMeshInfo::SelectInstances(AInstancedFoliageActor* InIFA, bool bSelect, TArray<int32>& InInstances)
 {
 	InIFA->Modify();
 
-	if( bSelect )
+	if (bSelect)
 	{
-		for( int32 i=0;i<InInstances.Num();i++ )
+		for (int32 i = 0; i < InInstances.Num(); i++)
 		{
 			SelectedIndices.AddUnique(InInstances[i]);
 		}
 
-		for( int32 ClusterIdx=0;ClusterIdx<InstanceClusters.Num();ClusterIdx++ )
+		for (int32 ClusterIdx = 0; ClusterIdx < InstanceClusters.Num(); ClusterIdx++)
 		{
 			FFoliageInstanceCluster& Cluster = InstanceClusters[ClusterIdx];
 
 			// Apply any selections in the component
 			Cluster.ClusterComponent->Modify();
 			Cluster.ClusterComponent->MarkRenderStateDirty();
-			if( Cluster.ClusterComponent->SelectedInstances.Num() != Cluster.InstanceIndices.Num() )
+			if (Cluster.ClusterComponent->SelectedInstances.Num() != Cluster.InstanceIndices.Num())
 			{
 				Cluster.ClusterComponent->SelectedInstances.Init(false, Cluster.InstanceIndices.Num());
 			}
 
-			for( int32 ClusterInstanceIdx=0;ClusterInstanceIdx<Cluster.InstanceIndices.Num();ClusterInstanceIdx++ )
+			for (int32 ClusterInstanceIdx = 0; ClusterInstanceIdx < Cluster.InstanceIndices.Num(); ClusterInstanceIdx++)
 			{
 				int32 InstanceIdx = Cluster.InstanceIndices[ClusterInstanceIdx];
-				if( InInstances.Find(InstanceIdx) != INDEX_NONE )
+				if (InInstances.Find(InstanceIdx) != INDEX_NONE)
 				{
 					Cluster.ClusterComponent->SelectedInstances[ClusterInstanceIdx] = true;
 				}
@@ -704,25 +704,25 @@ void FFoliageMeshInfo::SelectInstances( AInstancedFoliageActor* InIFA, bool bSel
 	}
 	else
 	{
-		for( int32 i=0;i<InInstances.Num();i++ )
+		for (int32 i = 0; i < InInstances.Num(); i++)
 		{
 			SelectedIndices.RemoveSingleSwap(InInstances[i]);
 		}
 
-		for( int32 ClusterIdx=0;ClusterIdx<InstanceClusters.Num();ClusterIdx++ )
+		for (int32 ClusterIdx = 0; ClusterIdx < InstanceClusters.Num(); ClusterIdx++)
 		{
 			FFoliageInstanceCluster& Cluster = InstanceClusters[ClusterIdx];
 
-			if( Cluster.ClusterComponent->SelectedInstances.Num() != 0 )
+			if (Cluster.ClusterComponent->SelectedInstances.Num() != 0)
 			{
 				// Remove any selections from the component
 				Cluster.ClusterComponent->Modify();
 				Cluster.ClusterComponent->MarkRenderStateDirty();
 
-				for( int32 ClusterInstanceIdx=0;ClusterInstanceIdx<Cluster.InstanceIndices.Num();ClusterInstanceIdx++ )
+				for (int32 ClusterInstanceIdx = 0; ClusterInstanceIdx < Cluster.InstanceIndices.Num(); ClusterInstanceIdx++)
 				{
 					int32 InstanceIdx = Cluster.InstanceIndices[ClusterInstanceIdx];
-					if( InInstances.Find(InstanceIdx) != INDEX_NONE )
+					if (InInstances.Find(InstanceIdx) != INDEX_NONE)
 					{
 						Cluster.ClusterComponent->SelectedInstances[ClusterInstanceIdx] = false;
 					}
@@ -747,14 +747,14 @@ AInstancedFoliageActor::AInstancedFoliageActor(const class FPostConstructInitial
 }
 
 #if WITH_EDITOR
-AInstancedFoliageActor* AInstancedFoliageActor::GetInstancedFoliageActor(UWorld* InWorld,bool bCreateIfNone)
+AInstancedFoliageActor* AInstancedFoliageActor::GetInstancedFoliageActor(UWorld* InWorld, bool bCreateIfNone)
 {
-	for( TActorIterator<AInstancedFoliageActor> It(InWorld);It;++It )
+	for (TActorIterator<AInstancedFoliageActor> It(InWorld); It; ++It)
 	{
 		AInstancedFoliageActor* InstancedFoliageActor = *It;
-		if( InstancedFoliageActor )
+		if (InstancedFoliageActor)
 		{
-			if( InstancedFoliageActor->GetLevel()->IsCurrentLevel() && !InstancedFoliageActor->IsPendingKill() )
+			if (InstancedFoliageActor->GetLevel()->IsCurrentLevel() && !InstancedFoliageActor->IsPendingKill())
 			{
 				return InstancedFoliageActor;
 			}
@@ -783,41 +783,41 @@ AInstancedFoliageActor* AInstancedFoliageActor::GetInstancedFoliageActorForLevel
 	return NULL;
 }
 
-void AInstancedFoliageActor::SnapInstancesForLandscape( class ULandscapeHeightfieldCollisionComponent* InLandscapeComponent, const FBox& InInstanceBox )
+void AInstancedFoliageActor::SnapInstancesForLandscape(class ULandscapeHeightfieldCollisionComponent* InLandscapeComponent, const FBox& InInstanceBox)
 {
-	for( TMap<class UStaticMesh*, struct FFoliageMeshInfo>::TIterator MeshIt(FoliageMeshes); MeshIt; ++MeshIt )
+	for (auto& MeshPair : FoliageMeshes)
 	{
 		// Find the per-mesh info matching the mesh.
-		FFoliageMeshInfo& MeshInfo = MeshIt.Value();
+		FFoliageMeshInfo& MeshInfo = *MeshPair.Value;
 
 		FFoliageComponentHashInfo* ComponentHashInfo = MeshInfo.ComponentHash.Find(InLandscapeComponent);
-		if( ComponentHashInfo )
+		if (ComponentHashInfo)
 		{
 			float TraceExtentSize = InLandscapeComponent->Bounds.SphereRadius * 2.f + 10.f; // extend a little
-			FVector TraceVector = InLandscapeComponent->GetOwner()->GetRootComponent()->ComponentToWorld.GetUnitAxis( EAxis::Z ) * TraceExtentSize;
+			FVector TraceVector = InLandscapeComponent->GetOwner()->GetRootComponent()->ComponentToWorld.GetUnitAxis(EAxis::Z) * TraceExtentSize;
 
 			bool bFirst = true;
 			TArray<int32> InstancesToRemove;
-			for( TSet<int32>::TConstIterator InstIt(ComponentHashInfo->Instances); InstIt; ++InstIt )
+			for (TSet<int32>::TConstIterator InstIt(ComponentHashInfo->Instances); InstIt; ++InstIt)
 			{
 				int32 InstanceIndex = *InstIt;
 				FFoliageInstance& Instance = MeshInfo.Instances[InstanceIndex];
 
 				// Test location should remove any Z offset
 				FVector TestLocation = FMath::Abs(Instance.ZOffset) > KINDA_SMALL_NUMBER
-					?	(FVector)Instance.GetInstanceWorldTransform().TransformPosition(FVector(0,0,-Instance.ZOffset))
-					:	Instance.Location;
+					? (FVector)Instance.GetInstanceWorldTransform().TransformPosition(FVector(0, 0, -Instance.ZOffset))
+					: Instance.Location;
 
-				if( InInstanceBox.IsInside(TestLocation) )
+				if (InInstanceBox.IsInside(TestLocation))
 				{
-					if( bFirst )
+					if (bFirst)
 					{
 						bFirst = false;
 						Modify();
 					}
 
 					FVector Start = TestLocation + TraceVector;
-					FVector End   = TestLocation - TraceVector;
+					FVector End = TestLocation - TraceVector;
 
 					static FName TraceTag = FName(TEXT("FoliageSnapToLandscape"));
 					TArray<FHitResult> Results;
@@ -825,23 +825,23 @@ void AInstancedFoliageActor::SnapInstancesForLandscape( class ULandscapeHeightfi
 					check(World);
 					// Editor specific landscape heightfield uses ECC_Visibility collision channel
 					World->LineTraceMulti(Results, Start, End, FCollisionQueryParams(TraceTag, true), FCollisionObjectQueryParams(ECollisionChannel::ECC_Visibility));
-				
+
 					bool bFoundHit = false;
-					for( int32 HitIdx=0; HitIdx<Results.Num(); HitIdx++ )
+					for (int32 HitIdx = 0; HitIdx < Results.Num(); HitIdx++)
 					{
 						const FHitResult& Hit = Results[HitIdx];
-						if( Hit.Component == InLandscapeComponent )
+						if (Hit.Component == InLandscapeComponent)
 						{
 							bFoundHit = true;
-							if( (TestLocation - Hit.Location).SizeSquared() > KINDA_SMALL_NUMBER )
+							if ((TestLocation - Hit.Location).SizeSquared() > KINDA_SMALL_NUMBER)
 							{
 								// Remove instance location from the hash. Do not need to update ComponentHash as we re-add below.
 								MeshInfo.InstanceHash->RemoveInstance(Instance.Location, InstanceIndex);
-							
+
 								// Update the instance editor data
 								Instance.Location = Hit.Location;
 
-								if( Instance.Flags & FOLIAGE_AlignToNormal )
+								if (Instance.Flags & FOLIAGE_AlignToNormal)
 								{
 									// Remove previous alignment and align to new normal.
 									Instance.Rotation = Instance.PreAlignRotation;
@@ -849,30 +849,30 @@ void AInstancedFoliageActor::SnapInstancesForLandscape( class ULandscapeHeightfi
 								}
 
 								// Reapply the Z offset in local space
-								if( FMath::Abs(Instance.ZOffset) > KINDA_SMALL_NUMBER )
+								if (FMath::Abs(Instance.ZOffset) > KINDA_SMALL_NUMBER)
 								{
-									Instance.Location = Instance.GetInstanceWorldTransform().TransformPosition(FVector(0,0,Instance.ZOffset));
+									Instance.Location = Instance.GetInstanceWorldTransform().TransformPosition(FVector(0, 0, Instance.ZOffset));
 								}
 
 								// Todo: add do validation with other parameters such as max/min height etc.
 
 								// Update this instances' transform in the UInstancedStaticMeshComponent
-								if( MeshInfo.InstanceClusters.IsValidIndex(Instance.ClusterIndex) )
+								if (MeshInfo.InstanceClusters.IsValidIndex(Instance.ClusterIndex))
 								{
 									FFoliageInstanceCluster& Cluster = MeshInfo.InstanceClusters[Instance.ClusterIndex];
 
 									Cluster.ClusterComponent->Modify();
-									
+
 									int32 ClusterInstanceDataIndex = Cluster.InstanceIndices.Find(InstanceIndex);
-									
-									check(ClusterInstanceDataIndex != INDEX_NONE);					
+
+									check(ClusterInstanceDataIndex != INDEX_NONE);
 									check(Cluster.ClusterComponent->PerInstanceSMData.IsValidIndex(ClusterInstanceDataIndex));
 									Cluster.ClusterComponent->MarkRenderStateDirty();
-									
+
 									// Update bounds
 									FTransform InstanceToWorld = Instance.GetInstanceWorldTransform();
 									Cluster.Bounds = Cluster.Bounds + Cluster.ClusterComponent->StaticMesh->GetBounds().TransformBy(InstanceToWorld);
-									
+
 									// Update transform in InstancedStaticMeshComponent
 									FTransform InstanceToLocal = InstanceToWorld.GetRelativeTransform(Cluster.ClusterComponent->ComponentToWorld);
 									Cluster.ClusterComponent->PerInstanceSMData[ClusterInstanceDataIndex].Transform = InstanceToLocal.ToMatrixWithScale();
@@ -886,7 +886,7 @@ void AInstancedFoliageActor::SnapInstancesForLandscape( class ULandscapeHeightfi
 						}
 					}
 
-					if( !bFoundHit )
+					if (!bFoundHit)
 					{
 						// Couldn't find new spot - remove instance
 						InstancesToRemove.Add(InstanceIndex);
@@ -900,18 +900,18 @@ void AInstancedFoliageActor::SnapInstancesForLandscape( class ULandscapeHeightfi
 	}
 }
 
-void AInstancedFoliageActor::MoveInstancesForMovedComponent( class UActorComponent* InComponent )
+void AInstancedFoliageActor::MoveInstancesForMovedComponent(class UActorComponent* InComponent)
 {
 	bool bUpdatedInstances = false;
 	bool bFirst = true;
 
-	for( TMap<class UStaticMesh*, struct FFoliageMeshInfo>::TIterator MeshIt(FoliageMeshes); MeshIt; ++MeshIt )
+	for (auto& MeshPair : FoliageMeshes)
 	{
-		FFoliageMeshInfo& MeshInfo = MeshIt.Value();
+		FFoliageMeshInfo& MeshInfo = *MeshPair.Value;
 		FFoliageComponentHashInfo* ComponentHashInfo = MeshInfo.ComponentHash.Find(InComponent);
-		if( ComponentHashInfo )
+		if (ComponentHashInfo)
 		{
-			if( bFirst )
+			if (bFirst)
 			{
 				bFirst = false;
 				Modify();
@@ -922,16 +922,16 @@ void AInstancedFoliageActor::MoveInstancesForMovedComponent( class UActorCompone
 			FVector OldDrawScale = ComponentHashInfo->CachedDrawScale;
 			ComponentHashInfo->UpdateLocationFromActor(InComponent);
 
-			for( TSet<int32>::TConstIterator InstIt(ComponentHashInfo->Instances); InstIt; ++InstIt )
+			for (TSet<int32>::TConstIterator InstIt(ComponentHashInfo->Instances); InstIt; ++InstIt)
 			{
 				int32 InstanceIndex = *InstIt;
 				FFoliageInstance& Instance = MeshInfo.Instances[InstanceIndex];
 
 				MeshInfo.InstanceHash->RemoveInstance(Instance.Location, InstanceIndex);
-				
+
 				// Apply change
-				FMatrix DeltaTransform = 
-					FRotationMatrix(Instance.Rotation) * 
+				FMatrix DeltaTransform =
+					FRotationMatrix(Instance.Rotation) *
 					FTranslationMatrix(Instance.Location) *
 					FTranslationMatrix(-OldLocation) *
 					FInverseRotationMatrix(OldRotation) *
@@ -944,12 +944,12 @@ void AInstancedFoliageActor::MoveInstancesForMovedComponent( class UActorCompone
 				Instance.Rotation = DeltaTransform.Rotator();
 
 				// Update this instances' transform in the UInstancedStaticMeshComponent
-				if( MeshInfo.InstanceClusters.IsValidIndex(Instance.ClusterIndex) )
+				if (MeshInfo.InstanceClusters.IsValidIndex(Instance.ClusterIndex))
 				{
 					FFoliageInstanceCluster& Cluster = MeshInfo.InstanceClusters[Instance.ClusterIndex];
 					Cluster.ClusterComponent->Modify();
 					int32 ClusterInstanceDataIndex = Cluster.InstanceIndices.Find(InstanceIndex);
-					check(ClusterInstanceDataIndex != INDEX_NONE);					
+					check(ClusterInstanceDataIndex != INDEX_NONE);
 					check(Cluster.ClusterComponent->PerInstanceSMData.IsValidIndex(ClusterInstanceDataIndex));
 					Cluster.ClusterComponent->MarkRenderStateDirty();
 					bUpdatedInstances = true;
@@ -964,46 +964,42 @@ void AInstancedFoliageActor::MoveInstancesForMovedComponent( class UActorCompone
 
 				// Re-add the new instance location to the hash
 				MeshInfo.InstanceHash->InsertInstance(Instance.Location, InstanceIndex);
-			}	
+			}
 		}
 	}
 }
 
-void AInstancedFoliageActor::DeleteInstancesForComponent( class UActorComponent* InComponent )
+void AInstancedFoliageActor::DeleteInstancesForComponent(class UActorComponent* InComponent)
 {
-	for( TMap<class UStaticMesh*, struct FFoliageMeshInfo>::TIterator MeshIt(FoliageMeshes); MeshIt; ++MeshIt )
+	for (auto& MeshPair : FoliageMeshes)
 	{
-		FFoliageMeshInfo& MeshInfo = MeshIt.Value();
+		FFoliageMeshInfo& MeshInfo = *MeshPair.Value;
 		const FFoliageComponentHashInfo* ComponentHashInfo = MeshInfo.ComponentHash.Find(InComponent);
-		if( ComponentHashInfo )
+		if (ComponentHashInfo)
 		{
 			MeshInfo.RemoveInstances(this, ComponentHashInfo->Instances.Array());
 		}
 	}
 }
 
-void AInstancedFoliageActor::MoveInstancesForComponentToCurrentLevel( class UActorComponent* InComponent )
+void AInstancedFoliageActor::MoveInstancesForComponentToCurrentLevel(class UActorComponent* InComponent)
 {
 	AInstancedFoliageActor* NewIFA = AInstancedFoliageActor::GetInstancedFoliageActor(InComponent->GetWorld());
 
-	for( TMap<class UStaticMesh*, struct FFoliageMeshInfo>::TIterator MeshIt(FoliageMeshes); MeshIt; ++MeshIt )
+	for (auto& MeshPair : FoliageMeshes)
 	{
-		FFoliageMeshInfo& MeshInfo = MeshIt.Value();
+		FFoliageMeshInfo& MeshInfo = *MeshPair.Value;
 
 		const FFoliageComponentHashInfo* ComponentHashInfo = MeshInfo.ComponentHash.Find(InComponent);
-		if( ComponentHashInfo )
+		if (ComponentHashInfo)
 		{
-			FFoliageMeshInfo* NewMeshInfo = NewIFA->FoliageMeshes.Find( MeshIt.Key() );
-			if( !NewMeshInfo )
-			{
-				NewMeshInfo = NewIFA->AddMesh( MeshIt.Key() );
-			}
+			FFoliageMeshInfo* NewMeshInfo = NewIFA->FindOrAddMesh(MeshPair.Key);
 
 			// Add the foliage to the new level
-			for( TSet<int32>::TConstIterator InstIt(ComponentHashInfo->Instances); InstIt; ++InstIt )
+			for (TSet<int32>::TConstIterator InstIt(ComponentHashInfo->Instances); InstIt; ++InstIt)
 			{
 				int32 InstanceIndex = *InstIt;
-				NewMeshInfo->AddInstance( NewIFA, MeshIt.Key(), MeshInfo.Instances[InstanceIndex] );
+				NewMeshInfo->AddInstance(NewIFA, MeshPair.Key, MeshInfo.Instances[InstanceIndex]);
 			}
 
 			// Remove from old level
@@ -1012,7 +1008,7 @@ void AInstancedFoliageActor::MoveInstancesForComponentToCurrentLevel( class UAct
 	}
 }
 
-void AInstancedFoliageActor::MoveInstancesToNewComponent( class UActorComponent* InOldComponent, class UActorComponent* InNewComponent )
+void AInstancedFoliageActor::MoveInstancesToNewComponent(class UActorComponent* InOldComponent, class UActorComponent* InNewComponent)
 {
 	AInstancedFoliageActor* NewIFA = AInstancedFoliageActor::GetInstancedFoliageActorForLevel(InNewComponent->GetTypedOuter<ULevel>());
 	if (!NewIFA)
@@ -1021,9 +1017,9 @@ void AInstancedFoliageActor::MoveInstancesToNewComponent( class UActorComponent*
 	}
 	check(NewIFA);
 
-	for (TPair<UStaticMesh*, FFoliageMeshInfo>& Entry : FoliageMeshes)
+	for (auto& MeshPair : FoliageMeshes)
 	{
-		FFoliageMeshInfo& MeshInfo = Entry.Value;
+		FFoliageMeshInfo& MeshInfo = *MeshPair.Value;
 
 		const FFoliageComponentHashInfo* ComponentHashInfo = MeshInfo.ComponentHash.Find(InOldComponent);
 		if (ComponentHashInfo)
@@ -1046,11 +1042,7 @@ void AInstancedFoliageActor::MoveInstancesToNewComponent( class UActorComponent*
 			}
 			else
 			{
-				FFoliageMeshInfo* NewMeshInfo = NewIFA->FoliageMeshes.Find(Entry.Key);
-				if (!NewMeshInfo)
-				{
-					NewMeshInfo = NewIFA->AddMesh(Entry.Key);
-				}
+				FFoliageMeshInfo* NewMeshInfo = NewIFA->FindOrAddMesh(MeshPair.Key);
 
 				// Add the foliage to the new level
 				for (int32 InstanceIndex : ComponentHashInfo->Instances)
@@ -1058,7 +1050,7 @@ void AInstancedFoliageActor::MoveInstancesToNewComponent( class UActorComponent*
 					FFoliageInstance NewInstance = MeshInfo.Instances[InstanceIndex];
 					NewInstance.Base = InNewComponent;
 					NewInstance.ClusterIndex = INDEX_NONE;
-					NewMeshInfo->AddInstance(NewIFA, Entry.Key, NewInstance);
+					NewMeshInfo->AddInstance(NewIFA, MeshPair.Key, NewInstance);
 				}
 
 				// Remove from old level
@@ -1069,21 +1061,21 @@ void AInstancedFoliageActor::MoveInstancesToNewComponent( class UActorComponent*
 	}
 }
 
-TMap<class UStaticMesh*,TArray<const FFoliageInstancePlacementInfo*> > AInstancedFoliageActor::GetInstancesForComponent( class UActorComponent* InComponent )
+TMap<class UStaticMesh*, TArray<const FFoliageInstancePlacementInfo*> > AInstancedFoliageActor::GetInstancesForComponent(class UActorComponent* InComponent)
 {
-	TMap<class UStaticMesh*,TArray<const struct FFoliageInstancePlacementInfo*> > Result;
+	TMap<class UStaticMesh*, TArray<const struct FFoliageInstancePlacementInfo*> > Result;
 
-	for( TMap<class UStaticMesh*, struct FFoliageMeshInfo>::TConstIterator MeshIt(FoliageMeshes); MeshIt; ++MeshIt )
+	for (auto& MeshPair : FoliageMeshes)
 	{
-		const FFoliageMeshInfo& MeshInfo = MeshIt.Value();
+		const FFoliageMeshInfo& MeshInfo = *MeshPair.Value;
 
 		const FFoliageComponentHashInfo* ComponentHashInfo = MeshInfo.ComponentHash.Find(InComponent);
-		if( ComponentHashInfo )
+		if (ComponentHashInfo)
 		{
-			TArray<const FFoliageInstancePlacementInfo*>& Array = Result.Add(MeshIt.Key(),TArray<const FFoliageInstancePlacementInfo*>() );
+			TArray<const FFoliageInstancePlacementInfo*>& Array = Result.Add(MeshPair.Key, TArray<const FFoliageInstancePlacementInfo*>());
 			Array.Empty(ComponentHashInfo->Instances.Num());
 
-			for( TSet<int32>::TConstIterator InstIt(ComponentHashInfo->Instances); InstIt; ++InstIt )
+			for (TSet<int32>::TConstIterator InstIt(ComponentHashInfo->Instances); InstIt; ++InstIt)
 			{
 				int32 InstanceIndex = *InstIt;
 				const FFoliageInstancePlacementInfo* Instance = &MeshInfo.Instances[InstanceIndex];
@@ -1095,22 +1087,36 @@ TMap<class UStaticMesh*,TArray<const FFoliageInstancePlacementInfo*> > AInstance
 	return Result;
 }
 
-struct FFoliageMeshInfo* AInstancedFoliageActor::AddMesh( class UStaticMesh* InMesh )
+FFoliageMeshInfo* AInstancedFoliageActor::FindMesh(UStaticMesh* InMesh)
+{
+	TUniqueObj<FFoliageMeshInfo>* MeshInfoEntry = FoliageMeshes.Find(InMesh);
+	FFoliageMeshInfo* MeshInfo = MeshInfoEntry ? &MeshInfoEntry->Get() : nullptr;
+	return MeshInfo;
+}
+
+FFoliageMeshInfo* AInstancedFoliageActor::FindOrAddMesh(UStaticMesh* InMesh)
+{
+	TUniqueObj<FFoliageMeshInfo>* MeshInfoEntry = FoliageMeshes.Find(InMesh);
+	FFoliageMeshInfo* MeshInfo = MeshInfoEntry ? &MeshInfoEntry->Get() : AddMesh(InMesh);
+	return MeshInfo;
+}
+
+struct FFoliageMeshInfo* AInstancedFoliageActor::AddMesh(class UStaticMesh* InMesh)
 {
-	check( FoliageMeshes.Find(InMesh) == NULL );
+	check(FindMesh(InMesh) == NULL);
 
 	MarkPackageDirty();
 
 	int32 MaxDisplayOrder = 0;
-	for( TMap<class UStaticMesh*, struct FFoliageMeshInfo>::TIterator It(FoliageMeshes); It; ++It )
+	for (auto& MeshPair : FoliageMeshes)
 	{
-		if( It.Value().Settings->DisplayOrder > MaxDisplayOrder )
+		if (MeshPair.Value->Settings->DisplayOrder > MaxDisplayOrder)
 		{
-			MaxDisplayOrder = It.Value().Settings->DisplayOrder;
+			MaxDisplayOrder = MeshPair.Value->Settings->DisplayOrder;
 		}
 	}
 
-	FFoliageMeshInfo& MeshInfo = FoliageMeshes.Add(InMesh, FFoliageMeshInfo());
+	FFoliageMeshInfo& MeshInfo = *FoliageMeshes.Add(InMesh);
 #if WITH_EDITORONLY_DATA
 	UInstancedFoliageSettings* DefaultSettings = InMesh->FoliageDefaultSettings;
 	if( DefaultSettings != NULL )
@@ -1129,10 +1135,10 @@ struct FFoliageMeshInfo* AInstancedFoliageActor::AddMesh( class UStaticMesh* InM
 	// the component will be treated as destroyed, in that undo an add will
 	// actually work
 	MeshInfo.Settings->SetFlags(RF_PendingKill);
-	MeshInfo.Settings->Modify( false );
+	MeshInfo.Settings->Modify(false);
 	MeshInfo.Settings->ClearFlags(RF_PendingKill);
 	MeshInfo.Settings->IsSelected = true;
-	MeshInfo.Settings->DisplayOrder = MaxDisplayOrder+1;
+	MeshInfo.Settings->DisplayOrder = MaxDisplayOrder + 1;
 	MeshInfo.Settings->MeshBounds = MeshBounds;
 
 	// Make bottom only bound
@@ -1145,9 +1151,9 @@ struct FFoliageMeshInfo* AInstancedFoliageActor::AddMesh( class UStaticMesh* InM
 	if (InMesh->RenderData)
 	{
 		FPositionVertexBuffer& PositionVertexBuffer = InMesh->RenderData->LODResources[0].PositionVertexBuffer;
-		for( uint32 Index = 0; Index < PositionVertexBuffer.GetNumVertices(); ++Index )
+		for (uint32 Index = 0; Index < PositionVertexBuffer.GetNumVertices(); ++Index)
 		{
-			const FVector& Pos = PositionVertexBuffer.VertexPosition( Index );
+			const FVector& Pos = PositionVertexBuffer.VertexPosition(Index);
 			if (Pos.Z < LowBound.Max.Z)
 			{
 				MinX = FMath::Min(MinX, Pos.X);
@@ -1158,25 +1164,25 @@ struct FFoliageMeshInfo* AInstancedFoliageActor::AddMesh( class UStaticMesh* InM
 		}
 	}
 
-	MeshInfo.Settings->LowBoundOriginRadius = FVector((MinX+MaxX), (MinY+MaxY), FMath::Sqrt( FMath::Square(MaxX-MinX) + FMath::Square(MaxY-MinY)) ) * 0.5f;
+	MeshInfo.Settings->LowBoundOriginRadius = FVector((MinX + MaxX), (MinY + MaxY), FMath::Sqrt(FMath::Square(MaxX - MinX) + FMath::Square(MaxY - MinY))) * 0.5f;
 
 	return &MeshInfo;
 }
 
-void AInstancedFoliageActor::RemoveMesh( class UStaticMesh* InMesh )
+void AInstancedFoliageActor::RemoveMesh(class UStaticMesh* InMesh)
 {
 	Modify();
 	MarkPackageDirty();
 	UnregisterAllComponents();
 
 	// Remove all components for this mesh from the Components array.
-	FFoliageMeshInfo* MeshInfo =  FoliageMeshes.Find(InMesh);
-	if( MeshInfo )
+	FFoliageMeshInfo* MeshInfo = FindMesh(InMesh);
+	if (MeshInfo)
 	{
-		for( int32 ClusterIndex=0;ClusterIndex < MeshInfo->InstanceClusters.Num(); ClusterIndex++ )
+		for (int32 ClusterIndex = 0; ClusterIndex < MeshInfo->InstanceClusters.Num(); ClusterIndex++)
 		{
 			UInstancedStaticMeshComponent* Component = MeshInfo->InstanceClusters[ClusterIndex].ClusterComponent;
-			if( Component != NULL )
+			if (Component != NULL)
 			{
 				Component->bAutoRegister = false;
 			}
@@ -1188,25 +1194,25 @@ void AInstancedFoliageActor::RemoveMesh( class UStaticMesh* InMesh )
 	CheckSelection();
 }
 
-void AInstancedFoliageActor::SelectInstance( class UInstancedStaticMeshComponent* InComponent, int32 InComponentInstanceIndex, bool bToggle )
+void AInstancedFoliageActor::SelectInstance(class UInstancedStaticMeshComponent* InComponent, int32 InComponentInstanceIndex, bool bToggle)
 {
 	bool bNeedsUpdate = false;
 
 	Modify();
 
 	// If we're not toggling, we need to first deselect everything else
-	if( !bToggle )
+	if (!bToggle)
 	{
-		for( TMap<class UStaticMesh*, struct FFoliageMeshInfo>::TIterator MeshIt(FoliageMeshes); MeshIt; ++MeshIt )
+		for (auto& MeshPair : FoliageMeshes)
 		{
-			FFoliageMeshInfo& Mesh = MeshIt.Value();
+			FFoliageMeshInfo& MeshInfo = *MeshPair.Value;
 
-			if( Mesh.SelectedIndices.Num() > 0 )
+			if (MeshInfo.SelectedIndices.Num() > 0)
 			{
-				for( int32 ClusterIdx=0;ClusterIdx<Mesh.InstanceClusters.Num();ClusterIdx++ )
+				for (int32 ClusterIdx = 0; ClusterIdx < MeshInfo.InstanceClusters.Num(); ClusterIdx++)
 				{
-					FFoliageInstanceCluster& Cluster = Mesh.InstanceClusters[ClusterIdx];
-					if( Cluster.ClusterComponent->SelectedInstances.Num() > 0 )
+					FFoliageInstanceCluster& Cluster = MeshInfo.InstanceClusters[ClusterIdx];
+					if (Cluster.ClusterComponent->SelectedInstances.Num() > 0)
 					{
 						Cluster.ClusterComponent->Modify();
 						Cluster.ClusterComponent->SelectedInstances.Empty();
@@ -1216,41 +1222,41 @@ void AInstancedFoliageActor::SelectInstance( class UInstancedStaticMeshComponent
 				bNeedsUpdate = true;
 			}
 
-			Mesh.SelectedIndices.Empty();
+			MeshInfo.SelectedIndices.Empty();
 		}
 	}
 
-	if( InComponent )
+	if (InComponent)
 	{
-		FFoliageMeshInfo* Mesh = FoliageMeshes.Find(InComponent->StaticMesh);
+		FFoliageMeshInfo* MeshInfo = FindMesh(InComponent->StaticMesh);
 
-		if( Mesh )
+		if (MeshInfo)
 		{
-			for( int32 ClusterIdx=0;ClusterIdx<Mesh->InstanceClusters.Num();ClusterIdx++ )
+			for (int32 ClusterIdx = 0; ClusterIdx < MeshInfo->InstanceClusters.Num(); ClusterIdx++)
 			{
-				FFoliageInstanceCluster& Cluster = Mesh->InstanceClusters[ClusterIdx];
-				if( Cluster.ClusterComponent == InComponent )
+				FFoliageInstanceCluster& Cluster = MeshInfo->InstanceClusters[ClusterIdx];
+				if (Cluster.ClusterComponent == InComponent)
 				{
 					InComponent->Modify();
 					int32 InstanceIndex = Cluster.InstanceIndices[InComponentInstanceIndex];
-					int32 SelectedIndex = Mesh->SelectedIndices.Find(InstanceIndex);
+					int32 SelectedIndex = MeshInfo->SelectedIndices.Find(InstanceIndex);
 
 					bNeedsUpdate = true;
 
 					// Deselect if it's already selected.
-					if( InComponentInstanceIndex < InComponent->SelectedInstances.Num() )
+					if (InComponentInstanceIndex < InComponent->SelectedInstances.Num())
 					{
 						InComponent->SelectedInstances[InComponentInstanceIndex] = false;
 						InComponent->MarkRenderStateDirty();
 					}
-					if( SelectedIndex != INDEX_NONE )
+					if (SelectedIndex != INDEX_NONE)
 					{
-						Mesh->SelectedIndices.RemoveAt(SelectedIndex);
+						MeshInfo->SelectedIndices.RemoveAt(SelectedIndex);
 					}
 
-					if( bToggle && SelectedIndex != INDEX_NONE)
+					if (bToggle && SelectedIndex != INDEX_NONE)
 					{
-						if( SelectedMesh == InComponent->StaticMesh && Mesh->SelectedIndices.Num() == 0 )
+						if (SelectedMesh == InComponent->StaticMesh && MeshInfo->SelectedIndices.Num() == 0)
 						{
 							SelectedMesh = NULL;
 						}
@@ -1258,7 +1264,7 @@ void AInstancedFoliageActor::SelectInstance( class UInstancedStaticMeshComponent
 					else
 					{
 						// Add the selection
-						if( InComponent->SelectedInstances.Num() < InComponent->PerInstanceSMData.Num() )
+						if (InComponent->SelectedInstances.Num() < InComponent->PerInstanceSMData.Num())
 						{
 							InComponent->SelectedInstances.Init(false, Cluster.InstanceIndices.Num());
 						}
@@ -1266,7 +1272,7 @@ void AInstancedFoliageActor::SelectInstance( class UInstancedStaticMeshComponent
 						InComponent->MarkRenderStateDirty();
 
 						SelectedMesh = InComponent->StaticMesh;
-						Mesh->SelectedIndices.Insert(InstanceIndex,0);
+						MeshInfo->SelectedIndices.Insert(InstanceIndex, 0);
 					}
 					break;
 				}
@@ -1274,7 +1280,7 @@ void AInstancedFoliageActor::SelectInstance( class UInstancedStaticMeshComponent
 		}
 	}
 
-	if( bNeedsUpdate )
+	if (bNeedsUpdate)
 	{
 		// Update selection
 		MarkComponentsRenderStateDirty();
@@ -1283,31 +1289,31 @@ void AInstancedFoliageActor::SelectInstance( class UInstancedStaticMeshComponent
 	CheckSelection();
 }
 
-void AInstancedFoliageActor::ApplySelectionToComponents( bool bApply )
+void AInstancedFoliageActor::ApplySelectionToComponents(bool bApply)
 {
 	bool bNeedsUpdate = false;
 
-	for( TMap<class UStaticMesh*, struct FFoliageMeshInfo>::TIterator MeshIt(FoliageMeshes); MeshIt; ++MeshIt )
+	for (auto& MeshPair : FoliageMeshes)
 	{
-		FFoliageMeshInfo& Mesh = MeshIt.Value();
+		FFoliageMeshInfo& MeshInfo = *MeshPair.Value;
 
-		if( bApply )
+		if (bApply)
 		{
-			if( Mesh.SelectedIndices.Num() > 0 )
+			if (MeshInfo.SelectedIndices.Num() > 0)
 			{
-				for( int32 ClusterIdx=0;ClusterIdx<Mesh.InstanceClusters.Num();ClusterIdx++ )
+				for (int32 ClusterIdx = 0; ClusterIdx < MeshInfo.InstanceClusters.Num(); ClusterIdx++)
 				{
-					FFoliageInstanceCluster& Cluster = Mesh.InstanceClusters[ClusterIdx];
+					FFoliageInstanceCluster& Cluster = MeshInfo.InstanceClusters[ClusterIdx];
 
 					// Apply any selections in the component
 					Cluster.ClusterComponent->SelectedInstances.Init(false, Cluster.InstanceIndices.Num());
 					Cluster.ClusterComponent->MarkRenderStateDirty();
 					bNeedsUpdate = true;
 
-					for( int32 ClusterInstanceIdx=0;ClusterInstanceIdx<Cluster.InstanceIndices.Num();ClusterInstanceIdx++ )
+					for (int32 ClusterInstanceIdx = 0; ClusterInstanceIdx < Cluster.InstanceIndices.Num(); ClusterInstanceIdx++)
 					{
 						int32 InstanceIdx = Cluster.InstanceIndices[ClusterInstanceIdx];
-						if( Mesh.SelectedIndices.Find(InstanceIdx) != INDEX_NONE )
+						if (MeshInfo.SelectedIndices.Find(InstanceIdx) != INDEX_NONE)
 						{
 							Cluster.ClusterComponent->SelectedInstances[ClusterInstanceIdx] = true;
 						}
@@ -1317,11 +1323,11 @@ void AInstancedFoliageActor::ApplySelectionToComponents( bool bApply )
 		}
 		else
 		{
-			for( int32 ClusterIdx=0;ClusterIdx<Mesh.InstanceClusters.Num();ClusterIdx++ )
+			for (int32 ClusterIdx = 0; ClusterIdx < MeshInfo.InstanceClusters.Num(); ClusterIdx++)
 			{
-				FFoliageInstanceCluster& Cluster = Mesh.InstanceClusters[ClusterIdx];
+				FFoliageInstanceCluster& Cluster = MeshInfo.InstanceClusters[ClusterIdx];
 
-				if( Cluster.ClusterComponent->SelectedInstances.Num() > 0 )
+				if (Cluster.ClusterComponent->SelectedInstances.Num() > 0)
 				{
 					// remove any selections in the component
 					Cluster.ClusterComponent->SelectedInstances.Empty();
@@ -1332,7 +1338,7 @@ void AInstancedFoliageActor::ApplySelectionToComponents( bool bApply )
 		}
 	}
 
-	if( bNeedsUpdate )
+	if (bNeedsUpdate)
 	{
 		// Update selection
 		MarkComponentsRenderStateDirty();
@@ -1342,10 +1348,10 @@ void AInstancedFoliageActor::ApplySelectionToComponents( bool bApply )
 void AInstancedFoliageActor::CheckSelection()
 {
 	// Check if we have to change the selection.
-	if( SelectedMesh != NULL )
+	if (SelectedMesh != NULL)
 	{
-		FFoliageMeshInfo* Mesh = FoliageMeshes.Find(SelectedMesh);
-		if( Mesh && Mesh->SelectedIndices.Num() > 0 )
+		FFoliageMeshInfo* MeshInfo = FindMesh(SelectedMesh);
+		if (MeshInfo && MeshInfo->SelectedIndices.Num() > 0)
 		{
 			return;
 		}
@@ -1354,12 +1360,12 @@ void AInstancedFoliageActor::CheckSelection()
 	SelectedMesh = NULL;
 
 	// Try to find a new selection
-	for( TMap<class UStaticMesh*, struct FFoliageMeshInfo>::TIterator MeshIt(FoliageMeshes); MeshIt; ++MeshIt )
+	for (auto& MeshPair : FoliageMeshes)
 	{
-		FFoliageMeshInfo& Mesh = MeshIt.Value();
-		if( Mesh.SelectedIndices.Num() > 0 )
+		FFoliageMeshInfo& MeshInfo = *MeshPair.Value;
+		if (MeshInfo.SelectedIndices.Num() > 0)
 		{
-			SelectedMesh = MeshIt.Key();
+			SelectedMesh = MeshPair.Key;
 			return;
 		}
 	}
@@ -1367,14 +1373,14 @@ void AInstancedFoliageActor::CheckSelection()
 
 FVector AInstancedFoliageActor::GetSelectionLocation()
 {
-	FVector Result(0,0,0);
+	FVector Result(0, 0, 0);
 
-	if( SelectedMesh != NULL )
+	if (SelectedMesh != NULL)
 	{
-		FFoliageMeshInfo* Mesh = FoliageMeshes.Find(SelectedMesh);
-		if( Mesh && Mesh->SelectedIndices.Num() > 0 )
+		FFoliageMeshInfo* MeshInfo = FindMesh(SelectedMesh);
+		if (MeshInfo && MeshInfo->SelectedIndices.Num() > 0)
 		{
-			Result = Mesh->Instances[Mesh->SelectedIndices[0]].Location;
+			Result = MeshInfo->Instances[MeshInfo->SelectedIndices[0]].Location;
 		}
 	}
 
@@ -1392,18 +1398,18 @@ void AInstancedFoliageActor::MapRebuild()
 	check(World);
 
 	// For each foliage brush, represented by the mesh/info pair
-	for( TMap<UStaticMesh*, FFoliageMeshInfo>::TConstIterator MeshIt(FoliageMeshes); MeshIt; ++MeshIt )
+	for (auto& MeshPair : FoliageMeshes)
 	{
 		// each target component has some foliage instances
-		FFoliageMeshInfo const& MeshInfo = MeshIt.Value();
-		UStaticMesh* Mesh = MeshIt.Key();
+		FFoliageMeshInfo const& MeshInfo = *MeshPair.Value;
+		UStaticMesh* Mesh = MeshPair.Key;
 		check(Mesh);
-		
+
 		for (TMap<UActorComponent*, FFoliageComponentHashInfo >::TConstIterator ComponentFoliageIt(MeshInfo.ComponentHash); ComponentFoliageIt; ++ComponentFoliageIt)
 		{
 			// BSP components are UModelComponents - they are the only ones we need to change
 			UActorComponent* TargetComponent = ComponentFoliageIt.Key();
-			if(!TargetComponent || TargetComponent->IsA(UModelComponent::StaticClass()))
+			if (!TargetComponent || TargetComponent->IsA(UModelComponent::StaticClass()))
 			{
 				// Delete it later
 				RemovedComponents.Add(TargetComponent);
@@ -1411,7 +1417,7 @@ void AInstancedFoliageActor::MapRebuild()
 				FFoliageComponentHashInfo const& FoliageInfo = ComponentFoliageIt.Value();
 
 				// We have to test each instance to see if we can migrate it across
-				for( TSet<int32>::TConstIterator InstanceIt(FoliageInfo.Instances); InstanceIt; ++InstanceIt )
+				for (TSet<int32>::TConstIterator InstanceIt(FoliageInfo.Instances); InstanceIt; ++InstanceIt)
 				{
 					// Use a line test against the world, similar to FoliageEditMode.
 					check(MeshInfo.Instances.IsValidIndex(*InstanceIt));
@@ -1419,7 +1425,7 @@ void AInstancedFoliageActor::MapRebuild()
 
 					FFoliageInstance NewInstance = Instance;
 					NewInstance.ClusterIndex = -1;
-					
+
 					FTransform InstanceToWorld = Instance.GetInstanceWorldTransform();
 					FVector Down(-FVector::UpVector);
 					FVector Start(InstanceToWorld.TransformPosition(FVector::UpVector));
@@ -1427,8 +1433,8 @@ void AInstancedFoliageActor::MapRebuild()
 
 					FHitResult Result;
 					bool bHit = World->LineTraceSingle(Result, Start, End, FCollisionQueryParams(true), FCollisionObjectQueryParams(ECC_WorldStatic));
-					
-					if(bHit && Result.Component.IsValid() && Result.Component->IsA(UModelComponent::StaticClass()))
+
+					if (bHit && Result.Component.IsValid() && Result.Component->IsA(UModelComponent::StaticClass()))
 					{
 						NewInstance.Base = Cast<UActorComponent>(Result.Component.Get());
 					}
@@ -1455,8 +1461,8 @@ void AInstancedFoliageActor::MapRebuild()
 	{
 		UStaticMesh* Mesh = NewInstanceIt.Key();
 		check(Mesh);
-		FFoliageMeshInfo& MeshInfo = FoliageMeshes.FindOrAdd(Mesh);
-		for(TArray<FFoliageInstance>::TConstIterator InstanceIt(NewInstanceIt.Value()); InstanceIt; ++InstanceIt)
+		FFoliageMeshInfo& MeshInfo = *FindOrAddMesh(Mesh);
+		for (TArray<FFoliageInstance>::TConstIterator InstanceIt(NewInstanceIt.Value()); InstanceIt; ++InstanceIt)
 		{
 			MeshInfo.AddInstance(this, Mesh, *InstanceIt);
 		}
@@ -1472,9 +1478,9 @@ void AInstancedFoliageActor::Serialize(FArchive& Ar)
 
 	if (Ar.UE4Ver() < VER_UE4_FOLIAGE_MOVABLE_MOBILITY)
 	{
-		for (const TPair<UStaticMesh*, FFoliageMeshInfo>& MeshInfo : FoliageMeshes)
+		for (auto& MeshInfo : FoliageMeshes)
 		{
-			for (const FFoliageInstanceCluster& Cluster : MeshInfo.Value.InstanceClusters)
+			for (const FFoliageInstanceCluster& Cluster : MeshInfo.Value->InstanceClusters)
 			{
 				if (Ar.UE4Ver() < VER_UE4_FOLIAGE_MOVABLE_MOBILITY)
 					Cluster.ClusterComponent->SetMobility(EComponentMobility::Movable);
@@ -1493,7 +1499,7 @@ void AInstancedFoliageActor::PostLoad()
 	{
 		for (auto& MeshPair : FoliageMeshes)
 		{
-			for (FFoliageInstance& Instance : MeshPair.Value.Instances)
+			for (FFoliageInstance& Instance : MeshPair.Value->Instances)
 			{
 				// Clear out the Base for any instances based on blueprint-created components,
 				// as those components will be destroyed when the construction scripts are
@@ -1506,12 +1512,12 @@ void AInstancedFoliageActor::PostLoad()
 		}
 	}
 
-	if( GIsEditor )
+	if (GIsEditor)
 	{
-		for( TMap<class UStaticMesh*, struct FFoliageMeshInfo>::TIterator MeshIt(FoliageMeshes); MeshIt; ++MeshIt )
+		for (auto MeshIt = FoliageMeshes.CreateIterator(); MeshIt; ++MeshIt)
 		{
 			// Remove any NULL entries.
-			if( MeshIt.Key() == NULL )
+			if (MeshIt.Key() == NULL)
 			{
 				FMessageLog("MapCheck").Warning()
 					->AddToken(FUObjectToken::Create(this))
@@ -1522,13 +1528,13 @@ void AInstancedFoliageActor::PostLoad()
 			}
 
 			// Find the per-mesh info matching the mesh.
-			FFoliageMeshInfo& MeshInfo = MeshIt.Value();
+			FFoliageMeshInfo& MeshInfo = *MeshIt.Value();
 
 			// Update the FreeInstanceIndices list and hash.
-			for( int32 InstanceIdx=0; InstanceIdx<MeshInfo.Instances.Num(); InstanceIdx++ )
+			for (int32 InstanceIdx = 0; InstanceIdx < MeshInfo.Instances.Num(); InstanceIdx++)
 			{
 				FFoliageInstance& Instance = MeshInfo.Instances[InstanceIdx];
-				if( Instance.ClusterIndex == -1 )
+				if (Instance.ClusterIndex == -1)
 				{
 					// Add invalid instances to the FreeInstanceIndices list.
 					MeshInfo.FreeInstanceIndices.Add(InstanceIdx);
@@ -1539,7 +1545,7 @@ void AInstancedFoliageActor::PostLoad()
 					// Add valid instances to the hash.
 					MeshInfo.InstanceHash->InsertInstance(Instance.Location, InstanceIdx);
 					FFoliageComponentHashInfo* ComponentHashInfo = MeshInfo.ComponentHash.Find(Instance.Base);
-					if( ComponentHashInfo == NULL )
+					if (ComponentHashInfo == NULL)
 					{
 						ComponentHashInfo = &MeshInfo.ComponentHash.Add(Instance.Base, FFoliageComponentHashInfo(Instance.Base));
 					}
@@ -1549,21 +1555,21 @@ void AInstancedFoliageActor::PostLoad()
 
 			// Check if any of the clusters are NULL
 			int32 NumMissingComponents = 0;
-			for( int32 ClusterIdx=0;ClusterIdx<MeshInfo.InstanceClusters.Num();ClusterIdx++ )
+			for (int32 ClusterIdx = 0; ClusterIdx < MeshInfo.InstanceClusters.Num(); ClusterIdx++)
 			{
-				if( MeshInfo.InstanceClusters[ClusterIdx].ClusterComponent == NULL )
+				if (MeshInfo.InstanceClusters[ClusterIdx].ClusterComponent == NULL)
 				{
 					NumMissingComponents++;
 				}
 			}
-			if( NumMissingComponents> 0 )
+			if (NumMissingComponents > 0)
 			{
 				FFormatNamedArguments Arguments;
 				Arguments.Add(TEXT("MissingCount"), NumMissingComponents);
 				Arguments.Add(TEXT("MeshName"), FText::FromString(MeshIt.Key()->GetName()));
 				FMessageLog("MapCheck").Warning()
 					->AddToken(FUObjectToken::Create(this))
-					->AddToken(FTextToken::Create(FText::Format(LOCTEXT("MapCheck_Message_FoliageMissingClusterComponent", "Foliage in this map is missing {MissingCount} cluster component(s) for static mesh {MeshName}. Opening the Foliage tool will this problem."), Arguments) ))
+					->AddToken(FTextToken::Create(FText::Format(LOCTEXT("MapCheck_Message_FoliageMissingClusterComponent", "Foliage in this map is missing {MissingCount} cluster component(s) for static mesh {MeshName}. Opening the Foliage tool will this problem."), Arguments)))
 					->AddToken(FMapErrorToken::Create(FMapErrors::FoliageMissingClusterComponent));
 			}
 		}
@@ -1576,42 +1582,42 @@ void AInstancedFoliageActor::PostLoad()
 void AInstancedFoliageActor::AddReferencedObjects(UObject* InThis, FReferenceCollector& Collector)
 {
 	AInstancedFoliageActor* This = CastChecked<AInstancedFoliageActor>(InThis);
-	if( This->SelectedMesh )
+	if (This->SelectedMesh)
 	{
-		Collector.AddReferencedObject( This->SelectedMesh, This );
+		Collector.AddReferencedObject(This->SelectedMesh, This);
 	}
 
-	for( TMap<class UStaticMesh*, struct FFoliageMeshInfo>::TIterator MeshIt(This->FoliageMeshes); MeshIt; ++MeshIt )
+	for (auto& MeshPair : This->FoliageMeshes)
 	{
-		Collector.AddReferencedObject( MeshIt.Key(), This );
-		FFoliageMeshInfo& MeshInfo = MeshIt.Value();
+		Collector.AddReferencedObject(MeshPair.Key, This);
+		FFoliageMeshInfo& MeshInfo = *MeshPair.Value;
 
-		if(MeshInfo.Settings != NULL )
+		if (MeshInfo.Settings != NULL)
 		{
-			Collector.AddReferencedObject( MeshInfo.Settings, This );
+			Collector.AddReferencedObject(MeshInfo.Settings, This);
 		}
 
-		for( int32 InstanceIdx=0;InstanceIdx<MeshInfo.Instances.Num();InstanceIdx++ )
+		for (int32 InstanceIdx = 0; InstanceIdx < MeshInfo.Instances.Num(); InstanceIdx++)
 		{
 			class UActorComponent* Base = MeshInfo.Instances[InstanceIdx].Base;
-			if( Base != NULL )
+			if (Base != NULL)
 			{
-				Collector.AddReferencedObject( Base, This );
+				Collector.AddReferencedObject(Base, This);
 			}
 		}
 
-		for( int32 ClusterIdx=0;ClusterIdx<MeshInfo.InstanceClusters.Num();ClusterIdx++ )
+		for (int32 ClusterIdx = 0; ClusterIdx < MeshInfo.InstanceClusters.Num(); ClusterIdx++)
 		{
-			Collector.AddReferencedObject( MeshInfo.InstanceClusters[ClusterIdx].ClusterComponent, This );	
+			Collector.AddReferencedObject(MeshInfo.InstanceClusters[ClusterIdx].ClusterComponent, This);
 		}
 	}
 	Super::AddReferencedObjects(This, Collector);
 }
 
 /** InstancedStaticMeshInstance hit proxy */
-void HInstancedStaticMeshInstance::AddReferencedObjects( FReferenceCollector& Collector )
+void HInstancedStaticMeshInstance::AddReferencedObjects(FReferenceCollector& Collector)
 {
-	Collector.AddReferencedObject( Component );
+	Collector.AddReferencedObject(Component);
 }
 
 void AInstancedFoliageActor::ApplyWorldOffset(const FVector& InOffset, bool bWorldShift)
@@ -1626,21 +1632,21 @@ void AInstancedFoliageActor::ApplyWorldOffset(const FVector& InOffset, bool bWor
 
 	if (GIsEditor)
 	{
-		for (auto MeshIt = FoliageMeshes.CreateIterator(); MeshIt; ++MeshIt)
+		for (auto& MeshPair : FoliageMeshes)
 		{
-			FFoliageMeshInfo& MeshInfo = MeshIt.Value();
-			
+			FFoliageMeshInfo& MeshInfo = *MeshPair.Value;
+
 			for (int32 ClusterIdx = 0; ClusterIdx < MeshInfo.InstanceClusters.Num(); ClusterIdx++)
 			{
 				FFoliageInstanceCluster& Cluster = MeshInfo.InstanceClusters[ClusterIdx];
-				Cluster.Bounds.Origin+= InOffset;
+				Cluster.Bounds.Origin += InOffset;
 			}
-			
+
 			MeshInfo.InstanceHash->Empty();
 			for (int32 InstanceIdx = 0; InstanceIdx < MeshInfo.Instances.Num(); InstanceIdx++)
 			{
 				FFoliageInstance& Instance = MeshInfo.Instances[InstanceIdx];
-				Instance.Location+= InOffset;
+				Instance.Location += InOffset;
 				// Rehash instance location
 				MeshInfo.InstanceHash->InsertInstance(Instance.Location, InstanceIdx);
 			}
@@ -1649,7 +1655,7 @@ void AInstancedFoliageActor::ApplyWorldOffset(const FVector& InOffset, bool bWor
 			{
 				// We assume here that component we painted foliage on will be shifted by same value
 				FFoliageComponentHashInfo& Info = It.Value();
-				Info.CachedLocation+= InOffset;
+				Info.CachedLocation += InOffset;
 			}
 		}
 	}
