Removed EXPERIMENTAL_PARALLEL_CODE

diff --git a/Engine/Source/Runtime/Core/Public/Async/PhysXASync.h b/Engine/Source/Runtime/Core/Public/Async/PhysXASync.h
--- a/Engine/Source/Runtime/Core/Public/Async/PhysXASync.h
+++ b/Engine/Source/Runtime/Core/Public/Async/PhysXASync.h
@@ -2,7 +2,7 @@
 
 #pragma once
 
-#if EXPERIMENTAL_PARALLEL_CODE
+#if 0
 
 /** Implements basic multi-reader and single-writer synchronization primitive. */
 class FMultiReaderSingleWriter
diff --git a/Engine/Source/Runtime/Core/Public/HAL/ThreadingBase.h b/Engine/Source/Runtime/Core/Public/HAL/ThreadingBase.h
--- a/Engine/Source/Runtime/Core/Public/HAL/ThreadingBase.h
+++ b/Engine/Source/Runtime/Core/Public/HAL/ThreadingBase.h
@@ -6,10 +6,6 @@
 
 #pragma once
 
-#ifndef EXPERIMENTAL_PARALLEL_CODE
-	#error EXPERIMENTAL_PARALLEL_CODE not defined!
-#endif
-
 /**
  * The list of enumerated thread priorities we support
  */
diff --git a/Engine/Source/Runtime/Core/Public/Misc/CoreMiscDefines.h b/Engine/Source/Runtime/Core/Public/Misc/CoreMiscDefines.h
--- a/Engine/Source/Runtime/Core/Public/Misc/CoreMiscDefines.h
+++ b/Engine/Source/Runtime/Core/Public/Misc/CoreMiscDefines.h
@@ -94,19 +94,3 @@ enum EForceInit
 	ForceInitToZero
 };
 enum ENoInit {NoInit};
-
-/**
- *	Whether to use experimental code to tick skeletal in parallel.
- *	Also forces shared pointers to be thread-safe.
- */
-#ifndef EXPERIMENTAL_PARALLEL_CODE
-	#define EXPERIMENTAL_PARALLEL_CODE 0
-#endif
-
-#if EXPERIMENTAL_PARALLEL_CODE
-	#define	FORCE_THREADSAFE_SHAREDPTRS 1
-#else
-	#define	FORCE_THREADSAFE_SHAREDPTRS 0
-#endif
-
-
diff --git a/Engine/Source/Runtime/Core/Public/Templates/SharedPointerInternals.h b/Engine/Source/Runtime/Core/Public/Templates/SharedPointerInternals.h
--- a/Engine/Source/Runtime/Core/Public/Templates/SharedPointerInternals.h
+++ b/Engine/Source/Runtime/Core/Public/Templates/SharedPointerInternals.h
@@ -6,9 +6,8 @@
 
 #pragma once
 
-#ifndef EXPERIMENTAL_PARALLEL_CODE
-	#error EXPERIMENTAL_PARALLEL_CODE not defined!
-#endif
+/** Default behavior. */
+#define	FORCE_THREADSAFE_SHAREDPTRS 0
 
 /**
  * ESPMode is used select between either 'fast' or 'thread safe' shared pointer types.  This is
@@ -23,7 +22,7 @@ namespace ESPMode
 
 		/**
 		 *	Fast, doesn't ever use atomic interlocks.
-		 *	Experimental parallel code requires that all shared pointers are thread-safe.
+		 *	Some code requires that all shared pointers are thread-safe.
 		 *	It's better to change it here, instead of replacing ESPMode::Fast to ESPMode::ThreadSafe throughout the code.
 		 */
 		Fast = FORCE_THREADSAFE_SHAREDPTRS ? 1 : 0,
diff --git a/Engine/Source/Runtime/Engine/Classes/Engine/EngineBaseTypes.h b/Engine/Source/Runtime/Engine/Classes/Engine/EngineBaseTypes.h
--- a/Engine/Source/Runtime/Engine/Classes/Engine/EngineBaseTypes.h
+++ b/Engine/Source/Runtime/Engine/Classes/Engine/EngineBaseTypes.h
@@ -45,17 +45,7 @@ enum ETickingGroup
 	TG_PrePhysics,
 	/**
 	 * Temp while we transition away from tick groups, this will only have one task, the one that starts physics
-	 */
-#if EXPERIMENTAL_PARALLEL_CODE
-	/**
-	 * Add a specific group for parallel animation work
-	 */
-	TG_ParallelAnimWork,
-	/**
-	 * Add a specific group for parallel post animation work
-	 */
-	TG_ParallelPostAnimWork,
-#endif								
+	 */							
 	TG_StartPhysics,
 	/**
 	 * Any item that can be run in parallel of our async work
diff --git a/Engine/Source/Runtime/Engine/Classes/Engine/World.h b/Engine/Source/Runtime/Engine/Classes/Engine/World.h
--- a/Engine/Source/Runtime/Engine/Classes/Engine/World.h
+++ b/Engine/Source/Runtime/Engine/Classes/Engine/World.h
@@ -53,34 +53,26 @@ public:
 	inline UWorld* operator->()
 	{
 		// GWorld is changed often on the game thread when in PIE, accessing on any other thread is going to be a race condition
-		// In general, the rendering thread should not dereference UObjects, unless there is a mechanism in place to make it safe
-#if !EXPERIMENTAL_PARALLEL_CODE		
-		checkSlow(IsInGameThread());
-#endif								
+		// In general, the rendering thread should not dereference UObjects, unless there is a mechanism in place to make it safe	
+		checkSlow(IsInGameThread());							
 		return World;
 	}
 
 	inline const UWorld* operator->() const
 	{
-#if !EXPERIMENTAL_PARALLEL_CODE
 		checkSlow(IsInGameThread());
-#endif
 		return World;
 	}
 
 	inline UWorld& operator*()
 	{
-#if !EXPERIMENTAL_PARALLEL_CODE
 		checkSlow(IsInGameThread());
-#endif
 		return *World;
 	}
 
 	inline const UWorld& operator*() const
 	{
-#if !EXPERIMENTAL_PARALLEL_CODE
 		checkSlow(IsInGameThread());
-#endif
 		return *World;
 	}
 
@@ -103,17 +95,13 @@ public:
 
 	inline operator UWorld*() const
 	{
-#if !EXPERIMENTAL_PARALLEL_CODE
 		checkSlow(IsInGameThread());
-#endif
 		return World;
 	}
 
 	inline UWorld* GetReference() 
 	{
-#if !EXPERIMENTAL_PARALLEL_CODE
 		checkSlow(IsInGameThread());
-#endif
 		return World;
 	}
 
@@ -656,11 +644,6 @@ private:
 	/** Set of components that need recreates at the end of the frame */
 	TSet<TWeakObjectPtr<class UActorComponent> > ComponentsThatNeedEndOfFrameUpdate_OnGameThread;
 
-#if EXPERIMENTAL_PARALLEL_CODE
-	/** Sync object for ComponentsThatNeedEndOfFrameUpdate */
-	FCriticalSection	ComponentsThatNeedEndOfFrameUpdateSynchronizationObject;
-#endif
-
 	/** The state of async tracing - abstracted into its own object for easier reference */
 	FWorldAsyncTraceState AsyncTraceState;
 
diff --git a/Engine/Source/Runtime/Engine/Private/ActorComponent.cpp b/Engine/Source/Runtime/Engine/Private/ActorComponent.cpp
--- a/Engine/Source/Runtime/Engine/Private/ActorComponent.cpp
+++ b/Engine/Source/Runtime/Engine/Private/ActorComponent.cpp
@@ -7,10 +7,6 @@
 #include "UObjectToken.h"
 #include "MapErrors.h"
 
-#ifndef EXPERIMENTAL_PARALLEL_CODE  
-	#error EXPERIMENTAL_PARALLEL_CODE must be defined as either zero or one
-#endif
-
 #define LOCTEXT_NAMESPACE "ActorComponent"
 
 DEFINE_LOG_CATEGORY(LogActorComponent);
@@ -969,18 +965,6 @@ void UActorComponent::MarkForNeededEndOfFrameRecreate()
 
 bool UActorComponent::RequiresGameThreadEndOfFrameUpdates() const
 {
-#if	EXPERIMENTAL_PARALLEL_CODE
-	AActor* Owner = GetOwner();
-	if (Owner != NULL)
-	{
-		if(Owner->GetRootComponent() == this)
-		{
-			// Root components end up touching all other components on the same actor when they update their render transforms.  This needs to be done on the game thread.
-			return true;
-		}
-	}
-#endif // EXPERIMENTAL_PARALLEL_CODE
-
 	return false;
 }
 
diff --git a/Engine/Source/Runtime/Engine/Private/LevelTick.cpp b/Engine/Source/Runtime/Engine/Private/LevelTick.cpp
--- a/Engine/Source/Runtime/Engine/Private/LevelTick.cpp
+++ b/Engine/Source/Runtime/Engine/Private/LevelTick.cpp
@@ -22,10 +22,6 @@
 #include "STaskGraph.h"
 #endif
 
-#ifndef EXPERIMENTAL_PARALLEL_CODE  
-	#error EXPERIMENTAL_PARALLEL_CODE must be defined as either zero or one
-#endif
-
 // this will log out all of the objects that were ticked in the FDetailedTickStats struct so you can isolate what is expensive
 #define LOG_DETAILED_DUMPSTATS 0
 
@@ -697,7 +693,7 @@ void UWorld::RunTickGroup(ETickingGroup Group, bool bBlockTillComplete = true)
 
 static TAutoConsoleVariable<int32> CVarAllowAsyncRenderThreadUpdates(
 	TEXT("AllowAsyncRenderThreadUpdates"),
-	EXPERIMENTAL_PARALLEL_CODE ? 1 : 0,
+	0,
 	TEXT("Used to control async renderthread updates."));
 
 void UWorld::MarkActorComponentForNeededEndOfFrameUpdate(class UActorComponent* Component, bool bForceGameThread)
@@ -709,10 +705,6 @@ void UWorld::MarkActorComponentForNeededEndOfFrameUpdate(class UActorComponent*
 		bForceGameThread = !bAllowConcurrentUpdates;
 	}
 
-#if EXPERIMENTAL_PARALLEL_CODE
-	FScopeLock ScopeLock( &ComponentsThatNeedEndOfFrameUpdateSynchronizationObject );
-#endif
-
 	if (bForceGameThread)
 	{
 		ComponentsThatNeedEndOfFrameUpdate_OnGameThread.Add(Component);
@@ -728,9 +720,6 @@ void UWorld::MarkActorComponentForNeededEndOfFrameUpdate(class UActorComponent*
 	*/
 void UWorld::SendAllEndOfFrameUpdates(FGraphEventArray* OutCompletion)
 {
-#if EXPERIMENTAL_PARALLEL_CODE
-	FScopeLock ScopeLock( &ComponentsThatNeedEndOfFrameUpdateSynchronizationObject );
-#endif
 	SCOPE_CYCLE_COUNTER(STAT_PostTickComponentUpdate);
 	// update all dirty components. 
 	bPostTickComponentUpdate = true;
@@ -1115,14 +1104,6 @@ void UWorld::Tick( ELevelTick TickType, float DeltaSeconds )
         bInTick = false;
         EnsureCollisionTreeIsBuilt();
         bInTick = true;
-#if EXPERIMENTAL_PARALLEL_CODE
-		{			
-			RunTickGroup(TG_ParallelAnimWork);			
-		}
-		{			
-			RunTickGroup(TG_ParallelPostAnimWork);			
-		}
-#endif // EXPERIMENTAL_PARALLEL_CODE
 		RunTickGroup(TG_StartPhysics); 
 		RunTickGroup(TG_DuringPhysics, false); // No wait here, we should run until idle though. We don't care if all of the async ticks are done before we start running post-phys stuff
 		TickGroup = TG_EndPhysics; // set this here so the current tick group is correct during collision notifies, though I am not sure it matters. 'cause of the false up there^^^
diff --git a/Engine/Source/Runtime/Engine/Private/SkeletalMeshComponent.cpp b/Engine/Source/Runtime/Engine/Private/SkeletalMeshComponent.cpp
--- a/Engine/Source/Runtime/Engine/Private/SkeletalMeshComponent.cpp
+++ b/Engine/Source/Runtime/Engine/Private/SkeletalMeshComponent.cpp
@@ -20,10 +20,6 @@
 #include "Particles/ParticleSystemComponent.h"
 #include "Animation/AnimSingleNodeInstance.h"
 
-#ifndef EXPERIMENTAL_PARALLEL_CODE  
-	#error EXPERIMENTAL_PARALLEL_CODE must be defined as either zero or one
-#endif
-
 TAutoConsoleVariable<int32> CVarUseParallelAnimationEvaluation(TEXT("a.ParallelAnimEvaluation"), 0, TEXT("If 1, animation evaluation will be run across the task graph system. If 0, evaluation will run purely on the game thread"));
 
 USkeletalMeshComponent::USkeletalMeshComponent(const class FPostConstructInitializeProperties& PCIP)
@@ -69,11 +65,6 @@ USkeletalMeshComponent::USkeletalMeshComponent(const class FPostConstructInitial
 	bEnablePhysicsOnDedicatedServer = false;
 	bEnableUpdateRateOptimizations = false;
 
-#if EXPERIMENTAL_PARALLEL_CODE
-	PrimaryComponentTick.TickGroup = TG_ParallelAnimWork;
-	PrimaryComponentTick.bRunOnAnyThread = true;	
-#endif
-
 	bTickInEditor = true;
 
 	ParallelEvaluationDelegate = FSimpleDelegateGraphTask::FDelegate::CreateUObject(this, &USkeletalMeshComponent::ParallelAnimationEvaluation);
@@ -223,12 +214,6 @@ void USkeletalMeshComponent::InitializeComponent()
 	Super::InitializeComponent();
 
 	InitAnim(false);
-
-#if EXPERIMENTAL_PARALLEL_CODE
-	PrimaryComponentTick.TickGroup = TG_ParallelAnimWork;
-	PrimaryComponentTick.bRunOnAnyThread = true;	
-#endif
-
 }
 
 #if CHART_DISTANCE_FACTORS
@@ -493,13 +478,6 @@ void USkeletalMeshComponent::TickPose(float DeltaTime)
 
 void USkeletalMeshComponent::TickComponent(float DeltaTime, enum ELevelTick TickType, FActorComponentTickFunction *ThisTickFunction)
 {
-#if EXPERIMENTAL_PARALLEL_CODE
-	if (ThisTickFunction->TickGroup != TG_ParallelAnimWork)
-	{			
-		UE_LOG(LogAnimation, Warning, TEXT("USkeletalMeshComponent %s is not in TickGroup TG_ParallelAnimWork! Group=%d, GoWide=%d"), *(GetOwner()->GetName()), (int32)ThisTickFunction->TickGroup, ThisTickFunction->bRunOnAnyThread);
-	}
-#endif
-
 	Super::TickComponent(DeltaTime, TickType, ThisTickFunction);
 
 	// Update bOldForceRefPose
@@ -973,35 +951,27 @@ void USkeletalMeshComponent::PostAnimEvaluation(bool bDoInterpolation, bool bDup
 		FAnimationRuntime::LerpBoneTransforms(SpaceBases, CachedSpaceBases, Alpha, RequiredBones);
 	}
 
-	{
-		FScopeLockPhysXWriter LockPhysXForWriting;
-
-		// Transforms updated, cached local bounds are now out of date.
-		InvalidateCachedBounds();
+	// Transforms updated, cached local bounds are now out of date.
+	InvalidateCachedBounds();
 
-		// update physics data from animated data
-		UpdateKinematicBonesToPhysics(false, true);
-		UpdateRBJointMotors();
-	}
+	// update physics data from animated data
+	UpdateKinematicBonesToPhysics(false, true);
+	UpdateRBJointMotors();
 
 	// @todo anim : hack TTP 224385	ANIM: Skeletalmesh double buffer
 	// this is problem because intermediate buffer changes physics position as well
 	// this causes issue where a half of frame, physics position is fixed with anim pose, and the other half is real simulated position
 	// if you enable physics in tick, since that's before physics update, you'll get animation pose dominating physics pose, which isn't what you want. (Or what you'll see)
 	// so do not update transform if physics is on. This problem will be solved by double buffer, when we keep one buffer for intermediate, and the other buffer for result query
+	if( !IsSimulatingPhysics() )
 	{
-		FScopeLockPhysXReader LockPhysXForReading;
+		SCOPE_CYCLE_COUNTER(STAT_UpdateLocalToWorldAndOverlaps);
 
-		if( !IsSimulatingPhysics() )
-		{
-			SCOPE_CYCLE_COUNTER(STAT_UpdateLocalToWorldAndOverlaps);
-
-			// New bone positions need to be sent to render thread
-			UpdateComponentToWorld();
+		// New bone positions need to be sent to render thread
+		UpdateComponentToWorld();
 
-			// animation often change overlap. 
-			UpdateOverlaps();
-		}
+		// animation often change overlap. 
+		UpdateOverlaps();
 	}
 
 	MarkRenderDynamicDataDirty();
diff --git a/Engine/Source/Runtime/Engine/Private/TickTaskManager.cpp b/Engine/Source/Runtime/Engine/Private/TickTaskManager.cpp
--- a/Engine/Source/Runtime/Engine/Private/TickTaskManager.cpp
+++ b/Engine/Source/Runtime/Engine/Private/TickTaskManager.cpp
@@ -7,10 +7,6 @@
 #include "EnginePrivate.h"
 #include "TickTaskManagerInterface.h"
 
-#ifndef EXPERIMENTAL_PARALLEL_CODE  
-	#error EXPERIMENTAL_PARALLEL_CODE must be defined as either zero or one
-#endif
-
 DEFINE_LOG_CATEGORY_STATIC(LogTick, Log, All);
 
 DECLARE_CYCLE_STAT(TEXT("Queue Ticks"),STAT_QueueTicks,STATGROUP_Game);
@@ -25,7 +21,7 @@ static TAutoConsoleVariable<int32> CVarLogTicks(
 
 static TAutoConsoleVariable<int32> CVarAllowAsyncComponentTicks(
 	TEXT("AllowAsyncComponentTicks"),
-	EXPERIMENTAL_PARALLEL_CODE ? 1  : 0,
+	0,
 	TEXT("Used to control async component ticks."));
 
 struct FTickContext
