Allow compiling specific material with remote 'recompileshaders material' command

diff --git a/Engine/Source/Editor/UnrealEd/Private/CookOnTheFlyServer.cpp b/Engine/Source/Editor/UnrealEd/Private/CookOnTheFlyServer.cpp
--- a/Engine/Source/Editor/UnrealEd/Private/CookOnTheFlyServer.cpp
+++ b/Engine/Source/Editor/UnrealEd/Private/CookOnTheFlyServer.cpp
@@ -989,5 +989,6 @@ void UCookOnTheFlyServer::HandleNetworkFileServerRecompileShaders(const FShaderR
 		RecompileData.MaterialsToLoad, 
 		RecompileData.SerializedShaderResources, 
 		RecompileData.MeshMaterialMaps, 
-		RecompileData.ModifiedFiles);
+		RecompileData.ModifiedFiles,
+		RecompileData.bCompileChangedShaders);
 }
diff --git a/Engine/Source/Runtime/Engine/Private/GlobalShader.cpp b/Engine/Source/Runtime/Engine/Private/GlobalShader.cpp
--- a/Engine/Source/Runtime/Engine/Private/GlobalShader.cpp
+++ b/Engine/Source/Runtime/Engine/Private/GlobalShader.cpp
@@ -499,7 +499,8 @@ void RecompileShadersForRemote(
 	const TArray<FString>& MaterialsToLoad, 
 	const TArray<uint8>& SerializedShaderResources, 
 	TArray<uint8>* MeshMaterialMaps, 
-	TArray<FString>* ModifiedFiles )
+	TArray<FString>* ModifiedFiles,
+	bool bCompileChangedShaders )
 {
 	// figure out what shader platforms to recompile
 	ITargetPlatformManagerModule* TPM = GetTargetPlatformManager();
@@ -532,9 +533,11 @@ void RecompileShadersForRemote(
 	// Pick up new changes to shader files
 	FlushShaderFileCache();
 
-	FShaderType::GetOutdatedTypes(OutdatedShaderTypes, OutdatedFactoryTypes);
-	
-	UE_LOG(LogShaders, Display, TEXT("We found %d out of date shader types, and %d out of date VF types!"), OutdatedShaderTypes.Num(), OutdatedFactoryTypes.Num());
+	if( bCompileChangedShaders )
+	{
+		FShaderType::GetOutdatedTypes( OutdatedShaderTypes, OutdatedFactoryTypes );
+		UE_LOG( LogShaders, Display, TEXT( "We found %d out of date shader types, and %d out of date VF types!" ), OutdatedShaderTypes.Num(), OutdatedFactoryTypes.Num() );
+	}
 
 	{
 		for (int32 FormatIndex = 0; FormatIndex < DesiredShaderFormats.Num(); FormatIndex++)
@@ -545,14 +548,17 @@ void RecompileShadersForRemote(
 			// Only compile for the desired platform if requested
 			if (ShaderPlatform == ShaderPlatformToCompile || ShaderPlatformToCompile == SP_NumPlatforms)
 			{
-				// Kick off global shader recompiles
-				BeginRecompileGlobalShaders(OutdatedShaderTypes, ShaderPlatform);
+				if( bCompileChangedShaders )
+				{
+					// Kick off global shader recompiles
+					BeginRecompileGlobalShaders( OutdatedShaderTypes, ShaderPlatform );
 
-				// Block on global shaders
-				FinishRecompileGlobalShaders();
+					// Block on global shaders
+					FinishRecompileGlobalShaders();
+				}
 
 				// we only want to actually compile mesh shaders if a client directly requested it, and there's actually some work to do
-				if (MeshMaterialMaps != NULL && (OutdatedShaderTypes.Num() || OutdatedFactoryTypes.Num()))
+				if (MeshMaterialMaps != NULL && (OutdatedShaderTypes.Num() || OutdatedFactoryTypes.Num() || bCompileChangedShaders == false))
 				{
 					TMap<FString, TArray<TRefCountPtr<FMaterialShaderMap> > > CompiledShaderMaps;
 					UMaterial::CompileMaterialsForRemoteRecompile(MaterialsToCompile, ShaderPlatform, CompiledShaderMaps);
diff --git a/Engine/Source/Runtime/Engine/Private/ShaderCompiler/ShaderCompiler.cpp b/Engine/Source/Runtime/Engine/Private/ShaderCompiler/ShaderCompiler.cpp
--- a/Engine/Source/Runtime/Engine/Private/ShaderCompiler/ShaderCompiler.cpp
+++ b/Engine/Source/Runtime/Engine/Private/ShaderCompiler/ShaderCompiler.cpp
@@ -2074,26 +2074,60 @@ protected:
 
 class FRecompileShaderMessageHandler : public IPlatformFile::IFileServerMessageHandler
 {
+public:
+	FRecompileShaderMessageHandler( const TCHAR* InCmd ) :
+		Cmd( InCmd )
+	{
+	}
+
 	/** Subclass fills out an archive to send to the server */
 	virtual void FillPayload(FArchive& Payload) override
 	{
-		// tell other side all the materials to load, by pathname
-		for (TObjectIterator<UMaterialInterface> It; It; ++It)
+		bool bCompileChangedShaders = true;
+
+		const TCHAR* CmdString = *Cmd;
+		FString CmdName = FParse::Token(CmdString, 0);
+
+		if( !CmdName.IsEmpty() && FCString::Stricmp(*CmdName,TEXT("Material"))==0 )
+		{
+			bCompileChangedShaders = false;
+
+			// tell other side the material to load, by pathname
+			FString RequestedMaterialName( FParse::Token( CmdString, 0 ) );
+
+			for( TObjectIterator<UMaterialInterface> It; It; ++It )
+			{
+				UMaterial* Material = It->GetMaterial();
+
+				if( Material && Material->GetName() == RequestedMaterialName)
+				{
+					MaterialsToLoad.Add( It->GetPathName() );
+					break;
+				}
+			}
+		}
+		else
 		{
-			MaterialsToLoad.Add(It->GetPathName());
+			// tell other side all the materials to load, by pathname
+			for( TObjectIterator<UMaterialInterface> It; It; ++It )
+			{
+				MaterialsToLoad.Add( It->GetPathName() );
+			}
 		}
+
 		Payload << MaterialsToLoad;
-		uint32 ShaderPlatform = (uint32)GRHIShaderPlatform;
+		uint32 ShaderPlatform = ( uint32 )GRHIShaderPlatform;
 		Payload << ShaderPlatform;
 		// tell the other side the Ids we have so it doesn't send back duplicates
 		// (need to serialize this into a TArray since FShaderResourceId isn't known in the file server)
 		TArray<FShaderResourceId> AllIds;
-		FShaderResource::GetAllShaderResourceId(AllIds);
+		FShaderResource::GetAllShaderResourceId( AllIds );
 
 		TArray<uint8> SerializedBytes;
-		FMemoryWriter Ar(SerializedBytes);
+		FMemoryWriter Ar( SerializedBytes );
 		Ar << AllIds;
 		Payload << SerializedBytes;
+		Payload << bCompileChangedShaders;
 	}
 
 	/** Subclass pulls data response from the server */
@@ -2137,8 +2171,12 @@ class FRecompileShaderMessageHandler : public IPlatformFile::IFileServerMessageH
 		}
 	}
 
+private:
 	/** The materials we send over the network and expect maps for on the return */
 	TArray<FString> MaterialsToLoad;
+
+	/** The recompileshader console command to parse */
+	FString Cmd;
 };
 
 bool RecompileShaders(const TCHAR* Cmd, FOutputDevice& Ar)
@@ -2146,7 +2184,7 @@ bool RecompileShaders(const TCHAR* Cmd, FOutputDevice& Ar)
 	// if this platform can't compile shaders, then we try to send a message to a file/cooker server
 	if (FPlatformProperties::RequiresCookedData())
 	{
-		FRecompileShaderMessageHandler Handler;
+		FRecompileShaderMessageHandler Handler( Cmd );
 
 		// send the info, the handler will process the response (and update shaders, etc)
 		IFileManager::Get().SendMessageToServer(TEXT("RecompileShaders"), &Handler);
diff --git a/Engine/Source/Runtime/Engine/Public/GlobalShader.h b/Engine/Source/Runtime/Engine/Public/GlobalShader.h
--- a/Engine/Source/Runtime/Engine/Public/GlobalShader.h
+++ b/Engine/Source/Runtime/Engine/Public/GlobalShader.h
@@ -206,6 +206,7 @@ extern void ProcessCompiledGlobalShaders(const TArray<FShaderCompileJob*>& Compi
  * @param SerializedShaderResources		Serialized shader resources
  * @param MeshMaterialMaps				Mesh material maps
  * @param ModifiedFiles					Returns the list of modified files if not NULL
+ * @param bCompileChangedShaders		Whether to compile all changed shaders or the specific material that is passed
  **/
 extern ENGINE_API void RecompileShadersForRemote( 
 	const FString& PlatformName, 
@@ -214,4 +215,5 @@ extern ENGINE_API void RecompileShadersForRemote(
 	const TArray<FString>& MaterialsToLoad, 
 	const TArray<uint8>& SerializedShaderResources, 
 	TArray<uint8>* MeshMaterialMaps, 
-	TArray<FString>* ModifiedFiles );
+	TArray<FString>* ModifiedFiles,
+	bool bCompileChangedShaders = true);
diff --git a/Engine/Source/Runtime/NetworkFileSystem/Private/NetworkFileServerConnection.cpp b/Engine/Source/Runtime/NetworkFileSystem/Private/NetworkFileServerConnection.cpp
--- a/Engine/Source/Runtime/NetworkFileSystem/Private/NetworkFileServerConnection.cpp
+++ b/Engine/Source/Runtime/NetworkFileSystem/Private/NetworkFileServerConnection.cpp
@@ -858,6 +858,7 @@ void FNetworkFileServerClientConnection::ProcessRecompileShaders( FArchive& In,
 	In << RecompileData.MaterialsToLoad;
 	In << RecompileData.ShaderPlatform;
 	In << RecompileData.SerializedShaderResources;
+	In << RecompileData.bCompileChangedShaders;
 
 	RecompileShadersDelegate.ExecuteIfBound(RecompileData);
 
diff --git a/Engine/Source/Runtime/NetworkFileSystem/Public/Interfaces/INetworkFileSystemModule.h b/Engine/Source/Runtime/NetworkFileSystem/Public/Interfaces/INetworkFileSystemModule.h
--- a/Engine/Source/Runtime/NetworkFileSystem/Public/Interfaces/INetworkFileSystemModule.h
+++ b/Engine/Source/Runtime/NetworkFileSystem/Public/Interfaces/INetworkFileSystemModule.h
@@ -20,9 +20,11 @@ struct FShaderRecompileData
 	TArray<uint8>* MeshMaterialMaps;
 	TArray<FString> MaterialsToLoad;
 	TArray<uint8> SerializedShaderResources;
+	bool bCompileChangedShaders;
 
 	FShaderRecompileData() :
-		ShaderPlatform(-1)
+		ShaderPlatform(-1),
+		bCompileChangedShaders(true)
 	{}
 
 	FShaderRecompileData& operator=(const FShaderRecompileData& Other)
@@ -33,6 +35,7 @@ struct FShaderRecompileData
 		MeshMaterialMaps = Other.MeshMaterialMaps;
 		MaterialsToLoad = Other.MaterialsToLoad;
 		SerializedShaderResources = Other.SerializedShaderResources;
+		bCompileChangedShaders = Other.bCompileChangedShaders;
 
 		return *this;
 	}
