Fix incorrect warning when geom is not found, but saying the scale is incorrect
TTP#336704

diff --git a/Engine/Source/Runtime/Engine/Private/PhysicsEngine/BodyInstance.cpp b/Engine/Source/Runtime/Engine/Private/PhysicsEngine/BodyInstance.cpp
--- a/Engine/Source/Runtime/Engine/Private/PhysicsEngine/BodyInstance.cpp
+++ b/Engine/Source/Runtime/Engine/Private/PhysicsEngine/BodyInstance.cpp
@@ -1552,6 +1552,7 @@ bool FBodyInstance::UpdateBodyScale(const FVector& InScale3D)
 
 	for(int32 ShapeIdx=0; ShapeIdx<PShapes.Num(); ShapeIdx++)
 	{
+		bool bInvalid = false;	//we only mark invalid if actually found geom and it's invalid scale
 		PxGeometry* UpdatedGeometry = NULL;
 		PxShape* PShape = PShapes[ShapeIdx];
 		PxScene* PScene = PShape->getActor()->getScene();
@@ -1578,6 +1579,10 @@ bool FBodyInstance::UpdateBodyScale(const FVector& InScale3D)
 					UpdatedGeometry = &PSphereGeom;
 					bSuccess = true;
 				}
+				else
+				{
+					bInvalid = true;
+				}
 				break;
 			}
 			case PxGeometryType::eBOX:
@@ -1597,6 +1602,10 @@ bool FBodyInstance::UpdateBodyScale(const FVector& InScale3D)
 					UpdatedGeometry = &PBoxGeom;
 					bSuccess = true;
 				}
+				else
+				{
+					bInvalid = true;
+				}
 				break;
 			}
 			case PxGeometryType::eCAPSULE:
@@ -1618,6 +1627,10 @@ bool FBodyInstance::UpdateBodyScale(const FVector& InScale3D)
 					UpdatedGeometry = &PCapsuleGeom;
 					bSuccess = true;
 				}
+				else
+				{
+					bInvalid = true;
+				}
 
 				break;
 			}
@@ -1657,6 +1670,10 @@ bool FBodyInstance::UpdateBodyScale(const FVector& InScale3D)
 								UpdatedGeometry = &PConvexGeom;
 								bSuccess = true;
 							}
+							else
+							{
+								bInvalid = true;
+							}
 							break;
 						}
 					}
@@ -1697,6 +1714,10 @@ bool FBodyInstance::UpdateBodyScale(const FVector& InScale3D)
 							bSuccess = true;
 
 						}
+						else
+						{
+							bInvalid = true;
+						}
 					}
 				}
 				break;
@@ -1714,7 +1735,7 @@ bool FBodyInstance::UpdateBodyScale(const FVector& InScale3D)
 			PShape->setLocalPose(PLocalPose);
 			PShape->setGeometry(*UpdatedGeometry);
 		}
-		else
+		else if (bInvalid)
 		{
 			FMessageLog("PIE").Warning(FText::Format(LOCTEXT("PhysicsInvalidScale", "Scale ''{0}'' is not valid on object '{1}'."), FText::FromString(InScale3DAdjusted.ToString()), FText::FromString(GetBodyDebugName())));
 		}
