Improvement to Navigation Link debug drawing's bounds calculation #UE4

diff --git a/Engine/Source/Runtime/Engine/Private/AI/Navigation/NavLinkRenderingComponent.cpp b/Engine/Source/Runtime/Engine/Private/AI/Navigation/NavLinkRenderingComponent.cpp
--- a/Engine/Source/Runtime/Engine/Private/AI/Navigation/NavLinkRenderingComponent.cpp
+++ b/Engine/Source/Runtime/Engine/Private/AI/Navigation/NavLinkRenderingComponent.cpp
@@ -25,7 +25,45 @@ UNavLinkRenderingComponent::UNavLinkRenderingComponent(const class FPostConstruc
 
 FBoxSphereBounds UNavLinkRenderingComponent::CalcBounds(const FTransform & LocalToWorld) const
 {
-	return FBoxSphereBounds( FVector::ZeroVector, FVector(512), 512 ).TransformBy(LocalToWorld);
+	AActor* LinkOwnerActor = Cast<AActor>(GetOwner());
+	INavLinkHostInterface* LinkOwnerHost = InterfaceCast<INavLinkHostInterface>(GetOwner());
+
+	if (LinkOwnerActor != NULL && LinkOwnerHost != NULL)
+	{
+		FBox BoundingBox(0);
+		const FTransform LocalToWorld = LinkOwnerActor->ActorToWorld();
+		TArray<TSubclassOf<UNavLinkDefinition> > NavLinkClasses;
+		TArray<FNavigationLink> SimpleLinks;
+		TArray<FNavigationSegmentLink> DummySegmentLinks;
+
+		if (LinkOwnerHost->GetNavigationLinksClasses(NavLinkClasses))
+		{
+			for (int32 NavLinkClassIdx = 0; NavLinkClassIdx < NavLinkClasses.Num(); ++NavLinkClassIdx)
+			{
+				if (NavLinkClasses[NavLinkClassIdx] != NULL)
+				{
+					const TArray<FNavigationLink>& Links = UNavLinkDefinition::GetLinksDefinition(NavLinkClasses[NavLinkClassIdx]);
+					for (const auto& Link : Links)
+					{
+						BoundingBox += Link.Left;
+						BoundingBox += Link.Right;
+					}
+				}
+			}
+		}
+		if (LinkOwnerHost->GetNavigationLinksArray(SimpleLinks, DummySegmentLinks))
+		{
+			for (const auto& Link : SimpleLinks)
+			{
+				BoundingBox += Link.Left;
+				BoundingBox += Link.Right;
+			}
+		}
+
+		return FBoxSphereBounds(BoundingBox).TransformBy(LocalToWorld);
+	}
+
+	return FBoxSphereBounds();
 }
 
 FPrimitiveSceneProxy* UNavLinkRenderingComponent::CreateSceneProxy()
