Altered FAnimNotifyEvent equality operator to correctly consider non-blueprint notifies.
#ttp 339248 - CRITICAL: Regression: PERSONA: Unable to delete new notify

diff --git a/Engine/Source/Runtime/Engine/Classes/Animation/AnimSequenceBase.h b/Engine/Source/Runtime/Engine/Classes/Animation/AnimSequenceBase.h
--- a/Engine/Source/Runtime/Engine/Classes/Animation/AnimSequenceBase.h
+++ b/Engine/Source/Runtime/Engine/Classes/Animation/AnimSequenceBase.h
@@ -109,7 +109,11 @@ struct FAnimNotifyEvent
 
 	bool operator ==(const FAnimNotifyEvent& Other)
 	{
-		return((Notify && Notify == Other.Notify) || (NotifyStateClass && NotifyStateClass == Other.NotifyStateClass));
+		return(
+			(Notify && Notify == Other.Notify) || 
+			(NotifyStateClass && NotifyStateClass == Other.NotifyStateClass) ||
+			(!IsBlueprintNotify() && NotifyName == Other.NotifyName)
+			);
 	}
 };
 
