Fix hit event not firing when hitting a BSP
TTP#339182

diff --git a/Engine/Source/Runtime/Engine/Private/PhysicsEngine/PhysScene.cpp b/Engine/Source/Runtime/Engine/Private/PhysicsEngine/PhysScene.cpp
--- a/Engine/Source/Runtime/Engine/Private/PhysicsEngine/PhysScene.cpp
+++ b/Engine/Source/Runtime/Engine/Private/PhysicsEngine/PhysScene.cpp
@@ -712,13 +712,13 @@ void FPhysScene::DispatchPhysNotifications()
 			FCollisionNotifyInfo& NotifyInfo = PendingCollisionNotifies[i];
 			if (NotifyInfo.RigidCollisionData.ContactInfos.Num() > 0)
 			{
-				if (NotifyInfo.bCallEvent0 && NotifyInfo.IsValidForNotify())
+				if (NotifyInfo.bCallEvent0 && NotifyInfo.IsValidForNotify() && NotifyInfo.Info0.Actor.IsValid())
 				{
 					NotifyInfo.Info0.Actor->DispatchPhysicsCollisionHit(NotifyInfo.Info0, NotifyInfo.Info1, NotifyInfo.RigidCollisionData);
 				}
 
 				// Need to check IsValidForNotify again in case first call broke something.
-				if (NotifyInfo.bCallEvent1 && NotifyInfo.IsValidForNotify())
+				if (NotifyInfo.bCallEvent1 && NotifyInfo.IsValidForNotify() && NotifyInfo.Info1.Actor.IsValid())
 				{
 					NotifyInfo.RigidCollisionData.SwapContactOrders();
 					NotifyInfo.Info1.Actor->DispatchPhysicsCollisionHit(NotifyInfo.Info1, NotifyInfo.Info0, NotifyInfo.RigidCollisionData);
diff --git a/Engine/Source/Runtime/Engine/Private/PhysicsEngine/PhysUtils.cpp b/Engine/Source/Runtime/Engine/Private/PhysicsEngine/PhysUtils.cpp
--- a/Engine/Source/Runtime/Engine/Private/PhysicsEngine/PhysUtils.cpp
+++ b/Engine/Source/Runtime/Engine/Private/PhysicsEngine/PhysUtils.cpp
@@ -176,7 +176,7 @@ FBodyInstance* FRigidBodyCollisionInfo::GetBodyInstance() const
 
 bool FCollisionNotifyInfo::IsValidForNotify() const
 {
-	return (Info0.Actor.IsValid() && Info1.Actor.IsValid());
+	return (Info0.Component.IsValid() && Info1.Component.IsValid());
 }
 
 /** Iterate over ContactInfos array and swap order of information */
