Pass newly created player controller to FindPlayerStart for GameMode Login

diff --git a/Engine/Source/Runtime/Engine/Private/GameMode.cpp b/Engine/Source/Runtime/Engine/Private/GameMode.cpp
--- a/Engine/Source/Runtime/Engine/Private/GameMode.cpp
+++ b/Engine/Source/Runtime/Engine/Private/GameMode.cpp
@@ -987,7 +987,7 @@ APlayerController* AGameMode::Login(UPlayer* NewPlayer, const FString& Portal, c
 	InitNewPlayer(NewPlayerController, UniqueId, Options);
 
 	// Find a start spot.
-	AActor* const StartSpot = FindPlayerStart( NULL, Portal );
+	AActor* const StartSpot = FindPlayerStart( NewPlayerController, Portal );
 	if( StartSpot == NULL )
 	{
 		ErrorMessage = FString::Printf(TEXT("Failed to find PlayerStart"));
