Fix bug where setKinematicTarget was using the globalPosition for early out.

diff --git a/Engine/Source/Runtime/Engine/Private/PhysicsEngine/BodyInstance.cpp b/Engine/Source/Runtime/Engine/Private/PhysicsEngine/BodyInstance.cpp
--- a/Engine/Source/Runtime/Engine/Private/PhysicsEngine/BodyInstance.cpp
+++ b/Engine/Source/Runtime/Engine/Private/PhysicsEngine/BodyInstance.cpp
@@ -1960,10 +1960,26 @@ void FBodyInstance::SetBodyTransform(const FTransform& NewTransform, bool bTelep
 		// Do nothing if already in correct place
 		{
 			SCOPED_SCENE_READ_LOCK(RigidActor->getScene());
-			const PxTransform PCurrentPose = RigidActor->getGlobalPose();
-			if (NewTransform.Equals(P2UTransform(PCurrentPose)))
+			if (PRigidDynamic && !IsRigidDynamicNonKinematic(PRigidDynamic) && !bTeleport)
 			{
-				return;
+				const PxScene* PScene = PRigidDynamic->getScene();
+				FPhysScene* PhysScene = FPhysxUserData::Get<FPhysScene>(PScene->userData);
+				FTransform CurrentTarget;
+				if (PhysScene->GetKinematicTarget(this, CurrentTarget))
+				{
+					if (CurrentTarget.Equals(NewTransform))
+					{
+						return;
+					}
+				}
+			}
+			else
+			{
+				const PxTransform PCurrentPose = RigidActor->getGlobalPose();
+				if (NewTransform.Equals(P2UTransform(PCurrentPose)))
+				{
+					return;
+				}
 			}
 		}
 
diff --git a/Engine/Source/Runtime/Engine/Private/PhysicsEngine/PhysScene.cpp b/Engine/Source/Runtime/Engine/Private/PhysicsEngine/PhysScene.cpp
--- a/Engine/Source/Runtime/Engine/Private/PhysicsEngine/PhysScene.cpp
+++ b/Engine/Source/Runtime/Engine/Private/PhysicsEngine/PhysScene.cpp
@@ -161,6 +161,36 @@ bool UseSyncTime(uint32 SceneType)
 
 }
 
+bool FPhysScene::GetKinematicTarget(const FBodyInstance * BodyInstance, FTransform & OutTM) const
+{
+#if WITH_PHYSX
+	if (PxRigidDynamic * PRigidDynamic = BodyInstance->GetPxRigidDynamic())
+	{
+#if WITH_SUBSTEPPING
+		if (IsSubstepping())
+		{
+			FPhysSubstepTask * PhysSubStepper = PhysSubSteppers[SceneType(BodyInstance)];
+			return PhysSubStepper->GetKinematicTarget(BodyInstance, OutTM);
+		}
+		else
+#endif
+		{
+
+			SCOPED_SCENE_READ_LOCK(PRigidDynamic->getScene());
+			PxTransform POutTM;
+			bool validTM = PRigidDynamic->getKinematicTarget(POutTM);
+			if (validTM)
+			{
+				OutTM = P2UTransform(POutTM);
+				return true;
+			}
+		}
+	}
+#endif
+
+	return false;
+}
+
 void FPhysScene::SetKinematicTarget(FBodyInstance * BodyInstance, const FTransform & TargetTransform, bool bAllowSubstepping)
 {
 	TargetTransform.DiagnosticCheckNaN_All();
diff --git a/Engine/Source/Runtime/Engine/Private/PhysicsEngine/PhysSubstepTasks.cpp b/Engine/Source/Runtime/Engine/Private/PhysicsEngine/PhysSubstepTasks.cpp
--- a/Engine/Source/Runtime/Engine/Private/PhysicsEngine/PhysSubstepTasks.cpp
+++ b/Engine/Source/Runtime/Engine/Private/PhysicsEngine/PhysSubstepTasks.cpp
@@ -66,6 +66,24 @@ void FPhysSubstepTask::SetKinematicTarget(FBodyInstance * Body, const FTransform
 #endif
 }
 
+bool FPhysSubstepTask::GetKinematicTarget(const FBodyInstance * Body, FTransform & OutTM) const
+{
+#if WITH_PHYSX
+	const PxRigidDynamic * PRigidDynamic = Body->GetPxRigidDynamic();
+	SCOPED_SCENE_READ_LOCK(PRigidDynamic->getScene());
+	if (const FPhysTarget * TargetState = PhysTargetBuffers[External].Find(Body))
+	{
+		if (TargetState->bKinematicTarget)
+		{
+			OutTM = TargetState->KinematicTarget.TargetTM;
+			return true;
+		}
+	}
+#endif
+
+	return false;
+}
+
 void FPhysSubstepTask::AddForce(FBodyInstance * Body, const FVector & Force)
 {
 #if WITH_PHYSX
diff --git a/Engine/Source/Runtime/Engine/Private/PhysicsEngine/PhysSubstepTasks.h b/Engine/Source/Runtime/Engine/Private/PhysicsEngine/PhysSubstepTasks.h
--- a/Engine/Source/Runtime/Engine/Private/PhysicsEngine/PhysSubstepTasks.h
+++ b/Engine/Source/Runtime/Engine/Private/PhysicsEngine/PhysSubstepTasks.h
@@ -105,6 +105,7 @@ public:
 #endif
 
 	void SetKinematicTarget(FBodyInstance * Body, const FTransform & TM);
+	bool GetKinematicTarget(const FBodyInstance * Body, FTransform & OutTM) const;
 	void AddForce(FBodyInstance * Body, const FVector & Force);
 	void AddForceAtPosition(FBodyInstance * Body, const FVector & Force, const FVector & Position);
 	void AddTorque(FBodyInstance * Body, const FVector & Torque);
diff --git a/Engine/Source/Runtime/Engine/Public/PhysicsPublic.h b/Engine/Source/Runtime/Engine/Public/PhysicsPublic.h
--- a/Engine/Source/Runtime/Engine/Public/PhysicsPublic.h
+++ b/Engine/Source/Runtime/Engine/Public/PhysicsPublic.h
@@ -347,7 +347,7 @@ public:
 	static bool SupportsOriginShifting() { return true; }
 
 	/** @return Whether physics scene is using substepping */
-	bool IsSubstepping()
+	bool IsSubstepping() const
 	{
 #if WITH_SUBSTEPPING
 		return bSubstepping;
@@ -377,6 +377,10 @@ public:
 	/** Sets a Kinematic actor's target position - We need to do this here to support substepping*/
 	void SetKinematicTarget(FBodyInstance * BodyInstance, const FTransform & TargetTM, bool bAllowSubstepping);
 
+	/** Gets a Kinematic actor's target position - We need to do this here to support substepping
+	  * Returns true if kinematic target has been set. If false the OutTM is invalid */
+	bool GetKinematicTarget(const FBodyInstance * BodyInstance, FTransform & OutTM) const;
+
 	/** Gets the collision disable table */
 	const TMap<uint32, TMap<struct FRigidBodyIndexPair, bool> *> & GetCollisionDisableTableLookup()
 	{
