Updated FSceneViewport::InputGesture so it sets keyboard focus to viewport when user is manipulating it using trackpad

diff --git a/Engine/Source/Runtime/Engine/Private/Slate/SceneViewport.cpp b/Engine/Source/Runtime/Engine/Private/Slate/SceneViewport.cpp
--- a/Engine/Source/Runtime/Engine/Private/Slate/SceneViewport.cpp
+++ b/Engine/Source/Runtime/Engine/Private/Slate/SceneViewport.cpp
@@ -630,6 +630,8 @@ FReply FSceneViewport::OnTouchGesture( const FGeometry& MyGeometry, const FPoint
 		// Switch to the viewport clients world before processing input
 		FScopedConditionalWorldSwitcher WorldSwitcher( ViewportClient );
 
+		FSlateApplication::Get().SetKeyboardFocus(ViewportWidget.Pin());
+
 		if( !ViewportClient->InputGesture( this, GestureEvent.GetGestureType(), GestureEvent.GetGestureDelta() ) )
 		{
 			CurrentReplyState = FReply::Unhandled();
