Use WindowedFullscreen mode on OS X when enabling fullscreen via Alt+Enter as that scales the window to the monitor size which is all we currently do anyway. 'Proper' fullscreen mode requires an explicit CGL back-buffer size and mouse<->display coordinate conversions (same as Retina mode support) that we don't have code for yet.

diff --git a/Engine/Source/Runtime/Engine/Private/GameViewportClient.cpp b/Engine/Source/Runtime/Engine/Private/GameViewportClient.cpp
--- a/Engine/Source/Runtime/Engine/Private/GameViewportClient.cpp
+++ b/Engine/Source/Runtime/Engine/Private/GameViewportClient.cpp
@@ -2475,7 +2475,12 @@ bool UGameViewportClient::SetDisplayConfiguration(const FIntPoint* Dimensions, E
 
 bool UGameViewportClient::HandleToggleFullscreenCommand(const TCHAR* Cmd, FOutputDevice& Ar)
 {
+#if PLATFORM_MAC // @TODO: Fullscreen mode isn't working fully on OS X as it requires explicit CGL back-buffer size and mouse<->display coordinate conversions.
+				 // For now use the windowed fullscreen mode which merely stretches to fill the display, which is what OS X does for us by default.
+	return SetDisplayConfiguration(NULL, Viewport->IsFullscreen() ? EWindowMode::Windowed : EWindowMode::WindowedFullscreen);
+#else
 	return SetDisplayConfiguration(NULL, Viewport->IsFullscreen() ? EWindowMode::Windowed : EWindowMode::Fullscreen);
+#endif
 }
 
 bool UGameViewportClient::HandleSetResCommand( const TCHAR* Cmd, FOutputDevice& Ar )
