Fixed TTP #333943 : CRASH: EDITOR: If a piece of a destructible mesh leaves a level streaming volume, the editor will crash when the player character also leaves the volume.

The reason was accessing to a already unregistered component and transform data were all deallocated.

diff --git a/Engine/Source/Runtime/Engine/Private/PhysicsEngine/PhysScene.cpp b/Engine/Source/Runtime/Engine/Private/PhysicsEngine/PhysScene.cpp
--- a/Engine/Source/Runtime/Engine/Private/PhysicsEngine/PhysScene.cpp
+++ b/Engine/Source/Runtime/Engine/Private/PhysicsEngine/PhysScene.cpp
@@ -684,7 +684,12 @@ void FPhysScene::SyncComponentsToBodies(uint32 SceneType)
 					const physx::PxTransform Transform(ChunkPoseRT);
 					if (UDestructibleComponent * DestructibleComponent = Cast<UDestructibleComponent>(FPhysxUserData::Get<UPrimitiveComponent>(DestructibleActor->userData)))
 					{
-						DestructibleComponent->SetChunkWorldRT(ChunkIndex, P2UQuat(Transform.q), P2UVector(Transform.p));
+						// if this component was already unregistered, transform data were deallocated
+						// transform data become empty when a piece of destructible actor goes out from a streaming volume
+						if (DestructibleComponent->IsRegistered())
+						{
+							DestructibleComponent->SetChunkWorldRT(ChunkIndex, P2UQuat(Transform.q), P2UVector(Transform.p));
+						}
 					}
 				}
 			}
