Minor refactor to GameUserSettings, making sure that cmd line overrides don't get stored to disc.

Fixes TTP 334080 - UE4: RENDERING: FRAMEWORK: Game will launch in fullscreen every time if the user launches with the -fullscreen command once

diff --git a/Engine/Source/Runtime/Engine/Classes/GameFramework/GameUserSettings.h b/Engine/Source/Runtime/Engine/Classes/GameFramework/GameUserSettings.h
--- a/Engine/Source/Runtime/Engine/Classes/GameFramework/GameUserSettings.h
+++ b/Engine/Source/Runtime/Engine/Classes/GameFramework/GameUserSettings.h
@@ -92,7 +92,8 @@ public:
 	/** Loads the user .ini settings into GConfig */
 	static void LoadConfigIni( bool bForceReload=false );
 
-
+	/** Apply any runtime overrides to the settings */
+	static void RequestResolutionChange(int32 InResolutionX, int32 InResolutionY, EWindowMode::Type InWindowMode);
 
 	/** Whether to use VSync or not. (public to allow UI to connect to it) */
 	UPROPERTY(config)
diff --git a/Engine/Source/Runtime/Engine/Private/GameUserSettings.cpp b/Engine/Source/Runtime/Engine/Private/GameUserSettings.cpp
--- a/Engine/Source/Runtime/Engine/Private/GameUserSettings.cpp
+++ b/Engine/Source/Runtime/Engine/Private/GameUserSettings.cpp
@@ -195,7 +195,7 @@ void UGameUserSettings::ApplySettings()
 	}
 
 	// Request a resolution change
-	FSystemResolution::RequestResolutionChange(ResolutionSizeX, ResolutionSizeY, NewWindowMode);
+	RequestResolutionChange(ResolutionSizeX, ResolutionSizeY, NewWindowMode);
 
 	// Update vsync cvar
 	{
@@ -229,25 +229,22 @@ void UGameUserSettings::LoadSettings( bool bForceReload/*=false*/ )
 	ScalabilityQuality = Scalability::GetQualityLevels();
 
 	// Allow override using command-line settings
-	const EWindowMode::Type FullscreenWindowMode = GetFullscreenMode();
-	EWindowMode::Type OverrideFullscreenWindowMode = FullscreenWindowMode;
-	int32 OverrideResolutionSizeX = ResolutionSizeX;
-	int32 OverrideResolutionSizeY = ResolutionSizeY;
 	bool bDetectingResolution = ResolutionSizeX == 0 || ResolutionSizeY == 0;
 
-	UGameEngine::ConditionallyOverrideSettings(OverrideResolutionSizeX, OverrideResolutionSizeY, OverrideFullscreenWindowMode);
-
-	ResolutionSizeX = OverrideResolutionSizeX;
-	ResolutionSizeY = OverrideResolutionSizeY;
-	if (OverrideFullscreenWindowMode != FullscreenWindowMode)
+	if (bDetectingResolution)
 	{
-		SetFullscreenMode(OverrideFullscreenWindowMode);
+		ConfirmVideoMode();
 	}
+}
 
-	if (bDetectingResolution)
+void UGameUserSettings::RequestResolutionChange(int32 InResolutionX, int32 InResolutionY, EWindowMode::Type InWindowMode)
+{
+	UGameEngine::ConditionallyOverrideSettings(InResolutionX, InResolutionY, InWindowMode);
 	{
-		ConfirmVideoMode();
+		IConsoleVariable* CVar = IConsoleManager::Get().FindConsoleVariable(TEXT("r.FullScreenMode"));
+		CVar->Set(InWindowMode);
 	}
+	FSystemResolution::RequestResolutionChange(InResolutionX, InResolutionY, InWindowMode);
 }
 
 void UGameUserSettings::SaveSettings()
@@ -307,8 +304,7 @@ void UGameUserSettings::PreloadResolutionSettings()
 		CVar->Set((int32)WindowMode);
 	}
 
-	UGameEngine::ConditionallyOverrideSettings(ResolutionX, ResolutionY, WindowMode);
-	FSystemResolution::RequestResolutionChange(ResolutionX, ResolutionY, WindowMode);
+	RequestResolutionChange(ResolutionX, ResolutionY, WindowMode);
 
 	IConsoleManager::Get().CallAllConsoleVariableSinks();
 }
