Made foliage not receive decals by default, added checkbox to re-enable the user wants.

diff --git a/Engine/Source/Editor/FoliageEdit/Private/SFoliageEditMeshDisplayItem.cpp b/Engine/Source/Editor/FoliageEdit/Private/SFoliageEditMeshDisplayItem.cpp
--- a/Engine/Source/Editor/FoliageEdit/Private/SFoliageEditMeshDisplayItem.cpp
+++ b/Engine/Source/Editor/FoliageEdit/Private/SFoliageEditMeshDisplayItem.cpp
@@ -1854,7 +1854,7 @@ FString SFoliageEditMeshDisplayItem::GetSaveRemoveSettingsTooltip() const
 bool SFoliageEditMeshDisplayItem::IsPropertyVisible(UProperty const * const InProperty) const
 {
 	const FString Category = FObjectEditorUtils::GetCategory(InProperty);
-	return Category == TEXT("Clustering") || Category == TEXT("Culling") || Category == TEXT("Lighting") || Category == TEXT("Collision");
+	return Category == TEXT("Clustering") || Category == TEXT("Culling") || Category == TEXT("Lighting") || Category == TEXT("Collision") || Category == TEXT("Rendering");
 }
 
 FString SFoliageEditMeshDisplayItem::GetInstanceCountString() const
diff --git a/Engine/Source/Runtime/Engine/Classes/Foliage/InstancedFoliageSettings.h b/Engine/Source/Runtime/Engine/Classes/Foliage/InstancedFoliageSettings.h
--- a/Engine/Source/Runtime/Engine/Classes/Foliage/InstancedFoliageSettings.h
+++ b/Engine/Source/Runtime/Engine/Classes/Foliage/InstancedFoliageSettings.h
@@ -180,34 +180,39 @@ class UInstancedFoliageSettings : public UObject
 	UPROPERTY()
 	uint32 ShowInstanceSettings:1;
 
-	/** Controls whether the primitive component should cast a shadow or not. **/
+	/** Controls whether the foliage should cast a shadow or not. **/
 	UPROPERTY(EditAnywhere, BlueprintReadOnly, Category=Lighting)
 	uint32 CastShadow:1;
 
-	/** Controls whether the primitive should inject light into the Light Propagation Volume.  This flag is only used if CastShadow is true. **/
+	/** Controls whether the foliage should inject light into the Light Propagation Volume.  This flag is only used if CastShadow is true. **/
 	UPROPERTY(EditAnywhere, BlueprintReadOnly, Category=Lighting)
 	uint32 bAffectDynamicIndirectLighting:1;
 
-	/** Controls whether the primitive should cast shadows in the case of non precomputed shadowing.  This flag is only used if CastShadow is true. **/
+	/** Controls whether the foliage should cast shadows in the case of non precomputed shadowing.  This flag is only used if CastShadow is true. **/
 	UPROPERTY(EditAnywhere, BlueprintReadOnly, Category=Lighting)
 	uint32 bCastDynamicShadow:1;
-	/** Whether the object should cast a static shadow from shadow casting lights.  This flag is only used if CastShadow is true. */
+
+	/** Whether the foliage should cast a static shadow from shadow casting lights.  This flag is only used if CastShadow is true. */
 	UPROPERTY(EditAnywhere, BlueprintReadOnly, Category=Lighting)
 	uint32 bCastStaticShadow:1;
 
 	/** 
-	 *	If true, the primitive will cast shadows even if bHidden is true.
+	 *	If true, the foliage will cast shadows even if bHidden is true.
 	 *	Controls whether the primitive should cast shadows when hidden.
 	 *	This flag is only used if CastShadow is true.
 	 */
 	UPROPERTY(EditAnywhere, AdvancedDisplay, BlueprintReadOnly, Category=Lighting)
 	uint32 bCastHiddenShadow:1;
 
-	/** Whether this primitive should cast dynamic shadows as if it were a two sided material. */
+	/** Whether this foliage should cast dynamic shadows as if it were a two sided material. */
 	UPROPERTY(EditAnywhere, AdvancedDisplay, BlueprintReadOnly, Category=Lighting)
 	uint32 bCastShadowAsTwoSided:1;
 
-	/** Custom Collision for foliage */
+	/** Whether the foliage receives decals. */
+	UPROPERTY(EditAnywhere, AdvancedDisplay, BlueprintReadOnly, Category=Rendering)
+	uint32 bReceivesDecals : 1;
+
+	/** Custom collision for foliage */
 	UPROPERTY(EditAnywhere, Category=Collision, meta=(HideObjectType=true))
 	struct FBodyInstance BodyInstance;
 };
diff --git a/Engine/Source/Runtime/Engine/Private/InstancedFoliage.cpp b/Engine/Source/Runtime/Engine/Private/InstancedFoliage.cpp
--- a/Engine/Source/Runtime/Engine/Private/InstancedFoliage.cpp
+++ b/Engine/Source/Runtime/Engine/Private/InstancedFoliage.cpp
@@ -132,6 +132,7 @@ UInstancedFoliageSettings::UInstancedFoliageSettings(const class FPostConstructI
 	bAffectDynamicIndirectLighting = false;
 	bCastHiddenShadow = false;
 	bCastShadowAsTwoSided = false;
+	bReceivesDecals = false;
 
 	BodyInstance.SetCollisionProfileName(UCollisionProfile::NoCollision_ProfileName);
 }
@@ -311,6 +312,7 @@ void FFoliageMeshInfo::AddInstance(AInstancedFoliageActor* InIFA, UStaticMesh* I
 		BestCluster->ClusterComponent->bAffectDynamicIndirectLighting = Settings->bAffectDynamicIndirectLighting;
 		BestCluster->ClusterComponent->bCastHiddenShadow = Settings->bCastHiddenShadow;
 		BestCluster->ClusterComponent->bCastShadowAsTwoSided = Settings->bCastShadowAsTwoSided;
+		BestCluster->ClusterComponent->bReceivesDecals = Settings->bReceivesDecals;
 
 		BestCluster->ClusterComponent->BodyInstance.CopyBodyInstancePropertiesFrom(&Settings->BodyInstance);
 
