Fix array access, possible cause of rare crash

diff --git a/Engine/Source/Runtime/Engine/Private/SkinnedMeshComponent.cpp b/Engine/Source/Runtime/Engine/Private/SkinnedMeshComponent.cpp
--- a/Engine/Source/Runtime/Engine/Private/SkinnedMeshComponent.cpp
+++ b/Engine/Source/Runtime/Engine/Private/SkinnedMeshComponent.cpp
@@ -1709,7 +1709,7 @@ TArray<FActiveVertexAnim> USkinnedMeshComponent::UpdateActiveVertexAnims(const T
 				// If it does, use the max weight
 				else
 				{
-					const float CurrentWeight = ActiveVertexAnims[AnimIndex].Weight;
+					const float CurrentWeight = OutVertexAnims[AnimIndex].Weight;
 					OutVertexAnims[AnimIndex].Weight = FMath::Max<float>(CurrentWeight, Weight);
 				}
 			}
