Restore previous fallback behavior when a player start cannot be found

diff --git a/Engine/Source/Runtime/Engine/Private/GameMode.cpp b/Engine/Source/Runtime/Engine/Private/GameMode.cpp
--- a/Engine/Source/Runtime/Engine/Private/GameMode.cpp
+++ b/Engine/Source/Runtime/Engine/Private/GameMode.cpp
@@ -353,7 +353,7 @@ AActor* AGameMode::FindPlayerStart( AController* Player, const FString& Incoming
 
 	AActor* BestStart = ChoosePlayerStart(Player);
 
-	if ( (BestStart == NULL) && (Player == NULL) )
+	if (BestStart == NULL)
 	{
 		// no playerstart found
 		UE_LOG(LogGameMode, Log, TEXT("Warning - PATHS NOT DEFINED or NO PLAYERSTART with positive rating"));
