Expose constraint limits and motion types to blueprint

diff --git a/Engine/Source/Runtime/Engine/Classes/PhysicsEngine/ConstraintInstance.h b/Engine/Source/Runtime/Engine/Classes/PhysicsEngine/ConstraintInstance.h
--- a/Engine/Source/Runtime/Engine/Classes/PhysicsEngine/ConstraintInstance.h
+++ b/Engine/Source/Runtime/Engine/Classes/PhysicsEngine/ConstraintInstance.h
@@ -348,6 +348,39 @@ struct ENGINE_API FConstraintInstance
 	UPROPERTY(EditAnywhere, BlueprintReadOnly, Category=AngularMotor)
 	float AngularDriveForceLimit;    //	Due to the way the ConstraintInstance pooling works, this MUST BE LAST PROPERTY OF THE CLASS.
 
+	float AverageMass;
+
+
+	/** Sets the LinearX Motion Type
+	*	@param MotionType	New Motion Type
+	*/
+	void SetLinearXLimit(ELinearConstraintMotion ConstraintType, float LinearLimitSize);
+
+	/** Sets the LinearY Motion Type
+	*	@param MotionType	New Motion Type
+	*/
+	void SetLinearYLimit(ELinearConstraintMotion ConstraintType, float LinearLimitSize);
+
+	/** Sets the LinearZ Motion Type
+	*	@param MotionType	New Motion Type
+	*/
+	void SetLinearZLimit(ELinearConstraintMotion ConstraintType, float LinearLimitSize);
+
+	/** Sets the Angular Swing1 Motion Type
+	*	@param MotionType	New Motion Type
+	*/
+	void SetAngularSwing1Limit(EAngularConstraintMotion MotionType, float Swing1LimitAngle);
+
+	/** Sets the Angular Swing2 Motion Type
+	*	@param MotionType	New Motion Type
+	*/
+	void SetAngularSwing2Limit(EAngularConstraintMotion MotionType, float Swing2LimitAngle);
+
+	/** Sets the Angular Twist Motion Type
+	*	@param MotionType	New Motion Type
+	*/
+	void SetAngularTwistLimit(EAngularConstraintMotion MotionType, float TwistLimitAngle);
+
 #if WITH_PHYSX
 	FPhysxUserData PhysxUserData;
 #endif
@@ -398,6 +431,9 @@ public:
 	void SetAngularVelocityTarget(const FVector& InVelTarget);
 	void SetAngularDriveParams(float InSpring, float InDamping, float InForceLimit);
 
+	void UpdateLinearLimit();
+	void UpdateAngularLimit();
+
 
 	/** Scale Angular Limit Constraints (as defined in RB_ConstraintSetup) */
 	void	SetAngularDOFLimitScale(float InSwing1LimitScale, float InSwing2LimitScale, float InTwistLimitScale);
diff --git a/Engine/Source/Runtime/Engine/Classes/PhysicsEngine/PhysicsConstraintComponent.h b/Engine/Source/Runtime/Engine/Classes/PhysicsEngine/PhysicsConstraintComponent.h
--- a/Engine/Source/Runtime/Engine/Classes/PhysicsEngine/PhysicsConstraintComponent.h
+++ b/Engine/Source/Runtime/Engine/Classes/PhysicsEngine/PhysicsConstraintComponent.h
@@ -175,6 +175,49 @@ public:
 	UFUNCTION(BlueprintCallable, Category="Physics|Components|PhysicsConstraint")
 	ENGINE_API void	SetAngularDriveParams(float InSpring, float InDamping, float InForceLimit);
 
+
+	/** Sets the LinearX Motion Type
+	*	@param ConstraintType	New Constraint Type
+	*	@param LimitSize		Size of limit
+	*/
+	UFUNCTION(BlueprintCallable, Category = "Physics|Components|PhysicsConstraint")
+	ENGINE_API void	SetLinearXLimit(ELinearConstraintMotion ConstraintType, float LimitSize);
+
+	/** Sets the LinearY Motion Type
+	*	@param ConstraintType	New Constraint Type
+	*	@param LimitSize		Size of limit
+	*/
+	UFUNCTION(BlueprintCallable, Category = "Physics|Components|PhysicsConstraint")
+	ENGINE_API void	SetLinearYLimit(ELinearConstraintMotion ConstraintType, float LimitSize);
+
+	/** Sets the LinearZ Motion Type
+	*	@param ConstraintType	New Constraint Type
+	*	@param LimitSize		Size of limit
+	*/
+	UFUNCTION(BlueprintCallable, Category = "Physics|Components|PhysicsConstraint")
+	ENGINE_API void	SetLinearZLimit(ELinearConstraintMotion ConstraintType, float LimitSize);
+
+	/** Sets the Angular Swing1 Motion Type
+	*	@param ConstraintType	New Constraint Type
+	*	@param Swing1LimitAngle	Size of limit in degrees
+	*/
+	UFUNCTION(BlueprintCallable, Category = "Physics|Components|PhysicsConstraint")
+	ENGINE_API void	SetAngularSwing1Limit(EAngularConstraintMotion MotionType, float Swing1LimitAngle);
+
+	/** Sets the Angular Swing2 Motion Type
+	*	@param ConstraintType	New Constraint Type
+	*	@param Swing2LimitAngle	Size of limit in degrees
+	*/
+	UFUNCTION(BlueprintCallable, Category = "Physics|Components|PhysicsConstraint")
+	ENGINE_API void	SetAngularSwing2Limit(EAngularConstraintMotion MotionType, float Swing2LimitAngle);
+
+	/** Sets the Angular Twist Motion Type
+	*	@param ConstraintType	New Constraint Type
+	*	@param TwistLimitAngle	Size of limit in degrees
+	*/
+	UFUNCTION(BlueprintCallable, Category = "Physics|Components|PhysicsConstraint")
+	ENGINE_API void	SetAngularTwistLimit(EAngularConstraintMotion ConstraintType, float TwistLimitAngle);
+
 	/** 
 	 *	Update the reference frames held inside the constraint that indicate the joint location in the reference frame 
 	 *	of the two connected bodies. You should call this whenever the constraint or either Component moves, or if you change
diff --git a/Engine/Source/Runtime/Engine/Private/PhysicsEngine/ConstraintInstance.cpp b/Engine/Source/Runtime/Engine/Private/PhysicsEngine/ConstraintInstance.cpp
--- a/Engine/Source/Runtime/Engine/Private/PhysicsEngine/ConstraintInstance.cpp
+++ b/Engine/Source/Runtime/Engine/Private/PhysicsEngine/ConstraintInstance.cpp
@@ -39,6 +39,90 @@ static void SetLinearMovement(PxD6Joint* PD6Joint, PxD6Axis::Enum PAxis, uint8 M
 		PD6Joint->setMotion(PAxis, PxD6Motion::eFREE);
 	}
 }
+
+void FConstraintInstance::UpdateLinearLimit()
+{
+	PxD6Joint* Joint = ConstraintData;
+	if (Joint && !(Joint->getConstraintFlags()&PxConstraintFlag::eBROKEN))
+	{
+		bool bLockLimitSize = (LinearLimitSize < RB_MinSizeToLockDOF);
+
+		SetLinearMovement(Joint, PxD6Axis::eX, LinearXMotion, bLockLimitSize);
+		SetLinearMovement(Joint, PxD6Axis::eY, LinearYMotion, bLockLimitSize);
+		SetLinearMovement(Joint, PxD6Axis::eZ, LinearZMotion, bLockLimitSize);
+
+		// If any DOF is locked/limited, set up the joint limit
+		if (LinearXMotion != LCM_Free || LinearYMotion != LCM_Free || LinearZMotion != LCM_Free)
+		{
+			// If limit drops below RB_MinSizeToLockDOF, just pass RB_MinSizeToLockDOF to physics - that axis will be locked anyway, and PhysX dislikes 0 here
+			float LinearLimit = FMath::Max<float>(LinearLimitSize, RB_MinSizeToLockDOF);
+			PxJointLinearLimit PLinearLimit(GPhysXSDK->getTolerancesScale(), LinearLimit, 0.05f * GPhysXSDK->getTolerancesScale().length);
+			SetSoftLimitParams(&PLinearLimit, bLinearLimitSoft, LinearLimitStiffness*AverageMass*CVarConstraintSitffnessScale->GetValueOnGameThread(), LinearLimitDamping*AverageMass*CVarConstraintDamplingScale->GetValueOnGameThread());
+			Joint->setLinearLimit(PLinearLimit);
+		}
+	}
+}
+
+void FConstraintInstance::UpdateAngularLimit()
+{
+	PxD6Joint* Joint = ConstraintData;
+	if (Joint && !(Joint->getConstraintFlags()&PxConstraintFlag::eBROKEN))
+	{
+		/////////////// TWIST LIMIT
+		PxD6Motion::Enum TwistMotion = PxD6Motion::eFREE;
+		if (AngularTwistMotion == ACM_Limited)
+		{
+			TwistMotion = PxD6Motion::eLIMITED;
+			// If angle drops below RB_MinAngleToLockDOF, just pass RB_MinAngleToLockDOF to physics - that axis will be locked anyway, and PhysX dislikes 0 here
+			float TwistLimitRad = TwistLimitAngle * (PI / 180.0f);
+			PxJointAngularLimitPair PTwistLimitPair(-TwistLimitRad, TwistLimitRad, 1.f * (PI / 180.0f));
+			SetSoftLimitParams(&PTwistLimitPair, bTwistLimitSoft, TwistLimitStiffness*AverageMass*CVarConstraintSitffnessScale->GetValueOnGameThread(), TwistLimitDamping*AverageMass*CVarConstraintDamplingScale->GetValueOnGameThread());
+			Joint->setTwistLimit(PTwistLimitPair);
+		}
+		else if (AngularTwistMotion == ACM_Locked)
+		{
+			TwistMotion = PxD6Motion::eLOCKED;
+		}
+		Joint->setMotion(PxD6Axis::eTWIST, TwistMotion);
+
+		/////////////// SWING1 LIMIT
+		PxD6Motion::Enum Swing1Motion = PxD6Motion::eFREE;
+		PxD6Motion::Enum Swing2Motion = PxD6Motion::eFREE;
+
+		if (AngularSwing1Motion == ACM_Limited)
+		{
+			Swing1Motion = PxD6Motion::eLIMITED;
+		}
+		else if (AngularSwing1Motion == ACM_Locked)
+		{
+			Swing1Motion = PxD6Motion::eLOCKED;
+		}
+
+		/////////////// SWING2 LIMIT
+		if (AngularSwing2Motion == ACM_Limited)
+		{
+			Swing2Motion = PxD6Motion::eLIMITED;
+		}
+		else if (AngularSwing2Motion == ACM_Locked)
+		{
+			Swing2Motion = PxD6Motion::eLOCKED;
+		}
+
+		if (AngularSwing1Motion == ACM_Limited || AngularSwing2Motion == ACM_Limited)
+		{
+			//Clamp the limit value to valid range which PhysX won't ignore, both value have to be clamped even there is only one degree limit in constraint
+			float Limit1Rad = FMath::ClampAngle(Swing1LimitAngle, KINDA_SMALL_NUMBER, 179.9999f) * (PI / 180.0f);
+			float Limit2Rad = FMath::ClampAngle(Swing2LimitAngle, KINDA_SMALL_NUMBER, 179.9999f) * (PI / 180.0f);
+			PxJointLimitCone PSwingLimitCone(Limit2Rad, Limit1Rad, 1.f * (PI / 180.0f));
+			SetSoftLimitParams(&PSwingLimitCone, bSwingLimitSoft, SwingLimitStiffness*AverageMass*CVarConstraintSitffnessScale->GetValueOnGameThread(), SwingLimitDamping*AverageMass*CVarConstraintDamplingScale->GetValueOnGameThread());
+			Joint->setSwingLimit(PSwingLimitCone);
+		}
+
+		Joint->setMotion(PxD6Axis::eSWING2, Swing1Motion);
+		Joint->setMotion(PxD6Axis::eSWING1, Swing2Motion);
+	}
+	
+}
 #endif // WITH_PHYSX
 
 /** Constructor **/
@@ -72,6 +156,7 @@ FConstraintInstance::FConstraintInstance()
 	, AngularDriveSpring(50.0f)
 	, AngularDriveDamping(1.0f)
 	, AngularDriveForceLimit(0)
+	, AverageMass(0.f)
 #if WITH_PHYSX
 	, PhysxUserData(this)
 #endif
@@ -145,7 +230,7 @@ void FConstraintInstance::InitConstraint(USceneComponent* Owner, FBodyInstance*
 		}
 	}
 
-	float AverageMass = 0;
+	AverageMass = 0;
 	{
 		float TotalMass = 0;
 		int NumDynamic = 0;
@@ -251,75 +336,11 @@ void FConstraintInstance::InitConstraint(USceneComponent* Owner, FBodyInstance*
 
 	PD6Joint->setConstraintFlags(Flags);
 
-	/////////////// TWIST LIMIT
-	PxD6Motion::Enum TwistMotion = PxD6Motion::eFREE;
-	if(AngularTwistMotion==ACM_Limited)
-	{
-		TwistMotion = PxD6Motion::eLIMITED;
-		// If angle drops below RB_MinAngleToLockDOF, just pass RB_MinAngleToLockDOF to physics - that axis will be locked anyway, and PhysX dislikes 0 here
-		float TwistLimitRad = TwistLimitAngle * (PI/180.0f);
-		PxJointAngularLimitPair PTwistLimitPair(-TwistLimitRad, TwistLimitRad, 1.f * (PI/180.0f));
-		SetSoftLimitParams(&PTwistLimitPair, bTwistLimitSoft, TwistLimitStiffness*AverageMass*CVarConstraintSitffnessScale->GetValueOnGameThread(), TwistLimitDamping*AverageMass*CVarConstraintDamplingScale->GetValueOnGameThread());
-		PD6Joint->setTwistLimit(PTwistLimitPair);
-	}
-	else if (AngularTwistMotion==ACM_Locked)
-	{
-		TwistMotion = PxD6Motion::eLOCKED;
-	}
-	PD6Joint->setMotion(PxD6Axis::eTWIST, TwistMotion);
-
-	/////////////// SWING1 LIMIT
-	PxD6Motion::Enum Swing1Motion = PxD6Motion::eFREE;
-	PxD6Motion::Enum Swing2Motion = PxD6Motion::eFREE;
-
-	if(AngularSwing1Motion==ACM_Limited)
-	{
-		Swing1Motion = PxD6Motion::eLIMITED;
-	}
-	else if (AngularSwing1Motion==ACM_Locked)
-	{
-		Swing1Motion = PxD6Motion::eLOCKED;
-	}
-
-	/////////////// SWING2 LIMIT
-	if(AngularSwing2Motion==ACM_Limited)
-	{
-		Swing2Motion = PxD6Motion::eLIMITED;
-	}
-	else if (AngularSwing2Motion==ACM_Locked)
-	{
-		Swing2Motion = PxD6Motion::eLOCKED;
-	}
-
-	if (AngularSwing1Motion==ACM_Limited || AngularSwing2Motion==ACM_Limited)
-	{
-		//Clamp the limit value to valid range which PhysX won't ignore, both value have to be clamped even there is only one degree limit in constraint
-		float Limit1Rad = FMath::ClampAngle(Swing1LimitAngle, KINDA_SMALL_NUMBER, 179.9999f) * (PI/180.0f);
-		float Limit2Rad = FMath::ClampAngle(Swing2LimitAngle, KINDA_SMALL_NUMBER, 179.9999f) * (PI/180.0f);
-		PxJointLimitCone PSwingLimitCone(Limit2Rad, Limit1Rad, 1.f * (PI/180.0f));
-		SetSoftLimitParams(&PSwingLimitCone, bSwingLimitSoft, SwingLimitStiffness*AverageMass*CVarConstraintSitffnessScale->GetValueOnGameThread(), SwingLimitDamping*AverageMass*CVarConstraintDamplingScale->GetValueOnGameThread());
-		PD6Joint->setSwingLimit(PSwingLimitCone);
-	}
-
-	PD6Joint->setMotion(PxD6Axis::eSWING2, Swing1Motion);
-	PD6Joint->setMotion(PxD6Axis::eSWING1, Swing2Motion);
+	/////////////// ANUGULAR LIMIT
+	UpdateAngularLimit();
 
 	/////////////// LINEAR LIMIT
-	bool bLockLimitSize = (LinearLimitSize < RB_MinSizeToLockDOF);
-
-	SetLinearMovement(PD6Joint, PxD6Axis::eX, LinearXMotion, bLockLimitSize);
-	SetLinearMovement(PD6Joint, PxD6Axis::eY, LinearYMotion, bLockLimitSize);
-	SetLinearMovement(PD6Joint, PxD6Axis::eZ, LinearZMotion, bLockLimitSize);
-
-	// If any DOF is locked/limited, set up the joint limit
-	if(LinearXMotion != LCM_Free || LinearYMotion != LCM_Free || LinearZMotion != LCM_Free)
-	{
-		// If limit drops below RB_MinSizeToLockDOF, just pass RB_MinSizeToLockDOF to physics - that axis will be locked anyway, and PhysX dislikes 0 here
-		float LinearLimit = FMath::Max<float>(LinearLimitSize, RB_MinSizeToLockDOF);
-		PxJointLinearLimit PLinearLimit(GPhysXSDK->getTolerancesScale(), LinearLimit, 0.05f * GPhysXSDK->getTolerancesScale().length);
-		SetSoftLimitParams(&PLinearLimit, bLinearLimitSoft, LinearLimitStiffness*AverageMass*CVarConstraintSitffnessScale->GetValueOnGameThread(), LinearLimitDamping*AverageMass*CVarConstraintDamplingScale->GetValueOnGameThread());
-		PD6Joint->setLinearLimit(PLinearLimit);
-	}
+	UpdateLinearLimit();
 
 	///////// BREAKABLE
 	float LinearBreakForce = PX_MAX_REAL;
@@ -1018,4 +1039,52 @@ void FConstraintInstance::OnConstraintBroken()
 
 }
 
+void FConstraintInstance::SetLinearXLimit(ELinearConstraintMotion XMotion, float InLinearLimitSize)
+{
+	LinearXMotion = XMotion;
+	LinearLimitSize = InLinearLimitSize;
+
+	UpdateLinearLimit();
+}
+
+void FConstraintInstance::SetLinearYLimit(ELinearConstraintMotion YMotion, float InLinearLimitSize)
+{
+	LinearYMotion = YMotion;
+	LinearLimitSize = InLinearLimitSize;
+	
+	UpdateLinearLimit();
+}
+
+void FConstraintInstance::SetLinearZLimit(ELinearConstraintMotion ZMotion, float InLinearLimitSize)
+{
+	LinearZMotion = ZMotion;
+	LinearLimitSize = InLinearLimitSize;
+
+	UpdateLinearLimit();
+}
+
+void FConstraintInstance::SetAngularSwing1Limit(EAngularConstraintMotion Swing1Motion, float InSwing1LimitAngle)
+{
+	AngularSwing1Motion = Swing1Motion;
+	Swing1LimitAngle = InSwing1LimitAngle;
+	
+	UpdateAngularLimit();
+}
+
+void FConstraintInstance::SetAngularSwing2Limit(EAngularConstraintMotion Swing2Motion, float InSwing2LimitAngle)
+{
+	AngularSwing2Motion = Swing2Motion;
+	Swing2LimitAngle = InSwing2LimitAngle;
+
+	UpdateAngularLimit();
+}
+
+void FConstraintInstance::SetAngularTwistLimit(EAngularConstraintMotion TwistMotion, float InTwistLimitAngle)
+{
+	AngularTwistMotion = TwistMotion;
+	TwistLimitAngle = InTwistLimitAngle;
+
+	UpdateAngularLimit();
+}
+
 #undef LOCTEXT_NAMESPACE
\ No newline at end of file
diff --git a/Engine/Source/Runtime/Engine/Private/PhysicsEngine/PhysicsConstraintComponent.cpp b/Engine/Source/Runtime/Engine/Private/PhysicsEngine/PhysicsConstraintComponent.cpp
--- a/Engine/Source/Runtime/Engine/Private/PhysicsEngine/PhysicsConstraintComponent.cpp
+++ b/Engine/Source/Runtime/Engine/Private/PhysicsEngine/PhysicsConstraintComponent.cpp
@@ -476,4 +476,34 @@ void UPhysicsConstraintComponent::SetAngularDriveParams( float InSpring, float I
 	ConstraintInstance.SetAngularDriveParams(InSpring, InDamping, InForceLimit);
 }
 
+void UPhysicsConstraintComponent::SetLinearXLimit(ELinearConstraintMotion Motion, float LinearLimit)
+{
+	ConstraintInstance.SetLinearXLimit(Motion, LinearLimit);
+}
+
+void UPhysicsConstraintComponent::SetLinearYLimit(ELinearConstraintMotion Motion, float LinearLimit)
+{
+	ConstraintInstance.SetLinearYLimit(Motion, LinearLimit);
+}
+
+void UPhysicsConstraintComponent::SetLinearZLimit(ELinearConstraintMotion Motion, float LinearLimit)
+{
+	ConstraintInstance.SetLinearZLimit(Motion, LinearLimit);
+}
+
+void UPhysicsConstraintComponent::SetAngularSwing1Limit(EAngularConstraintMotion Motion, float Swing1LimitAngle)
+{
+	ConstraintInstance.SetAngularSwing1Limit(Motion, Swing1LimitAngle);
+}
+
+void UPhysicsConstraintComponent::SetAngularSwing2Limit(EAngularConstraintMotion Motion, float Swing2LimitAngle)
+{
+	ConstraintInstance.SetAngularSwing2Limit(Motion, Swing2LimitAngle);
+}
+
+void UPhysicsConstraintComponent::SetAngularTwistLimit(EAngularConstraintMotion Motion, float TwistLimitAngle)
+{
+	ConstraintInstance.SetAngularTwistLimit(Motion, TwistLimitAngle);
+}
+
 #undef LOCTEXT_NAMESPACE
