Fix compiler for WITH_PHYSX

diff --git a/Engine/Source/Runtime/Engine/Private/PhysicsEngine/ConstraintInstance.cpp b/Engine/Source/Runtime/Engine/Private/PhysicsEngine/ConstraintInstance.cpp
--- a/Engine/Source/Runtime/Engine/Private/PhysicsEngine/ConstraintInstance.cpp
+++ b/Engine/Source/Runtime/Engine/Private/PhysicsEngine/ConstraintInstance.cpp
@@ -39,9 +39,11 @@ static void SetLinearMovement(PxD6Joint* PD6Joint, PxD6Axis::Enum PAxis, uint8 M
 		PD6Joint->setMotion(PAxis, PxD6Motion::eFREE);
 	}
 }
+#endif //WITH_PHYSX
 
 void FConstraintInstance::UpdateLinearLimit()
 {
+#if WITH_PHYSX
 	PxD6Joint* Joint = ConstraintData;
 	if (Joint && !(Joint->getConstraintFlags()&PxConstraintFlag::eBROKEN))
 	{
@@ -61,10 +63,12 @@ void FConstraintInstance::UpdateLinearLimit()
 			Joint->setLinearLimit(PLinearLimit);
 		}
 	}
+#endif
 }
 
 void FConstraintInstance::UpdateAngularLimit()
 {
+#if WITH_PHYSX
 	PxD6Joint* Joint = ConstraintData;
 	if (Joint && !(Joint->getConstraintFlags()&PxConstraintFlag::eBROKEN))
 	{
@@ -121,9 +125,9 @@ void FConstraintInstance::UpdateAngularLimit()
 		Joint->setMotion(PxD6Axis::eSWING2, Swing1Motion);
 		Joint->setMotion(PxD6Axis::eSWING1, Swing2Motion);
 	}
+#endif
 	
 }
-#endif // WITH_PHYSX
 
 /** Constructor **/
 FConstraintInstance::FConstraintInstance()
