#ue4 - Restore some perf counters for GeomSweepMulti and GeomOverlapMulti lost in a refactor.

diff --git a/Engine/Source/Runtime/Engine/Private/Collision/PhysXCollision.cpp b/Engine/Source/Runtime/Engine/Private/Collision/PhysXCollision.cpp
--- a/Engine/Source/Runtime/Engine/Private/Collision/PhysXCollision.cpp
+++ b/Engine/Source/Runtime/Engine/Private/Collision/PhysXCollision.cpp
@@ -1273,6 +1273,7 @@ bool GeomSweepSingle(const UWorld* World, const struct FCollisionShape& Collisio
 
 bool GeomSweepMulti_PhysX(const UWorld* World, const PxGeometry& PGeom, const PxQuat& PGeomRot, TArray<FHitResult>& OutHits, FVector Start, FVector End, ECollisionChannel TraceChannel, const struct FCollisionQueryParams& Params, const struct FCollisionResponseParams& ResponseParams, const struct FCollisionObjectQueryParams& ObjectParams)
 {
+	SCOPE_CYCLE_COUNTER(STAT_Collision_GeomSweepMultiple);
 	STARTQUERYTIMER();
 	bool bBlockingHit = false;
 
@@ -1589,6 +1590,7 @@ bool GeomOverlapSingle(const UWorld* World, const struct FCollisionShape& Collis
 #if WITH_PHYSX
 bool GeomOverlapMulti_PhysX(const UWorld* World, const PxGeometry& PGeom, const PxTransform& PGeomPose, TArray<FOverlapResult>& OutOverlaps, ECollisionChannel TraceChannel, const struct FCollisionQueryParams& Params, const struct FCollisionResponseParams& ResponseParams, const struct FCollisionObjectQueryParams& ObjectParams)
 {
+	SCOPE_CYCLE_COUNTER(STAT_Collision_GeomOverlapMultiple);
 	bool bHaveBlockingHit = false;
 
 	// overlapMultiple only supports sphere/capsule/box 
diff --git a/Engine/Source/Runtime/Engine/Private/Collision/WorldCollision.cpp b/Engine/Source/Runtime/Engine/Private/Collision/WorldCollision.cpp
--- a/Engine/Source/Runtime/Engine/Private/Collision/WorldCollision.cpp
+++ b/Engine/Source/Runtime/Engine/Private/Collision/WorldCollision.cpp
@@ -328,6 +328,7 @@ bool UWorld::ComponentSweepMulti(TArray<struct FHitResult>& OutHits, class UPrim
 	}
 #endif
 
+	SCOPE_CYCLE_COUNTER(STAT_Collision_GeomSweepMultiple);
 	bool bHaveBlockingHit = false;
 
 #if WITH_PHYSX
diff --git a/Engine/Source/Runtime/Engine/Private/PhysicsEngine/BodyInstance.cpp b/Engine/Source/Runtime/Engine/Private/PhysicsEngine/BodyInstance.cpp
--- a/Engine/Source/Runtime/Engine/Private/PhysicsEngine/BodyInstance.cpp
+++ b/Engine/Source/Runtime/Engine/Private/PhysicsEngine/BodyInstance.cpp
@@ -3177,6 +3177,7 @@ bool FBodyInstance::OverlapMulti(TArray<struct FOverlapResult>& InOutOverlaps, c
 		return false;
 	}
 
+	SCOPE_CYCLE_COUNTER(STAT_Collision_GeomOverlapMultiple);
 	bool bHaveBlockingHit = false;
 
 	// Determine how to convert the local space of this body instance to the test space
