Fix dedicated servers taking the oculus hmd connection

diff --git a/Engine/Source/Runtime/Engine/Private/UnrealEngine.cpp b/Engine/Source/Runtime/Engine/Private/UnrealEngine.cpp
--- a/Engine/Source/Runtime/Engine/Private/UnrealEngine.cpp
+++ b/Engine/Source/Runtime/Engine/Private/UnrealEngine.cpp
@@ -1611,7 +1611,8 @@ bool UEngine::InitializeHMDDevice()
 			TSharedPtr<FFakeStereoRenderingDevice> FakeStereoDevice(new FFakeStereoRenderingDevice());
 			StereoRenderingDevice = FakeStereoDevice;
 		}
-		else if(!HMDDevice.IsValid() && !FParse::Param(FCommandLine::Get(),TEXT("nohmd")))
+		// No reason to connect an HMD on a dedicated server.  Also fixes dedicated servers stealing the oculus connection.
+		else if(!HMDDevice.IsValid() && !FParse::Param(FCommandLine::Get(),TEXT("nohmd")) && !IsRunningDedicatedServer())
 		{
 			// Get a list of plugins that implement this feature
 			TArray<IHeadMountedDisplayModule*> HMDImplementations = IModularFeatures::Get().GetModularFeatureImplementations<IHeadMountedDisplayModule>( IHeadMountedDisplayModule::GetModularFeatureName() );
