Fix PC crash when all audio devices are disabled

diff --git a/Engine/Source/Runtime/Engine/Private/GameViewportClient.cpp b/Engine/Source/Runtime/Engine/Private/GameViewportClient.cpp
--- a/Engine/Source/Runtime/Engine/Private/GameViewportClient.cpp
+++ b/Engine/Source/Runtime/Engine/Private/GameViewportClient.cpp
@@ -763,9 +763,7 @@ void UGameViewportClient::Draw(FViewport* InViewport, FCanvas* SceneCanvas)
 
 								FTransform ListenerTransform(FRotationMatrix::MakeFromXY(ProjFront, ProjRight));
 								ListenerTransform.SetTranslation(Location);
-								ListenerTransform.NormalizeRotation();
-
-								bReverbSettingsFound = true;
+								ListenerTransform.NormalizeRotation();								
 
 								FReverbSettings PlayerReverbSettings;
 								FInteriorSettings PlayerInteriorSettings;
@@ -777,6 +775,8 @@ void UGameViewportClient::Draw(FViewport* InViewport, FCanvas* SceneCanvas)
 									ReverbSettings = PlayerReverbSettings;
 								}
 
+								bReverbSettingsFound = ReverbVolume != nullptr;
+
 								uint32 ViewportIndex = PlayerViewMap.Num()-1;
 								AudioDevice->SetListener(ViewportIndex, ListenerTransform, (View->bCameraCut ? 0.f : GetWorld()->GetDeltaSeconds()), PlayerReverbVolume, PlayerInteriorSettings);
 							}
