Roll back crash fix.  Could cause reverb hang later instead.
wwise probably to blame, needs to be fixed elsewhere.

diff --git a/Engine/Source/Runtime/Engine/Private/GameViewportClient.cpp b/Engine/Source/Runtime/Engine/Private/GameViewportClient.cpp
--- a/Engine/Source/Runtime/Engine/Private/GameViewportClient.cpp
+++ b/Engine/Source/Runtime/Engine/Private/GameViewportClient.cpp
@@ -763,7 +763,9 @@ void UGameViewportClient::Draw(FViewport* InViewport, FCanvas* SceneCanvas)
 
 								FTransform ListenerTransform(FRotationMatrix::MakeFromXY(ProjFront, ProjRight));
 								ListenerTransform.SetTranslation(Location);
-								ListenerTransform.NormalizeRotation();								
+								ListenerTransform.NormalizeRotation();
+
+								bReverbSettingsFound = true;
 
 								FReverbSettings PlayerReverbSettings;
 								FInteriorSettings PlayerInteriorSettings;
@@ -775,8 +777,6 @@ void UGameViewportClient::Draw(FViewport* InViewport, FCanvas* SceneCanvas)
 									ReverbSettings = PlayerReverbSettings;
 								}
 
-								bReverbSettingsFound = ReverbVolume != nullptr;
-
 								uint32 ViewportIndex = PlayerViewMap.Num()-1;
 								AudioDevice->SetListener(ViewportIndex, ListenerTransform, (View->bCameraCut ? 0.f : GetWorld()->GetDeltaSeconds()), PlayerReverbVolume, PlayerInteriorSettings);
 							}
