If it snap to, it should stay with whatever current scale is.

diff --git a/Engine/Source/Runtime/Engine/Private/SceneComponent.cpp b/Engine/Source/Runtime/Engine/Private/SceneComponent.cpp
--- a/Engine/Source/Runtime/Engine/Private/SceneComponent.cpp
+++ b/Engine/Source/Runtime/Engine/Private/SceneComponent.cpp
@@ -720,7 +720,7 @@ void USceneComponent::AttachTo(class USceneComponent* Parent, FName InSocketName
 		
 
 		// Make sure we are detached
-		bool bMaintainWorldPosition = AttachType == EAttachLocation::KeepWorldPosition;
+		bool bMaintainWorldPosition = AttachType == EAttachLocation::KeepWorldPosition || EAttachLocation::SnapToTarget;
 		DetachFromParent(bMaintainWorldPosition);
 
 		// Detach removes all Prerequisite, so will need to add after Detach happens
