Added flag to allow landscapes to cast two-sided shadows

diff --git a/Engine/Source/Editor/LandscapeEditor/Private/LandscapeEdMode.cpp b/Engine/Source/Editor/LandscapeEditor/Private/LandscapeEdMode.cpp
--- a/Engine/Source/Editor/LandscapeEditor/Private/LandscapeEdMode.cpp
+++ b/Engine/Source/Editor/LandscapeEditor/Private/LandscapeEdMode.cpp
@@ -2551,6 +2551,7 @@ ALandscape* FEdModeLandscape::ChangeComponentSetting(int32 NumComponentsX, int32
 			Landscape->LODDistanceFactor           = OldLandscapeProxy->LODDistanceFactor;
 			Landscape->StaticLightingResolution    = OldLandscapeProxy->StaticLightingResolution;
 			Landscape->bCastStaticShadow           = OldLandscapeProxy->bCastStaticShadow;
+			Landscape->bCastShadowAsTwoSided	   = OldLandscapeProxy->bCastShadowAsTwoSided;
 			Landscape->LightmassSettings           = OldLandscapeProxy->LightmassSettings;
 			Landscape->CollisionThickness          = OldLandscapeProxy->CollisionThickness;
 			Landscape->BodyInstance.SetCollisionProfileName(OldLandscapeProxy->BodyInstance.GetCollisionProfileName());
diff --git a/Engine/Source/Editor/LandscapeEditor/Private/LandscapeEdModeComponentTools.cpp b/Engine/Source/Editor/LandscapeEditor/Private/LandscapeEdModeComponentTools.cpp
--- a/Engine/Source/Editor/LandscapeEditor/Private/LandscapeEdModeComponentTools.cpp
+++ b/Engine/Source/Editor/LandscapeEditor/Private/LandscapeEdModeComponentTools.cpp
@@ -811,6 +811,7 @@ public:
 
 						// Assign shared properties
 						LandscapeComponent->bCastStaticShadow = Landscape->bCastStaticShadow;
+						LandscapeComponent->bCastShadowAsTwoSided = Landscape->bCastShadowAsTwoSided;
 
 						int32 ComponentVerts = (Landscape->SubsectionSizeQuads+1) * Landscape->NumSubsections;
 						// Update Weightmap Scale Bias
diff --git a/Engine/Source/Runtime/Engine/Classes/Landscape/LandscapeProxy.h b/Engine/Source/Runtime/Engine/Classes/Landscape/LandscapeProxy.h
--- a/Engine/Source/Runtime/Engine/Classes/Landscape/LandscapeProxy.h
+++ b/Engine/Source/Runtime/Engine/Classes/Landscape/LandscapeProxy.h
@@ -257,6 +257,10 @@ public:
 	UPROPERTY(EditAnywhere, Category=Lighting)
 	uint32 bCastStaticShadow:1;
 
+	/** Whether this primitive should cast dynamic shadows as if it were a two sided material. */
+	UPROPERTY(EditAnywhere, AdvancedDisplay, Category=Lighting)
+	uint32 bCastShadowAsTwoSided:1;
+
 	UPROPERTY()
 	uint32 bIsProxy:1;
 
diff --git a/Engine/Source/Runtime/Engine/Private/Landscape/Landscape.cpp b/Engine/Source/Runtime/Engine/Private/Landscape/Landscape.cpp
--- a/Engine/Source/Runtime/Engine/Private/Landscape/Landscape.cpp
+++ b/Engine/Source/Runtime/Engine/Private/Landscape/Landscape.cpp
@@ -351,7 +351,8 @@ void ULandscapeComponent::PostLoad()
 	if (ensure(LandscapeProxy))
 	{
 		bCastStaticShadow = LandscapeProxy->bCastStaticShadow;
-		
+		bCastShadowAsTwoSided = LandscapeProxy->bCastShadowAsTwoSided;
+
 		// convert from deprecated layer names to direct LayerInfo references
 		static const FName DataWeightmapName = FName("__DataLayer__");
 		for (int32 i = 0; i < WeightmapLayerAllocations.Num(); i++)
@@ -505,6 +506,7 @@ ALandscapeProxy::ALandscapeProxy(const class FPostConstructInitializeProperties&
 	LODDistanceFactor = 1.0f;
 	LODFalloff = ELandscapeLODFalloff::Linear;
 	bCastStaticShadow = true;
+	bCastShadowAsTwoSided = false;
 	bUsedForNavigation = true;
 	CollisionThickness = 16;
 	BodyInstance.SetCollisionProfileName(UCollisionProfile::BlockAll_ProfileName);
@@ -1304,6 +1306,7 @@ void ALandscapeProxy::GetSharedProperties(ALandscapeProxy* Landscape)
 		//PrePivot = Landscape->PrePivot;
 		StaticLightingResolution = Landscape->StaticLightingResolution;
 		bCastStaticShadow = Landscape->bCastStaticShadow;
+		bCastShadowAsTwoSided = Landscape->bCastShadowAsTwoSided;
 		ComponentSizeQuads = Landscape->ComponentSizeQuads;
 		NumSubsections = Landscape->NumSubsections;
 		SubsectionSizeQuads = Landscape->SubsectionSizeQuads;
diff --git a/Engine/Source/Runtime/Engine/Private/Landscape/LandscapeEdit.cpp b/Engine/Source/Runtime/Engine/Private/Landscape/LandscapeEdit.cpp
--- a/Engine/Source/Runtime/Engine/Private/Landscape/LandscapeEdit.cpp
+++ b/Engine/Source/Runtime/Engine/Private/Landscape/LandscapeEdit.cpp
@@ -1727,6 +1727,7 @@ void ALandscapeProxy::Import(FGuid Guid, int32 VertsX, int32 VertsY,
 
 			// Assign shared properties
 			LandscapeComponent->bCastStaticShadow = bCastStaticShadow;
+			LandscapeComponent->bCastShadowAsTwoSided = bCastShadowAsTwoSided;
 		}
 	}
 
@@ -3571,7 +3572,7 @@ void ALandscapeProxy::PostEditChangeProperty(FPropertyChangedEvent& PropertyChan
 	{
 		RecreateCollisionComponents();
 	}
-	else if( PropertyName == FName(TEXT("bCastStaticShadow")) )
+	else if (PropertyName == FName(TEXT("bCastStaticShadow")) || PropertyName == FName(TEXT("bCastShadowAsTwoSided")))
 	{
 		// Replicate shared properties to all components.
 		for(int32 ComponentIndex = 0; ComponentIndex < LandscapeComponents.Num(); ComponentIndex++ )
@@ -3580,6 +3581,7 @@ void ALandscapeProxy::PostEditChangeProperty(FPropertyChangedEvent& PropertyChan
 			if( Comp )
 			{
 				Comp->bCastStaticShadow = bCastStaticShadow;
+				Comp->bCastShadowAsTwoSided = bCastShadowAsTwoSided;
 			}
 		}
 	}
