Fixed octree implementation of origin rebasing

diff --git a/Engine/Source/Runtime/Engine/Public/GenericOctree.h b/Engine/Source/Runtime/Engine/Public/GenericOctree.h
--- a/Engine/Source/Runtime/Engine/Public/GenericOctree.h
+++ b/Engine/Source/Runtime/Engine/Public/GenericOctree.h
@@ -721,11 +721,8 @@ public:
 		RootNode.ShrinkElements();
 	}
 
-	void ApplyOffset(const FVector& InOffset)
-	{
-		RootNodeContext.Bounds.Center+= FVector4(InOffset, 0.f);
-		RootNode.ApplyOffset(InOffset);
-	}
+	/** Apply an arbitrary offset to all elements in the tree */
+	void ApplyOffset(const FVector& InOffset);
 
 	/** Initialization constructor. */
 	TOctree(const FVector& InOrigin,float InExtent);
diff --git a/Engine/Source/Runtime/Engine/Public/GenericOctree.inl b/Engine/Source/Runtime/Engine/Public/GenericOctree.inl
--- a/Engine/Source/Runtime/Engine/Public/GenericOctree.inl
+++ b/Engine/Source/Runtime/Engine/Public/GenericOctree.inl
@@ -345,4 +345,41 @@ TOctree<ElementType,OctreeSemantics>::TOctree(const FVector& InOrigin,float InEx
 {
 }
 
+template<typename ElementType,typename OctreeSemantics>
+void TOctree<ElementType,OctreeSemantics>::ApplyOffset(const FVector& InOffset)
+{
+	// Shift elements
+	RootNode.ApplyOffset(InOffset);
+	
+	// Make a local copy of all nodes
+	FNode OldRootNode = RootNode;
+	
+	// Assign to root node a NULL node, so Destroy procedure will not delete child nodes
+	RootNode = FNode(nullptr);
+	// Call destroy to clean up octree
+	Destroy();
+
+	// Add all elements from a saved nodes to a new empty octree
+	for (TConstIterator<> NodeIt(OldRootNode, RootNodeContext); NodeIt.HasPendingNodes(); NodeIt.Advance())
+	{
+		const auto& CurrentNode = NodeIt.GetCurrentNode();
+
+		FOREACH_OCTREE_CHILD_NODE(ChildRef)
+		{
+			if (CurrentNode.HasChild(ChildRef))
+			{
+				NodeIt.PushChild(ChildRef);
+			}
+		}
+
+		for (auto ElementIt = CurrentNode.GetElementIt(); ElementIt; ++ElementIt)
+		{
+			AddElement(*ElementIt);
+		}
+	}
+
+	// Saved nodes will be deleted on scope exit
+}
+
+
 #endif
diff --git a/Engine/Source/Runtime/Renderer/Private/RendererScene.cpp b/Engine/Source/Runtime/Renderer/Private/RendererScene.cpp
--- a/Engine/Source/Runtime/Renderer/Private/RendererScene.cpp
+++ b/Engine/Source/Runtime/Renderer/Private/RendererScene.cpp
@@ -1879,6 +1879,8 @@ static void StaticMeshDrawListApplyWorldOffset(T(&InList)[N], FVector InOffset)
 
 void FScene::ApplyWorldOffset_RenderThread(FVector InOffset)
 {
+	QUICK_SCOPE_CYCLE_COUNTER(STAT_SceneApplyWorldOffset);
+	
 	// Primitives
 	for (auto It = Primitives.CreateIterator(); It; ++It)
 	{
@@ -1904,7 +1906,7 @@ void FScene::ApplyWorldOffset_RenderThread(FVector InOffset)
 	}
 
 	// Lights
-	VectorRegister OffsetReg = VectorLoadFloat3(&InOffset);
+	VectorRegister OffsetReg = VectorLoadFloat3_W0(&InOffset);
 	for (auto It = Lights.CreateIterator(); It; ++It)
 	{
 		(*It).BoundingSphereVector = VectorAdd((*It).BoundingSphereVector, OffsetReg);
