Further CVar fixups.

diff --git a/Engine/Source/Runtime/Engine/Private/CharacterMovementComponent.cpp b/Engine/Source/Runtime/Engine/Private/CharacterMovementComponent.cpp
--- a/Engine/Source/Runtime/Engine/Private/CharacterMovementComponent.cpp
+++ b/Engine/Source/Runtime/Engine/Private/CharacterMovementComponent.cpp
@@ -51,7 +51,7 @@ static TAutoConsoleVariable<float> CVarNetProxyShrinkHalfHeight(
 	ECVF_Default);
 
 #if !(UE_BUILD_SHIPPING || UE_BUILD_TEST)
-static TAutoConsoleVariable<int32> CvarNetShowCorrections(
+static TAutoConsoleVariable<int32> CVarNetShowCorrections(
 	TEXT("p.NetShowCorrections"),
 	0,
 	TEXT("Whether to draw client position corrections (red is incorrect, green is corrected).\n")
@@ -5366,7 +5366,7 @@ void UCharacterMovementComponent::ServerMoveHandleClientError(float TimeStamp, f
 		ServerData->PendingAdjustment.NewRot = CharacterOwner->GetActorRotation();
 
 #if !(UE_BUILD_SHIPPING || UE_BUILD_TEST)
-		if (CvarNetShowCorrections.GetValueOnGameThread() != 0)
+		if (CVarNetShowCorrections.GetValueOnGameThread() != 0)
 		{
 			UE_LOG(LogNetPlayerMovement, Warning, TEXT("******** Client Error at %f is %f Accel %s LocDiff %s ClientLoc %s, ServerLoc: %s, Base: %s"), TimeStamp, LocDiff.Size(), *Accel.ToString(), *LocDiff.ToString(), *ClientLoc.ToString(), *CharacterOwner->GetActorLocation().ToString(), *GetNameSafe(MovementBase));
 			const float DebugLifetime = CVarNetCorrectionLifetime.GetValueOnGameThread();
diff --git a/Engine/Source/Runtime/Engine/Private/ContentStreaming.cpp b/Engine/Source/Runtime/Engine/Private/ContentStreaming.cpp
--- a/Engine/Source/Runtime/Engine/Private/ContentStreaming.cpp
+++ b/Engine/Source/Runtime/Engine/Private/ContentStreaming.cpp
@@ -138,7 +138,7 @@ static TAutoConsoleVariable<int32> CVarSetTextureStreamingEnabled(
 	ECVF_Default | ECVF_RenderThreadSafe);
 #endif
 
-static TAutoConsoleVariable<int32> CvarStreamingUseFixedPoolSize(
+static TAutoConsoleVariable<int32> CVarStreamingUseFixedPoolSize(
 	TEXT("r.Streaming.UseFixedPoolSize"),
 	0,
 	TEXT("If non-zero, do not allow the pool size to change at run time."),
@@ -3657,7 +3657,7 @@ void FStreamingManagerTexture::StreamTexturesUnlimited( FStreamingContext& Conte
 void FStreamingManagerTexture::CheckUserSettings()
 {	
 	int32 PoolSizeSetting = CVarStreamingPoolSize.GetValueOnGameThread();
-	int32 FixedPoolSizeSetting = CvarStreamingUseFixedPoolSize.GetValueOnGameThread();
+	int32 FixedPoolSizeSetting = CVarStreamingUseFixedPoolSize.GetValueOnGameThread();
 
 	if (FixedPoolSizeSetting == 0)
 	{
diff --git a/Engine/Source/Runtime/Engine/Private/PhysicsEngine/ConstraintInstance.cpp b/Engine/Source/Runtime/Engine/Private/PhysicsEngine/ConstraintInstance.cpp
--- a/Engine/Source/Runtime/Engine/Private/PhysicsEngine/ConstraintInstance.cpp
+++ b/Engine/Source/Runtime/Engine/Private/PhysicsEngine/ConstraintInstance.cpp
@@ -9,8 +9,19 @@
 /** Handy macro for setting BIT of VAR based on the bool CONDITION */
 #define SET_DRIVE_PARAM(VAR, CONDITION, BIT)   (VAR) = (CONDITION) ? ((VAR) | (BIT)) : ((VAR) & ~(BIT))
 
-static TConsoleVariableData<float>* CVarConstraintDamplingScale = NULL; 
-static TConsoleVariableData<float>* CVarConstraintSitffnessScale = NULL; 
+// Cvars
+static TAutoConsoleVariable<float> CVarConstraintDampingScale(
+	TEXT("p.ConstraintDampingScale"),
+	100000.f,
+	TEXT("The multiplier of constraint damping in simulation. Default: 100000"),
+	ECVF_ReadOnly);
+
+static TAutoConsoleVariable<float> CVarConstraintStiffnessScale(
+	TEXT("p.ConstraintStiffnessScale"),
+	100000.f,
+	TEXT("The multiplier of constraint stiffness in simulation. Default: 100000"),
+	ECVF_ReadOnly);
+
 
 #if WITH_PHYSX
 /** Util for setting soft limit params */
@@ -59,7 +70,7 @@ void FConstraintInstance::UpdateLinearLimit()
 			// If limit drops below RB_MinSizeToLockDOF, just pass RB_MinSizeToLockDOF to physics - that axis will be locked anyway, and PhysX dislikes 0 here
 			float LinearLimit = FMath::Max<float>(LinearLimitSize, RB_MinSizeToLockDOF);
 			PxJointLinearLimit PLinearLimit(GPhysXSDK->getTolerancesScale(), LinearLimit, 0.05f * GPhysXSDK->getTolerancesScale().length);
-			SetSoftLimitParams(&PLinearLimit, bLinearLimitSoft, LinearLimitStiffness*AverageMass*CVarConstraintSitffnessScale->GetValueOnGameThread(), LinearLimitDamping*AverageMass*CVarConstraintDamplingScale->GetValueOnGameThread());
+			SetSoftLimitParams(&PLinearLimit, bLinearLimitSoft, LinearLimitStiffness*AverageMass*CVarConstraintStiffnessScale.GetValueOnGameThread(), LinearLimitDamping*AverageMass*CVarConstraintDampingScale.GetValueOnGameThread());
 			Joint->setLinearLimit(PLinearLimit);
 		}
 	}
@@ -80,7 +91,7 @@ void FConstraintInstance::UpdateAngularLimit()
 			// If angle drops below RB_MinAngleToLockDOF, just pass RB_MinAngleToLockDOF to physics - that axis will be locked anyway, and PhysX dislikes 0 here
 			float TwistLimitRad = TwistLimitAngle * (PI / 180.0f);
 			PxJointAngularLimitPair PTwistLimitPair(-TwistLimitRad, TwistLimitRad, 1.f * (PI / 180.0f));
-			SetSoftLimitParams(&PTwistLimitPair, bTwistLimitSoft, TwistLimitStiffness*AverageMass*CVarConstraintSitffnessScale->GetValueOnGameThread(), TwistLimitDamping*AverageMass*CVarConstraintDamplingScale->GetValueOnGameThread());
+			SetSoftLimitParams(&PTwistLimitPair, bTwistLimitSoft, TwistLimitStiffness*AverageMass*CVarConstraintStiffnessScale.GetValueOnGameThread(), TwistLimitDamping*AverageMass*CVarConstraintDampingScale.GetValueOnGameThread());
 			Joint->setTwistLimit(PTwistLimitPair);
 		}
 		else if (AngularTwistMotion == ACM_Locked)
@@ -118,7 +129,7 @@ void FConstraintInstance::UpdateAngularLimit()
 			float Limit1Rad = FMath::ClampAngle(Swing1LimitAngle, KINDA_SMALL_NUMBER, 179.9999f) * (PI / 180.0f);
 			float Limit2Rad = FMath::ClampAngle(Swing2LimitAngle, KINDA_SMALL_NUMBER, 179.9999f) * (PI / 180.0f);
 			PxJointLimitCone PSwingLimitCone(Limit2Rad, Limit1Rad, 1.f * (PI / 180.0f));
-			SetSoftLimitParams(&PSwingLimitCone, bSwingLimitSoft, SwingLimitStiffness*AverageMass*CVarConstraintSitffnessScale->GetValueOnGameThread(), SwingLimitDamping*AverageMass*CVarConstraintDamplingScale->GetValueOnGameThread());
+			SetSoftLimitParams(&PSwingLimitCone, bSwingLimitSoft, SwingLimitStiffness*AverageMass*CVarConstraintStiffnessScale.GetValueOnGameThread(), SwingLimitDamping*AverageMass*CVarConstraintDampingScale.GetValueOnGameThread());
 			Joint->setSwingLimit(PSwingLimitCone);
 		}
 
@@ -165,12 +176,6 @@ FConstraintInstance::FConstraintInstance()
 	, PhysxUserData(this)
 #endif
 {
-	CVarConstraintDamplingScale = IConsoleManager::Get().FindTConsoleVariableDataFloat(TEXT("p.ConstraintDampingScale"));
-	CVarConstraintSitffnessScale = IConsoleManager::Get().FindTConsoleVariableDataFloat(TEXT("p.ConstraintStiffnessScale"));
-
-	check (CVarConstraintSitffnessScale);
-	check (CVarConstraintDamplingScale);
-
 	bDisableCollision = false;
 
 	LinearXMotion = LCM_Locked;
@@ -1013,12 +1018,12 @@ void FConstraintInstance::PostSerialize(const FArchive& Ar)
 {
 	if (Ar.IsLoading() && Ar.UE4Ver() < VER_UE4_FIXUP_STIFFNESS_AND_DAMPING_SCALE)
 	{
-		LinearLimitStiffness	/= CVarConstraintSitffnessScale->GetValueOnGameThread();
-		SwingLimitStiffness		/= CVarConstraintSitffnessScale->GetValueOnGameThread();
-		TwistLimitStiffness		/= CVarConstraintSitffnessScale->GetValueOnGameThread();
-		LinearLimitDamping		/=  CVarConstraintDamplingScale->GetValueOnGameThread();
-		SwingLimitDamping		/=  CVarConstraintDamplingScale->GetValueOnGameThread();
-		TwistLimitDamping		/=  CVarConstraintDamplingScale->GetValueOnGameThread();
+		LinearLimitStiffness	/= CVarConstraintStiffnessScale.GetValueOnGameThread();
+		SwingLimitStiffness		/= CVarConstraintStiffnessScale.GetValueOnGameThread();
+		TwistLimitStiffness		/= CVarConstraintStiffnessScale.GetValueOnGameThread();
+		LinearLimitDamping		/=  CVarConstraintDampingScale.GetValueOnGameThread();
+		SwingLimitDamping		/=  CVarConstraintDampingScale.GetValueOnGameThread();
+		TwistLimitDamping		/=  CVarConstraintDampingScale.GetValueOnGameThread();
 	}
 
 	if (Ar.IsLoading() && Ar.UE4Ver() < VER_UE4_FIXUP_MOTOR_UNITS)
