Fix landscape collision problems where the player was falling through.

TTP#318613

diff --git a/Engine/Source/Runtime/Engine/Private/Collision/CollisionConversions.cpp b/Engine/Source/Runtime/Engine/Private/Collision/CollisionConversions.cpp
--- a/Engine/Source/Runtime/Engine/Private/Collision/CollisionConversions.cpp
+++ b/Engine/Source/Runtime/Engine/Private/Collision/CollisionConversions.cpp
@@ -486,6 +486,71 @@ bool AddSweepResults(int32 NumHits, PxSweepHit* Hits, float CheckLength, const P
 
 #define DRAW_OVERLAPPING_TRIS 0
 
+/* Function to find the best normal from the list of triangles that are overlapping our geom. */
+template<typename GeomType>
+FVector FindBestOverlappingNormal(const PxGeometry& Geom, const PxTransform& QueryTM, const GeomType& ShapeGeom, const PxTransform& PShapeWorldPose, PxU32* HitTris, int32 NumTrisHit)
+{
+#if DRAW_OVERLAPPING_TRIS
+	TArray<FOverlapResult> Overlaps;
+	DrawGeomOverlaps(World, Geom, QueryTM, Overlaps);
+
+	TArray<FBatchedLine> Lines;
+	const FLinearColor LineColor = FLinearColor(1.f, 0.7f, 0.7f);
+	const FLinearColor NormalColor = FLinearColor(1.f, 1.f, 1.f);
+	const float Lifetime = 5.f;
+#endif // DRAW_OVERLAPPING_TRIS
+
+	// Track the best triangle plane distance
+	float BestPlaneDist = -BIG_NUMBER;
+	FVector BestPlaneNormal(0, 0, 1);
+	FVector BestPointOnPlane(0, 0, 0);
+	// Iterate over triangles
+	for (int32 TriIdx = 0; TriIdx < NumTrisHit; TriIdx++)
+	{
+		PxTriangle Tri;
+		PxMeshQuery::getTriangle(ShapeGeom, PShapeWorldPose, HitTris[TriIdx], Tri);
+
+		const FVector A = P2UVector(Tri.verts[0]);
+		const FVector B = P2UVector(Tri.verts[1]);
+		const FVector C = P2UVector(Tri.verts[2]);
+
+		FVector TriNormal = ((B - A) ^ (C - A));
+
+		// Use a more accurate normalization that avoids InvSqrtEst
+		const float TriNormalSize = TriNormal.Size();
+		TriNormal = (TriNormalSize >= KINDA_SMALL_NUMBER ? TriNormal / TriNormalSize : FVector::ZeroVector);
+
+		const FPlane TriPlane(A, TriNormal);
+
+		const FVector QueryCenter = P2UVector(QueryTM.p);
+		const float DistToPlane = TriPlane.PlaneDot(QueryCenter);
+
+		if (DistToPlane > BestPlaneDist)
+		{
+			BestPlaneDist = DistToPlane;
+			BestPlaneNormal = TriNormal;
+			BestPointOnPlane = A;
+		}
+
+#if DRAW_OVERLAPPING_TRIS
+		Lines.Add(FBatchedLine(A, B, LineColor, Lifetime, 0.1f, SDPG_Foreground));
+		Lines.Add(FBatchedLine(B, C, LineColor, Lifetime, 0.1f, SDPG_Foreground));
+		Lines.Add(FBatchedLine(C, A, LineColor, Lifetime, 0.1f, SDPG_Foreground));
+		Lines.Add(FBatchedLine(A, A + (50.f*TriNormal), NormalColor, Lifetime, 0.1f, SDPG_Foreground));
+#endif // DRAW_OVERLAPPING_TRIS
+	}
+
+#if DRAW_OVERLAPPING_TRIS
+	if (World->PersistentLineBatcher)
+	{
+		World->PersistentLineBatcher->DrawLines(Lines);
+	}
+#endif // DRAW_OVERLAPPING_TRIS
+
+	return BestPlaneNormal;
+}
+
+
 /** Util to convert an overlapped shape into a sweep hit result, returns whether it was a blocking hit. */
 static bool ConvertOverlappedShapeToImpactHit(const PxLocationHit& PHit, const FVector& StartLoc, const FVector& EndLoc, FHitResult& OutResult, const PxGeometry& Geom, const PxTransform& QueryTM, const PxFilterData& QueryFilter, bool bReturnPhysMat)
 {
@@ -515,6 +580,7 @@ static bool ConvertOverlappedShapeToImpactHit(const PxLocationHit& PHit, const F
 	const PxTransform PShapeWorldPose = PxShapeExt::getGlobalPose(*PShape, *PActor); 
 
 	PxTriangleMeshGeometry PTriMeshGeom;
+	PxHeightFieldGeometry PHeightfieldGeom;
 	if(PShape->getTriangleMeshGeometry(PTriMeshGeom))
 	{
 		PxU32 HitTris[64];
@@ -522,64 +588,7 @@ static bool ConvertOverlappedShapeToImpactHit(const PxLocationHit& PHit, const F
 		const int32 NumTrisHit = PxMeshQuery::findOverlapTriangleMesh(Geom, QueryTM, PTriMeshGeom, PShapeWorldPose, HitTris, 64, 0, bOverflow);
 		if (NumTrisHit > 0)
 		{
-			#if DRAW_OVERLAPPING_TRIS
-			TArray<FOverlapResult> Overlaps;
-			DrawGeomOverlaps(World, Geom, QueryTM, Overlaps);
-
-			TArray<FBatchedLine> Lines;
-			const FLinearColor LineColor = FLinearColor(1.f,0.7f,0.7f);
-			const FLinearColor NormalColor = FLinearColor(1.f,1.f,1.f);
-			const float Lifetime = 5.f;
-			#endif // DRAW_OVERLAPPING_TRIS
-
-			// Track the best triangle plane distance
-			float BestPlaneDist = -BIG_NUMBER;
-			FVector BestPlaneNormal(0,0,1);
-			FVector BestPointOnPlane(0,0,0);
-			// Iterate over triangles
-			for(int32 TriIdx = 0; TriIdx<NumTrisHit; TriIdx++)
-			{
-				PxTriangle Tri;
-				PxMeshQuery::getTriangle(PTriMeshGeom, PShapeWorldPose, HitTris[TriIdx], Tri);
-
-				const FVector A = P2UVector(Tri.verts[0]);
-				const FVector B = P2UVector(Tri.verts[1]);
-				const FVector C = P2UVector(Tri.verts[2]);
-
-				FVector TriNormal = ((B-A) ^ (C-A));
-
-				// Use a more accurate normalization that avoids InvSqrtEst
-				const float TriNormalSize = TriNormal.Size();
-				TriNormal = (TriNormalSize >= KINDA_SMALL_NUMBER ? TriNormal/TriNormalSize : FVector::ZeroVector);
-
-				const FPlane TriPlane(A, TriNormal);
-
-				const FVector QueryCenter = P2UVector(QueryTM.p);
-				const float DistToPlane = TriPlane.PlaneDot(QueryCenter);
-
-				if(DistToPlane > BestPlaneDist)
-				{
-					BestPlaneDist = DistToPlane;
-					BestPlaneNormal = TriNormal;
-					BestPointOnPlane = A;
-				}
-
-				#if DRAW_OVERLAPPING_TRIS
-				Lines.Add(FBatchedLine(A, B, LineColor, Lifetime, 0.1f, SDPG_Foreground));
-				Lines.Add(FBatchedLine(B, C, LineColor, Lifetime, 0.1f, SDPG_Foreground));
-				Lines.Add(FBatchedLine(C, A, LineColor, Lifetime, 0.1f, SDPG_Foreground));
-				Lines.Add(FBatchedLine(A, A+(50.f*TriNormal), NormalColor, Lifetime, 0.1f, SDPG_Foreground));
-				#endif // DRAW_OVERLAPPING_TRIS
-			}
-
-			#if DRAW_OVERLAPPING_TRIS
-			if ( World->PersistentLineBatcher )
-			{
-				World->PersistentLineBatcher->DrawLines(Lines);
-			}
-			#endif // DRAW_OVERLAPPING_TRIS
-
-			OutResult.ImpactNormal = BestPlaneNormal;
+			OutResult.ImpactNormal = FindBestOverlappingNormal(Geom, QueryTM, PTriMeshGeom, PShapeWorldPose, HitTris, NumTrisHit);
 		}
 		else
 		{
@@ -593,6 +602,22 @@ static bool ConvertOverlappedShapeToImpactHit(const PxLocationHit& PHit, const F
 		}
 	}
 	else
+	if (PShape->getHeightFieldGeometry(PHeightfieldGeom))
+	{
+		PxU32 HitTris[64];
+		bool bOverflow = false;
+		const int32 NumTrisHit = PxMeshQuery::findOverlapHeightField(Geom, QueryTM, PHeightfieldGeom, PShapeWorldPose, HitTris, 64, 0, bOverflow);
+		if (NumTrisHit > 0)
+		{
+			OutResult.ImpactNormal = FindBestOverlappingNormal(Geom, QueryTM, PHeightfieldGeom, PShapeWorldPose, HitTris, NumTrisHit);
+		}
+		else
+		{
+			// PxMeshQuery::findOverlapHeightField failed.
+			OutResult.ImpactNormal = (PHit.flags & PxHitFlag::eNORMAL) ? P2UVector(PHit.normal) : FVector(0.f, 0.f, 1.f);
+		}
+	}
+	else
 	{
 		// Non tri-mesh
 		// use vector center of shape to query as good direction to move in
