Fix hit box behaviors when using touch and mouse faking touch

diff --git a/Engine/Source/Runtime/Engine/Classes/GameFramework/HUD.h b/Engine/Source/Runtime/Engine/Classes/GameFramework/HUD.h
--- a/Engine/Source/Runtime/Engine/Classes/GameFramework/HUD.h
+++ b/Engine/Source/Runtime/Engine/Classes/GameFramework/HUD.h
@@ -579,7 +579,7 @@ public:
 	 * Update the list of hitboxes and dispatch events for any hits.
 	 * @param	InEventType	Input event that triggered the call.
 	 */
-	bool UpdateAndDispatchHitBoxClickEvents(const FVector2D ClickLocation, const EInputEvent InEventType, const bool bDispatchOverOutEvent);
+	bool UpdateAndDispatchHitBoxClickEvents(const FVector2D ClickLocation, const EInputEvent InEventType, const bool bIsTouchEvent);
 	
 	/**
 	 * Update a the list of hitboxes that have been hit this frame.
diff --git a/Engine/Source/Runtime/Engine/Private/HUD.cpp b/Engine/Source/Runtime/Engine/Private/HUD.cpp
--- a/Engine/Source/Runtime/Engine/Private/HUD.cpp
+++ b/Engine/Source/Runtime/Engine/Private/HUD.cpp
@@ -10,6 +10,7 @@
 #include "MessageLog.h"
 #include "UObjectToken.h"
 #include "DisplayDebugHelpers.h"
+#include "Slate.h"
 
 DEFINE_LOG_CATEGORY_STATIC(LogHUD, Log, All);
 
@@ -132,7 +133,12 @@ void AHUD::PostRender()
 		if (LocalPlayer)
 		{
 			TArray<FVector2D> ContactPoints;
-			ContactPoints.Add(LocalPlayer->ViewportClient->GetMousePosition());
+
+			// TODO: This needs a check whether the platform has a mouse at all, but there is no way to currently do that
+			if (!FSlateApplication::Get().IsFakingTouchEvents())
+			{
+				ContactPoints.Add(LocalPlayer->ViewportClient->GetMousePosition());
+			}
 
 			for (int32 FingerIndex = 0; FingerIndex < EKeys::NUM_TOUCH_KEYS; ++FingerIndex)
 			{
@@ -867,37 +873,52 @@ const TArray<FIntRect>& AHUD::GetUIBlurRectangles() const
 	return UIBlurOverrideRectangles;
 }
 
-bool AHUD::UpdateAndDispatchHitBoxClickEvents(FVector2D ClickLocation, const EInputEvent InEventType, const bool bDispatchOverOutEvent)
+bool AHUD::UpdateAndDispatchHitBoxClickEvents(FVector2D ClickLocation, const EInputEvent InEventType, const bool bIsTouchEvent)
 {
 	ClickLocation += GetCoordinateOffset();
 
+	const bool bIsClickEvent = (InEventType == IE_Pressed || InEventType == IE_DoubleClick);
 	bool bHit = false;
-	for (FHUDHitBox* HitBoxHit : HitBoxHits)
+
+	// If this is a touch event we likely don't have the hit box in the hit list yet so we need to check all HitBoxes
+	if (bIsTouchEvent && bIsClickEvent)
 	{
-		if (HitBoxHit->Contains(ClickLocation))
+		for (FHUDHitBox& HitBox : HitBoxMap)
 		{
-			bHit = true;
-
-			if (InEventType == IE_Pressed || InEventType == IE_DoubleClick)
+			if (HitBox.Contains(ClickLocation))
 			{
-				ReceiveHitBoxClick(HitBoxHit->GetName());
-				if (bDispatchOverOutEvent)
+				bHit = true;
+
+				ReceiveHitBoxClick(HitBox.GetName());
+
+				if (HitBox.ConsumesInput())
 				{
-					ReceiveHitBoxBeginCursorOver(HitBoxHit->GetName());
+					break;	//Early out if this box consumed the click
 				}
 			}
-			else if (InEventType == IE_Released)
+		}
+	}
+	else
+	{
+		for (FHUDHitBox* HitBoxHit : HitBoxHits)
+		{
+			if (HitBoxHit->Contains(ClickLocation))
 			{
-				ReceiveHitBoxRelease(HitBoxHit->GetName());
-				if (bDispatchOverOutEvent)
+				bHit = true;
+
+				if (bIsClickEvent)
 				{
-					ReceiveHitBoxEndCursorOver(HitBoxHit->GetName());
+					ReceiveHitBoxClick(HitBoxHit->GetName());
+				}
+				else if (InEventType == IE_Released)
+				{
+					ReceiveHitBoxRelease(HitBoxHit->GetName());
 				}
-			}
 
-			if (HitBoxHit->ConsumesInput() == true)
-			{
-				break;	//Early out if this box consumed the click
+				if (HitBoxHit->ConsumesInput() == true)
+				{
+					break;	//Early out if this box consumed the click
+				}
 			}
 		}
 	}
