Initial support for other foliage types

diff --git a/Engine/Source/Runtime/Engine/Classes/Engine/StaticMesh.h b/Engine/Source/Runtime/Engine/Classes/Engine/StaticMesh.h
--- a/Engine/Source/Runtime/Engine/Classes/Engine/StaticMesh.h
+++ b/Engine/Source/Runtime/Engine/Classes/Engine/StaticMesh.h
@@ -8,6 +8,9 @@
 /** The maximum number of static mesh LODs allowed. */
 #define MAX_STATIC_MESH_LODS 4
 
+// Forward declarations
+class UFoliageType_InstancedStaticMesh;
+
 /*-----------------------------------------------------------------------------
 	Legacy mesh optimization settings.
 -----------------------------------------------------------------------------*/
@@ -90,71 +93,71 @@ struct FStaticMeshOptimizationSettings
 	uint8 ShadingImportance;
 
 
-		FStaticMeshOptimizationSettings()
-		: ReductionMethod( OT_MaxDeviation )
-		, NumOfTrianglesPercentage( 1.0f )
-		, MaxDeviationPercentage( 0.0f )
-		, WeldingThreshold( 0.1f )
-		, bRecalcNormals( true )
-		, NormalsThreshold( 60.0f )
-		, SilhouetteImportance( IL_Normal )
-		, TextureImportance( IL_Normal )
-		, ShadingImportance( IL_Normal )
-		{
-		}
+	FStaticMeshOptimizationSettings()
+	: ReductionMethod( OT_MaxDeviation )
+	, NumOfTrianglesPercentage( 1.0f )
+	, MaxDeviationPercentage( 0.0f )
+	, WeldingThreshold( 0.1f )
+	, bRecalcNormals( true )
+	, NormalsThreshold( 60.0f )
+	, SilhouetteImportance( IL_Normal )
+	, TextureImportance( IL_Normal )
+	, ShadingImportance( IL_Normal )
+	{
+	}
 
-		/** Serialization for FStaticMeshOptimizationSettings. */
-		inline friend FArchive& operator<<( FArchive& Ar, FStaticMeshOptimizationSettings& Settings )
+	/** Serialization for FStaticMeshOptimizationSettings. */
+	inline friend FArchive& operator<<( FArchive& Ar, FStaticMeshOptimizationSettings& Settings )
+	{
+		if( Ar.UE4Ver() < VER_UE4_ADDED_EXTRA_MESH_OPTIMIZATION_SETTINGS )
 		{
-			if( Ar.UE4Ver() < VER_UE4_ADDED_EXTRA_MESH_OPTIMIZATION_SETTINGS )
+			Ar << Settings.MaxDeviationPercentage;
+
+			//Remap Importance Settings
+			Ar << Settings.SilhouetteImportance;
+			Ar << Settings.TextureImportance;
+
+			//IL_Normal was previously the first enum value. We add the new index of IL_Normal to correctly offset the old values.
+			Settings.SilhouetteImportance += IL_Normal;
+			Settings.TextureImportance += IL_Normal;
+
+			//Remap NormalMode enum value to new threshold variable.
+			uint8 NormalMode;
+			Ar << NormalMode;
+
+			const float NormalThresholdTable[] =
+			{
+				60.0f, // Recompute
+				80.0f, // Recompute (Smooth)
+				45.0f  // Recompute (Hard)
+			};
+
+			if( NormalMode == NM_PreserveSmoothingGroups)
 			{
-				Ar << Settings.MaxDeviationPercentage;
-
-				//Remap Importance Settings
-				Ar << Settings.SilhouetteImportance;
-				Ar << Settings.TextureImportance;
-
-				//IL_Normal was previously the first enum value. We add the new index of IL_Normal to correctly offset the old values.
-				Settings.SilhouetteImportance += IL_Normal;
-				Settings.TextureImportance += IL_Normal;
-
-				//Remap NormalMode enum value to new threshold variable.
-				uint8 NormalMode;
-				Ar << NormalMode;
-
-				const float NormalThresholdTable[] =
-				{
-					60.0f, // Recompute
-					80.0f, // Recompute (Smooth)
-					45.0f  // Recompute (Hard)
-				};
-
-				if( NormalMode == NM_PreserveSmoothingGroups)
-				{
-					Settings.bRecalcNormals = false;
-				}
-				else
-				{
-					Settings.bRecalcNormals = true;
-					Settings.NormalsThreshold = NormalThresholdTable[NormalMode];
-				}
+				Settings.bRecalcNormals = false;
 			}
 			else
 			{
-				Ar << Settings.ReductionMethod;
+				Settings.bRecalcNormals = true;
+				Settings.NormalsThreshold = NormalThresholdTable[NormalMode];
+			}
+		}
+		else
+		{
+			Ar << Settings.ReductionMethod;
 			Ar << Settings.MaxDeviationPercentage;
-				Ar << Settings.NumOfTrianglesPercentage;
+			Ar << Settings.NumOfTrianglesPercentage;
 			Ar << Settings.SilhouetteImportance;
 			Ar << Settings.TextureImportance;
-				Ar << Settings.ShadingImportance;
-				Ar << Settings.bRecalcNormals;
-				Ar << Settings.NormalsThreshold;
-				Ar << Settings.WeldingThreshold;
-			}
-
-			return Ar;
+			Ar << Settings.ShadingImportance;
+			Ar << Settings.bRecalcNormals;
+			Ar << Settings.NormalsThreshold;
+			Ar << Settings.WeldingThreshold;
 		}
 
+		return Ar;
+	}
+
 };
 
 /*-----------------------------------------------------------------------------
@@ -400,7 +403,7 @@ class UStaticMesh : public UObject, public IInterface_CollisionDataProvider, pub
 #if WITH_EDITORONLY_DATA
 	/** Default settings when using this mesh for instanced foliage */
 	UPROPERTY(EditAnywhere, AdvancedDisplay, editinline, Category=StaticMesh)
-	class UInstancedFoliageSettings* FoliageDefaultSettings;
+	UFoliageType_InstancedStaticMesh* FoliageDefaultSettings;
 
 	/** Importing data and options used for this mesh */
 	UPROPERTY(VisibleAnywhere, editinline, Category=ImportSettings)
diff --git a/Engine/Source/Runtime/Engine/Classes/Foliage/FoliageType.h b/Engine/Source/Runtime/Engine/Classes/Foliage/FoliageType.h
new file mode 100644
--- /dev/null
+++ b/Engine/Source/Runtime/Engine/Classes/Foliage/FoliageType.h
@@ -0,0 +1,223 @@
+// Copyright 1998-2014 Epic Games, Inc. All Rights Reserved.
+
+#pragma once
+
+#include "FoliageType.generated.h"
+
+UENUM()
+enum FoliageVertexColorMask
+{
+	FOLIAGEVERTEXCOLORMASK_Disabled,
+	FOLIAGEVERTEXCOLORMASK_Red,
+	FOLIAGEVERTEXCOLORMASK_Green,
+	FOLIAGEVERTEXCOLORMASK_Blue,
+	FOLIAGEVERTEXCOLORMASK_Alpha,
+};
+
+
+UCLASS(hidecategories=Object, editinlinenew, MinimalAPI)
+class UFoliageType : public UObject
+{
+	GENERATED_UCLASS_BODY()
+
+	UPROPERTY(EditAnywhere, Category=Painting)
+	float Density;
+
+	UPROPERTY(EditAnywhere, Category=Painting)
+	float Radius;
+
+	UPROPERTY(EditAnywhere, Category=Painting)
+	float ScaleMinX;
+
+	UPROPERTY(EditAnywhere, Category=Painting)
+	float ScaleMinY;
+
+	UPROPERTY(EditAnywhere, Category=Painting)
+	float ScaleMinZ;
+
+	UPROPERTY(EditAnywhere, Category=Painting)
+	float ScaleMaxX;
+
+	UPROPERTY(EditAnywhere, Category=Painting)
+	float ScaleMaxY;
+
+	UPROPERTY(EditAnywhere, Category=Painting)
+	float ScaleMaxZ;
+
+	UPROPERTY(EditAnywhere, Category=Painting)
+	uint32 LockScaleX:1;
+
+	UPROPERTY(EditAnywhere, Category=Painting)
+	uint32 LockScaleY:1;
+
+	UPROPERTY(EditAnywhere, Category=Painting)
+	uint32 LockScaleZ:1;
+
+	UPROPERTY(EditAnywhere, Category=Painting)
+	float AlignMaxAngle;
+
+	UPROPERTY(EditAnywhere, Category=Painting)
+	float RandomPitchAngle;
+
+	UPROPERTY(EditAnywhere, Category=Painting)
+	float GroundSlope;
+
+	UPROPERTY(EditAnywhere, Category=Painting)
+	float HeightMin;
+
+	UPROPERTY(EditAnywhere, Category=Painting)
+	float HeightMax;
+
+	UPROPERTY(EditAnywhere, Category=Painting)
+	FName LandscapeLayer;
+
+	UPROPERTY(EditAnywhere, Category=Painting)
+	uint32 AlignToNormal:1;
+
+	UPROPERTY(EditAnywhere, Category=Painting)
+	uint32 RandomYaw:1;
+
+	UPROPERTY(EditAnywhere, Category=Painting)
+	uint32 UniformScale:1;
+
+	UPROPERTY(EditAnywhere, Category=Painting)
+	float ZOffsetMin;
+
+	UPROPERTY(EditAnywhere, Category=Painting)
+	float ZOffsetMax;
+
+	UPROPERTY(EditAnywhere, Category=Painting)
+	uint32 CollisionWithWorld:1;
+
+	UPROPERTY(EditAnywhere, Category=Painting)
+	FVector CollisionScale;
+
+	UPROPERTY()
+	FBoxSphereBounds MeshBounds;
+
+	// X, Y is origin position and Z is radius...
+	UPROPERTY()
+	FVector LowBoundOriginRadius;
+
+	UPROPERTY(EditAnywhere, Category=Painting)
+	TEnumAsByte<enum FoliageVertexColorMask> VertexColorMask;
+
+	UPROPERTY(EditAnywhere, Category=Painting)
+	float VertexColorMaskThreshold;
+
+	UPROPERTY(EditAnywhere, Category=Painting)
+	uint32 VertexColorMaskInvert:1;
+
+	UPROPERTY(EditAnywhere, Category=Clustering)
+	int32 MaxInstancesPerCluster;
+
+	UPROPERTY(EditAnywhere, Category=Clustering)
+	float MaxClusterRadius;
+
+	UPROPERTY()
+	float ReapplyDensityAmount;
+
+	UPROPERTY()
+	uint32 ReapplyDensity:1;
+
+	UPROPERTY()
+	uint32 ReapplyRadius:1;
+
+	UPROPERTY()
+	uint32 ReapplyAlignToNormal:1;
+
+	UPROPERTY()
+	uint32 ReapplyRandomYaw:1;
+
+	UPROPERTY()
+	uint32 ReapplyScaleX:1;
+
+	UPROPERTY()
+	uint32 ReapplyScaleY:1;
+
+	UPROPERTY()
+	uint32 ReapplyScaleZ:1;
+
+	UPROPERTY()
+	uint32 ReapplyRandomPitchAngle:1;
+
+	UPROPERTY()
+	uint32 ReapplyGroundSlope:1;
+
+	UPROPERTY()
+	uint32 ReapplyHeight:1;
+
+	UPROPERTY()
+	uint32 ReapplyLandscapeLayer:1;
+
+	UPROPERTY()
+	uint32 ReapplyZOffset:1;
+
+	UPROPERTY()
+	uint32 ReapplyCollisionWithWorld:1;
+
+	UPROPERTY()
+	uint32 ReapplyVertexColorMask:1;
+
+	UPROPERTY(EditAnywhere, Category=Culling)
+	int32 StartCullDistance;
+
+	UPROPERTY(EditAnywhere, Category=Culling)
+	int32 EndCullDistance;
+
+	UPROPERTY()
+	int32 DisplayOrder;
+
+	UPROPERTY()
+	uint32 IsSelected:1;
+
+	UPROPERTY()
+	uint32 ShowNothing:1;
+
+	UPROPERTY()
+	uint32 ShowPaintSettings:1;
+
+	UPROPERTY()
+	uint32 ShowInstanceSettings:1;
+
+	/** Controls whether the foliage should cast a shadow or not. **/
+	UPROPERTY(EditAnywhere, BlueprintReadOnly, Category=Lighting)
+	uint32 CastShadow:1;
+
+	/** Controls whether the foliage should inject light into the Light Propagation Volume.  This flag is only used if CastShadow is true. **/
+	UPROPERTY(EditAnywhere, BlueprintReadOnly, Category=Lighting)
+	uint32 bAffectDynamicIndirectLighting:1;
+
+	/** Controls whether the foliage should cast shadows in the case of non precomputed shadowing.  This flag is only used if CastShadow is true. **/
+	UPROPERTY(EditAnywhere, BlueprintReadOnly, Category=Lighting)
+	uint32 bCastDynamicShadow:1;
+
+	/** Whether the foliage should cast a static shadow from shadow casting lights.  This flag is only used if CastShadow is true. */
+	UPROPERTY(EditAnywhere, BlueprintReadOnly, Category=Lighting)
+	uint32 bCastStaticShadow:1;
+
+	/** 
+	 *	If true, the foliage will cast shadows even if bHidden is true.
+	 *	Controls whether the primitive should cast shadows when hidden.
+	 *	This flag is only used if CastShadow is true.
+	 */
+	UPROPERTY(EditAnywhere, AdvancedDisplay, BlueprintReadOnly, Category=Lighting)
+	uint32 bCastHiddenShadow:1;
+
+	/** Whether this foliage should cast dynamic shadows as if it were a two sided material. */
+	UPROPERTY(EditAnywhere, AdvancedDisplay, BlueprintReadOnly, Category=Lighting)
+	uint32 bCastShadowAsTwoSided:1;
+
+	/** Whether the foliage receives decals. */
+	UPROPERTY(EditAnywhere, AdvancedDisplay, BlueprintReadOnly, Category=Rendering)
+	uint32 bReceivesDecals : 1;
+
+	/** Custom collision for foliage */
+	UPROPERTY(EditAnywhere, Category=Collision, meta=(HideObjectType=true))
+	struct FBodyInstance BodyInstance;
+
+	// Gets the mesh associated with this FoliageType
+	virtual UStaticMesh* GetStaticMesh() PURE_VIRTUAL(UFoliageType::GetStaticMesh, return nullptr; );
+};
+
+
diff --git a/Engine/Source/Runtime/Engine/Classes/Foliage/InstancedFoliageActor.h b/Engine/Source/Runtime/Engine/Classes/Foliage/InstancedFoliageActor.h
--- a/Engine/Source/Runtime/Engine/Classes/Foliage/InstancedFoliageActor.h
+++ b/Engine/Source/Runtime/Engine/Classes/Foliage/InstancedFoliageActor.h
@@ -7,6 +7,10 @@
 #include "InstancedFoliage.h"
 #include "InstancedFoliageActor.generated.h"
 
+// Forward declarations
+class UFoliageType;
+class ULandscapeHeightfieldCollisionComponent;
+struct FFoliageInstancePlacementInfo;
 
 UCLASS(notplaceable, hidecategories = Object, MinimalAPI, NotBlueprintable)
 class AInstancedFoliageActor : public AActor
@@ -15,10 +19,10 @@ class AInstancedFoliageActor : public AActor
 
 	/** The static mesh type that will be used to show the widget */
 	UPROPERTY(transient)
-	UStaticMesh* SelectedMesh;
+	UFoliageType* SelectedMesh;
 
 public:
-	TMap<UStaticMesh*, TUniqueObj<FFoliageMeshInfo>> FoliageMeshes;
+	TMap<UFoliageType*, TUniqueObj<FFoliageMeshInfo>> FoliageMeshes;
 
 public:
 	// Begin UObject interface.
@@ -40,41 +44,45 @@ public:
 	ENGINE_API void MapRebuild();
 
 	// Called from editor code to manage instances for components
-	ENGINE_API void SnapInstancesForLandscape( class ULandscapeHeightfieldCollisionComponent* InComponent, const FBox& InInstanceBox );
+	ENGINE_API void SnapInstancesForLandscape(ULandscapeHeightfieldCollisionComponent* InComponent, const FBox& InInstanceBox);
 
 	// Moves instances based on the specified component to the current streaming level
-	ENGINE_API void MoveInstancesForComponentToCurrentLevel( class UActorComponent* InComponent );
+	ENGINE_API void MoveInstancesForComponentToCurrentLevel(UActorComponent* InComponent);
 
 	// Change all instances based on one component to a new component (possible in another level).
 	// The instances keep the same world locations
-	ENGINE_API void MoveInstancesToNewComponent(class UActorComponent* InOldComponent, class UActorComponent* InNewComponent);
+	ENGINE_API void MoveInstancesToNewComponent(UActorComponent* InOldComponent, UActorComponent* InNewComponent);
 
 	// Move instances based on a component that has just been moved.
-	void MoveInstancesForMovedComponent( class UActorComponent* InComponent );
+	void MoveInstancesForMovedComponent(UActorComponent* InComponent);
 
 	// Returns a map of Static Meshes and their placed instances attached to a component.
-	ENGINE_API TMap<class UStaticMesh*,TArray<const struct FFoliageInstancePlacementInfo*> > GetInstancesForComponent( class UActorComponent* InComponent );
+	ENGINE_API TMap<UFoliageType*, TArray<const FFoliageInstancePlacementInfo*>> GetInstancesForComponent(UActorComponent* InComponent);
 
 	// Deletes the instances attached to a component
-	ENGINE_API void DeleteInstancesForComponent( class UActorComponent* InComponent );
+	ENGINE_API void DeleteInstancesForComponent(UActorComponent* InComponent);
 
 	// Finds a mesh entry
-	ENGINE_API FFoliageMeshInfo* FindMesh(UStaticMesh* InMesh);
+	ENGINE_API FFoliageMeshInfo* FindMesh(UFoliageType* InType);
 
 	// Finds a mesh entry or adds it if it doesn't already exist
-	ENGINE_API FFoliageMeshInfo* FindOrAddMesh(UStaticMesh* InMesh);
+	ENGINE_API FFoliageMeshInfo* FindOrAddMesh(UFoliageType* InType);
 
 	// Add a new static mesh.
-	ENGINE_API FFoliageMeshInfo* AddMesh(UStaticMesh* InMesh);
+	ENGINE_API FFoliageMeshInfo* AddMesh(UStaticMesh* InMesh, UFoliageType** OutSettings = nullptr);
+	ENGINE_API FFoliageMeshInfo* AddMesh(UFoliageType* InType);
 
 	// Remove the static mesh from the mesh list, and all its instances.
-	ENGINE_API void RemoveMesh( class UStaticMesh* InMesh );
+	ENGINE_API void RemoveMesh(UFoliageType* InFoliageType);
+
+	// Performs a reverse lookup from a mesh to its settings
+	ENGINE_API UFoliageType* GetSettingsForMesh(UStaticMesh* InMesh, FFoliageMeshInfo** OutMeshInfo = nullptr);
 
 	// Select an individual instance.
-	ENGINE_API void SelectInstance( class UInstancedStaticMeshComponent* InComponent, int32 InComponentInstanceIndex, bool bToggle );
+	ENGINE_API void SelectInstance(UInstancedStaticMeshComponent* InComponent, int32 InComponentInstanceIndex, bool bToggle);
 
 	// Propagate the selected instances to the actual render components
-	ENGINE_API void ApplySelectionToComponents( bool bApply );
+	ENGINE_API void ApplySelectionToComponents(bool bApply);
 
 	// Updates the SelectedMesh property of the actor based on the actual selected instances
 	ENGINE_API void CheckSelection();
@@ -83,14 +91,14 @@ public:
 	ENGINE_API FVector GetSelectionLocation();
 
 	/*
- 	* Get the instanced foliage actor for the current streaming level.
- 	*
- 	* @param InCreationWorldIfNone			World to create the foliage instance in
+	* Get the instanced foliage actor for the current streaming level.
+	*
+	* @param InCreationWorldIfNone			World to create the foliage instance in
 	* @param bCreateIfNone					Create if doesnt already exist
- 	* returns								pointer to foliage object instance
- 	*/
- 	static ENGINE_API AInstancedFoliageActor* GetInstancedFoliageActor(UWorld* InWorld,bool bCreateIfNone=true);
-	
+	* returns								pointer to foliage object instance
+	*/
+	static ENGINE_API AInstancedFoliageActor* GetInstancedFoliageActorForCurrentLevel(UWorld* InWorld, bool bCreateIfNone = true);
+
 
 	// Get the instanced foliage actor for the specified streaming level. Never creates a new IFA.
 	static ENGINE_API AInstancedFoliageActor* GetInstancedFoliageActorForLevel(ULevel* Level);
diff --git a/Engine/Source/Runtime/Engine/Classes/Landscape/LandscapeProxy.h b/Engine/Source/Runtime/Engine/Classes/Landscape/LandscapeProxy.h
--- a/Engine/Source/Runtime/Engine/Classes/Landscape/LandscapeProxy.h
+++ b/Engine/Source/Runtime/Engine/Classes/Landscape/LandscapeProxy.h
@@ -153,6 +153,9 @@ struct FLandscapeImportLayerInfo
 #endif
 };
 
+// this is only here because putting it in LandscapeEditorObject.h (where it belongs)
+// results in Engine being dependent on LandscapeEditor, as the actual landscape editing
+// code (e.g. LandscapeEdit.h) is in /Engine/ for some reason...
 UENUM()
 namespace ELandscapeLayerPaintingRestriction
 {
diff --git a/Engine/Source/Runtime/Engine/Private/InstancedFoliage.cpp b/Engine/Source/Runtime/Engine/Private/InstancedFoliage.cpp
--- a/Engine/Source/Runtime/Engine/Private/InstancedFoliage.cpp
+++ b/Engine/Source/Runtime/Engine/Private/InstancedFoliage.cpp
@@ -10,7 +10,7 @@
 #include "UObjectToken.h"
 #include "MapErrors.h"
 #include "Foliage/InstancedFoliageActor.h"
-#include "Foliage/InstancedFoliageSettings.h"
+#include "Foliage/FoliageType_InstancedStaticMesh.h"
 #include "Landscape/LandscapeHeightfieldCollisionComponent.h"
 
 #define LOCTEXT_NAMESPACE "InstancedFoliage"
@@ -45,8 +45,13 @@ FArchive& operator<<(FArchive& Ar, FFoliageInstanceCluster& Cluster)
 	Ar << Cluster.Bounds;
 	Ar << Cluster.ClusterComponent;
 
-	//!! editor only
-	Ar << Cluster.InstanceIndices;
+#if WITH_EDITORONLY_DATA
+	if (!Ar.ArIsFilterEditorOnly ||
+		Ar.UE4Ver() < VER_UE4_FOLIAGE_SETTINGS_TYPE)
+	{
+		Ar << Cluster.InstanceIndices;
+	}
+#endif
 
 	return Ar;
 }
@@ -54,8 +59,13 @@ FArchive& operator<<(FArchive& Ar, FFoliageInstanceCluster& Cluster)
 FArchive& operator<<(FArchive& Ar, FFoliageMeshInfo& MeshInfo)
 {
 	Ar << MeshInfo.InstanceClusters;
-	//!! editor only
-	Ar << MeshInfo.Instances;
+
+#if WITH_EDITORONLY_DATA
+	if (!Ar.ArIsFilterEditorOnly ||
+		Ar.UE4Ver() < VER_UE4_FOLIAGE_SETTINGS_TYPE)
+	{
+		Ar << MeshInfo.Instances;
+	}
 
 	// Serialize the transient data for undo.
 	if (Ar.IsTransacting())
@@ -65,8 +75,7 @@ FArchive& operator<<(FArchive& Ar, FFoliageMeshInfo& MeshInfo)
 		Ar << MeshInfo.FreeInstanceIndices;
 		Ar << MeshInfo.SelectedIndices;
 	}
-
-	Ar << MeshInfo.Settings;
+#endif
 
 	return Ar;
 }
@@ -77,23 +86,12 @@ FArchive& operator<<(FArchive& Ar, FFoliageComponentHashInfo& ComponentHashInfo)
 }
 
 //
-// UInstancedFoliageSettings
+// UFoliageType
 //
 
-UInstancedFoliageSettings::UInstancedFoliageSettings(const class FPostConstructInitializeProperties& PCIP)
+UFoliageType::UFoliageType(const FPostConstructInitializeProperties& PCIP)
 	: Super(PCIP)
 {
-	// Structure to hold one-time initialization
-	struct FConstructorStatics
-	{
-		FName NAME_None;
-		FConstructorStatics()
-			: NAME_None(TEXT("None"))
-		{
-		}
-	};
-	static FConstructorStatics ConstructorStatics;
-
 	Density = 100.0f;
 	Radius = 0.0f;
 	AlignToNormal = true;
@@ -112,7 +110,7 @@ UInstancedFoliageSettings::UInstancedFoliageSettings(const class FPostConstructI
 	HeightMax = 262144.0f;
 	ZOffsetMin = 0.0f;
 	ZOffsetMax = 0.0f;
-	LandscapeLayer = ConstructorStatics.NAME_None;
+	LandscapeLayer = NAME_None;
 	MaxInstancesPerCluster = 100;
 	MaxClusterRadius = 10000.0f;
 	DisplayOrder = 0;
@@ -137,60 +135,44 @@ UInstancedFoliageSettings::UInstancedFoliageSettings(const class FPostConstructI
 	BodyInstance.SetCollisionProfileName(UCollisionProfile::NoCollision_ProfileName);
 }
 
+UFoliageType_InstancedStaticMesh::UFoliageType_InstancedStaticMesh(const FPostConstructInitializeProperties& PCIP)
+	: Super(PCIP)
+{
+	Mesh = nullptr;
+}
+
 //
 // FFoliageMeshInfo
 //
 FFoliageMeshInfo::FFoliageMeshInfo()
-	: InstanceHash(NULL)
-	, Settings(NULL)
-{
-	if (GIsEditor)
-	{
-		InstanceHash = new FFoliageInstanceHash();
-	}
-}
-
-FFoliageMeshInfo::FFoliageMeshInfo(const FFoliageMeshInfo& Other)
-	: InstanceHash(NULL)
-	, Settings(NULL)
-{
-	if (GIsEditor)
-	{
-		InstanceHash = new FFoliageInstanceHash();
-	}
-}
+#if WITH_EDITOR
+	: InstanceHash(GIsEditor ? new FFoliageInstanceHash() : nullptr)
+#endif
+{ }
 
-FFoliageMeshInfo::~FFoliageMeshInfo()
-{
-	if (GIsEditor)
-	{
-		delete InstanceHash;
-	}
-}
 
 #if WITH_EDITOR
 
-
 void FFoliageMeshInfo::CheckValid()
 {
 #if DO_FOLIAGE_CHECK
-	int32 ClusterTotal=0;
-	int32 ComponentTotal=0;
+	int32 ClusterTotal = 0;
+	int32 ComponentTotal = 0;
 
-	for( int32 ClusterIdx=0;ClusterIdx<InstanceClusters.Num();ClusterIdx++ )
+	for (FFoliageInstanceCluster& Cluster : InstanceClusters)
 	{
-		check( InstanceClusters[ClusterIdx].ClusterComponent != NULL);
-		ClusterTotal += InstanceClusters[ClusterIdx].InstanceIndices.Num();
-		ComponentTotal += InstanceClusters[ClusterIdx].ClusterComponent->PerInstanceSMData.Num();
+		check(Cluster.ClusterComponent != nullptr);
+		ClusterTotal += Cluster.InstanceIndices.Num();
+		ComponentTotal += Cluster.ClusterComponent->PerInstanceSMData.Num();
 	}
 
-	check( ClusterTotal == ComponentTotal );
+	check(ClusterTotal == ComponentTotal);
 
 	int32 FreeTotal = 0;
 	int32 InstanceTotal = 0;
-	for( int32 InstanceIdx=0; InstanceIdx<Instances.Num(); InstanceIdx++ )
+	for (int32 InstanceIdx = 0; InstanceIdx < Instances.Num(); InstanceIdx++)
 	{
-		if( Instances[InstanceIdx].ClusterIndex != -1 )
+		if (Instances[InstanceIdx].ClusterIndex != -1)
 		{
 			InstanceTotal++;
 		}
@@ -206,7 +188,7 @@ void FFoliageMeshInfo::CheckValid()
 	InstanceHash->CheckInstanceCount(InstanceTotal);
 
 	int32 ComponentHashTotal = 0;
-	for( TMap<class UActorComponent*, FFoliageComponentHashInfo >::TConstIterator It(ComponentHash); It; ++It )
+	for( TMap<UActorComponent*, FFoliageComponentHashInfo >::TConstIterator It(ComponentHash); It; ++It )
 	{
 		ComponentHashTotal += It.Value().Instances.Num();
 	}
@@ -243,12 +225,12 @@ void FFoliageMeshInfo::CheckValid()
 #endif
 }
 
-void FFoliageMeshInfo::AddInstance(AInstancedFoliageActor* InIFA, UStaticMesh* InMesh, const FFoliageInstance& InNewInstance)
+void FFoliageMeshInfo::AddInstance(AInstancedFoliageActor* InIFA, UFoliageType* InSettings, const FFoliageInstance& InNewInstance)
 {
 	InIFA->Modify();
 
 	// Add the instance taking either a free slot or adding a new item.
-	int32 InstanceIndex = FreeInstanceIndices.Num() > 0 ? FreeInstanceIndices.Pop() : Instances.AddUninitialized();
+	int32 InstanceIndex = FreeInstanceIndices.Num() > 0 ?  FreeInstanceIndices.Pop() : Instances.AddUninitialized();
 
 	FFoliageInstance& AddedInstance = Instances[InstanceIndex];
 	AddedInstance = InNewInstance;
@@ -256,19 +238,19 @@ void FFoliageMeshInfo::AddInstance(AInstancedFoliageActor* InIFA, UStaticMesh* I
 	// Add the instance to the hash
 	InstanceHash->InsertInstance(InNewInstance.Location, InstanceIndex);
 	FFoliageComponentHashInfo* ComponentHashInfo = ComponentHash.Find(InNewInstance.Base);
-	if (ComponentHashInfo == NULL)
+	if (ComponentHashInfo == nullptr)
 	{
 		ComponentHashInfo = &ComponentHash.Add(InNewInstance.Base, FFoliageComponentHashInfo(InNewInstance.Base));
 	}
 	ComponentHashInfo->Instances.Add(InstanceIndex);
 
 	// Find the best cluster to allocate the instance to.
-	FFoliageInstanceCluster* BestCluster = NULL;
+	FFoliageInstanceCluster* BestCluster = nullptr;
 	int32 BestClusterIndex = INDEX_NONE;
 	float BestClusterDistSq = FLT_MAX;
 
-	int32 MaxInstancesPerCluster = Settings->MaxInstancesPerCluster;
-	float MaxClusterRadiusSq = FMath::Square(Settings->MaxClusterRadius);
+	int32 MaxInstancesPerCluster = InSettings->MaxInstancesPerCluster;
+	float MaxClusterRadiusSq = FMath::Square(InSettings->MaxClusterRadius);
 
 	for (int32 ClusterIdx = 0; ClusterIdx < InstanceClusters.Num(); ClusterIdx++)
 	{
@@ -288,33 +270,33 @@ void FFoliageMeshInfo::AddInstance(AInstancedFoliageActor* InIFA, UStaticMesh* I
 	// Calculate transform for the instance
 	FTransform InstanceToWorld = InNewInstance.GetInstanceWorldTransform();
 
-	if (BestCluster == NULL)
+	if (BestCluster == nullptr)
 	{
 		BestClusterIndex = InstanceClusters.Num();
 		BestCluster = new(InstanceClusters)FFoliageInstanceCluster(
 			ConstructObject<UInstancedStaticMeshComponent>(UInstancedStaticMeshComponent::StaticClass(), InIFA, NAME_None, RF_Transactional),
-			InMesh->GetBounds().TransformBy(InstanceToWorld)
+			InSettings->GetStaticMesh()->GetBounds().TransformBy(InstanceToWorld)
 			);
 
 		// Make the instanced static mesh component movable so it doesn't get statically lit; see the comment below about CastShadow for more details
 		BestCluster->ClusterComponent->Mobility = EComponentMobility::Movable;
 
-		BestCluster->ClusterComponent->StaticMesh = InMesh;
+		BestCluster->ClusterComponent->StaticMesh = InSettings->GetStaticMesh();
 		BestCluster->ClusterComponent->bSelectable = true;
 		BestCluster->ClusterComponent->bHasPerInstanceHitProxies = true;
 		BestCluster->ClusterComponent->InstancingRandomSeed = FMath::Rand();
-		BestCluster->ClusterComponent->InstanceStartCullDistance = Settings->StartCullDistance;
-		BestCluster->ClusterComponent->InstanceEndCullDistance = Settings->EndCullDistance;
+		BestCluster->ClusterComponent->InstanceStartCullDistance = InSettings->StartCullDistance;
+		BestCluster->ClusterComponent->InstanceEndCullDistance = InSettings->EndCullDistance;
 
-		BestCluster->ClusterComponent->CastShadow = Settings->CastShadow;
-		BestCluster->ClusterComponent->bCastDynamicShadow = Settings->bCastDynamicShadow;
-		BestCluster->ClusterComponent->bCastStaticShadow = Settings->bCastStaticShadow;
-		BestCluster->ClusterComponent->bAffectDynamicIndirectLighting = Settings->bAffectDynamicIndirectLighting;
-		BestCluster->ClusterComponent->bCastHiddenShadow = Settings->bCastHiddenShadow;
-		BestCluster->ClusterComponent->bCastShadowAsTwoSided = Settings->bCastShadowAsTwoSided;
-		BestCluster->ClusterComponent->bReceivesDecals = Settings->bReceivesDecals;
+		BestCluster->ClusterComponent->CastShadow = InSettings->CastShadow;
+		BestCluster->ClusterComponent->bCastDynamicShadow = InSettings->bCastDynamicShadow;
+		BestCluster->ClusterComponent->bCastStaticShadow = InSettings->bCastStaticShadow;
+		BestCluster->ClusterComponent->bAffectDynamicIndirectLighting = InSettings->bAffectDynamicIndirectLighting;
+		BestCluster->ClusterComponent->bCastHiddenShadow = InSettings->bCastHiddenShadow;
+		BestCluster->ClusterComponent->bCastShadowAsTwoSided = InSettings->bCastShadowAsTwoSided;
+		BestCluster->ClusterComponent->bReceivesDecals = InSettings->bReceivesDecals;
 
-		BestCluster->ClusterComponent->BodyInstance.CopyBodyInstancePropertiesFrom(&Settings->BodyInstance);
+		BestCluster->ClusterComponent->BodyInstance.CopyBodyInstancePropertiesFrom(&InSettings->BodyInstance);
 
 		BestCluster->ClusterComponent->SetWorldTransform(InstanceToWorld);
 		BestCluster->ClusterComponent->RegisterComponent();
@@ -343,12 +325,12 @@ void FFoliageMeshInfo::AddInstance(AInstancedFoliageActor* InIFA, UStaticMesh* I
 
 	// Update PrimitiveComponent's culling distance taking into account the radius of the bounds, as
 	// it is based on the center of the component's bounds.
-	float CullDistance = Settings->EndCullDistance > 0 ? (float)Settings->EndCullDistance + BestCluster->Bounds.SphereRadius : 0.f;
+	float CullDistance = InSettings->EndCullDistance > 0 ? (float)InSettings->EndCullDistance + BestCluster->Bounds.SphereRadius : 0.f;
 	BestCluster->ClusterComponent->LDMaxDrawDistance = CullDistance;
 	BestCluster->ClusterComponent->CachedMaxDrawDistance = CullDistance;
 
 	CheckValid();
-}
+		}
 
 void FFoliageMeshInfo::RemoveInstances(AInstancedFoliageActor* InIFA, const TArray<int32>& InInstancesToRemove)
 {
@@ -424,8 +406,8 @@ void FFoliageMeshInfo::RemoveInstances(AInstancedFoliageActor* InIFA, const TArr
 
 				// Invalidate this instance's editor data so we reuse the slot later
 				FFoliageInstance& InstanceEditorData = Instances[InstanceIndex];
-				InstanceEditorData.ClusterIndex = -1;
-				InstanceEditorData.Base = NULL;
+				InstanceEditorData.ClusterIndex = INDEX_NONE;
+				InstanceEditorData.Base = nullptr;
 
 				// And remember the slot to reuse it.
 				FreeInstanceIndices.Add(InstanceIndex);
@@ -478,7 +460,7 @@ void FFoliageMeshInfo::RemoveInstances(AInstancedFoliageActor* InIFA, const TArr
 	}
 }
 
-void FFoliageMeshInfo::PreMoveInstances(class AInstancedFoliageActor* InIFA, const TArray<int32>& InInstancesToMove)
+void FFoliageMeshInfo::PreMoveInstances(AInstancedFoliageActor* InIFA, const TArray<int32>& InInstancesToMove)
 {
 	// Remove instances from the hash
 	for (TArray<int32>::TConstIterator It(InInstancesToMove); It; ++It)
@@ -490,7 +472,7 @@ void FFoliageMeshInfo::PreMoveInstances(class AInstancedFoliageActor* InIFA, con
 }
 
 
-void FFoliageMeshInfo::PostUpdateInstances(class AInstancedFoliageActor* InIFA, const TArray<int32>& InInstancesUpdated, bool bReAddToHash)
+void FFoliageMeshInfo::PostUpdateInstances(AInstancedFoliageActor* InIFA, const TArray<int32>& InInstancesUpdated, bool bReAddToHash)
 {
 	for (TArray<int32>::TConstIterator It(InInstancesUpdated); It; ++It)
 	{
@@ -525,18 +507,17 @@ void FFoliageMeshInfo::PostUpdateInstances(class AInstancedFoliageActor* InIFA,
 	}
 }
 
-void FFoliageMeshInfo::PostMoveInstances(class AInstancedFoliageActor* InIFA, const TArray<int32>& InInstancesMoved)
+void FFoliageMeshInfo::PostMoveInstances(AInstancedFoliageActor* InIFA, const TArray<int32>& InInstancesMoved)
 {
 	PostUpdateInstances(InIFA, InInstancesMoved, true);
 }
 
-void FFoliageMeshInfo::DuplicateInstances(class AInstancedFoliageActor* InIFA, class UStaticMesh* InMesh, const TArray<int32>& InInstancesToDuplicate)
+void FFoliageMeshInfo::DuplicateInstances(AInstancedFoliageActor* InIFA, UFoliageType* InSettings, const TArray<int32>& InInstancesToDuplicate)
 {
-	for (TArray<int32>::TConstIterator It(InInstancesToDuplicate); It; ++It)
+	for (int32 InstanceIndex : InInstancesToDuplicate)
 	{
-		int32 InstanceIndex = *It;
 		const FFoliageInstance TempInstance = Instances[InstanceIndex];
-		AddInstance(InIFA, InMesh, TempInstance);
+		AddInstance(InIFA, InSettings, TempInstance);
 	}
 }
 
@@ -555,34 +536,23 @@ int32 FFoliageMeshInfo::GetInstanceCount() const
 }
 
 // Destroy existing clusters and reassign all instances to new clusters
-void FFoliageMeshInfo::ReallocateClusters(AInstancedFoliageActor* InIFA, UStaticMesh* InMesh)
+void FFoliageMeshInfo::ReallocateClusters(AInstancedFoliageActor* InIFA, UFoliageType* InSettings)
 {
 	// Detach all components
 	InIFA->UnregisterAllComponents();
 
 	// Delete any existing clusters from the components array
-	for (int32 ClusterIndex = 0; ClusterIndex < InstanceClusters.Num(); ClusterIndex++)
+	for (FFoliageInstanceCluster& Cluster : InstanceClusters)
 	{
-		UInstancedStaticMeshComponent* Component = InstanceClusters[ClusterIndex].ClusterComponent;
-		if (Component != NULL)
+		UInstancedStaticMeshComponent* Component = Cluster.ClusterComponent;
+		if (Component != nullptr)
 		{
 			Component->bAutoRegister = false;
 		}
 	}
 
 	TArray<FFoliageInstance> OldInstances = Instances;
-	for (int32 Idx = 0; Idx < OldInstances.Num(); Idx++)
-	{
-		if (OldInstances[Idx].ClusterIndex == -1)
-		{
-			OldInstances.RemoveAtSwap(Idx);
-			Idx--;
-		}
-		else
-		{
-			OldInstances[Idx].ClusterIndex = -1;
-		}
-	}
+	OldInstances.RemoveAllSwap([](FFoliageInstance& Instance){ return Instance.ClusterIndex == INDEX_NONE; });
 
 	// Remove everything
 	InstanceClusters.Empty();
@@ -593,47 +563,46 @@ void FFoliageMeshInfo::ReallocateClusters(AInstancedFoliageActor* InIFA, UStatic
 	SelectedIndices.Empty();
 
 	// Re-add
-	for (int32 Idx = 0; Idx < OldInstances.Num(); Idx++)
+	for (FFoliageInstance& Instance : OldInstances)
 	{
-		AddInstance(InIFA, InMesh, OldInstances[Idx]);
+		Instance.ClusterIndex = INDEX_NONE;
+		AddInstance(InIFA, InSettings, Instance);
 	}
 
 	InIFA->RegisterAllComponents();
 }
 
 // Update settings in the clusters based on the current settings (eg culling distance, collision, ...)
-void FFoliageMeshInfo::UpdateClusterSettings(AInstancedFoliageActor* InIFA)
-{
-	for (int32 ClusterIdx = 0; ClusterIdx < InstanceClusters.Num(); ClusterIdx++)
-	{
-		UInstancedStaticMeshComponent* ClusterComponent = InstanceClusters[ClusterIdx].ClusterComponent;
-		ClusterComponent->Modify();
-		ClusterComponent->MarkRenderStateDirty();
-
-		// Copy settings
-		ClusterComponent->InstanceStartCullDistance = Settings->StartCullDistance;
-		ClusterComponent->InstanceEndCullDistance = Settings->EndCullDistance;
-
-		// Update PrimitiveComponent's culling distance taking into account the radius of the bounds, as
-		// it is based on the center of the component's bounds.
-		float CullDistance = Settings->EndCullDistance > 0 ? (float)Settings->EndCullDistance + InstanceClusters[ClusterIdx].Bounds.SphereRadius : 0.f;
-		ClusterComponent->LDMaxDrawDistance = CullDistance;
-		ClusterComponent->CachedMaxDrawDistance = CullDistance;
-	}
-
-	InIFA->MarkComponentsRenderStateDirty();
-}
+//void FFoliageMeshInfo::UpdateClusterSettings(AInstancedFoliageActor* InIFA)
+//{
+//	for (FFoliageInstanceCluster& Cluster : InstanceClusters)
+//	{
+//		UInstancedStaticMeshComponent* ClusterComponent = Cluster.ClusterComponent;
+//		ClusterComponent->Modify();
+//		ClusterComponent->MarkRenderStateDirty();
+//
+//		// Copy settings
+//		ClusterComponent->InstanceStartCullDistance = Settings->StartCullDistance;
+//		ClusterComponent->InstanceEndCullDistance = Settings->EndCullDistance;
+//
+//		// Update PrimitiveComponent's culling distance taking into account the radius of the bounds, as
+//		// it is based on the center of the component's bounds.
+//		float CullDistance = Settings->EndCullDistance > 0 ? (float)Settings->EndCullDistance + Cluster.Bounds.SphereRadius : 0.f;
+//		ClusterComponent->LDMaxDrawDistance = CullDistance;
+//		ClusterComponent->CachedMaxDrawDistance = CullDistance;
+//	}
+//
+//	InIFA->MarkComponentsRenderStateDirty();
+//}
 
 void FFoliageMeshInfo::GetInstancesInsideSphere(const FSphere& Sphere, TArray<int32>& OutInstances)
 {
-	TSet<int32> TempInstances;
-	InstanceHash->GetInstancesOverlappingBox(FBox::BuildAABB(Sphere.Center, FVector(Sphere.W, Sphere.W, Sphere.W)), TempInstances);
-
-	for (TSet<int32>::TConstIterator It(TempInstances); It; ++It)
+	auto TempInstances = InstanceHash->GetInstancesOverlappingBox(FBox::BuildAABB(Sphere.Center, FVector(Sphere.W)));
+	for (int32 Idx : TempInstances)
 	{
-		if (FSphere(Instances[*It].Location, 0.f).IsInside(Sphere))
+		if (FSphere(Instances[Idx].Location, 0.f).IsInside(Sphere))
 		{
-			OutInstances.Add(*It);
+			OutInstances.Add(Idx);
 		}
 	}
 }
@@ -641,12 +610,10 @@ void FFoliageMeshInfo::GetInstancesInsideSphere(const FSphere& Sphere, TArray<in
 // Returns whether or not there is are any instances overlapping the sphere specified
 bool FFoliageMeshInfo::CheckForOverlappingSphere(const FSphere& Sphere)
 {
-	TSet<int32> TempInstances;
-	InstanceHash->GetInstancesOverlappingBox(FBox::BuildAABB(Sphere.Center, FVector(Sphere.W, Sphere.W, Sphere.W)), TempInstances);
-
-	for (TSet<int32>::TConstIterator It(TempInstances); It; ++It)
+	auto TempInstances = InstanceHash->GetInstancesOverlappingBox(FBox::BuildAABB(Sphere.Center, FVector(Sphere.W)));
+	for (int32 Idx : TempInstances)
 	{
-		if (FSphere(Instances[*It].Location, 0.f).IsInside(Sphere))
+		if (FSphere(Instances[Idx].Location, 0.f).IsInside(Sphere))
 		{
 			return true;
 		}
@@ -658,12 +625,11 @@ bool FFoliageMeshInfo::CheckForOverlappingSphere(const FSphere& Sphere)
 bool FFoliageMeshInfo::CheckForOverlappingInstanceExcluding(int32 TestInstanceIdx, float Radius, TSet<int32>& ExcludeInstances)
 {
 	FSphere Sphere(Instances[TestInstanceIdx].Location, Radius);
-	TSet<int32> TempInstances;
-	InstanceHash->GetInstancesOverlappingBox(FBox::BuildAABB(Sphere.Center, FVector(Sphere.W, Sphere.W, Sphere.W)), TempInstances);
 
-	for (TSet<int32>::TConstIterator It(TempInstances); It; ++It)
+	auto TempInstances = InstanceHash->GetInstancesOverlappingBox(FBox::BuildAABB(Sphere.Center, FVector(Sphere.W)));
+	for (int32 Idx : TempInstances)
 	{
-		if (*It != TestInstanceIdx && !ExcludeInstances.Contains(*It) && FSphere(Instances[*It].Location, 0.f).IsInside(Sphere))
+		if (Idx != TestInstanceIdx && !ExcludeInstances.Contains(Idx) && FSphere(Instances[Idx].Location, 0.f).IsInside(Sphere))
 		{
 			return true;
 		}
@@ -677,15 +643,13 @@ void FFoliageMeshInfo::SelectInstances(AInstancedFoliageActor* InIFA, bool bSele
 
 	if (bSelect)
 	{
-		for (int32 i = 0; i < InInstances.Num(); i++)
+		for (int32 i : InInstances)
 		{
-			SelectedIndices.AddUnique(InInstances[i]);
+			SelectedIndices.AddUnique(i);
 		}
 
-		for (int32 ClusterIdx = 0; ClusterIdx < InstanceClusters.Num(); ClusterIdx++)
+		for (FFoliageInstanceCluster& Cluster : InstanceClusters)
 		{
-			FFoliageInstanceCluster& Cluster = InstanceClusters[ClusterIdx];
-
 			// Apply any selections in the component
 			Cluster.ClusterComponent->Modify();
 			Cluster.ClusterComponent->MarkRenderStateDirty();
@@ -706,15 +670,13 @@ void FFoliageMeshInfo::SelectInstances(AInstancedFoliageActor* InIFA, bool bSele
 	}
 	else
 	{
-		for (int32 i = 0; i < InInstances.Num(); i++)
+		for (int32 i : InInstances)
 		{
-			SelectedIndices.RemoveSingleSwap(InInstances[i]);
+			SelectedIndices.RemoveSingleSwap(i);
 		}
 
-		for (int32 ClusterIdx = 0; ClusterIdx < InstanceClusters.Num(); ClusterIdx++)
+		for (FFoliageInstanceCluster& Cluster : InstanceClusters)
 		{
-			FFoliageInstanceCluster& Cluster = InstanceClusters[ClusterIdx];
-
 			if (Cluster.ClusterComponent->SelectedInstances.Num() != 0)
 			{
 				// Remove any selections from the component
@@ -739,7 +701,7 @@ void FFoliageMeshInfo::SelectInstances(AInstancedFoliageActor* InIFA, bool bSele
 //
 // AInstancedFoliageActor
 //
-AInstancedFoliageActor::AInstancedFoliageActor(const class FPostConstructInitializeProperties& PCIP)
+AInstancedFoliageActor::AInstancedFoliageActor(const FPostConstructInitializeProperties& PCIP)
 	: Super(PCIP)
 {
 	SetActorEnableCollision(true);
@@ -749,7 +711,7 @@ AInstancedFoliageActor::AInstancedFoliageActor(const class FPostConstructInitial
 }
 
 #if WITH_EDITOR
-AInstancedFoliageActor* AInstancedFoliageActor::GetInstancedFoliageActor(UWorld* InWorld, bool bCreateIfNone)
+AInstancedFoliageActor* AInstancedFoliageActor::GetInstancedFoliageActorForCurrentLevel(UWorld* InWorld, bool bCreateIfNone)
 {
 	for (TActorIterator<AInstancedFoliageActor> It(InWorld); It; ++It)
 	{
@@ -763,7 +725,7 @@ AInstancedFoliageActor* AInstancedFoliageActor::GetInstancedFoliageActor(UWorld*
 		}
 	}
 
-	return bCreateIfNone ? InWorld->SpawnActor<AInstancedFoliageActor>() : NULL;
+	return bCreateIfNone ? InWorld->SpawnActor<AInstancedFoliageActor>() : nullptr;
 }
 
 
@@ -771,7 +733,7 @@ AInstancedFoliageActor* AInstancedFoliageActor::GetInstancedFoliageActorForLevel
 {
 	if (!InLevel)
 	{
-		return NULL;
+		return nullptr;
 	}
 
 	for (int32 ActorIndex = 0; ActorIndex < InLevel->Actors.Num(); ++ActorIndex)
@@ -782,14 +744,15 @@ AInstancedFoliageActor* AInstancedFoliageActor::GetInstancedFoliageActorForLevel
 			return InstancedFoliageActor;
 		}
 	}
-	return NULL;
+	return nullptr;
 }
 
-void AInstancedFoliageActor::SnapInstancesForLandscape(class ULandscapeHeightfieldCollisionComponent* InLandscapeComponent, const FBox& InInstanceBox)
+void AInstancedFoliageActor::SnapInstancesForLandscape(ULandscapeHeightfieldCollisionComponent* InLandscapeComponent, const FBox& InInstanceBox)
 {
 	for (auto& MeshPair : FoliageMeshes)
 	{
 		// Find the per-mesh info matching the mesh.
+		UFoliageType* Settings = MeshPair.Key;
 		FFoliageMeshInfo& MeshInfo = *MeshPair.Value;
 
 		FFoliageComponentHashInfo* ComponentHashInfo = MeshInfo.ComponentHash.Find(InLandscapeComponent);
@@ -800,9 +763,8 @@ void AInstancedFoliageActor::SnapInstancesForLandscape(class ULandscapeHeightfie
 
 			bool bFirst = true;
 			TArray<int32> InstancesToRemove;
-			for (TSet<int32>::TConstIterator InstIt(ComponentHashInfo->Instances); InstIt; ++InstIt)
+			for (int32 InstanceIndex : ComponentHashInfo->Instances)
 			{
-				int32 InstanceIndex = *InstIt;
 				FFoliageInstance& Instance = MeshInfo.Instances[InstanceIndex];
 
 				// Test location should remove any Z offset
@@ -829,9 +791,8 @@ void AInstancedFoliageActor::SnapInstancesForLandscape(class ULandscapeHeightfie
 					World->LineTraceMulti(Results, Start, End, FCollisionQueryParams(TraceTag, true), FCollisionObjectQueryParams(ECollisionChannel::ECC_Visibility));
 
 					bool bFoundHit = false;
-					for (int32 HitIdx = 0; HitIdx < Results.Num(); HitIdx++)
+					for (const FHitResult& Hit : Results)
 					{
-						const FHitResult& Hit = Results[HitIdx];
 						if (Hit.Component == InLandscapeComponent)
 						{
 							bFoundHit = true;
@@ -847,7 +808,7 @@ void AInstancedFoliageActor::SnapInstancesForLandscape(class ULandscapeHeightfie
 								{
 									// Remove previous alignment and align to new normal.
 									Instance.Rotation = Instance.PreAlignRotation;
-									Instance.AlignToNormal(Hit.Normal, MeshInfo.Settings->AlignMaxAngle);
+									Instance.AlignToNormal(Hit.Normal, Settings->AlignMaxAngle);
 								}
 
 								// Reapply the Z offset in local space
@@ -902,7 +863,7 @@ void AInstancedFoliageActor::SnapInstancesForLandscape(class ULandscapeHeightfie
 	}
 }
 
-void AInstancedFoliageActor::MoveInstancesForMovedComponent(class UActorComponent* InComponent)
+void AInstancedFoliageActor::MoveInstancesForMovedComponent(UActorComponent* InComponent)
 {
 	bool bUpdatedInstances = false;
 	bool bFirst = true;
@@ -924,9 +885,8 @@ void AInstancedFoliageActor::MoveInstancesForMovedComponent(class UActorComponen
 			FVector OldDrawScale = ComponentHashInfo->CachedDrawScale;
 			ComponentHashInfo->UpdateLocationFromActor(InComponent);
 
-			for (TSet<int32>::TConstIterator InstIt(ComponentHashInfo->Instances); InstIt; ++InstIt)
+			for (int32 InstanceIndex : ComponentHashInfo->Instances)
 			{
-				int32 InstanceIndex = *InstIt;
 				FFoliageInstance& Instance = MeshInfo.Instances[InstanceIndex];
 
 				MeshInfo.InstanceHash->RemoveInstance(Instance.Location, InstanceIndex);
@@ -971,7 +931,7 @@ void AInstancedFoliageActor::MoveInstancesForMovedComponent(class UActorComponen
 	}
 }
 
-void AInstancedFoliageActor::DeleteInstancesForComponent(class UActorComponent* InComponent)
+void AInstancedFoliageActor::DeleteInstancesForComponent(UActorComponent* InComponent)
 {
 	for (auto& MeshPair : FoliageMeshes)
 	{
@@ -984,9 +944,9 @@ void AInstancedFoliageActor::DeleteInstancesForComponent(class UActorComponent*
 	}
 }
 
-void AInstancedFoliageActor::MoveInstancesForComponentToCurrentLevel(class UActorComponent* InComponent)
+void AInstancedFoliageActor::MoveInstancesForComponentToCurrentLevel(UActorComponent* InComponent)
 {
-	AInstancedFoliageActor* NewIFA = AInstancedFoliageActor::GetInstancedFoliageActor(InComponent->GetWorld());
+	AInstancedFoliageActor* NewIFA = AInstancedFoliageActor::GetInstancedFoliageActorForCurrentLevel(InComponent->GetWorld());
 
 	for (auto& MeshPair : FoliageMeshes)
 	{
@@ -998,9 +958,8 @@ void AInstancedFoliageActor::MoveInstancesForComponentToCurrentLevel(class UActo
 			FFoliageMeshInfo* NewMeshInfo = NewIFA->FindOrAddMesh(MeshPair.Key);
 
 			// Add the foliage to the new level
-			for (TSet<int32>::TConstIterator InstIt(ComponentHashInfo->Instances); InstIt; ++InstIt)
+			for (int32 InstanceIndex : ComponentHashInfo->Instances)
 			{
-				int32 InstanceIndex = *InstIt;
 				NewMeshInfo->AddInstance(NewIFA, MeshPair.Key, MeshInfo.Instances[InstanceIndex]);
 			}
 
@@ -1010,7 +969,7 @@ void AInstancedFoliageActor::MoveInstancesForComponentToCurrentLevel(class UActo
 	}
 }
 
-void AInstancedFoliageActor::MoveInstancesToNewComponent(class UActorComponent* InOldComponent, class UActorComponent* InNewComponent)
+void AInstancedFoliageActor::MoveInstancesToNewComponent(UActorComponent* InOldComponent, UActorComponent* InNewComponent)
 {
 	AInstancedFoliageActor* NewIFA = AInstancedFoliageActor::GetInstancedFoliageActorForLevel(InNewComponent->GetTypedOuter<ULevel>());
 	if (!NewIFA)
@@ -1063,9 +1022,9 @@ void AInstancedFoliageActor::MoveInstancesToNewComponent(class UActorComponent*
 	}
 }
 
-TMap<class UStaticMesh*, TArray<const FFoliageInstancePlacementInfo*> > AInstancedFoliageActor::GetInstancesForComponent(class UActorComponent* InComponent)
+TMap<UFoliageType*, TArray<const FFoliageInstancePlacementInfo*>> AInstancedFoliageActor::GetInstancesForComponent(UActorComponent* InComponent)
 {
-	TMap<class UStaticMesh*, TArray<const struct FFoliageInstancePlacementInfo*> > Result;
+	TMap<UFoliageType*, TArray<const FFoliageInstancePlacementInfo*>> Result;
 
 	for (auto& MeshPair : FoliageMeshes)
 	{
@@ -1077,9 +1036,8 @@ TMap<class UStaticMesh*, TArray<const FFoliageInstancePlacementInfo*> > AInstanc
 			TArray<const FFoliageInstancePlacementInfo*>& Array = Result.Add(MeshPair.Key, TArray<const FFoliageInstancePlacementInfo*>());
 			Array.Empty(ComponentHashInfo->Instances.Num());
 
-			for (TSet<int32>::TConstIterator InstIt(ComponentHashInfo->Instances); InstIt; ++InstIt)
+			for (int32 InstanceIndex : ComponentHashInfo->Instances)
 			{
-				int32 InstanceIndex = *InstIt;
 				const FFoliageInstancePlacementInfo* Instance = &MeshInfo.Instances[InstanceIndex];
 				Array.Add(Instance);
 			}
@@ -1089,66 +1047,52 @@ TMap<class UStaticMesh*, TArray<const FFoliageInstancePlacementInfo*> > AInstanc
 	return Result;
 }
 
-FFoliageMeshInfo* AInstancedFoliageActor::FindMesh(UStaticMesh* InMesh)
+FFoliageMeshInfo* AInstancedFoliageActor::FindMesh(UFoliageType* InType)
 {
-	TUniqueObj<FFoliageMeshInfo>* MeshInfoEntry = FoliageMeshes.Find(InMesh);
+	TUniqueObj<FFoliageMeshInfo>* MeshInfoEntry = FoliageMeshes.Find(InType);
 	FFoliageMeshInfo* MeshInfo = MeshInfoEntry ? &MeshInfoEntry->Get() : nullptr;
 	return MeshInfo;
 }
 
-FFoliageMeshInfo* AInstancedFoliageActor::FindOrAddMesh(UStaticMesh* InMesh)
+FFoliageMeshInfo* AInstancedFoliageActor::FindOrAddMesh(UFoliageType* InType)
 {
-	TUniqueObj<FFoliageMeshInfo>* MeshInfoEntry = FoliageMeshes.Find(InMesh);
-	FFoliageMeshInfo* MeshInfo = MeshInfoEntry ? &MeshInfoEntry->Get() : AddMesh(InMesh);
+	TUniqueObj<FFoliageMeshInfo>* MeshInfoEntry = FoliageMeshes.Find(InType);
+	FFoliageMeshInfo* MeshInfo = MeshInfoEntry ? &MeshInfoEntry->Get() : AddMesh(InType);
 	return MeshInfo;
 }
 
-struct FFoliageMeshInfo* AInstancedFoliageActor::AddMesh(class UStaticMesh* InMesh)
+FFoliageMeshInfo* AInstancedFoliageActor::AddMesh(UStaticMesh* InMesh, UFoliageType** OutSettings)
 {
-	check(FindMesh(InMesh) == NULL);
+	check(GetSettingsForMesh(InMesh) == nullptr);
 
 	MarkPackageDirty();
 
-	int32 MaxDisplayOrder = 0;
-	for (auto& MeshPair : FoliageMeshes)
-	{
-		if (MeshPair.Value->Settings->DisplayOrder > MaxDisplayOrder)
-		{
-			MaxDisplayOrder = MeshPair.Value->Settings->DisplayOrder;
-		}
-	}
-
-	FFoliageMeshInfo& MeshInfo = *FoliageMeshes.Add(InMesh);
+	UFoliageType_InstancedStaticMesh* Settings = nullptr;
 #if WITH_EDITORONLY_DATA
-	UInstancedFoliageSettings* DefaultSettings = InMesh->FoliageDefaultSettings;
-	if( DefaultSettings != NULL )
+	if (InMesh->FoliageDefaultSettings)
 	{
-		MeshInfo.Settings = Cast<UInstancedFoliageSettings>(StaticDuplicateObject(DefaultSettings,this,TEXT("None")));
+		// TODO: Can't we just use this directly?
+		Settings = DuplicateObject<UFoliageType_InstancedStaticMesh>(InMesh->FoliageDefaultSettings, this);
 	}
 	else
 #endif
 	{
-		MeshInfo.Settings = ConstructObject<UInstancedFoliageSettings>(UInstancedFoliageSettings::StaticClass(), this);
+		Settings = ConstructObject<UFoliageType_InstancedStaticMesh>(UFoliageType_InstancedStaticMesh::StaticClass(), this);
 	}
+	Settings->Mesh = InMesh;
+
+	FFoliageMeshInfo* MeshInfo = AddMesh(Settings);
 
 	const FBoxSphereBounds MeshBounds = InMesh->GetBounds();
 
-	// Set IsPendingKill() to true so that when the initial undo record is made,
-	// the component will be treated as destroyed, in that undo an add will
-	// actually work
-	MeshInfo.Settings->SetFlags(RF_PendingKill);
-	MeshInfo.Settings->Modify(false);
-	MeshInfo.Settings->ClearFlags(RF_PendingKill);
-	MeshInfo.Settings->IsSelected = true;
-	MeshInfo.Settings->DisplayOrder = MaxDisplayOrder + 1;
-	MeshInfo.Settings->MeshBounds = MeshBounds;
+	Settings->MeshBounds = MeshBounds;
 
 	// Make bottom only bound
 	FBox LowBound = MeshBounds.GetBox();
 	LowBound.Max.Z = LowBound.Min.Z + (LowBound.Max.Z - LowBound.Min.Z) * 0.1f;
 
 	float MinX = FLT_MAX, MaxX = FLT_MIN, MinY = FLT_MAX, MaxY = FLT_MIN;
-	MeshInfo.Settings->LowBoundOriginRadius = FVector::ZeroVector;
+	Settings->LowBoundOriginRadius = FVector::ZeroVector;
 
 	if (InMesh->RenderData)
 	{
@@ -1166,37 +1110,90 @@ struct FFoliageMeshInfo* AInstancedFoliageActor::AddMesh(class UStaticMesh* InMe
 		}
 	}
 
-	MeshInfo.Settings->LowBoundOriginRadius = FVector((MinX + MaxX), (MinY + MaxY), FMath::Sqrt(FMath::Square(MaxX - MinX) + FMath::Square(MaxY - MinY))) * 0.5f;
+	Settings->LowBoundOriginRadius = FVector((MinX + MaxX), (MinY + MaxY), FMath::Sqrt(FMath::Square(MaxX - MinX) + FMath::Square(MaxY - MinY))) * 0.5f;
+
+	if (OutSettings)
+	{
+		*OutSettings = Settings;
+	}
 
-	return &MeshInfo;
+	return MeshInfo;
 }
 
-void AInstancedFoliageActor::RemoveMesh(class UStaticMesh* InMesh)
+FFoliageMeshInfo* AInstancedFoliageActor::AddMesh(UFoliageType* InType)
+{
+	check(FoliageMeshes.Find(InType) == nullptr);
+
+	MarkPackageDirty();
+
+	if (InType->DisplayOrder == 0)
+	{
+		int32 MaxDisplayOrder = 0;
+		for (auto& MeshPair : FoliageMeshes)
+		{
+			if (MeshPair.Key->DisplayOrder > MaxDisplayOrder)
+			{
+				MaxDisplayOrder = MeshPair.Key->DisplayOrder;
+			}
+		}
+		InType->DisplayOrder = MaxDisplayOrder + 1;
+	}
+
+	FFoliageMeshInfo* MeshInfo = &*FoliageMeshes.Add(InType);
+	InType->IsSelected = true;
+
+	return MeshInfo;
+}
+
+void AInstancedFoliageActor::RemoveMesh(UFoliageType* InSettings)
 {
 	Modify();
 	MarkPackageDirty();
 	UnregisterAllComponents();
 
 	// Remove all components for this mesh from the Components array.
-	FFoliageMeshInfo* MeshInfo = FindMesh(InMesh);
+	FFoliageMeshInfo* MeshInfo = FindMesh(InSettings);
 	if (MeshInfo)
 	{
-		for (int32 ClusterIndex = 0; ClusterIndex < MeshInfo->InstanceClusters.Num(); ClusterIndex++)
+		for (FFoliageInstanceCluster& Cluster : MeshInfo->InstanceClusters)
 		{
-			UInstancedStaticMeshComponent* Component = MeshInfo->InstanceClusters[ClusterIndex].ClusterComponent;
-			if (Component != NULL)
+			UInstancedStaticMeshComponent* Component = Cluster.ClusterComponent;
+			if (Component != nullptr)
 			{
 				Component->bAutoRegister = false;
 			}
 		}
 	}
 
-	FoliageMeshes.Remove(InMesh);
+	FoliageMeshes.Remove(InSettings);
 	RegisterAllComponents();
 	CheckSelection();
 }
 
-void AInstancedFoliageActor::SelectInstance(class UInstancedStaticMeshComponent* InComponent, int32 InComponentInstanceIndex, bool bToggle)
+UFoliageType* AInstancedFoliageActor::GetSettingsForMesh(UStaticMesh* InMesh, FFoliageMeshInfo** OutMeshInfo)
+{
+	UFoliageType* Type = nullptr;
+	FFoliageMeshInfo* MeshInfo = nullptr;
+
+	for (auto& MeshPair : FoliageMeshes)
+	{
+		UFoliageType* Settings = MeshPair.Key;
+		if (Settings && Settings->GetStaticMesh() == InMesh)
+		{
+			Type = MeshPair.Key;
+			MeshInfo = &*MeshPair.Value;
+			break;
+		}
+	}
+
+	if (OutMeshInfo)
+	{
+		*OutMeshInfo = MeshInfo;
+	}
+	return Type;
+}
+
+void AInstancedFoliageActor::SelectInstance(UInstancedStaticMeshComponent* InComponent, int32 InComponentInstanceIndex, bool bToggle)
 {
 	bool bNeedsUpdate = false;
 
@@ -1230,53 +1227,52 @@ void AInstancedFoliageActor::SelectInstance(class UInstancedStaticMeshComponent*
 
 	if (InComponent)
 	{
-		FFoliageMeshInfo* MeshInfo = FindMesh(InComponent->StaticMesh);
+		UFoliageType* Type = nullptr;
+		FFoliageMeshInfo* MeshInfo = nullptr;
+
+		Type = GetSettingsForMesh(InComponent->StaticMesh, &MeshInfo);
 
 		if (MeshInfo)
 		{
-			for (int32 ClusterIdx = 0; ClusterIdx < MeshInfo->InstanceClusters.Num(); ClusterIdx++)
+			FFoliageInstanceCluster* Cluster = MeshInfo->InstanceClusters.FindByPredicate([InComponent](FFoliageInstanceCluster& Cluster){ return Cluster.ClusterComponent == InComponent; });
+			if (ensure(Cluster))
 			{
-				FFoliageInstanceCluster& Cluster = MeshInfo->InstanceClusters[ClusterIdx];
-				if (Cluster.ClusterComponent == InComponent)
-				{
-					InComponent->Modify();
-					int32 InstanceIndex = Cluster.InstanceIndices[InComponentInstanceIndex];
-					int32 SelectedIndex = MeshInfo->SelectedIndices.Find(InstanceIndex);
+				InComponent->Modify();
+				int32 InstanceIndex = Cluster->InstanceIndices[InComponentInstanceIndex];
+				int32 SelectedIndex = MeshInfo->SelectedIndices.Find(InstanceIndex);
 
-					bNeedsUpdate = true;
+				bNeedsUpdate = true;
 
-					// Deselect if it's already selected.
-					if (InComponentInstanceIndex < InComponent->SelectedInstances.Num())
-					{
-						InComponent->SelectedInstances[InComponentInstanceIndex] = false;
-						InComponent->MarkRenderStateDirty();
-					}
-					if (SelectedIndex != INDEX_NONE)
-					{
-						MeshInfo->SelectedIndices.RemoveAt(SelectedIndex);
-					}
+				// Deselect if it's already selected.
+				if (InComponentInstanceIndex < InComponent->SelectedInstances.Num())
+				{
+					InComponent->SelectedInstances[InComponentInstanceIndex] = false;
+					InComponent->MarkRenderStateDirty();
+				}
+				if (SelectedIndex != INDEX_NONE)
+				{
+					MeshInfo->SelectedIndices.RemoveAt(SelectedIndex);
+				}
 
-					if (bToggle && SelectedIndex != INDEX_NONE)
+				if (bToggle && SelectedIndex != INDEX_NONE)
+				{
+					if (SelectedMesh == Type && MeshInfo->SelectedIndices.Num() == 0)
 					{
-						if (SelectedMesh == InComponent->StaticMesh && MeshInfo->SelectedIndices.Num() == 0)
-						{
-							SelectedMesh = NULL;
-						}
+						SelectedMesh = nullptr;
 					}
-					else
+				}
+				else
+				{
+					// Add the selection
+					if (InComponent->SelectedInstances.Num() < InComponent->PerInstanceSMData.Num())
 					{
-						// Add the selection
-						if (InComponent->SelectedInstances.Num() < InComponent->PerInstanceSMData.Num())
-						{
-							InComponent->SelectedInstances.Init(false, Cluster.InstanceIndices.Num());
-						}
-						InComponent->SelectedInstances[InComponentInstanceIndex] = true;
-						InComponent->MarkRenderStateDirty();
-
-						SelectedMesh = InComponent->StaticMesh;
-						MeshInfo->SelectedIndices.Insert(InstanceIndex, 0);
+						InComponent->SelectedInstances.Init(false, Cluster->InstanceIndices.Num());
 					}
-					break;
+					InComponent->SelectedInstances[InComponentInstanceIndex] = true;
+					InComponent->MarkRenderStateDirty();
+
+					SelectedMesh = Type;
+					MeshInfo->SelectedIndices.Insert(InstanceIndex, 0);
 				}
 			}
 		}
@@ -1303,10 +1299,8 @@ void AInstancedFoliageActor::ApplySelectionToComponents(bool bApply)
 		{
 			if (MeshInfo.SelectedIndices.Num() > 0)
 			{
-				for (int32 ClusterIdx = 0; ClusterIdx < MeshInfo.InstanceClusters.Num(); ClusterIdx++)
+				for (FFoliageInstanceCluster& Cluster : MeshInfo.InstanceClusters)
 				{
-					FFoliageInstanceCluster& Cluster = MeshInfo.InstanceClusters[ClusterIdx];
-
 					// Apply any selections in the component
 					Cluster.ClusterComponent->SelectedInstances.Init(false, Cluster.InstanceIndices.Num());
 					Cluster.ClusterComponent->MarkRenderStateDirty();
@@ -1325,10 +1319,8 @@ void AInstancedFoliageActor::ApplySelectionToComponents(bool bApply)
 		}
 		else
 		{
-			for (int32 ClusterIdx = 0; ClusterIdx < MeshInfo.InstanceClusters.Num(); ClusterIdx++)
+			for (FFoliageInstanceCluster& Cluster : MeshInfo.InstanceClusters)
 			{
-				FFoliageInstanceCluster& Cluster = MeshInfo.InstanceClusters[ClusterIdx];
-
 				if (Cluster.ClusterComponent->SelectedInstances.Num() > 0)
 				{
 					// remove any selections in the component
@@ -1350,7 +1342,7 @@ void AInstancedFoliageActor::ApplySelectionToComponents(bool bApply)
 void AInstancedFoliageActor::CheckSelection()
 {
 	// Check if we have to change the selection.
-	if (SelectedMesh != NULL)
+	if (SelectedMesh != nullptr)
 	{
 		FFoliageMeshInfo* MeshInfo = FindMesh(SelectedMesh);
 		if (MeshInfo && MeshInfo->SelectedIndices.Num() > 0)
@@ -1359,7 +1351,7 @@ void AInstancedFoliageActor::CheckSelection()
 		}
 	}
 
-	SelectedMesh = NULL;
+	SelectedMesh = nullptr;
 
 	// Try to find a new selection
 	for (auto& MeshPair : FoliageMeshes)
@@ -1377,7 +1369,7 @@ FVector AInstancedFoliageActor::GetSelectionLocation()
 {
 	FVector Result(0, 0, 0);
 
-	if (SelectedMesh != NULL)
+	if (SelectedMesh != nullptr)
 	{
 		FFoliageMeshInfo* MeshInfo = FindMesh(SelectedMesh);
 		if (MeshInfo && MeshInfo->SelectedIndices.Num() > 0)
@@ -1394,7 +1386,7 @@ void AInstancedFoliageActor::MapRebuild()
 	// Map rebuilt - this may have modified the BSP components and thrown the previous ones away - if so we need to migrate the foliage across.
 	UE_LOG(LogInstancedFoliage, Log, TEXT("Map Rebuilt - Update all BSP painted foliage!"));
 
-	TMap< UStaticMesh*, TArray<FFoliageInstance> > NewInstances;
+	TMap<UFoliageType*, TArray<FFoliageInstance>> NewInstances;
 	TArray<UActorComponent*> RemovedComponents;
 	UWorld* World = GetWorld();
 	check(World);
@@ -1404,29 +1396,29 @@ void AInstancedFoliageActor::MapRebuild()
 	{
 		// each target component has some foliage instances
 		FFoliageMeshInfo const& MeshInfo = *MeshPair.Value;
-		UStaticMesh* Mesh = MeshPair.Key;
-		check(Mesh);
+		UFoliageType* Settings = MeshPair.Key;
+		check(Settings);
 
-		for (TMap<UActorComponent*, FFoliageComponentHashInfo >::TConstIterator ComponentFoliageIt(MeshInfo.ComponentHash); ComponentFoliageIt; ++ComponentFoliageIt)
+		for (auto& ComponentFoliagePair : MeshInfo.ComponentHash)
 		{
 			// BSP components are UModelComponents - they are the only ones we need to change
-			UActorComponent* TargetComponent = ComponentFoliageIt.Key();
+			UActorComponent* TargetComponent = ComponentFoliagePair.Key;
 			if (!TargetComponent || TargetComponent->IsA(UModelComponent::StaticClass()))
 			{
 				// Delete it later
 				RemovedComponents.Add(TargetComponent);
 
-				FFoliageComponentHashInfo const& FoliageInfo = ComponentFoliageIt.Value();
+				FFoliageComponentHashInfo const& FoliageInfo = ComponentFoliagePair.Value;
 
 				// We have to test each instance to see if we can migrate it across
-				for (TSet<int32>::TConstIterator InstanceIt(FoliageInfo.Instances); InstanceIt; ++InstanceIt)
+				for (int32 InstanceIdx : FoliageInfo.Instances)
 				{
 					// Use a line test against the world, similar to FoliageEditMode.
-					check(MeshInfo.Instances.IsValidIndex(*InstanceIt));
-					FFoliageInstance const& Instance = MeshInfo.Instances[*InstanceIt];
+					check(MeshInfo.Instances.IsValidIndex(InstanceIdx));
+					FFoliageInstance const& Instance = MeshInfo.Instances[InstanceIdx];
 
 					FFoliageInstance NewInstance = Instance;
-					NewInstance.ClusterIndex = -1;
+					NewInstance.ClusterIndex = INDEX_NONE;
 
 					FTransform InstanceToWorld = Instance.GetInstanceWorldTransform();
 					FVector Down(-FVector::UpVector);
@@ -1443,40 +1435,88 @@ void AInstancedFoliageActor::MapRebuild()
 					else
 					{
 						// Untargeted instances remain for manual deletion.
-						NewInstance.Base = NULL;
+						NewInstance.Base = nullptr;
 					}
 
-					NewInstances.FindOrAdd(Mesh).Add(NewInstance);
+					NewInstances.FindOrAdd(Settings).Add(NewInstance);
 				}
 			}
 		}
 	}
 
 	// Remove all existing & broken instances & component references.
-	for (TArray<UActorComponent*>::TConstIterator ComponentIt(RemovedComponents); ComponentIt; ++ComponentIt)
+	for (UActorComponent* Component : RemovedComponents)
 	{
-		DeleteInstancesForComponent(*ComponentIt);
+		DeleteInstancesForComponent(Component);
 	}
 
 	// And then finally add our new instances to the correct target components.
-	for (TMap< UStaticMesh*, TArray<FFoliageInstance> >::TConstIterator NewInstanceIt(NewInstances); NewInstanceIt; ++NewInstanceIt)
+	for (auto& NewInstancePair : NewInstances)
 	{
-		UStaticMesh* Mesh = NewInstanceIt.Key();
-		check(Mesh);
-		FFoliageMeshInfo& MeshInfo = *FindOrAddMesh(Mesh);
-		for (TArray<FFoliageInstance>::TConstIterator InstanceIt(NewInstanceIt.Value()); InstanceIt; ++InstanceIt)
+		UFoliageType* Settings = NewInstancePair.Key;
+		check(Settings);
+		FFoliageMeshInfo& MeshInfo = *FindOrAddMesh(Settings);
+		for (FFoliageInstance& Instance : NewInstancePair.Value)
 		{
-			MeshInfo.AddInstance(this, Mesh, *InstanceIt);
+			MeshInfo.AddInstance(this, Settings, Instance);
 		}
 	}
 }
 
 #endif
 
+struct FFoliageMeshInfo_Old
+{
+	TArray<FFoliageInstanceCluster> InstanceClusters;
+	TArray<FFoliageInstance> Instances;
+	//FFoliageInstanceHash* InstanceHash;
+	//TMap<UActorComponent*, FFoliageComponentHashInfo> ComponentHash;
+	//TArray<int32> FreeInstanceIndices;
+	//TArray<int32> SelectedIndices;
+	UFoliageType_InstancedStaticMesh* Settings; // Type remapped via +ActiveClassRedirects
+};
+FArchive& operator<<(FArchive& Ar, FFoliageMeshInfo_Old& MeshInfo)
+{
+	Ar << MeshInfo.InstanceClusters;
+	Ar << MeshInfo.Instances;
+	Ar << MeshInfo.Settings;
+
+	return Ar;
+}
+
 void AInstancedFoliageActor::Serialize(FArchive& Ar)
 {
 	Super::Serialize(Ar);
-	Ar << FoliageMeshes;
+	if (Ar.UE4Ver() < VER_UE4_FOLIAGE_SETTINGS_TYPE)
+	{
+		TMap<UStaticMesh*, FFoliageMeshInfo_Old> OldFoliageMeshes;
+		Ar << OldFoliageMeshes;
+		for (auto& OldMeshInfo : OldFoliageMeshes)
+		{
+			FFoliageMeshInfo NewMeshInfo;
+			NewMeshInfo.InstanceClusters = MoveTemp(OldMeshInfo.Value.InstanceClusters);
+#if WITH_EDITORONLY_DATA
+			NewMeshInfo.Instances = MoveTemp(OldMeshInfo.Value.Instances);
+#endif
+			UFoliageType_InstancedStaticMesh* FoliageType = OldMeshInfo.Value.Settings;
+			if (FoliageType->Mesh == nullptr)
+			{
+				FoliageType->Modify();
+				FoliageType->Mesh = OldMeshInfo.Key;
+			}
+			else if (FoliageType->Mesh != OldMeshInfo.Key)
+			{
+				// If mesh doesn't match (two meshes sharing the same settings object?) then we need to duplicate as that is no longer supported
+				FoliageType = (UFoliageType_InstancedStaticMesh*)StaticDuplicateObject(FoliageType, this, nullptr, RF_AllFlags & ~(RF_Standalone | RF_Public));
+				FoliageType->Mesh = OldMeshInfo.Key;
+			}
+			FoliageMeshes.Add(FoliageType, TUniqueObj<FFoliageMeshInfo>(MoveTemp(NewMeshInfo)));
+		}
+	}
+	else
+	{
+		Ar << FoliageMeshes;
+	}
 
 	if (Ar.UE4Ver() < VER_UE4_FOLIAGE_MOVABLE_MOBILITY)
 	{
@@ -1497,6 +1537,7 @@ void AInstancedFoliageActor::PostLoad()
 {
 	Super::PostLoad();
 
+#if WITH_EDITORONLY_DATA
 	if (GetLinkerUE4Version() < VER_UE4_DISALLOW_FOLIAGE_ON_BLUEPRINTS)
 	{
 		for (auto& MeshPair : FoliageMeshes)
@@ -1513,41 +1554,46 @@ void AInstancedFoliageActor::PostLoad()
 			}
 		}
 	}
+#endif
 
+#if WITH_EDITOR
 	if (GIsEditor)
 	{
-		for (auto MeshIt = FoliageMeshes.CreateIterator(); MeshIt; ++MeshIt)
 		{
-			// Remove any NULL entries.
-			if (MeshIt.Key() == NULL)
+			bool bContainsNull = FoliageMeshes.Remove(nullptr) > 0;
+			if (bContainsNull)
 			{
 				FMessageLog("MapCheck").Warning()
 					->AddToken(FUObjectToken::Create(this))
 					->AddToken(FTextToken::Create(LOCTEXT("MapCheck_Message_FoliageMissingStaticMesh", "Foliage instances for a missing static mesh have been removed.")))
 					->AddToken(FMapErrorToken::Create(FMapErrors::FoliageMissingStaticMesh));
-				MeshIt.RemoveCurrent();
-				continue;
+				while (bContainsNull)
+				{
+					bContainsNull = FoliageMeshes.Remove(nullptr) > 0;
+				}
 			}
-
+		}
+		for (auto& MeshPair : FoliageMeshes)
+		{
 			// Find the per-mesh info matching the mesh.
-			FFoliageMeshInfo& MeshInfo = *MeshIt.Value();
+			FFoliageMeshInfo& MeshInfo = *MeshPair.Value;
 
 			// Update the FreeInstanceIndices list and hash.
 			for (int32 InstanceIdx = 0; InstanceIdx < MeshInfo.Instances.Num(); InstanceIdx++)
 			{
 				FFoliageInstance& Instance = MeshInfo.Instances[InstanceIdx];
-				if (Instance.ClusterIndex == -1)
+				if (Instance.ClusterIndex == INDEX_NONE)
 				{
 					// Add invalid instances to the FreeInstanceIndices list.
 					MeshInfo.FreeInstanceIndices.Add(InstanceIdx);
-					Instance.Base = NULL;
+					Instance.Base = nullptr;
 				}
 				else
 				{
 					// Add valid instances to the hash.
 					MeshInfo.InstanceHash->InsertInstance(Instance.Location, InstanceIdx);
 					FFoliageComponentHashInfo* ComponentHashInfo = MeshInfo.ComponentHash.Find(Instance.Base);
-					if (ComponentHashInfo == NULL)
+					if (ComponentHashInfo == nullptr)
 					{
 						ComponentHashInfo = &MeshInfo.ComponentHash.Add(Instance.Base, FFoliageComponentHashInfo(Instance.Base));
 					}
@@ -1555,11 +1601,11 @@ void AInstancedFoliageActor::PostLoad()
 				}
 			}
 
-			// Check if any of the clusters are NULL
+			// Check if any of the clusters are nullptr
 			int32 NumMissingComponents = 0;
-			for (int32 ClusterIdx = 0; ClusterIdx < MeshInfo.InstanceClusters.Num(); ClusterIdx++)
+			for (FFoliageInstanceCluster& Cluster : MeshInfo.InstanceClusters)
 			{
-				if (MeshInfo.InstanceClusters[ClusterIdx].ClusterComponent == NULL)
+				if (Cluster.ClusterComponent == nullptr)
 				{
 					NumMissingComponents++;
 				}
@@ -1568,7 +1614,7 @@ void AInstancedFoliageActor::PostLoad()
 			{
 				FFormatNamedArguments Arguments;
 				Arguments.Add(TEXT("MissingCount"), NumMissingComponents);
-				Arguments.Add(TEXT("MeshName"), FText::FromString(MeshIt.Key()->GetName()));
+				Arguments.Add(TEXT("MeshName"), FText::FromString(MeshPair.Key->GetName()));
 				FMessageLog("MapCheck").Warning()
 					->AddToken(FUObjectToken::Create(this))
 					->AddToken(FTextToken::Create(FText::Format(LOCTEXT("MapCheck_Message_FoliageMissingClusterComponent", "Foliage in this map is missing {MissingCount} cluster component(s) for static mesh {MeshName}. Opening the Foliage tool will this problem."), Arguments)))
@@ -1576,6 +1622,7 @@ void AInstancedFoliageActor::PostLoad()
 			}
 		}
 	}
+#endif
 }
 
 //
@@ -1594,23 +1641,20 @@ void AInstancedFoliageActor::AddReferencedObjects(UObject* InThis, FReferenceCol
 		Collector.AddReferencedObject(MeshPair.Key, This);
 		FFoliageMeshInfo& MeshInfo = *MeshPair.Value;
 
-		if (MeshInfo.Settings != NULL)
-		{
-			Collector.AddReferencedObject(MeshInfo.Settings, This);
-		}
-
-		for (int32 InstanceIdx = 0; InstanceIdx < MeshInfo.Instances.Num(); InstanceIdx++)
+#if WITH_EDITORONLY_DATA
+		for (FFoliageInstance& Instance : MeshInfo.Instances)
 		{
-			class UActorComponent* Base = MeshInfo.Instances[InstanceIdx].Base;
-			if (Base != NULL)
+			UActorComponent* Base = Instance.Base;
+			if (Base != nullptr)
 			{
 				Collector.AddReferencedObject(Base, This);
 			}
 		}
+#endif
 
-		for (int32 ClusterIdx = 0; ClusterIdx < MeshInfo.InstanceClusters.Num(); ClusterIdx++)
+		for (FFoliageInstanceCluster& Cluster : MeshInfo.InstanceClusters)
 		{
-			Collector.AddReferencedObject(MeshInfo.InstanceClusters[ClusterIdx].ClusterComponent, This);
+			Collector.AddReferencedObject(Cluster.ClusterComponent, This);
 		}
 	}
 	Super::AddReferencedObjects(This, Collector);
@@ -1627,9 +1671,9 @@ void AInstancedFoliageActor::ApplyWorldOffset(const FVector& InOffset, bool bWor
 	TArray<UActorComponent*> Components;
 	GetComponents(Components);
 
-	for (int32 CompIndex = 0; CompIndex < Components.Num(); CompIndex++)
+	for (UActorComponent* Component : Components)
 	{
-		Components[CompIndex]->ApplyWorldOffset(InOffset, bWorldShift);
+		Component->ApplyWorldOffset(InOffset, bWorldShift);
 	}
 
 	if (GIsEditor)
@@ -1638,12 +1682,12 @@ void AInstancedFoliageActor::ApplyWorldOffset(const FVector& InOffset, bool bWor
 		{
 			FFoliageMeshInfo& MeshInfo = *MeshPair.Value;
 
-			for (int32 ClusterIdx = 0; ClusterIdx < MeshInfo.InstanceClusters.Num(); ClusterIdx++)
+			for (FFoliageInstanceCluster& Cluster : MeshInfo.InstanceClusters)
 			{
-				FFoliageInstanceCluster& Cluster = MeshInfo.InstanceClusters[ClusterIdx];
 				Cluster.Bounds.Origin += InOffset;
 			}
 
+#if WITH_EDITORONLY_DATA
 			MeshInfo.InstanceHash->Empty();
 			for (int32 InstanceIdx = 0; InstanceIdx < MeshInfo.Instances.Num(); InstanceIdx++)
 			{
@@ -1659,6 +1703,7 @@ void AInstancedFoliageActor::ApplyWorldOffset(const FVector& InOffset, bool bWor
 				FFoliageComponentHashInfo& Info = It.Value();
 				Info.CachedLocation += InOffset;
 			}
+#endif
 		}
 	}
 }
diff --git a/Engine/Source/Runtime/Engine/Private/Landscape/LandscapeEdit.cpp b/Engine/Source/Runtime/Engine/Private/Landscape/LandscapeEdit.cpp
--- a/Engine/Source/Runtime/Engine/Private/Landscape/LandscapeEdit.cpp
+++ b/Engine/Source/Runtime/Engine/Private/Landscape/LandscapeEdit.cpp
@@ -706,7 +706,7 @@ void ULandscapeComponent::UpdateCollisionHeightData(const FColor* HeightmapTextu
 		// Move any foliage instances if we created a new collision component.
 		if( OldCollisionComponent && OldCollisionComponent != CollisionComp )
 		{
-			AInstancedFoliageActor* IFA = AInstancedFoliageActor::GetInstancedFoliageActor(OldCollisionComponent->GetWorld());
+			AInstancedFoliageActor* IFA = AInstancedFoliageActor::GetInstancedFoliageActorForCurrentLevel(OldCollisionComponent->GetWorld());
 			IFA->MoveInstancesToNewComponent(OldCollisionComponent, CollisionComp);
 		}
 	}
@@ -1265,24 +1265,7 @@ void ULandscapeComponent::UpdateMipsTempl(int32 InNumSubsections, int32 InSubsec
 	int32 WeightmapSizeV = Texture->Source.GetSizeY();
 
 	// Find the maximum mip where each texel's data comes from just one subsection.
-	int32 MaxWholeSubsectionMip = 1;
-	for(int32 Mip=1;;++Mip)
-	{
-		int32 MipSubsectionSizeQuads = ((InSubsectionSizeQuads+1)>>Mip)-1;
-
-		int32 MipSizeU = FMath::Max<int32>(WeightmapSizeU >> Mip,1);
-		int32 MipSizeV = FMath::Max<int32>(WeightmapSizeV >> Mip,1);
-
-		// Mip must represent at least one quad to store valid weight data
-		if( MipSubsectionSizeQuads > 0 )
-		{
-			MaxWholeSubsectionMip = Mip;
-		}
-		else
-		{
-			break;
-		}
-	}
+	int32 MaxWholeSubsectionMip = FMath::FloorLog2(InSubsectionSizeQuads + 1) - 1;
 
 	// Update the mip where each texel's data comes from just one subsection.
 	for( int32 SubsectionY = 0;SubsectionY < InNumSubsections;SubsectionY++ )
@@ -1401,43 +1384,40 @@ void ULandscapeComponent::UpdateMipsTempl(int32 InNumSubsections, int32 InSubsec
 	}
 
 	// Handle mips that have texels from multiple subsections
-	for(int32 Mip=1;;++Mip)
+	// not valid weight data, so just average the texels of the previous mip.
+	for (int32 Mip = MaxWholeSubsectionMip + 1;; ++Mip)
 	{
-		int32 MipSubsectionSizeQuads = ((InSubsectionSizeQuads+1)>>Mip)-1;
+		int32 MipSubsectionSizeQuads = ((InSubsectionSizeQuads + 1) >> Mip) - 1;
+		checkSlow(MipSubsectionSizeQuads <= 0);
 
-		int32 MipSizeU = FMath::Max<int32>(WeightmapSizeU >> Mip,1);
-		int32 MipSizeV = FMath::Max<int32>(WeightmapSizeV >> Mip,1);
+		int32 MipSizeU = FMath::Max<int32>(WeightmapSizeU >> Mip, 1);
+		int32 MipSizeV = FMath::Max<int32>(WeightmapSizeV >> Mip, 1);
 
-		// Mip must represent at least one quad to store valid weight data
-		if( MipSubsectionSizeQuads <= 0 )
-		{
-			int32 PrevMipSizeU = WeightmapSizeU >> (Mip-1);
-			int32 PrevMipSizeV = WeightmapSizeV >> (Mip-1);
+		int32 PrevMipSizeU = FMath::Max<int32>(WeightmapSizeU >> (Mip - 1), 1);
+		int32 PrevMipSizeV = FMath::Max<int32>(WeightmapSizeV >> (Mip - 1), 1);
 
-			// not valid weight data, so just average the texels of the previous mip.
-			for( int32 Y = 0;Y < MipSizeV;Y++ )
+		for (int32 Y = 0; Y < MipSizeV; Y++)
+		{
+			for (int32 X = 0; X < MipSizeU; X++)
 			{
-				for( int32 X = 0;X < MipSizeU;X++ )
-				{
-					DataType* TexData = &(TextureMipData[Mip])[ X + Y * MipSizeU ];
+				DataType* TexData = &(TextureMipData[Mip])[X + Y * MipSizeU];
 
-					DataType *PreMipTexData00 = &(TextureMipData[Mip-1])[ (X*2+0) + (Y*2+0)  * PrevMipSizeU ];
-					DataType *PreMipTexData01 = &(TextureMipData[Mip-1])[ (X*2+0) + (Y*2+1)  * PrevMipSizeU ];
-					DataType *PreMipTexData10 = &(TextureMipData[Mip-1])[ (X*2+1) + (Y*2+0)  * PrevMipSizeU ];
-					DataType *PreMipTexData11 = &(TextureMipData[Mip-1])[ (X*2+1) + (Y*2+1)  * PrevMipSizeU ];
+				DataType *PreMipTexData00 = &(TextureMipData[Mip - 1])[(X * 2 + 0) + (Y * 2 + 0)  * PrevMipSizeU];
+				DataType *PreMipTexData01 = &(TextureMipData[Mip - 1])[(X * 2 + 0) + (Y * 2 + 1)  * PrevMipSizeU];
+				DataType *PreMipTexData10 = &(TextureMipData[Mip - 1])[(X * 2 + 1) + (Y * 2 + 0)  * PrevMipSizeU];
+				DataType *PreMipTexData11 = &(TextureMipData[Mip - 1])[(X * 2 + 1) + (Y * 2 + 1)  * PrevMipSizeU];
 
-					AverageTexData<DataType>(TexData, PreMipTexData00, PreMipTexData10, PreMipTexData01, PreMipTexData11);
-				}
+				AverageTexData<DataType>(TexData, PreMipTexData00, PreMipTexData10, PreMipTexData01, PreMipTexData11);
 			}
+		}
 
-			if( TextureDataInfo )
-			{
-				// These mip sizes are small enough that we may as well just update the whole mip.
-				TextureDataInfo->AddMipUpdateRegion(Mip,0,0,MipSizeU-1,MipSizeV-1);
-			}
+		if (TextureDataInfo)
+		{
+			// These mip sizes are small enough that we may as well just update the whole mip.
+			TextureDataInfo->AddMipUpdateRegion(Mip, 0, 0, MipSizeU - 1, MipSizeV - 1);
 		}
 
-		if( MipSizeU == 1 && MipSizeV == 1 )
+		if (MipSizeU == 1 && MipSizeV == 1)
 		{
 			break;
 		}
diff --git a/Engine/Source/Runtime/Engine/Public/InstancedFoliage.h b/Engine/Source/Runtime/Engine/Public/InstancedFoliage.h
--- a/Engine/Source/Runtime/Engine/Public/InstancedFoliage.h
+++ b/Engine/Source/Runtime/Engine/Public/InstancedFoliage.h
@@ -2,24 +2,26 @@
 
 /*=============================================================================
 	InstancedFoliage.h: Instanced foliage type definitions.
-=============================================================================*/
+  =============================================================================*/
 
 #pragma once
 
-#include "Components/SceneComponent.h"
 //
 // Forward declarations.
 //
 class UInstancedStaticMeshComponent;
+class AInstancedFoliageActor;
+class UFoliageType;
+struct FFoliageInstanceHash;
 
 /**
  * Flags stored with each instance
  */
 enum EFoliageInstanceFlags
 {
-	FOLIAGE_AlignToNormal				= 0x00000001,
-	FOLIAGE_NoRandomYaw					= 0x00000002,
-	FOLIAGE_Readjusted					= 0x00000004,
+	FOLIAGE_AlignToNormal = 0x00000001,
+	FOLIAGE_NoRandomYaw   = 0x00000002,
+	FOLIAGE_Readjusted    = 0x00000004,
 };
 
 /**
@@ -35,12 +37,12 @@ struct FFoliageInstancePlacementInfo
 	uint32 Flags;
 
 	FFoliageInstancePlacementInfo()
-	:	Location(0.f,0.f,0.f)
-	,	Rotation(0,0,0)
-	,	PreAlignRotation(0,0,0)
-	,	DrawScale3D(1.f,1.f,1.f)
-	,	ZOffset(0.f)
-	,	Flags(0)
+		: Location(0.f, 0.f, 0.f)
+		, Rotation(0, 0, 0)
+		, PreAlignRotation(0, 0, 0)
+		, DrawScale3D(1.f, 1.f, 1.f)
+		, ZOffset(0.f)
+		, Flags(0)
 	{}
 };
 
@@ -49,23 +51,23 @@ struct FFoliageInstancePlacementInfo
  */
 struct FFoliageInstance : public FFoliageInstancePlacementInfo
 {
-	class UActorComponent* Base;
-	int32 ClusterIndex;	// -1 if this instance is invalid and its array slot can be reused.
+	UActorComponent* Base;
+	int32 ClusterIndex;	// -1/INDEX_NONE if this instance is invalid and its array slot can be reused.
 
 	FFoliageInstance()
-	:	Base(NULL)
-	,	ClusterIndex(-1)
+		: Base(NULL)
+		, ClusterIndex(INDEX_NONE)
 	{}
 
 
-	friend FArchive& operator<<( FArchive& Ar, FFoliageInstance& Instance );
+	friend FArchive& operator<<(FArchive& Ar, FFoliageInstance& Instance);
 
 	FTransform GetInstanceWorldTransform() const
 	{
 		return FTransform(Rotation, Location, DrawScale3D);
 	}
 
-	void AlignToNormal( const FVector& InNormal, float AlignMaxAngle=0.f )
+	void AlignToNormal(const FVector& InNormal, float AlignMaxAngle = 0.f)
 	{
 		Flags |= FOLIAGE_AlignToNormal;
 
@@ -76,22 +78,22 @@ struct FFoliageInstance : public FFoliageInstancePlacementInfo
 		AlignRotation.Pitch = FRotator::NormalizeAxis(AlignRotation.Pitch);
 
 		// limit the maximum pitch angle if it's > 0.
-		if( AlignMaxAngle > 0.f )
+		if (AlignMaxAngle > 0.f)
 		{
 			int32 MaxPitch = AlignMaxAngle;
-			if( AlignRotation.Pitch > MaxPitch )
+			if (AlignRotation.Pitch > MaxPitch)
 			{
 				AlignRotation.Pitch = MaxPitch;
 			}
 			else
-			if( AlignRotation.Pitch < -MaxPitch )
-			{
+				if (AlignRotation.Pitch < -MaxPitch)
+				{
 				AlignRotation.Pitch = -MaxPitch;
-			}								
+				}
 		}
 
 		PreAlignRotation = Rotation;
-		Rotation = FRotator( FQuat(AlignRotation) * FQuat(Rotation) );
+		Rotation = FRotator(FQuat(AlignRotation) * FQuat(Rotation));
 	}
 };
 
@@ -101,20 +103,22 @@ struct FFoliageInstance : public FFoliageInstancePlacementInfo
 struct FFoliageInstanceCluster
 {
 	UInstancedStaticMeshComponent* ClusterComponent;
-
 	FBoxSphereBounds Bounds;
+
+#if WITH_EDITORONLY_DATA
 	TArray<int32> InstanceIndices;	// index into editor editor Instances array
+#endif
 
 	FFoliageInstanceCluster()
 	{}
 	FFoliageInstanceCluster(UInstancedStaticMeshComponent* InClusterComponent, const FBoxSphereBounds& InBounds)
-	:	ClusterComponent(InClusterComponent)
-	,	Bounds(InBounds)
+		: ClusterComponent(InClusterComponent)
+		, Bounds(InBounds)
 	{}
-	friend FArchive& operator<<( FArchive& Ar, FFoliageInstanceCluster& Cluster );
+	friend FArchive& operator<<(FArchive& Ar, FFoliageInstanceCluster& Cluster);
 };
 
-/** 
+/**
  * FFoliageComponentHashInfo
  * Cached instance list and component location info stored in the ComponentHash.
  * Used for moving quick updates after operations on components with foliage painted on them.
@@ -123,29 +127,29 @@ struct FFoliageComponentHashInfo
 {
 	// tors
 	FFoliageComponentHashInfo()
-	:	CachedLocation(0,0,0)
-	,	CachedRotation(0,0,0)
-	,	CachedDrawScale(1,1,1)
+		: CachedLocation(0, 0, 0)
+		, CachedRotation(0, 0, 0)
+		, CachedDrawScale(1, 1, 1)
 	{}
 
-	FFoliageComponentHashInfo( UActorComponent* InComponent )
-	:	CachedLocation(0,0,0)
-	,	CachedRotation(0,0,0)
-	,	CachedDrawScale(1,1,1)
+	FFoliageComponentHashInfo(UActorComponent* InComponent)
+		: CachedLocation(0, 0, 0)
+		, CachedRotation(0, 0, 0)
+		, CachedDrawScale(1, 1, 1)
 	{
 		UpdateLocationFromActor(InComponent);
 	}
 
 	// Cache the location and rotation from the actor
-	void UpdateLocationFromActor( UActorComponent* InComponent )
+	void UpdateLocationFromActor(UActorComponent* InComponent)
 	{
-		if( InComponent )
+		if (InComponent)
 		{
 			AActor* Owner = Cast<AActor>(InComponent->GetOuter());
-			if( Owner )
+			if (Owner)
 			{
 				const USceneComponent* RootComponent = Owner->GetRootComponent();
-				if( RootComponent )
+				if (RootComponent)
 				{
 					CachedLocation = RootComponent->RelativeLocation;
 					CachedRotation = RootComponent->RelativeRotation;
@@ -156,7 +160,7 @@ struct FFoliageComponentHashInfo
 	}
 
 	// serializer
-	friend FArchive& operator<<( FArchive& Ar, FFoliageComponentHashInfo& ComponentHashInfo );
+	friend FArchive& operator<<(FArchive& Ar, FFoliageComponentHashInfo& ComponentHashInfo);
 
 	FVector CachedLocation;
 	FRotator CachedRotation;
@@ -172,45 +176,71 @@ struct FFoliageMeshInfo
 	// Cluster allocation data and pointers to components
 	TArray<FFoliageInstanceCluster> InstanceClusters;
 
+#if WITH_EDITORONLY_DATA
 	// Editor-only placed instances
 	TArray<FFoliageInstance> Instances;
-	
+
 	// Transient, editor-only locality hash of instances
-	struct FFoliageInstanceHash* InstanceHash;
+	TUniquePtr<FFoliageInstanceHash> InstanceHash;
 
 	// Transient, editor-only set of instances per component
-	TMap<class UActorComponent*, FFoliageComponentHashInfo > ComponentHash;
+	TMap<UActorComponent*, FFoliageComponentHashInfo> ComponentHash;
 
 	// Transient, editor-only list of free instance slots.
 	TArray<int32> FreeInstanceIndices;
 
 	// Transient, editor-only list of selected instances.
 	TArray<int32> SelectedIndices;
+#endif
 
-	// UI settings
-	class UInstancedFoliageSettings* Settings;
+	ENGINE_API FFoliageMeshInfo();
 
-	FFoliageMeshInfo();
-	FFoliageMeshInfo( const FFoliageMeshInfo& Other );
+	~FFoliageMeshInfo() // =default
+	{ }
 
-	~FFoliageMeshInfo();
+	FFoliageMeshInfo(FFoliageMeshInfo&& Other)
+		// even VC++2013 doesn't support "=default" on move constructors
+		: InstanceClusters(MoveTemp(Other.InstanceClusters))
+#if WITH_EDITORONLY_DATA
+		, Instances(MoveTemp(Other.Instances))
+		, InstanceHash(MoveTemp(Other.InstanceHash))
+		, ComponentHash(MoveTemp(Other.ComponentHash))
+		, FreeInstanceIndices(MoveTemp(Other.FreeInstanceIndices))
+		, SelectedIndices(MoveTemp(Other.SelectedIndices))
+#endif
+	{ }
+
+	FFoliageMeshInfo& operator=(FFoliageMeshInfo&& Other)
+		// even VC++2013 doesn't support "=default" on move assignment
+	{
+		InstanceClusters = MoveTemp(Other.InstanceClusters);
+#if WITH_EDITORONLY_DATA
+		Instances = MoveTemp(Other.Instances);
+		InstanceHash = MoveTemp(Other.InstanceHash);
+		ComponentHash = MoveTemp(Other.ComponentHash);
+		FreeInstanceIndices = MoveTemp(Other.FreeInstanceIndices);
+		SelectedIndices = MoveTemp(Other.SelectedIndices);
+#endif
+
+		return *this;
+	}
 
 #if WITH_EDITOR
-	ENGINE_API void AddInstance( class AInstancedFoliageActor* InIFA, class UStaticMesh* InMesh, const FFoliageInstance& InNewInstance );
-	ENGINE_API void RemoveInstances( class AInstancedFoliageActor* InIFA, const TArray<int32>& InInstancesToRemove );
-	ENGINE_API void PreMoveInstances( class AInstancedFoliageActor* InIFA, const TArray<int32>& InInstancesToMove );
-	ENGINE_API void PostMoveInstances( class AInstancedFoliageActor* InIFA, const TArray<int32>& InInstancesMoved );
-	ENGINE_API void PostUpdateInstances( class AInstancedFoliageActor* InIFA, const TArray<int32>& InInstancesUpdated, bool bReAddToHash=false );
-	ENGINE_API void DuplicateInstances( class AInstancedFoliageActor* InIFA, class UStaticMesh* InMesh, const TArray<int32>& InInstancesToDuplicate );
-	ENGINE_API void GetInstancesInsideSphere( const FSphere& Sphere, TArray<int32>& OutInstances );
-	ENGINE_API bool CheckForOverlappingSphere( const FSphere& Sphere );
-	ENGINE_API bool CheckForOverlappingInstanceExcluding( int32 TestInstanceIdx, float Radius, TSet<int32>& ExcludeInstances );
+	ENGINE_API void AddInstance(AInstancedFoliageActor* InIFA, UFoliageType* InSettings, const FFoliageInstance& InNewInstance);
+	ENGINE_API void RemoveInstances(AInstancedFoliageActor* InIFA, const TArray<int32>& InInstancesToRemove);
+	ENGINE_API void PreMoveInstances(AInstancedFoliageActor* InIFA, const TArray<int32>& InInstancesToMove);
+	ENGINE_API void PostMoveInstances(AInstancedFoliageActor* InIFA, const TArray<int32>& InInstancesMoved);
+	ENGINE_API void PostUpdateInstances(AInstancedFoliageActor* InIFA, const TArray<int32>& InInstancesUpdated, bool bReAddToHash = false);
+	ENGINE_API void DuplicateInstances(AInstancedFoliageActor* InIFA, UFoliageType* InSettings, const TArray<int32>& InInstancesToDuplicate);
+	ENGINE_API void GetInstancesInsideSphere(const FSphere& Sphere, TArray<int32>& OutInstances);
+	ENGINE_API bool CheckForOverlappingSphere(const FSphere& Sphere);
+	ENGINE_API bool CheckForOverlappingInstanceExcluding(int32 TestInstanceIdx, float Radius, TSet<int32>& ExcludeInstances);
 
 	// Destroy existing clusters and reassign all instances to new clusters
-	ENGINE_API void ReallocateClusters( class AInstancedFoliageActor* InIFA, class UStaticMesh* InMesh );
+	ENGINE_API void ReallocateClusters(AInstancedFoliageActor* InIFA, UFoliageType* InSettings);
 	// Update settings in the clusters based on the current settings (eg culling distance, collision, ...)
-	ENGINE_API void UpdateClusterSettings( class AInstancedFoliageActor* InIFA );
-	ENGINE_API void SelectInstances( class AInstancedFoliageActor* InIFA, bool bSelect, TArray<int32>& Instances );
+	//ENGINE_API void UpdateClusterSettings(AInstancedFoliageActor* InIFA);
+	ENGINE_API void SelectInstances(AInstancedFoliageActor* InIFA, bool bSelect, TArray<int32>& Instances);
 
 	// Get the number of placed instances
 	ENGINE_API int32 GetInstanceCount() const;
@@ -219,7 +249,11 @@ struct FFoliageMeshInfo
 	void CheckValid();
 #endif
 
-	friend FArchive& operator<<( FArchive& Ar, FFoliageMeshInfo& MeshInfo );
+	friend FArchive& operator<<(FArchive& Ar, FFoliageMeshInfo& MeshInfo);
+
+//private:
+	const FFoliageMeshInfo(const FFoliageMeshInfo& Other) = delete;
+	const FFoliageMeshInfo& operator=(const FFoliageMeshInfo& Other) = delete;
 };
 
 //
@@ -230,23 +264,24 @@ struct FFoliageMeshInfo
 
 struct FFoliageInstanceHash
 {
-	int32 HashCellBits;
 private:
+	const int32 HashCellBits;
+	TMultiMap<uint64, int32> CellMap;
 
-	uint64 MakeKey( int32 CellX, int32 CellY )
+	uint64 MakeKey(int32 CellX, int32 CellY)
 	{
 		return ((uint64)(*(uint32*)(&CellX)) << 32) | (*(uint32*)(&CellY) & 0xffffffff);
 	}
 
-	uint64 MakeKey( const FVector& Location )
+	uint64 MakeKey(const FVector& Location)
 	{
-		return  MakeKey( FMath::FloorToInt(Location.X) >> HashCellBits, FMath::FloorToInt(Location.Y) >> HashCellBits );
+		return  MakeKey(FMath::FloorToInt(Location.X) >> HashCellBits, FMath::FloorToInt(Location.Y) >> HashCellBits);
 	}
 
 	// Locality map
-	TMap<uint64, TSet<int32> > CellMap;
+	//TMap<uint64, TSet<int32>> CellMap;
 public:
-	FFoliageInstanceHash( int32 InHashCellBits = FOLIAGE_HASH_CELL_BITS )
+	FFoliageInstanceHash(int32 InHashCellBits = FOLIAGE_HASH_CELL_BITS)
 	:	HashCellBits(InHashCellBits)
 	{}
 
@@ -254,50 +289,45 @@ public:
 	{
 		uint64 Key = MakeKey(InstanceLocation);
 
-		TSet<int32>& NewSet = CellMap.FindOrAdd(Key);
-		NewSet.Add(InstanceIndex);
+		CellMap.AddUnique(Key, InstanceIndex);
 	}
 
 	void RemoveInstance(const FVector& InstanceLocation, int32 InstanceIndex)
 	{
 		uint64 Key = MakeKey(InstanceLocation);
-
-		TSet<int32>* Set = CellMap.Find(Key);
-		check( Set != NULL );
-		int32 RemoveCount = Set->Remove(InstanceIndex);
-		check( RemoveCount == 1 );
+		
+		int32 RemoveCount = CellMap.RemoveSingle(Key, InstanceIndex);
+		check(RemoveCount == 1);
 	}
 
-	void GetInstancesOverlappingBox( const FBox& InBox, TSet<int32>& OutInstanceIndices )
+	void GetInstancesOverlappingBox(const FBox& InBox, TArray<int32>& OutInstanceIndices)
 	{
 		int32 MinX = FMath::FloorToInt(InBox.Min.X) >> HashCellBits;
 		int32 MinY = FMath::FloorToInt(InBox.Min.Y) >> HashCellBits;
 		int32 MaxX = FMath::FloorToInt(InBox.Max.X) >> HashCellBits;
 		int32 MaxY = FMath::FloorToInt(InBox.Max.Y) >> HashCellBits;
 
-		for( int32 y=MinY;y<=MaxY;y++ )
+		for (int32 y = MinY; y <= MaxY; y++)
 		{
-			for( int32 x=MinX;x<=MaxX;x++ )
+			for (int32 x = MinX; x <= MaxX; x++)
 			{
-				uint64 Key = MakeKey(x,y);
-				TSet<int32>* Set = CellMap.Find(Key);
-				if( Set != NULL )
-				{
-					OutInstanceIndices.Append( *Set );
-				}
+				uint64 Key = MakeKey(x, y);
+				CellMap.MultiFind(Key, OutInstanceIndices);
 			}
 		}
 	}
 
+	TArray<int32> GetInstancesOverlappingBox(const FBox& InBox)
+	{
+		TArray<int32> Result;
+		GetInstancesOverlappingBox(InBox, Result);
+		return Result;
+	}
+
 #if UE_BUILD_DEBUG
 	void CheckInstanceCount(int32 InCount)
 	{
-		int32 HashCount = 0;
-		for( TMap<uint64, TSet<int32> >::TConstIterator It(CellMap); It; ++It )
-		{
-			HashCount += It.Value().Num();
-		}	
-		check( HashCount == InCount );
+		check(CellMap.Num() == InCount);
 	}
 #endif
 
@@ -306,9 +336,27 @@ public:
 		CellMap.Empty();
 	}
 
-	friend FArchive& operator<<( FArchive& Ar, FFoliageInstanceHash& Hash )
+	friend FArchive& operator<<(FArchive& Ar, FFoliageInstanceHash& Hash)
 	{
-		Ar << Hash.CellMap;
+		if (Ar.UE4Ver() < VER_UE4_FOLIAGE_SETTINGS_TYPE)
+		{
+			Hash.CellMap.Reset();
+
+			TMap<uint64, TSet<int32>> OldCellMap;
+			Ar << OldCellMap;
+			for (auto& CellPair : OldCellMap)
+			{
+				for (int32 Idx : CellPair.Value)
+				{
+					Hash.CellMap.AddUnique(CellPair.Key, Idx);
+				}
+			}
+		}
+		else
+		{
+			Ar << Hash.CellMap;
+		}
+
 		return Ar;
 	}
 };
@@ -320,10 +368,10 @@ struct HInstancedStaticMeshInstance : public HHitProxy
 	UInstancedStaticMeshComponent* Component;
 	int32 InstanceIndex;
 
-	DECLARE_HIT_PROXY( ENGINE_API );
-	HInstancedStaticMeshInstance(UInstancedStaticMeshComponent* InComponent, int32 InInstanceIndex): HHitProxy(HPP_World), Component(InComponent), InstanceIndex(InInstanceIndex) {}
+	DECLARE_HIT_PROXY(ENGINE_API);
+	HInstancedStaticMeshInstance(UInstancedStaticMeshComponent* InComponent, int32 InInstanceIndex) : HHitProxy(HPP_World), Component(InComponent), InstanceIndex(InInstanceIndex) {}
 
-	virtual void AddReferencedObjects( FReferenceCollector& Collector ) override;
+	virtual void AddReferencedObjects(FReferenceCollector& Collector) override;
 
 	virtual EMouseCursor::Type GetMouseCursor()
 	{
diff --git a/Engine/Source/Runtime/Engine/Public/Landscape/LandscapeEdit.h b/Engine/Source/Runtime/Engine/Public/Landscape/LandscapeEdit.h
--- a/Engine/Source/Runtime/Engine/Public/Landscape/LandscapeEdit.h
+++ b/Engine/Source/Runtime/Engine/Public/Landscape/LandscapeEdit.h
@@ -11,7 +11,7 @@ LandscapeEdit.h: Classes for the editor to access to Landscape data
 
 #if WITH_EDITOR
 
-#include "Landscape/LandscapeProxy.h"
+#include "Landscape/LandscapeProxy.h" // for ELandscapeLayerPaintingRestriction
 
 struct FLandscapeTextureDataInfo
 {
