[GitHub] - CheckWaterJump() fix when colliding with BSP.

pull request: https://github.com/EpicGames/UnrealEngine/pull/249

diff --git a/Engine/Source/Runtime/Engine/Private/CharacterMovementComponent.cpp b/Engine/Source/Runtime/Engine/Private/CharacterMovementComponent.cpp
--- a/Engine/Source/Runtime/Engine/Private/CharacterMovementComponent.cpp
+++ b/Engine/Source/Runtime/Engine/Private/CharacterMovementComponent.cpp
@@ -3701,7 +3701,7 @@ bool UCharacterMovementComponent::CheckWaterJump(FVector CheckPoint, FVector& Wa
 	const ECollisionChannel CollisionChannel = UpdatedComponent->GetCollisionObjectType();
 	bool bHit = GetWorld()->SweepSingle( HitInfo, CharacterOwner->GetActorLocation(), CheckPoint, FQuat::Identity, CollisionChannel, CapsuleShape, CapsuleParams, ResponseParam);
 	
-	if ( HitInfo.GetActor() && !Cast<APawn>(HitInfo.GetActor()) )
+	if ( bHit && !Cast<APawn>(HitInfo.GetActor()) )
 	{
 		// hit a wall - check if it is low enough
 		WallNormal = -1.f * HitInfo.ImpactNormal;
