Fixed comment

diff --git a/Engine/Source/Runtime/Engine/Classes/GameFramework/Actor.h b/Engine/Source/Runtime/Engine/Classes/GameFramework/Actor.h
--- a/Engine/Source/Runtime/Engine/Classes/GameFramework/Actor.h
+++ b/Engine/Source/Runtime/Engine/Classes/GameFramework/Actor.h
@@ -1999,7 +1999,7 @@ public:
 	/** Returns true if this is a replicated actor that was placed in the map */
 	bool IsNetStartupActor() const;
 
-	/** Searches components array and returns first encountered array by type. */
+	/** Searches components array and returns first encountered component of the specified class. */
 	virtual UActorComponent* FindComponentByClass(const TSubclassOf<UActorComponent> ComponentClass) const;
 	
 	/** Script exposed version of FindComponentByClass */
