Fix FFoliageMeshInfo for old compilers without "=delete"

diff --git a/Engine/Source/Runtime/Engine/Public/InstancedFoliage.h b/Engine/Source/Runtime/Engine/Public/InstancedFoliage.h
--- a/Engine/Source/Runtime/Engine/Public/InstancedFoliage.h
+++ b/Engine/Source/Runtime/Engine/Public/InstancedFoliage.h
@@ -251,9 +251,14 @@ struct FFoliageMeshInfo
 
 	friend FArchive& operator<<(FArchive& Ar, FFoliageMeshInfo& MeshInfo);
 
-//private:
+#if PLATFORM_COMPILER_HAS_DEFAULTED_FUNCTIONS
 	const FFoliageMeshInfo(const FFoliageMeshInfo& Other) = delete;
 	const FFoliageMeshInfo& operator=(const FFoliageMeshInfo& Other) = delete;
+#else
+private:
+	const FFoliageMeshInfo(const FFoliageMeshInfo& Other);
+	const FFoliageMeshInfo& operator=(const FFoliageMeshInfo& Other);
+#endif
 };
 
 //
