Add warning when trying to simulate physics with no collision preset

diff --git a/Engine/Source/Runtime/Engine/Private/PrimitiveComponent.cpp b/Engine/Source/Runtime/Engine/Private/PrimitiveComponent.cpp
--- a/Engine/Source/Runtime/Engine/Private/PrimitiveComponent.cpp
+++ b/Engine/Source/Runtime/Engine/Private/PrimitiveComponent.cpp
@@ -789,6 +789,12 @@ bool UPrimitiveComponent::ShouldCreatePhysicsState() const
 	bool bShouldCreatePhysicsState = IsRegistered() && (bAlwaysCreatePhysicsState || IsCollisionEnabled());
 
 #if WITH_EDITOR
+	if (BodyInstance.bSimulatePhysics && GetCollisionEnabled() == ECollisionEnabled::NoCollision)
+	{
+		FMessageLog("PIE").Warning(FText::Format(LOCTEXT("InvalidSimulateOptions", "Invalid Simulate Options: Body ({0}) is set to simulate physics but Collision Enabled is set to No Collision"),
+			FText::FromString(GetReadableName())));
+	}
+
 	// if it shouldn't create physics state, but if world wants to enable trace collision for components, allow it
 	if (!bShouldCreatePhysicsState && World && World->bEnableTraceCollision)
 	{
