Support for async loading of network objects, work in progress

diff --git a/Engine/Source/Runtime/Engine/Classes/Engine/ActorChannel.h b/Engine/Source/Runtime/Engine/Classes/Engine/ActorChannel.h
--- a/Engine/Source/Runtime/Engine/Classes/Engine/ActorChannel.h
+++ b/Engine/Source/Runtime/Engine/Classes/Engine/ActorChannel.h
@@ -68,6 +68,10 @@ class ENGINE_API UActorChannel : public UChannel
 
 	TMap< TWeakObjectPtr< UObject >, TSharedRef< FObjectReplicator > > ReplicationMap;
 
+	TArray< class FInBunch * > QueuedBunches;		// Queued bunches waiting on pending guids to resolve
+
+	TSet< FNetworkGUID > PendingGuidResolves;		// These guids are waiting for their resolves, we need to queue up bunches until these are resolved
+
 	/**
 	 * Default constructor
 	 */
@@ -86,6 +90,9 @@ class ENGINE_API UActorChannel : public UChannel
 	virtual void SetClosingFlag() override;
 
 	virtual void ReceivedBunch( FInBunch& Bunch ) override;
+
+	virtual void Tick() override;
+	void ProcessBunch( FInBunch & Bunch );
 	
 	virtual void ReceivedNak( int32 NakPacketId ) override;
 	
diff --git a/Engine/Source/Runtime/Engine/Classes/Engine/PackageMapClient.h b/Engine/Source/Runtime/Engine/Classes/Engine/PackageMapClient.h
--- a/Engine/Source/Runtime/Engine/Classes/Engine/PackageMapClient.h
+++ b/Engine/Source/Runtime/Engine/Classes/Engine/PackageMapClient.h
@@ -22,27 +22,12 @@
 #include "Net/DataBunch.h"
 #include "PackageMapClient.generated.h"
 
-class FMarkSerializeNewActorScoped
-{
-public:
-	FMarkSerializeNewActorScoped( bool & InIsSerializingNewActor ) : IsSerializingNewActorRef( InIsSerializingNewActor )
-	{ 
-		IsSerializingNewActorRef = true; 
-	};
 
-	~FMarkSerializeNewActorScoped()
-	{
-		IsSerializingNewActorRef = false;
-	}
-
-	bool & IsSerializingNewActorRef;
-};
-	
 /** Stores an object with path associated with FNetworkGUID */
 class FNetGuidCacheObject
 {
 public:
-	FNetGuidCacheObject() : GuidSequence( 0 ), bNoLoad( 0 ), bIgnoreWhenMissing( 0 ), bIsPending( 0 ), bIsBroken( 0 )
+	FNetGuidCacheObject() : GuidSequence( 0 ), InitialQueryTime( -1.0f ), bNoLoad( 0 ), bIgnoreWhenMissing( 0 ), bIsPending( 0 ), bIsBroken( 0 )
 	{
 	}
 
@@ -53,6 +38,8 @@ public:
 	FName						PathName;
 	int32						GuidSequence;				// We remember the guid sequence this net guid was generated on so we can tell how old it is
 
+	float						InitialQueryTime;
+
 	uint8						bNoLoad				: 1;	// Don't load this, only do a find
 	uint8						bIgnoreWhenMissing	: 1;	// Don't warn when this asset can't be found or loaded
 	uint8						bIsPending			: 1;	// This object is waiting to be fully loaded
@@ -76,6 +63,9 @@ public:
 	void			RegisterNetGUIDFromPath_Client( const FNetworkGUID & NetGUID, const FString & PathName, const FNetworkGUID & OuterGUID, const bool bNoLoad, const bool bIgnoreWhenMissing );
 	UObject *		GetObjectFromNetGUID( const FNetworkGUID & NetGUID );
 	bool			ShouldIgnoreWhenMissing( const FNetworkGUID & NetGUID );
+	bool			IsGUIDRegistered( const FNetworkGUID & NetGUID );
+	bool			IsGUIDLoaded( const FNetworkGUID & NetGUID );
+
 	UObject *		ResolvePath( const FString & PathName, UObject * ObjOuter, const bool bNoLoad );
 
 	TMap< FNetworkGUID, FNetGuidCacheObject >		ObjectLookup;
@@ -83,6 +73,7 @@ public:
 	int32											UniqueNetIDs[2];
 
 	int32											GuidSequence;
+	bool											IsExportingNetGUIDBunch;
 
 	UNetDriver *									Driver;
 
@@ -103,8 +94,6 @@ class UPackageMapClient : public UPackageMap
 	{
 		GuidCache				= InNetGUIDCache;
 		ExportNetGUIDCount		= 0;
-		IsExportingNetGUIDBunch = false;
-		IsSerializingNewActor	= false;
 	}
 
 	virtual ~UPackageMapClient()
@@ -148,6 +137,11 @@ class UPackageMapClient : public UPackageMap
 	virtual void		LogDebugInfo( FOutputDevice & Ar) override;
 	virtual UObject *	GetObjectFromNetGUID( const FNetworkGUID & NetGUID ) override;
 
+	TArray< FNetworkGUID > & GetMustBeMappedGuidsInLastBunch()
+	{
+		return MustBeMappedGuidsInLastBunch;
+	}
+
 protected:
 
 	bool	ExportNetGUID( FNetworkGUID NetGUID, const UObject * Object, FString PathName, UObject * ObjOuter );
@@ -162,8 +156,6 @@ protected:
 	
 	bool IsNetGUIDAuthority();
 
-	bool IsExportingNetGUIDBunch;
-
 	class UNetConnection* Connection;
 
 	bool ObjectLevelHasFinishedLoading(UObject* Obj);
@@ -181,5 +173,5 @@ protected:
 
 	TSharedPtr< FNetGUIDCache >			GuidCache;
 
-	bool								IsSerializingNewActor;
+	TArray< FNetworkGUID >				MustBeMappedGuidsInLastBunch;
 };
diff --git a/Engine/Source/Runtime/Engine/Private/DataBunch.cpp b/Engine/Source/Runtime/Engine/Private/DataBunch.cpp
--- a/Engine/Source/Runtime/Engine/Private/DataBunch.cpp
+++ b/Engine/Source/Runtime/Engine/Private/DataBunch.cpp
@@ -102,9 +102,10 @@ FOutBunch::FOutBunch( UChannel* InChannel, bool bInClose )
 ,	bDormant	( 0 )
 ,	bReliable	( 0 )
 ,	bPartial	( 0 )
-,	bPartialInitial ( 0 )
-,	bPartialFinal	( 0 )
-,	bHasGUIDs ( 0 )
+,	bPartialInitial			( 0 )
+,	bPartialFinal			( 0 )
+,	bHasGUIDs				( 0 )
+,	bHasMustBeMappedGUIDs	( 0 )
 {
 	checkSlow(!Channel->Closing);
 	checkSlow(Channel->Connection->Channels[Channel->ChIndex]==Channel);
@@ -136,7 +137,8 @@ FOutBunch::FOutBunch( UPackageMap *InPackageMap, int64 MaxBits )
 ,	bPartial	( 0 )
 ,	bPartialInitial ( 0 )
 ,	bPartialFinal	( 0 )
-,	bHasGUIDs ( 0 )
+,	bHasGUIDs				( 0 )
+,	bHasMustBeMappedGUIDs	( 0 )
 {
 }
 
diff --git a/Engine/Source/Runtime/Engine/Private/DataChannel.cpp b/Engine/Source/Runtime/Engine/Private/DataChannel.cpp
--- a/Engine/Source/Runtime/Engine/Private/DataChannel.cpp
+++ b/Engine/Source/Runtime/Engine/Private/DataChannel.cpp
@@ -269,9 +269,9 @@ bool UChannel::ReceivedSequencedBunch( FInBunch& Bunch )
 void UChannel::ReceivedRawBunch( FInBunch & Bunch, bool & bOutSkipAck )
 {
 	// Immediately consume the NetGUID portion of this bunch, regardless if it is partial or reliable.
-	if (Bunch.bHasGUIDs)
+	if ( Bunch.bHasGUIDs )
 	{
-		Cast<UPackageMapClient>(Connection->PackageMap)->ReceiveNetGUIDBunch(Bunch);
+		Cast<UPackageMapClient>( Connection->PackageMap )->ReceiveNetGUIDBunch( Bunch );
 
 		if ( Bunch.IsError() )
 		{
@@ -280,7 +280,6 @@ void UChannel::ReceivedRawBunch( FInBunch & Bunch, bool & bOutSkipAck )
 		}
 	}
 
-
 	check(Connection->Channels[ChIndex]==this);
 
 	if ( Bunch.bReliable && Bunch.ChSequence != Connection->InReliable[ChIndex] + 1 )
@@ -459,10 +458,15 @@ bool UChannel::ReceivedNextBunch( FInBunch & Bunch, bool & bOutSkipAck )
 						UE_LOG(LogNetPartialBunch, Verbose, TEXT("Completed Partial Bunch: Channel: %d ChSequence: %d. Num: %d Rel: %d CRC 0x%X"), InPartialBunch->ChIndex, InPartialBunch->ChSequence, InPartialBunch->GetNumBits(), Bunch.bReliable, FCrc::MemCrc_DEPRECATED(InPartialBunch->GetData(), InPartialBunch->GetNumBytes()));
 					}
 
+					check( !Bunch.bHasGUIDs );		// Shouldn't have these, they only go in initial partial export bunches
+
 					HandleBunch = InPartialBunch;
-					InPartialBunch->bPartialFinal = true;
-					InPartialBunch->bClose = Bunch.bClose;
-					InPartialBunch->bDormant = Bunch.bDormant;
+
+					InPartialBunch->bPartialFinal			= true;
+					InPartialBunch->bClose					= Bunch.bClose;
+					InPartialBunch->bDormant				= Bunch.bDormant;
+					InPartialBunch->bHasMustBeMappedGUIDs	= Bunch.bHasMustBeMappedGUIDs;
+
 					if (InPartialBunch->bOpen)
 					{
 						UE_LOG(LogNetPartialBunch, Verbose, TEXT("Channel: %d is now Open! [%d] - [%d]"), ChIndex, OpenPacketId.First, OpenPacketId.Last );
@@ -532,6 +536,7 @@ FPacketIdRange UChannel::SendBunch( FOutBunch* Bunch, bool Merge )
 	check(!Closing);
 	check(Connection->Channels[ChIndex]==this);
 	check(!Bunch->IsError());
+	check( !Bunch->bHasGUIDs );
 
 	// Set bunch flags.
 	if( OpenPacketId.First==INDEX_NONE && OpenedLocally )
@@ -554,13 +559,47 @@ FPacketIdRange UChannel::SendBunch( FOutBunch* Bunch, bool Merge )
 
 	TArray<FOutBunch *> OutgoingBunches;
 
+	UPackageMapClient * PackageMapClient = CastChecked< UPackageMapClient >( Connection->PackageMap );
+
 	// Let the package map add any outgoing bunches it needs to send
-	if (Cast<UPackageMapClient>(Connection->PackageMap)->AppendExportBunches(OutgoingBunches))
+	if ( PackageMapClient->AppendExportBunches( OutgoingBunches ) )
 	{
 		// Dont merge if we have multiple bunches to send.
 		Merge = false;
 	}
 
+	// Append any "must be mapped" guids to front of bunch
+	// This is so the client will know to make sure they are loaded, and pause the stream until they are
+	TArray< FNetworkGUID > & MustBeMappedGuidsInLastBunch = PackageMapClient->GetMustBeMappedGuidsInLastBunch();
+
+	if ( MustBeMappedGuidsInLastBunch.Num() > 0 && Connection->Driver->IsServer() )
+	{
+		//UE_LOG( LogNet, Warning, TEXT( "SendBunch: Appending must be mapped guids. NumGuids: %i" ), MustBeMappedGuidsInLastBunch.Num() );
+
+		// Rewrite the bunch with the unique guids in front
+		FOutBunch TempBunch( *Bunch );
+
+		Bunch->Reset();
+
+		// Write all the guids out
+		uint16 NumMustBeMappedGUIDs = MustBeMappedGuidsInLastBunch.Num();
+		*Bunch << NumMustBeMappedGUIDs;
+		for ( int32 i = 0; i < MustBeMappedGuidsInLastBunch.Num(); i++ )
+		{
+			*Bunch << MustBeMappedGuidsInLastBunch[i];
+		}
+
+		// Append the original bunch data at the end
+		Bunch->SerializeBits( TempBunch.GetData(), TempBunch.GetNumBits() );
+
+		Bunch->bHasMustBeMappedGUIDs = 1;
+
+		// We can't merge with this, since we need all the unique static guids in the front
+		Merge = false;
+	}
+
+	MustBeMappedGuidsInLastBunch.Empty();
+
 	//-----------------------------------------------------
 	// Contemplate merging.
 	//-----------------------------------------------------
@@ -596,7 +635,7 @@ FPacketIdRange UChannel::SendBunch( FOutBunch* Bunch, bool Merge )
 		uint8 *data = Bunch->GetData();
 		int64 bitsLeft = Bunch->GetNumBits();
 		Merge = false;
-		
+
 		while(bitsLeft > 0)
 		{
 			FOutBunch * PartialBunch = new FOutBunch(this, false);
@@ -655,6 +694,11 @@ FPacketIdRange UChannel::SendBunch( FOutBunch* Bunch, bool Merge )
 		NextBunch->ChIndex = Bunch->ChIndex;
 		NextBunch->ChType = Bunch->ChType;
 
+		if ( !NextBunch->bHasGUIDs )
+		{
+			NextBunch->bHasMustBeMappedGUIDs |= Bunch->bHasMustBeMappedGUIDs;
+		}
+
 		if (OutgoingBunches.Num() > 1)
 		{
 			NextBunch->bPartial = 1;
@@ -1316,6 +1360,17 @@ void UActorChannel::CleanUp()
 		Connection->SentTemporaries.Remove(Actor);
 	}
 
+	// We don't care about any leftover pending guids at this point
+	PendingGuidResolves.Empty();
+
+	// Free any queued bunches
+	for ( int32 i = 0; i < QueuedBunches.Num(); i++ )
+	{
+		delete QueuedBunches[i];
+	}
+
+	QueuedBunches.Empty();
+
 	Super::CleanUp();
 }
 
@@ -1417,15 +1472,86 @@ void UActorChannel::SetChannelActorForDestroy( FActorDestructionInfo *DestructIn
 	}
 }
 
-void UActorChannel::ReceivedBunch( FInBunch& Bunch )
+void UActorChannel::Tick()
 {
-	check(!Closing);
+	Super::Tick();
+	
+	// Try to resolve any guids that are holding up the network stream on this channel
+	Connection->Driver->GuidCache->IsExportingNetGUIDBunch = true;		// Hack to make assert happy
+	for ( auto It = PendingGuidResolves.CreateIterator(); It; ++It )
+	{
+		if ( Connection->Driver->GuidCache->GetObjectFromNetGUID( *It ) != NULL )
+		{
+			// This guid is now resolved, we can remove it from the pending guid list
+			It.RemoveCurrent();
+		}
+	}
+	Connection->Driver->GuidCache->IsExportingNetGUIDBunch = false;
+
+	// If we don't have any pending guids to resolve, process any queued bunches now
+	if ( PendingGuidResolves.Num() == 0 && QueuedBunches.Num() > 0 )
+	{
+		UE_LOG( LogNet, Warning, TEXT( "UActorChannel::Tick: Purging queued bunches: ChIndex: %i, Actor: %s, Queued: %i" ), ChIndex, Actor != NULL ? *Actor->GetPathName() : TEXT( "NULL" ), QueuedBunches.Num() );
+
+		for ( int32 i = 0; i < QueuedBunches.Num(); i++ )
+		{
+			ProcessBunch( *QueuedBunches[i] );
+			delete QueuedBunches[i];
+		}
+
+		QueuedBunches.Empty();
+	}
+}
+
+void UActorChannel::ReceivedBunch( FInBunch & Bunch )
+{
+	check( !Closing );
+
 	if ( Broken || bTornOff )
 	{
 		return;
 	}
 
-	const bool bIsServer = (Connection->Driver->ServerConnection == NULL);
+	if ( Bunch.bHasMustBeMappedGUIDs )
+	{
+		// If this bunch has any guids that must be mapped, we need to wait until they resolve before we can 
+		// process the rest of the stream on this channel
+		uint16 NumMustBeMappedGUIDs = 0;
+		Bunch << NumMustBeMappedGUIDs;
+
+		//UE_LOG( LogNetTraffic, Warning, TEXT( "Read must be mapped GUID's. NumMustBeMappedGUIDs: %i" ), NumMustBeMappedGUIDs );
+
+		UPackageMapClient * PackageMapClient = CastChecked< UPackageMapClient >( Connection->PackageMap );
+
+		for ( int32 i = 0; i < NumMustBeMappedGUIDs; i++ )
+		{
+			FNetworkGUID NetGUID;
+			Bunch << NetGUID;
+
+			// This GUID better have been exported before we get here, which means it must be registered by now
+			check( Connection->Driver->GuidCache->IsGUIDRegistered( NetGUID ) );
+
+			if ( !Connection->Driver->GuidCache->IsGUIDLoaded( NetGUID ) )
+			{
+				PendingGuidResolves.Add( NetGUID );
+			}
+		}
+	}
+
+	// If we have guids that need to be resolved, or we have existing pending bunches, we must process this bunch later
+	if ( PendingGuidResolves.Num() > 0 || QueuedBunches.Num() > 0 )
+	{
+		QueuedBunches.Add( new FInBunch( Bunch ) );
+		return;
+	}
+
+	// We can process this bunch now
+	ProcessBunch( Bunch );
+}
+
+void UActorChannel::ProcessBunch( FInBunch & Bunch )
+{
+	const bool bIsServer = Connection->Driver->IsServer();
 
 	FReplicationFlags RepFlags;
 
diff --git a/Engine/Source/Runtime/Engine/Private/NetConnection.cpp b/Engine/Source/Runtime/Engine/Private/NetConnection.cpp
--- a/Engine/Source/Runtime/Engine/Private/NetConnection.cpp
+++ b/Engine/Source/Runtime/Engine/Private/NetConnection.cpp
@@ -832,16 +832,17 @@ void UNetConnection::ReceivedPacket( FBitReader& Reader )
 		{
 			// Parse the incoming data.
 			FInBunch Bunch( this );
-			int32 IncomingStartPos = Reader.GetPosBits();
-			uint8 bControl      = Reader.ReadBit();
-			Bunch.PacketId     = PacketId;
-			Bunch.bOpen        = bControl ? Reader.ReadBit() : 0;
-			Bunch.bClose       = bControl ? Reader.ReadBit() : 0;
-			Bunch.bDormant	   = Bunch.bClose ? Reader.ReadBit() : 0;
-			Bunch.bReliable    = Reader.ReadBit();
-			Bunch.ChIndex      = Reader.ReadInt( MAX_CHANNELS );
-			Bunch.bHasGUIDs	   = Reader.ReadBit();
-			Bunch.bPartial	   = Reader.ReadBit();
+			int32 IncomingStartPos		= Reader.GetPosBits();
+			uint8 bControl				= Reader.ReadBit();
+			Bunch.PacketId				= PacketId;
+			Bunch.bOpen					= bControl ? Reader.ReadBit() : 0;
+			Bunch.bClose				= bControl ? Reader.ReadBit() : 0;
+			Bunch.bDormant				= Bunch.bClose ? Reader.ReadBit() : 0;
+			Bunch.bReliable				= Reader.ReadBit();
+			Bunch.ChIndex				= Reader.ReadInt( MAX_CHANNELS );
+			Bunch.bHasGUIDs				= Reader.ReadBit();
+			Bunch.bHasMustBeMappedGUIDs	= Reader.ReadBit();
+			Bunch.bPartial				= Reader.ReadBit();
 
 			if ( Bunch.bReliable )
 			{
@@ -884,7 +885,8 @@ void UNetConnection::ReceivedPacket( FBitReader& Reader )
 				UE_LOG(LogNetTraffic, VeryVerbose, TEXT("Received: %s"), *Bunch.ToString());
 			}
 
-			int32 HeaderPos      = Reader.GetPosBits();
+			const int32 HeaderPos = Reader.GetPosBits();
+
 			if( Reader.IsError() )
 			{
 				UE_LOG( LogNetTraffic, Error, TEXT( "Bunch header overflowed" ) );
@@ -1188,6 +1190,7 @@ int32 UNetConnection::SendRawBunch( FOutBunch& Bunch, bool InAllowMerge )
 	Header.WriteBit( Bunch.bReliable );
 	Header.WriteIntWrapped(Bunch.ChIndex, MAX_CHANNELS);
 	Header.WriteBit( Bunch.bHasGUIDs );
+	Header.WriteBit( Bunch.bHasMustBeMappedGUIDs );
 	Header.WriteBit( Bunch.bPartial );
 	if (Bunch.bReliable || Bunch.bPartial)
 	{
diff --git a/Engine/Source/Runtime/Engine/Private/PackageMapClient.cpp b/Engine/Source/Runtime/Engine/Private/PackageMapClient.cpp
--- a/Engine/Source/Runtime/Engine/Private/PackageMapClient.cpp
+++ b/Engine/Source/Runtime/Engine/Private/PackageMapClient.cpp
@@ -30,6 +30,9 @@ static const int INTERNAL_LOAD_OBJECT_RECURSION_LIMIT = 16;
 #define COMPOSE_RELATIVE_NET_GUID( GUID, Index )	( COMPOSE_NET_GUID( GET_NET_GUID_INDEX( GUID ) + Index, GET_NET_GUID_STATIC_MOD( GUID ) ) )
 #endif
 
+static TAutoConsoleVariable<int32> CVarAllowAsyncLoading( TEXT( "net.AllowAsyncLoading" ), 0, TEXT( "Allow async loading" ) );
+static TAutoConsoleVariable<int32> CVarSimulateAsyncLoading( TEXT( "net.SimulateAsyncLoading" ), 0, TEXT( "Simulate async loading" ) );
+
 /*-----------------------------------------------------------------------------
 	UPackageMapClient implementation.
 -----------------------------------------------------------------------------*/
@@ -169,9 +172,6 @@ bool UPackageMapClient::SerializeNewActor(FArchive& Ar, class UActorChannel *Cha
 {
 	UE_LOG( LogNetPackageMap, VeryVerbose, TEXT( "SerializeNewActor START" ) );
 
-	// This will mark IsSerializingNewActor true for the scope of the function
-	FMarkSerializeNewActorScoped SerializeNewActorScoped( IsSerializingNewActor );
-
 	uint8 bIsClosingChannel = 0;
 
 	if (Ar.IsLoading() )
@@ -421,9 +421,51 @@ struct FExportFlags
 	}
 };
 
+static bool CanClientLoadObject_Server( const UObject * Object, const FNetworkGUID & NetGUID )
+{
+	if ( !NetGUID.IsValid() || NetGUID.IsDynamic() )
+	{
+		return false;
+	}
+
+	const UPackage * Package = Cast< UPackage >( Object );
+
+	if ( Package != NULL )
+	{
+		if ( Package->GetFName() != Package->FileName )
+		{
+			return false;		// Client can't load renamed packages
+		}
+	}
+
+	// We never allow the client to load levels on demand, we wait until the client has streamed the level in
+	while ( Object != NULL )
+	{
+		if ( Cast< const ULevel >( Object ) != NULL )
+		{
+			return false;
+		}
+
+		Object = Object->GetOuter();
+	}
+
+	// We can load everything else
+	return true;
+}
+
 /** Writes an object NetGUID given the NetGUID and either the object itself, or FString full name of the object. Appends full name/path if necessary */
 void UPackageMapClient::InternalWriteObject( FArchive & Ar, FNetworkGUID NetGUID, const UObject * Object, FString ObjectPathName, UObject * ObjectOuter )
 {
+	check( Ar.IsSaving() );
+
+	if ( CVarAllowAsyncLoading.GetValueOnGameThread() > 0 && IsNetGUIDAuthority() && !GuidCache->IsExportingNetGUIDBunch && CanClientLoadObject_Server( Object, NetGUID ) )
+	{
+		// These are guids that must exist on the client in a package
+		// The client needs to know about these so it can determine if it has finished loading them
+		// and pause the network stream for that channel if it hasn't
+		MustBeMappedGuidsInLastBunch.AddUnique( NetGUID );
+	}
+
 	Ar << NetGUID;
 	NET_CHECKSUM( Ar );
 
@@ -439,10 +481,15 @@ void UPackageMapClient::InternalWriteObject( FArchive & Ar, FNetworkGUID NetGUID
 
 	if ( NetGUID.IsDefault() )
 	{
+		// Only the client sends default guids
+		check( !IsNetGUIDAuthority() );
 		ExportFlags.bHasPath = 1;
 	}		
-	else if ( IsExportingNetGUIDBunch )
+	else if ( GuidCache->IsExportingNetGUIDBunch )
 	{
+		// Only the server should be exporting guids
+		check( IsNetGUIDAuthority() );
+
 		if ( Object != NULL )
 		{
 			ExportFlags.bHasPath = ShouldSendFullPath( Object, NetGUID ) ? 1 : 0;
@@ -485,7 +532,7 @@ void UPackageMapClient::InternalWriteObject( FArchive & Ar, FNetworkGUID NetGUID
 		// Serialize Name of object
 		Ar << ObjectPathName;
 
-		if ( IsExportingNetGUIDBunch )
+		if ( GuidCache->IsExportingNetGUIDBunch )
 		{
 			CurrentExportNetGUIDs.Add( NetGUID );
 
@@ -575,7 +622,7 @@ FNetworkGUID UPackageMapClient::InternalLoadObject( FArchive & Ar, UObject *& Ob
 	{
 		ExportFlags.bHasPath = 1;
 	}
-	else if ( IsExportingNetGUIDBunch )
+	else if ( GuidCache->IsExportingNetGUIDBunch )
 	{
 		Ar << ExportFlags.Value;
 
@@ -752,11 +799,11 @@ bool UPackageMapClient::ExportNetGUID( FNetworkGUID NetGUID, const UObject * Obj
 		FBitWriterMark LastExportMark;
 		LastExportMark.Init( *CurrentExportBunch );
 
-		IsExportingNetGUIDBunch = true;
+		GuidCache->IsExportingNetGUIDBunch = true;
 
 		InternalWriteObject( *CurrentExportBunch, NetGUID, Object, PathName, ObjOuter );
 
-		IsExportingNetGUIDBunch = false;
+		GuidCache->IsExportingNetGUIDBunch = false;
 
 		if ( CurrentExportNetGUIDs.Num() == 0 )
 		{
@@ -846,7 +893,7 @@ void UPackageMapClient::ExportNetGUIDHeader()
 
 void UPackageMapClient::ReceiveNetGUIDBunch( FInBunch &InBunch )
 {
-	IsExportingNetGUIDBunch = true;
+	GuidCache->IsExportingNetGUIDBunch = true;
 	check(InBunch.bHasGUIDs);
 
 	int32 NumGUIDsInBunch = 0;
@@ -860,6 +907,7 @@ void UPackageMapClient::ReceiveNetGUIDBunch( FInBunch &InBunch )
 	{
 		// Older sequence, ignore
 		UE_LOG( LogNetPackageMap, Warning, TEXT( "UPackageMapClient::ReceiveNetGUIDBunch: NewGuidSequence < GuidSequence (%i / %i)" ), GuidCache->GuidSequence, NewGuidSequence );
+		GuidCache->IsExportingNetGUIDBunch = false;
 		return;
 	}
 	else if ( NewGuidSequence > GuidCache->GuidSequence )
@@ -876,6 +924,7 @@ void UPackageMapClient::ReceiveNetGUIDBunch( FInBunch &InBunch )
 	{
 		UE_LOG( LogNetPackageMap, Error, TEXT( "UPackageMapClient::ReceiveNetGUIDBunch: NumGUIDsInBunch > MAX_GUID_COUNT (%i)" ), NumGUIDsInBunch );
 		InBunch.SetError();
+		GuidCache->IsExportingNetGUIDBunch = false;
 		return;
 	}
 
@@ -892,13 +941,14 @@ void UPackageMapClient::ReceiveNetGUIDBunch( FInBunch &InBunch )
 		if ( InBunch.IsError() )
 		{
 			UE_LOG( LogNetPackageMap, Error, TEXT( "UPackageMapClient::ReceiveNetGUIDBunch: InBunch.IsError() after InternalLoadObject" ) );
+			GuidCache->IsExportingNetGUIDBunch = false;
 			return;
 		}
 		NumGUIDsRead++;
 	}
 
 	UE_LOG(LogNetPackageMap, Log, TEXT("UPackageMapClient::ReceiveNetGUIDBunch end. BitPos: %d"), InBunch.GetPosBits() );
-	IsExportingNetGUIDBunch = false;
+	GuidCache->IsExportingNetGUIDBunch = false;
 }
 
 bool UPackageMapClient::AppendExportBunches(TArray<FOutBunch *>& OutgoingBunches)
@@ -1280,6 +1330,7 @@ FNetGUIDCache::FNetGUIDCache( UNetDriver * InDriver ) : Driver( InDriver )
 {
 	UniqueNetIDs[0] = UniqueNetIDs[1] = 0;
 	GuidSequence = 0;
+	IsExportingNetGUIDBunch = false;
 }
 
 class FArchiveCountMemGUID : public FArchive
@@ -1713,7 +1764,7 @@ void FNetGUIDCache::RegisterNetGUIDFromPath_Client( const FNetworkGUID & NetGUID
 			}
 		}
 
-		ObjectLookup.Remove( NetGUID );
+		return;
 	}
 
 	// Register a new guid with this path
@@ -1807,10 +1858,37 @@ UObject * FNetGUIDCache::GetObjectFromNetGUID( const FNetworkGUID & NetGUID )
 			{
 				UE_LOG( LogNetPackageMap, Warning, TEXT( "GetObjectFromNetGUID: Failed to find outer. Path: %s, NetGUID: %s" ), *CacheObjectPtr->PathName.ToString(), *NetGUID.ToString() );
 			}
+
 			return NULL;
 		}
 	}
 
+	if ( CVarAllowAsyncLoading.GetValueOnGameThread() > 0 && CVarSimulateAsyncLoading.GetValueOnGameThread() > 0 && !IsNetGUIDAuthority() && !CacheObjectPtr->OuterGUID.IsValid() )
+	{
+		if ( !CacheObjectPtr->bNoLoad && !CacheObjectPtr->bIsPending && CacheObjectPtr->InitialQueryTime <= 0.0f )
+		{
+			CacheObjectPtr->InitialQueryTime = FPlatformTime::Seconds();
+			CacheObjectPtr->bIsPending = true;
+		}
+
+		if ( CacheObjectPtr->bIsPending && CacheObjectPtr->InitialQueryTime >= 0.0f )
+		{
+			if ( FPlatformTime::Seconds() - CacheObjectPtr->InitialQueryTime > (float)CVarSimulateAsyncLoading.GetValueOnGameThread() )
+			{
+				//UE_LOG( LogNetPackageMap, Warning, TEXT( "GetObjectFromNetGUID: FINISH simulating pending load. Path: %s, NetGUID: %s" ), *CacheObjectPtr->PathName.ToString(), *NetGUID.ToString() );
+				CacheObjectPtr->bIsPending = false;
+			}
+			else
+			{
+				// We shouldn't get here unless we are trying to import
+				// The reason is because the actor channel code should be using the "must be mapped" list to pause the stream
+				// until these guids are resolved
+				check( IsExportingNetGUIDBunch );
+				return NULL;
+			}
+		}
+	}
+
 	// Resolve the path of this object relative to outer
 	Object = ResolvePath( CacheObjectPtr->PathName.ToString(), ObjOuter, CacheObjectPtr->bNoLoad || CacheObjectPtr->bIsPending );
 
@@ -1875,6 +1953,43 @@ bool FNetGUIDCache::ShouldIgnoreWhenMissing( const FNetworkGUID & NetGUID )
 	return CacheObject->bIgnoreWhenMissing || CacheObject->bIsPending;
 }
 
+bool FNetGUIDCache::IsGUIDRegistered( const FNetworkGUID & NetGUID )
+{
+	if ( !NetGUID.IsValid() )
+	{
+		return false;
+	}
+
+	if ( NetGUID.IsDefault() )
+	{
+		return false;
+	}
+
+	return ObjectLookup.Contains( NetGUID );
+}
+
+bool FNetGUIDCache::IsGUIDLoaded( const FNetworkGUID & NetGUID )
+{
+	if ( !NetGUID.IsValid() )
+	{
+		return false;
+	}
+
+	if ( NetGUID.IsDefault() )
+	{
+		return false;
+	}
+
+	FNetGuidCacheObject * CacheObjectPtr = ObjectLookup.Find( NetGUID );
+
+	if ( CacheObjectPtr == NULL )
+	{
+		return false;
+	}
+
+	return CacheObjectPtr->Object != NULL;
+}
+
 UObject * FNetGUIDCache::ResolvePath( const FString & PathName, UObject * ObjOuter, const bool bNoLoad )
 {
 	UObject * Object = StaticFindObject( UObject::StaticClass(), ObjOuter, *PathName, false );
diff --git a/Engine/Source/Runtime/Engine/Public/Net/DataBunch.h b/Engine/Source/Runtime/Engine/Public/Net/DataBunch.h
--- a/Engine/Source/Runtime/Engine/Public/Net/DataBunch.h
+++ b/Engine/Source/Runtime/Engine/Public/Net/DataBunch.h
@@ -24,12 +24,13 @@ public:
 	uint8					bClose;
 	uint8					bDormant;
 	uint8					bReliable;
-	uint8					bPartial;			// Not a complete bunch
-	uint8					bPartialInitial;	// The first bunch of a partial bunch
-	uint8					bPartialFinal;		// The final bunch of a partial bunch
-	uint8					bHasGUIDs;			// This bunch has networkGUID name/id pairs
+	uint8					bPartial;				// Not a complete bunch
+	uint8					bPartialInitial;		// The first bunch of a partial bunch
+	uint8					bPartialFinal;			// The final bunch of a partial bunch
+	uint8					bHasGUIDs;				// This bunch has networkGUID name/id pairs
+	uint8					bHasMustBeMappedGUIDs;	// This bunch has guids that must be mapped before we can process this bunch
 
-	TArray< FNetworkGUID >	ExportNetGUIDs;		// List of GUIDs that went out on this bunch
+	TArray< FNetworkGUID >	ExportNetGUIDs;			// List of GUIDs that went out on this bunch
 
 #if !(UE_BUILD_SHIPPING || UE_BUILD_TEST)
 	FString			DebugString;
@@ -89,19 +90,20 @@ class FInBunch : public FNetBitReader
 public:
 	// Variables.
 	int32				PacketId;	// Note this must stay as first member variable in FInBunch for FInBunch(FInBunch, bool) to work
-	FInBunch*		Next;
-	UNetConnection*	Connection;
+	FInBunch *			Next;
+	UNetConnection *	Connection;
 	int32				ChIndex;
 	int32				ChType;
 	int32				ChSequence;
-	uint8			bOpen;
-	uint8			bClose;
-	uint8			bDormant;		// Close, but go dormant
-	uint8			bReliable;
-	uint8			bPartial;		// Not a complete bunch
-	uint8			bPartialInitial;// The first bunch of a partial bunch
-	uint8			bPartialFinal;	// The final bunch of a partial bunch
-	uint8			bHasGUIDs; // This bunch has networkGUID name/id pairs
+	uint8				bOpen;
+	uint8				bClose;
+	uint8				bDormant;				// Close, but go dormant
+	uint8				bReliable;
+	uint8				bPartial;				// Not a complete bunch
+	uint8				bPartialInitial;		// The first bunch of a partial bunch
+	uint8				bPartialFinal;			// The final bunch of a partial bunch
+	uint8				bHasGUIDs;				// This bunch has networkGUID name/id pairs
+	uint8				bHasMustBeMappedGUIDs;	// This bunch has guids that must be mapped before we can process this bunch
 
 	FString	ToString()
 	{
