Fixed so that when snaps, it only applies child bone(or socket)'s transform

diff --git a/Engine/Source/Runtime/Engine/Private/SceneComponent.cpp b/Engine/Source/Runtime/Engine/Private/SceneComponent.cpp
--- a/Engine/Source/Runtime/Engine/Private/SceneComponent.cpp
+++ b/Engine/Source/Runtime/Engine/Private/SceneComponent.cpp
@@ -720,7 +720,7 @@ void USceneComponent::AttachTo(class USceneComponent* Parent, FName InSocketName
 		
 
 		// Make sure we are detached
-		bool bMaintainWorldPosition = AttachType == EAttachLocation::KeepWorldPosition || EAttachLocation::SnapToTarget;
+		bool bMaintainWorldPosition = AttachType == EAttachLocation::KeepWorldPosition;
 		DetachFromParent(bMaintainWorldPosition);
 
 		// Detach removes all Prerequisite, so will need to add after Detach happens
@@ -781,8 +781,11 @@ void USceneComponent::AttachTo(class USceneComponent* Parent, FName InSocketName
 				}
 				else
 				{
-					FTransform ParentToWorld = AttachParent->GetSocketTransform(AttachSocketName);
-					FTransform RelativeTM = ComponentToWorld.GetRelativeTransform(ParentToWorld);
+					// when snap, we'd like to give socket or bone scale only
+					// to do so, get parent and get socket relative to parent
+					FTransform ParentToWorld = AttachParent->GetComponentToWorld();
+					FTransform SocketTransform = AttachParent->GetSocketTransform(AttachSocketName);
+					FTransform RelativeTM = SocketTransform.GetRelativeTransform(ParentToWorld);
 					RelativeScale3D = RelativeTM.GetScale3D();
 				}
 			}
