UObject improvement: DependsOn in engine code
Removing DependsOn and replacing them with proper #includes.
#ttp 312783

diff --git a/Engine/Plugins/Experimental/LiveEditor/Source/LiveEditor/Classes/K2Node_LiveEditObject.h b/Engine/Plugins/Experimental/LiveEditor/Source/LiveEditor/Classes/K2Node_LiveEditObject.h
--- a/Engine/Plugins/Experimental/LiveEditor/Source/LiveEditor/Classes/K2Node_LiveEditObject.h
+++ b/Engine/Plugins/Experimental/LiveEditor/Source/LiveEditor/Classes/K2Node_LiveEditObject.h
@@ -1,9 +1,12 @@
 // Copyright 1998-2014 Epic Games, Inc. All Rights Reserved.
 
 #pragma once
+
+#include "LiveEditorTypes.h"
+
 #include "K2Node_LiveEditObject.generated.h"
 
-UCLASS(DependsOn=ULiveEditorTypes,MinimalAPI)
+UCLASS(MinimalAPI)
 class UK2Node_LiveEditObject : public UK2Node
 {
 	GENERATED_UCLASS_BODY()
diff --git a/Engine/Plugins/Experimental/LiveEditor/Source/LiveEditor/Classes/LiveEditorBlueprint.h b/Engine/Plugins/Experimental/LiveEditor/Source/LiveEditor/Classes/LiveEditorBlueprint.h
--- a/Engine/Plugins/Experimental/LiveEditor/Source/LiveEditor/Classes/LiveEditorBlueprint.h
+++ b/Engine/Plugins/Experimental/LiveEditor/Source/LiveEditor/Classes/LiveEditorBlueprint.h
@@ -1,9 +1,11 @@
 // Copyright 1998-2014 Epic Games, Inc. All Rights Reserved.
 #pragma once
 
+#include "LiveEditorTypes.h"
+
 #include "LiveEditorBlueprint.generated.h"
 
-UCLASS(Abstract, Blueprintable,DependsOn=ULiveEditorTypes)
+UCLASS(Abstract, Blueprintable)
 class ULiveEditorBlueprint : public UObject
 {
 	GENERATED_UCLASS_BODY()
diff --git a/Engine/Plugins/Experimental/LiveEditor/Source/LiveEditor/Classes/LiveEditorBroadcaster.h b/Engine/Plugins/Experimental/LiveEditor/Source/LiveEditor/Classes/LiveEditorBroadcaster.h
--- a/Engine/Plugins/Experimental/LiveEditor/Source/LiveEditor/Classes/LiveEditorBroadcaster.h
+++ b/Engine/Plugins/Experimental/LiveEditor/Source/LiveEditor/Classes/LiveEditorBroadcaster.h
@@ -1,11 +1,14 @@
 // Copyright 1998-2014 Epic Games, Inc. All Rights Reserved.
 
 #pragma once
+
+#include "LiveEditorTypes.h"
+
 #include "LiveEditorBroadcaster.generated.h"
 
 DECLARE_DYNAMIC_MULTICAST_DELEGATE_ThreeParams( FLiveEditorEventMIDIMultiCast, int32, Delta, int32, MidiValue, ELiveEditControllerType::Type, ControlType );
 
-UCLASS(DependsOn=ULiveEditorTypes,MinimalAPI)
+UCLASS(MinimalAPI)
 class ULiveEditorBroadcaster : public UObject
 {
 	GENERATED_UCLASS_BODY()
diff --git a/Engine/Plugins/Experimental/LiveEditor/Source/LiveEditor/Classes/LiveEditorKismetLibrary.h b/Engine/Plugins/Experimental/LiveEditor/Source/LiveEditor/Classes/LiveEditorKismetLibrary.h
--- a/Engine/Plugins/Experimental/LiveEditor/Source/LiveEditor/Classes/LiveEditorKismetLibrary.h
+++ b/Engine/Plugins/Experimental/LiveEditor/Source/LiveEditor/Classes/LiveEditorKismetLibrary.h
@@ -2,9 +2,11 @@
 
 #pragma once
 
+#include "LiveEditorTypes.h"
+
 #include "LiveEditorKismetLibrary.generated.h"
 
-UCLASS(DependsOn=ULiveEditorTypes)
+UCLASS()
 class ULiveEditorKismetLibrary : public UBlueprintFunctionLibrary
 {
 	GENERATED_UCLASS_BODY()
diff --git a/Engine/Plugins/Experimental/Paper2D/Source/Paper2D/Classes/PaperRenderComponent.h b/Engine/Plugins/Experimental/Paper2D/Source/Paper2D/Classes/PaperRenderComponent.h
--- a/Engine/Plugins/Experimental/Paper2D/Source/Paper2D/Classes/PaperRenderComponent.h
+++ b/Engine/Plugins/Experimental/Paper2D/Source/Paper2D/Classes/PaperRenderComponent.h
@@ -2,9 +2,11 @@
 
 #pragma once
 
+#include "PaperSprite.h"
+
 #include "PaperRenderComponent.generated.h"
 
-UCLASS(DependsOn=UPaperSprite, MinimalAPI, ShowCategories=(Mobility), EarlyAccessPreview, meta=(BlueprintSpawnableComponent))
+UCLASS(MinimalAPI, ShowCategories=(Mobility), EarlyAccessPreview, meta=(BlueprintSpawnableComponent))
 class UPaperRenderComponent : public UPrimitiveComponent
 {
 	GENERATED_UCLASS_BODY()
diff --git a/Engine/Plugins/Experimental/Paper2D/Source/Paper2D/Classes/PaperSprite.h b/Engine/Plugins/Experimental/Paper2D/Source/Paper2D/Classes/PaperSprite.h
--- a/Engine/Plugins/Experimental/Paper2D/Source/Paper2D/Classes/PaperSprite.h
+++ b/Engine/Plugins/Experimental/Paper2D/Source/Paper2D/Classes/PaperSprite.h
@@ -3,6 +3,8 @@
 #pragma once
 
 #include "SpriteEditorOnlyTypes.h"
+#include "Engine/EngineTypes.h"
+
 #include "PaperSprite.generated.h"
 
 UENUM()
@@ -37,7 +39,7 @@ struct FPaperSpriteSocket
  * Can also contain collision shapes for the sprite.
  */
 
-UCLASS(DependsOn=UEngineTypes, BlueprintType, meta=(DisplayThumbnail = "true"))
+UCLASS(BlueprintType, meta=(DisplayThumbnail = "true"))
 class PAPER2D_API UPaperSprite : public UObject, public IInterface_CollisionDataProvider
 {
 	GENERATED_UCLASS_BODY()
diff --git a/Engine/Plugins/Experimental/Paper2D/Source/Paper2D/Classes/PaperTerrainActor.h b/Engine/Plugins/Experimental/Paper2D/Source/Paper2D/Classes/PaperTerrainActor.h
--- a/Engine/Plugins/Experimental/Paper2D/Source/Paper2D/Classes/PaperTerrainActor.h
+++ b/Engine/Plugins/Experimental/Paper2D/Source/Paper2D/Classes/PaperTerrainActor.h
@@ -2,13 +2,15 @@
 
 #pragma once
 
+#include "Engine/EngineTypes.h"
+
 #include "PaperTerrainActor.generated.h"
 
 /**
  * An instance of a piece of 2D terrain in the level
  */
 
-UCLASS(DependsOn=UEngineTypes, BlueprintType, Experimental)
+UCLASS(BlueprintType, Experimental)
 class PAPER2D_API APaperTerrainActor : public AActor
 {
 	GENERATED_UCLASS_BODY()
diff --git a/Engine/Plugins/Experimental/Paper2D/Source/Paper2D/Classes/PaperTileLayer.h b/Engine/Plugins/Experimental/Paper2D/Source/Paper2D/Classes/PaperTileLayer.h
--- a/Engine/Plugins/Experimental/Paper2D/Source/Paper2D/Classes/PaperTileLayer.h
+++ b/Engine/Plugins/Experimental/Paper2D/Source/Paper2D/Classes/PaperTileLayer.h
@@ -2,6 +2,8 @@
 
 #pragma once
 
+#include "PaperTileSet.h"
+
 #include "PaperTileLayer.generated.h"
 
 USTRUCT(BlueprintType)
@@ -32,7 +34,7 @@ struct FPaperTileInfo
 	}
 };
 
-UCLASS(DependsOn=UPaperTileSet)
+UCLASS()
 class PAPER2D_API UPaperTileLayer : public UObject
 {
 	GENERATED_UCLASS_BODY()
diff --git a/Engine/Plugins/Experimental/Paper2D/Source/Paper2D/Classes/PaperTileMap.h b/Engine/Plugins/Experimental/Paper2D/Source/Paper2D/Classes/PaperTileMap.h
--- a/Engine/Plugins/Experimental/Paper2D/Source/Paper2D/Classes/PaperTileMap.h
+++ b/Engine/Plugins/Experimental/Paper2D/Source/Paper2D/Classes/PaperTileMap.h
@@ -2,10 +2,11 @@
 
 #pragma once
 
+#include "Engine/EngineTypes.h"
 #include "PaperSprite.h"
 #include "PaperTileMap.generated.h"
 
-UCLASS(DependsOn=UEngineTypes)
+UCLASS()
 class PAPER2D_API UPaperTileMap : public UDataAsset //@TODO: Just to make it easy to spawn for now
 {
 	GENERATED_UCLASS_BODY()
diff --git a/Engine/Plugins/Experimental/Paper2D/Source/Paper2D/Classes/PaperTileMapActor.h b/Engine/Plugins/Experimental/Paper2D/Source/Paper2D/Classes/PaperTileMapActor.h
--- a/Engine/Plugins/Experimental/Paper2D/Source/Paper2D/Classes/PaperTileMapActor.h
+++ b/Engine/Plugins/Experimental/Paper2D/Source/Paper2D/Classes/PaperTileMapActor.h
@@ -2,9 +2,11 @@
 
 #pragma once
 
+#include "Engine/EngineTypes.h"
+
 #include "PaperTileMapActor.generated.h"
 
-UCLASS(DependsOn=UEngineTypes, MinimalAPI)
+UCLASS(MinimalAPI)
 class APaperTileMapActor : public AActor
 {
 	GENERATED_UCLASS_BODY()
diff --git a/Engine/Plugins/Experimental/Paper2D/Source/Paper2D/Classes/PaperTileMapRenderComponent.h b/Engine/Plugins/Experimental/Paper2D/Source/Paper2D/Classes/PaperTileMapRenderComponent.h
--- a/Engine/Plugins/Experimental/Paper2D/Source/Paper2D/Classes/PaperTileMapRenderComponent.h
+++ b/Engine/Plugins/Experimental/Paper2D/Source/Paper2D/Classes/PaperTileMapRenderComponent.h
@@ -2,9 +2,11 @@
 
 #pragma once
 
+#include "PaperSprite.h"
+
 #include "PaperTileMapRenderComponent.generated.h"
 
-UCLASS(DependsOn=UPaperSprite, hideCategories=Object, Experimental)
+UCLASS(hideCategories=Object, Experimental)
 class PAPER2D_API UPaperTileMapRenderComponent : public UPrimitiveComponent
 {
 	GENERATED_UCLASS_BODY()
diff --git a/Engine/Plugins/Experimental/Paper2D/Source/Paper2D/Classes/PaperTileSet.h b/Engine/Plugins/Experimental/Paper2D/Source/Paper2D/Classes/PaperTileSet.h
--- a/Engine/Plugins/Experimental/Paper2D/Source/Paper2D/Classes/PaperTileSet.h
+++ b/Engine/Plugins/Experimental/Paper2D/Source/Paper2D/Classes/PaperTileSet.h
@@ -2,9 +2,11 @@
 
 #pragma once
 
+#include "Engine/EngineTypes.h"
+
 #include "PaperTileSet.generated.h"
 
-UCLASS(DependsOn=UEngineTypes)
+UCLASS()
 class PAPER2D_API UPaperTileSet : public UDataAsset //@TODO: Just to make it easy to spawn for now
 {
 	GENERATED_UCLASS_BODY()
diff --git a/Engine/Source/Editor/AnimGraph/Classes/AnimStateTransitionNode.h b/Engine/Source/Editor/AnimGraph/Classes/AnimStateTransitionNode.h
--- a/Engine/Source/Editor/AnimGraph/Classes/AnimStateTransitionNode.h
+++ b/Engine/Source/Editor/AnimGraph/Classes/AnimStateTransitionNode.h
@@ -1,10 +1,13 @@
 // Copyright 1998-2014 Epic Games, Inc. All Rights Reserved.
 
 #pragma once
+
+#include "Animation/AnimInstance.h"
 #include "AnimStateNodeBase.h"
+
 #include "AnimStateTransitionNode.generated.h"
 
-UCLASS(MinimalAPI, DependsOn=UAnimInstance, config=Editor)
+UCLASS(MinimalAPI, config=Editor)
 class UAnimStateTransitionNode : public UAnimStateNodeBase
 {
 	GENERATED_UCLASS_BODY()
diff --git a/Engine/Source/Editor/UnrealEd/Classes/MaterialEditor/DEditorFontParameterValue.h b/Engine/Source/Editor/UnrealEd/Classes/MaterialEditor/DEditorFontParameterValue.h
--- a/Engine/Source/Editor/UnrealEd/Classes/MaterialEditor/DEditorFontParameterValue.h
+++ b/Engine/Source/Editor/UnrealEd/Classes/MaterialEditor/DEditorFontParameterValue.h
@@ -1,10 +1,12 @@
 // Copyright 1998-2014 Epic Games, Inc. All Rights Reserved.
 
-
 #pragma once
+
+#include "Editor/UnrealEdTypes.h"
+
 #include "DEditorFontParameterValue.generated.h"
 
-UCLASS(hidecategories=Object, dependson=UUnrealEdTypes, collapsecategories)
+UCLASS(hidecategories=Object, collapsecategories)
 class UNREALED_API UDEditorFontParameterValue : public UDEditorParameterValue
 {
 	GENERATED_UCLASS_BODY()
diff --git a/Engine/Source/Editor/UnrealEd/Classes/MaterialEditor/DEditorParameterValue.h b/Engine/Source/Editor/UnrealEd/Classes/MaterialEditor/DEditorParameterValue.h
--- a/Engine/Source/Editor/UnrealEd/Classes/MaterialEditor/DEditorParameterValue.h
+++ b/Engine/Source/Editor/UnrealEd/Classes/MaterialEditor/DEditorParameterValue.h
@@ -1,10 +1,12 @@
 // Copyright 1998-2014 Epic Games, Inc. All Rights Reserved.
 
-
 #pragma once
+
+#include "Editor/UnrealEdTypes.h"
+
 #include "DEditorParameterValue.generated.h"
 
-UCLASS(hidecategories=Object, dependson=UUnrealEdTypes, collapsecategories, editinlinenew)
+UCLASS(hidecategories=Object, collapsecategories, editinlinenew)
 class UNREALED_API UDEditorParameterValue : public UObject
 {
 	GENERATED_UCLASS_BODY()
diff --git a/Engine/Source/Editor/UnrealEd/Classes/MaterialEditor/DEditorScalarParameterValue.h b/Engine/Source/Editor/UnrealEd/Classes/MaterialEditor/DEditorScalarParameterValue.h
--- a/Engine/Source/Editor/UnrealEd/Classes/MaterialEditor/DEditorScalarParameterValue.h
+++ b/Engine/Source/Editor/UnrealEd/Classes/MaterialEditor/DEditorScalarParameterValue.h
@@ -1,10 +1,12 @@
 // Copyright 1998-2014 Epic Games, Inc. All Rights Reserved.
 
-
 #pragma once
+
+#include "Editor/UnrealEdTypes.h"
+
 #include "DEditorScalarParameterValue.generated.h"
 
-UCLASS(hidecategories=Object, dependson=UUnrealEdTypes, collapsecategories, editinlinenew)
+UCLASS(hidecategories=Object, collapsecategories, editinlinenew)
 class UNREALED_API UDEditorScalarParameterValue : public UDEditorParameterValue
 {
 	GENERATED_UCLASS_BODY()
diff --git a/Engine/Source/Editor/UnrealEd/Classes/MaterialEditor/DEditorStaticComponentMaskParameterValue.h b/Engine/Source/Editor/UnrealEd/Classes/MaterialEditor/DEditorStaticComponentMaskParameterValue.h
--- a/Engine/Source/Editor/UnrealEd/Classes/MaterialEditor/DEditorStaticComponentMaskParameterValue.h
+++ b/Engine/Source/Editor/UnrealEd/Classes/MaterialEditor/DEditorStaticComponentMaskParameterValue.h
@@ -1,7 +1,9 @@
 // Copyright 1998-2014 Epic Games, Inc. All Rights Reserved.
 
-
 #pragma once
+
+#include "Editor/UnrealEdTypes.h"
+
 #include "DEditorStaticComponentMaskParameterValue.generated.h"
 
 USTRUCT()
@@ -41,7 +43,7 @@ struct UNREALED_API FDComponentMaskParameter
 	
 };
 
-UCLASS(hidecategories=Object, dependson=UUnrealEdTypes, collapsecategories)
+UCLASS(hidecategories=Object, collapsecategories)
 class UNREALED_API UDEditorStaticComponentMaskParameterValue : public UDEditorParameterValue
 {
 	GENERATED_UCLASS_BODY()
diff --git a/Engine/Source/Editor/UnrealEd/Classes/MaterialEditor/DEditorStaticSwitchParameterValue.h b/Engine/Source/Editor/UnrealEd/Classes/MaterialEditor/DEditorStaticSwitchParameterValue.h
--- a/Engine/Source/Editor/UnrealEd/Classes/MaterialEditor/DEditorStaticSwitchParameterValue.h
+++ b/Engine/Source/Editor/UnrealEd/Classes/MaterialEditor/DEditorStaticSwitchParameterValue.h
@@ -1,10 +1,12 @@
 // Copyright 1998-2014 Epic Games, Inc. All Rights Reserved.
 
-
 #pragma once
+
+#include "Editor/UnrealEdTypes.h"
+
 #include "DEditorStaticSwitchParameterValue.generated.h"
 
-UCLASS(hidecategories=Object, dependson=UUnrealEdTypes, collapsecategories)
+UCLASS(hidecategories=Object, collapsecategories)
 class UNREALED_API UDEditorStaticSwitchParameterValue : public UDEditorParameterValue
 {
 	GENERATED_UCLASS_BODY()
diff --git a/Engine/Source/Editor/UnrealEd/Classes/MaterialEditor/DEditorTextureParameterValue.h b/Engine/Source/Editor/UnrealEd/Classes/MaterialEditor/DEditorTextureParameterValue.h
--- a/Engine/Source/Editor/UnrealEd/Classes/MaterialEditor/DEditorTextureParameterValue.h
+++ b/Engine/Source/Editor/UnrealEd/Classes/MaterialEditor/DEditorTextureParameterValue.h
@@ -1,10 +1,12 @@
 // Copyright 1998-2014 Epic Games, Inc. All Rights Reserved.
 
-
 #pragma once
+
+#include "Editor/UnrealEdTypes.h"
+
 #include "DEditorTextureParameterValue.generated.h"
 
-UCLASS(hidecategories=Object, dependson=UUnrealEdTypes, collapsecategories)
+UCLASS(hidecategories=Object, collapsecategories)
 class UNREALED_API UDEditorTextureParameterValue : public UDEditorParameterValue
 {
 	GENERATED_UCLASS_BODY()
diff --git a/Engine/Source/Editor/UnrealEd/Classes/MaterialEditor/DEditorVectorParameterValue.h b/Engine/Source/Editor/UnrealEd/Classes/MaterialEditor/DEditorVectorParameterValue.h
--- a/Engine/Source/Editor/UnrealEd/Classes/MaterialEditor/DEditorVectorParameterValue.h
+++ b/Engine/Source/Editor/UnrealEd/Classes/MaterialEditor/DEditorVectorParameterValue.h
@@ -1,10 +1,12 @@
 // Copyright 1998-2014 Epic Games, Inc. All Rights Reserved.
 
-
 #pragma once
+
+#include "Editor/UnrealEdTypes.h"
+
 #include "DEditorVectorParameterValue.generated.h"
 
-UCLASS(hidecategories=Object, dependson=UUnrealEdTypes, collapsecategories, editinlinenew)
+UCLASS(hidecategories=Object, collapsecategories, editinlinenew)
 class UNREALED_API UDEditorVectorParameterValue : public UDEditorParameterValue
 {
 	GENERATED_UCLASS_BODY()
diff --git a/Engine/Source/Editor/UnrealEd/Classes/MaterialEditor/MaterialEditorInstanceConstant.h b/Engine/Source/Editor/UnrealEd/Classes/MaterialEditor/MaterialEditorInstanceConstant.h
--- a/Engine/Source/Editor/UnrealEd/Classes/MaterialEditor/MaterialEditorInstanceConstant.h
+++ b/Engine/Source/Editor/UnrealEd/Classes/MaterialEditor/MaterialEditorInstanceConstant.h
@@ -5,7 +5,11 @@
  */
 
 #pragma once
+
+#include "Editor/UnrealEdTypes.h"
+
 #include "Materials/MaterialInstanceBasePropertyOverrides.h"
+
 #include "MaterialEditorInstanceConstant.generated.h"
 
 USTRUCT()
@@ -197,7 +201,7 @@ struct FEditorStaticComponentMaskParameterValue : public FEditorParameterValue
 	
 };
 
-UCLASS(hidecategories=Object, dependson=UUnrealEdTypes, collapsecategories)
+UCLASS(hidecategories=Object, collapsecategories)
 class UNREALED_API UMaterialEditorInstanceConstant : public UObject
 {
 	GENERATED_UCLASS_BODY()
diff --git a/Engine/Source/Editor/UnrealEd/Classes/MaterialEditor/PreviewMaterial.h b/Engine/Source/Editor/UnrealEd/Classes/MaterialEditor/PreviewMaterial.h
--- a/Engine/Source/Editor/UnrealEd/Classes/MaterialEditor/PreviewMaterial.h
+++ b/Engine/Source/Editor/UnrealEd/Classes/MaterialEditor/PreviewMaterial.h
@@ -1,10 +1,12 @@
 // Copyright 1998-2014 Epic Games, Inc. All Rights Reserved.
 
-
 #pragma once
+
+#include "Materials/Material.h"
+
 #include "PreviewMaterial.generated.h"
 
-UCLASS(dependson=UMaterial)
+UCLASS()
 class UNREALED_API UPreviewMaterial : public UMaterial
 {
 	GENERATED_UCLASS_BODY()
diff --git a/Engine/Source/Editor/UnrealEd/Classes/Preferences/LightmassOptionsObject.h b/Engine/Source/Editor/UnrealEd/Classes/Preferences/LightmassOptionsObject.h
--- a/Engine/Source/Editor/UnrealEd/Classes/Preferences/LightmassOptionsObject.h
+++ b/Engine/Source/Editor/UnrealEd/Classes/Preferences/LightmassOptionsObject.h
@@ -7,9 +7,12 @@
  */
 
 #pragma once
+
+#include "Engine/EngineTypes.h"
+
 #include "LightmassOptionsObject.generated.h"
 
-UCLASS(hidecategories=Object, dependson=UEngineTypes, editinlinenew)
+UCLASS(hidecategories=Object, editinlinenew)
 class ULightmassOptionsObject : public UObject
 {
 	GENERATED_UCLASS_BODY()
diff --git a/Engine/Source/Editor/UnrealEd/Classes/ThumbnailRendering/SceneThumbnailInfoWithPrimitive.h b/Engine/Source/Editor/UnrealEd/Classes/ThumbnailRendering/SceneThumbnailInfoWithPrimitive.h
--- a/Engine/Source/Editor/UnrealEd/Classes/ThumbnailRendering/SceneThumbnailInfoWithPrimitive.h
+++ b/Engine/Source/Editor/UnrealEd/Classes/ThumbnailRendering/SceneThumbnailInfoWithPrimitive.h
@@ -5,10 +5,13 @@
  */
 
 #pragma once
+
+#include "ThumbnailManager.h"
+
 #include "SceneThumbnailInfoWithPrimitive.generated.h"
 
 
-UCLASS(dependson=UThumbnailManager,MinimalAPI)
+UCLASS(MinimalAPI)
 class USceneThumbnailInfoWithPrimitive : public USceneThumbnailInfo
 {
 	GENERATED_UCLASS_BODY()
diff --git a/Engine/Source/Programs/UnrealHeaderTool/Private/HeaderParser.cpp b/Engine/Source/Programs/UnrealHeaderTool/Private/HeaderParser.cpp
--- a/Engine/Source/Programs/UnrealHeaderTool/Private/HeaderParser.cpp
+++ b/Engine/Source/Programs/UnrealHeaderTool/Private/HeaderParser.cpp
@@ -4229,6 +4229,8 @@ void FHeaderParser::CompileClassDeclaration(FClasses& AllClasses)
 		}
 		else if (Specifier == TEXT("dependsOn"))
 		{
+			FError::Throwf(TEXT("The dependsOn specifier is deprecated. Please use proper #include instead."));
+
 			// Make sure the syntax matches but don't do anything with it; that's handled in MakeCommandlet.cpp
 			RequireSpecifierValue(PropSpecifier);
 		}
diff --git a/Engine/Source/Runtime/Engine/Classes/AI/Navigation/NavigationGraphNode.h b/Engine/Source/Runtime/Engine/Classes/AI/Navigation/NavigationGraphNode.h
--- a/Engine/Source/Runtime/Engine/Classes/AI/Navigation/NavigationGraphNode.h
+++ b/Engine/Source/Runtime/Engine/Classes/AI/Navigation/NavigationGraphNode.h
@@ -1,9 +1,12 @@
 // Copyright 1998-2014 Epic Games, Inc. All Rights Reserved.
 
 #pragma once
+
+#include "NavigationSystem.h"
+
 #include "NavigationGraphNode.generated.h"
 
-UCLASS(config=Engine, dependson=UNavigationSystem, MinimalAPI, NotBlueprintable, abstract)
+UCLASS(config=Engine, MinimalAPI, NotBlueprintable, abstract)
 class ANavigationGraphNode : public AActor
 {
 	GENERATED_UCLASS_BODY()
diff --git a/Engine/Source/Runtime/Engine/Classes/Animation/AimOffsetBlendSpace.h b/Engine/Source/Runtime/Engine/Classes/Animation/AimOffsetBlendSpace.h
--- a/Engine/Source/Runtime/Engine/Classes/Animation/AimOffsetBlendSpace.h
+++ b/Engine/Source/Runtime/Engine/Classes/Animation/AimOffsetBlendSpace.h
@@ -8,11 +8,13 @@
 #pragma once 
 
 #include "BlendSpace.h"
+#include "AnimInstance.h"
+
 #include "AimOffsetBlendSpace.generated.h"
 
 struct FBlendSample;
 
-UCLASS(config=Engine, hidecategories=Object, dependson=UAnimInstance, MinimalAPI, BlueprintType)
+UCLASS(config=Engine, hidecategories=Object, MinimalAPI, BlueprintType)
 class UAimOffsetBlendSpace : public UBlendSpace
 {
 	GENERATED_UCLASS_BODY()
diff --git a/Engine/Source/Runtime/Engine/Classes/Animation/AimOffsetBlendSpace1D.h b/Engine/Source/Runtime/Engine/Classes/Animation/AimOffsetBlendSpace1D.h
--- a/Engine/Source/Runtime/Engine/Classes/Animation/AimOffsetBlendSpace1D.h
+++ b/Engine/Source/Runtime/Engine/Classes/Animation/AimOffsetBlendSpace1D.h
@@ -8,10 +8,12 @@
 #pragma once 
 
 #include "BlendSpace1D.h"
+#include "AnimInstance.h"
+
 #include "AimOffsetBlendSpace1D.generated.h"
 
 
-UCLASS(config=Engine, hidecategories=Object, dependson=UAnimInstance, MinimalAPI, BlueprintType)
+UCLASS(config=Engine, hidecategories=Object, MinimalAPI, BlueprintType)
 class UAimOffsetBlendSpace1D : public UBlendSpace1D
 {
 	GENERATED_UCLASS_BODY()
diff --git a/Engine/Source/Runtime/Engine/Classes/Animation/AnimBlueprint.h b/Engine/Source/Runtime/Engine/Classes/Animation/AnimBlueprint.h
--- a/Engine/Source/Runtime/Engine/Classes/Animation/AnimBlueprint.h
+++ b/Engine/Source/Runtime/Engine/Classes/Animation/AnimBlueprint.h
@@ -3,6 +3,8 @@
 #pragma once
 
 #include "Engine/Blueprint.h"
+#include "AnimInstance.h"
+
 #include "AnimBlueprint.generated.h"
 
 USTRUCT()
@@ -52,7 +54,7 @@ struct FAnimParentNodeAssetOverride
  * It can perform blending of animations, directly control the bones of the skeleton, and output a final pose
  * for a Skeletal Mesh each frame.
  */
-UCLASS(dependson=(UBlueprint, UAnimInstance), BlueprintType)
+UCLASS(BlueprintType)
 class ENGINE_API UAnimBlueprint : public UBlueprint
 {
 	GENERATED_UCLASS_BODY()
diff --git a/Engine/Source/Runtime/Engine/Classes/Animation/AnimBlueprintGeneratedClass.h b/Engine/Source/Runtime/Engine/Classes/Animation/AnimBlueprintGeneratedClass.h
--- a/Engine/Source/Runtime/Engine/Classes/Animation/AnimBlueprintGeneratedClass.h
+++ b/Engine/Source/Runtime/Engine/Classes/Animation/AnimBlueprintGeneratedClass.h
@@ -5,6 +5,8 @@
 #include "AnimStateMachineTypes.h"
 #include "AnimSequenceBase.h"
 #include "Engine/BlueprintGeneratedClass.h"
+#include "AnimBlueprint.h"
+
 #include "AnimBlueprintGeneratedClass.generated.h"
 
 struct FAnimNotifyEvent;
@@ -156,7 +158,7 @@ namespace EPropertySearchMode
 }
 #endif
 
-UCLASS(dependson=(UAnimBlueprint, UAnimStateMachineTypes))
+UCLASS()
 class ENGINE_API UAnimBlueprintGeneratedClass : public UBlueprintGeneratedClass
 {
 	GENERATED_UCLASS_BODY()
diff --git a/Engine/Source/Runtime/Engine/Classes/Animation/AnimCompress.h b/Engine/Source/Runtime/Engine/Classes/Animation/AnimCompress.h
--- a/Engine/Source/Runtime/Engine/Classes/Animation/AnimCompress.h
+++ b/Engine/Source/Runtime/Engine/Classes/Animation/AnimCompress.h
@@ -10,7 +10,7 @@
 #include "AnimSequence.h"
 #include "AnimCompress.generated.h"
 
-UCLASS(abstract, DependsOn=UAnimSequence, hidecategories=Object, MinimalAPI)
+UCLASS(abstract, hidecategories=Object, MinimalAPI)
 class UAnimCompress : public UObject
 {
 	GENERATED_UCLASS_BODY()
diff --git a/Engine/Source/Runtime/Engine/Classes/Animation/AnimInstance.h b/Engine/Source/Runtime/Engine/Classes/Animation/AnimInstance.h
--- a/Engine/Source/Runtime/Engine/Classes/Animation/AnimInstance.h
+++ b/Engine/Source/Runtime/Engine/Classes/Animation/AnimInstance.h
@@ -5,6 +5,8 @@
 #include "AnimationAsset.h"
 #include "Components/SkeletalMeshComponent.h"
 #include "AnimSequence.h"
+#include "AnimStateMachineTypes.h"
+
 #include "AnimInstance.generated.h"
 
 struct FBoneContainer;
@@ -233,7 +235,7 @@ struct FSlotEvaluationPose
 	}
 };
 
-UCLASS(transient, Blueprintable, hideCategories=AnimInstance, DependsOn(UAnimStateMachineTypes), BlueprintType)
+UCLASS(transient, Blueprintable, hideCategories=AnimInstance, BlueprintType)
 class ENGINE_API UAnimInstance : public UObject
 {
 	GENERATED_UCLASS_BODY()
diff --git a/Engine/Source/Runtime/Engine/Classes/Animation/AnimSequenceBase.h b/Engine/Source/Runtime/Engine/Classes/Animation/AnimSequenceBase.h
--- a/Engine/Source/Runtime/Engine/Classes/Animation/AnimSequenceBase.h
+++ b/Engine/Source/Runtime/Engine/Classes/Animation/AnimSequenceBase.h
@@ -266,7 +266,7 @@ enum ETypeAdvanceAnim
 	ETAA_Looped
 };
 
-UCLASS(abstract, MinimalAPI, dependson=UCurveBase, BlueprintType)
+UCLASS(abstract, MinimalAPI, BlueprintType)
 class UAnimSequenceBase : public UAnimationAsset
 {
 	GENERATED_UCLASS_BODY()
diff --git a/Engine/Source/Runtime/Engine/Classes/Animation/BlendSpace.h b/Engine/Source/Runtime/Engine/Classes/Animation/BlendSpace.h
--- a/Engine/Source/Runtime/Engine/Classes/Animation/BlendSpace.h
+++ b/Engine/Source/Runtime/Engine/Classes/Animation/BlendSpace.h
@@ -4,12 +4,14 @@
 
 #include "BlendSpaceBase.h"
 #include "AnimInterpFilter.h"
+#include "AnimInstance.h"
+
 #include "BlendSpace.generated.h"
 
 /**
  * Contains a grid of data points with weights from sample points in the space
  */
-UCLASS(config=Engine, hidecategories=Object, dependson=UAnimInstance, MinimalAPI, BlueprintType)
+UCLASS(config=Engine, hidecategories=Object, MinimalAPI, BlueprintType)
 class UBlendSpace : public UBlendSpaceBase
 {
 	GENERATED_UCLASS_BODY()
diff --git a/Engine/Source/Runtime/Engine/Classes/Animation/BlendSpace1D.h b/Engine/Source/Runtime/Engine/Classes/Animation/BlendSpace1D.h
--- a/Engine/Source/Runtime/Engine/Classes/Animation/BlendSpace1D.h
+++ b/Engine/Source/Runtime/Engine/Classes/Animation/BlendSpace1D.h
@@ -11,7 +11,7 @@
 #include "BlendSpaceBase.h"
 #include "BlendSpace1D.generated.h"
 
-UCLASS(config=Engine, hidecategories=Object, dependson=UVimInstance, MinimalAPI, BlueprintType)
+UCLASS(config=Engine, hidecategories=Object, MinimalAPI, BlueprintType)
 class UBlendSpace1D : public UBlendSpaceBase
 {
 	GENERATED_UCLASS_BODY()
diff --git a/Engine/Source/Runtime/Engine/Classes/Animation/BlendSpaceBase.h b/Engine/Source/Runtime/Engine/Classes/Animation/BlendSpaceBase.h
--- a/Engine/Source/Runtime/Engine/Classes/Animation/BlendSpaceBase.h
+++ b/Engine/Source/Runtime/Engine/Classes/Animation/BlendSpaceBase.h
@@ -195,7 +195,7 @@ namespace ENotifyTriggerMode
 /**
  * Allows multiple animations to be blended between based on input parameters
  */
-UCLASS(config=Engine, hidecategories=Object, dependson=UVimInstance, MinimalAPI, BlueprintType)
+UCLASS(config=Engine, hidecategories=Object, MinimalAPI, BlueprintType)
 class UBlendSpaceBase : public UAnimationAsset
 {
 	GENERATED_UCLASS_BODY()
diff --git a/Engine/Source/Runtime/Engine/Classes/Camera/CameraModifier_CameraShake.h b/Engine/Source/Runtime/Engine/Classes/Camera/CameraModifier_CameraShake.h
--- a/Engine/Source/Runtime/Engine/Classes/Camera/CameraModifier_CameraShake.h
+++ b/Engine/Source/Runtime/Engine/Classes/Camera/CameraModifier_CameraShake.h
@@ -7,6 +7,8 @@
 #pragma once
 
 #include "CameraTypes.h"
+#include "CameraShake.h"
+
 #include "CameraModifier_CameraShake.generated.h"
 
 USTRUCT()
@@ -83,7 +85,7 @@ struct ENGINE_API FCameraShakeInstance
 
 };
 
-UCLASS(dependson=UCameraShake, config=Camera)
+UCLASS(config=Camera)
 class ENGINE_API UCameraModifier_CameraShake : public UCameraModifier
 {
 	GENERATED_UCLASS_BODY()
diff --git a/Engine/Source/Runtime/Engine/Classes/Camera/PlayerCameraManager.h b/Engine/Source/Runtime/Engine/Classes/Camera/PlayerCameraManager.h
--- a/Engine/Source/Runtime/Engine/Classes/Camera/PlayerCameraManager.h
+++ b/Engine/Source/Runtime/Engine/Classes/Camera/PlayerCameraManager.h
@@ -3,6 +3,9 @@
 #pragma once
 
 #include "CameraTypes.h"
+#include "Engine/Scene.h"
+#include "Engine/EngineBaseTypes.h"
+
 #include "PlayerCameraManager.generated.h"
 
 UENUM()
@@ -133,7 +136,7 @@ struct FViewTargetTransitionParams
 /**
  *	Defines the point of view of a player in world space.
  */
-UCLASS(notplaceable, dependson=UEngineBaseTypes, transient, dependson=UScene, BlueprintType, Blueprintable)
+UCLASS(notplaceable, transient, BlueprintType, Blueprintable)
 class ENGINE_API APlayerCameraManager : public AActor
 {
 	GENERATED_UCLASS_BODY()
diff --git a/Engine/Source/Runtime/Engine/Classes/Components/AudioComponent.h b/Engine/Source/Runtime/Engine/Classes/Components/AudioComponent.h
--- a/Engine/Source/Runtime/Engine/Classes/Components/AudioComponent.h
+++ b/Engine/Source/Runtime/Engine/Classes/Components/AudioComponent.h
@@ -3,7 +3,10 @@
 
 #pragma once
 #include "Audio.h"
+#include "Sound/ReverbVolume.h"
 #include "Sound/SoundAttenuation.h"
+#include "Engine/EngineTypes.h"
+
 #include "AudioComponent.generated.h"
 
 /** called when we finish playing audio, either because it played to completion or because a Stop() call turned it off early */
@@ -91,7 +94,7 @@ struct FSoundParseParameters
 /**
  * Used to provide objects with audio.
  */
-UCLASS(ClassGroup=(Audio, Common), hidecategories=(Object, ActorComponent, Physics, Rendering, Mobility, LOD), ShowCategories=Trigger, dependson=(AReverbVolume, UEngineTypes, USoundAttenuation), meta=(BlueprintSpawnableComponent))
+UCLASS(ClassGroup=(Audio, Common), hidecategories=(Object, ActorComponent, Physics, Rendering, Mobility, LOD), ShowCategories=Trigger, meta=(BlueprintSpawnableComponent))
 class ENGINE_API UAudioComponent : public USceneComponent
 {
 	GENERATED_UCLASS_BODY()
diff --git a/Engine/Source/Runtime/Engine/Classes/Components/DirectionalLightComponent.h b/Engine/Source/Runtime/Engine/Classes/Components/DirectionalLightComponent.h
--- a/Engine/Source/Runtime/Engine/Classes/Components/DirectionalLightComponent.h
+++ b/Engine/Source/Runtime/Engine/Classes/Components/DirectionalLightComponent.h
@@ -1,13 +1,15 @@
 // Copyright 1998-2014 Epic Games, Inc. All Rights Reserved.
 
-
 #pragma once
+
+#include "Engine/EngineTypes.h"
+
 #include "DirectionalLightComponent.generated.h"
 
 /**
  * A light component that has parallel rays. Will provide a uniform lighting across any affected surface (eg. The Sun). This will affect all objects in the defined light-mass importance volume.
  */
-UCLASS(ClassGroup=Lights, hidecategories=(Object, LightProfiles), dependson=UEngineTypes, editinlinenew, meta=(BlueprintSpawnableComponent), MinimalAPI)
+UCLASS(ClassGroup=Lights, hidecategories=(Object, LightProfiles), editinlinenew, meta=(BlueprintSpawnableComponent), MinimalAPI)
 class UDirectionalLightComponent : public ULightComponent
 {
 	GENERATED_UCLASS_BODY()
diff --git a/Engine/Source/Runtime/Engine/Classes/Components/LightComponent.h b/Engine/Source/Runtime/Engine/Classes/Components/LightComponent.h
--- a/Engine/Source/Runtime/Engine/Classes/Components/LightComponent.h
+++ b/Engine/Source/Runtime/Engine/Classes/Components/LightComponent.h
@@ -1,8 +1,10 @@
 // Copyright 1998-2014 Epic Games, Inc. All Rights Reserved.
 
-
 #pragma once
+
+#include "LightComponentBase.h"
 #include "SceneTypes.h"
+
 #include "LightComponent.generated.h"
 
 UCLASS(abstract, HideCategories=(Trigger,Activation,"Components|Activation",Physics), ShowCategories=(Mobility))
diff --git a/Engine/Source/Runtime/Engine/Classes/Components/LightComponentBase.h b/Engine/Source/Runtime/Engine/Classes/Components/LightComponentBase.h
--- a/Engine/Source/Runtime/Engine/Classes/Components/LightComponentBase.h
+++ b/Engine/Source/Runtime/Engine/Classes/Components/LightComponentBase.h
@@ -1,7 +1,9 @@
 // Copyright 1998-2014 Epic Games, Inc. All Rights Reserved.
 
-
 #pragma once
+
+#include "Engine/Texture2D.h"
+
 #include "LightComponentBase.generated.h"
 
 UCLASS(abstract, HideCategories=(Trigger,Activation,"Components|Activation",Physics), ShowCategories=(Mobility))
diff --git a/Engine/Source/Runtime/Engine/Classes/Components/PrimitiveComponent.h b/Engine/Source/Runtime/Engine/Classes/Components/PrimitiveComponent.h
--- a/Engine/Source/Runtime/Engine/Classes/Components/PrimitiveComponent.h
+++ b/Engine/Source/Runtime/Engine/Classes/Components/PrimitiveComponent.h
@@ -5,9 +5,13 @@
 #include "StaticArray.h"
 #include "PhysicsEngine/BodyInstance.h"
 #include "Components/SceneComponent.h"
+#include "Components/LightComponent.h"
 #include "Materials/MaterialInterface.h"
 #include "SceneTypes.h"
 #include "CollisionQueryParams.h"
+#include "Engine/Scene.h"
+#include "Engine/EngineTypes.h"
+
 #include "PrimitiveComponent.generated.h"
 
 class FPrimitiveSceneProxy;
@@ -116,7 +120,7 @@ DECLARE_DYNAMIC_MULTICAST_DELEGATE_TwoParams( FComponentOnInputTouchEndSignature
 DECLARE_DYNAMIC_MULTICAST_DELEGATE_TwoParams( FComponentBeginTouchOverSignature, ETouchIndex::Type, FingerIndex, UPrimitiveComponent*, TouchedComponent );
 DECLARE_DYNAMIC_MULTICAST_DELEGATE_TwoParams( FComponentEndTouchOverSignature, ETouchIndex::Type, FingerIndex, UPrimitiveComponent*, TouchedComponent );
 
-UCLASS(dependson=(UScene, ULightComponent, UEngineTypes), abstract, HideCategories=(Mobility))
+UCLASS(abstract, HideCategories=(Mobility))
 class ENGINE_API UPrimitiveComponent : public USceneComponent
 {
 	GENERATED_UCLASS_BODY()
@@ -1438,4 +1442,3 @@ public:
 	void DispatchOnInputTouchBegin(const ETouchIndex::Type Key);
 	void DispatchOnInputTouchEnd(const ETouchIndex::Type Key);
 };
-
diff --git a/Engine/Source/Runtime/Engine/Classes/Components/StaticMeshComponent.h b/Engine/Source/Runtime/Engine/Classes/Components/StaticMeshComponent.h
--- a/Engine/Source/Runtime/Engine/Classes/Components/StaticMeshComponent.h
+++ b/Engine/Source/Runtime/Engine/Classes/Components/StaticMeshComponent.h
@@ -1,10 +1,12 @@
 // Copyright 1998-2014 Epic Games, Inc. All Rights Reserved.
 
-
 #pragma once
+
 #include "SceneTypes.h"
 #include "Components/MeshComponent.h"
 #include "Runtime/RenderCore/Public/PackedNormal.h"
+#include "Lightmass/LightmassPrimitiveSettingsObject.h"
+
 #include "StaticMeshComponent.generated.h"
 
 class FColorVertexBuffer;
@@ -96,7 +98,7 @@ private:
 };
 
 /** A StaticMeshComponent is a mesh that does not animate. */
-UCLASS(ClassGroup=(Rendering, Common), hidecategories=(Object,Activation,"Components|Activation"), ShowCategories=(Mobility), dependson=ULightmassPrimitiveSettingsObject, editinlinenew, meta=(BlueprintSpawnableComponent))
+UCLASS(ClassGroup=(Rendering, Common), hidecategories=(Object,Activation,"Components|Activation"), ShowCategories=(Mobility), editinlinenew, meta=(BlueprintSpawnableComponent))
 class ENGINE_API UStaticMeshComponent : public UMeshComponent
 {
 	GENERATED_UCLASS_BODY()
diff --git a/Engine/Source/Runtime/Engine/Classes/EdGraph/EdGraph.h b/Engine/Source/Runtime/Engine/Classes/EdGraph/EdGraph.h
--- a/Engine/Source/Runtime/Engine/Classes/EdGraph/EdGraph.h
+++ b/Engine/Source/Runtime/Engine/Classes/EdGraph/EdGraph.h
@@ -1,8 +1,11 @@
 // Copyright 1998-2014 Epic Games, Inc. All Rights Reserved.
 
 #pragma once
+
 #include "BlueprintUtilities.h"
 #include "EdGraphSchema.h"
+#include "EdGraphNode.h"
+
 #include "EdGraph.generated.h"
 
 USTRUCT()
@@ -51,7 +54,7 @@ struct TStructOpsTypeTraits<FGraphReference> : public TStructOpsTypeTraitsBase
 	};
 };
 
-UCLASS(dependson=(UEdGraphNode))
+UCLASS()
 class ENGINE_API UEdGraph : public UObject
 {
 	GENERATED_UCLASS_BODY()
diff --git a/Engine/Source/Runtime/Engine/Classes/EdGraph/EdGraphSchema.h b/Engine/Source/Runtime/Engine/Classes/EdGraph/EdGraphSchema.h
--- a/Engine/Source/Runtime/Engine/Classes/EdGraph/EdGraphSchema.h
+++ b/Engine/Source/Runtime/Engine/Classes/EdGraph/EdGraphSchema.h
@@ -2,9 +2,12 @@
 
 #pragma once
 
+#include "EdGraphPin.h"
+
 #include "EdGraphSchema.generated.h"
 
 class FSlateRect;
+class UEdGraphNode;
 
 /** Distinguishes between different graph types. Graphs can have different properties; for example: functions have one entry point, ubergraphs can have multiples. */
 UENUM()
diff --git a/Engine/Source/Runtime/Engine/Classes/Engine/Blueprint.h b/Engine/Source/Runtime/Engine/Classes/Engine/Blueprint.h
--- a/Engine/Source/Runtime/Engine/Classes/Engine/Blueprint.h
+++ b/Engine/Source/Runtime/Engine/Classes/Engine/Blueprint.h
@@ -2,8 +2,11 @@
 
 #pragma once
 
+#include "EdGraph/EdGraph.h"
 #include "EdGraph/EdGraphPin.h"
+#include "EdGraph/EdGraphNode.h"
 #include "BlueprintCore.h"
+
 #include "Blueprint.generated.h"
 
 /** States a blueprint can be in */
@@ -248,7 +251,7 @@ struct FEditedDocumentInfo
  * and script level events; giving designers and gameplay programmers the tools to quickly create and iterate gameplay from
  * within Unreal Editor without ever needing to write a line of code.
  */
-UCLASS(config=Engine, dependson=(UEdGraphPin,UEdGraph,UEdGraphNode), BlueprintType)
+UCLASS(config=Engine, BlueprintType)
 class ENGINE_API UBlueprint : public UBlueprintCore
 {
 	GENERATED_UCLASS_BODY()
diff --git a/Engine/Source/Runtime/Engine/Classes/Engine/BlueprintGeneratedClass.h b/Engine/Source/Runtime/Engine/Classes/Engine/BlueprintGeneratedClass.h
--- a/Engine/Source/Runtime/Engine/Classes/Engine/BlueprintGeneratedClass.h
+++ b/Engine/Source/Runtime/Engine/Classes/Engine/BlueprintGeneratedClass.h
@@ -379,7 +379,7 @@ public:
 #endif
 };
 
-UCLASS(dependson=DynamicBlueprintBinding)
+UCLASS()
 class ENGINE_API UBlueprintGeneratedClass : public UClass
 {
 	GENERATED_UCLASS_BODY()
diff --git a/Engine/Source/Runtime/Engine/Classes/Engine/Font.h b/Engine/Source/Runtime/Engine/Classes/Engine/Font.h
--- a/Engine/Source/Runtime/Engine/Classes/Engine/Font.h
+++ b/Engine/Source/Runtime/Engine/Classes/Engine/Font.h
@@ -74,7 +74,7 @@ struct TStructOpsTypeTraits<FFontCharacter> : public TStructOpsTypeTraitsBase
  * the font database only contains the coordinates of the individual
  * glyph.
  */
-UCLASS(hidecategories=Object, dependson=UFontImportOptions, autoexpandcategories=Font, MinimalAPI, BlueprintType)
+UCLASS(hidecategories=Object, autoexpandcategories=Font, MinimalAPI, BlueprintType)
 class UFont : public UObject
 {
 	GENERATED_UCLASS_BODY()
diff --git a/Engine/Source/Runtime/Engine/Classes/Engine/GameEngine.h b/Engine/Source/Runtime/Engine/Classes/Engine/GameEngine.h
--- a/Engine/Source/Runtime/Engine/Classes/Engine/GameEngine.h
+++ b/Engine/Source/Runtime/Engine/Classes/Engine/GameEngine.h
@@ -5,6 +5,9 @@
 //=============================================================================
 
 #pragma once
+
+#include "Engine/Engine.h"
+
 #include "GameEngine.generated.h"
 
 UCLASS(config=Engine, transient)
diff --git a/Engine/Source/Runtime/Engine/Classes/Engine/LevelBase.h b/Engine/Source/Runtime/Engine/Classes/Engine/LevelBase.h
--- a/Engine/Source/Runtime/Engine/Classes/Engine/LevelBase.h
+++ b/Engine/Source/Runtime/Engine/Classes/Engine/LevelBase.h
@@ -7,7 +7,7 @@
 #pragma once
 #include "LevelBase.generated.h"
 
-UCLASS(Abstract, customConstructor, dependsOn=UGameEngine)
+UCLASS(Abstract, customConstructor)
 class ULevelBase : public UObject
 {
 	GENERATED_UCLASS_BODY()
diff --git a/Engine/Source/Runtime/Engine/Classes/Engine/LevelScriptBlueprint.h b/Engine/Source/Runtime/Engine/Classes/Engine/LevelScriptBlueprint.h
--- a/Engine/Source/Runtime/Engine/Classes/Engine/LevelScriptBlueprint.h
+++ b/Engine/Source/Runtime/Engine/Classes/Engine/LevelScriptBlueprint.h
@@ -1,9 +1,12 @@
 // Copyright 1998-2014 Epic Games, Inc. All Rights Reserved.
 
 #pragma once
+
+#include "Blueprint.h"
+
 #include "LevelScriptBlueprint.generated.h"
 
-UCLASS(dependson=UBlueprint,MinimalAPI, NotBlueprintType)
+UCLASS(MinimalAPI, NotBlueprintType)
 class ULevelScriptBlueprint : public UBlueprint
 {
 	GENERATED_UCLASS_BODY()
diff --git a/Engine/Source/Runtime/Engine/Classes/Engine/NetConnection.h b/Engine/Source/Runtime/Engine/Classes/Engine/NetConnection.h
--- a/Engine/Source/Runtime/Engine/Classes/Engine/NetConnection.h
+++ b/Engine/Source/Runtime/Engine/Classes/Engine/NetConnection.h
@@ -5,10 +5,12 @@
 //
 
 #pragma once
+
 #include "Net/DataBunch.h"
 #include "Engine/Channel.h"
 #include "Engine/Player.h"
 #include "Engine/NetDriver.h"
+
 #include "NetConnection.generated.h"
 
 class FObjectReplicator;
@@ -95,7 +97,7 @@ struct DelayedPacket
 #define PING_ACK_DELAY 0.5
 
 
-UCLASS(customConstructor, Abstract, MinimalAPI, transient, config=Engine, dependsOn=UGameEngine)
+UCLASS(customConstructor, Abstract, MinimalAPI, transient, config=Engine)
 class UNetConnection : public UPlayer
 {
 	GENERATED_UCLASS_BODY()
diff --git a/Engine/Source/Runtime/Engine/Classes/Engine/NetDriver.h b/Engine/Source/Runtime/Engine/Classes/Engine/NetDriver.h
--- a/Engine/Source/Runtime/Engine/Classes/Engine/NetDriver.h
+++ b/Engine/Source/Runtime/Engine/Classes/Engine/NetDriver.h
@@ -3,7 +3,9 @@
 //
 // Base class of a network driver attached to an active or pending level.
 #pragma once
-//
+
+#include "EngineTypes.h"
+
 #include "NetDriver.generated.h"
 
 class FRepChangedPropertyTracker;
@@ -90,7 +92,7 @@ struct FActorDestructionInfo
 };
 
 
-UCLASS(dependson=UEngineTypes, Abstract, customConstructor, transient, MinimalAPI, config=Engine)
+UCLASS(Abstract, customConstructor, transient, MinimalAPI, config=Engine)
 class UNetDriver : public UObject, public FExec
 {
 	GENERATED_UCLASS_BODY()
diff --git a/Engine/Source/Runtime/Engine/Classes/Engine/PendingNetGame.h b/Engine/Source/Runtime/Engine/Classes/Engine/PendingNetGame.h
--- a/Engine/Source/Runtime/Engine/Classes/Engine/PendingNetGame.h
+++ b/Engine/Source/Runtime/Engine/Classes/Engine/PendingNetGame.h
@@ -1,6 +1,9 @@
 // Copyright 1998-2014 Epic Games, Inc. All Rights Reserved.
 
 #pragma once
+
+#include "GameViewportClient.h"
+
 #include "PendingNetGame.generated.h"
 
 /**
@@ -61,7 +64,7 @@ public:
 	virtual void NotifyControlMessage(UNetConnection* Connection, uint8 MessageType, class FInBunch& Bunch) PURE_VIRTUAL(FNetworkNotify::NotifyReceivedText,);
 };
 
-UCLASS(customConstructor, transient, dependsOn=UGameViewportClient)
+UCLASS(customConstructor, transient)
 class UPendingNetGame :
 	public UObject,
 	public FNetworkNotify
diff --git a/Engine/Source/Runtime/Engine/Classes/Engine/PostProcessVolume.h b/Engine/Source/Runtime/Engine/Classes/Engine/PostProcessVolume.h
--- a/Engine/Source/Runtime/Engine/Classes/Engine/PostProcessVolume.h
+++ b/Engine/Source/Runtime/Engine/Classes/Engine/PostProcessVolume.h
@@ -6,11 +6,15 @@
 //=============================================================================
 
 #pragma once
+
+#include "Scene.h"
 #include "GameFramework/Volume.h"
+#include "Interfaces/Interface_PostProcessVolume.h"
+
 #include "PostProcessVolume.generated.h"
 
 	// for FPostprocessSettings
-UCLASS(dependson=UScene, autoexpandcategories=PostProcessVolume, hidecategories=(Advanced, Collision, Volume, Brush, Attachment), MinimalAPI)
+UCLASS(autoexpandcategories=PostProcessVolume, hidecategories=(Advanced, Collision, Volume, Brush, Attachment), MinimalAPI)
 class APostProcessVolume : public AVolume, public IInterface_PostProcessVolume
 {
 	GENERATED_UCLASS_BODY()
diff --git a/Engine/Source/Runtime/Engine/Classes/Engine/TriggerBase.h b/Engine/Source/Runtime/Engine/Classes/Engine/TriggerBase.h
--- a/Engine/Source/Runtime/Engine/Classes/Engine/TriggerBase.h
+++ b/Engine/Source/Runtime/Engine/Classes/Engine/TriggerBase.h
@@ -5,9 +5,12 @@
  */
 
 #pragma once
+
+#include "EngineBaseTypes.h"
+
 #include "TriggerBase.generated.h"
 
-UCLASS(dependson=AMatineeActor, ClassGroup=Common, abstract, ConversionRoot, MinimalAPI)
+UCLASS(ClassGroup=Common, abstract, ConversionRoot, MinimalAPI)
 class ATriggerBase : public AActor
 {
 	GENERATED_UCLASS_BODY()
diff --git a/Engine/Source/Runtime/Engine/Classes/GameFramework/Actor.h b/Engine/Source/Runtime/Engine/Classes/GameFramework/Actor.h
--- a/Engine/Source/Runtime/Engine/Classes/GameFramework/Actor.h
+++ b/Engine/Source/Runtime/Engine/Classes/GameFramework/Actor.h
@@ -2,6 +2,10 @@
 
 #pragma once
 
+#include "Engine/EngineBaseTypes.h"
+#include "Components/ActorComponent.h"
+#include "Engine/EngineTypes.h"
+
 #include "Runtime/InputCore/Classes/InputCoreTypes.h"
 #include "Runtime/RenderCore/Public/RenderCommandFence.h"
 
@@ -143,7 +147,7 @@ public:
 * AActor::PostInitializeComponents - Called after the actor's components have been initialized
 * AActor::BeginPlay - Called when the level is started
 */
-UCLASS(abstract, dependson=(UEngineBaseTypes, UActorComponent, UEngineTypes), BlueprintType, Blueprintable, config=Engine)
+UCLASS(abstract, BlueprintType, Blueprintable, config=Engine)
 class ENGINE_API AActor : public UObject
 {
 	GENERATED_UCLASS_BODY()
diff --git a/Engine/Source/Runtime/Engine/Classes/GameFramework/Character.h b/Engine/Source/Runtime/Engine/Classes/GameFramework/Character.h
--- a/Engine/Source/Runtime/Engine/Classes/GameFramework/Character.h
+++ b/Engine/Source/Runtime/Engine/Classes/GameFramework/Character.h
@@ -1,7 +1,10 @@
 // Copyright 1998-2014 Epic Games, Inc. All Rights Reserved.
 
 #pragma once
+
+#include "Controller.h"
 #include "GameFramework/CharacterMovementComponent.h"
+
 #include "Character.generated.h"
 
 DECLARE_DYNAMIC_MULTICAST_DELEGATE_ThreeParams(FMovementModeChangedSignature, class ACharacter*, Character, EMovementMode, PrevMovementMode, uint8, PreviousCustomMode);
@@ -136,7 +139,8 @@ struct FRepRelativeMovement
 // to a CharacterMovementComponent that handles movement of the collision capsule, and they
 // also have implementations of basic networking and input models.
 //=============================================================================
-UCLASS(abstract, config=Game, dependson=(AController, UCharacterMovementComponent), BlueprintType, hidecategories=("Pawn|Character|InternalEvents"))
+
+UCLASS(abstract, config=Game, BlueprintType, hidecategories=("Pawn|Character|InternalEvents"))
 class ENGINE_API ACharacter : public APawn
 {
 	GENERATED_UCLASS_BODY()
diff --git a/Engine/Source/Runtime/Engine/Classes/GameFramework/CharacterMovementComponent.h b/Engine/Source/Runtime/Engine/Classes/GameFramework/CharacterMovementComponent.h
--- a/Engine/Source/Runtime/Engine/Classes/GameFramework/CharacterMovementComponent.h
+++ b/Engine/Source/Runtime/Engine/Classes/GameFramework/CharacterMovementComponent.h
@@ -8,6 +8,7 @@
 #pragma once
 
 #include "Interfaces/NetworkPredictionInterface.h"
+
 #include "CharacterMovementComponent.generated.h"
 
 /** Movement modes for Characters.  */
@@ -267,7 +268,7 @@ struct FCharacterMovementComponentPreClothTickFunction : public FTickFunction
 	virtual FString DiagnosticMessage() override;
 };
 
-UCLASS(dependson=(UNetworkPredictionInterface))
+UCLASS()
 class ENGINE_API UCharacterMovementComponent : public UPawnMovementComponent, public INetworkPredictionInterface
 {
 	GENERATED_UCLASS_BODY()
diff --git a/Engine/Source/Runtime/Engine/Classes/GameFramework/HUD.h b/Engine/Source/Runtime/Engine/Classes/GameFramework/HUD.h
--- a/Engine/Source/Runtime/Engine/Classes/GameFramework/HUD.h
+++ b/Engine/Source/Runtime/Engine/Classes/GameFramework/HUD.h
@@ -1,6 +1,9 @@
 // Copyright 1998-2014 Epic Games, Inc. All Rights Reserved.
 
 #pragma once
+
+#include "Canvas.h"
+
 #include "HUD.generated.h"
 
 /** List of actors and debug text to draw, @see AddDebugText(), RemoveDebugText(), and DrawDebugTextList() */
@@ -151,7 +154,7 @@ private:
 // Base class of the heads-up display.
 //
 //=============================================================================
-UCLASS(config=Game, hidecategories=(Rendering,Actor,Input,Replication), showcategories=("Input|MouseInput", "Input|TouchInput"), notplaceable, transient, dependson=UCanvas, BlueprintType, Blueprintable)
+UCLASS(config=Game, hidecategories=(Rendering,Actor,Input,Replication), showcategories=("Input|MouseInput", "Input|TouchInput"), notplaceable, transient, BlueprintType, Blueprintable)
 class ENGINE_API AHUD : public AActor
 {
 	GENERATED_UCLASS_BODY()
diff --git a/Engine/Source/Runtime/Engine/Classes/GameFramework/NavMovementComponent.h b/Engine/Source/Runtime/Engine/Classes/GameFramework/NavMovementComponent.h
--- a/Engine/Source/Runtime/Engine/Classes/GameFramework/NavMovementComponent.h
+++ b/Engine/Source/Runtime/Engine/Classes/GameFramework/NavMovementComponent.h
@@ -5,12 +5,16 @@
  */
 
 #pragma once
+
+#include "AI/Navigation/NavigationSystem.h"
+#include "Components/PrimitiveComponent.h"
 #include "AI/Navigation/NavAgentInterface.h"
 #include "AI/Navigation/NavigationTypes.h"
 #include "GameFramework/MovementComponent.h"
+
 #include "NavMovementComponent.generated.h"
 
-UCLASS(abstract, dependson=(UNavigationSystem, UPrimitiveComponent, INavAgentInterface))
+UCLASS(abstract)
 class ENGINE_API UNavMovementComponent : public UMovementComponent
 {
 	GENERATED_UCLASS_BODY()
diff --git a/Engine/Source/Runtime/Engine/Classes/GameFramework/Pawn.h b/Engine/Source/Runtime/Engine/Classes/GameFramework/Pawn.h
--- a/Engine/Source/Runtime/Engine/Classes/GameFramework/Pawn.h
+++ b/Engine/Source/Runtime/Engine/Classes/GameFramework/Pawn.h
@@ -1,7 +1,10 @@
 // Copyright 1998-2014 Epic Games, Inc. All Rights Reserved.
 
 #pragma once
+
+#include "Controller.h"
 #include "PawnMovementComponent.h"
+
 #include "Pawn.generated.h"
 
 ENGINE_API DECLARE_LOG_CATEGORY_EXTERN(LogDamage, Warning, All);
@@ -11,7 +14,7 @@ ENGINE_API DECLARE_LOG_CATEGORY_EXTERN(LogDamage, Warning, All);
  *	Pawns are the physical representations of players and creatures in a level.
  *	They are responsible for all physical interaction between the player or AI and the world.
  */
-UCLASS(abstract, config=Game, dependson=AController, BlueprintType, Blueprintable, hidecategories=(Navigation, "AI|Navigation"))
+UCLASS(abstract, config=Game, BlueprintType, Blueprintable, hidecategories=(Navigation, "AI|Navigation"))
 class ENGINE_API APawn : public AActor, public INavAgentInterface
 {
 	GENERATED_UCLASS_BODY()
diff --git a/Engine/Source/Runtime/Engine/Classes/GameFramework/PlayerInput.h b/Engine/Source/Runtime/Engine/Classes/GameFramework/PlayerInput.h
--- a/Engine/Source/Runtime/Engine/Classes/GameFramework/PlayerInput.h
+++ b/Engine/Source/Runtime/Engine/Classes/GameFramework/PlayerInput.h
@@ -7,9 +7,12 @@
 //=============================================================================
 
 #pragma once
+
+#include "Engine/EngineTypes.h"
 #include "Components/InputComponent.h"
 #include "GestureRecognizer.h"
 #include "KeyState.h"
+
 #include "PlayerInput.generated.h"
 
 ENGINE_API DECLARE_LOG_CATEGORY_EXTERN(LogInput, Log, All);
@@ -227,7 +230,7 @@ struct FAxisKeyDetails
 	}
 };
 
-UCLASS(Within=PlayerController, config=Input, dependsOn=UEngineTypes, transient)
+UCLASS(Within=PlayerController, config=Input, transient)
 class ENGINE_API UPlayerInput : public UObject
 {
 	GENERATED_UCLASS_BODY()
diff --git a/Engine/Source/Runtime/Engine/Classes/GameFramework/SpectatorPawnMovement.h b/Engine/Source/Runtime/Engine/Classes/GameFramework/SpectatorPawnMovement.h
--- a/Engine/Source/Runtime/Engine/Classes/GameFramework/SpectatorPawnMovement.h
+++ b/Engine/Source/Runtime/Engine/Classes/GameFramework/SpectatorPawnMovement.h
@@ -5,10 +5,13 @@
  */
 
 #pragma once
+
+#include "FloatingPawnMovement.h"
 #include "GameFramework/FloatingPawnMovement.h"
+
 #include "SpectatorPawnMovement.generated.h"
 
-UCLASS(DependsOn=UFloatingPawnMovement)
+UCLASS()
 class ENGINE_API USpectatorPawnMovement : public UFloatingPawnMovement
 {
 	GENERATED_UCLASS_BODY()
diff --git a/Engine/Source/Runtime/Engine/Classes/GameFramework/WheeledVehicle.h b/Engine/Source/Runtime/Engine/Classes/GameFramework/WheeledVehicle.h
--- a/Engine/Source/Runtime/Engine/Classes/GameFramework/WheeledVehicle.h
+++ b/Engine/Source/Runtime/Engine/Classes/GameFramework/WheeledVehicle.h
@@ -1,9 +1,12 @@
 // Copyright 1998-2014 Epic Games, Inc. All Rights Reserved.
 
 #pragma once
+
+#include "Controller.h"
+
 #include "WheeledVehicle.generated.h"
 
-UCLASS(abstract, config=Game, dependson=(AController), BlueprintType)
+UCLASS(abstract, config=Game, BlueprintType)
 class ENGINE_API AWheeledVehicle : public APawn
 {
 	GENERATED_UCLASS_BODY()
diff --git a/Engine/Source/Runtime/Engine/Classes/GameFramework/WorldSettings.h b/Engine/Source/Runtime/Engine/Classes/GameFramework/WorldSettings.h
--- a/Engine/Source/Runtime/Engine/Classes/GameFramework/WorldSettings.h
+++ b/Engine/Source/Runtime/Engine/Classes/GameFramework/WorldSettings.h
@@ -1,7 +1,11 @@
 // Copyright 1998-2014 Epic Games, Inc. All Rights Reserved.
 
 #pragma once
+
 #include "Sound/ReverbVolume.h"
+#include "Engine/PostProcessVolume.h"
+#include "MusicTrackDataStructures.h"
+
 #include "WorldSettings.generated.h"
 
 UENUM()
@@ -174,7 +178,7 @@ struct ENGINE_API FNetViewer
 /**
  * Actor containing all script accessible world properties.
  */
-UCLASS(config=game, showcategories=(Rendering), hidecategories=(Actor, Advanced, Display, Events, Object, Attachment, Info, Input, Blueprint, Layers), showcategories=("Input|MouseInput", "Input|TouchInput"), notplaceable, dependson=(APostProcessVolume, UMusicTrackDataStructures))
+UCLASS(config=game, showcategories=(Rendering), hidecategories=(Actor, Advanced, Display, Events, Object, Attachment, Info, Input, Blueprint, Layers), showcategories=("Input|MouseInput", "Input|TouchInput"), notplaceable)
 class ENGINE_API AWorldSettings : public AInfo, public IInterface_AssetUserData
 {
 	GENERATED_UCLASS_BODY()
diff --git a/Engine/Source/Runtime/Engine/Classes/Landscape/Landscape.h b/Engine/Source/Runtime/Engine/Classes/Landscape/Landscape.h
--- a/Engine/Source/Runtime/Engine/Classes/Landscape/Landscape.h
+++ b/Engine/Source/Runtime/Engine/Classes/Landscape/Landscape.h
@@ -1,8 +1,11 @@
 // Copyright 1998-2014 Epic Games, Inc. All Rights Reserved.
 
 #pragma once
+
 #include "LandscapeInfo.h"
 #include "LandscapeProxy.h"
+#include "Components/LightComponent.h"
+
 #include "Landscape.generated.h"
 
 UENUM()
@@ -18,7 +21,7 @@ enum ELandscapeSetupErrors
 	LSE_MAX,
 };
 
-UCLASS(dependson=ULightComponent, Placeable, hidecategories=LandscapeProxy, showcategories=(Display, Movement, Collision, Lighting, LOD, Input), MinimalAPI)
+UCLASS(Placeable, hidecategories=LandscapeProxy, showcategories=(Display, Movement, Collision, Lighting, LOD, Input), MinimalAPI)
 class ALandscape : public ALandscapeProxy
 {
 	GENERATED_UCLASS_BODY()
diff --git a/Engine/Source/Runtime/Engine/Classes/Landscape/LandscapeProxy.h b/Engine/Source/Runtime/Engine/Classes/Landscape/LandscapeProxy.h
--- a/Engine/Source/Runtime/Engine/Classes/Landscape/LandscapeProxy.h
+++ b/Engine/Source/Runtime/Engine/Classes/Landscape/LandscapeProxy.h
@@ -1,6 +1,8 @@
 // Copyright 1998-2014 Epic Games, Inc. All Rights Reserved.
 
 #pragma once
+
+#include "Engine/EngineTypes.h"
 #include "LandscapeInfo.h"
 #include "PhysicsEngine/BodyInstance.h"
 
@@ -184,7 +186,7 @@ namespace ELandscapeLODFalloff
 	};
 }
 
-UCLASS(dependson=UEngineTypes, NotPlaceable, hidecategories=(Display, Attachment, Physics, Debug, Lighting, LOD), showcategories=(Rendering, "Utilities|Transformation"), MinimalAPI)
+UCLASS(NotPlaceable, hidecategories=(Display, Attachment, Physics, Debug, Lighting, LOD), showcategories=(Rendering, "Utilities|Transformation"), MinimalAPI)
 class ALandscapeProxy : public AActor, public INavRelevantActorInterface
 {
 	GENERATED_UCLASS_BODY()
diff --git a/Engine/Source/Runtime/Engine/Classes/Lightmass/LightmassPrimitiveSettingsObject.h b/Engine/Source/Runtime/Engine/Classes/Lightmass/LightmassPrimitiveSettingsObject.h
--- a/Engine/Source/Runtime/Engine/Classes/Lightmass/LightmassPrimitiveSettingsObject.h
+++ b/Engine/Source/Runtime/Engine/Classes/Lightmass/LightmassPrimitiveSettingsObject.h
@@ -5,9 +5,12 @@
  */
 
 #pragma once
+
+#include "Engine/EngineTypes.h"
+
 #include "LightmassPrimitiveSettingsObject.generated.h"
 
-UCLASS(hidecategories=Object, dependson=UEngineTypes, editinlinenew, MinimalAPI,collapsecategories)
+UCLASS(hidecategories=Object, editinlinenew, MinimalAPI,collapsecategories)
 class ULightmassPrimitiveSettingsObject : public UObject
 {
 	GENERATED_UCLASS_BODY()
diff --git a/Engine/Source/Runtime/Engine/Classes/Materials/Material.h b/Engine/Source/Runtime/Engine/Classes/Materials/Material.h
--- a/Engine/Source/Runtime/Engine/Classes/Materials/Material.h
+++ b/Engine/Source/Runtime/Engine/Classes/Materials/Material.h
@@ -1,9 +1,12 @@
 // Copyright 1998-2014 Epic Games, Inc. All Rights Reserved.
 
 #pragma once
-#include "../Engine/BlendableInterface.h"
-#include "../Materials/MaterialInterface.h"
+
+#include "Engine/BlendableInterface.h"
 #include "MaterialExpressionIO.h"
+#include "Materials/MaterialExpression.h"
+#include "Materials/MaterialInterface.h"
+
 #include "Material.generated.h"
 
 #if WITH_EDITOR
@@ -256,7 +259,7 @@ struct FMaterialParameterCollectionInfo
  * when light from the scene hits the surface, the shading model of the material is used to calculate how
  * that light interacts with the surface. 
  */
-UCLASS(hidecategories=Object, MinimalAPI, BlueprintType, DependsOn=UMaterialExpression)
+UCLASS(hidecategories=Object, MinimalAPI, BlueprintType)
 class UMaterial : public UMaterialInterface
 {
 	GENERATED_UCLASS_BODY()
diff --git a/Engine/Source/Runtime/Engine/Classes/Matinee/InterpTrackMoveAxis.h b/Engine/Source/Runtime/Engine/Classes/Matinee/InterpTrackMoveAxis.h
--- a/Engine/Source/Runtime/Engine/Classes/Matinee/InterpTrackMoveAxis.h
+++ b/Engine/Source/Runtime/Engine/Classes/Matinee/InterpTrackMoveAxis.h
@@ -23,7 +23,7 @@ enum EInterpMoveAxis
 	AXIS_RotationZ,
 };
 
-UCLASS(dependson=UInterpTrackMove, MinimalAPI, meta=( DisplayName = "Move Axis Track" ) )
+UCLASS(MinimalAPI, meta=( DisplayName = "Move Axis Track" ) )
 class UInterpTrackMoveAxis : public UInterpTrackFloatBase
 {
 	GENERATED_UCLASS_BODY()
diff --git a/Engine/Source/Runtime/Engine/Classes/Matinee/MatineeActor.h b/Engine/Source/Runtime/Engine/Classes/Matinee/MatineeActor.h
--- a/Engine/Source/Runtime/Engine/Classes/Matinee/MatineeActor.h
+++ b/Engine/Source/Runtime/Engine/Classes/Matinee/MatineeActor.h
@@ -3,6 +3,9 @@
 //  AActor  used to controll matinee's and to replicate activation, playback, and other relevant flags to net clients
 
 #pragma once
+
+#include "Engine/EngineBaseTypes.h"
+
 #include "MatineeActor.generated.h"
 
 /** Signature of function to handle a matinee event track key */
@@ -49,7 +52,7 @@ struct FInterpGroupActorInfo
 
 };
 
-UCLASS(dependson=UEngineBaseTypes, MinimalAPI, NotBlueprintable, hidecategories=(Collision, Game, Input), showcategories=("Input|MouseInput", "Input|TouchInput", "Game|Damage"))
+UCLASS(MinimalAPI, NotBlueprintable, hidecategories=(Collision, Game, Input), showcategories=("Input|MouseInput", "Input|TouchInput", "Game|Damage"))
 class AMatineeActor : public AActor
 {
 	GENERATED_UCLASS_BODY()
diff --git a/Engine/Source/Runtime/Engine/Classes/Matinee/MatineeActorCameraAnim.h b/Engine/Source/Runtime/Engine/Classes/Matinee/MatineeActorCameraAnim.h
--- a/Engine/Source/Runtime/Engine/Classes/Matinee/MatineeActorCameraAnim.h
+++ b/Engine/Source/Runtime/Engine/Classes/Matinee/MatineeActorCameraAnim.h
@@ -1,13 +1,16 @@
 // Copyright 1998-2014 Epic Games, Inc. All Rights Reserved.
 
 #pragma once
+
 #include "Matinee/MatineeActor.h"
+#include "Engine/EngineBaseTypes.h"
+
 #include "MatineeActorCameraAnim.generated.h"
 
 /**
  * Actor used to control temporary matinees for camera anims that only exist in the editor
  */
-UCLASS(dependson=UEngineBaseTypes, notplaceable, MinimalAPI, NotBlueprintable)
+UCLASS(notplaceable, MinimalAPI, NotBlueprintable)
 class AMatineeActorCameraAnim : public AMatineeActor
 {
 	GENERATED_UCLASS_BODY()
diff --git a/Engine/Source/Runtime/Engine/Classes/Particles/EmitterCameraLensEffectBase.h b/Engine/Source/Runtime/Engine/Classes/Particles/EmitterCameraLensEffectBase.h
--- a/Engine/Source/Runtime/Engine/Classes/Particles/EmitterCameraLensEffectBase.h
+++ b/Engine/Source/Runtime/Engine/Classes/Particles/EmitterCameraLensEffectBase.h
@@ -7,10 +7,13 @@
  */
 
 #pragma once
+
 #include "Particles/Emitter.h"
+#include "Emitter.h"
+
 #include "EmitterCameraLensEffectBase.generated.h"
 
-UCLASS(abstract, Blueprintable, dependson=AEmitter, MinimalAPI)
+UCLASS(abstract, Blueprintable, MinimalAPI)
 class AEmitterCameraLensEffectBase : public AEmitter
 {
 	GENERATED_UCLASS_BODY()
diff --git a/Engine/Source/Runtime/Engine/Classes/Particles/ParticleEmitter.h b/Engine/Source/Runtime/Engine/Classes/Particles/ParticleEmitter.h
--- a/Engine/Source/Runtime/Engine/Classes/Particles/ParticleEmitter.h
+++ b/Engine/Source/Runtime/Engine/Classes/Particles/ParticleEmitter.h
@@ -6,8 +6,11 @@
 //=============================================================================
 
 #pragma once
+
+#include "ParticleLODLevel.h"
 #include "ParticleHelper.h"
 #include "ParticleEmitterInstances.h"
+
 #include "ParticleEmitter.generated.h"
 
 //=============================================================================
@@ -77,7 +80,7 @@ struct FParticleBurst
 	
 };
 
-UCLASS(dependson=UParticleLODLevel, hidecategories=Object, editinlinenew, abstract, MinimalAPI)
+UCLASS(hidecategories=Object, editinlinenew, abstract, MinimalAPI)
 class UParticleEmitter : public UObject
 {
 	GENERATED_UCLASS_BODY()
diff --git a/Engine/Source/Runtime/Engine/Classes/Particles/ParticleEventManager.h b/Engine/Source/Runtime/Engine/Classes/Particles/ParticleEventManager.h
--- a/Engine/Source/Runtime/Engine/Classes/Particles/ParticleEventManager.h
+++ b/Engine/Source/Runtime/Engine/Classes/Particles/ParticleEventManager.h
@@ -2,7 +2,10 @@
 
 
 #pragma once
+
 #include "GameFramework/Actor.h"
+#include "ParticleSystemComponent.h"
+
 #include "ParticleEventManager.generated.h"
 
 struct FParticleEventSpawnData;
@@ -10,7 +13,7 @@ struct FParticleEventDeathData;
 struct FParticleEventCollideData;
 struct FParticleEventBurstData;
 
-UCLASS(dependson=UParticleSystemComponent, config=Game, NotBlueprintable, notplaceable)
+UCLASS(config=Game, NotBlueprintable, notplaceable)
 class AParticleEventManager : public AActor
 {
 	GENERATED_UCLASS_BODY()
diff --git a/Engine/Source/Runtime/Engine/Classes/Particles/ParticleSystemComponent.h b/Engine/Source/Runtime/Engine/Classes/Particles/ParticleSystemComponent.h
--- a/Engine/Source/Runtime/Engine/Classes/Particles/ParticleSystemComponent.h
+++ b/Engine/Source/Runtime/Engine/Classes/Particles/ParticleSystemComponent.h
@@ -249,7 +249,7 @@ struct FParticleEventKismetData : public FParticleEventData
 /** 
  * A particle emmitter.
  */
-UCLASS(ClassGroup=Rendering, hidecategories=Object, hidecategories=Physics, hidecategories=Collision, showcategories=Trigger, editinlinenew, dependson=(AEmitter, UParticleSystem), meta=(BlueprintSpawnableComponent))
+UCLASS(ClassGroup=Rendering, hidecategories=Object, hidecategories=Physics, hidecategories=Collision, showcategories=Trigger, editinlinenew, meta=(BlueprintSpawnableComponent))
 class ENGINE_API UParticleSystemComponent : public UPrimitiveComponent
 {
 	GENERATED_UCLASS_BODY()
diff --git a/Engine/Source/Runtime/Engine/Classes/PhysicsEngine/BodySetup.h b/Engine/Source/Runtime/Engine/Classes/PhysicsEngine/BodySetup.h
--- a/Engine/Source/Runtime/Engine/Classes/PhysicsEngine/BodySetup.h
+++ b/Engine/Source/Runtime/Engine/Classes/PhysicsEngine/BodySetup.h
@@ -392,7 +392,7 @@ struct ENGINE_API FKAggregateGeom
 	float GetVolume(const FVector& Scale3D) const;
 };
 
-UCLASS(hidecategories=Object, MinimalAPI, dependson=BodyInstance)
+UCLASS(hidecategories=Object, MinimalAPI)
 class UBodySetup : public UObject
 {
 	GENERATED_UCLASS_BODY()
diff --git a/Engine/Source/Runtime/Engine/Classes/PhysicsEngine/BodySetup2D.h b/Engine/Source/Runtime/Engine/Classes/PhysicsEngine/BodySetup2D.h
--- a/Engine/Source/Runtime/Engine/Classes/PhysicsEngine/BodySetup2D.h
+++ b/Engine/Source/Runtime/Engine/Classes/PhysicsEngine/BodySetup2D.h
@@ -3,10 +3,11 @@
 #pragma once
 
 #include "AggregateGeometry2D.h"
+#include "Engine/EngineTypes.h"
 
 #include "BodySetup2D.generated.h"
 
-UCLASS(DependsOn=UEngineTypes)
+UCLASS()
 class ENGINE_API UBodySetup2D : public UBodySetup
 {
 	GENERATED_UCLASS_BODY()
diff --git a/Engine/Source/Runtime/Engine/Classes/PhysicsEngine/PhysicsConstraintComponent.h b/Engine/Source/Runtime/Engine/Classes/PhysicsEngine/PhysicsConstraintComponent.h
--- a/Engine/Source/Runtime/Engine/Classes/PhysicsEngine/PhysicsConstraintComponent.h
+++ b/Engine/Source/Runtime/Engine/Classes/PhysicsEngine/PhysicsConstraintComponent.h
@@ -9,7 +9,7 @@ DECLARE_DYNAMIC_MULTICAST_DELEGATE_OneParam(FConstraintBrokenSignature, int32, C
 /**
  *	This is effectively a joint that allows you to connect 2 rigid bodies together. You can create different types of joints using the various parameters of this component.
  */
-UCLASS(ClassGroup=Physics, dependson(ConstraintInstance), MinimalAPI, meta=(BlueprintSpawnableComponent), HideCategories=(Activation,"Components|Activation", Physics, Mobility), ShowCategories=("Physics|Components|PhysicsConstraint"))
+UCLASS(ClassGroup=Physics, MinimalAPI, meta=(BlueprintSpawnableComponent), HideCategories=(Activation,"Components|Activation", Physics, Mobility), ShowCategories=("Physics|Components|PhysicsConstraint"))
 class UPhysicsConstraintComponent : public USceneComponent
 {
 	GENERATED_UCLASS_BODY()
diff --git a/Engine/Source/Runtime/Engine/Classes/Sound/DialogueVoice.h b/Engine/Source/Runtime/Engine/Classes/Sound/DialogueVoice.h
--- a/Engine/Source/Runtime/Engine/Classes/Sound/DialogueVoice.h
+++ b/Engine/Source/Runtime/Engine/Classes/Sound/DialogueVoice.h
@@ -6,9 +6,11 @@
  * Playable sound object for spoken dialogue 
  */
 
+#include "DialogueTypes.h"
+
 #include "DialogueVoice.generated.h"
 
-UCLASS(DependsOn=UDialogueTypes, hidecategories=Object, editinlinenew, MinimalAPI, BlueprintType)
+UCLASS(hidecategories=Object, editinlinenew, MinimalAPI, BlueprintType)
 class UDialogueVoice : public UObject
 {
 	GENERATED_UCLASS_BODY()
diff --git a/Engine/Source/Runtime/Engine/Classes/Sound/DialogueWave.h b/Engine/Source/Runtime/Engine/Classes/Sound/DialogueWave.h
--- a/Engine/Source/Runtime/Engine/Classes/Sound/DialogueWave.h
+++ b/Engine/Source/Runtime/Engine/Classes/Sound/DialogueWave.h
@@ -6,6 +6,8 @@
  * Context to sound wave map for spoken dialogue 
  */
 
+#include "DialogueTypes.h"
+
 #include "DialogueWave.generated.h"
 
 class USoundWave;
@@ -34,7 +36,7 @@ struct FDialogueContextMapping
 bool operator==(const FDialogueContextMapping& LHS, const FDialogueContextMapping& RHS);
 bool operator!=(const FDialogueContextMapping& LHS, const FDialogueContextMapping& RHS);
 
-UCLASS(DependsOn=UDialogueTypes, hidecategories=Object, editinlinenew, MinimalAPI, BlueprintType)
+UCLASS(hidecategories=Object, editinlinenew, MinimalAPI, BlueprintType)
 class UDialogueWave : public UObject
 {
 	GENERATED_UCLASS_BODY()
diff --git a/Engine/Source/Runtime/Engine/Classes/Sound/SoundCue.h b/Engine/Source/Runtime/Engine/Classes/Sound/SoundCue.h
--- a/Engine/Source/Runtime/Engine/Classes/Sound/SoundCue.h
+++ b/Engine/Source/Runtime/Engine/Classes/Sound/SoundCue.h
@@ -1,7 +1,10 @@
 // Copyright 1998-2014 Epic Games, Inc. All Rights Reserved.
 
 #pragma once
+
+#include "SoundClass.h"
 #include "Sound/SoundBase.h"
+
 #include "SoundCue.generated.h"
 
 struct FActiveSound;
@@ -33,7 +36,7 @@ struct FSoundNodeEditorData
 /**
  * The behavior of audio playback is defined within Sound Cues.
  */
-UCLASS(dependson=USoundClass, hidecategories=object, MinimalAPI, BlueprintType)
+UCLASS(hidecategories=object, MinimalAPI, BlueprintType)
 class USoundCue : public USoundBase
 {
 	GENERATED_UCLASS_BODY()
diff --git a/Engine/Source/Runtime/Engine/Classes/Sound/SoundMix.h b/Engine/Source/Runtime/Engine/Classes/Sound/SoundMix.h
--- a/Engine/Source/Runtime/Engine/Classes/Sound/SoundMix.h
+++ b/Engine/Source/Runtime/Engine/Classes/Sound/SoundMix.h
@@ -2,6 +2,7 @@
 
 #pragma once
 
+#include "SoundClass.h"
 #include "Audio.h"
 #include "SoundMix.generated.h"
 
@@ -112,7 +113,7 @@ struct FSoundClassAdjuster
 	
 };
 
-UCLASS(BlueprintType, dependson=(USoundClass), hidecategories=object, MinimalAPI)
+UCLASS(BlueprintType, hidecategories=object, MinimalAPI)
 class USoundMix : public UObject
 {
 	GENERATED_UCLASS_BODY()
diff --git a/Engine/Source/Runtime/Engine/Classes/Sound/SoundNodeAmbient.h b/Engine/Source/Runtime/Engine/Classes/Sound/SoundNodeAmbient.h
--- a/Engine/Source/Runtime/Engine/Classes/Sound/SoundNodeAmbient.h
+++ b/Engine/Source/Runtime/Engine/Classes/Sound/SoundNodeAmbient.h
@@ -6,7 +6,10 @@
  */
 
 #pragma once
+
+#include "SoundNodeAttenuation.h"
 #include "Sound/SoundNodeDeprecated.h"
+
 #include "SoundNodeAmbient.generated.h"
 
 USTRUCT()
@@ -39,7 +42,7 @@ struct FAmbientSoundSlot
 		}
 };
 
-UCLASS(dependson=USoundNodeAttenuation, editinlinenew, MinimalAPI)
+UCLASS(editinlinenew, MinimalAPI)
 class UDEPRECATED_SoundNodeAmbient : public UDEPRECATED_SoundNodeDeprecated
 {
 	GENERATED_UCLASS_BODY()
diff --git a/Engine/Source/Runtime/Engine/Classes/Sound/SoundNodeAmbientNonLoop.h b/Engine/Source/Runtime/Engine/Classes/Sound/SoundNodeAmbientNonLoop.h
--- a/Engine/Source/Runtime/Engine/Classes/Sound/SoundNodeAmbientNonLoop.h
+++ b/Engine/Source/Runtime/Engine/Classes/Sound/SoundNodeAmbientNonLoop.h
@@ -6,10 +6,14 @@
  */
 
 #pragma once
-#include "Sound/SoundNodeAmbient.h"
+
+#include "SoundNodeAttenuation.h"
+#include "Distributions/DistributionFloat.h"
+#include "SoundNodeAmbient.h"
+
 #include "SoundNodeAmbientNonLoop.generated.h"
 
-UCLASS(dependson=(USoundNodeAttenuation, UDistributionFloat), editinlinenew)
+UCLASS(editinlinenew)
 class UDEPRECATED_SoundNodeAmbientNonLoop : public UDEPRECATED_SoundNodeAmbient
 {
 	GENERATED_UCLASS_BODY()
diff --git a/Engine/Source/Runtime/Engine/Classes/Sound/SoundNodeAmbientNonLoopToggle.h b/Engine/Source/Runtime/Engine/Classes/Sound/SoundNodeAmbientNonLoopToggle.h
--- a/Engine/Source/Runtime/Engine/Classes/Sound/SoundNodeAmbientNonLoopToggle.h
+++ b/Engine/Source/Runtime/Engine/Classes/Sound/SoundNodeAmbientNonLoopToggle.h
@@ -6,10 +6,13 @@
  */
 
 #pragma once
+
+#include "SoundNodeAttenuation.h"
 #include "Sound/SoundNodeAmbientNonLoop.h"
+
 #include "SoundNodeAmbientNonLoopToggle.generated.h"
 
-UCLASS(hidecategories=Object, dependson=USoundNodeAttenuation, editinlinenew)
+UCLASS(hidecategories=Object, editinlinenew)
 class UDEPRECATED_SoundNodeAmbientNonLoopToggle : public UDEPRECATED_SoundNodeAmbientNonLoop
 {
 	GENERATED_UCLASS_BODY()
diff --git a/Engine/Source/Runtime/Engine/Classes/Sound/SoundWave.h b/Engine/Source/Runtime/Engine/Classes/Sound/SoundWave.h
--- a/Engine/Source/Runtime/Engine/Classes/Sound/SoundWave.h
+++ b/Engine/Source/Runtime/Engine/Classes/Sound/SoundWave.h
@@ -8,6 +8,8 @@
 
 #include "Sound/SoundBase.h"
 #include "Sound/SoundGroups.h"
+#include "SoundGroups.h"
+
 #include "SoundWave.generated.h"
 
 struct FActiveSound;
@@ -105,7 +107,7 @@ struct FStreamedAudioPlatformData
 
 };
 
-UCLASS(hidecategories=Object, editinlinenew, BlueprintType, dependson=USoundGroups)
+UCLASS(hidecategories=Object, editinlinenew, BlueprintType)
 class ENGINE_API USoundWave : public USoundBase
 {
 	GENERATED_UCLASS_BODY()
diff --git a/Engine/Source/Runtime/Engine/Classes/Vehicles/WheeledVehicleMovementComponent.h b/Engine/Source/Runtime/Engine/Classes/Vehicles/WheeledVehicleMovementComponent.h
--- a/Engine/Source/Runtime/Engine/Classes/Vehicles/WheeledVehicleMovementComponent.h
+++ b/Engine/Source/Runtime/Engine/Classes/Vehicles/WheeledVehicleMovementComponent.h
@@ -3,6 +3,8 @@
 #pragma once
 
 #include "VehicleWheel.h"
+#include "Curves/CurveBase.h"
+
 #include "WheeledVehicleMovementComponent.generated.h"
 
 #if WITH_VEHICLE
@@ -160,7 +162,7 @@ struct FVehicleInputRate
 /**
  * Component to handle the vehicle simulation for an actor.
  */
-UCLASS(Abstract, dependson=UCurveBase, hidecategories=(PlanarMovement, "Components|Movement|Planar", Activation, "Components|Activation"))
+UCLASS(Abstract, hidecategories=(PlanarMovement, "Components|Movement|Planar", Activation, "Components|Activation"))
 class ENGINE_API UWheeledVehicleMovementComponent : public UPawnMovementComponent
 {
 	GENERATED_UCLASS_BODY()
diff --git a/Engine/Source/Runtime/GameplayAbilities/Classes/GameplayAbilityBlueprint.h b/Engine/Source/Runtime/GameplayAbilities/Classes/GameplayAbilityBlueprint.h
--- a/Engine/Source/Runtime/GameplayAbilities/Classes/GameplayAbilityBlueprint.h
+++ b/Engine/Source/Runtime/GameplayAbilities/Classes/GameplayAbilityBlueprint.h
@@ -9,7 +9,7 @@
  * A Gameplay Ability Blueprint is essentially a specialized Blueprint whose graphs control a gameplay ability.
  */
 
-UCLASS(dependson = (UBlueprint), BlueprintType)
+UCLASS(BlueprintType)
 class GAMEPLAYABILITIES_API UGameplayAbilityBlueprint : public UBlueprint
 {
 	GENERATED_UCLASS_BODY()
diff --git a/Engine/Source/Runtime/MovieSceneCoreTypes/Classes/MovieSceneTransformSection.h b/Engine/Source/Runtime/MovieSceneCoreTypes/Classes/MovieSceneTransformSection.h
--- a/Engine/Source/Runtime/MovieSceneCoreTypes/Classes/MovieSceneTransformSection.h
+++ b/Engine/Source/Runtime/MovieSceneCoreTypes/Classes/MovieSceneTransformSection.h
@@ -2,12 +2,14 @@
 
 #pragma once
 
+#include "MovieSceneSection.h"
+
 #include "MovieSceneTransformSection.generated.h"
 
 /**
  * A transform section
  */
-UCLASS( DependsOn=UMovieSceneSection, MinimalAPI )
+UCLASS(MinimalAPI)
 class UMovieSceneTransformSection : public UMovieSceneSection
 {
 	GENERATED_UCLASS_BODY()
diff --git a/Engine/Source/Runtime/MovieSceneCoreTypes/Classes/MovieSceneVectorSection.h b/Engine/Source/Runtime/MovieSceneCoreTypes/Classes/MovieSceneVectorSection.h
--- a/Engine/Source/Runtime/MovieSceneCoreTypes/Classes/MovieSceneVectorSection.h
+++ b/Engine/Source/Runtime/MovieSceneCoreTypes/Classes/MovieSceneVectorSection.h
@@ -2,12 +2,14 @@
 
 #pragma once
 
+#include "MovieSceneSection.h"
+
 #include "MovieSceneVectorSection.generated.h"
 
 /**
  * A vector section
  */
-UCLASS( DependsOn=UMovieSceneSection, MinimalAPI )
+UCLASS(MinimalAPI )
 class UMovieSceneVectorSection : public UMovieSceneSection
 {
 	GENERATED_UCLASS_BODY()
diff --git a/Engine/Source/Runtime/MovieSceneCoreTypes/Classes/MovieSceneVectorTrack.h b/Engine/Source/Runtime/MovieSceneCoreTypes/Classes/MovieSceneVectorTrack.h
--- a/Engine/Source/Runtime/MovieSceneCoreTypes/Classes/MovieSceneVectorTrack.h
+++ b/Engine/Source/Runtime/MovieSceneCoreTypes/Classes/MovieSceneVectorTrack.h
@@ -5,12 +5,14 @@
 #include "MovieScene.h"
 #include "MovieSceneTrack.h"
 #include "MovieScenePropertyTrack.h"
+#include "MovieScenePropertyTrack.h"
+
 #include "MovieSceneVectorTrack.generated.h"
 
 /**
  * Handles manipulation of component transforms in a movie scene
  */
-UCLASS( dependson=UMovieScenePropertyTrack, MinimalAPI )
+UCLASS(MinimalAPI)
 class UMovieSceneVectorTrack : public UMovieScenePropertyTrack
 {
 	GENERATED_UCLASS_BODY()
diff --git a/Engine/Source/Runtime/Online/OnlineSubsystemUtils/Classes/PartyBeaconClient.h b/Engine/Source/Runtime/Online/OnlineSubsystemUtils/Classes/PartyBeaconClient.h
--- a/Engine/Source/Runtime/Online/OnlineSubsystemUtils/Classes/PartyBeaconClient.h
+++ b/Engine/Source/Runtime/Online/OnlineSubsystemUtils/Classes/PartyBeaconClient.h
@@ -1,7 +1,10 @@
 // Copyright 1998-2014 Epic Games, Inc. All Rights Reserved.
 
 #pragma once
+
 #include "Runtime/Online/OnlineSubsystemUtils/Classes/OnlineBeaconClient.h"
+#include "PartyBeaconHost.h"
+
 #include "PartyBeaconClient.generated.h"
 
 /**
@@ -19,7 +22,7 @@ DECLARE_DELEGATE(FOnHostConnectionFailure);
 /**
  * A beacon client used for making reservations with an existing game session
  */
-UCLASS(transient, notplaceable, config=Engine, DependsOn=APartyBeaconHost)
+UCLASS(transient, notplaceable, config=Engine)
 class ONLINESUBSYSTEMUTILS_API APartyBeaconClient : public AOnlineBeaconClient
 {
 	GENERATED_UCLASS_BODY()
diff --git a/Engine/Source/Runtime/UMG/Public/Blueprint/WidgetBlueprint.h b/Engine/Source/Runtime/UMG/Public/Blueprint/WidgetBlueprint.h
--- a/Engine/Source/Runtime/UMG/Public/Blueprint/WidgetBlueprint.h
+++ b/Engine/Source/Runtime/UMG/Public/Blueprint/WidgetBlueprint.h
@@ -2,6 +2,8 @@
 
 #pragma once
 
+#include "Engine/Blueprint.h"
+#include "UserWidget.h"
 #include "WidgetBlueprintGeneratedClass.h"
 
 #include "WidgetBlueprint.generated.h"
@@ -38,7 +40,7 @@ class UMovieScene;
 /**
  * The widget blueprint enables extending UUserWidget the user extensible UWidget.
  */
-UCLASS(dependson=(UBlueprint, UUserWidget, UWidgetBlueprintGeneratedClass), BlueprintType)
+UCLASS(BlueprintType)
 class UMG_API UWidgetBlueprint : public UBlueprint
 {
 	GENERATED_UCLASS_BODY()
