Updated comments for a world composition flag

diff --git a/Engine/Source/Runtime/Engine/Classes/GameFramework/WorldSettings.h b/Engine/Source/Runtime/Engine/Classes/GameFramework/WorldSettings.h
--- a/Engine/Source/Runtime/Engine/Classes/GameFramework/WorldSettings.h
+++ b/Engine/Source/Runtime/Engine/Classes/GameFramework/WorldSettings.h
@@ -193,11 +193,15 @@ class ENGINE_API AWorldSettings : public AInfo, public IInterface_AssetUserData
 	UPROPERTY(EditAnywhere, BlueprintReadOnly, Category=Config)
 	uint32 bEnableNavigationSystem:1;
 
-	/** if set to true enables tools for composing a tiled world */
+	/** 
+	 * Enables tools for composing a tiled world. 
+	 * Level has to be saved to enable this option.
+	 * Warning: Enabling this option will remove all streaming levels from your persistent level.
+	 */
 	UPROPERTY(EditAnywhere, BlueprintReadOnly, Category=Config)
 	uint32 bEnableWorldComposition:1;
 
-	/** if set to true world origin will shift to a camera position when camera goes far away from current origin */
+	/** World origin will shift to a camera position when camera goes far away from current origin */
 	UPROPERTY(EditAnywhere, BlueprintReadOnly, Category=Config, meta=(editcondition = "bEnableWorldComposition"))
 	uint32 bEnableWorldOriginRebasing:1;
 		
