Fix for a foliage serialization, when folliage settings is null

diff --git a/Engine/Source/Runtime/Engine/Private/InstancedFoliage.cpp b/Engine/Source/Runtime/Engine/Private/InstancedFoliage.cpp
--- a/Engine/Source/Runtime/Engine/Private/InstancedFoliage.cpp
+++ b/Engine/Source/Runtime/Engine/Private/InstancedFoliage.cpp
@@ -1498,7 +1498,14 @@ void AInstancedFoliageActor::Serialize(FArchive& Ar)
 #if WITH_EDITORONLY_DATA
 			NewMeshInfo.Instances = MoveTemp(OldMeshInfo.Value.Instances);
 #endif
+
 			UFoliageType_InstancedStaticMesh* FoliageType = OldMeshInfo.Value.Settings;
+			if (FoliageType == nullptr)
+			{
+				// If the Settings object was null, eg the user forgot to save their settings asset, create a new one.
+				FoliageType = ConstructObject<UFoliageType_InstancedStaticMesh>(UFoliageType_InstancedStaticMesh::StaticClass(), this);
+			}
+
 			if (FoliageType->Mesh == nullptr)
 			{
 				FoliageType->Modify();
