Removed FLogVisualizer (and other parts of visual logger) from compilation for builds with disabled logging #UE4
It fixes TTP #313754 "FLogVisualizer is being compiled when NO_LOGGING is defined"

diff --git a/Engine/Source/Developer/LogVisualizer/Private/LogVisualizer.cpp b/Engine/Source/Developer/LogVisualizer/Private/LogVisualizer.cpp
--- a/Engine/Source/Developer/LogVisualizer/Private/LogVisualizer.cpp
+++ b/Engine/Source/Developer/LogVisualizer/Private/LogVisualizer.cpp
@@ -3,7 +3,7 @@
 #include "LogVisualizerPCH.h"
 #include "CollisionDebugDrawingPublic.h"
 
-
+#if ENABLE_VISUAL_LOG
 //////////////////////////////////////////////////////////////////////////
 
 void FLogVisualizer::SummonUI(UWorld* InWorld) 
@@ -134,3 +134,5 @@ int32 FLogVisualizer::GetLogIndexForActor(const AActor* Actor)
 
 	return ResultIndex;
 }
+
+#endif //ENABLE_VISUAL_LOG
diff --git a/Engine/Source/Developer/LogVisualizer/Private/LogVisualizer.h b/Engine/Source/Developer/LogVisualizer/Private/LogVisualizer.h
--- a/Engine/Source/Developer/LogVisualizer/Private/LogVisualizer.h
+++ b/Engine/Source/Developer/LogVisualizer/Private/LogVisualizer.h
@@ -4,6 +4,8 @@
 
 #include "VisualLog.h"
 
+#if ENABLE_VISUAL_LOG
+
 /** Actual implementation of LogVisualizer, private inside module */
 class FLogVisualizer : public ILogVisualizer
 {
@@ -51,4 +53,6 @@ private:
 	FVisLogsChangedEvent	LogAddedEvent;
 	TWeakObjectPtr<UWorld>	World;
 	TWeakPtr<SWindow>		LogWindow;
-};
\ No newline at end of file
+};
+
+#endif //ENABLE_VISUAL_LOG
\ No newline at end of file
diff --git a/Engine/Source/Developer/LogVisualizer/Private/LogVisualizerModule.cpp b/Engine/Source/Developer/LogVisualizer/Private/LogVisualizerModule.cpp
--- a/Engine/Source/Developer/LogVisualizer/Private/LogVisualizerModule.cpp
+++ b/Engine/Source/Developer/LogVisualizer/Private/LogVisualizerModule.cpp
@@ -7,7 +7,9 @@ DEFINE_LOG_CATEGORY(LogLogVisualizer);
 
 void FLogVisualizerModule::StartupModule() 
 {
+#if ENABLE_VISUAL_LOG
 	LogVisualizer = new FLogVisualizer();
+#endif
 }
 
 void FLogVisualizerModule::ShutdownModule() 
diff --git a/Engine/Source/Developer/LogVisualizer/Private/SLogBar.cpp b/Engine/Source/Developer/LogVisualizer/Private/SLogBar.cpp
--- a/Engine/Source/Developer/LogVisualizer/Private/SLogBar.cpp
+++ b/Engine/Source/Developer/LogVisualizer/Private/SLogBar.cpp
@@ -104,7 +104,7 @@ int32 SLogBar::OnPaint( const FGeometry& AllottedGeometry, const FSlateRect& MyC
 				FillImage,
 				ForegroundClippingRect,
 				DrawEffects,
-				SLogVisualizer::GetColorForUsedCategory(0)
+				FColor(0xff00A480)
 			);
 		}
 
diff --git a/Engine/Source/Developer/LogVisualizer/Private/SLogVisualizer.cpp b/Engine/Source/Developer/LogVisualizer/Private/SLogVisualizer.cpp
--- a/Engine/Source/Developer/LogVisualizer/Private/SLogVisualizer.cpp
+++ b/Engine/Source/Developer/LogVisualizer/Private/SLogVisualizer.cpp
@@ -23,6 +23,8 @@
 
 #include "SFilterList.h"
 
+#if ENABLE_VISUAL_LOG
+
 #define LOCTEXT_NAMESPACE "SLogVisualizer"
 
 const FName SLogVisualizer::NAME_LogName = TEXT("LogName");
@@ -1900,4 +1902,6 @@ void SLogVisualizer::SaveSelectedLogs(FString& Filename)
 	}
 }
 
-#undef LOCTEXT_NAMESPACE
\ No newline at end of file
+#undef LOCTEXT_NAMESPACE
+
+#endif //ENABLE_VISUAL_LOG
diff --git a/Engine/Source/Developer/LogVisualizer/Private/SLogVisualizer.h b/Engine/Source/Developer/LogVisualizer/Private/SLogVisualizer.h
--- a/Engine/Source/Developer/LogVisualizer/Private/SLogVisualizer.h
+++ b/Engine/Source/Developer/LogVisualizer/Private/SLogVisualizer.h
@@ -57,6 +57,7 @@ struct FLogStatusItem
 /** Main LogVisualizer UI widget */
 class SLogVisualizer : public SCompoundWidget
 {
+#if ENABLE_VISUAL_LOG
 public:
 	static const FName NAME_LogName;
 	static const FName NAME_StartTime;
@@ -271,4 +272,5 @@ private:
 	TSharedPtr<STimeline> Timeline;
 	TSharedPtr<SScrollBar> ScrollBar;
 	TSharedPtr<SSlider> ZoomSlider;
+#endif
 };
diff --git a/Engine/Source/Developer/LogVisualizer/Private/SLogsTableRow.cpp b/Engine/Source/Developer/LogVisualizer/Private/SLogsTableRow.cpp
--- a/Engine/Source/Developer/LogVisualizer/Private/SLogsTableRow.cpp
+++ b/Engine/Source/Developer/LogVisualizer/Private/SLogsTableRow.cpp
@@ -3,6 +3,7 @@
 #include "LogVisualizerPCH.h"
 #include "SLogBar.h"
 
+#if ENABLE_VISUAL_LOG
 void SLogsTableRow::Construct(const FArguments& InArgs, const TSharedRef<STableViewBase>& InOwnerTableView)
 {
 	Item = InArgs._Item;
@@ -108,3 +109,4 @@ void SLogsTableRow::OnBarGeometryChanged( FGeometry Geometry )
 		OwnerVisualizerWidget->GetTimeline()->SetDrawingGeometry( Geometry );
 	}
 }
+#endif //ENABLE_VISUAL_LOG
diff --git a/Engine/Source/Developer/LogVisualizer/Private/SLogsTableRow.h b/Engine/Source/Developer/LogVisualizer/Private/SLogsTableRow.h
--- a/Engine/Source/Developer/LogVisualizer/Private/SLogsTableRow.h
+++ b/Engine/Source/Developer/LogVisualizer/Private/SLogsTableRow.h
@@ -7,6 +7,7 @@ struct FLogsListItem;
 /** Implements a row widget for log list. */
 class SLogsTableRow : public SMultiColumnTableRow< TSharedPtr<FLogsListItem> >
 {
+#if ENABLE_VISUAL_LOG
 	typedef SMultiColumnTableRow< TSharedPtr<FLogsListItem> > Super;
 public:
 
@@ -38,4 +39,5 @@ private:
 	TWeakPtr<SLogVisualizer> OwnerVisualizerWidgetPtr;
 
 	TSharedPtr<class SLogBar> LogBar;
+#endif //ENABLE_VISUAL_LOG
 };
\ No newline at end of file
diff --git a/Engine/Source/Runtime/Engine/Private/UnrealEngine.cpp b/Engine/Source/Runtime/Engine/Private/UnrealEngine.cpp
--- a/Engine/Source/Runtime/Engine/Private/UnrealEngine.cpp
+++ b/Engine/Source/Runtime/Engine/Private/UnrealEngine.cpp
@@ -6333,6 +6333,7 @@ void DrawStatsHUD( UWorld* World, FViewport* Viewport, FCanvas* Canvas, UCanvas*
 				MessageY += 20;
 			}
 
+#if ENABLE_VISUAL_LOG
 			FVisualLog* VisLog = FVisualLog::Get();
 			if (VisLog && (VisLog->IsRecording() || VisLog->IsRecordingOnServer()))
 			{
@@ -6348,7 +6349,7 @@ void DrawStatsHUD( UWorld* World, FViewport* Viewport, FCanvas* Canvas, UCanvas*
 				Canvas->DrawItem(SmallTextItem);
 				SmallTextItem.SetColor(FLinearColor::White);
 			}
-
+#endif
 
 			/* @todo ue4 temporarily disabled
 			AWorldSettings* WorldSettings = World->GetWorldSettings();
diff --git a/Engine/Source/Runtime/Engine/Private/VisualLog.cpp b/Engine/Source/Runtime/Engine/Private/VisualLog.cpp
--- a/Engine/Source/Runtime/Engine/Private/VisualLog.cpp
+++ b/Engine/Source/Runtime/Engine/Private/VisualLog.cpp
@@ -556,6 +556,7 @@ void FVisualLog::LogLine(const AActor* Actor, const FName& CategoryName, ELogVer
 		Entry->LogLines.Add(FVisLogEntry::FLogLine(CategoryName, Verbosity, Line));
 	}
 }
+#endif // ENABLE_VISUAL_LOG
 
 #if !(UE_BUILD_SHIPPING || UE_BUILD_TEST)
 
@@ -569,6 +570,7 @@ public:
 	{
 		if (FParse::Command(&Cmd, TEXT("VISLOG")))
 		{
+#if ENABLE_VISUAL_LOG
 			FString Command = FParse::Token(Cmd, 0);
 			if (Command == TEXT("record"))
 			{
@@ -601,11 +603,13 @@ public:
 				FLogVisualizerModule::Get()->SummonUI(InWorld);
 				return true;
 			}
+#else
+			UE_LOG(LogVisual, Warning, TEXT("Unable to open LogVisualizer - logs are disabled"));
+#endif
 		}
 		return false;
 	}
 } LogVisualizerExec;
 
-#endif // !(UE_BUILD_SHIPPING || UE_BUILD_TEST)
 
 #endif // ENABLE_VISUAL_LOG
\ No newline at end of file
