Shader compilation jobs that are superceded by new jobs are now removed from the compile queue.
This not only eliminates the issue of spamming jobs with Undos or undo/redo, it should also help speed up iteration times.

TTP 338558 - Editor: Material Shaders: Pressing Ctrl+Z repeatedly will increase the numbers of shaders compiling infinitely

diff --git a/Engine/Source/Editor/MaterialEditor/Private/MaterialEditor.cpp b/Engine/Source/Editor/MaterialEditor/Private/MaterialEditor.cpp
--- a/Engine/Source/Editor/MaterialEditor/Private/MaterialEditor.cpp
+++ b/Engine/Source/Editor/MaterialEditor/Private/MaterialEditor.cpp
@@ -3125,13 +3125,6 @@ void FMaterialEditor::UndoGraphAction()
 	{
 		Material->BuildEditorParameterList();
 	}
-
-	// Update the current preview material.
-	UpdatePreviewMaterial();
-
-	RefreshExpressionPreviews();
-	GraphEditor->NotifyGraphChanged();
-	SetMaterialDirty();
 }
 
 void FMaterialEditor::RedoGraphAction()
@@ -3147,26 +3140,23 @@ void FMaterialEditor::RedoGraphAction()
 		Material->BuildEditorParameterList();
 	}
 
-	// Update the current preview material.
-	UpdatePreviewMaterial();
-
-	RefreshExpressionPreviews();
-	GraphEditor->NotifyGraphChanged();
-	SetMaterialDirty();
 }
 
 void FMaterialEditor::PostUndo(bool bSuccess)
 {
-	GraphEditor->ClearSelectionSet();
-
-	Material->BuildEditorParameterList();
+	if (bSuccess)
+	{	
+		GraphEditor->ClearSelectionSet();
+		
+		Material->BuildEditorParameterList();
 
-	// Update the current preview material.
-	UpdatePreviewMaterial();
+		// Update the current preview material.
+		UpdatePreviewMaterial();
 
-	RefreshExpressionPreviews();
-	GraphEditor->NotifyGraphChanged();
-	SetMaterialDirty();
+		RefreshExpressionPreviews();
+		GraphEditor->NotifyGraphChanged();
+		SetMaterialDirty();
+	}
 }
 
 void FMaterialEditor::NotifyPreChange(UProperty* PropertyAboutToChange)
diff --git a/Engine/Source/Editor/MaterialEditor/Private/MaterialEditor.h b/Engine/Source/Editor/MaterialEditor/Private/MaterialEditor.h
--- a/Engine/Source/Editor/MaterialEditor/Private/MaterialEditor.h
+++ b/Engine/Source/Editor/MaterialEditor/Private/MaterialEditor.h
@@ -34,6 +34,7 @@ public:
 
 	~FMatExpressionPreview()
 	{
+		CancelCompilation();
 		ReleaseResource();
 	}
 
diff --git a/Engine/Source/Runtime/Engine/Classes/Materials/Material.h b/Engine/Source/Runtime/Engine/Classes/Materials/Material.h
--- a/Engine/Source/Runtime/Engine/Classes/Materials/Material.h
+++ b/Engine/Source/Runtime/Engine/Classes/Materials/Material.h
@@ -794,6 +794,11 @@ public:
 	ENGINE_API static void AddReferencedObjects(UObject* InThis, FReferenceCollector& Collector);
 	// End UObject Interface
 
+#if WITH_EDITOR
+	/** Cancels any currently outstanding compilation jobs for this material. Useful in the material editor when some edits superceds existing, in flight compilation jobs.*/
+	ENGINE_API virtual void CancelOutstandingCompilation();
+#endif // WITH_EDITOR
+
 	/**
 	 * Return the default material, loading it if necessary
 	 */
diff --git a/Engine/Source/Runtime/Engine/Private/Materials/Material.cpp b/Engine/Source/Runtime/Engine/Private/Materials/Material.cpp
--- a/Engine/Source/Runtime/Engine/Private/Materials/Material.cpp
+++ b/Engine/Source/Runtime/Engine/Private/Materials/Material.cpp
@@ -2197,6 +2197,9 @@ void UMaterial::PostEditChangeProperty(FPropertyChangedEvent& PropertyChangedEve
 	Super::PostEditChangeProperty(PropertyChangedEvent);
 
 	UProperty* PropertyThatChanged = PropertyChangedEvent.Property;
+	
+	//Cancel any current compilation jobs that are in flight for this material.
+	CancelOutstandingCompilation();
 
 	// check for distortion in material 
 	{
@@ -2701,6 +2704,16 @@ void UMaterial::AddReferencedObjects(UObject* InThis, FReferenceCollector& Colle
 	Super::AddReferencedObjects(This, Collector);
 }
 
+#if WITH_EDITOR
+void UMaterial::CancelOutstandingCompilation()
+{
+	if (FMaterialResource* Res = GetMaterialResource(GRHIFeatureLevel))
+	{
+		Res->CancelCompilation();
+	}
+}
+#endif
+
 void UMaterial::UpdateMaterialShaders(TArray<FShaderType*>& ShaderTypesToFlush, TArray<const FVertexFactoryType*>& VFTypesToFlush, EShaderPlatform ShaderPlatform)
 {
 	// Create a material update context so we can safely update materials.
diff --git a/Engine/Source/Runtime/Engine/Private/Materials/MaterialShared.cpp b/Engine/Source/Runtime/Engine/Private/Materials/MaterialShared.cpp
--- a/Engine/Source/Runtime/Engine/Private/Materials/MaterialShared.cpp
+++ b/Engine/Source/Runtime/Engine/Private/Materials/MaterialShared.cpp
@@ -347,19 +347,35 @@ EMaterialTessellationMode FMaterial::GetTessellationMode() const
 	return MTM_NoTessellation; 
 };
 
-void FMaterial::FinishCompilation()
+void FMaterial::GetShaderMapIDsWithUnfinishedCompilation(TArray<int32>& ShaderMapIds)
 {
-	TArray<int32> ShaderMapIdsToFinish;
-
-	// Build an array of the shader map Id's that we need to be compiled
+	// Build an array of the shader map Id's are not finished compiling.
 	if (GameThreadShaderMap && !GameThreadShaderMap->IsCompilationFinalized())
 	{
-		ShaderMapIdsToFinish.Add(GameThreadShaderMap->GetCompilingId());
+		ShaderMapIds.Add(GameThreadShaderMap->GetCompilingId());
 	}
 	else if (OutstandingCompileShaderMapId != INDEX_NONE)
 	{
-		ShaderMapIdsToFinish.Add(OutstandingCompileShaderMapId);
+		ShaderMapIds.Add(OutstandingCompileShaderMapId);
+	}
+}
+
+void FMaterial::CancelCompilation()
+{
+	TArray<int32> ShaderMapIdsToCancel;
+	GetShaderMapIDsWithUnfinishedCompilation(ShaderMapIdsToCancel);
+
+	if (ShaderMapIdsToCancel.Num() > 0)
+	{
+		// Cancel all compile jobs for these shader maps.
+		GShaderCompilingManager->CancelCompilation(*GetFriendlyName(), ShaderMapIdsToCancel);
 	}
+}
+
+void FMaterial::FinishCompilation()
+{
+	TArray<int32> ShaderMapIdsToFinish;
+	GetShaderMapIDsWithUnfinishedCompilation(ShaderMapIdsToFinish);
 
 	if (ShaderMapIdsToFinish.Num() > 0)
 	{
diff --git a/Engine/Source/Runtime/Engine/Private/ShaderCompiler/ShaderCompiler.cpp b/Engine/Source/Runtime/Engine/Private/ShaderCompiler/ShaderCompiler.cpp
--- a/Engine/Source/Runtime/Engine/Private/ShaderCompiler/ShaderCompiler.cpp
+++ b/Engine/Source/Runtime/Engine/Private/ShaderCompiler/ShaderCompiler.cpp
@@ -1744,6 +1744,51 @@ bool FShaderCompilingManager::HandlePotentialRetryOnError(TMap<int32, FShaderMap
 	return bRetryCompile;
 }
 
+void FShaderCompilingManager::CancelCompilation(const TCHAR* MaterialName, const TArray<int32>& ShaderMapIdsToCancel)
+{
+	check(!FPlatformProperties::RequiresCookedData());
+	UE_LOG(LogShaders, Log, TEXT("CancelCompilation %s "), MaterialName ? MaterialName : TEXT(""));
+
+	// Lock CompileQueueSection so we can access the input and output queues
+	FScopeLock Lock(&CompileQueueSection);
+
+	int32 TotalNumJobsRemoved = 0;
+	for (int32 IdIndex = 0; IdIndex < ShaderMapIdsToCancel.Num(); ++IdIndex)
+	{
+		int32 MapIdx = ShaderMapIdsToCancel[IdIndex];
+		if (FShaderMapCompileResults* ShaderMapJob = ShaderMapJobs.Find(MapIdx))
+		{
+			int32 NumJobsRemoved = 0;
+
+			int32 JobIndex = CompileQueue.Num();
+			while ( --JobIndex >= 0 )
+			{
+				if (FShaderCompileJob* Job = CompileQueue[JobIndex])
+				{
+					if (Job->Id == MapIdx)
+					{
+						++TotalNumJobsRemoved;
+						++NumJobsRemoved;
+						CompileQueue.RemoveAt(JobIndex, 1, false);
+					}
+				}
+			}
+
+			ShaderMapJob->NumJobsQueued -= NumJobsRemoved;
+
+			if (ShaderMapJob->NumJobsQueued == 0)
+			{
+				//We've removed all the jobs for this shader map so remove it.
+				ShaderMapJobs.Remove(MapIdx);
+			}
+		}
+	}
+	CompileQueue.Shrink();
+
+	// Using atomics to update NumOutstandingJobs since it is read outside of the critical section
+	FPlatformAtomics::InterlockedAdd(&NumOutstandingJobs, -TotalNumJobsRemoved);
+}
+
 void FShaderCompilingManager::FinishCompilation(const TCHAR* MaterialName, const TArray<int32>& ShaderMapIdsToFinishCompiling)
 {
 	check(!FPlatformProperties::RequiresCookedData());
diff --git a/Engine/Source/Runtime/Engine/Public/MaterialShared.h b/Engine/Source/Runtime/Engine/Public/MaterialShared.h
--- a/Engine/Source/Runtime/Engine/Public/MaterialShared.h
+++ b/Engine/Source/Runtime/Engine/Public/MaterialShared.h
@@ -943,6 +943,11 @@ public:
 	*/
 	virtual void NotifyCompilationFinished() { }
 
+	/**
+	* Cancels all outstanding compilation jobs for this materail.
+	*/
+	ENGINE_API void CancelCompilation();
+
 	/** 
 	 * Blocks until compilation has completed. Returns immediately if a compilation is not outstanding.
 	 */
@@ -1055,6 +1060,11 @@ public:
 
 protected:
 
+	/**
+	* Fills the passed array with IDs of shader maps unfinished compilation jobs.
+	*/
+	void GetShaderMapIDsWithUnfinishedCompilation(TArray<int32>& ShaderMapIds);
+
 	/** Entry point for compiling a specific material property.  This must call SetMaterialProperty. */
 	virtual int32 CompileProperty(EMaterialProperty Property,EShaderFrequency InShaderFrequency,class FMaterialCompiler* Compiler) const = 0;
 
diff --git a/Engine/Source/Runtime/Engine/Public/ShaderCompiler.h b/Engine/Source/Runtime/Engine/Public/ShaderCompiler.h
--- a/Engine/Source/Runtime/Engine/Public/ShaderCompiler.h
+++ b/Engine/Source/Runtime/Engine/Public/ShaderCompiler.h
@@ -277,6 +277,11 @@ public:
 	 */
 	ENGINE_API void AddJobs(TArray<FShaderCompileJob*>& NewJobs, bool bApplyCompletedShaderMapForRendering, bool bOptimizeForLowLatency);
 
+	/**
+	* Removes all outstanding compile jobs for the passed shader maps.
+	*/
+	ENGINE_API void CancelCompilation(const TCHAR* MaterialName, const TArray<int32>& ShaderMapIdsToCancel);
+
 	/** 
 	 * Blocks until completion of the requested shader maps.  
 	 * This will not assign the shader map to any materials, the caller is responsible for that.
