Expose double sided flag to allow scene queries on double sided geometry

diff --git a/Engine/Source/Editor/UnrealEd/Private/StaticMeshEdit.cpp b/Engine/Source/Editor/UnrealEd/Private/StaticMeshEdit.cpp
--- a/Engine/Source/Editor/UnrealEd/Private/StaticMeshEdit.cpp
+++ b/Engine/Source/Editor/UnrealEd/Private/StaticMeshEdit.cpp
@@ -938,6 +938,7 @@ void RestoreExistingMeshData(struct ExistingStaticMeshData* ExistingMeshDataPtr,
 			{
 				// New collision geometry, but we still want the original settings
 				NewMesh->BodySetup->PhysicsType = ExistingMeshDataPtr->ExistingBodySetup->PhysicsType;
+				NewMesh->BodySetup->bDoubleSidedGeometry = ExistingMeshDataPtr->ExistingBodySetup->bDoubleSidedGeometry;
 				NewMesh->BodySetup->BoneName = ExistingMeshDataPtr->ExistingBodySetup->BoneName;
 				NewMesh->BodySetup->PhysMaterial = ExistingMeshDataPtr->ExistingBodySetup->PhysMaterial;
 				NewMesh->BodySetup->CollisionTraceFlag = ExistingMeshDataPtr->ExistingBodySetup->CollisionTraceFlag;
diff --git a/Engine/Source/Runtime/Engine/Classes/PhysicsEngine/BodySetup.h b/Engine/Source/Runtime/Engine/Classes/PhysicsEngine/BodySetup.h
--- a/Engine/Source/Runtime/Engine/Classes/PhysicsEngine/BodySetup.h
+++ b/Engine/Source/Runtime/Engine/Classes/PhysicsEngine/BodySetup.h
@@ -432,6 +432,13 @@ class UBodySetup : public UObject
 	UPROPERTY(Transient)
 	uint32 bMeshCollideAll:1;
 
+	/**
+	*	If true, the physics triangle mesh will use double sided faces when doing scene queries.
+	*	This is useful for planes and single sided meshes that need traces to work on both sides.
+	*/
+	UPROPERTY(EditAnywhere, Category=Physics)
+	uint32 bDoubleSidedGeometry : 1;
+
 	/**	Should we generate data necessary to support collision on normal (non-mirrored) versions of this body. */
 	UPROPERTY()
 	uint32 bGenerateNonMirroredCollision:1;
diff --git a/Engine/Source/Runtime/Engine/Private/Collision/CollisionConversions.cpp b/Engine/Source/Runtime/Engine/Private/Collision/CollisionConversions.cpp
--- a/Engine/Source/Runtime/Engine/Private/Collision/CollisionConversions.cpp
+++ b/Engine/Source/Runtime/Engine/Private/Collision/CollisionConversions.cpp
@@ -151,7 +151,7 @@ static void TransformNormalToShapeSpace(const PxMeshScale& meshScale, const PxVe
 
 
 /** Util to find the normal of the face that we hit */
-static bool FindGeomOpposingNormal(const PxLocationHit& PHit, FVector& OutNormal, FVector& OutPointOnGeom)
+static bool FindGeomOpposingNormal(const PxLocationHit& PHit, const FVector & Direction, FVector& OutNormal, FVector& OutPointOnGeom)
 {
 	OutPointOnGeom = (PHit.flags & PxHitFlag::ePOSITION) ? P2UVector(PHit.position) : FVector::ZeroVector;
 
@@ -196,6 +196,13 @@ static bool FindGeomOpposingNormal(const PxLocationHit& PHit, FVector& OutNormal
 
 		const PxVec3 PWorldTriNormal = PShapeWorldPose.rotate(PLocalTriNormal);
 		OutNormal = P2UVector(PWorldTriNormal);
+
+		if (PTriMeshGeom.meshFlags & PxMeshGeometryFlag::eDOUBLE_SIDED)
+		{
+			//double sided mesh so we need to consider direction of query
+			float sign = FVector::DotProduct(OutNormal, Direction) > 0 ? -1.f : 1.f;
+			OutNormal *= sign;
+		}
 		
 		if (PTriMeshGeom.scale.isIdentity())
 		{
@@ -376,7 +383,7 @@ void ConvertQueryImpactHit(const PxLocationHit& PHit, FHitResult& OutResult, flo
 		// Compute better ImpactNormal. This is the same as Normal except when we hit convex/trimesh, and then its the most opposing normal of the geom at point of impact.
 		FVector PointOnGeom;
 		check(PHit.flags & PxHitFlag::ePOSITION);
-		if (FindGeomOpposingNormal(PHit, OutResult.ImpactNormal, PointOnGeom))
+		if (FindGeomOpposingNormal(PHit, TraceDir, OutResult.ImpactNormal, PointOnGeom))
 		{
 			ConvertConvexNormalToCapsuleNormal(Capsule->halfHeight, Capsule->radius, OutResult, &PointOnGeom);
 		}
@@ -393,7 +400,7 @@ void ConvertQueryImpactHit(const PxLocationHit& PHit, FHitResult& OutResult, flo
 		// Compute better ImpactNormal. This is the same as Normal except when we hit convex/trimesh, and then its the most opposing normal of the geom at point of impact.
 		FVector PointOnGeom;
 		check(PHit.flags & PxHitFlag::ePOSITION);
-		if (FindGeomOpposingNormal(PHit, OutResult.ImpactNormal, PointOnGeom))
+		if (FindGeomOpposingNormal(PHit, TraceDir, OutResult.ImpactNormal, PointOnGeom))
 		{
 			const FPlane GeomPlane(PointOnGeom, OutResult.ImpactNormal);
 			const float DistFromPlane = FMath::Abs(GeomPlane.PlaneDot(OutResult.Location));
@@ -556,6 +563,7 @@ static bool ConvertOverlappedShapeToImpactHit(const PxLocationHit& PHit, const F
 {
 	OutResult.TraceStart = StartLoc;
 	OutResult.TraceEnd = EndLoc;
+	FVector TraceDir = EndLoc - StartLoc;
 
 	const PxShape* PShape = PHit.shape;
 	const PxRigidActor* PActor = PHit.actor;
@@ -594,7 +602,7 @@ static bool ConvertOverlappedShapeToImpactHit(const PxLocationHit& PHit, const F
 		{
 			// PxMeshQuery::findOverlapTriangleMesh failed, but we might have a valid face index at impact, so use that to grab the normal from that face.
 			FVector PointOnGeom;
-			if (!FindGeomOpposingNormal(PHit, OutResult.ImpactNormal, PointOnGeom))
+			if (!FindGeomOpposingNormal(PHit, TraceDir, OutResult.ImpactNormal, PointOnGeom))
 			{
 				// That failed (invalid face index). This really shouldn't happen.
 				OutResult.ImpactNormal = (PHit.flags & PxHitFlag::eNORMAL) ? P2UVector(PHit.normal) : FVector(0.f, 0.f, 1.f);
diff --git a/Engine/Source/Runtime/Engine/Private/PhysicsEngine/BodySetup.cpp b/Engine/Source/Runtime/Engine/Private/PhysicsEngine/BodySetup.cpp
--- a/Engine/Source/Runtime/Engine/Private/PhysicsEngine/BodySetup.cpp
+++ b/Engine/Source/Runtime/Engine/Private/PhysicsEngine/BodySetup.cpp
@@ -69,6 +69,7 @@ void UBodySetup::CopyBodyPropertiesFrom(class UBodySetup* FromSetup)
 	DefaultInstance.CopyBodyInstancePropertiesFrom(&FromSetup->DefaultInstance);
 	PhysMaterial = FromSetup->PhysMaterial;
 	PhysicsType = FromSetup->PhysicsType;
+	bDoubleSidedGeometry = FromSetup->bDoubleSidedGeometry;
 }
 
 void UBodySetup::AddCollisionFrom(class UBodySetup* FromSetup)
@@ -491,9 +492,16 @@ void UBodySetup::AddShapesToRigidActor(PxRigidActor* PDestActor, FVector& Scale3
 		PxTriangleMesh* UseTriMesh = bUseNegX ? TriMeshNegX : TriMesh;
 		if(UseTriMesh != NULL)
 		{
+			
+
 			PxTriangleMeshGeometry PTriMeshGeom;
 			PTriMeshGeom.triangleMesh = bUseNegX ? TriMeshNegX : TriMesh;
 			PTriMeshGeom.scale.scale = U2PVector(Scale3DAbs);
+			if (bDoubleSidedGeometry)
+			{
+				PTriMeshGeom.meshFlags |= PxMeshGeometryFlag::eDOUBLE_SIDED;
+			}
+			
 
 			if(PTriMeshGeom.isValid())
 			{
