* fixed TTP 339866 CRASH: CASTLE: Castle zone randomly crashes

diff --git a/Engine/Source/Runtime/Engine/Private/PointLightSceneProxy.h b/Engine/Source/Runtime/Engine/Private/PointLightSceneProxy.h
--- a/Engine/Source/Runtime/Engine/Private/PointLightSceneProxy.h
+++ b/Engine/Source/Runtime/Engine/Private/PointLightSceneProxy.h
@@ -239,7 +239,9 @@ private:
 	void UpdateRadius(float ComponentRadius)
 	{
 		Radius = ComponentRadius;
-		InvRadius = 1.0f / ComponentRadius;
+
+		// Min to avoid div by 0 (NaN in InvRadius)
+		InvRadius = 1.0f / FMath::Max(0.00001f, ComponentRadius);
 	}
 };
 
