Better handling of cases where object of replicator gets GC'd

diff --git a/Engine/Source/Runtime/Engine/Private/DataChannel.cpp b/Engine/Source/Runtime/Engine/Private/DataChannel.cpp
--- a/Engine/Source/Runtime/Engine/Private/DataChannel.cpp
+++ b/Engine/Source/Runtime/Engine/Private/DataChannel.cpp
@@ -1644,7 +1644,10 @@ void UActorChannel::ProcessBunch( FInBunch & Bunch )
 	
 	for (auto RepComp = ReplicationMap.CreateIterator(); RepComp; ++RepComp)
 	{
-		RepComp.Value()->PostReceivedBunch();
+		if ( RepComp.Key().IsValid() )
+		{
+			RepComp.Value()->PostReceivedBunch();
+		}
 	}
 
 	// After all properties have been initialized, call PostNetInit. This should call BeginPlay() so initialization can be done with proper starting values.
@@ -1716,7 +1719,6 @@ bool UActorChannel::ReplicateActor()
 			{
 				RepComp.Value()->ForceRefreshUnreliableProperties();
 			}
-
 		}
 	}
 	else
@@ -1927,8 +1929,13 @@ void UActorChannel::BecomeDormant()
 
 bool UActorChannel::ReadyForDormancy(bool suppressLogs)
 {
-	for (auto MapIt = ReplicationMap.CreateIterator(); MapIt; ++MapIt)
+	for ( auto MapIt = ReplicationMap.CreateIterator(); MapIt; ++MapIt )
 	{
+		if ( !MapIt.Key().IsValid() )
+		{
+			continue;
+		}
+
 		if (!MapIt.Value()->ReadyForDormancy(suppressLogs))
 		{
 			return false;
diff --git a/Engine/Source/Runtime/Engine/Private/DataReplication.cpp b/Engine/Source/Runtime/Engine/Private/DataReplication.cpp
--- a/Engine/Source/Runtime/Engine/Private/DataReplication.cpp
+++ b/Engine/Source/Runtime/Engine/Private/DataReplication.cpp
@@ -770,6 +770,12 @@ bool FObjectReplicator::ReceivedBunch( FInBunch &Bunch, const FReplicationFlags
 
 void FObjectReplicator::PostReceivedBunch()
 {
+	if ( GetObject() == NULL )
+	{
+		UE_LOG( LogNet, Error, TEXT( "PostReceivedBunch: Object == NULL" ) );
+		return;
+	}
+
 	// Call PostNetReceive
 	const bool bIsServer = (OwningChannel->Connection->Driver->ServerConnection == NULL);
 	if (!bIsServer && bHasReplicatedProperties)
@@ -1012,6 +1018,12 @@ bool FObjectReplicator::ReplicateProperties( FOutBunch & Bunch, FReplicationFlag
 
 void FObjectReplicator::ForceRefreshUnreliableProperties()
 {
+	if ( GetObject() == NULL )
+	{
+		UE_LOG( LogNet, Error, TEXT( "ForceRefreshUnreliableProperties: Object == NULL" ) );
+		return;
+	}
+
 	check( !bOpenAckCalled );
 
 	RepLayout->OpenAcked( RepState );
@@ -1021,6 +1033,12 @@ void FObjectReplicator::ForceRefreshUnreliableProperties()
 
 void FObjectReplicator::PostSendBunch( FPacketIdRange & PacketRange, uint8 bReliable )
 {
+	if ( GetObject() == NULL )
+	{
+		UE_LOG( LogNet, Error, TEXT( "PostSendBunch: Object == NULL" ) );
+		return;
+	}
+
 	RepLayout->PostReplicate( RepState, PacketRange, bReliable ? true : false );
 
 	for ( int32 i = 0; i < LifetimeCustomDeltaProperties.Num(); i++ )
@@ -1097,6 +1115,12 @@ void FObjectReplicator::QueueRemoteFunctionBunch( UFunction* Func, FOutBunch &Bu
 
 bool FObjectReplicator::ReadyForDormancy(bool suppressLogs)
 {
+	if ( GetObject() == NULL )
+	{
+		UE_LOG( LogNet, Error, TEXT( "ReadyForDormancy: Object == NULL" ) );
+		return true;		// Technically, we don't want to hold up dormancy, but the owner needs to clean us up, so we warn
+	}
+
 	// Can't go dormant until last update produced no new property updates
 	if ( !bLastUpdateEmpty )
 	{
@@ -1126,6 +1150,12 @@ bool FObjectReplicator::ReadyForDormancy(bool suppressLogs)
 
 void FObjectReplicator::StartBecomingDormant()
 {
+	if ( GetObject() == NULL )
+	{
+		UE_LOG( LogNet, Error, TEXT( "StartBecomingDormant: Object == NULL" ) );
+		return;
+	}
+
 	bLastUpdateEmpty = false; // Ensure we get one more attempt to update properties
 }
 
