Fix game mode losing map default when a value is passed in the URL but fails to find a valid class.

diff --git a/Engine/Source/Runtime/Engine/Private/World.cpp b/Engine/Source/Runtime/Engine/Private/World.cpp
--- a/Engine/Source/Runtime/Engine/Private/World.cpp
+++ b/Engine/Source/Runtime/Engine/Private/World.cpp
@@ -2567,7 +2567,16 @@ bool UWorld::SetGameMode(const FURL& InURL)
 				GameEngine->LoadPackagesFully(this, FULLYLOAD_Game_PreLoadClass, *GameClassName);
 			}
 
-			GameClass = StaticLoadClass(AGameMode::StaticClass(), NULL, *GameClassName, NULL, LOAD_None, NULL);
+			// Don't overwrite the map's world settings if we failed to load the value off the command line parameter
+			UClass* GameModeParamClass = StaticLoadClass(AGameMode::StaticClass(), NULL, *GameClassName, NULL, LOAD_None, NULL);
+			if (GameModeParamClass)
+			{
+				GameClass = GameModeParamClass;
+			}
+			else
+			{
+				UE_LOG(LogWorld, Warning, TEXT("Failed to load game mode '%s' specified by URL options."), *GameClassName);
+			}
 		}
 
 		if (!GameClass)
