Only supporting component types return true from HasStaticLighting()
* Fixes issues where skeletal mesh components with Static Mobility would cause the 'lighting needs to be rebuilt' message

diff --git a/Engine/Source/Runtime/Engine/Classes/Components/InstancedStaticMeshComponent.h b/Engine/Source/Runtime/Engine/Classes/Components/InstancedStaticMeshComponent.h
--- a/Engine/Source/Runtime/Engine/Classes/Components/InstancedStaticMeshComponent.h
+++ b/Engine/Source/Runtime/Engine/Classes/Components/InstancedStaticMeshComponent.h
@@ -152,5 +152,13 @@ private:
 	/** Sets up new instance data to sensible defaults, creates physics counterparts if possible */
 	void SetupNewInstanceData(FInstancedStaticMeshInstanceData& InOutNewInstanceData, int32 InInstanceIndex, const FTransform& InInstanceTransform);
 
+protected:
+
+	/** Whether the component type supports static lighting. */
+	virtual bool SupportsStaticLighting() const override
+	{
+		//@todo - support for instanced meshes
+		return false;
+	}
 };
 
diff --git a/Engine/Source/Runtime/Engine/Classes/Components/ModelComponent.h b/Engine/Source/Runtime/Engine/Classes/Components/ModelComponent.h
--- a/Engine/Source/Runtime/Engine/Classes/Components/ModelComponent.h
+++ b/Engine/Source/Runtime/Engine/Classes/Components/ModelComponent.h
@@ -160,6 +160,13 @@ private:
 	friend void SetDebugLightmapSample(TArray<UActorComponent*>* Components, UModel* Model, int32 iSurf, FVector ClickLocation);
 	friend class UModel;
 	friend class FStaticLightingSystem;
+
+protected:
+	/** Whether the component type supports static lighting. */
+	virtual bool SupportsStaticLighting() const override
+	{
+		return true;
+	}
 };
 
 
diff --git a/Engine/Source/Runtime/Engine/Classes/Components/PrimitiveComponent.h b/Engine/Source/Runtime/Engine/Classes/Components/PrimitiveComponent.h
--- a/Engine/Source/Runtime/Engine/Classes/Components/PrimitiveComponent.h
+++ b/Engine/Source/Runtime/Engine/Classes/Components/PrimitiveComponent.h
@@ -1046,6 +1046,12 @@ protected:
 	*/
 	virtual void OnChildAttached(USceneComponent* ChildComponent) override;
 
+	/** Whether the component type supports static lighting. */
+	virtual bool SupportsStaticLighting() const 
+	{
+		return false;
+	}
+
 public:
 	virtual bool IsSimulatingPhysics(FName BoneName = NAME_None) const override;
 
diff --git a/Engine/Source/Runtime/Engine/Classes/Components/StaticMeshComponent.h b/Engine/Source/Runtime/Engine/Classes/Components/StaticMeshComponent.h
--- a/Engine/Source/Runtime/Engine/Classes/Components/StaticMeshComponent.h
+++ b/Engine/Source/Runtime/Engine/Classes/Components/StaticMeshComponent.h
@@ -327,7 +327,6 @@ public:
 		int32& VertexLightMapMemoryUsage, int32& VertexShadowMapMemoryUsage,
 		int32& StaticLightingResolution, bool& bIsUsingTextureMapping, bool& bHasLightmapTexCoords) const;
 
-
 	/**
 	 * Determines whether any of the component's LODs require override vertex color fixups
 	 *
@@ -370,25 +369,25 @@ private:
 
 	/** Update the vertex override colors */
 	void PrivateFixupOverrideColors( const TArray<int32>& LODsToUpdate );
-public:
-
-	void ReleaseResources();
-
 
+protected:
 
+	/** Whether the component type supports static lighting. */
+	virtual bool SupportsStaticLighting() const override
+	{
+		return true;
+	}
 
+public:
 
+	void ReleaseResources();
 
 	/** Allocates an implementation of FStaticLightingMesh that will handle static lighting for this component */
 	virtual class FStaticMeshStaticLightingMesh* AllocateStaticLightingMesh(int32 LODIndex, const TArray<ULightComponent*>& InRelevantLights);
 
-
-
-
 	/** Add or remove elements to have the size in the specified range. Reconstructs elements if MaxSize<MinSize */
 	void SetLODDataCount( const uint32 MinSize, const uint32 MaxSize );
 
-
 	/**
 	 *	Switches the static mesh component to use either Texture or Vertex static lighting.
 	 *
diff --git a/Engine/Source/Runtime/Engine/Classes/Landscape/LandscapeComponent.h b/Engine/Source/Runtime/Engine/Classes/Landscape/LandscapeComponent.h
--- a/Engine/Source/Runtime/Engine/Classes/Landscape/LandscapeComponent.h
+++ b/Engine/Source/Runtime/Engine/Classes/Landscape/LandscapeComponent.h
@@ -470,6 +470,14 @@ public:
 	friend struct FLandscapeComponentDataInterface;
 
 	void SetLOD(bool bForced, int32 InLODValue);
+
+protected:
+
+	/** Whether the component type supports static lighting. */
+	virtual bool SupportsStaticLighting() const override
+	{
+		return true;
+	}
 };
 
 
diff --git a/Engine/Source/Runtime/Engine/Private/PrimitiveComponent.cpp b/Engine/Source/Runtime/Engine/Private/PrimitiveComponent.cpp
--- a/Engine/Source/Runtime/Engine/Private/PrimitiveComponent.cpp
+++ b/Engine/Source/Runtime/Engine/Private/PrimitiveComponent.cpp
@@ -150,7 +150,7 @@ bool UPrimitiveComponent::IsEditorOnly() const
 
 bool UPrimitiveComponent::HasStaticLighting() const
 {
-	return (Mobility == EComponentMobility::Static) || bLightAsIfStatic;
+	return ((Mobility == EComponentMobility::Static) || bLightAsIfStatic) && SupportsStaticLighting();
 }
 
 
