Hidden DebugCameraController and LogVisualizerCameraComponent from showing up in controller-selection dropdowns #UE4

diff --git a/Engine/Source/Runtime/Engine/Classes/Debug/LogVisualizerCameraController.h b/Engine/Source/Runtime/Engine/Classes/Debug/LogVisualizerCameraController.h
--- a/Engine/Source/Runtime/Engine/Classes/Debug/LogVisualizerCameraController.h
+++ b/Engine/Source/Runtime/Engine/Classes/Debug/LogVisualizerCameraController.h
@@ -4,7 +4,7 @@
 
 #include "LogVisualizerCameraController.generated.h"
 
-UCLASS(config=Input)
+UCLASS(config=Input, hidedropdown)
 class ALogVisualizerCameraController : public ADebugCameraController
 {
 public:
diff --git a/Engine/Source/Runtime/Engine/Classes/Engine/DebugCameraController.h b/Engine/Source/Runtime/Engine/Classes/Engine/DebugCameraController.h
--- a/Engine/Source/Runtime/Engine/Classes/Engine/DebugCameraController.h
+++ b/Engine/Source/Runtime/Engine/Classes/Engine/DebugCameraController.h
@@ -12,7 +12,7 @@
 #pragma once
 #include "DebugCameraController.generated.h"
 
-UCLASS(config=Game)
+UCLASS(config=Game, hidedropdown)
 class ENGINE_API ADebugCameraController : public APlayerController
 {
 	GENERATED_UCLASS_BODY()
