Merged github pull request for new batched element blend mode that writes to dest alpha

diff --git a/Engine/Source/Runtime/Engine/Private/BatchedElements.cpp b/Engine/Source/Runtime/Engine/Private/BatchedElements.cpp
--- a/Engine/Source/Runtime/Engine/Private/BatchedElements.cpp
+++ b/Engine/Source/Runtime/Engine/Private/BatchedElements.cpp
@@ -299,6 +299,9 @@ static void SetBlendState(FRHICommandList& RHICmdList, ESimpleElementBlendMode B
 	case SE_BLEND_AlphaComposite:
 		RHICmdList.SetBlendState(TStaticBlendState<CW_RGBA, BO_Add, BF_One, BF_InverseSourceAlpha, BO_Add, BF_One, BF_InverseSourceAlpha>::GetRHI());
 		break;
+	case SE_BLEND_AlphaBlend:
+		RHICmdList.SetBlendState(TStaticBlendState<CW_RGBA,BO_Add,BF_SourceAlpha,BF_InverseSourceAlpha,BO_Add,BF_SourceAlpha,BF_InverseSourceAlpha>::GetRHI());
+		break;
 	}
 }
 
@@ -316,6 +319,7 @@ static void SetHitTestingBlendState(FRHICommandList& RHICmdList, ESimpleElementB
 		RHICmdList.SetBlendState(TStaticBlendState<CW_RGBA, BO_Add, BF_One, BF_Zero, BO_Add, BF_One, BF_Zero>::GetRHI());
 		break;
 	case SE_BLEND_AlphaComposite:
+	case SE_BLEND_AlphaBlend:
 	case SE_BLEND_Translucent:
 	case SE_BLEND_TranslucentDistanceField:
 	case SE_BLEND_TranslucentDistanceFieldShadowed:
diff --git a/Engine/Source/Runtime/Engine/Public/BatchedElements.h b/Engine/Source/Runtime/Engine/Public/BatchedElements.h
--- a/Engine/Source/Runtime/Engine/Public/BatchedElements.h
+++ b/Engine/Source/Runtime/Engine/Public/BatchedElements.h
@@ -22,6 +22,8 @@ enum ESimpleElementBlendMode
 	SE_BLEND_TranslucentDistanceField,
 	SE_BLEND_TranslucentDistanceFieldShadowed,
 	SE_BLEND_AlphaComposite,
+	// Like SE_BLEND_Translucent, but modifies destination alpha
+	SE_BLEND_AlphaBlend,
 
 	SE_BLEND_RGBA_MASK_START,
 	SE_BLEND_RGBA_MASK_END = SE_BLEND_RGBA_MASK_START+31, //Using 5bit bit-field for red, green, blue, alpha and desaturation
