#ue4 - Improve physx penetration handling and improve reported normals when in penetration.

- Enable auto MTD results within PhysX sweep code, rather than performing the test separately, yielding better accuracy.
- Prefer MTD result as the opposing normal and safe direction for movement during overlaps, if the MTD distance is non-zero.
- Initial overlaps try to use FindGeomOpposingNormal() if faceIndex is valid and distance is zero.
- Add FindGeomOpposingNormal() implementation for HeightFields.
- Remove our own computation of capsule and sphere impact point and normal, those are much better with the newer PhysX versions so this was unnecessary and slow.

diff --git a/Engine/Source/Runtime/Engine/Private/Collision/CollisionConversions.cpp b/Engine/Source/Runtime/Engine/Private/Collision/CollisionConversions.cpp
--- a/Engine/Source/Runtime/Engine/Private/Collision/CollisionConversions.cpp
+++ b/Engine/Source/Runtime/Engine/Private/Collision/CollisionConversions.cpp
@@ -8,6 +8,9 @@
 #include "CollisionDebugDrawing.h"
 #include "CollisionConversions.h"
 
+// Sentinel for invalid query results.
+static const PxQueryHit InvalidQueryHit;
+
 // Forward declare, I don't want to move the entire function right now or we lose change history.
 static bool ConvertOverlappedShapeToImpactHit(const PxLocationHit& PHit, const FVector& StartLoc, const FVector& EndLoc, FHitResult& OutResult, const PxGeometry& Geom, const PxTransform& QueryTM, const PxFilterData& QueryFilter, bool bReturnPhysMat);
 
@@ -36,92 +39,6 @@ static void CheckHitResultNormal(const FHitResult& OutResult, const TCHAR* Messa
 
 #endif // !(UE_BUILD_SHIPPING || UE_BUILD_TEST)
 
-/** 
- * Convert Convex Normal to Capsule Normal
- * For capsule, we'd like to get normal of capsule rather than convex, so the normal can be smooth when we move character
- * 
- * @param HalfHeight	: HalfHeight of the capsule
- * @param Radius		: Radius of the capsule
- * @param OutHit		: The result of hit from physX
- @ @param PointOnGeom	: If not null, use this point and the impact normal to recompute the impact point (which can yield a better normal).
- * @result OutHit.Normal is converted to Capsule Normal
- */
-static void ConvertConvexNormalToCapsuleNormal(float HalfHeight, float Radius, struct FHitResult& OutHit, const FVector* PointOnGeom)
-{
-	if ( OutHit.Time > 0.f )
-	{
-		// Impact point is contact point where it meet 
-		const float ContactZ = OutHit.ImpactPoint.Z;
-		// Hit.Location is center of the object
-		const FVector CenterOfCapsule = OutHit.Location; 
-		// We're looking for end of the normal for this capsule
-		FVector NormalEnd = CenterOfCapsule;
-
-		// see if it's upper hemisphere 
-		if (ContactZ > CenterOfCapsule.Z + HalfHeight)
-		{
-			// if upper hemisphere, the normal should point to the center of upper hemisphere
-			NormalEnd.Z = CenterOfCapsule.Z + HalfHeight;
-		}
-		else if (ContactZ < CenterOfCapsule.Z - HalfHeight)
-		{
-			// if lower hemisphere, the normal should point to the center of lower hemisphere	
-			NormalEnd.Z = CenterOfCapsule.Z - HalfHeight;
-		}
-		else
-		{
-			// otherwise, normal end is going to be same as contact point 
-			// since it's the tube section of the capsule
-			NormalEnd.Z = ContactZ;
-		}
-
-		// Recompute ImpactPoint if requested
-		if (PointOnGeom)
-		{
-			const FPlane GeomPlane(*PointOnGeom, OutHit.ImpactNormal);
-			const float DistToPlane = GeomPlane.PlaneDot(NormalEnd);
-			if (DistToPlane >= Radius)
-			{
-				OutHit.ImpactPoint = NormalEnd - DistToPlane * OutHit.ImpactNormal;
-			}
-		}
-
-		//DrawDebugLine(ContactPoint, NormalEnd, FColor(0, 255, 0), true);
-		//DrawDebugLine(NormalEnd, CenterOfCapsule, FColor(0, 0, 255), true);
-
-		FVector ContactNormal = (NormalEnd - OutHit.ImpactPoint);
-		// if ContactNormal is not nearly zero, set it
-		// otherwise use previous normal
-		const float Size = ContactNormal.Size();
-		if (Size >= KINDA_SMALL_NUMBER)
-		{
-			ContactNormal /= Size;
-			OutHit.Normal = ContactNormal;
-		}
-
-#if !(UE_BUILD_SHIPPING || UE_BUILD_TEST || !WITH_EDITOR)
-		CheckHitResultNormal(OutHit, TEXT("ConvertConvexNormalToCapsuleNormal"));
-#endif
-	}
-}
-
-static void ConvertConvexNormalToSphereNormal(float Radius, struct FHitResult& OutHit)
-{
-	if (OutHit.Time > 0.f)
-	{
-		FVector ContactNormal = (OutHit.Location - OutHit.ImpactPoint);
-		const float Size = ContactNormal.Size();
-		if (Size >= KINDA_SMALL_NUMBER)
-		{
-			ContactNormal /= Size;
-			OutHit.Normal = ContactNormal;
-		}
-
-#if !(UE_BUILD_SHIPPING || UE_BUILD_TEST || !WITH_EDITOR)
-		CheckHitResultNormal(OutHit, TEXT("ConvertConvexNormalToSphereNormal"));
-#endif
-	}
-}
 
 
 /** Helper to transform a normal when non-uniform scale is present. */
@@ -150,18 +67,30 @@ static void TransformNormalToShapeSpace(const PxMeshScale& meshScale, const PxVe
 
 
 
-/** Util to find the normal of the face that we hit */
-static bool FindGeomOpposingNormal(const PxLocationHit& PHit, const FVector & Direction, FVector& OutNormal, FVector& OutPointOnGeom)
+/**
+ * Util to find the normal of the face that we hit.
+ * @param PHit - incoming hit from PhysX
+ * @param Direction - direction of sweep test (not normalized)
+ * @param OutNormal - normal we may recompute based on the faceIndex of the hit
+ * @return true if we compute a new normal for the geometry.
+ */
+static bool FindGeomOpposingNormal(const PxLocationHit& PHit, const FVector& Direction, FVector& OutNormal)
 {
-	OutPointOnGeom = (PHit.flags & PxHitFlag::ePOSITION) ? P2UVector(PHit.position) : FVector::ZeroVector;
+	checkf(InvalidQueryHit.faceIndex == 0xFFFFffff, TEXT("Engine code needs fixing: PhysX invalid face index sentinel has changed or is not part of default PxQueryHit!"));
+	if (PHit.faceIndex == InvalidQueryHit.faceIndex)
+	{
+		return false;
+	}
 
 	PxTriangleMeshGeometry PTriMeshGeom;
 	PxConvexMeshGeometry PConvexMeshGeom;
+	PxHeightFieldGeometry PHeightFieldGeom;
+
 	if(	PHit.shape->getTriangleMeshGeometry(PTriMeshGeom) && 
 		PTriMeshGeom.triangleMesh != NULL &&
 		PHit.faceIndex < PTriMeshGeom.triangleMesh->getNbTriangles() )
 	{
-		const int32 TriIndex = PHit.faceIndex;
+		const PxU32 TriIndex = PHit.faceIndex;
 		const void* Triangles = PTriMeshGeom.triangleMesh->getTriangles();
 
 		// Grab triangle indices that we hit
@@ -203,11 +132,6 @@ static bool FindGeomOpposingNormal(const PxLocationHit& PHit, const FVector & Di
 			float sign = FVector::DotProduct(OutNormal, Direction) > 0 ? -1.f : 1.f;
 			OutNormal *= sign;
 		}
-		
-		if (PTriMeshGeom.scale.isIdentity())
-		{
-			OutPointOnGeom = P2UVector(PShapeWorldPose.transform(V0));
-		}
 
 #if !(UE_BUILD_SHIPPING || UE_BUILD_TEST)
 		if (!OutNormal.IsNormalized())
@@ -218,11 +142,11 @@ static bool FindGeomOpposingNormal(const PxLocationHit& PHit, const FVector & Di
 #endif
 		return true;
 	}
-	else if(PHit.shape->getConvexMeshGeometry(PConvexMeshGeom) && 
+	else if( PHit.shape->getConvexMeshGeometry(PConvexMeshGeom) && 
 		PConvexMeshGeom.convexMesh != NULL &&
 		PHit.faceIndex < PConvexMeshGeom.convexMesh->getNbPolygons() )
 	{
-		const int32 PolyIndex = PHit.faceIndex;
+		const PxU32 PolyIndex = PHit.faceIndex;
 		PxHullPolygon PPoly;
 		bool bSuccess = PConvexMeshGeom.convexMesh->getPolygonData(PolyIndex, PPoly);
 		if(bSuccess)
@@ -233,7 +157,7 @@ static bool FindGeomOpposingNormal(const PxLocationHit& PHit, const FVector & Di
 			// TODO: can we just change the hit normal within the physx code where we choose the most opposing normal, and avoid doing this here again?
 			PxVec3 PLocalPolyNormal;
 			TransformNormalToShapeSpace(PConvexMeshGeom.scale, PPlaneNormal.getNormalized(), PLocalPolyNormal);			
-			const PxTransform PShapeWorldPose = PHit.actor->getGlobalPose() * PHit.shape->getLocalPose(); 
+			const PxTransform PShapeWorldPose = PxShapeExt::getGlobalPose(*PHit.shape, *PHit.actor); 
 
 			const PxVec3 PWorldPolyNormal = PShapeWorldPose.rotate(PLocalPolyNormal);
 			OutNormal = P2UVector(PWorldPolyNormal);
@@ -248,6 +172,21 @@ static bool FindGeomOpposingNormal(const PxLocationHit& PHit, const FVector & Di
 			return true;
 		}
 	}
+	else if( PHit.shape->getHeightFieldGeometry(PHeightFieldGeom) &&
+		PHeightFieldGeom.heightField != NULL)
+	{
+		const PxU32 TriIndex = PHit.faceIndex;
+		const PxTransform PShapeWorldPose = PxShapeExt::getGlobalPose(*PHit.shape, *PHit.actor);
+		
+		PxTriangle Tri;
+		PxMeshQuery::getTriangle(PHeightFieldGeom, PShapeWorldPose, TriIndex, Tri);
+
+		PxVec3 TriNormal;
+		Tri.normal(TriNormal);
+		OutNormal = P2UVector(TriNormal);
+
+		return true;
+	}
 
 	return false;
 }
@@ -338,7 +277,7 @@ static void SetHitResultFromShapeAndFaceIndex(const PxShape* PShape,  const PxRi
 
 void ConvertQueryImpactHit(const PxLocationHit& PHit, FHitResult& OutResult, float CheckLength, const PxFilterData& QueryFilter, const FVector& StartLoc, const FVector& EndLoc, const PxGeometry* const Geom, const PxTransform& QueryTM, bool bReturnFaceIndex, bool bReturnPhysMat)
 {
-	check(PHit.flags & PxHitFlag::eDISTANCE);
+	checkSlow(PHit.flags & PxHitFlag::eDISTANCE);
 	if (Geom != NULL && PHit.hadInitialOverlap())
 	{
 		ConvertOverlappedShapeToImpactHit(PHit, StartLoc, EndLoc, OutResult, *Geom, QueryTM, QueryFilter, bReturnPhysMat);
@@ -378,46 +317,15 @@ void ConvertQueryImpactHit(const PxLocationHit& PHit, FHitResult& OutResult, flo
 	// Special handling for swept-capsule results
 	if(Geom && Geom->getType() == PxGeometryType::eCAPSULE)
 	{
-		const PxCapsuleGeometry* Capsule = static_cast<const PxCapsuleGeometry*>(Geom);
-		
 		// Compute better ImpactNormal. This is the same as Normal except when we hit convex/trimesh, and then its the most opposing normal of the geom at point of impact.
-		FVector PointOnGeom;
-		check(PHit.flags & PxHitFlag::ePOSITION);
-		if (FindGeomOpposingNormal(PHit, TraceDir, OutResult.ImpactNormal, PointOnGeom))
-		{
-			ConvertConvexNormalToCapsuleNormal(Capsule->halfHeight, Capsule->radius, OutResult, &PointOnGeom);
-		}
-		else
-		{
-			ConvertConvexNormalToCapsuleNormal(Capsule->halfHeight, Capsule->radius, OutResult, NULL);
-			OutResult.ImpactNormal = OutResult.Normal;
-		}
+		const PxCapsuleGeometry* Capsule = static_cast<const PxCapsuleGeometry*>(Geom);
+		FindGeomOpposingNormal(PHit, TraceDir, OutResult.ImpactNormal);
 	}
 	else if (Geom && Geom->getType() == PxGeometryType::eSPHERE)
 	{
-		const PxSphereGeometry* Sphere = static_cast<const PxSphereGeometry*>(Geom);
-
 		// Compute better ImpactNormal. This is the same as Normal except when we hit convex/trimesh, and then its the most opposing normal of the geom at point of impact.
-		FVector PointOnGeom;
-		check(PHit.flags & PxHitFlag::ePOSITION);
-		if (FindGeomOpposingNormal(PHit, TraceDir, OutResult.ImpactNormal, PointOnGeom))
-		{
-			const FPlane GeomPlane(PointOnGeom, OutResult.ImpactNormal);
-			const float DistFromPlane = FMath::Abs(GeomPlane.PlaneDot(OutResult.Location));
-			if (DistFromPlane >= Sphere->radius)
-			{
-				// Use the new (better) impact normal to compute a new impact point by projecting the sphere's location onto the geometry.
-				OutResult.ImpactPoint = OutResult.Location - DistFromPlane * OutResult.ImpactNormal;
-			}
-
-			// Use the impact point to compute a better sphere surface normal (impact point to center of sphere).
-			ConvertConvexNormalToSphereNormal(Sphere->radius, OutResult);
-		}
-		else
-		{
-			ConvertConvexNormalToSphereNormal(Sphere->radius, OutResult);
-			OutResult.ImpactNormal = OutResult.Normal;
-		}
+		const PxSphereGeometry* Sphere = static_cast<const PxSphereGeometry*>(Geom);
+		FindGeomOpposingNormal(PHit, TraceDir, OutResult.ImpactNormal);
 	}
 	
 	if( PHit.shape->getGeometryType() == PxGeometryType::eHEIGHTFIELD)
@@ -433,7 +341,6 @@ void ConvertQueryImpactHit(const PxLocationHit& PHit, FHitResult& OutResult, flo
 	if(bReturnFaceIndex && PHit.shape->getGeometryType() == PxGeometryType::eTRIANGLEMESH)
 	{
 		PxTriangleMeshGeometry PTriMeshGeom;
-		PxConvexMeshGeometry PConvexMeshGeom;
 		if(	PHit.shape->getTriangleMeshGeometry(PTriMeshGeom) && 
 			PTriMeshGeom.triangleMesh != NULL &&
 			PHit.faceIndex < PTriMeshGeom.triangleMesh->getNbTriangles() )
@@ -476,7 +383,7 @@ bool AddSweepResults(int32 NumHits, PxSweepHit* Hits, float CheckLength, const P
 	for(int32 i=0; i<NumHits; i++)
 	{
 		const PxSweepHit& PHit = Hits[i];
-		check(PHit.flags & PxHitFlag::eDISTANCE);
+		checkSlow(PHit.flags & PxHitFlag::eDISTANCE);
 		if(PHit.distance <= MaxDistance)
 		{
 			FHitResult NewResult(ForceInit);
@@ -510,7 +417,6 @@ FVector FindBestOverlappingNormal(const PxGeometry& Geom, const PxTransform& Que
 	// Track the best triangle plane distance
 	float BestPlaneDist = -BIG_NUMBER;
 	FVector BestPlaneNormal(0, 0, 1);
-	FVector BestPointOnPlane(0, 0, 0);
 	// Iterate over triangles
 	for (int32 TriIdx = 0; TriIdx < NumTrisHit; TriIdx++)
 	{
@@ -536,7 +442,6 @@ FVector FindBestOverlappingNormal(const PxGeometry& Geom, const PxTransform& Que
 		{
 			BestPlaneDist = DistToPlane;
 			BestPlaneNormal = TriNormal;
-			BestPointOnPlane = A;
 		}
 
 #if DRAW_OVERLAPPING_TRIS
@@ -563,7 +468,7 @@ static bool ConvertOverlappedShapeToImpactHit(const PxLocationHit& PHit, const F
 {
 	OutResult.TraceStart = StartLoc;
 	OutResult.TraceEnd = EndLoc;
-	FVector TraceDir = EndLoc - StartLoc;
+	const FVector TraceDir = EndLoc - StartLoc;
 
 	const PxShape* PShape = PHit.shape;
 	const PxRigidActor* PActor = PHit.actor;
@@ -584,86 +489,74 @@ static bool ConvertOverlappedShapeToImpactHit(const PxLocationHit& PHit, const F
 	OutResult.Location = P2UVector(QueryTM.p);
 	OutResult.ImpactPoint = OutResult.Location; // @todo not really sure of a better thing to do here...
 
-	// Now find the normal that most opposed the movement (closest to point of penetration)
-	const PxTransform PShapeWorldPose = PxShapeExt::getGlobalPose(*PShape, *PActor); 
+	// Fallback normal if we can't find it with MTD or otherwise.
+	OutResult.ImpactNormal = (PHit.flags & PxHitFlag::eNORMAL) ? P2UVector(PHit.normal) : FVector::UpVector;
 
-	PxTriangleMeshGeometry PTriMeshGeom;
-	PxHeightFieldGeometry PHeightfieldGeom;
-	if(PShape->getTriangleMeshGeometry(PTriMeshGeom))
+	// Use MTD result if possible. We interpret the MTD vector as both the direction to move and the opposing normal.
+	const bool bFiniteNormal = PHit.normal.isFinite();
+	const bool bValidMtdNormal = (PHit.flags & PxHitFlag::eNORMAL) && bFiniteNormal;
+	if (bValidMtdNormal)
 	{
-		PxU32 HitTris[64];
-		bool bOverflow = false;
-		const int32 NumTrisHit = PxMeshQuery::findOverlapTriangleMesh(Geom, QueryTM, PTriMeshGeom, PShapeWorldPose, HitTris, 64, 0, bOverflow);
-		if (NumTrisHit > 0)
-		{
-			OutResult.ImpactNormal = FindBestOverlappingNormal(Geom, QueryTM, PTriMeshGeom, PShapeWorldPose, HitTris, NumTrisHit);
-		}
-		else
-		{
-			// PxMeshQuery::findOverlapTriangleMesh failed, but we might have a valid face index at impact, so use that to grab the normal from that face.
-			FVector PointOnGeom;
-			if (!FindGeomOpposingNormal(PHit, TraceDir, OutResult.ImpactNormal, PointOnGeom))
-			{
-				// That failed (invalid face index). This really shouldn't happen.
-				OutResult.ImpactNormal = (PHit.flags & PxHitFlag::eNORMAL) ? P2UVector(PHit.normal) : FVector(0.f, 0.f, 1.f);
-			}
-		}
+		OutResult.Normal = P2UVector(PHit.normal);
+		OutResult.ImpactNormal = P2UVector(PHit.normal);
+		OutResult.PenetrationDepth = FMath::Abs(PHit.distance);
 	}
-	else
-	if (PShape->getHeightFieldGeometry(PHeightfieldGeom))
+	else if (!bFiniteNormal)
 	{
-		PxU32 HitTris[64];
-		bool bOverflow = false;
-		const int32 NumTrisHit = PxMeshQuery::findOverlapHeightField(Geom, QueryTM, PHeightfieldGeom, PShapeWorldPose, HitTris, 64, 0, bOverflow);
-		if (NumTrisHit > 0)
-		{
-			OutResult.ImpactNormal = FindBestOverlappingNormal(Geom, QueryTM, PHeightfieldGeom, PShapeWorldPose, HitTris, NumTrisHit);
-		}
-		else
-		{
-			// PxMeshQuery::findOverlapHeightField failed.
-			OutResult.ImpactNormal = (PHit.flags & PxHitFlag::eNORMAL) ? P2UVector(PHit.normal) : FVector(0.f, 0.f, 1.f);
-		}
+		UE_LOG(LogPhysics, Warning, TEXT("ConvertOverlappedShapeToImpactHit, MTD returned NaN :( normal: (X:%f, Y:%f, Z:%f)"), PHit.normal.x, PHit.normal.y, PHit.normal.z);
+		OutResult.PenetrationDepth = 0.f;
 	}
-	else
-	{
-		// Non tri-mesh
-		// use vector center of shape to query as good direction to move in
-		PxGeometry& PGeom = PShape->getGeometry().any();
-		PxVec3 PClosestPoint;
-		float Distance = PxGeometryQuery::pointDistance(QueryTM.p, PGeom, PShapeWorldPose, &PClosestPoint);
 
-		if(Distance < KINDA_SMALL_NUMBER)
+	// Zero-distance hits are often valid hits and we can extract the hit normal from the face index
+	// For invalid normals we can try other methods as well (get overlapping triangles).
+	if (PHit.distance == 0.f || !bValidMtdNormal)
+	{
+		const PxTransform PShapeWorldPose = PxShapeExt::getGlobalPose(*PShape, *PActor); 
+		PxTriangleMeshGeometry PTriMeshGeom;
+		PxHeightFieldGeometry PHeightfieldGeom;
+		if (PShape->getTriangleMeshGeometry(PTriMeshGeom))
 		{
-			//UE_LOG(LogCollision, Warning, TEXT("ConvertOverlappedShapeToImpactHit: Query origin inside shape, giving poor MTD."));			
-			PClosestPoint = PxShapeExt::getWorldBounds(*PShape, *PActor).getCenter(); 
+			if (!FindGeomOpposingNormal(PHit, TraceDir, OutResult.ImpactNormal))
+			{
+				PxU32 HitTris[64];
+				bool bOverflow = false;
+				const int32 NumTrisHit = PxMeshQuery::findOverlapTriangleMesh(Geom, QueryTM, PTriMeshGeom, PShapeWorldPose, HitTris, ARRAY_COUNT(HitTris), 0, bOverflow);
+				if (NumTrisHit > 0)
+				{
+					OutResult.ImpactNormal = FindBestOverlappingNormal(Geom, QueryTM, PTriMeshGeom, PShapeWorldPose, HitTris, NumTrisHit);
+				}
+			}
 		}
-
-		OutResult.ImpactNormal = (OutResult.Location - P2UVector(PClosestPoint)).SafeNormal();
-	}
-
-	// Compute depenetration vector and distance if possible.
-	PxVec3 PxMtdNormal(0.f);
-	PxF32 PxMtdDepth = 0.f;
-	PxGeometry& POtherGeom = PShape->getGeometry().any();
-
-	const bool bMtdResult = PxGeometryQuery::computePenetration(PxMtdNormal, PxMtdDepth, Geom, QueryTM, POtherGeom, PShapeWorldPose);
-	if (bMtdResult)
-	{
-		if (PxMtdNormal.isFinite())
+		else if (PShape->getHeightFieldGeometry(PHeightfieldGeom))
 		{
-			const FVector MtdNormal = P2UVector(PxMtdNormal);
-			OutResult.Normal = MtdNormal;
-			OutResult.PenetrationDepth = FMath::Abs(PxMtdDepth) + KINDA_SMALL_NUMBER;
+			if (!FindGeomOpposingNormal(PHit, TraceDir, OutResult.ImpactNormal))
+			{
+				PxU32 HitTris[64];
+				bool bOverflow = false;
+				const int32 NumTrisHit = PxMeshQuery::findOverlapHeightField(Geom, QueryTM, PHeightfieldGeom, PShapeWorldPose, HitTris, ARRAY_COUNT(HitTris), 0, bOverflow);
+				if (NumTrisHit > 0)
+				{
+					OutResult.ImpactNormal = FindBestOverlappingNormal(Geom, QueryTM, PHeightfieldGeom, PShapeWorldPose, HitTris, NumTrisHit);
+				}
+			}
 		}
 		else
 		{
-			OutResult.Normal = OutResult.ImpactNormal;
-			UE_LOG(LogPhysics, Warning, TEXT("ConvertOverlappedShapeToImpactHit, PxGeometryQuery::computePenetration returned NaN :( PxMtdNormal: (X:%f, Y:%f, Z:%f)"), PxMtdNormal.x, PxMtdNormal.y, PxMtdNormal.z);
+			// Non tri-mesh or heightfield
+			// Note: faceIndex seems to be unreliable for convex meshes in these cases, so not using FindGeomOpposingNormal() for them here.
+			PxGeometry& PGeom = PShape->getGeometry().any();
+			PxVec3 PClosestPoint;
+			const float Distance = PxGeometryQuery::pointDistance(QueryTM.p, PGeom, PShapeWorldPose, &PClosestPoint);
+
+			if (Distance < KINDA_SMALL_NUMBER)
+			{
+				UE_LOG(LogCollision, Warning, TEXT("ConvertOverlappedShapeToImpactHit: Query origin inside shape, giving poor MTD."));
+				PClosestPoint = PxShapeExt::getWorldBounds(*PShape, *PActor).getCenter(); 
+			}
+
+			OutResult.ImpactNormal = (OutResult.Location - P2UVector(PClosestPoint)).SafeNormal();
 		}
-	}
-	else
-	{
+
 		OutResult.Normal = OutResult.ImpactNormal;
 	}
 
@@ -729,8 +622,7 @@ static void AddUniqueOverlap(TArray<FOverlapResult>& OutOverlaps, const FOverlap
 
 bool ConvertOverlapResults(int32 NumOverlaps, PxOverlapHit* POverlapResults, const PxFilterData& QueryFilter, TArray<FOverlapResult>& OutOverlaps)
 {
-	OutOverlaps.Empty();
-
+	OutOverlaps.Reserve(NumOverlaps);
 	bool bBlockingFound = false;
 
 	for(int32 i=0; i<NumOverlaps; i++)
diff --git a/Engine/Source/Runtime/Engine/Private/Collision/PhysXCollision.cpp b/Engine/Source/Runtime/Engine/Private/Collision/PhysXCollision.cpp
--- a/Engine/Source/Runtime/Engine/Private/Collision/PhysXCollision.cpp
+++ b/Engine/Source/Runtime/Engine/Private/Collision/PhysXCollision.cpp
@@ -857,7 +857,7 @@ bool RaycastSingle(const UWorld* World, struct FHitResult& OutHit, const FVector
 			// Create filter data used to filter collisions
 			PxFilterData PFilter = CreateQueryFilterData(TraceChannel, Params.bTraceComplex, ResponseParams.CollisionResponse, ObjectParams, false);
 			PxSceneQueryFilterData PQueryFilterData(PFilter, PxSceneQueryFilterFlag::eSTATIC | PxSceneQueryFilterFlag::eDYNAMIC | PxSceneQueryFilterFlag::ePREFILTER);
-			PxSceneQueryFlags POutputFlags = PxSceneQueryFlag::ePOSITION | PxSceneQueryFlag::eNORMAL | PxSceneQueryFlag::eDISTANCE;
+			PxSceneQueryFlags POutputFlags = PxSceneQueryFlag::ePOSITION | PxSceneQueryFlag::eNORMAL | PxSceneQueryFlag::eDISTANCE | PxSceneQueryFlag::eMTD;
 			FPxQueryFilterCallback PQueryCallback(Params.IgnoreActors);
 			PQueryCallback.bSingleQuery = true; // this needs to be set for raycastSingle so we return eNONE instead of eTOUCH
 
@@ -953,7 +953,7 @@ bool RaycastMulti(const UWorld* World, TArray<struct FHitResult>& OutHits, const
 		// Create filter data used to filter collisions
 		PxFilterData PFilter = CreateQueryFilterData(TraceChannel, Params.bTraceComplex, ResponseParams.CollisionResponse, ObjectParams, true);
 		PxSceneQueryFilterData PQueryFilterData(PFilter, PxSceneQueryFilterFlag::eSTATIC | PxSceneQueryFilterFlag::eDYNAMIC | PxSceneQueryFilterFlag::ePREFILTER);
-		PxSceneQueryFlags POutputFlags = PxSceneQueryFlag::ePOSITION | PxSceneQueryFlag::eNORMAL | PxSceneQueryFlag::eDISTANCE;
+		PxSceneQueryFlags POutputFlags = PxSceneQueryFlag::ePOSITION | PxSceneQueryFlag::eNORMAL | PxSceneQueryFlag::eDISTANCE | PxSceneQueryFlag::eMTD;
 		FPxQueryFilterCallback PQueryCallback(Params.IgnoreActors);
 
 		// Create buffer for hits. Note: memory is not initialized (for perf reasons), since API does not require it.
@@ -1185,7 +1185,7 @@ bool GeomSweepSingle(const UWorld* World, const struct FCollisionShape& Collisio
 		PxFilterData PFilter = CreateQueryFilterData(TraceChannel, Params.bTraceComplex, ResponseParams.CollisionResponse, ObjectParams, false);
 		//UE_LOG(LogCollision, Log, TEXT("PFilter: %x %x %x %x"), PFilter.word0, PFilter.word1, PFilter.word2, PFilter.word3);
 		PxSceneQueryFilterData PQueryFilterData(PFilter, PxSceneQueryFilterFlag::eSTATIC | PxSceneQueryFilterFlag::eDYNAMIC | PxSceneQueryFilterFlag::ePREFILTER);
-		PxSceneQueryFlags POutputFlags = PxSceneQueryFlag::ePOSITION | PxSceneQueryFlag::eNORMAL | PxSceneQueryFlag::eDISTANCE;
+		PxSceneQueryFlags POutputFlags = PxSceneQueryFlag::ePOSITION | PxSceneQueryFlag::eNORMAL | PxSceneQueryFlag::eDISTANCE | PxSceneQueryFlag::eMTD;
 		FPxQueryFilterCallbackSweep PQueryCallbackSweep(Params.IgnoreActors);
 		PQueryCallbackSweep.bSingleQuery = true; //LOC_MOD need to do this to return eNONE instead of eTOUCH for r
 		PQueryCallbackSweep.DiscardInitialOverlaps = !Params.bFindInitialOverlaps;	
@@ -1280,7 +1280,7 @@ bool GeomSweepMulti_PhysX(const UWorld* World, const PxGeometry& PGeom, const Px
 	// Create filter data used to filter collisions
 	PxFilterData PFilter = CreateQueryFilterData(TraceChannel, Params.bTraceComplex, ResponseParams.CollisionResponse, ObjectParams, true);
 	PxSceneQueryFilterData PQueryFilterData(PFilter, PxSceneQueryFilterFlag::eSTATIC | PxSceneQueryFilterFlag::eDYNAMIC | PxSceneQueryFilterFlag::ePREFILTER | PxSceneQueryFilterFlag::ePOSTFILTER);
-	PxSceneQueryFlags POutputFlags = PxSceneQueryFlag::ePOSITION | PxSceneQueryFlag::eNORMAL | PxSceneQueryFlag::eDISTANCE;
+	PxSceneQueryFlags POutputFlags = PxSceneQueryFlag::ePOSITION | PxSceneQueryFlag::eNORMAL | PxSceneQueryFlag::eDISTANCE | PxSceneQueryFlag::eMTD;
 	FPxQueryFilterCallbackSweep PQueryCallbackSweep(Params.IgnoreActors);
 	PQueryCallbackSweep.DiscardInitialOverlaps = !Params.bFindInitialOverlaps;
 
@@ -1699,7 +1699,7 @@ bool GeomOverlapMulti(const UWorld* World, const struct FCollisionShape& Collisi
 
 #if WITH_PHYSX
 
-static PxQuat CapsuleRotator(0.f, 0.707106781f, 0.f, 0.707106781f );
+static const PxQuat CapsuleRotator(0.f, 0.707106781f, 0.f, 0.707106781f );
 
 PxQuat ConvertToPhysXCapsuleRot(const FQuat& GeomRot)
 {
diff --git a/Engine/Source/Runtime/Engine/Private/PrimitiveComponent.cpp b/Engine/Source/Runtime/Engine/Private/PrimitiveComponent.cpp
--- a/Engine/Source/Runtime/Engine/Private/PrimitiveComponent.cpp
+++ b/Engine/Source/Runtime/Engine/Private/PrimitiveComponent.cpp
@@ -1173,8 +1173,12 @@ static bool ShouldIgnoreHitResult(const UWorld* InWorld, FHitResult const& TestH
 		{
 			if(CVarShowInitialOverlaps.GetValueOnGameThread() != 0)
 			{
-				UE_LOG(LogTemp, Log, TEXT("Overlapping %s Dir %s Dot %f"), *GetNameSafe(TestHit.GetActor()), *MovementDir.ToString(), MoveDot);
-				DrawDebugDirectionalArrow(InWorld, TestHit.TraceStart, TestHit.TraceStart + 30.f * TestHit.ImpactNormal, 5.f, bMovingOut ? FColor(64,128,255) : FColor(255,64,64), true, 1.f);
+				UE_LOG(LogTemp, Log, TEXT("Overlapping %s Dir %s Dot %f Normal %s Depth %f"), *GetNameSafe(TestHit.Component.Get()), *MovementDir.ToString(), MoveDot, *TestHit.ImpactNormal.ToString(), TestHit.PenetrationDepth);
+				DrawDebugDirectionalArrow(InWorld, TestHit.TraceStart, TestHit.TraceStart + 30.f * TestHit.ImpactNormal, 5.f, bMovingOut ? FColor(64,128,255) : FColor(255,64,64), true, 4.f);
+				if (TestHit.PenetrationDepth > 0.f)
+				{
+					DrawDebugDirectionalArrow(InWorld, TestHit.TraceStart, TestHit.TraceStart + TestHit.PenetrationDepth * TestHit.Normal, 5.f, FColor(64,255,64), true, 4.f);
+				}
 			}
 		}
 #endif
