#ue4 - CharacterMovementComponent::bAlwaysCheckFloor now defaults to true.

- Fixes issues with vertical moving platforms starting below the character and then moving up through the character and being ignored.
- Performance is only affected for characters when they are not moving (by adding a floor check).

diff --git a/Engine/Source/Runtime/Engine/Classes/GameFramework/CharacterMovementComponent.h b/Engine/Source/Runtime/Engine/Classes/GameFramework/CharacterMovementComponent.h
--- a/Engine/Source/Runtime/Engine/Classes/GameFramework/CharacterMovementComponent.h
+++ b/Engine/Source/Runtime/Engine/Classes/GameFramework/CharacterMovementComponent.h
@@ -591,8 +591,8 @@ public:
 	uint32 bFastAttachedMove:1;
 
 	/**
-	 * True to always force floor checks for non-moving Characters.
-	 * Normally floor checks are avoided if possible, but this can be used to force them if there are use-cases where they are being skipped erroneously.
+	 * True to always force floor checks for stationary Characters.
+	 * Normally floor checks are avoided if possible when not moving, but this can be used to force them if there are use-cases where they are being skipped erroneously.
 	 */
 	UPROPERTY(Category="Character Movement", EditAnywhere, BlueprintReadWrite, AdvancedDisplay)
 	uint32 bAlwaysCheckFloor:1;
diff --git a/Engine/Source/Runtime/Engine/Private/CharacterMovementComponent.cpp b/Engine/Source/Runtime/Engine/Private/CharacterMovementComponent.cpp
--- a/Engine/Source/Runtime/Engine/Private/CharacterMovementComponent.cpp
+++ b/Engine/Source/Runtime/Engine/Private/CharacterMovementComponent.cpp
@@ -219,6 +219,7 @@ UCharacterMovementComponent::UCharacterMovementComponent(const class FPostConstr
 	bImpartBaseVelocityY = true;
 	bImpartBaseVelocityZ = true;
 	bImpartBaseAngularVelocity = true;
+	bAlwaysCheckFloor = true;
 
 	// default character can jump and walk
 	NavAgentProps.bCanJump = true;
