UE4 - rework RHI user-facing APIs to be safer and incorporate the desired sematics

diff --git a/Engine/Plugins/OculusRift/Source/OculusRift/Private/OculusRiftRender.cpp b/Engine/Plugins/OculusRift/Source/OculusRift/Private/OculusRiftRender.cpp
--- a/Engine/Plugins/OculusRift/Source/OculusRift/Private/OculusRiftRender.cpp
+++ b/Engine/Plugins/OculusRift/Source/OculusRift/Private/OculusRiftRender.cpp
@@ -102,7 +102,7 @@ void FOculusRiftHMD::DrawDistortionMesh_RenderThread(FRenderingCompositePassCont
 
 	const FSceneViewFamily& ViewFamily = *(View.Family);
 	FIntPoint ViewportSize = ViewFamily.RenderTarget->GetSizeXY();
-	RHISetViewport(0, 0, 0.0f, ViewportSize.X, ViewportSize.Y, 1.0f);
+	Context.RHICmdList.SetViewport(0, 0, 0.0f, ViewportSize.X, ViewportSize.Y, 1.0f);
 
 	Ptr<FDistortionMesh> mesh = RenderParams_RenderThread.pDistortionMesh[(View.StereoPass == eSSP_LEFT_EYE) ? 0 : 1];
 
diff --git a/Engine/Plugins/Runtime/CableComponent/Source/CableComponent/Private/CableComponent.cpp b/Engine/Plugins/Runtime/CableComponent/Source/CableComponent/Private/CableComponent.cpp
--- a/Engine/Plugins/Runtime/CableComponent/Source/CableComponent/Private/CableComponent.cpp
+++ b/Engine/Plugins/Runtime/CableComponent/Source/CableComponent/Private/CableComponent.cpp
@@ -8,7 +8,7 @@
 class FCableVertexBuffer : public FVertexBuffer 
 {
 public:
-	virtual void InitRHI()
+	virtual void InitRHI() override
 	{
 		FRHIResourceCreateInfo CreateInfo;
 		VertexBufferRHI = RHICreateVertexBuffer(NumVerts * sizeof(FDynamicMeshVertex), BUF_Dynamic, CreateInfo);
@@ -21,7 +21,7 @@ public:
 class FCableIndexBuffer : public FIndexBuffer 
 {
 public:
-	virtual void InitRHI()
+	virtual void InitRHI() override
 	{
 		FRHIResourceCreateInfo CreateInfo;
 		IndexBufferRHI = RHICreateIndexBuffer(sizeof(int32), NumIndices * sizeof(int32), BUF_Dynamic, CreateInfo);
diff --git a/Engine/Plugins/Runtime/CustomMeshComponent/Source/CustomMeshComponent/Private/CustomMeshComponent.cpp b/Engine/Plugins/Runtime/CustomMeshComponent/Source/CustomMeshComponent/Private/CustomMeshComponent.cpp
--- a/Engine/Plugins/Runtime/CustomMeshComponent/Source/CustomMeshComponent/Private/CustomMeshComponent.cpp
+++ b/Engine/Plugins/Runtime/CustomMeshComponent/Source/CustomMeshComponent/Private/CustomMeshComponent.cpp
@@ -10,7 +10,7 @@ class FCustomMeshVertexBuffer : public FVertexBuffer
 public:
 	TArray<FDynamicMeshVertex> Vertices;
 
-	virtual void InitRHI()
+	virtual void InitRHI() override
 	{
 		FRHIResourceCreateInfo CreateInfo;
 		VertexBufferRHI = RHICreateVertexBuffer(Vertices.Num() * sizeof(FDynamicMeshVertex),BUF_Static,CreateInfo);
@@ -29,7 +29,7 @@ class FCustomMeshIndexBuffer : public FIndexBuffer
 public:
 	TArray<int32> Indices;
 
-	virtual void InitRHI()
+	virtual void InitRHI() override
 	{
 		FRHIResourceCreateInfo CreateInfo;
 		IndexBufferRHI = RHICreateIndexBuffer(sizeof(int32),Indices.Num() * sizeof(int32),BUF_Static, CreateInfo);
diff --git a/Engine/Source/Developer/MovieSceneTools/Private/DirectorTrackEditor.cpp b/Engine/Source/Developer/MovieSceneTools/Private/DirectorTrackEditor.cpp
--- a/Engine/Source/Developer/MovieSceneTools/Private/DirectorTrackEditor.cpp
+++ b/Engine/Source/Developer/MovieSceneTools/Private/DirectorTrackEditor.cpp
@@ -133,7 +133,7 @@ void FShotThumbnail::CopyTextureIn(FSlateRenderTargetRHI* InTexture)
 		FSlateRenderTargetRHI*, RenderTarget, InTexture,
 		FSlateTexture2DRHIRef*, TargetTexture, Texture,
 	{
-		RHICopyToResolveTarget(RenderTarget->GetRHIRef(), TargetTexture->GetTypedResource(), false, FResolveParams());
+		RHICmdList.CopyToResolveTarget(RenderTarget->GetRHIRef(), TargetTexture->GetTypedResource(), false, FResolveParams());
 	});
 }
 
diff --git a/Engine/Source/Editor/GraphEditor/Private/MaterialNodes/SGraphNodeMaterialBase.cpp b/Engine/Source/Editor/GraphEditor/Private/MaterialNodes/SGraphNodeMaterialBase.cpp
--- a/Engine/Source/Editor/GraphEditor/Private/MaterialNodes/SGraphNodeMaterialBase.cpp
+++ b/Engine/Source/Editor/GraphEditor/Private/MaterialNodes/SGraphNodeMaterialBase.cpp
@@ -172,11 +172,12 @@ bool FPreviewElement::BeginRenderingCanvas( const FIntRect& InCanvasRect, const
 	return false;
 }
 
-void FPreviewElement::DrawRenderThread( const void* InWindowBackBuffer )
+void FPreviewElement::DrawRenderThread(FRHICommandListImmediate& RHICmdList, const void* InWindowBackBuffer)
 {
 	// Clip the canvas to avoid having to set UV values
 	FIntRect ClippingRect = RenderTarget->GetClippingRect();
-	RHISetScissorRect(	true,
+
+	RHICmdList.SetScissorRect(true,
 						ClippingRect.Min.X,
 						ClippingRect.Min.Y,
 						ClippingRect.Max.X,
@@ -198,7 +199,7 @@ void FPreviewElement::DrawRenderThread( const void* InWindowBackBuffer )
 		Canvas.Flush(true);
 	}
 	RenderTarget->ClearRenderTargetTexture();
-	RHISetScissorRect(false,0,0,0,0);
+	RHICmdList.SetScissorRect(false, 0, 0, 0, 0);
 }
 
 /////////////////////////////////////////////////////
diff --git a/Engine/Source/Editor/GraphEditor/Private/MaterialNodes/SGraphNodeMaterialBase.h b/Engine/Source/Editor/GraphEditor/Private/MaterialNodes/SGraphNodeMaterialBase.h
--- a/Engine/Source/Editor/GraphEditor/Private/MaterialNodes/SGraphNodeMaterialBase.h
+++ b/Engine/Source/Editor/GraphEditor/Private/MaterialNodes/SGraphNodeMaterialBase.h
@@ -44,7 +44,7 @@ private:
 	/**
 	 * ICustomSlateElement interface 
 	 */
-	virtual void DrawRenderThread( const void* InWindowBackBuffer ) override;
+	virtual void DrawRenderThread(FRHICommandListImmediate& RHICmdList, const void* InWindowBackBuffer) override;
 
 private:
 	/** Render target that the canvas renders to */
diff --git a/Engine/Source/Editor/MeshPaint/Private/MeshPaintEdMode.cpp b/Engine/Source/Editor/MeshPaint/Private/MeshPaintEdMode.cpp
--- a/Engine/Source/Editor/MeshPaint/Private/MeshPaintEdMode.cpp
+++ b/Engine/Source/Editor/MeshPaint/Private/MeshPaintEdMode.cpp
@@ -44,7 +44,7 @@ class FMeshPaintBatchedElementParameters : public FBatchedElementParameters
 public:
 
 	/** Binds vertex and pixel shaders for this element */
-	virtual void BindShaders_RenderThread(FRHICommandList& RHICmdList, const FMatrix& InTransform, const float InGamma, const FMatrix& ColorWeights, const FTexture* Texture ) override
+	virtual void BindShaders_RenderThread(FRHICommandListImmediate& RHICmdList, const FMatrix& InTransform, const float InGamma, const FMatrix& ColorWeights, const FTexture* Texture) override
 	{
 		MeshPaintRendering::SetMeshPaintShaders_RenderThread(RHICmdList, InTransform, InGamma, ShaderParams );
 	}
@@ -63,7 +63,7 @@ class FMeshPaintDilateBatchedElementParameters : public FBatchedElementParameter
 public:
 
 	/** Binds vertex and pixel shaders for this element */
-	virtual void BindShaders_RenderThread(FRHICommandList& RHICmdList, const FMatrix& InTransform, const float InGamma, const FMatrix& ColorWeights, const FTexture* Texture ) override
+	virtual void BindShaders_RenderThread(FRHICommandListImmediate& RHICmdList, const FMatrix& InTransform, const float InGamma, const FMatrix& ColorWeights, const FTexture* Texture) override
 	{
 		MeshPaintRendering::SetMeshPaintDilateShaders_RenderThread(RHICmdList, InTransform, InGamma, ShaderParams );
 	}
@@ -2271,7 +2271,7 @@ void FEdModeMeshPaint::PaintTexture( const FMeshPaintParameters& InParams,
 			FTextureRenderTargetResource*, BrushRenderTargetResource, BrushRenderTargetResource,
 		{
 			// Copy (resolve) the rendered image from the frame buffer to its render target texture
-			RHICopyToResolveTarget(
+			RHICmdList.CopyToResolveTarget(
 			BrushRenderTargetResource->GetRenderTargetTexture(),	// Source texture
 			BrushRenderTargetResource->TextureRHI,
 			true,													// Do we need the source image content again?
@@ -2290,7 +2290,7 @@ void FEdModeMeshPaint::PaintTexture( const FMeshPaintParameters& InParams,
 				FTextureRenderTargetResource*, BrushMaskRenderTargetResource, BrushMaskRenderTargetResource,
 			{
 				// Copy (resolve) the rendered image from the frame buffer to its render target texture
-				RHICopyToResolveTarget(
+				RHICmdList.CopyToResolveTarget(
 					BrushMaskRenderTargetResource->GetRenderTargetTexture(),		// Source texture
 					BrushMaskRenderTargetResource->TextureRHI,
 					true,												// Do we need the source image content again?
@@ -2390,7 +2390,7 @@ void FEdModeMeshPaint::PaintTexture( const FMeshPaintParameters& InParams,
 				FTextureRenderTargetResource*, RenderTargetResource, RenderTargetResource,
 			{
 				// Copy (resolve) the rendered image from the frame buffer to its render target texture
-				RHICopyToResolveTarget(
+				RHICmdList.CopyToResolveTarget(
 					RenderTargetResource->GetRenderTargetTexture(),		// Source texture
 					RenderTargetResource->TextureRHI,
 					true,												// Do we need the source image content again?
@@ -3303,7 +3303,7 @@ void FEdModeMeshPaint::CopyTextureToRenderTargetTexture( UTexture* SourceTexture
 			FTextureRenderTargetResource*, RenderTargetResource, RenderTargetResource,
 		{
 			// Copy (resolve) the rendered image from the frame buffer to its render target texture
-			RHICopyToResolveTarget(
+			RHICmdList.CopyToResolveTarget(
 				RenderTargetResource->GetRenderTargetTexture(),		// Source texture
 				RenderTargetResource->TextureRHI,					// Dest texture
 				true,												// Do we need the source image content again?
@@ -3514,7 +3514,7 @@ bool FEdModeMeshPaint::GenerateSeamMask(UStaticMeshComponent* StaticMeshComponen
 			FTextureRenderTargetResource*, RenderTargetResource, RenderTargetResource,
 		{
 			// Copy (resolve) the rendered image from the frame buffer to its render target texture
-			RHICopyToResolveTarget(
+			RHICmdList.CopyToResolveTarget(
 				RenderTargetResource->GetRenderTargetTexture(),		// Source texture
 				RenderTargetResource->TextureRHI,
 				true,												// Do we need the source image content again?
diff --git a/Engine/Source/Editor/UnrealEd/Private/CubemapUnwrapUtils.cpp b/Engine/Source/Editor/UnrealEd/Private/CubemapUnwrapUtils.cpp
--- a/Engine/Source/Editor/UnrealEd/Private/CubemapUnwrapUtils.cpp
+++ b/Engine/Source/Editor/UnrealEd/Private/CubemapUnwrapUtils.cpp
@@ -124,7 +124,7 @@ void FCubemapTexturePropertiesPS<bHDROutput>::SetParameters( FRHICommandList& RH
 	SetShaderValue(RHICmdList, GetPixelShader(), Gamma, GammaValue);
 }
 
-void FMipLevelBatchedElementParameters::BindShaders_RenderThread(FRHICommandList& RHICmdList, const FMatrix& InTransform, const float InGamma, const FMatrix& ColorWeights, const FTexture* Texture)
+void FMipLevelBatchedElementParameters::BindShaders_RenderThread(FRHICommandListImmediate& RHICmdList, const FMatrix& InTransform, const float InGamma, const FMatrix& ColorWeights, const FTexture* Texture)
 {
 	if(bHDROutput)
 	{
@@ -137,18 +137,18 @@ void FMipLevelBatchedElementParameters::BindShaders_RenderThread(FRHICommandList
 }
 
 template<typename TPixelShader>
-void FMipLevelBatchedElementParameters::BindShaders_RenderThread( FRHICommandList& RHICmdList, const FMatrix& InTransform, const float InGamma, const FMatrix& ColorWeights, const FTexture* Texture )
+void FMipLevelBatchedElementParameters::BindShaders_RenderThread(FRHICommandListImmediate& RHICmdList, const FMatrix& InTransform, const float InGamma, const FMatrix& ColorWeights, const FTexture* Texture)
 {
 	TShaderMapRef<FCubemapTexturePropertiesVS> VertexShader(GetGlobalShaderMap());
 	TShaderMapRef<TPixelShader> PixelShader(GetGlobalShaderMap());
-	RHICmdList.CheckIsNull(); // need new approach for "static FGlobalBoundShaderState" for parallel rendering
+	
 
 	static FGlobalBoundShaderState BoundShaderState;
 	SetGlobalBoundShaderState(RHICmdList, BoundShaderState, GSimpleElementVertexDeclaration.VertexDeclarationRHI, *VertexShader, *PixelShader);
 
 	VertexShader->SetParameters(RHICmdList, InTransform);
 
-	RHISetBlendState(TStaticBlendState<>::GetRHI());
+	RHICmdList.SetBlendState(TStaticBlendState<>::GetRHI());
 
 	PixelShader->SetParameters(RHICmdList, Texture, ColorWeights, MipLevel, InGamma);
 }
@@ -161,15 +161,15 @@ void FIESLightProfilePS::SetParameters( FRHICommandList& RHICmdList, const FText
 	SetShaderValue(RHICmdList, ShaderRHI, BrightnessInLumens, InBrightnessInLumens);
 }
 
-void FIESLightProfileBatchedElementParameters::BindShaders_RenderThread( FRHICommandList& RHICmdList, const FMatrix& InTransform, const float InGamma, const FMatrix& ColorWeights, const FTexture* Texture )
+void FIESLightProfileBatchedElementParameters::BindShaders_RenderThread( FRHICommandListImmediate& RHICmdList, const FMatrix& InTransform, const float InGamma, const FMatrix& ColorWeights, const FTexture* Texture )
 {
 	TShaderMapRef<FSimpleElementVS> VertexShader(GetGlobalShaderMap());
 	TShaderMapRef<FIESLightProfilePS> PixelShader(GetGlobalShaderMap());
 
-	RHICmdList.CheckIsNull(); // need new approach for "static FGlobalBoundShaderState" for parallel rendering
+	
 	static FGlobalBoundShaderState BoundShaderState;
 	SetGlobalBoundShaderState(RHICmdList, BoundShaderState, GSimpleElementVertexDeclaration.VertexDeclarationRHI, *VertexShader, *PixelShader);
-	RHISetBlendState(TStaticBlendState<>::GetRHI());
+	RHICmdList.SetBlendState(TStaticBlendState<>::GetRHI());
 
 	VertexShader->SetParameters(RHICmdList, InTransform);
 	PixelShader->SetParameters(RHICmdList, Texture, BrightnessInLumens);
diff --git a/Engine/Source/Editor/UnrealEd/Private/EditorExporters.cpp b/Engine/Source/Editor/UnrealEd/Private/EditorExporters.cpp
--- a/Engine/Source/Editor/UnrealEd/Private/EditorExporters.cpp
+++ b/Engine/Source/Editor/UnrealEd/Private/EditorExporters.cpp
@@ -2769,7 +2769,7 @@ namespace MaterialExportUtils
 					FTextureRenderTargetResource*, RenderTargetResource, RenderTargetResource,
 				{
 					// Copy (resolve) the rendered thumbnail from the render target to its texture
-					RHICopyToResolveTarget(
+					RHICmdList.CopyToResolveTarget(
 						RenderTargetResource->GetRenderTargetTexture(),		// Source texture
 						RenderTargetResource->TextureRHI,					// Dest texture
 						false,												// Do we need the source image content again?
@@ -2853,7 +2853,7 @@ namespace MaterialExportUtils
 					FTextureRenderTargetResource*, RenderTargetResource, RenderTargetResource,
 				{
 					// Copy (resolve) the rendered thumbnail from the render target to its texture
-					RHICopyToResolveTarget(
+					RHICmdList.CopyToResolveTarget(
 						RenderTargetResource->GetRenderTargetTexture(),		// Source texture
 						RenderTargetResource->TextureRHI,					// Dest texture
 						false,												// Do we need the source image content again?
diff --git a/Engine/Source/Editor/UnrealEd/Private/MeshPaintRendering.cpp b/Engine/Source/Editor/UnrealEd/Private/MeshPaintRendering.cpp
--- a/Engine/Source/Editor/UnrealEd/Private/MeshPaintRendering.cpp
+++ b/Engine/Source/Editor/UnrealEd/Private/MeshPaintRendering.cpp
@@ -348,14 +348,14 @@ namespace MeshPaintRendering
 
 
 	/** Binds the mesh paint vertex and pixel shaders to the graphics device */
-	void SetMeshPaintShaders_RenderThread(FRHICommandList& RHICmdList, const FMatrix& InTransform,
+	void SetMeshPaintShaders_RenderThread(FRHICommandListImmediate& RHICmdList, const FMatrix& InTransform,
 										   const float InGamma,
 										   const FMeshPaintShaderParameters& InShaderParams )
 	{
 		TShaderMapRef< TMeshPaintVertexShader > VertexShader( GetGlobalShaderMap() );
 		TShaderMapRef< TMeshPaintPixelShader > PixelShader( GetGlobalShaderMap() );
 
-		RHICmdList.CheckIsNull(); // need new approach for "static FGlobalBoundShaderState" for parallel rendering
+		
 		static FGlobalBoundShaderState BoundShaderState;
 		SetGlobalBoundShaderState(RHICmdList, BoundShaderState, GMeshPaintVertexDeclaration.VertexDeclarationRHI, *VertexShader, *PixelShader);
 
@@ -369,14 +369,14 @@ namespace MeshPaintRendering
 	}
 
 	/** Binds the mesh paint vertex and pixel shaders to the graphics device */
-	void SetMeshPaintDilateShaders_RenderThread(FRHICommandList& RHICmdList, const FMatrix& InTransform,
+	void SetMeshPaintDilateShaders_RenderThread(FRHICommandListImmediate& RHICmdList, const FMatrix& InTransform,
 												 const float InGamma,
 												 const FMeshPaintDilateShaderParameters& InShaderParams )
 	{
 		TShaderMapRef< TMeshPaintDilateVertexShader > VertexShader( GetGlobalShaderMap() );
 		TShaderMapRef< TMeshPaintDilatePixelShader > PixelShader( GetGlobalShaderMap() );
 
-		RHICmdList.CheckIsNull(); // need new approach for "static FGlobalBoundShaderState" for parallel rendering
+		
 		static FGlobalBoundShaderState BoundShaderState;
 		SetGlobalBoundShaderState(RHICmdList, BoundShaderState, GMeshPaintDilateVertexDeclaration.VertexDeclarationRHI, *VertexShader, *PixelShader);
 
diff --git a/Engine/Source/Editor/UnrealEd/Private/NormalMapPreview.cpp b/Engine/Source/Editor/UnrealEd/Private/NormalMapPreview.cpp
--- a/Engine/Source/Editor/UnrealEd/Private/NormalMapPreview.cpp
+++ b/Engine/Source/Editor/UnrealEd/Private/NormalMapPreview.cpp
@@ -75,7 +75,7 @@ IMPLEMENT_SHADER_TYPE(,FSimpleElementNormalMapPS,TEXT("SimpleElementNormalMapPix
 
 /** Binds vertex and pixel shaders for this element */
 void FNormalMapBatchedElementParameters::BindShaders_RenderThread(
-	FRHICommandList& RHICmdList, 
+	FRHICommandListImmediate& RHICmdList,
 	const FMatrix& InTransform,
 	const float InGamma,
 	const FMatrix& ColorWeights,
@@ -84,13 +84,13 @@ void FNormalMapBatchedElementParameters::BindShaders_RenderThread(
 	TShaderMapRef<FSimpleElementVS> VertexShader(GetGlobalShaderMap());
 	TShaderMapRef<FSimpleElementNormalMapPS> PixelShader(GetGlobalShaderMap());
 
-	RHICmdList.CheckIsNull(); // need new approach for "static FGlobalBoundShaderState" for parallel rendering
+	
 	static FGlobalBoundShaderState BoundShaderState;
 	SetGlobalBoundShaderState(RHICmdList, BoundShaderState, GSimpleElementVertexDeclaration.VertexDeclarationRHI, *VertexShader, *PixelShader);
 
 	VertexShader->SetParameters(RHICmdList, InTransform);
 	PixelShader->SetParameters(RHICmdList, Texture);
 
-	RHISetBlendState(TStaticBlendState<>::GetRHI());
+	RHICmdList.SetBlendState(TStaticBlendState<>::GetRHI());
 }
 
diff --git a/Engine/Source/Editor/UnrealEd/Private/ObjectTools.cpp b/Engine/Source/Editor/UnrealEd/Private/ObjectTools.cpp
--- a/Engine/Source/Editor/UnrealEd/Private/ObjectTools.cpp
+++ b/Engine/Source/Editor/UnrealEd/Private/ObjectTools.cpp
@@ -3468,7 +3468,7 @@ namespace ThumbnailTools
 					FTextureRenderTargetResource*, RenderTargetResource, RenderTargetResource,
 				{
 					// Copy (resolve) the rendered thumbnail from the render target to its texture
-					RHICopyToResolveTarget(
+					RHICmdList.CopyToResolveTarget(
 						RenderTargetResource->GetRenderTargetTexture(),		// Source texture
 						RenderTargetResource->TextureRHI,					// Dest texture
 						false,												// Do we need the source image content again?
diff --git a/Engine/Source/Editor/UnrealEd/Private/Texture2DPreview.cpp b/Engine/Source/Editor/UnrealEd/Private/Texture2DPreview.cpp
--- a/Engine/Source/Editor/UnrealEd/Private/Texture2DPreview.cpp
+++ b/Engine/Source/Editor/UnrealEd/Private/Texture2DPreview.cpp
@@ -77,7 +77,7 @@ IMPLEMENT_SHADER_TYPE(,FSimpleElementTexture2DPreviewPS,TEXT("SimpleElementTextu
 
 /** Binds vertex and pixel shaders for this element */
 void FBatchedElementTexture2DPreviewParameters::BindShaders_RenderThread(
-	FRHICommandList& RHICmdList, 
+	FRHICommandListImmediate& RHICmdList,
 	const FMatrix& InTransform,
 	const float InGamma,
 	const FMatrix& ColorWeights,
@@ -86,7 +86,7 @@ void FBatchedElementTexture2DPreviewParameters::BindShaders_RenderThread(
 	TShaderMapRef<FSimpleElementVS> VertexShader(GetGlobalShaderMap());
 	TShaderMapRef<FSimpleElementTexture2DPreviewPS> PixelShader(GetGlobalShaderMap());
 
-	RHICmdList.CheckIsNull(); // need new approach for "static FGlobalBoundShaderState" for parallel rendering
+	
 	static FGlobalBoundShaderState BoundShaderState;
 	SetGlobalBoundShaderState(RHICmdList, BoundShaderState, GSimpleElementVertexDeclaration.VertexDeclarationRHI, *VertexShader, *PixelShader);
 
@@ -95,6 +95,6 @@ void FBatchedElementTexture2DPreviewParameters::BindShaders_RenderThread(
 
 	if( bIsSingleChannelFormat )
 	{
-		RHISetBlendState(TStaticBlendState<>::GetRHI());
+		RHICmdList.SetBlendState(TStaticBlendState<>::GetRHI());
 	}
 }
diff --git a/Engine/Source/Editor/UnrealEd/Public/CubemapUnwrapUtils.h b/Engine/Source/Editor/UnrealEd/Public/CubemapUnwrapUtils.h
--- a/Engine/Source/Editor/UnrealEd/Public/CubemapUnwrapUtils.h
+++ b/Engine/Source/Editor/UnrealEd/Public/CubemapUnwrapUtils.h
@@ -120,10 +120,10 @@ public:
 	}
 
 	/** Binds vertex and pixel shaders for this element */
-	virtual void BindShaders_RenderThread(FRHICommandList& RHICmdList, const FMatrix& InTransform, const float InGamma, const FMatrix& ColorWeights, const FTexture* Texture) override;
+	virtual void BindShaders_RenderThread(FRHICommandListImmediate& RHICmdList, const FMatrix& InTransform, const float InGamma, const FMatrix& ColorWeights, const FTexture* Texture) override;
 
 private:
-	template<typename TPixelShader> void BindShaders_RenderThread(FRHICommandList& RHICmdList, const FMatrix& InTransform, const float InGamma, const FMatrix& ColorWeights, const FTexture* Texture);
+	template<typename TPixelShader> void BindShaders_RenderThread(FRHICommandListImmediate& RHICmdList, const FMatrix& InTransform, const float InGamma, const FMatrix& ColorWeights, const FTexture* Texture);
 
 	bool bHDROutput;
 	/** Parameters that need to be passed to the shader */
@@ -180,7 +180,7 @@ public:
 	}
 
 	/** Binds vertex and pixel shaders for this element */
-	virtual void BindShaders_RenderThread(FRHICommandList& RHICmdList, const FMatrix& InTransform, const float InGamma, const FMatrix& ColorWeights, const FTexture* Texture ) override;
+	virtual void BindShaders_RenderThread(FRHICommandListImmediate& RHICmdList, const FMatrix& InTransform, const float InGamma, const FMatrix& ColorWeights, const FTexture* Texture) override;
 
 private:
 	float BrightnessInLumens;
diff --git a/Engine/Source/Editor/UnrealEd/Public/MeshPaintRendering.h b/Engine/Source/Editor/UnrealEd/Public/MeshPaintRendering.h
--- a/Engine/Source/Editor/UnrealEd/Public/MeshPaintRendering.h
+++ b/Engine/Source/Editor/UnrealEd/Public/MeshPaintRendering.h
@@ -57,12 +57,12 @@ namespace MeshPaintRendering
 
 
 	/** Binds the mesh paint vertex and pixel shaders to the graphics device */
-	void UNREALED_API SetMeshPaintShaders_RenderThread(FRHICommandList& RHICmdList, const FMatrix& InTransform,
+	void UNREALED_API SetMeshPaintShaders_RenderThread(FRHICommandListImmediate& RHICmdList, const FMatrix& InTransform,
 										   const float InGamma,
 										   const FMeshPaintShaderParameters& InShaderParams );
 
 	/** Binds the mesh paint dilation vertex and pixel shaders to the graphics device */
-	void UNREALED_API SetMeshPaintDilateShaders_RenderThread(FRHICommandList& RHICmdList, const FMatrix& InTransform,
+	void UNREALED_API SetMeshPaintDilateShaders_RenderThread(FRHICommandListImmediate& RHICmdList, const FMatrix& InTransform,
 												const float InGamma,
 												const FMeshPaintDilateShaderParameters& InShaderParams );
 
diff --git a/Engine/Source/Editor/UnrealEd/Public/NormalMapPreview.h b/Engine/Source/Editor/UnrealEd/Public/NormalMapPreview.h
--- a/Engine/Source/Editor/UnrealEd/Public/NormalMapPreview.h
+++ b/Engine/Source/Editor/UnrealEd/Public/NormalMapPreview.h
@@ -15,7 +15,7 @@ class UNREALED_API FNormalMapBatchedElementParameters : public FBatchedElementPa
 {
 	/** Binds vertex and pixel shaders for this element */
 	virtual void BindShaders_RenderThread(
-		FRHICommandList& RHICmdList, 
+	FRHICommandListImmediate& RHICmdList,
 		const FMatrix& InTransform,
 		const float InGamma,
 		const FMatrix& ColorWeights,
diff --git a/Engine/Source/Editor/UnrealEd/Public/Texture2DPreview.h b/Engine/Source/Editor/UnrealEd/Public/Texture2DPreview.h
--- a/Engine/Source/Editor/UnrealEd/Public/Texture2DPreview.h
+++ b/Engine/Source/Editor/UnrealEd/Public/Texture2DPreview.h
@@ -22,7 +22,7 @@ public:
 	}
 
 	/** Binds vertex and pixel shaders for this element */
-	virtual void BindShaders_RenderThread(FRHICommandList& RHICmdList, const FMatrix& InTransform, const float InGamma, const FMatrix& ColorWeights, const FTexture* Texture ) override;
+	virtual void BindShaders_RenderThread(FRHICommandListImmediate& RHICmdList, const FMatrix& InTransform, const float InGamma, const FMatrix& ColorWeights, const FTexture* Texture) override;
 
 private:
 
diff --git a/Engine/Source/Runtime/Core/Private/HAL/ThreadingBase.cpp b/Engine/Source/Runtime/Core/Private/HAL/ThreadingBase.cpp
--- a/Engine/Source/Runtime/Core/Private/HAL/ThreadingBase.cpp
+++ b/Engine/Source/Runtime/Core/Private/HAL/ThreadingBase.cpp
@@ -21,6 +21,19 @@ CORE_API bool IsInSlateThread()
 	return GSlateLoadingThreadId != 0 && FPlatformTLS::GetCurrentThreadId() == GSlateLoadingThreadId;
 }
 
+CORE_API int32 GIsRenderingThreadSuspended = 0;
+
+CORE_API FRunnableThread* GRenderingThread = NULL;
+
+CORE_API bool IsInActualRenderingThread()
+{
+	return GRenderingThread && FPlatformTLS::GetCurrentThreadId() == GRenderingThread->GetThreadID();
+}
+
+CORE_API bool IsInRenderingThread()
+{
+	return !GRenderingThread || GIsRenderingThreadSuspended || (FPlatformTLS::GetCurrentThreadId() == GRenderingThread->GetThreadID());
+}
 
 
 // Fake threads
diff --git a/Engine/Source/Runtime/Core/Public/HAL/ThreadingBase.h b/Engine/Source/Runtime/Core/Public/HAL/ThreadingBase.h
--- a/Engine/Source/Runtime/Core/Public/HAL/ThreadingBase.h
+++ b/Engine/Source/Runtime/Core/Public/HAL/ThreadingBase.h
@@ -867,11 +867,23 @@ private:
 
 
 /** @return True if called from the game thread. */
-extern CORE_API bool IsInGameThread ();
+extern CORE_API bool IsInGameThread();
 
 /** @return True if called from the slate thread, and not merely a thread calling slate functions. */
-extern CORE_API bool IsInSlateThread ();
+extern CORE_API bool IsInSlateThread();
 
+/** @return True if called from the rendering thread, or if called from ANY thread during single threaded rendering */
+extern CORE_API bool IsInRenderingThread();
+
+/** @return True if called from the rendering thread. */
+// Unlike IsInRenderingThread, this will always return false if we are running single threaded. It only returns true if this is actually a separate rendering thread. Mostly useful for checks
+extern CORE_API bool IsInActualRenderingThread();
+
+/** Thread used for rendering */
+extern CORE_API FRunnableThread* GRenderingThread;
+
+/** Whether the rendering thread is suspended (not even processing the tickables) */
+extern CORE_API int32 GIsRenderingThreadSuspended;
 
 /** Minimal base class for the thread singleton. */
 struct FThreadSingleton
diff --git a/Engine/Source/Runtime/EmptyRHI/Private/EmptyGlobalUniformBuffer.h b/Engine/Source/Runtime/EmptyRHI/Private/EmptyGlobalUniformBuffer.h
--- a/Engine/Source/Runtime/EmptyRHI/Private/EmptyGlobalUniformBuffer.h
+++ b/Engine/Source/Runtime/EmptyRHI/Private/EmptyGlobalUniformBuffer.h
@@ -20,8 +20,8 @@ public:
 	~FEmptyGlobalUniformBuffer();
 
 	// FRenderResource interface.
-	virtual void	InitDynamicRHI();
-	virtual void	ReleaseDynamicRHI();
+	virtual void	InitDynamicRHI() override;
+	virtual void	ReleaseDynamicRHI() override;
 
 	void			UpdateConstant(const uint8* Data, uint16 Offset, uint16 Size);
 
diff --git a/Engine/Source/Runtime/EmptyRHI/Private/EmptyViewport.cpp b/Engine/Source/Runtime/EmptyRHI/Private/EmptyViewport.cpp
--- a/Engine/Source/Runtime/EmptyRHI/Private/EmptyViewport.cpp
+++ b/Engine/Source/Runtime/EmptyRHI/Private/EmptyViewport.cpp
@@ -47,7 +47,7 @@ void FEmptyDynamicRHI::RHIBeginDrawingViewport(FViewportRHIParamRef ViewportRHI,
 {
 	DYNAMIC_CAST_EMPTYRESOURCE(Viewport,Viewport);
 
-	RHISetRenderTarget(RHIGetViewportBackBuffer(ViewportRHI), NULL);
+	SetRenderTarget(RHICmdList, RHIGetViewportBackBuffer(ViewportRHI), NULL);
 }
 
 void FEmptyDynamicRHI::RHIEndDrawingViewport(FViewportRHIParamRef ViewportRHI,bool bPresent,bool bLockToVsync)
diff --git a/Engine/Source/Runtime/EmptyRHI/Public/EmptyRHI.h b/Engine/Source/Runtime/EmptyRHI/Public/EmptyRHI.h
--- a/Engine/Source/Runtime/EmptyRHI/Public/EmptyRHI.h
+++ b/Engine/Source/Runtime/EmptyRHI/Public/EmptyRHI.h
@@ -36,7 +36,7 @@ public:
 	virtual void PopEvent();
 
 	// The RHI methods are defined as virtual functions in URenderHardwareInterface.
-	#define DEFINE_RHIMETHOD(Type,Name,ParameterTypesAndNames,ParameterNames,ReturnStatement,NullImplementation) virtual Type Name ParameterTypesAndNames
+	#define DEFINE_RHIMETHOD(Type,Name,ParameterTypesAndNames,ParameterNames,ReturnStatement,NullImplementation) virtual Type RHI##Name ParameterTypesAndNames
 	#include "RHIMethods.h"
 	#undef DEFINE_RHIMETHOD
 
diff --git a/Engine/Source/Runtime/Engine/Classes/Components/SkyLightComponent.h b/Engine/Source/Runtime/Engine/Classes/Components/SkyLightComponent.h
--- a/Engine/Source/Runtime/Engine/Classes/Components/SkyLightComponent.h
+++ b/Engine/Source/Runtime/Engine/Classes/Components/SkyLightComponent.h
@@ -29,9 +29,9 @@ public:
 		Format = InFormat;
 	}
 
-	virtual void InitRHI();
+	virtual void InitRHI() override;
 
-	virtual void ReleaseRHI()
+	virtual void ReleaseRHI() override
 	{
 		TextureCubeRHI.SafeRelease();
 		FTexture::ReleaseRHI();
diff --git a/Engine/Source/Runtime/Engine/Classes/Engine/CodecMovie.h b/Engine/Source/Runtime/Engine/Classes/Engine/CodecMovie.h
--- a/Engine/Source/Runtime/Engine/Classes/Engine/CodecMovie.h
+++ b/Engine/Source/Runtime/Engine/Classes/Engine/CodecMovie.h
@@ -89,7 +89,7 @@ class UCodecMovie : public UObject
 	 *
 	 * @param InTextureMovieResource - output from movie decoding is written to this resource
 	 */
-	virtual void GetFrame( class FTextureMovieResource* InTextureMovieResource ) {}
+	virtual void GetFrame(FRHICommandListImmediate& RHICmdList, class FTextureMovieResource* InTextureMovieResource) {}
 
 	/**
 	 * Returns the playback time of the movie.
diff --git a/Engine/Source/Runtime/Engine/Classes/Engine/CodecMovieFallback.h b/Engine/Source/Runtime/Engine/Classes/Engine/CodecMovieFallback.h
--- a/Engine/Source/Runtime/Engine/Classes/Engine/CodecMovieFallback.h
+++ b/Engine/Source/Runtime/Engine/Classes/Engine/CodecMovieFallback.h
@@ -37,7 +37,7 @@ class UCodecMovieFallback : public UCodecMovie
 
 	virtual void ResetStream() override;
 
-	virtual void GetFrame( class FTextureMovieResource* InTextureMovieResource ) override;
+	virtual void GetFrame(FRHICommandListImmediate& RHICmdList, class FTextureMovieResource* InTextureMovieResource) override;
 
 	// End CodecMovie Interface
 };
diff --git a/Engine/Source/Runtime/Engine/Private/BatchedElements.cpp b/Engine/Source/Runtime/Engine/Private/BatchedElements.cpp
--- a/Engine/Source/Runtime/Engine/Private/BatchedElements.cpp
+++ b/Engine/Source/Runtime/Engine/Private/BatchedElements.cpp
@@ -350,7 +350,7 @@ float GBatchedElementSmoothWidth = 12;
  * Sets the appropriate vertex and pixel shader.
  */
 void FBatchedElements::PrepareShaders(
-	FRHICommandList& RHICmdList,
+	FRHICommandListImmediate& RHICmdList,
 	ESimpleElementBlendMode BlendMode,
 	const FMatrix& Transform,
 	bool bSwitchVerticalAxis,
@@ -554,7 +554,7 @@ void FBatchedElements::PrepareShaders(
 PRAGMA_DISABLE_OPTIMIZATION
 #endif
 
-void FBatchedElements::DrawPointElements(FRHICommandList& RHICmdList, const FMatrix& Transform, const uint32 ViewportSizeX, const uint32 ViewportSizeY, const FVector& CameraX, const FVector& CameraY) const
+void FBatchedElements::DrawPointElements(FRHICommandListImmediate& RHICmdList, const FMatrix& Transform, const uint32 ViewportSizeX, const uint32 ViewportSizeY, const FVector& CameraX, const FVector& CameraY) const
 {
 	// Draw the point elements.
 	if( Points.Num() > 0 )
@@ -596,7 +596,7 @@ void FBatchedElements::DrawPointElements(FRHICommandList& RHICmdList, const FMat
 }
 
 
-bool FBatchedElements::Draw(FRHICommandList& RHICmdList, bool bNeedToSwitchVerticalAxis, const FMatrix& Transform, uint32 ViewportSizeX, uint32 ViewportSizeY, bool bHitTesting, float Gamma, const FSceneView* View, FTexture2DRHIRef DepthTexture) const
+bool FBatchedElements::Draw(FRHICommandListImmediate& RHICmdList, bool bNeedToSwitchVerticalAxis, const FMatrix& Transform, uint32 ViewportSizeX, uint32 ViewportSizeY, bool bHitTesting, float Gamma, const FSceneView* View, FTexture2DRHIRef DepthTexture) const
 {
 	if( HasPrimsToDraw() )
 	{
diff --git a/Engine/Source/Runtime/Engine/Private/CanvasRenderTarget2D.cpp b/Engine/Source/Runtime/Engine/Private/CanvasRenderTarget2D.cpp
--- a/Engine/Source/Runtime/Engine/Private/CanvasRenderTarget2D.cpp
+++ b/Engine/Source/Runtime/Engine/Private/CanvasRenderTarget2D.cpp
@@ -44,9 +44,6 @@ void UCanvasRenderTarget2D::UpdateResource()
 		TextureRenderTarget,
 		(FTextureRenderTarget2DResource*)GameThread_GetRenderTargetResource(),
 		{
-			//@todo-rco: RHIPacketList
-			FRHICommandList& RHICmdList = FRHICommandList::GetNullRef();
-
 			SetRenderTarget(RHICmdList, TextureRenderTarget->GetRenderTargetTexture(), FTexture2DRHIRef());
 			RHICmdList.SetViewport(0, 0, 0.0f, TextureRenderTarget->GetSizeXY().X, TextureRenderTarget->GetSizeXY().Y, 1.0f);
 		}
@@ -76,7 +73,7 @@ void UCanvasRenderTarget2D::UpdateResource()
 		RenderTargetResource,
 		GameThread_GetRenderTargetResource(),
 		{
-			RHICopyToResolveTarget(RenderTargetResource->GetRenderTargetTexture(), RenderTargetResource->TextureRHI, true, FResolveParams());
+			RHICmdList.CopyToResolveTarget(RenderTargetResource->GetRenderTargetTexture(), RenderTargetResource->TextureRHI, true, FResolveParams());
 		}
 	);
 }
diff --git a/Engine/Source/Runtime/Engine/Private/ClearReplacementShaders.cpp b/Engine/Source/Runtime/Engine/Private/ClearReplacementShaders.cpp
--- a/Engine/Source/Runtime/Engine/Private/ClearReplacementShaders.cpp
+++ b/Engine/Source/Runtime/Engine/Private/ClearReplacementShaders.cpp
@@ -14,16 +14,14 @@ IMPLEMENT_SHADER_TYPE(,FClearBufferReplacementCS,TEXT("ClearReplacementShaders")
 
 void FClearTexture2DReplacementCS::SetParameters( FRHICommandList& RHICmdList, FUnorderedAccessViewRHIParamRef TextureRW, FLinearColor Value )
 {
-	RHICmdList.CheckIsNull();
 	FComputeShaderRHIParamRef ComputeShaderRHI = GetComputeShader();
 	SetShaderValue(RHICmdList, ComputeShaderRHI, ClearColor, Value);
-	RHISetUAVParameter(ComputeShaderRHI, ClearTextureRW.GetBaseIndex(), TextureRW);
+	RHICmdList.SetUAVParameter(ComputeShaderRHI, ClearTextureRW.GetBaseIndex(), TextureRW);
 }
 
 void FClearBufferReplacementCS::SetParameters( FRHICommandList& RHICmdList, FUnorderedAccessViewRHIParamRef BufferRW, uint32 Dword )
 {
-	RHICmdList.CheckIsNull();
 	FComputeShaderRHIParamRef ComputeShaderRHI = GetComputeShader();
 	SetShaderValue(RHICmdList, ComputeShaderRHI, ClearDword, Dword);
-	RHISetUAVParameter(ComputeShaderRHI, ClearBufferRW.GetBaseIndex(), BufferRW);
+	RHICmdList.SetUAVParameter(ComputeShaderRHI, ClearBufferRW.GetBaseIndex(), BufferRW);
 }
diff --git a/Engine/Source/Runtime/Engine/Private/ContentStreaming.cpp b/Engine/Source/Runtime/Engine/Private/ContentStreaming.cpp
--- a/Engine/Source/Runtime/Engine/Private/ContentStreaming.cpp
+++ b/Engine/Source/Runtime/Engine/Private/ContentStreaming.cpp
@@ -2308,16 +2308,15 @@ struct FTextureStreamingCompare
  *	@param RequiredMemorySize	- Amount of memory to try to free up, in bytes
  *	@param bSucceeded			- [out] Upon return, whether it succeeded or not
  **/
-void Renderthread_StreamOutTextureData( TArray<FTextureSortElement>* InCandidateTextures, int32 RequiredMemorySize, volatile bool* bSucceeded )
+void Renderthread_StreamOutTextureData(FRHICommandListImmediate& RHICmdList, TArray<FTextureSortElement>* InCandidateTextures, int32 RequiredMemorySize, volatile bool* bSucceeded)
 {
 	*bSucceeded = false;
-
 	// only for debugging?
 	FTextureMemoryStats OldStats;
-	RHIGetTextureMemoryStats(OldStats);
+	RHICmdList.GetTextureMemoryStats(OldStats);
 
 	// Makes sure that texture memory can get freed up right away.
-	RHIBlockUntilGPUIdle();
+	RHICmdList.BlockUntilGPUIdle();
 
 	FTextureSortElement* CandidateTextures = InCandidateTextures->GetTypedData();
 
@@ -2447,7 +2446,7 @@ bool FStreamingManagerTexture::StreamOutTextureData( int32 RequiredMemorySize )
 		int32,RequiredMemorySize,RequiredMemorySize,
 		volatile bool*,bSucceeded,&bSucceeded,
 	{
-		Renderthread_StreamOutTextureData( CandidateTextures, RequiredMemorySize, bSucceeded );
+		Renderthread_StreamOutTextureData(RHICmdList, CandidateTextures, RequiredMemorySize, bSucceeded);
 	});
 
 	// Block until the command has finished executing.
diff --git a/Engine/Source/Runtime/Engine/Private/DynamicMeshBuilder.cpp b/Engine/Source/Runtime/Engine/Private/DynamicMeshBuilder.cpp
--- a/Engine/Source/Runtime/Engine/Private/DynamicMeshBuilder.cpp
+++ b/Engine/Source/Runtime/Engine/Private/DynamicMeshBuilder.cpp
@@ -183,7 +183,7 @@ public:
 	TArray<int32> Indices;
 
 	// FRenderResource interface.
-	virtual void InitRHI()
+	virtual void InitRHI() override
 	{
 		uint32 SizeInBytes = Indices.Num() * sizeof(int32);
 		if(SizeInBytes <= FGlobalDynamicMeshIndexPolicy().GetPoolBucketSize(FGlobalDynamicMeshIndexPolicy::NumPoolBuckets - 1))
@@ -202,7 +202,7 @@ public:
 		RHIUnlockIndexBuffer(IndexBufferRHI);
 	}
 	
-	virtual void ReleaseRHI()
+	virtual void ReleaseRHI() override
 	{
 		if(IndexBufferRHI->GetSize() <= FGlobalDynamicMeshIndexPolicy().GetPoolBucketSize(FGlobalDynamicMeshIndexPolicy::NumPoolBuckets - 1))
 		{
@@ -232,7 +232,7 @@ public:
 	TArray<FDynamicMeshVertex> Vertices;
 
 	// FRenderResource interface.
-	virtual void InitRHI()
+	virtual void InitRHI() override
 	{
 		uint32 SizeInBytes = Vertices.Num() * sizeof(FDynamicMeshVertex);
 		if(SizeInBytes <= FGlobalDynamicMeshIndexPolicy().GetPoolBucketSize(FGlobalDynamicMeshIndexPolicy::NumPoolBuckets - 1))
@@ -251,7 +251,7 @@ public:
 		RHIUnlockVertexBuffer(VertexBufferRHI);
 	}
 	
-	virtual void ReleaseRHI()
+	virtual void ReleaseRHI() override
 	{
 		if(VertexBufferRHI->GetSize() <= FGlobalDynamicMeshVertexPolicy().GetPoolBucketSize(FGlobalDynamicMeshVertexPolicy::NumPoolBuckets - 1))
 		{
diff --git a/Engine/Source/Runtime/Engine/Private/GPUSkinCache.cpp b/Engine/Source/Runtime/Engine/Private/GPUSkinCache.cpp
--- a/Engine/Source/Runtime/Engine/Private/GPUSkinCache.cpp
+++ b/Engine/Source/Runtime/Engine/Private/GPUSkinCache.cpp
@@ -89,8 +89,6 @@ public:
 		SetShaderValue(RHICmdList, ComputeShaderRHI, SkinInputStreamFloatOffset, InputStreamFloatOffset);
 		SetShaderValue(RHICmdList, ComputeShaderRHI, SkinOutputBufferFloatOffset, OutputBufferFloatOffset);
 
-		RHICmdList.CheckIsNull();
-
 		if (UniformBuffer)
 		{
 			SetUniformBufferParameter(RHICmdList, ComputeShaderRHI, GetUniformBufferParameter<GPUSkinCacheBonesUniformShaderParameters>(), UniformBuffer);
@@ -98,13 +96,13 @@ public:
 		else
 		{
 #if GPUSKIN_USE_DATA_BUFFERS
-			RHISetShaderResourceViewParameter(ComputeShaderRHI, BoneMatrices.GetBaseIndex(), BoneBuffer.VertexBufferSRV);
+			RHICmdList.SetShaderResourceViewParameter(ComputeShaderRHI, BoneMatrices.GetBaseIndex(), BoneBuffer.VertexBufferSRV);
 #endif
 		}
 
-		RHISetShaderResourceViewParameter(ComputeShaderRHI, SkinInputStream.GetBaseIndex(), VertexBufferSRV);
+		RHICmdList.SetShaderResourceViewParameter(ComputeShaderRHI, SkinInputStream.GetBaseIndex(), VertexBufferSRV);
 
-		RHISetUAVParameter(
+		RHICmdList.SetUAVParameter(
 			ComputeShaderRHI,
 			SkinCacheBufferRW.GetBaseIndex(),
 			SkinBuffer.UAV
@@ -113,17 +111,15 @@ public:
 
 	void UnsetParameters(FRHICommandList& RHICmdList)
 	{
-		RHICmdList.CheckIsNull();
-
 		FComputeShaderRHIParamRef ComputeShaderRHI = GetComputeShader();
 		FShaderResourceViewRHIParamRef NullSRV = FShaderResourceViewRHIParamRef();
 #if GPUSKIN_USE_DATA_BUFFERS
-		RHISetShaderResourceViewParameter(ComputeShaderRHI, BoneMatrices.GetBaseIndex(), NullSRV);
+		RHICmdList.SetShaderResourceViewParameter(ComputeShaderRHI, BoneMatrices.GetBaseIndex(), NullSRV);
 #endif
 
-		RHISetShaderResourceViewParameter(ComputeShaderRHI, SkinInputStream.GetBaseIndex(), NullSRV);
+		RHICmdList.SetShaderResourceViewParameter(ComputeShaderRHI, SkinInputStream.GetBaseIndex(), NullSRV);
 
-		RHISetUAVParameter( ComputeShaderRHI, SkinCacheBufferRW.GetBaseIndex(), FUnorderedAccessViewRHIParamRef() );
+		RHICmdList.SetUAVParameter( ComputeShaderRHI, SkinCacheBufferRW.GetBaseIndex(), FUnorderedAccessViewRHIParamRef() );
 	}
 
 	virtual bool Serialize(FArchive& Ar)
@@ -197,7 +193,7 @@ void FGPUSkinCache::Cleanup()
 }
 
 // Kick off compute shader skinning for incoming chunk, return key for fast lookup in draw passes
-int32 FGPUSkinCache::StartCacheMesh(int32 Key, const FVertexFactory* VertexFactory, const FVertexFactory* TargetVertexFactory, const FSkelMeshChunk& BatchElement, const FSkeletalMeshObjectGPUSkin* Skin, bool bExtraBoneInfluences)
+int32 FGPUSkinCache::StartCacheMesh(FRHICommandListImmediate& RHICmdList, int32 Key, const FVertexFactory* VertexFactory, const FVertexFactory* TargetVertexFactory, const FSkelMeshChunk& BatchElement, const FSkeletalMeshObjectGPUSkin* Skin, bool bExtraBoneInfluences)
 {
 	if (!GEnableGPUSkinCache)
 	{
@@ -357,8 +353,7 @@ int32 FGPUSkinCache::StartCacheMesh(int32 Key, const FVertexFactory* VertexFacto
 	}
 
 	auto& ShaderData = GPUVertexFactory->GetShaderData();
-	//@todo-rco: RHIPacketList
-	FRHICommandList& RHICmdList = FRHICommandList::GetNullRef();
+	
 	DispatchSkinCacheProcess(RHICmdList, InputStreamFloatOffset, ECSInfo->StreamOffset, ShaderData.GetBoneBuffer(), ShaderData.GetUniformBuffer(), VBInfo, StreamStrides[0], ElementVertexCount, ShaderData.MeshOrigin, ShaderData.MeshExtension, bExtraBoneInfluences);
 
 	CacheCurrentFloatOffset += ProcessedFloatCount;
@@ -402,9 +397,7 @@ bool FGPUSkinCache::SetVertexStreamFromCache(FRHICommandList& RHICmdList, int32
 				return false;
 			}
 
-			RHICmdList.CheckIsNull();
-
-			RHISetStreamSource(0, SkinCacheBufferRW.Buffer, CacheInfo.StreamStride, CompensatedOffset);
+			RHICmdList.SetStreamSource(0, SkinCacheBufferRW.Buffer, CacheInfo.StreamStride, CompensatedOffset);
 
 			FVertexShaderRHIParamRef VertexShaderRHI = Shader->GetVertexShader();
 			if (VelocityPass && VertexShaderRHI)
@@ -427,7 +420,7 @@ bool FGPUSkinCache::SetVertexStreamFromCache(FRHICommandList& RHICmdList, int32
 
 				if (PreviousStreamBuffer.IsBound())
 				{
-					RHISetShaderResourceViewParameter(VertexShaderRHI, PreviousStreamBuffer.GetBaseIndex(), SkinCacheBufferRW.SRV);
+					RHICmdList.SetShaderResourceViewParameter(VertexShaderRHI, PreviousStreamBuffer.GetBaseIndex(), SkinCacheBufferRW.SRV);
 				}
 			}
 
@@ -478,7 +471,7 @@ FGPUSkinCache::FElementCacheStatusInfo* FGPUSkinCache::FindEvictableCacheStatusI
 	}
 }
 
-void FGPUSkinCache::DispatchSkinCacheProcess(FRHICommandList& RHICmdList, uint32 InputStreamFloatOffset, uint32 OutputBufferFloatOffset, FBoneBufferTypeRef BoneBuffer, FUniformBufferRHIRef UniformBuffer, const FVertexBufferInfo* VBInfo, uint32 VertexStride, uint32 VertexCount, const FVector& MeshOrigin, const FVector& MeshExtension, bool bUseExtraBoneInfluences)
+void FGPUSkinCache::DispatchSkinCacheProcess(FRHICommandListImmediate& RHICmdList, uint32 InputStreamFloatOffset, uint32 OutputBufferFloatOffset, FBoneBufferTypeRef BoneBuffer, FUniformBufferRHIRef UniformBuffer, const FVertexBufferInfo* VBInfo, uint32 VertexStride, uint32 VertexCount, const FVector& MeshOrigin, const FVector& MeshExtension, bool bUseExtraBoneInfluences)
 {
 	if (GRHIFeatureLevel < ERHIFeatureLevel::SM5)
 	{
@@ -493,10 +486,9 @@ void FGPUSkinCache::DispatchSkinCacheProcess(FRHICommandList& RHICmdList, uint32
 	{
 		TShaderMapRef<FGPUSkinCacheCS<true> > SkinCacheCS(GetGlobalShaderMap());
 		SkinCacheCS->SetParameters(RHICmdList, VertexStride, InputStreamFloatOffset, VertexCount, OutputBufferFloatOffset, BoneBuffer, UniformBuffer, VBInfo->VertexBufferSRV, SkinCacheBufferRW, MeshOrigin, MeshExtension);
-		RHICmdList.CheckIsNull();
 
-		RHISetComputeShader(SkinCacheCS->GetComputeShader());		
-		RHIDispatchComputeShader(VertexCountAlign64, 1, 1);
+		RHICmdList.SetComputeShader(SkinCacheCS->GetComputeShader());
+		RHICmdList.DispatchComputeShader(VertexCountAlign64, 1, 1);
 		SkinCacheCS->UnsetParameters(RHICmdList);
 	}
 	else
@@ -504,8 +496,8 @@ void FGPUSkinCache::DispatchSkinCacheProcess(FRHICommandList& RHICmdList, uint32
 		TShaderMapRef<FGPUSkinCacheCS<false> > SkinCacheCS(GetGlobalShaderMap());
 		SkinCacheCS->SetParameters(RHICmdList, VertexStride, InputStreamFloatOffset, VertexCount, OutputBufferFloatOffset, BoneBuffer, UniformBuffer, VBInfo->VertexBufferSRV, SkinCacheBufferRW, MeshOrigin, MeshExtension);
 
-		RHISetComputeShader(SkinCacheCS->GetComputeShader());		
-		RHIDispatchComputeShader(VertexCountAlign64, 1, 1);
+		RHICmdList.SetComputeShader(SkinCacheCS->GetComputeShader());
+		RHICmdList.DispatchComputeShader(VertexCountAlign64, 1, 1);
 		SkinCacheCS->UnsetParameters(RHICmdList);
 	}
 }
diff --git a/Engine/Source/Runtime/Engine/Private/GPUSkinCache.h b/Engine/Source/Runtime/Engine/Private/GPUSkinCache.h
--- a/Engine/Source/Runtime/Engine/Private/GPUSkinCache.h
+++ b/Engine/Source/Runtime/Engine/Private/GPUSkinCache.h
@@ -40,7 +40,7 @@ public:
 	//	For each Chunk:
 	//		Call Add*
 	//	Call End*
-	int32	StartCacheMesh(int32 Key, const class FVertexFactory* VertexFactory, const FVertexFactory* TargetVertexFactory, const struct FSkelMeshChunk& BatchElement, const FSkeletalMeshObjectGPUSkin* Skin, bool bExtraBoneInfluences);
+	int32	StartCacheMesh(FRHICommandListImmediate& RHICmdList, int32 Key, const class FVertexFactory* VertexFactory, const FVertexFactory* TargetVertexFactory, const struct FSkelMeshChunk& BatchElement, const FSkeletalMeshObjectGPUSkin* Skin, bool bExtraBoneInfluences);
 
 	bool	SetVertexStreamFromCache(FRHICommandList& RHICmdList, int32 Key, FShader* Shader, const FVertexFactory* VertexFactory, uint32 BaseVertexIndex, bool VelocityPass, FShaderParameter PreviousStreamFloatOffset, FShaderParameter PreviousStreamFloatStride, FShaderResourceParameter PreviousStreamBuffer);
 
@@ -113,7 +113,7 @@ private:
 
 	FElementCacheStatusInfo* FindEvictableCacheStatusInfo();
 
-	void	DispatchSkinCacheProcess(FRHICommandList& RHICmdList, uint32 InputStreamFloatOffset, uint32 OutputBufferFloatOffset, FBoneBufferTypeRef BoneBuffer, FUniformBufferRHIRef UniformBuffer, const FVertexBufferInfo* VBInfo, uint32 VertexStride, uint32 VertexCount, const FVector& MeshOrigin, const FVector& MeshExtension, bool bUseExtraBoneInfluences);
+	void	DispatchSkinCacheProcess(FRHICommandListImmediate& RHICmdList, uint32 InputStreamFloatOffset, uint32 OutputBufferFloatOffset, FBoneBufferTypeRef BoneBuffer, FUniformBufferRHIRef UniformBuffer, const FVertexBufferInfo* VBInfo, uint32 VertexStride, uint32 VertexCount, const FVector& MeshOrigin, const FVector& MeshExtension, bool bUseExtraBoneInfluences);
 
 	bool	Initialized : 1;
 
diff --git a/Engine/Source/Runtime/Engine/Private/GPUSkinVertexFactory.cpp b/Engine/Source/Runtime/Engine/Private/GPUSkinVertexFactory.cpp
--- a/Engine/Source/Runtime/Engine/Private/GPUSkinVertexFactory.cpp
+++ b/Engine/Source/Runtime/Engine/Private/GPUSkinVertexFactory.cpp
@@ -426,10 +426,8 @@ public:
 	/**
 	* Set any shader data specific to this vertex factory
 	*/
-	virtual void SetMesh(FRHICommandList& RHICmdList, FShader* Shader,const FVertexFactory* VertexFactory,const FSceneView& View,const FMeshBatchElement& BatchElement,uint32 DataFlags) const override
+	virtual void SetMesh(FRHICommandList& RHICmdList, FShader* Shader, const FVertexFactory* VertexFactory, const FSceneView& View, const FMeshBatchElement& BatchElement, uint32 DataFlags) const override
 	{
-		RHICmdList.CheckIsNull();
-
 		if(Shader->GetVertexShader())
 		{
 			const auto FeatureLevel = View.GetFeatureLevel();
@@ -453,9 +451,9 @@ public:
 				if(BoneMatrices.IsBound())
 				{
 #if GPUSKIN_USE_DATA_BUFFERS
-					RHISetShaderResourceViewParameter(Shader->GetVertexShader(), BoneMatrices.GetBaseIndex(), ShaderData.GetBoneBuffer().VertexBufferSRV);
+					RHICmdList.SetShaderResourceViewParameter(Shader->GetVertexShader(), BoneMatrices.GetBaseIndex(), ShaderData.GetBoneBuffer().VertexBufferSRV);
 #else
-					RHISetShaderTexture(Shader->GetVertexShader(), BoneMatrices.GetBaseIndex(), ShaderData.GetBoneBuffer());
+					RHICmdList.SetShaderTexture(Shader->GetVertexShader(), BoneMatrices.GetBaseIndex(), ShaderData.GetBoneBuffer());
 #endif
 				}
 			}
@@ -480,9 +478,9 @@ public:
 				if(PreviousBoneMatrices.IsBound())
 				{
 #if GPUSKIN_USE_DATA_BUFFERS
-					RHISetShaderResourceViewParameter(Shader->GetVertexShader(), PreviousBoneMatrices.GetBaseIndex(), GPrevPerBoneMotionBlur.GetReadData()->BoneBuffer.VertexBufferSRV);
+					RHICmdList.SetShaderResourceViewParameter(Shader->GetVertexShader(), PreviousBoneMatrices.GetBaseIndex(), GPrevPerBoneMotionBlur.GetReadData()->BoneBuffer.VertexBufferSRV);
 #else
-					RHISetShaderTexture(Shader->GetVertexShader(), PreviousBoneMatrices.GetBaseIndex(), GPrevPerBoneMotionBlur.GetReadData()->BoneBuffer);
+					RHICmdList.SetShaderTexture(Shader->GetVertexShader(), PreviousBoneMatrices.GetBaseIndex(), GPrevPerBoneMotionBlur.GetReadData()->BoneBuffer);
 #endif
 				}
 
@@ -711,7 +709,7 @@ public:
 		Ar << ClothBlendWeightParameter;
 	}
 
-	virtual void SetMesh(FRHICommandList& RHICmdList, FShader* Shader,const FVertexFactory* VertexFactory,const FSceneView& View,const FMeshBatchElement& BatchElement,uint32 DataFlags) const override
+	virtual void SetMesh(FRHICommandList& RHICmdList, FShader* Shader, const FVertexFactory* VertexFactory, const FSceneView& View, const FMeshBatchElement& BatchElement, uint32 DataFlags) const override
 	{
 		if(Shader->GetVertexShader())
 		{
@@ -725,24 +723,23 @@ public:
 				: ((const TGPUSkinAPEXClothVertexFactory<false>*)GPUSkinVertexFactory)->GetClothShaderData();
 
 			SetUniformBufferParameter(RHICmdList, Shader->GetVertexShader(),Shader->GetUniformBufferParameter<FAPEXClothUniformShaderParameters>(),ClothShaderData.GetClothUniformBuffer());
-			RHICmdList.CheckIsNull();
 
 			// we tell the shader where to pickup the data
 			if(ClothSimulPositionsParameter.IsBound())
 			{
 #if GPUSKIN_USE_DATA_BUFFERS
-				RHISetShaderResourceViewParameter(Shader->GetVertexShader(), ClothSimulPositionsParameter.GetBaseIndex(), ClothShaderData.GetClothSimulPositionBuffer().VertexBufferSRV);
+				RHICmdList.SetShaderResourceViewParameter(Shader->GetVertexShader(), ClothSimulPositionsParameter.GetBaseIndex(), ClothShaderData.GetClothSimulPositionBuffer().VertexBufferSRV);
 #else
-				RHISetShaderTexture(Shader->GetVertexShader(), ClothSimulPositionsParameter.GetBaseIndex(), ClothShaderData.GetClothSimulPositionBuffer());
+				RHICmdList.SetShaderTexture(Shader->GetVertexShader(), ClothSimulPositionsParameter.GetBaseIndex(), ClothShaderData.GetClothSimulPositionBuffer());
 #endif
 			}
 			
 			if(ClothSimulNormalsParameter.IsBound())
 			{
 #if GPUSKIN_USE_DATA_BUFFERS
-				RHISetShaderResourceViewParameter(Shader->GetVertexShader(), ClothSimulNormalsParameter.GetBaseIndex(), ClothShaderData.GetClothSimulNormalBuffer().VertexBufferSRV);
+				RHICmdList.SetShaderResourceViewParameter(Shader->GetVertexShader(), ClothSimulNormalsParameter.GetBaseIndex(), ClothShaderData.GetClothSimulNormalBuffer().VertexBufferSRV);
 #else
-				RHISetShaderTexture(Shader->GetVertexShader(), ClothSimulNormalsParameter.GetBaseIndex(), ClothShaderData.GetClothSimulNormalBuffer());
+				RHICmdList.SetShaderTexture(Shader->GetVertexShader(), ClothSimulNormalsParameter.GetBaseIndex(), ClothShaderData.GetClothSimulNormalBuffer());
 #endif
 			}
 
diff --git a/Engine/Source/Runtime/Engine/Private/GPUSort.cpp b/Engine/Source/Runtime/Engine/Private/GPUSort.cpp
--- a/Engine/Source/Runtime/Engine/Private/GPUSort.cpp
+++ b/Engine/Source/Runtime/Engine/Private/GPUSort.cpp
@@ -265,24 +265,24 @@ public:
 	/**
 	 * Set output buffer for this shader.
 	 */
-	void SetOutput( FUnorderedAccessViewRHIParamRef OutOffsetsUAV )
+	void SetOutput(FRHICommandList& RHICmdList, FUnorderedAccessViewRHIParamRef OutOffsetsUAV)
 	{
 		FComputeShaderRHIParamRef ComputeShaderRHI = GetComputeShader();
 		if ( OutOffsets.IsBound() )
 		{
-			RHISetUAVParameter( ComputeShaderRHI, OutOffsets.GetBaseIndex(), OutOffsetsUAV );
+			RHICmdList.SetUAVParameter(ComputeShaderRHI, OutOffsets.GetBaseIndex(), OutOffsetsUAV);
 		}
 	}
 
 	/**
 	 * Unbinds any buffers that have been bound.
 	 */
-	void UnbindBuffers()
+	void UnbindBuffers(FRHICommandList& RHICmdList)
 	{
 		FComputeShaderRHIParamRef ComputeShaderRHI = GetComputeShader();
 		if ( OutOffsets.IsBound() )
 		{
-			RHISetUAVParameter( ComputeShaderRHI, OutOffsets.GetBaseIndex(), FUnorderedAccessViewRHIParamRef() );
+			RHICmdList.SetUAVParameter(ComputeShaderRHI, OutOffsets.GetBaseIndex(), FUnorderedAccessViewRHIParamRef());
 		}
 	}
 
@@ -354,50 +354,49 @@ public:
 	/**
 	 * Set parameters for this shader.
 	 */
-	void SetParameters(FRHICommandList& RHICmdList, FShaderResourceViewRHIParamRef InKeysSRV, FRadixSortUniformBufferRef& RadixSortUniformBuffer, FShaderResourceViewRHIParamRef RadixSortParameterBufferSRV )
+	void SetParameters(FRHICommandList& RHICmdList, FShaderResourceViewRHIParamRef InKeysSRV, FRadixSortUniformBufferRef& RadixSortUniformBuffer, FShaderResourceViewRHIParamRef RadixSortParameterBufferSRV)
 	{
 		FComputeShaderRHIParamRef ComputeShaderRHI = GetComputeShader();
 		SetUniformBufferParameter(RHICmdList, ComputeShaderRHI, GetUniformBufferParameter<FRadixSortParameters>(), RadixSortUniformBuffer );
-		RHICmdList.CheckIsNull();
 		if ( InKeys.IsBound() )
 		{
-			RHISetShaderResourceViewParameter( ComputeShaderRHI, InKeys.GetBaseIndex(), InKeysSRV );
+			RHICmdList.SetShaderResourceViewParameter(ComputeShaderRHI, InKeys.GetBaseIndex(), InKeysSRV);
 		}
 		if ( RadixSortParameterBuffer.IsBound() )
 		{
-			RHISetShaderResourceViewParameter( ComputeShaderRHI, RadixSortParameterBuffer.GetBaseIndex(), RadixSortParameterBufferSRV );
+			RHICmdList.SetShaderResourceViewParameter(ComputeShaderRHI, RadixSortParameterBuffer.GetBaseIndex(), RadixSortParameterBufferSRV);
 		}
 	}
 
 	/**
 	 * Set output buffer for this shader.
 	 */
-	void SetOutput( FUnorderedAccessViewRHIParamRef OutOffsetsUAV )
+	void SetOutput(FRHICommandList& RHICmdList, FUnorderedAccessViewRHIParamRef OutOffsetsUAV)
 	{
 		FComputeShaderRHIParamRef ComputeShaderRHI = GetComputeShader();
 		if ( OutOffsets.IsBound() )
 		{
-			RHISetUAVParameter( ComputeShaderRHI, OutOffsets.GetBaseIndex(), OutOffsetsUAV );
+			RHICmdList.SetUAVParameter(ComputeShaderRHI, OutOffsets.GetBaseIndex(), OutOffsetsUAV);
 		}
 	}
 
 	/**
 	 * Unbinds any buffers that have been bound.
 	 */
-	void UnbindBuffers()
+	void UnbindBuffers(FRHICommandList& RHICmdList)
 	{
 		FComputeShaderRHIParamRef ComputeShaderRHI = GetComputeShader();
 		if ( RadixSortParameterBuffer.IsBound() )
 		{
-			RHISetShaderResourceViewParameter( ComputeShaderRHI, RadixSortParameterBuffer.GetBaseIndex(), FShaderResourceViewRHIParamRef() );
+			RHICmdList.SetShaderResourceViewParameter(ComputeShaderRHI, RadixSortParameterBuffer.GetBaseIndex(), FShaderResourceViewRHIParamRef());
 		}
 		if ( InKeys.IsBound() )
 		{
-			RHISetShaderResourceViewParameter( ComputeShaderRHI, InKeys.GetBaseIndex(), FShaderResourceViewRHIParamRef() );
+			RHICmdList.SetShaderResourceViewParameter(ComputeShaderRHI, InKeys.GetBaseIndex(), FShaderResourceViewRHIParamRef());
 		}
 		if ( OutOffsets.IsBound() )
 		{
-			RHISetUAVParameter( ComputeShaderRHI, OutOffsets.GetBaseIndex(), FUnorderedAccessViewRHIParamRef() );
+			RHICmdList.SetUAVParameter(ComputeShaderRHI, OutOffsets.GetBaseIndex(), FUnorderedAccessViewRHIParamRef());
 		}
 	}
 
@@ -462,42 +461,40 @@ public:
 	/**
 	 * Set parameters for this shader.
 	 */
-	void SetParameters(FRHICommandList& RHICmdList, FShaderResourceViewRHIParamRef InOffsetsSRV )
+	void SetParameters(FRHICommandList& RHICmdList, FShaderResourceViewRHIParamRef InOffsetsSRV)
 	{
 		FComputeShaderRHIParamRef ComputeShaderRHI = GetComputeShader();
 		if ( InOffsets.IsBound() )
 		{
-			RHICmdList.CheckIsNull();
-
-			RHISetShaderResourceViewParameter( ComputeShaderRHI, InOffsets.GetBaseIndex(), InOffsetsSRV );
+			RHICmdList.SetShaderResourceViewParameter(ComputeShaderRHI, InOffsets.GetBaseIndex(), InOffsetsSRV);
 		}
 	}
 
 	/**
 	 * Set output buffer for this shader.
 	 */
-	void SetOutput( FUnorderedAccessViewRHIParamRef OutOffsetsUAV )
+	void SetOutput(FRHICommandList& RHICmdList, FUnorderedAccessViewRHIParamRef OutOffsetsUAV)
 	{
 		FComputeShaderRHIParamRef ComputeShaderRHI = GetComputeShader();
 		if ( OutOffsets.IsBound() )
 		{
-			RHISetUAVParameter( ComputeShaderRHI, OutOffsets.GetBaseIndex(), OutOffsetsUAV );
+			RHICmdList.SetUAVParameter(ComputeShaderRHI, OutOffsets.GetBaseIndex(), OutOffsetsUAV);
 		}
 	}
 
 	/**
 	 * Unbinds any buffers that have been bound.
 	 */
-	void UnbindBuffers()
+	void UnbindBuffers(FRHICommandList& RHICmdList)
 	{
 		FComputeShaderRHIParamRef ComputeShaderRHI = GetComputeShader();
 		if ( InOffsets.IsBound() )
 		{
-			RHISetShaderResourceViewParameter( ComputeShaderRHI, InOffsets.GetBaseIndex(), FShaderResourceViewRHIParamRef() );
+			RHICmdList.SetShaderResourceViewParameter(ComputeShaderRHI, InOffsets.GetBaseIndex(), FShaderResourceViewRHIParamRef());
 		}
 		if ( OutOffsets.IsBound() )
 		{
-			RHISetUAVParameter( ComputeShaderRHI, OutOffsets.GetBaseIndex(), FUnorderedAccessViewRHIParamRef() );
+			RHICmdList.SetUAVParameter(ComputeShaderRHI, OutOffsets.GetBaseIndex(), FUnorderedAccessViewRHIParamRef());
 		}
 	}
 
@@ -580,7 +577,7 @@ public:
 	 * Set parameters for this shader.
 	 */
 	void SetParameters(
-		FRHICommandList& RHICmdList, 
+		FRHICommandList& RHICmdList,
 		FShaderResourceViewRHIParamRef InKeysSRV,
 		FShaderResourceViewRHIParamRef InValuesSRV,
 		FShaderResourceViewRHIParamRef InOffsetsSRV,
@@ -589,70 +586,69 @@ public:
 	{
 		FComputeShaderRHIParamRef ComputeShaderRHI = GetComputeShader();
 		SetUniformBufferParameter(RHICmdList, ComputeShaderRHI, GetUniformBufferParameter<FRadixSortParameters>(), RadixSortUniformBuffer );
-		RHICmdList.CheckIsNull();
 		if ( RadixSortParameterBuffer.IsBound() )
 		{
-			RHISetShaderResourceViewParameter( ComputeShaderRHI, RadixSortParameterBuffer.GetBaseIndex(), RadixSortParameterBufferSRV );
+			RHICmdList.SetShaderResourceViewParameter(ComputeShaderRHI, RadixSortParameterBuffer.GetBaseIndex(), RadixSortParameterBufferSRV);
 		}
 		if ( InKeys.IsBound() )
 		{
-			RHISetShaderResourceViewParameter( ComputeShaderRHI, InKeys.GetBaseIndex(), InKeysSRV );
+			RHICmdList.SetShaderResourceViewParameter(ComputeShaderRHI, InKeys.GetBaseIndex(), InKeysSRV);
 		}
 		if ( InValues.IsBound() )
 		{
-			RHISetShaderResourceViewParameter( ComputeShaderRHI, InValues.GetBaseIndex(), InValuesSRV );
+			RHICmdList.SetShaderResourceViewParameter(ComputeShaderRHI, InValues.GetBaseIndex(), InValuesSRV);
 		}
 		if ( InOffsets.IsBound() )
 		{
-			RHISetShaderResourceViewParameter( ComputeShaderRHI, InOffsets.GetBaseIndex(), InOffsetsSRV );
+			RHICmdList.SetShaderResourceViewParameter(ComputeShaderRHI, InOffsets.GetBaseIndex(), InOffsetsSRV);
 		}
 	}
 
 	/**
 	 * Set output buffer for this shader.
 	 */
-	void SetOutput( FUnorderedAccessViewRHIParamRef OutKeysUAV, FUnorderedAccessViewRHIParamRef OutValuesUAV )
+	void SetOutput(FRHICommandList& RHICmdList, FUnorderedAccessViewRHIParamRef OutKeysUAV, FUnorderedAccessViewRHIParamRef OutValuesUAV)
 	{
 		FComputeShaderRHIParamRef ComputeShaderRHI = GetComputeShader();
 		if ( OutKeys.IsBound() )
 		{
-			RHISetUAVParameter( ComputeShaderRHI, OutKeys.GetBaseIndex(), OutKeysUAV );
+			RHICmdList.SetUAVParameter(ComputeShaderRHI, OutKeys.GetBaseIndex(), OutKeysUAV);
 		}
 		if ( OutValues.IsBound() )
 		{
-			RHISetUAVParameter( ComputeShaderRHI, OutValues.GetBaseIndex(), OutValuesUAV );
+			RHICmdList.SetUAVParameter(ComputeShaderRHI, OutValues.GetBaseIndex(), OutValuesUAV);
 		}
 	}
 
 	/**
 	 * Unbinds any buffers that have been bound.
 	 */
-	void UnbindBuffers()
+	void UnbindBuffers(FRHICommandList& RHICmdList)
 	{
 		FComputeShaderRHIParamRef ComputeShaderRHI = GetComputeShader();
 		if ( RadixSortParameterBuffer.IsBound() )
 		{
-			RHISetShaderResourceViewParameter( ComputeShaderRHI, RadixSortParameterBuffer.GetBaseIndex(), FShaderResourceViewRHIParamRef() );
+			RHICmdList.SetShaderResourceViewParameter(ComputeShaderRHI, RadixSortParameterBuffer.GetBaseIndex(), FShaderResourceViewRHIParamRef());
 		}
 		if ( InKeys.IsBound() )
 		{
-			RHISetShaderResourceViewParameter( ComputeShaderRHI, InKeys.GetBaseIndex(), FShaderResourceViewRHIParamRef() );
+			RHICmdList.SetShaderResourceViewParameter(ComputeShaderRHI, InKeys.GetBaseIndex(), FShaderResourceViewRHIParamRef());
 		}
 		if ( InValues.IsBound() )
 		{
-			RHISetShaderResourceViewParameter( ComputeShaderRHI, InValues.GetBaseIndex(), FShaderResourceViewRHIParamRef() );
+			RHICmdList.SetShaderResourceViewParameter(ComputeShaderRHI, InValues.GetBaseIndex(), FShaderResourceViewRHIParamRef());
 		}
 		if ( InOffsets.IsBound() )
 		{
-			RHISetShaderResourceViewParameter( ComputeShaderRHI, InOffsets.GetBaseIndex(), FShaderResourceViewRHIParamRef() );
+			RHICmdList.SetShaderResourceViewParameter(ComputeShaderRHI, InOffsets.GetBaseIndex(), FShaderResourceViewRHIParamRef());
 		}
 		if ( OutKeys.IsBound() )
 		{
-			RHISetUAVParameter( ComputeShaderRHI, OutKeys.GetBaseIndex(), FUnorderedAccessViewRHIParamRef() );
+			RHICmdList.SetUAVParameter(ComputeShaderRHI, OutKeys.GetBaseIndex(), FUnorderedAccessViewRHIParamRef());
 		}
 		if ( OutValues.IsBound() )
 		{
-			RHISetUAVParameter( ComputeShaderRHI, OutValues.GetBaseIndex(), FUnorderedAccessViewRHIParamRef() );
+			RHICmdList.SetUAVParameter(ComputeShaderRHI, OutValues.GetBaseIndex(), FUnorderedAccessViewRHIParamRef());
 		}
 	}
 
@@ -686,7 +682,7 @@ IMPLEMENT_SHADER_TYPE(,FRadixSortDownsweepCS,TEXT("RadixSortShaders"),TEXT("Radi
  * @param Count - How many items in the buffer need to be sorted.
  * @returns The index of the buffer containing sorted results.
  */
-int32 SortGPUBuffers( FGPUSortBuffers SortBuffers, int32 BufferIndex, uint32 KeyMask, int32 Count )
+int32 SortGPUBuffers(FRHICommandListImmediate& RHICmdList, FGPUSortBuffers SortBuffers, int32 BufferIndex, uint32 KeyMask, int32 Count)
 {
 	FRadixSortParameters SortParameters;
 	FRadixSortUniformBufferRef SortUniformBufferRef;
@@ -741,9 +737,7 @@ int32 SortGPUBuffers( FGPUSortBuffers SortBuffers, int32 BufferIndex, uint32 Key
 	check( UpsweepCS->RequiresConstantBufferWorkaround() == DownsweepCS->RequiresConstantBufferWorkaround() );
 	const bool bUseConstantBufferWorkaround = UpsweepCS->RequiresConstantBufferWorkaround();
 
-	//@todo-rco: RHIPacketList
-	FRHICommandList& RHICmdList = FRHICommandList::GetNullRef();
-
+	
 	// Execute each pass as needed.
 	uint32 PassBits = DIGIT_COUNT - 1;
 	for ( int32 PassIndex = 0; PassIndex < PassCount; ++PassIndex )
@@ -764,17 +758,17 @@ int32 SortGPUBuffers( FGPUSortBuffers SortBuffers, int32 BufferIndex, uint32 Key
 			}
 
 			// Clear the offsets buffer.
-			RHISetComputeShader(ClearOffsetsCS->GetComputeShader());
-			ClearOffsetsCS->SetOutput( GSortOffsetBuffers.BufferUAVs[0] );
+			RHICmdList.SetComputeShader(ClearOffsetsCS->GetComputeShader());
+			ClearOffsetsCS->SetOutput(RHICmdList, GSortOffsetBuffers.BufferUAVs[0]);
 			DispatchComputeShader(RHICmdList, *ClearOffsetsCS, 1, 1 ,1 );
-			ClearOffsetsCS->UnbindBuffers();
+			ClearOffsetsCS->UnbindBuffers(RHICmdList);
 
 			// Phase 1: Scan upsweep to compute per-digit totals.
-			RHISetComputeShader(UpsweepCS->GetComputeShader());
-			UpsweepCS->SetOutput( GSortOffsetBuffers.BufferUAVs[0] );
+			RHICmdList.SetComputeShader(UpsweepCS->GetComputeShader());
+			UpsweepCS->SetOutput(RHICmdList, GSortOffsetBuffers.BufferUAVs[0]);
 			UpsweepCS->SetParameters(RHICmdList, SortBuffers.RemoteKeySRVs[BufferIndex], SortUniformBufferRef, GRadixSortParametersBuffer.SortParametersBufferSRV );
 			DispatchComputeShader(RHICmdList, *UpsweepCS, GroupCount, 1, 1 );
-			UpsweepCS->UnbindBuffers();
+			UpsweepCS->UnbindBuffers(RHICmdList);
 
 			if (bDebugOffsets)
 			{
@@ -783,11 +777,11 @@ int32 SortGPUBuffers( FGPUSortBuffers SortBuffers, int32 BufferIndex, uint32 Key
 			}
 
 			// Phase 2: Parallel prefix scan on the offsets buffer.
-			RHISetComputeShader(SpineCS->GetComputeShader());
-			SpineCS->SetOutput( GSortOffsetBuffers.BufferUAVs[1] );
+			RHICmdList.SetComputeShader(SpineCS->GetComputeShader());
+			SpineCS->SetOutput(RHICmdList, GSortOffsetBuffers.BufferUAVs[1]);
 			SpineCS->SetParameters(RHICmdList, GSortOffsetBuffers.BufferSRVs[0] );
 			DispatchComputeShader(RHICmdList, *SpineCS, 1, 1, 1 );
-			SpineCS->UnbindBuffers();
+			SpineCS->UnbindBuffers(RHICmdList);
 
 			if (bDebugOffsets)
 			{
@@ -796,11 +790,11 @@ int32 SortGPUBuffers( FGPUSortBuffers SortBuffers, int32 BufferIndex, uint32 Key
 			}
 
 			// Phase 3: Downsweep to compute final offsets and scatter keys.
-			RHISetComputeShader(DownsweepCS->GetComputeShader());
-			DownsweepCS->SetOutput( SortBuffers.RemoteKeyUAVs[BufferIndex ^ 0x1], SortBuffers.RemoteValueUAVs[BufferIndex ^ 0x1] );
+			RHICmdList.SetComputeShader(DownsweepCS->GetComputeShader());
+			DownsweepCS->SetOutput(RHICmdList, SortBuffers.RemoteKeyUAVs[BufferIndex ^ 0x1], SortBuffers.RemoteValueUAVs[BufferIndex ^ 0x1]);
 			DownsweepCS->SetParameters(RHICmdList, SortBuffers.RemoteKeySRVs[BufferIndex], SortBuffers.RemoteValueSRVs[BufferIndex], GSortOffsetBuffers.BufferSRVs[1], SortUniformBufferRef, GRadixSortParametersBuffer.SortParametersBufferSRV );
 			DispatchComputeShader(RHICmdList, *DownsweepCS, GroupCount, 1, 1 );
-			DownsweepCS->UnbindBuffers();
+			DownsweepCS->UnbindBuffers(RHICmdList);
 
 			// Flip buffers.
 			BufferIndex ^= 0x1;
@@ -836,7 +830,7 @@ enum
  * @param TestSize - The number of elements to sort.
  * @returns true if the sort succeeded.
  */
-static bool RunGPUSortTest(int32 TestSize)
+static bool RunGPUSortTest(FRHICommandListImmediate& RHICmdList, int32 TestSize)
 {
 	FRandomStream RandomStream(0x3819FFE4);
 	FGPUSortBuffers SortBuffers;
@@ -874,24 +868,24 @@ static bool RunGPUSortTest(int32 TestSize)
 	for (int32 BufferIndex = 0; BufferIndex < 2; ++BufferIndex)
 	{
 		FRHIResourceCreateInfo CreateInfo;
-		KeysBufferRHI[BufferIndex] = RHICreateVertexBuffer(BufferSize, BUF_Static|BUF_ShaderResource|BUF_UnorderedAccess, CreateInfo);
-		KeysBufferSRV[BufferIndex] = RHICreateShaderResourceView(KeysBufferRHI[BufferIndex], /*Stride=*/ sizeof(uint32), PF_R32_UINT);
-		KeysBufferUAV[BufferIndex] = RHICreateUnorderedAccessView(KeysBufferRHI[BufferIndex], PF_R32_UINT);
-		ValuesBufferRHI[BufferIndex] = RHICreateVertexBuffer(BufferSize, BUF_Static|BUF_ShaderResource|BUF_UnorderedAccess, CreateInfo);
-		ValuesBufferSRV[BufferIndex] = RHICreateShaderResourceView(ValuesBufferRHI[BufferIndex], /*Stride=*/ sizeof(uint32), PF_R32_UINT);
-		ValuesBufferUAV[BufferIndex] = RHICreateUnorderedAccessView(ValuesBufferRHI[BufferIndex], PF_R32_UINT);
+		KeysBufferRHI[BufferIndex] = RHICmdList.CreateVertexBuffer(BufferSize, BUF_Static | BUF_ShaderResource | BUF_UnorderedAccess, CreateInfo);
+		KeysBufferSRV[BufferIndex] = RHICmdList.CreateShaderResourceView(KeysBufferRHI[BufferIndex], /*Stride=*/ sizeof(uint32), PF_R32_UINT);
+		KeysBufferUAV[BufferIndex] = RHICmdList.CreateUnorderedAccessView(KeysBufferRHI[BufferIndex], PF_R32_UINT);
+		ValuesBufferRHI[BufferIndex] = RHICmdList.CreateVertexBuffer(BufferSize, BUF_Static | BUF_ShaderResource | BUF_UnorderedAccess, CreateInfo);
+		ValuesBufferSRV[BufferIndex] = RHICmdList.CreateShaderResourceView(ValuesBufferRHI[BufferIndex], /*Stride=*/ sizeof(uint32), PF_R32_UINT);
+		ValuesBufferUAV[BufferIndex] = RHICmdList.CreateUnorderedAccessView(ValuesBufferRHI[BufferIndex], PF_R32_UINT);
 	}
 
 	// Upload initial keys and values to the GPU.
 	{
 		uint32* Buffer;
 
-		Buffer = (uint32*)RHILockVertexBuffer(KeysBufferRHI[0], /*Offset=*/ 0, BufferSize, RLM_WriteOnly);
+		Buffer = (uint32*)RHICmdList.LockVertexBuffer(KeysBufferRHI[0], /*Offset=*/ 0, BufferSize, RLM_WriteOnly);
 		FMemory::Memcpy(Buffer, Keys.GetTypedData(), BufferSize);
-		RHIUnlockVertexBuffer(KeysBufferRHI[0]);
-		Buffer = (uint32*)RHILockVertexBuffer(ValuesBufferRHI[0], /*Offset=*/ 0, BufferSize, RLM_WriteOnly);
+		RHICmdList.UnlockVertexBuffer(KeysBufferRHI[0]);
+		Buffer = (uint32*)RHICmdList.LockVertexBuffer(ValuesBufferRHI[0], /*Offset=*/ 0, BufferSize, RLM_WriteOnly);
 		FMemory::Memcpy(Buffer, Keys.GetTypedData(), BufferSize);
-		RHIUnlockVertexBuffer(ValuesBufferRHI[0]);
+		RHICmdList.UnlockVertexBuffer(ValuesBufferRHI[0]);
 	}
 
 	// Execute the GPU sort.
@@ -902,7 +896,7 @@ static bool RunGPUSortTest(int32 TestSize)
 		SortBuffers.RemoteValueSRVs[BufferIndex] = ValuesBufferSRV[BufferIndex];
 		SortBuffers.RemoteValueUAVs[BufferIndex] = ValuesBufferUAV[BufferIndex];
 	}
-	ResultBufferIndex = SortGPUBuffers(SortBuffers, 0, 0xFFFFFFFF, TestSize);
+	ResultBufferIndex = SortGPUBuffers(RHICmdList, SortBuffers, 0, 0xFFFFFFFF, TestSize);
 
 	// Download results from the GPU.
 	{
@@ -913,12 +907,12 @@ static bool RunGPUSortTest(int32 TestSize)
 		SortedValues.Reserve(TestSize);
 		SortedValues.AddUninitialized(TestSize);
 
-		Buffer = (uint32*)RHILockVertexBuffer(KeysBufferRHI[ResultBufferIndex], /*Offset=*/ 0, BufferSize, RLM_ReadOnly);
+		Buffer = (uint32*)RHICmdList.LockVertexBuffer(KeysBufferRHI[ResultBufferIndex], /*Offset=*/ 0, BufferSize, RLM_ReadOnly);
 		FMemory::Memcpy(SortedKeys.GetTypedData(), Buffer, BufferSize);
-		RHIUnlockVertexBuffer(KeysBufferRHI[ResultBufferIndex]);
-		Buffer = (uint32*)RHILockVertexBuffer(ValuesBufferRHI[ResultBufferIndex], /*Offset=*/ 0, BufferSize, RLM_ReadOnly);
+		RHICmdList.UnlockVertexBuffer(KeysBufferRHI[ResultBufferIndex]);
+		Buffer = (uint32*)RHICmdList.LockVertexBuffer(ValuesBufferRHI[ResultBufferIndex], /*Offset=*/ 0, BufferSize, RLM_ReadOnly);
 		FMemory::Memcpy(SortedValues.GetTypedData(), Buffer, BufferSize);
-		RHIUnlockVertexBuffer(ValuesBufferRHI[ResultBufferIndex]);
+		RHICmdList.UnlockVertexBuffer(ValuesBufferRHI[ResultBufferIndex]);
 	}
 
 	// Verify results.
@@ -967,7 +961,7 @@ static bool RunGPUSortTest(int32 TestSize)
  * Executes a sort test with debug information enabled.
  * @param TestSize - The number of elements to sort.
  */
-static void RunGPUSortTestWithDebug(int32 TestSize)
+static void RunGPUSortTestWithDebug(FRHICommandListImmediate& RHICmdList, int32 TestSize)
 {
 	static IConsoleVariable* IVarDebugOffsets = IConsoleManager::Get().FindConsoleVariable(TEXT("GPUSort.DebugOffsets"));
 	static IConsoleVariable* IVarDebugSort = IConsoleManager::Get().FindConsoleVariable(TEXT("GPUSort.DebugSort"));
@@ -975,7 +969,7 @@ static void RunGPUSortTestWithDebug(int32 TestSize)
 	const bool bWasDebuggingSort = CVarDebugSort.GetValueOnRenderThread() != 0;
 	IVarDebugOffsets->Set(1);
 	IVarDebugSort->Set(1);
-	RunGPUSortTest(TestSize);
+	RunGPUSortTest(RHICmdList, TestSize);
 	IVarDebugOffsets->Set(bWasDebuggingOffsets ? 1 : 0);
 	IVarDebugSort->Set(bWasDebuggingSort ? 1 : 0);
 }
@@ -985,13 +979,13 @@ static void RunGPUSortTestWithDebug(int32 TestSize)
  * information enabled.
  * @param TestSize - The number of elements to sort.
  */
-static bool TestGPUSortForSize(int32 TestSize)
+static bool TestGPUSortForSize(FRHICommandListImmediate& RHICmdList, int32 TestSize)
 {
 	check(IsInRenderingThread());
-	const bool bResult = RunGPUSortTest(TestSize);
+	const bool bResult = RunGPUSortTest(RHICmdList, TestSize);
 	if (bResult == false)
 	{
-		RunGPUSortTestWithDebug(TestSize);
+		RunGPUSortTestWithDebug(RHICmdList, TestSize);
 	}
 	return bResult;
 }
@@ -1000,17 +994,17 @@ static bool TestGPUSortForSize(int32 TestSize)
  * Test that GPU sorting works.
  * @param TestToRun - The test to run.
  */
-static bool TestGPUSort_RenderThread(EGPUSortTest TestToRun)
+static bool TestGPUSort_RenderThread(FRHICommandListImmediate& RHICmdList, EGPUSortTest TestToRun)
 {
 	check(IsInRenderingThread());
 
 	switch (TestToRun)
 	{
 	case GPU_SORT_TEST_SMALL:
-		return TestGPUSortForSize(GPU_SORT_TEST_SIZE_SMALL);
+		return TestGPUSortForSize(RHICmdList, GPU_SORT_TEST_SIZE_SMALL);
 
 	case GPU_SORT_TEST_LARGE:
-		return TestGPUSortForSize(GPU_SORT_TEST_SIZE_LARGE);
+		return TestGPUSortForSize(RHICmdList, GPU_SORT_TEST_SIZE_LARGE);
 
 	case GPU_SORT_TEST_EXHAUSTIVE:
 		{
@@ -1018,7 +1012,7 @@ static bool TestGPUSort_RenderThread(EGPUSortTest TestToRun)
 			int32 TestSize = GPU_SORT_TEST_SIZE_MIN;
 			while (TestSize <= GPU_SORT_TEST_SIZE_MAX)
 			{
-				if (TestGPUSortForSize(TestSize) == false)
+				if (TestGPUSortForSize(RHICmdList, TestSize) == false)
 				{
 					return false;
 				}
@@ -1029,7 +1023,7 @@ static bool TestGPUSort_RenderThread(EGPUSortTest TestToRun)
 			TestSize = GPU_SORT_TEST_SIZE_MIN;
 			while (TestSize <= GPU_SORT_TEST_SIZE_MAX)
 			{
-				if (TestGPUSortForSize(TestSize - 1) == false)
+				if (TestGPUSortForSize(RHICmdList, TestSize - 1) == false)
 				{
 					return false;
 				}
@@ -1042,7 +1036,7 @@ static bool TestGPUSort_RenderThread(EGPUSortTest TestToRun)
 		for ( int32 i = 0; i < 1000; ++i )
 		{
 			int32 TestSize = FMath::TruncToInt(FMath::SRand() * (float)(GPU_SORT_TEST_SIZE_MAX - GPU_SORT_TEST_SIZE_MIN)) + GPU_SORT_TEST_SIZE_MIN;
-			if (TestGPUSortForSize( (TestSize + 0xF) & 0xFFFFFFF0 ) == false)
+			if (TestGPUSortForSize(RHICmdList, (TestSize + 0xF) & 0xFFFFFFF0) == false)
 			{
 				return false;
 			}
@@ -1063,6 +1057,6 @@ void TestGPUSort(EGPUSortTest TestToRun)
 		FTestGPUSortCommand,
 		EGPUSortTest, TestToRun, TestToRun,
 	{
-		TestGPUSort_RenderThread(TestToRun);
+		TestGPUSort_RenderThread(RHICmdList, TestToRun);
 	});
 }
diff --git a/Engine/Source/Runtime/Engine/Private/GPUSort.h b/Engine/Source/Runtime/Engine/Private/GPUSort.h
--- a/Engine/Source/Runtime/Engine/Private/GPUSort.h
+++ b/Engine/Source/Runtime/Engine/Private/GPUSort.h
@@ -38,7 +38,7 @@ struct FGPUSortBuffers
  * @param Count - How many items in the buffer need to be sorted.
  * @returns The index of the buffer containing sorted results.
  */
-int32 SortGPUBuffers( FGPUSortBuffers SortBuffers, int32 BufferIndex, uint32 KeyMask, int32 Count );
+int32 SortGPUBuffers(FRHICommandListImmediate& RHICmdList, FGPUSortBuffers SortBuffers, int32 BufferIndex, uint32 KeyMask, int32 Count);
 
 /**
  * GPU sorting tests.
diff --git a/Engine/Source/Runtime/Engine/Private/InstancedStaticMesh.cpp b/Engine/Source/Runtime/Engine/Private/InstancedStaticMesh.cpp
--- a/Engine/Source/Runtime/Engine/Private/InstancedStaticMesh.cpp
+++ b/Engine/Source/Runtime/Engine/Private/InstancedStaticMesh.cpp
@@ -146,7 +146,7 @@ public:
 	}
 
 	// FRenderResource interface.
-	virtual void InitRHI();
+	virtual void InitRHI() override;
 	virtual FString GetFriendlyName() const { return TEXT("Static-mesh instances"); }
 
 private:
@@ -373,7 +373,7 @@ public:
 	void Copy(const FInstancedStaticMeshVertexFactory& Other);
 
 	// FRenderResource interface.
-	virtual void InitRHI();
+	virtual void InitRHI() override;
 
 	static FVertexFactoryShaderParameters* ConstructShaderParameters(EShaderFrequency ShaderFrequency);
 
diff --git a/Engine/Source/Runtime/Engine/Private/MaterialBillboardComponent.cpp b/Engine/Source/Runtime/Engine/Private/MaterialBillboardComponent.cpp
--- a/Engine/Source/Runtime/Engine/Private/MaterialBillboardComponent.cpp
+++ b/Engine/Source/Runtime/Engine/Private/MaterialBillboardComponent.cpp
@@ -19,7 +19,7 @@ class FMaterialSpriteVertexBuffer : public FVertexBuffer
 {
 public:
 
-	virtual void InitRHI()
+	virtual void InitRHI() override
 	{
 		FRHIResourceCreateInfo CreateInfo;
 		VertexBufferRHI = RHICreateVertexBuffer(sizeof(FMaterialSpriteVertex),BUF_Static,CreateInfo);
diff --git a/Engine/Source/Runtime/Engine/Private/MovieCodecs.cpp b/Engine/Source/Runtime/Engine/Private/MovieCodecs.cpp
--- a/Engine/Source/Runtime/Engine/Private/MovieCodecs.cpp
+++ b/Engine/Source/Runtime/Engine/Private/MovieCodecs.cpp
@@ -85,7 +85,7 @@ float UCodecMovieFallback::GetFrameRate()
 }
 
 
-void UCodecMovieFallback::GetFrame( class FTextureMovieResource* InTextureMovieResource )
+void UCodecMovieFallback::GetFrame(FRHICommandListImmediate& RHICmdList, class FTextureMovieResource* InTextureMovieResource)
 {
 	CurrentTime += 1.f / GetFrameRate();
 	if( CurrentTime > PlaybackDuration )
@@ -96,9 +96,9 @@ void UCodecMovieFallback::GetFrame( class FTextureMovieResource* InTextureMovieR
 		InTextureMovieResource->IsInitialized() )
 	{
 		FLinearColor ClearColor(1.f,CurrentTime/PlaybackDuration,0.f,1.f);
-		RHISetRenderTarget(InTextureMovieResource->GetRenderTargetTexture(),FTextureRHIRef());
-		RHIClear(true,ClearColor,false,0.f,false,0, FIntRect());
-		RHICopyToResolveTarget(InTextureMovieResource->GetRenderTargetTexture(),InTextureMovieResource->TextureRHI,false,FResolveParams());
+		SetRenderTarget(RHICmdList, InTextureMovieResource->GetRenderTargetTexture(), FTextureRHIRef());
+		RHICmdList.Clear(true, ClearColor, false, 0.f, false, 0, FIntRect());
+		RHICmdList.CopyToResolveTarget(InTextureMovieResource->GetRenderTargetTexture(), InTextureMovieResource->TextureRHI, false, FResolveParams());
 	}
 }
 
diff --git a/Engine/Source/Runtime/Engine/Private/NavMeshRenderingHelpers.h b/Engine/Source/Runtime/Engine/Private/NavMeshRenderingHelpers.h
--- a/Engine/Source/Runtime/Engine/Private/NavMeshRenderingHelpers.h
+++ b/Engine/Source/Runtime/Engine/Private/NavMeshRenderingHelpers.h
@@ -311,9 +311,8 @@ public:
 
 			const bool bNeedToSwitchVerticalAxis = IsES2Platform(GRHIShaderPlatform) && !IsPCPlatform(GRHIShaderPlatform);
 			const FTexture2DRHIRef DepthTexture;
-			//@todo-rco: RHIPacketList
-			FRHICommandList& RHICmdList = FRHICommandList::GetNullRef();
-
+			
+			FRHICommandListImmediate& RHICmdList = FRHICommandListExecutor::GetImmediateCommandList();
 			BatchedElements.Draw(
 				RHICmdList,
 				bNeedToSwitchVerticalAxis,
diff --git a/Engine/Source/Runtime/Engine/Private/Particles/FXSystem.cpp b/Engine/Source/Runtime/Engine/Private/Particles/FXSystem.cpp
--- a/Engine/Source/Runtime/Engine/Private/Particles/FXSystem.cpp
+++ b/Engine/Source/Runtime/Engine/Private/Particles/FXSystem.cpp
@@ -314,19 +314,19 @@ void FFXSystem::PreInitViews()
 	}
 }
 
-void FFXSystem::PreRender()
+void FFXSystem::PreRender(FRHICommandListImmediate& RHICmdList)
 {
 	if (RHISupportsGPUParticles(FeatureLevel))
 	{
-		SimulateGPUParticles(EParticleSimulatePhase::Main,NULL,FTexture2DRHIParamRef(),FTexture2DRHIParamRef());
+		SimulateGPUParticles(RHICmdList, EParticleSimulatePhase::Main, NULL, FTexture2DRHIParamRef(), FTexture2DRHIParamRef());
 	}
 }
 
-void FFXSystem::PostRenderOpaque(const class FSceneView* CollisionView, FTexture2DRHIParamRef SceneDepthTexture, FTexture2DRHIParamRef GBufferATexture)
+void FFXSystem::PostRenderOpaque(FRHICommandListImmediate& RHICmdList, const class FSceneView* CollisionView, FTexture2DRHIParamRef SceneDepthTexture, FTexture2DRHIParamRef GBufferATexture)
 {
 	if (RHISupportsGPUParticles(FeatureLevel))
 	{
-		SimulateGPUParticles(EParticleSimulatePhase::Collision,CollisionView,SceneDepthTexture,GBufferATexture);
-		SortGPUParticles();
+		SimulateGPUParticles(RHICmdList, EParticleSimulatePhase::Collision, CollisionView, SceneDepthTexture, GBufferATexture);
+		SortGPUParticles(RHICmdList);
 	}
 }
diff --git a/Engine/Source/Runtime/Engine/Private/Particles/FXSystemPrivate.h b/Engine/Source/Runtime/Engine/Private/Particles/FXSystemPrivate.h
--- a/Engine/Source/Runtime/Engine/Private/Particles/FXSystemPrivate.h
+++ b/Engine/Source/Runtime/Engine/Private/Particles/FXSystemPrivate.h
@@ -65,8 +65,8 @@ public:
 	virtual void UpdateVectorField(UVectorFieldComponent* VectorFieldComponent) override;
 	virtual FParticleEmitterInstance* CreateGPUSpriteEmitterInstance(FGPUSpriteEmitterInfo& EmitterInfo) override;
 	virtual void PreInitViews() override;
-	virtual void PreRender() override;
-	virtual void PostRenderOpaque(const class FSceneView* CollisionView, FTexture2DRHIParamRef SceneDepthTexture, FTexture2DRHIParamRef GBufferATexture) override;
+	virtual void PreRender(FRHICommandListImmediate& RHICmdList) override;
+	virtual void PostRenderOpaque(FRHICommandListImmediate& RHICmdList, const class FSceneView* CollisionView, FTexture2DRHIParamRef SceneDepthTexture, FTexture2DRHIParamRef GBufferATexture) override;
 	// End FFXSystemInterface.
 
 	/*--------------------------------------------------------------------------
@@ -141,7 +141,7 @@ private:
 	 * Sorts all GPU particles that have called AddSortedGPUSimulation since the
 	 * last reset.
 	 */
-	void SortGPUParticles();
+	void SortGPUParticles(FRHICommandListImmediate& RHICmdList);
 
 	/**
 	 * Resets all simulations to be run in the main pass.
@@ -156,6 +156,7 @@ private:
 	 * @param GBufferATexture	GBuffer texture containing the world normal.
 	 */
 	void SimulateGPUParticles(
+		FRHICommandListImmediate& RHICmdList,
 		EParticleSimulatePhase::Type Phase,
 		const class FSceneView* CollisionView,
 		FTexture2DRHIParamRef SceneDepthTexture,
diff --git a/Engine/Source/Runtime/Engine/Private/Particles/ParticleCurveTexture.cpp b/Engine/Source/Runtime/Engine/Private/Particles/ParticleCurveTexture.cpp
--- a/Engine/Source/Runtime/Engine/Private/Particles/ParticleCurveTexture.cpp
+++ b/Engine/Source/Runtime/Engine/Private/Particles/ParticleCurveTexture.cpp
@@ -134,7 +134,7 @@ public:
 	/** The vertex declaration. */
 	FVertexDeclarationRHIRef VertexDeclarationRHI;
 
-	virtual void InitRHI()
+	virtual void InitRHI() override
 	{
 		FVertexDeclarationElementList Elements;
 
@@ -174,7 +174,7 @@ TGlobalResource<FParticleCurveInjectionVertexDeclaration> GParticleCurveInjectio
  * @param InPendingCurves - Curves to be stored on the GPU.
  */
 static void InjectCurves(
-	FRHICommandList& RHICmdList, 
+	FRHICommandListImmediate& RHICmdList,
 	FTexture2DRHIParamRef CurveTextureRHI,
 	FTexture2DRHIParamRef CurveTextureTargetRHI,
 	TArray<FCurveSamples>& InPendingCurves )
@@ -187,8 +187,6 @@ static void InjectCurves(
 
 	FVertexBufferRHIParamRef ScratchVertexBufferRHI = GParticleScratchVertexBuffer.VertexBufferRHI;
 
-	RHICmdList.CheckIsNull();
-
 	SetRenderTarget(RHICmdList, CurveTextureTargetRHI, FTextureRHIParamRef());
 	RHICmdList.SetScissorRect(false, 0, 0, 0, 0);
 	RHICmdList.SetViewport(0, 0, 0.0f, GParticleCurveTextureSizeX, GParticleCurveTextureSizeY, 1.0f);
@@ -213,7 +211,7 @@ static void InjectCurves(
 		const int32 SampleByteCount = SampleCount * sizeof(FColor);
 		FColor* RESTRICT DestSamples = (FColor*)RHILockVertexBuffer( ScratchVertexBufferRHI, 0, SampleByteCount, RLM_WriteOnly );
 		FMemory::Memcpy( DestSamples, Curve.Samples, SampleByteCount );
-		RHIUnlockVertexBuffer( ScratchVertexBufferRHI );
+		RHICmdList.UnlockVertexBuffer(ScratchVertexBufferRHI);
 		FMemory::Free( Curve.Samples );
 		Curve.Samples = NULL;
 
@@ -227,7 +225,7 @@ static void InjectCurves(
 		TShaderMapRef<FParticleCurveInjectionPS> PixelShader( GetGlobalShaderMap() );
 
 		// Bound shader state.
-		RHICmdList.CheckIsNull(); // need new approach for "static FGlobalBoundShaderState" for parallel rendering
+		
 		static FGlobalBoundShaderState BoundShaderState;
 		SetGlobalBoundShaderState(
 			RHICmdList,
@@ -595,9 +593,9 @@ void FParticleCurveTexture::SubmitPendingCurves()
 			FParticleCurveTexture*, ParticleCurveTexture, this,
 			TArray<FCurveSamples>, PendingCurves, PendingCurves,
 		{
-			//@todo-rco: RHIPacketList
+			
 			InjectCurves(
-				FRHICommandList::GetNullRef(),
+				RHICmdList,
 				ParticleCurveTexture->CurveTextureRHI,
 				ParticleCurveTexture->CurveTextureTargetRHI,
 				PendingCurves
