UE4 - rework RHI user-facing APIs to be safer and incorporate the desired sematics

diff --git a/Engine/Source/Runtime/Engine/Private/Particles/ParticleGpuSimulation.cpp b/Engine/Source/Runtime/Engine/Private/Particles/ParticleGpuSimulation.cpp
--- a/Engine/Source/Runtime/Engine/Private/Particles/ParticleGpuSimulation.cpp
+++ b/Engine/Source/Runtime/Engine/Private/Particles/ParticleGpuSimulation.cpp
@@ -644,7 +644,7 @@ public:
 /**
  * Set vertex factory shader parameters.
  */
-void FGPUSpriteVertexFactoryShaderParameters::SetMesh(FRHICommandList& RHICmdList, FShader* Shader,const FVertexFactory* VertexFactory,const FSceneView& View,const FMeshBatchElement& BatchElement,uint32 DataFlags) const
+void FGPUSpriteVertexFactoryShaderParameters::SetMesh(FRHICommandList& RHICmdList, FShader* Shader, const FVertexFactory* VertexFactory, const FSceneView& View, const FMeshBatchElement& BatchElement, uint32 DataFlags) const
 {
 	FGPUSpriteVertexFactory* GPUVF = (FGPUSpriteVertexFactory*)VertexFactory;
 	FVertexShaderRHIParamRef VertexShader = Shader->GetVertexShader();
@@ -654,9 +654,7 @@ void FGPUSpriteVertexFactoryShaderParameters::SetMesh(FRHICommandList& RHICmdLis
 	SetUniformBufferParameter(RHICmdList, VertexShader, Shader->GetUniformBufferParameter<FGPUSpriteEmitterDynamicUniformParameters>(), GPUVF->EmitterDynamicUniformBuffer );
 	if (ParticleIndices.IsBound())
 	{
-		RHICmdList.CheckIsNull();
-
-		RHISetShaderResourceViewParameter(VertexShader, ParticleIndices.GetBaseIndex(), GPUVF->ParticleIndicesBuffer->VertexBufferSRV);
+		RHICmdList.SetShaderResourceViewParameter(VertexShader, ParticleIndices.GetBaseIndex(), GPUVF->ParticleIndicesBuffer->VertexBufferSRV);
 	}
 	SetShaderValue(RHICmdList, VertexShader, ParticleIndicesOffset, GPUVF->ParticleIndicesOffset);
 	SetTextureParameter(RHICmdList, VertexShader, PositionTexture, PositionTextureSampler, SamplerStatePoint, GPUVF->PositionTextureRHI );
@@ -807,12 +805,10 @@ public:
 	/** Set parameters. */
 	void SetParameters(FRHICommandList& RHICmdList, FParticleShaderParamRef TileOffsetsRef)
 	{
-		RHICmdList.CheckIsNull();
-
 		FVertexShaderRHIParamRef VertexShaderRHI = GetVertexShader();
 		if (TileOffsets.IsBound())
 		{
-			RHISetShaderResourceViewParameter(VertexShaderRHI, TileOffsets.GetBaseIndex(), TileOffsetsRef);
+			RHICmdList.SetShaderResourceViewParameter(VertexShaderRHI, TileOffsets.GetBaseIndex(), TileOffsetsRef);
 		}
 	}
 
@@ -986,7 +982,7 @@ public:
 	/**
 	 * Unbinds buffers that may need to be bound as UAVs.
 	 */
-	void UnbindBuffers()
+	void UnbindBuffers(FRHICommandList& RHICmdList)
 	{
 		FPixelShaderRHIParamRef PixelShaderRHI = GetPixelShader();
 		FShaderResourceViewRHIParamRef NullSRV = FShaderResourceViewRHIParamRef();
@@ -995,7 +991,7 @@ public:
 			const uint32 ResourceCount = VectorFieldTextures.GetNumResources();
 			for (uint32 ResourceIndex = 0; ResourceIndex < ResourceCount; ++ResourceIndex)
 			{
-				RHISetShaderResourceViewParameter(PixelShaderRHI, VectorFieldTextures.GetBaseIndex() + ResourceIndex, NullSRV);
+				RHICmdList.SetShaderResourceViewParameter(PixelShaderRHI, VectorFieldTextures.GetBaseIndex() + ResourceIndex, NullSRV);
 			}
 		}
 	}
@@ -1155,12 +1151,12 @@ static void BuildTileVertexBuffer( FParticleBufferParamRef TileOffsetsRef, const
  * Issues a draw call for an aligned set of tiles.
  * @param TileCount - The number of tiles to be drawn.
  */
-static void DrawAlignedParticleTiles( int32 TileCount )
+static void DrawAlignedParticleTiles(FRHICommandList& RHICmdList, int32 TileCount)
 {
 	check((TileCount & (TILES_PER_INSTANCE-1)) == 0);
 
 	// Stream 0: TexCoord.
-	RHISetStreamSource(
+	RHICmdList.SetStreamSource(
 		0,
 		GParticleTexCoordVertexBuffer.VertexBufferRHI,
 		/*Stride=*/ sizeof(FVector2D),
@@ -1168,7 +1164,7 @@ static void DrawAlignedParticleTiles( int32 TileCount )
 		);
 
 	// Draw tiles.
-	RHIDrawIndexedPrimitive(
+	RHICmdList.DrawIndexedPrimitive(
 		GParticleIndexBuffer.IndexBufferRHI,
 		PT_TriangleList,
 		/*BaseVertexIndex=*/ 0,
@@ -1212,7 +1208,7 @@ struct FSimulationCommandGPU
  */
 template <bool bUseDepthBufferCollision>
 void ExecuteSimulationCommands(
-	FRHICommandList& RHICmdList, 
+	FRHICommandListImmediate& RHICmdList,
 	const TArray<FSimulationCommandGPU>& SimulationCommands,
 	const FParticleStateTextures& TextureResources,
 	const FParticleAttributesTexture& AttributeTexture,
@@ -1227,7 +1223,7 @@ void ExecuteSimulationCommands(
 	TShaderMapRef<TParticleSimulationPS<bUseDepthBufferCollision> > PixelShader(GetGlobalShaderMap());
 
 	// Bound shader state.
-	RHICmdList.CheckIsNull(); // need new approach for "static FGlobalBoundShaderState" for parallel rendering
+	
 	static FGlobalBoundShaderState BoundShaderState;
 	SetGlobalBoundShaderState(
 		RHICmdList,
@@ -1252,18 +1248,18 @@ void ExecuteSimulationCommands(
 			Command.VectorFieldsUniformBuffer,
 			Command.VectorFieldTexturesRHI
 			);
-		DrawAlignedParticleTiles(Command.TileCount);
+		DrawAlignedParticleTiles(RHICmdList, Command.TileCount);
 	}
 
 	// Unbind input buffers.
-	PixelShader->UnbindBuffers();
+	PixelShader->UnbindBuffers(RHICmdList);
 }
 
 /**
  * Invokes the clear simulation shader for each particle in each tile.
  * @param Tiles - The list of tiles to clear.
  */
-void ClearTiles(const TArray<uint32>& Tiles)
+void ClearTiles(FRHICommandListImmediate& RHICmdList, const TArray<uint32>& Tiles)
 {
 	SCOPED_DRAW_EVENT(ClearTiles, DEC_PARTICLE);
 
@@ -1281,12 +1277,9 @@ void ClearTiles(const TArray<uint32>& Tiles)
 	// Grab shaders.
 	TShaderMapRef<FParticleTileVS> VertexShader( GetGlobalShaderMap() );
 	TShaderMapRef<FParticleSimulationClearPS> PixelShader( GetGlobalShaderMap() );
-
-	//@todo-rco: RHIPacketList
-	FRHICommandList& RHICmdList = FRHICommandList::GetNullRef();
-
+	
 	// Bound shader state.
-	RHICmdList.CheckIsNull(); // need new approach for "static FGlobalBoundShaderState" for parallel rendering
+	
 	static FGlobalBoundShaderState BoundShaderState;
 	SetGlobalBoundShaderState(
 		RHICmdList,
@@ -1303,9 +1296,9 @@ void ClearTiles(const TArray<uint32>& Tiles)
 		const int32 TilesThisDrawCall = FMath::Min<int32>( TileCount, MaxTilesPerDrawCall );
 		const uint32* TilesPtr = Tiles.GetTypedData() + FirstTile;
 		BuildTileVertexBuffer( BufferParam, TilesPtr, TilesThisDrawCall );
-		//@todo-rco: RHIPacketList
+		
 		VertexShader->SetParameters(RHICmdList, ShaderParam);
-		DrawAlignedParticleTiles( ComputeAlignedTileCount(TilesThisDrawCall) );
+		DrawAlignedParticleTiles(RHICmdList, ComputeAlignedTileCount(TilesThisDrawCall));
 		TileCount -= TilesThisDrawCall;
 		FirstTile += TilesThisDrawCall;
 	}
@@ -1458,7 +1451,7 @@ public:
 	/** The vertex declaration. */
 	FVertexDeclarationRHIRef VertexDeclarationRHI;
 
-	virtual void InitRHI()
+	virtual void InitRHI() override
 	{
 		FVertexDeclarationElementList Elements;
 
@@ -1494,7 +1487,7 @@ public:
 		VertexDeclarationRHI = RHICreateVertexDeclaration( Elements );
 	}
 
-	virtual void ReleaseRHI()
+	virtual void ReleaseRHI() override
 	{
 		VertexDeclarationRHI.SafeRelease();
 	}
@@ -1507,16 +1500,14 @@ TGlobalResource<FParticleInjectionVertexDeclaration> GParticleInjectionVertexDec
  * Injects new particles in to the GPU simulation.
  * @param NewParticles - A list of particles to inject in to the simulation.
  */
-void InjectNewParticles( const TArray<FNewParticle>& NewParticles )
+void InjectNewParticles(FRHICommandListImmediate& RHICmdList, const TArray<FNewParticle>& NewParticles)
 {
 	const int32 MaxParticlesPerDrawCall = GParticleScratchVertexBufferSize / sizeof(FNewParticle);
 	FVertexBufferRHIParamRef ScratchVertexBufferRHI = GParticleScratchVertexBuffer.VertexBufferRHI;
 	int32 ParticleCount = NewParticles.Num();
 	int32 FirstParticle = 0;
 
-	//@todo-rco: RHIPacketList
-	FRHICommandList& RHICmdList = FRHICommandList::GetNullRef();
-
+	
 	while ( ParticleCount > 0 )
 	{
 		// Copy new particles in to the vertex buffer.
@@ -1533,7 +1524,7 @@ void InjectNewParticles( const TArray<FNewParticle>& NewParticles )
 		TShaderMapRef<FParticleInjectionPS> PixelShader( GetGlobalShaderMap() );
 
 		// Bound shader state.
-		RHICmdList.CheckIsNull(); // need new approach for "static FGlobalBoundShaderState" for parallel rendering
+		
 		static FGlobalBoundShaderState BoundShaderState;
 		SetGlobalBoundShaderState(
 			RHICmdList,
@@ -1544,11 +1535,11 @@ void InjectNewParticles( const TArray<FNewParticle>& NewParticles )
 			0
 			);
 
-		//@todo-rco: RHIPacketList
-		VertexShader->SetParameters(FRHICommandList::GetNullRef());
+		
+		VertexShader->SetParameters(RHICmdList);
 
 		// Stream 0: New particles.
-		RHISetStreamSource(
+		RHICmdList.SetStreamSource(
 			0,
 			ScratchVertexBufferRHI,
 			/*Stride=*/ sizeof(FNewParticle),
@@ -1556,7 +1547,7 @@ void InjectNewParticles( const TArray<FNewParticle>& NewParticles )
 			);
 
 		// Stream 1: TexCoord.
-		RHISetStreamSource(
+		RHICmdList.SetStreamSource(
 			1,
 			GParticleTexCoordVertexBuffer.VertexBufferRHI,
 			/*Stride=*/ sizeof(FVector2D),
@@ -1564,7 +1555,7 @@ void InjectNewParticles( const TArray<FNewParticle>& NewParticles )
 			);
 
 		// Inject particles.
-		RHIDrawIndexedPrimitive(
+		RHICmdList.DrawIndexedPrimitive(
 			GParticleIndexBuffer.IndexBufferRHI,
 			PT_TriangleList,
 			/*BaseVertexIndex=*/ 0,
@@ -1711,14 +1702,14 @@ public:
 	/** The vertex declaration. */
 	FVertexDeclarationRHIRef VertexDeclarationRHI;
 
-	virtual void InitRHI()
+	virtual void InitRHI() override
 	{
 		FVertexDeclarationElementList Elements;
 		Elements.Add(FVertexElement(0, 0, VET_Float2, 0, sizeof(FVector2D)));
 		VertexDeclarationRHI = RHICreateVertexDeclaration( Elements );
 	}
 
-	virtual void ReleaseRHI()
+	virtual void ReleaseRHI() override
 	{
 		VertexDeclarationRHI.SafeRelease();
 	}
@@ -1732,6 +1723,7 @@ TGlobalResource<FParticleSimVisualizeVertexDeclaration> GParticleSimVisualizeVer
  * @param RenderTarget - The render target on which to draw the visualization.
  */
 static void VisualizeGPUSimulation(
+	FRHICommandListImmediate& RHICmdList,
 	int32 VisualizationMode,
 	FRenderTarget* RenderTarget,
 	const FParticleStateTextures& StateTextures,
@@ -1746,10 +1738,7 @@ static void VisualizeGPUSimulation(
 	const float DisplaySizeY = 256.0f;
 	const float DisplayOffsetX = 60.0f;
 	const float DisplayOffsetY = 60.0f;
-
-	//@todo-rco: RHIPacketList
-	FRHICommandList& RHICmdList = FRHICommandList::GetNullRef();
-
+	
 	// Setup render states.
 	FIntPoint TargetSize = RenderTarget->GetSizeXY();
 	SetRenderTarget(RHICmdList, RenderTarget->GetRenderTargetTexture(), FTextureRHIParamRef());
@@ -1763,7 +1752,7 @@ static void VisualizeGPUSimulation(
 	TShaderMapRef<FParticleSimVisualizePS> PixelShader( GetGlobalShaderMap() );
 
 	// Bound shader state.
-	RHICmdList.CheckIsNull(); // need new approach for "static FGlobalBoundShaderState" for parallel rendering
+	
 	static FGlobalBoundShaderState BoundShaderState;
 	SetGlobalBoundShaderState(
 		RHICmdList,
@@ -1932,9 +1921,7 @@ public:
 		FComputeShaderRHIParamRef ComputeShaderRHI = GetComputeShader();
 		if ( OutBounds.IsBound() )
 		{
-			RHICmdList.CheckIsNull();
-
-			RHISetUAVParameter( ComputeShaderRHI, OutBounds.GetBaseIndex(), OutBoundsUAV );
+			RHICmdList.SetUAVParameter(ComputeShaderRHI, OutBounds.GetBaseIndex(), OutBoundsUAV);
 		}
 	}
 
@@ -1942,28 +1929,26 @@ public:
 	 * Set input parameters.
 	 */
 	void SetParameters(
-		FRHICommandList& RHICmdList, 
+		FRHICommandList& RHICmdList,
 		FParticleBoundsUniformBufferRef& UniformBuffer,
 		FShaderResourceViewRHIParamRef InIndicesSRV,
 		FTexture2DRHIParamRef PositionTextureRHI,
 		FTexture2DRHIParamRef PositionZWTextureRHI
 		)
 	{
-		RHICmdList.CheckIsNull();
-
 		FComputeShaderRHIParamRef ComputeShaderRHI = GetComputeShader();
 		SetUniformBufferParameter(RHICmdList, ComputeShaderRHI, GetUniformBufferParameter<FParticleBoundsParameters>(), UniformBuffer );
 		if ( InParticleIndices.IsBound() )
 		{
-			RHISetShaderResourceViewParameter( ComputeShaderRHI, InParticleIndices.GetBaseIndex(), InIndicesSRV );
+			RHICmdList.SetShaderResourceViewParameter(ComputeShaderRHI, InParticleIndices.GetBaseIndex(), InIndicesSRV);
 		}
 		if ( PositionTexture.IsBound() )
 		{
-			RHISetShaderTexture( ComputeShaderRHI, PositionTexture.GetBaseIndex(), PositionTextureRHI );
+			RHICmdList.SetShaderTexture(ComputeShaderRHI, PositionTexture.GetBaseIndex(), PositionTextureRHI);
 		}
 		if ( PositionZWTexture.IsBound() )
 		{
-			RHISetShaderTexture( ComputeShaderRHI, PositionZWTexture.GetBaseIndex(), PositionZWTextureRHI );
+			RHICmdList.SetShaderTexture(ComputeShaderRHI, PositionZWTexture.GetBaseIndex(), PositionZWTextureRHI);
 		}
 	}
 
@@ -1972,16 +1957,14 @@ public:
 	 */
 	void UnbindBuffers(FRHICommandList& RHICmdList)
 	{
-		RHICmdList.CheckIsNull();
-
 		FComputeShaderRHIParamRef ComputeShaderRHI = GetComputeShader();
 		if ( InParticleIndices.IsBound() )
 		{
-			RHISetShaderResourceViewParameter( ComputeShaderRHI, InParticleIndices.GetBaseIndex(), FShaderResourceViewRHIParamRef() );
+			RHICmdList.SetShaderResourceViewParameter(ComputeShaderRHI, InParticleIndices.GetBaseIndex(), FShaderResourceViewRHIParamRef());
 		}
 		if ( OutBounds.IsBound() )
 		{
-			RHISetUAVParameter( ComputeShaderRHI, OutBounds.GetBaseIndex(), FUnorderedAccessViewRHIParamRef() );
+			RHICmdList.SetUAVParameter(ComputeShaderRHI, OutBounds.GetBaseIndex(), FUnorderedAccessViewRHIParamRef());
 		}
 	}
 
@@ -2015,7 +1998,7 @@ static bool AreBoundsValid( const FVector& Mins, const FVector& Maxs )
  * @param ParticleCount - The number of particles in the emitter.
  */
 static FBox ComputeParticleBounds(
-	FRHICommandList& RHICmdList, 
+	FRHICommandListImmediate& RHICmdList,
 	FShaderResourceViewRHIParamRef VertexBufferSRV,
 	FTexture2DRHIParamRef PositionTextureRHI,
 	FTexture2DRHIParamRef PositionZWTextureRHI,
@@ -2039,8 +2022,6 @@ static FBox ComputeParticleBounds(
 		Parameters.ParticleCount = ParticleCount;
 		UniformBuffer = FParticleBoundsUniformBufferRef::CreateUniformBufferImmediate( Parameters, UniformBuffer_SingleFrame );
 
-		RHICmdList.CheckIsNull();
-
 		// Create a buffer for storing bounds.
 		const int32 BufferSize = GroupCount * 2 * sizeof(FVector4);
 		FRHIResourceCreateInfo CreateInfo;
@@ -2054,7 +2035,7 @@ static FBox ComputeParticleBounds(
 
 		// Grab the shader.
 		TShaderMapRef<FParticleBoundsCS> ParticleBoundsCS( GetGlobalShaderMap() );
-		RHISetComputeShader(ParticleBoundsCS->GetComputeShader());
+		RHICmdList.SetComputeShader(ParticleBoundsCS->GetComputeShader());
 
 		// Dispatch shader to compute bounds.
 		ParticleBoundsCS->SetOutput(RHICmdList, BoundsVertexBufferUAV );
@@ -2102,7 +2083,7 @@ static FBox ComputeParticleBounds(
 		}
 
 		// Release buffer.
-		RHIUnlockVertexBuffer( BoundsVertexBufferRHI );
+		RHICmdList.UnlockVertexBuffer(BoundsVertexBufferRHI);
 		BoundsVertexBufferUAV.SafeRelease();
 		BoundsVertexBufferRHI.SafeRelease();
 	}
@@ -2455,7 +2436,7 @@ public:
 	/**
 	 * Initialize RHI resources.
 	 */
-	virtual void InitRHI()
+	virtual void InitRHI() override
 	{
 		UniformBuffer = FGPUSpriteEmitterUniformBufferRef::CreateUniformBufferImmediate( UniformParameters, UniformBuffer_MultiFrame );
 		EmitterSimulationResources.SimulationUniformBuffer =
@@ -2465,7 +2446,7 @@ public:
 	/**
 	 * Release RHI resources.
 	 */
-	virtual void ReleaseRHI()
+	virtual void ReleaseRHI() override
 	{
 		UniformBuffer.SafeRelease();
 		EmitterSimulationResources.SimulationUniformBuffer.SafeRelease();
@@ -3286,9 +3267,9 @@ public:
 			FComputeGPUSpriteBoundsCommand,
 			FGPUSpriteParticleEmitterInstance*, EmitterInstance, this,
 		{
-			//@todo-rco: RHIPacketList
+			
 			EmitterInstance->ParticleBoundingBox = ComputeParticleBounds(
-				FRHICommandList::GetNullRef(),
+				RHICmdList,
 				EmitterInstance->Simulation->VertexBuffer.VertexBufferSRV, 
 				EmitterInstance->FXSystem->GetParticleSimulationResources()->GetCurrentStateTextures().PositionTextureRHI,
 				EmitterInstance->FXSystem->GetParticleSimulationResources()->GetCurrentStateTextures().PositionZWTextureRHI,
@@ -3940,11 +3921,12 @@ void FFXSystem::ResetSortedGPUParticles()
 	ParticleSimulationResources->SortedParticleCount = 0;
 }
 
-void FFXSystem::SortGPUParticles()
+void FFXSystem::SortGPUParticles(FRHICommandListImmediate& RHICmdList)
 {
 	if (ParticleSimulationResources->SimulationsToSort.Num() > 0)
 	{
 		int32 BufferIndex = SortParticlesGPU(
+			RHICmdList,
 			GParticleSortBuffers,
 			ParticleSimulationResources->GetCurrentStateTextures().PositionTextureRHI,
 			ParticleSimulationResources->GetCurrentStateTextures().PositionZWTextureRHI,
@@ -3999,6 +3981,7 @@ static void SetParametersForVectorField(FVectorFieldUniformParameters& OutParame
 }
 
 void FFXSystem::SimulateGPUParticles(
+	FRHICommandListImmediate& RHICmdList,
 	EParticleSimulatePhase::Type Phase,
 	const class FSceneView* CollisionView,
 	FTexture2DRHIParamRef SceneDepthTexture,
@@ -4020,9 +4003,7 @@ void FFXSystem::SimulateGPUParticles(
 	FParticleStateTextures& CurrentStateTextures = ParticleSimulationResources->GetCurrentStateTextures();
 	FParticleStateTextures& PrevStateTextures = ParticleSimulationResources->GetPreviousStateTextures();
 
-	//@todo-rco: RHIPacketList
-	FRHICommandList& RHICmdList = FRHICommandList::GetNullRef();
-
+	
 	// On some platforms, the textures are filled with garbage after creation, so we need to clear them to black the first time we use them
 	if ( !CurrentStateTextures.bTexturesCleared )
 	{
@@ -4106,7 +4087,7 @@ void FFXSystem::SimulateGPUParticles(
 					const float Intensity = Simulation->LocalVectorField.Intensity * Simulation->LocalVectorField.Resource->Intensity;
 					if (FMath::Abs(Intensity) > 0.0f)
 					{
-						Simulation->LocalVectorField.Resource->Update(Simulation->PendingDeltaSeconds);
+						Simulation->LocalVectorField.Resource->Update(RHICmdList, Simulation->PendingDeltaSeconds);
 						SimulationCommand->VectorFieldTexturesRHI[0] = Simulation->LocalVectorField.Resource->VolumeTextureRHI;
 						SetParametersForVectorField(VectorFieldParameters, &Simulation->LocalVectorField, /*EmitterScale=*/ 1.0f, /*EmitterTightness=*/ -1, VectorFieldParameters.Count++);
 					}
@@ -4178,7 +4159,7 @@ void FFXSystem::SimulateGPUParticles(
 		{
 			/// ?
 			ExecuteSimulationCommands<true>(
-				FRHICommandList::GetNullRef(),
+				RHICmdList,
 				SimulationCommands,
 				PrevStateTextures,
 				ParticleSimulationResources->SimulationAttributesTexture,
@@ -4192,7 +4173,7 @@ void FFXSystem::SimulateGPUParticles(
 		{
 			/// ?
 			ExecuteSimulationCommands<false>(
-				FRHICommandList::GetNullRef(),
+				RHICmdList,
 				SimulationCommands,
 				PrevStateTextures,
 				ParticleSimulationResources->SimulationAttributesTexture,
@@ -4207,7 +4188,7 @@ void FFXSystem::SimulateGPUParticles(
 	// Clear any newly allocated tiles.
 	if (TilesToClear.Num())
 	{
-		ClearTiles(TilesToClear);
+		ClearTiles(RHICmdList, TilesToClear);
 	}
 
 	// Inject any new particles that have spawned into the simulation.
@@ -4227,7 +4208,7 @@ void FFXSystem::SimulateGPUParticles(
 		RHICmdList.SetViewport(0, 0, 0.0f, GParticleSimulationTextureSizeX, GParticleSimulationTextureSizeY, 1.0f);
 
 		// Inject particles.
-		InjectNewParticles(NewParticles);
+		InjectNewParticles(RHICmdList, NewParticles);
 
 		// Resolve attributes textures. State textures are resolved later.
 		RHICmdList.CopyToResolveTarget(
@@ -4272,7 +4253,7 @@ void FFXSystem::VisualizeGPUParticles(FCanvas* Canvas)
 	{
 		FParticleSimulationResources* Resources = FXSystem->GetParticleSimulationResources();
 		FParticleStateTextures& CurrentStateTextures = Resources->GetCurrentStateTextures();
-		VisualizeGPUSimulation(VisualizationMode, RenderTarget, CurrentStateTextures, GParticleCurveTexture.GetCurveTexture());
+		VisualizeGPUSimulation(RHICmdList, VisualizationMode, RenderTarget, CurrentStateTextures, GParticleCurveTexture.GetCurveTexture());
 	});
 }
 
diff --git a/Engine/Source/Runtime/Engine/Private/Particles/ParticleSortingGPU.cpp b/Engine/Source/Runtime/Engine/Private/Particles/ParticleSortingGPU.cpp
--- a/Engine/Source/Runtime/Engine/Private/Particles/ParticleSortingGPU.cpp
+++ b/Engine/Source/Runtime/Engine/Private/Particles/ParticleSortingGPU.cpp
@@ -96,16 +96,14 @@ public:
 	 */
 	void SetOutput(FRHICommandList& RHICmdList, FUnorderedAccessViewRHIParamRef OutKeysUAV, FUnorderedAccessViewRHIParamRef OutIndicesUAV )
 	{
-		RHICmdList.CheckIsNull();
-
 		FComputeShaderRHIParamRef ComputeShaderRHI = GetComputeShader();
 		if ( OutKeys.IsBound() )
 		{
-			RHISetUAVParameter( ComputeShaderRHI, OutKeys.GetBaseIndex(), OutKeysUAV );
+			RHICmdList.SetUAVParameter(ComputeShaderRHI, OutKeys.GetBaseIndex(), OutKeysUAV);
 		}
 		if ( OutParticleIndices.IsBound() )
 		{
-			RHISetUAVParameter( ComputeShaderRHI, OutParticleIndices.GetBaseIndex(), OutIndicesUAV );
+			RHICmdList.SetUAVParameter(ComputeShaderRHI, OutParticleIndices.GetBaseIndex(), OutIndicesUAV);
 		}
 	}
 
@@ -122,9 +120,7 @@ public:
 		SetUniformBufferParameter(RHICmdList, ComputeShaderRHI, GetUniformBufferParameter<FParticleKeyGenParameters>(), UniformBuffer );
 		if ( InParticleIndices.IsBound() )
 		{
-			RHICmdList.CheckIsNull();
-
-			RHISetShaderResourceViewParameter( ComputeShaderRHI, InParticleIndices.GetBaseIndex(), InIndicesSRV );
+			RHICmdList.SetShaderResourceViewParameter(ComputeShaderRHI, InParticleIndices.GetBaseIndex(), InIndicesSRV);
 		}
 	}
 
@@ -136,15 +132,11 @@ public:
 		FComputeShaderRHIParamRef ComputeShaderRHI = GetComputeShader();
 		if (PositionTexture.IsBound())
 		{
-			RHICmdList.CheckIsNull();
-
-			RHISetShaderTexture(ComputeShaderRHI, PositionTexture.GetBaseIndex(), PositionTextureRHI);
+			RHICmdList.SetShaderTexture(ComputeShaderRHI, PositionTexture.GetBaseIndex(), PositionTextureRHI);
 		}
 		if (PositionZWTexture.IsBound())
 		{
-			RHICmdList.CheckIsNull();
-
-			RHISetShaderTexture(ComputeShaderRHI, PositionZWTexture.GetBaseIndex(), PositionZWTextureRHI);
+			RHICmdList.SetShaderTexture(ComputeShaderRHI, PositionZWTexture.GetBaseIndex(), PositionZWTextureRHI);
 		}
 	}
 
@@ -153,20 +145,18 @@ public:
 	 */
 	void UnbindBuffers(FRHICommandList& RHICmdList)
 	{
-		RHICmdList.CheckIsNull();
-
 		FComputeShaderRHIParamRef ComputeShaderRHI = GetComputeShader();
 		if ( InParticleIndices.IsBound() )
 		{
-			RHISetShaderResourceViewParameter( ComputeShaderRHI, InParticleIndices.GetBaseIndex(), FShaderResourceViewRHIParamRef() );
+			RHICmdList.SetShaderResourceViewParameter(ComputeShaderRHI, InParticleIndices.GetBaseIndex(), FShaderResourceViewRHIParamRef());
 		}
 		if ( OutKeys.IsBound() )
 		{
-			RHISetUAVParameter( ComputeShaderRHI, OutKeys.GetBaseIndex(), FUnorderedAccessViewRHIParamRef() );
+			RHICmdList.SetUAVParameter(ComputeShaderRHI, OutKeys.GetBaseIndex(), FUnorderedAccessViewRHIParamRef());
 		}
 		if ( OutParticleIndices.IsBound() )
 		{
-			RHISetUAVParameter( ComputeShaderRHI, OutParticleIndices.GetBaseIndex(), FUnorderedAccessViewRHIParamRef() );
+			RHICmdList.SetUAVParameter(ComputeShaderRHI, OutParticleIndices.GetBaseIndex(), FUnorderedAccessViewRHIParamRef());
 		}
 	}
 
@@ -195,7 +185,7 @@ IMPLEMENT_SHADER_TYPE(,FParticleSortKeyGenCS,TEXT("ParticleSortKeyGen"),TEXT("Ge
  * @returns the total number of particles being sorted.
  */
 static int32 GenerateParticleSortKeys(
-	FRHICommandList& RHICmdList,
+	FRHICommandListImmediate& RHICmdList,
 	FUnorderedAccessViewRHIParamRef KeyBufferUAV,
 	FUnorderedAccessViewRHIParamRef SortedVertexBufferUAV,
 	FTexture2DRHIParamRef PositionTextureRHI,
@@ -214,7 +204,7 @@ static int32 GenerateParticleSortKeys(
 
 	// Grab the shader, set output.
 	TShaderMapRef<FParticleSortKeyGenCS> KeyGenCS(GetGlobalShaderMap());
-	RHISetComputeShader(KeyGenCS->GetComputeShader());
+	RHICmdList.SetComputeShader(KeyGenCS->GetComputeShader());
 	KeyGenCS->SetOutput(RHICmdList, KeyBufferUAV, SortedVertexBufferUAV);
 	KeyGenCS->SetPositionTextures(RHICmdList, PositionTextureRHI, PositionZWTextureRHI);
 
@@ -328,6 +318,7 @@ FGPUSortBuffers FParticleSortBuffers::GetSortBuffers()
  * @returns the buffer index in which sorting results are stored.
  */
 int32 SortParticlesGPU(
+	FRHICommandListImmediate& RHICmdList,
 	FParticleSortBuffers& ParticleSortBuffers,
 	FTexture2DRHIParamRef PositionTextureRHI,
 	FTexture2DRHIParamRef PositionZWTextureRHI,
@@ -344,14 +335,14 @@ int32 SortParticlesGPU(
 		const int32 StreamCount = 2;
 		for (int32 StreamIndex = 0; StreamIndex < StreamCount; ++StreamIndex)
 		{
-			RHISetStreamSource(StreamIndex, FVertexBufferRHIParamRef(), 0, 0);
+			RHICmdList.SetStreamSource(StreamIndex, FVertexBufferRHIParamRef(), 0, 0);
 		}
 	}
 
 	// First generate keys for each emitter to be sorted.
-	//@todo-rco: RHIPacketList
+	
 	const int32 TotalParticleCount = GenerateParticleSortKeys(
-		FRHICommandList::GetNullRef(),
+		RHICmdList,
 		ParticleSortBuffers.GetKeyBufferUAV(),
 		ParticleSortBuffers.GetVertexBufferUAV(),
 		PositionTextureRHI,
@@ -367,5 +358,5 @@ int32 SortParticlesGPU(
 	const uint32 EmitterKeyMask = (1 << FMath::CeilLogTwo( SimulationsToSort.Num() )) - 1;
 	const uint32 KeyMask = (EmitterKeyMask << 16) | 0xFFFF;
 	FGPUSortBuffers SortBuffers = ParticleSortBuffers.GetSortBuffers();
-	return SortGPUBuffers( SortBuffers, 0, KeyMask, TotalParticleCount );
+	return SortGPUBuffers(RHICmdList, SortBuffers, 0, KeyMask, TotalParticleCount);
 }
diff --git a/Engine/Source/Runtime/Engine/Private/Particles/ParticleSortingGPU.h b/Engine/Source/Runtime/Engine/Private/Particles/ParticleSortingGPU.h
--- a/Engine/Source/Runtime/Engine/Private/Particles/ParticleSortingGPU.h
+++ b/Engine/Source/Runtime/Engine/Private/Particles/ParticleSortingGPU.h
@@ -122,6 +122,7 @@ struct FParticleSimulationSortInfo
  * @returns the buffer index in which sorting results are stored.
  */
 int32 SortParticlesGPU(
+	FRHICommandListImmediate& RHICmdList,
 	FParticleSortBuffers& ParticleSortBuffers,
 	FTexture2DRHIParamRef PositionTextureRHI,
 	FTexture2DRHIParamRef PositionZWTextureRHI,
diff --git a/Engine/Source/Runtime/Engine/Private/PointLightSceneProxy.h b/Engine/Source/Runtime/Engine/Private/PointLightSceneProxy.h
--- a/Engine/Source/Runtime/Engine/Private/PointLightSceneProxy.h
+++ b/Engine/Source/Runtime/Engine/Private/PointLightSceneProxy.h
@@ -148,7 +148,7 @@ public:
 		return true;
 	}
 
-	virtual void SetScissorRect(const FSceneView& View) const
+	virtual void SetScissorRect(FRHICommandListImmediate& RHICmdList, const FSceneView& View) const
 	{
 		// Calculate a scissor rectangle for the light's radius.
 		if((GetLightToWorld().GetOrigin() - View.ViewMatrices.ViewOrigin).Size() > Radius)
@@ -189,11 +189,11 @@ public:
 				return;
 			}
 
-			RHISetScissorRect(true,ScissorMinX,ScissorMinY,ScissorMaxX,ScissorMaxY);
+			RHICmdList.SetScissorRect(true, ScissorMinX, ScissorMinY, ScissorMaxX, ScissorMaxY);
 		}
 		else
 		{
-			RHISetScissorRect(false,0,0,0,0);
+			RHICmdList.SetScissorRect(false, 0, 0, 0, 0);
 		}
 	}
 
diff --git a/Engine/Source/Runtime/Engine/Private/ReflectionCaptureComponent.cpp b/Engine/Source/Runtime/Engine/Private/ReflectionCaptureComponent.cpp
--- a/Engine/Source/Runtime/Engine/Private/ReflectionCaptureComponent.cpp
+++ b/Engine/Source/Runtime/Engine/Private/ReflectionCaptureComponent.cpp
@@ -619,7 +619,7 @@ public:
 		SourceData = InSourceData;
 	}
 
-	virtual void InitRHI()
+	virtual void InitRHI() override
 	{
 		FRHIResourceCreateInfo CreateInfo;
 		TextureCubeRHI = RHICreateTextureCube(Size, Format, NumMips, 0, CreateInfo);
@@ -1162,7 +1162,7 @@ void UReflectionCaptureComponent::SetCaptureIsDirty()
 	}
 }
 
-void ReadbackFromSM4Cubemap_RenderingThread(FReflectionTextureCubeResource* SM4FullHDRCubemapTexture, FReflectionCaptureFullHDRDerivedData* OutDerivedData)
+void ReadbackFromSM4Cubemap_RenderingThread(FRHICommandListImmediate& RHICmdList, FReflectionTextureCubeResource* SM4FullHDRCubemapTexture, FReflectionCaptureFullHDRDerivedData* OutDerivedData)
 {
 	const int32 EffectiveTopMipSize = GReflectionCaptureSize;
 	const int32 NumMips = FMath::CeilLogTwo(EffectiveTopMipSize) + 1;
@@ -1184,6 +1184,8 @@ void ReadbackFromSM4Cubemap_RenderingThread(FReflectionTextureCubeResource* SM4F
 	CaptureData.AddZeroed(CaptureDataSize);
 	int32 MipBaseIndex = 0;
 
+
+
 	for (int32 MipIndex = 0; MipIndex < NumMips; MipIndex++)
 	{
 		check(SM4FullHDRCubemapTexture->GetTextureRHI()->GetFormat() == PF_FloatRGBA);
@@ -1196,7 +1198,7 @@ void ReadbackFromSM4Cubemap_RenderingThread(FReflectionTextureCubeResource* SM4F
 			// Read each mip face
 			//@todo - do this without blocking the GPU so many times
 			//@todo - pool the temporary textures in RHIReadSurfaceFloatData instead of always creating new ones
-			RHIReadSurfaceFloatData(SM4FullHDRCubemapTexture->GetTextureRHI(), FIntRect(0, 0, MipSize, MipSize), SurfaceData, (ECubeFace)CubeFace, 0, MipIndex);
+			RHICmdList.ReadSurfaceFloatData(SM4FullHDRCubemapTexture->GetTextureRHI(), FIntRect(0, 0, MipSize, MipSize), SurfaceData, (ECubeFace)CubeFace, 0, MipIndex);
 			const int32 DestIndex = MipBaseIndex + CubeFace * CubeFaceBytes;
 			uint8* FaceData = &CaptureData[DestIndex];
 			check(SurfaceData.Num() * SurfaceData.GetTypeSize() == CubeFaceBytes);
@@ -1216,7 +1218,7 @@ void ReadbackFromSM4Cubemap(FReflectionTextureCubeResource* SM4FullHDRCubemapTex
 		FReflectionTextureCubeResource*, SM4FullHDRCubemapTexture, SM4FullHDRCubemapTexture,
 		FReflectionCaptureFullHDRDerivedData*, OutDerivedData, OutDerivedData,
 	{
-		ReadbackFromSM4Cubemap_RenderingThread(SM4FullHDRCubemapTexture, OutDerivedData);
+		ReadbackFromSM4Cubemap_RenderingThread(RHICmdList, SM4FullHDRCubemapTexture, OutDerivedData);
 	});
 
 	FlushRenderingCommands();
diff --git a/Engine/Source/Runtime/Engine/Private/SkeletalRenderCPUSkin.cpp b/Engine/Source/Runtime/Engine/Private/SkeletalRenderCPUSkin.cpp
--- a/Engine/Source/Runtime/Engine/Private/SkeletalRenderCPUSkin.cpp
+++ b/Engine/Source/Runtime/Engine/Private/SkeletalRenderCPUSkin.cpp
@@ -179,7 +179,7 @@ void FSkeletalMeshObjectCPUSkin::Update(int32 LODIndex,USkinnedMeshComponent* In
 		FDynamicSkelMeshObjectData*, NewDynamicData, NewDynamicData,
 	{
 		FScopeCycleCounter Context(MeshObject->GetStatId());
-		MeshObject->UpdateDynamicData_RenderThread(NewDynamicData);
+		MeshObject->UpdateDynamicData_RenderThread(RHICmdList, NewDynamicData);
 	}
 	);
 
@@ -191,7 +191,7 @@ void FSkeletalMeshObjectCPUSkin::Update(int32 LODIndex,USkinnedMeshComponent* In
 	}
 }
 
-void FSkeletalMeshObjectCPUSkin::UpdateDynamicData_RenderThread(FDynamicSkelMeshObjectData* InDynamicData)
+void FSkeletalMeshObjectCPUSkin::UpdateDynamicData_RenderThread(FRHICommandListImmediate& RHICmdList, FDynamicSkelMeshObjectData* InDynamicData)
 {
 	// we should be done with the old data at this point
 	delete DynamicData;
diff --git a/Engine/Source/Runtime/Engine/Private/SkeletalRenderCPUSkin.h b/Engine/Source/Runtime/Engine/Private/SkeletalRenderCPUSkin.h
--- a/Engine/Source/Runtime/Engine/Private/SkeletalRenderCPUSkin.h
+++ b/Engine/Source/Runtime/Engine/Private/SkeletalRenderCPUSkin.h
@@ -142,7 +142,7 @@ public:
 	virtual void InitResources() override;
 	virtual void ReleaseResources() override;
 	virtual void Update(int32 LODIndex,USkinnedMeshComponent* InMeshComponent,const TArray<FActiveVertexAnim>& ActiveVertexAnims) override;
-	virtual void UpdateDynamicData_RenderThread(FDynamicSkelMeshObjectData* InDynamicData) override;
+	virtual void UpdateDynamicData_RenderThread(FRHICommandListImmediate& RHICmdList, FDynamicSkelMeshObjectData* InDynamicData) override;
 	virtual void EnableBlendWeightRendering(bool bEnabled, const TArray<int32>& InBonesOfInterest) override;
 	virtual void CacheVertices(int32 LODIndex, bool bForce) const override;
 	virtual bool IsCPUSkinned() const override { return true; }
diff --git a/Engine/Source/Runtime/Engine/Private/SkeletalRenderGPUSkin.cpp b/Engine/Source/Runtime/Engine/Private/SkeletalRenderGPUSkin.cpp
--- a/Engine/Source/Runtime/Engine/Private/SkeletalRenderGPUSkin.cpp
+++ b/Engine/Source/Runtime/Engine/Private/SkeletalRenderGPUSkin.cpp
@@ -179,7 +179,7 @@ void FSkeletalMeshObjectGPUSkin::Update(int32 LODIndex,USkinnedMeshComponent* In
 		FDynamicSkelMeshObjectData*, NewDynamicData, NewDynamicData,
 	{
 		FScopeCycleCounter Context(MeshObject->GetStatId());
-		MeshObject->UpdateDynamicData_RenderThread(NewDynamicData);
+		MeshObject->UpdateDynamicData_RenderThread(RHICmdList, NewDynamicData);
 	}
 	);
 
@@ -191,7 +191,7 @@ void FSkeletalMeshObjectGPUSkin::Update(int32 LODIndex,USkinnedMeshComponent* In
 	}
 }
 
-void FSkeletalMeshObjectGPUSkin::UpdateDynamicData_RenderThread(FDynamicSkelMeshObjectData* InDynamicData)
+void FSkeletalMeshObjectGPUSkin::UpdateDynamicData_RenderThread(FRHICommandListImmediate& RHICmdList, FDynamicSkelMeshObjectData* InDynamicData)
 {
 	SCOPE_CYCLE_COUNTER(STAT_GPUSkinUpdateRTTime);
 	bool bMorphNeedsUpdate=false;
@@ -287,7 +287,7 @@ void FSkeletalMeshObjectGPUSkin::UpdateDynamicData_RenderThread(FDynamicSkelMesh
 			// Try to use the GPU skinning cache if possible
 			if (bGPUSkinCacheEnabled && ChunkIdx < MAX_GPUSKINCACHE_CHUNKS_PER_LOD && !bClothFactory && Chunk.MaxBoneInfluences > 0 && DynamicData->NumWeightedActiveVertexAnims <= 0)
 			{
-				int32 Key = GGPUSkinCache.StartCacheMesh(GPUSkinCacheKeys[ChunkIdx], &VertexFactoryData.VertexFactories[ChunkIdx], &VertexFactoryData.PassthroughVertexFactories[ChunkIdx], Chunk, this, Chunk.HasExtraBoneInfluences());
+				int32 Key = GGPUSkinCache.StartCacheMesh(RHICmdList, GPUSkinCacheKeys[ChunkIdx], &VertexFactoryData.VertexFactories[ChunkIdx], &VertexFactoryData.PassthroughVertexFactories[ChunkIdx], Chunk, this, Chunk.HasExtraBoneInfluences());
 				if(Key >= 0)
 				{
 					GPUSkinCacheKeys[ChunkIdx] = (int16)Key;
diff --git a/Engine/Source/Runtime/Engine/Private/SkeletalRenderGPUSkin.h b/Engine/Source/Runtime/Engine/Private/SkeletalRenderGPUSkin.h
--- a/Engine/Source/Runtime/Engine/Private/SkeletalRenderGPUSkin.h
+++ b/Engine/Source/Runtime/Engine/Private/SkeletalRenderGPUSkin.h
@@ -248,7 +248,7 @@ public:
 	virtual void InitResources() override;
 	virtual void ReleaseResources() override;
 	virtual void Update(int32 LODIndex,USkinnedMeshComponent* InMeshComponent,const TArray<FActiveVertexAnim>& ActiveVertexAnims) override;
-	virtual void UpdateDynamicData_RenderThread(FDynamicSkelMeshObjectData* InDynamicData) override;
+	virtual void UpdateDynamicData_RenderThread(FRHICommandListImmediate& RHICmdList, FDynamicSkelMeshObjectData* InDynamicData) override;
 	virtual const FVertexFactory* GetVertexFactory(int32 LODIndex,int32 ChunkIdx) const override;
 	virtual void CacheVertices(int32 LODIndex, bool bForce) const override {}
 	virtual bool IsCPUSkinned() const override { return false; }
diff --git a/Engine/Source/Runtime/Engine/Private/Slate/DebugCanvas.cpp b/Engine/Source/Runtime/Engine/Private/Slate/DebugCanvas.cpp
--- a/Engine/Source/Runtime/Engine/Private/Slate/DebugCanvas.cpp
+++ b/Engine/Source/Runtime/Engine/Private/Slate/DebugCanvas.cpp
@@ -163,7 +163,7 @@ void FDebugCanvasDrawer::InitDebugCanvas(UWorld* InWorld)
 	}
 }
 
-void FDebugCanvasDrawer::DrawRenderThread( const void* InWindowBackBuffer )
+void FDebugCanvasDrawer::DrawRenderThread(FRHICommandListImmediate& RHICmdList, const void* InWindowBackBuffer)
 {
 	check( IsInRenderingThread() );
 
diff --git a/Engine/Source/Runtime/Engine/Private/Slate/DebugCanvas.h b/Engine/Source/Runtime/Engine/Private/Slate/DebugCanvas.h
--- a/Engine/Source/Runtime/Engine/Private/Slate/DebugCanvas.h
+++ b/Engine/Source/Runtime/Engine/Private/Slate/DebugCanvas.h
@@ -30,7 +30,7 @@ private:
 	/**
 	 * ICustomSlateElement interface 
 	 */
-	virtual void DrawRenderThread( const void* InWindowBackBuffer ) override;
+	virtual void DrawRenderThread(FRHICommandListImmediate& RHICmdList, const void* InWindowBackBuffer) override;
 
 	/**
 	 * Deletes the rendering thread canvas 
diff --git a/Engine/Source/Runtime/Engine/Private/Slate/SceneViewport.cpp b/Engine/Source/Runtime/Engine/Private/Slate/SceneViewport.cpp
--- a/Engine/Source/Runtime/Engine/Private/Slate/SceneViewport.cpp
+++ b/Engine/Source/Runtime/Engine/Private/Slate/SceneViewport.cpp
@@ -1012,42 +1012,41 @@ void FSceneViewport::EnqueueBeginRenderFrame()
 	FViewport::EnqueueBeginRenderFrame();
 }
 
-void FSceneViewport::BeginRenderFrame()
+void FSceneViewport::BeginRenderFrame(FRHICommandListImmediate& RHICmdList)
 {
 	check( IsInRenderingThread() );
-	
 	if( bUseSeparateRenderTarget )
 	{
-		RHISetRenderTarget( RenderTargetTextureRHI,  FTexture2DRHIRef() );
+		SetRenderTarget(RHICmdList,  RenderTargetTextureRHI,  FTexture2DRHIRef() );
 	}
 	else if( IsValidRef( ViewportRHI ) ) 
 	{
 		// Get the backbuffer render target to render directly to it
-		RenderTargetTextureRHI = RHIGetViewportBackBuffer( ViewportRHI );
+		RenderTargetTextureRHI = RHICmdList.GetViewportBackBuffer(ViewportRHI);
 		if (GRHIShaderPlatform != SP_METAL)
 		{
 			// unused set render targets are bad on Metal
-			RHISetRenderTarget(RenderTargetTextureRHI, FTexture2DRHIRef());
+			SetRenderTarget(RHICmdList, RenderTargetTextureRHI, FTexture2DRHIRef());
 		}
 	}
 }
 
-void FSceneViewport::EndRenderFrame( bool bPresent, bool bLockToVsync )
+void FSceneViewport::EndRenderFrame(FRHICommandListImmediate& RHICmdList, bool bPresent, bool bLockToVsync)
 {
 	check( IsInRenderingThread() );
-	if( bUseSeparateRenderTarget )
+	if (bUseSeparateRenderTarget)
 	{
 		// @todo-mobile
 		if (GRHIShaderPlatform == SP_OPENGL_ES2)
 		{
 			check(0);
 		}
-		RHICopyToResolveTarget( RenderTargetTextureRHI, SlateRenderTargetHandle->GetRHIRef(), false, FResolveParams() ); 
+		RHICmdList.CopyToResolveTarget( RenderTargetTextureRHI, SlateRenderTargetHandle->GetRHIRef(), false, FResolveParams() ); 
 	}
 	else
 	{
 		// Set the active render target(s) to nothing to release references in the case that the viewport is resized by slate before we draw again
-		RHISetRenderTarget( FTexture2DRHIRef(), FTexture2DRHIRef() );
+		SetRenderTarget(RHICmdList,  FTexture2DRHIRef(), FTexture2DRHIRef() );
 		// Note: this releases our reference but does not release the resource as it is owned by slate (this is intended)
 		RenderTargetTextureRHI.SafeRelease();
 	}
diff --git a/Engine/Source/Runtime/Engine/Private/Slate/SlateTextures.cpp b/Engine/Source/Runtime/Engine/Private/Slate/SlateTextures.cpp
--- a/Engine/Source/Runtime/Engine/Private/Slate/SlateTextures.cpp
+++ b/Engine/Source/Runtime/Engine/Private/Slate/SlateTextures.cpp
@@ -218,14 +218,15 @@ void FSlateTextureRenderTarget2DResource::ReleaseDynamicRHI()
 void FSlateTextureRenderTarget2DResource::UpdateResource()
 {
 	check(IsInRenderingThread());
+	FRHICommandListImmediate& RHICmdList = FRHICommandListExecutor::GetImmediateCommandList();
 
 	// clear the target surface to green
-	RHISetRenderTarget(RenderTargetTextureRHI,FTextureRHIRef());
-	RHISetViewport(0,0,0.0f,TargetSizeX,TargetSizeY,1.0f);
-	RHIClear(true,ClearColor,false,0.f,false,0, FIntRect());
+	SetRenderTarget(RHICmdList, RenderTargetTextureRHI,FTextureRHIRef());
+	RHICmdList.SetViewport(0,0,0.0f,TargetSizeX,TargetSizeY,1.0f);
+	RHICmdList.Clear(true,ClearColor,false,0.f,false,0, FIntRect());
 
 	// copy surface to the texture for use
-	RHICopyToResolveTarget(RenderTargetTextureRHI, TextureRHI, true, FResolveParams());
+	RHICmdList.CopyToResolveTarget(RenderTargetTextureRHI, TextureRHI, true, FResolveParams());
 }
 
 FIntPoint FSlateTextureRenderTarget2DResource::GetSizeXY() const
diff --git a/Engine/Source/Runtime/Engine/Private/SplineMeshSceneProxy.h b/Engine/Source/Runtime/Engine/Private/SplineMeshSceneProxy.h
--- a/Engine/Source/Runtime/Engine/Private/SplineMeshSceneProxy.h
+++ b/Engine/Source/Runtime/Engine/Private/SplineMeshSceneProxy.h
@@ -49,7 +49,7 @@ public:
 	}
 
 	// FRenderResource interface.
-	virtual void InitRHI()
+	virtual void InitRHI() override
 	{
 		FLocalVertexFactory::InitRHI();
 	}
diff --git a/Engine/Source/Runtime/Engine/Private/StaticBoundShaderState.cpp b/Engine/Source/Runtime/Engine/Private/StaticBoundShaderState.cpp
--- a/Engine/Source/Runtime/Engine/Private/StaticBoundShaderState.cpp
+++ b/Engine/Source/Runtime/Engine/Private/StaticBoundShaderState.cpp
@@ -41,7 +41,7 @@ FGlobalBoundShaderStateResource::~FGlobalBoundShaderStateResource()
  * Initializes a global bound shader state with a vanilla bound shader state and required information.
  */
 FBoundShaderStateRHIParamRef FGlobalBoundShaderStateResource::GetInitializedRHI(
-	FVertexDeclarationRHIParamRef VertexDeclaration, 
+	FVertexDeclarationRHIParamRef VertexDeclaration,
 	FVertexShaderRHIParamRef VertexShader, 
 	FPixelShaderRHIParamRef PixelShader,
 	FGeometryShaderRHIParamRef GeometryShader
@@ -93,7 +93,7 @@ void FGlobalBoundShaderStateResource::ReleaseRHI()
 }
 
 void SetGlobalBoundShaderState(
-	FRHICommandList& RHICmdList,
+	FRHICommandListImmediate& RHICmdList,
 	FGlobalBoundShaderState& GlobalBoundShaderState,
 	FVertexDeclarationRHIParamRef VertexDeclaration,
 	FShader* VertexShader,
diff --git a/Engine/Source/Runtime/Engine/Private/TextRenderComponent.cpp b/Engine/Source/Runtime/Engine/Private/TextRenderComponent.cpp
--- a/Engine/Source/Runtime/Engine/Private/TextRenderComponent.cpp
+++ b/Engine/Source/Runtime/Engine/Private/TextRenderComponent.cpp
@@ -103,7 +103,7 @@ class FTextRenderVertexBuffer : public FVertexBuffer
 public:
 	TArray<FDynamicMeshVertex> Vertices;
 
-	void InitRHI()
+	virtual void InitRHI() override
 	{
 		FRHIResourceCreateInfo CreateInfo;
 		VertexBufferRHI = RHICreateVertexBuffer(Vertices.Num() * sizeof(FDynamicMeshVertex),BUF_Static,CreateInfo);
diff --git a/Engine/Source/Runtime/Engine/Private/Texture.cpp b/Engine/Source/Runtime/Engine/Private/Texture.cpp
--- a/Engine/Source/Runtime/Engine/Private/Texture.cpp
+++ b/Engine/Source/Runtime/Engine/Private/Texture.cpp
@@ -61,7 +61,7 @@ UTexture::UTexture(const class FPostConstructInitializeProperties& PCIP)
 			FCreateTextureRefCommand,
 			UTexture*,Texture,this,
 		{
-			FTextureReferenceRHIRef TextureRef = RHICreateTextureReference();
+			FTextureReferenceRHIRef TextureRef = RHICmdList.CreateTextureReference();
 			TextureRef->AddRef();
 			Texture->TextureRefRHI = TextureRef;
 		});
diff --git a/Engine/Source/Runtime/Engine/Private/TextureCube.cpp b/Engine/Source/Runtime/Engine/Private/TextureCube.cpp
--- a/Engine/Source/Runtime/Engine/Private/TextureCube.cpp
+++ b/Engine/Source/Runtime/Engine/Private/TextureCube.cpp
@@ -217,7 +217,7 @@ public:
 	/**
 	 * Called when the resource is initialized. This is only called by the rendering thread.
 	 */
-	virtual void InitRHI()
+	virtual void InitRHI() override
 	{
 		INC_DWORD_STAT_BY( STAT_TextureMemory, TextureSize );
 		INC_DWORD_STAT_FNAME_BY( LODGroupStatName, TextureSize );
@@ -261,7 +261,7 @@ public:
 		bGreyScaleFormat = (PixelFormat == PF_G8) || (PixelFormat == PF_BC4);
 	}
 
-	virtual void ReleaseRHI()
+	virtual void ReleaseRHI() override
 	{
 		DEC_DWORD_STAT_BY( STAT_TextureMemory, TextureSize );
 		DEC_DWORD_STAT_FNAME_BY( LODGroupStatName, TextureSize );
diff --git a/Engine/Source/Runtime/Engine/Private/TextureMovie.cpp b/Engine/Source/Runtime/Engine/Private/TextureMovie.cpp
--- a/Engine/Source/Runtime/Engine/Private/TextureMovie.cpp
+++ b/Engine/Source/Runtime/Engine/Private/TextureMovie.cpp
@@ -418,7 +418,8 @@ void FTextureMovieResource::UpdateResource()
 	if( Owner->Decoder )
 	{
 		SCOPED_DRAW_EVENTF(EventDecode, DEC_SCENE_ITEMS, TEXT("Movie[%s]"),*Owner->GetName());
-		Owner->Decoder->GetFrame(this);
+		FRHICommandListImmediate& RHICmdList = FRHICommandListExecutor::GetImmediateCommandList();
+		Owner->Decoder->GetFrame(RHICmdList, this);
 	}
 }
 
diff --git a/Engine/Source/Runtime/Engine/Private/TextureRenderTarget2D.cpp b/Engine/Source/Runtime/Engine/Private/TextureRenderTarget2D.cpp
--- a/Engine/Source/Runtime/Engine/Private/TextureRenderTarget2D.cpp
+++ b/Engine/Source/Runtime/Engine/Private/TextureRenderTarget2D.cpp
@@ -366,14 +366,15 @@ void FTextureRenderTarget2DResource::ReleaseDynamicRHI()
 void FTextureRenderTarget2DResource::UpdateResource()
 {
 	RemoveFromDeferredUpdateList();
+	FRHICommandListImmediate& RHICmdList = FRHICommandListExecutor::GetImmediateCommandList();
 
  	// clear the target surface to green
- 	RHISetRenderTarget(RenderTargetTextureRHI,FTextureRHIRef());
- 	RHISetViewport(0,0,0.0f,TargetSizeX,TargetSizeY,1.0f);
- 	RHIClear(true,ClearColor,false,0.f,false,0, FIntRect());
+	SetRenderTarget(RHICmdList, RenderTargetTextureRHI, FTextureRHIRef());
+	RHICmdList.SetViewport(0, 0, 0.0f, TargetSizeX, TargetSizeY, 1.0f);
+	RHICmdList.Clear(true, ClearColor, false, 0.f, false, 0, FIntRect());
  
  	// copy surface to the texture for use
- 	RHICopyToResolveTarget(RenderTargetTextureRHI, TextureRHI, true, FResolveParams());
+	RHICmdList.CopyToResolveTarget(RenderTargetTextureRHI, TextureRHI, true, FResolveParams());
 }
 
 /** 
diff --git a/Engine/Source/Runtime/Engine/Private/TextureRenderTargetCube.cpp b/Engine/Source/Runtime/Engine/Private/TextureRenderTargetCube.cpp
--- a/Engine/Source/Runtime/Engine/Private/TextureRenderTargetCube.cpp
+++ b/Engine/Source/Runtime/Engine/Private/TextureRenderTargetCube.cpp
@@ -281,18 +281,20 @@ void FTextureRenderTargetCubeResource::ReleaseDynamicRHI()
  */
 void FTextureRenderTargetCubeResource::UpdateResource()
 {
+	FRHICommandListImmediate& RHICmdList = FRHICommandListExecutor::GetImmediateCommandList();
+
 	const FIntPoint Dims = GetSizeXY();
 	for(int32 FaceIdx = CubeFace_PosX; FaceIdx < CubeFace_MAX; FaceIdx++)
 	{
 		// clear each face of the cube target texture to ClearColor
-		RHISetRenderTarget(RenderTargetTextureRHI, FTextureRHIParamRef());
+		SetRenderTarget(RHICmdList, RenderTargetTextureRHI, FTextureRHIParamRef());
 
-		RHISetViewport(0, 0, 0.0f, Dims.X, Dims.Y, 1.0f);
-		RHIClear(true, Owner->ClearColor, false, 0.f, false, 0, FIntRect());
+		RHICmdList.SetViewport(0, 0, 0.0f, Dims.X, Dims.Y, 1.0f);
+		RHICmdList.Clear(true, Owner->ClearColor, false, 0.f, false, 0, FIntRect());
 		// copy surface to the texture for use
 		FResolveParams ResolveParams;
 		ResolveParams.CubeFace = (ECubeFace)FaceIdx;
-		RHICopyToResolveTarget(RenderTargetTextureRHI, TextureCubeRHI, true, ResolveParams);
+		RHICmdList.CopyToResolveTarget(RenderTargetTextureRHI, TextureCubeRHI, true, ResolveParams);
 	}
 }
 
@@ -355,7 +357,7 @@ bool FTextureRenderTargetCubeResource::ReadPixels(TArray< FColor >& OutImageData
 	  ReadSurfaceCommand,
 	  FReadSurfaceContext, Context, ReadSurfaceContext,
 	{
-		RHIReadSurfaceData(
+		RHICmdList.ReadSurfaceData(
 		  Context.SrcRenderTarget->TextureCubeRHI,
 		  Context.Rect,
 		  *Context.OutData,
@@ -401,7 +403,7 @@ bool FTextureRenderTargetCubeResource::ReadPixels(TArray<FFloat16Color>& OutImag
 	  ReadSurfaceFloatCommand,
 	  FReadSurfaceFloatContext, Context, ReadSurfaceFloatContext,
 	{
-		RHIReadSurfaceFloatData(
+		RHICmdList.ReadSurfaceFloatData(
 		  Context.SrcRenderTarget->TextureCubeRHI,
 		  Context.Rect,
 		  *Context.OutData,
diff --git a/Engine/Source/Runtime/Engine/Private/TileRendering.cpp b/Engine/Source/Runtime/Engine/Private/TileRendering.cpp
--- a/Engine/Source/Runtime/Engine/Private/TileRendering.cpp
+++ b/Engine/Source/Runtime/Engine/Private/TileRendering.cpp
@@ -42,7 +42,7 @@ public:
 	/** 
 	* Initialize the RHI for this rendering resource 
 	*/
-	void InitRHI()
+	virtual void InitRHI() override
 	{
 		// used with a tristrip, so only 4 vertices are needed
 		uint32 Size = 4 * sizeof(FMaterialTileVertex);
@@ -131,7 +131,7 @@ public:
 };
 TGlobalResource<FTileMesh> GTileMesh;
 
-void FTileRenderer::DrawTile(const class FSceneView& View, const FMaterialRenderProxy* MaterialRenderProxy, float X, float Y, float SizeX, float SizeY, float U, float V, float SizeU, float SizeV, bool bIsHitTesting, const FHitProxyId HitProxyId)
+void FTileRenderer::DrawTile(FRHICommandListImmediate& RHICmdList, const class FSceneView& View, const FMaterialRenderProxy* MaterialRenderProxy, float X, float Y, float SizeX, float SizeY, float U, float V, float SizeU, float SizeV, bool bIsHitTesting, const FHitProxyId HitProxyId)
 {
 	FMaterialTileVertex DestVertex[4];
 
@@ -157,5 +157,5 @@ void FTileRenderer::DrawTile(const class FSceneView& View, const FMaterialRender
 	Mesh.DynamicVertexData = DestVertex;
 	Mesh.MaterialRenderProxy = MaterialRenderProxy;
 
-	GetRendererModule().DrawTileMesh(View, Mesh, bIsHitTesting, HitProxyId);
+	GetRendererModule().DrawTileMesh(RHICmdList, View, Mesh, bIsHitTesting, HitProxyId);
 }
diff --git a/Engine/Source/Runtime/Engine/Private/UnrealClient.cpp b/Engine/Source/Runtime/Engine/Private/UnrealClient.cpp
--- a/Engine/Source/Runtime/Engine/Private/UnrealClient.cpp
+++ b/Engine/Source/Runtime/Engine/Private/UnrealClient.cpp
@@ -62,7 +62,7 @@ bool FRenderTarget::ReadPixels(TArray< FColor >& OutImageData, FReadSurfaceDataF
 		ReadSurfaceCommand,
 		FReadSurfaceContext,Context,ReadSurfaceContext,
 	{
-		RHIReadSurfaceData(
+		RHICmdList.ReadSurfaceData(
 			Context.SrcRenderTarget->GetRenderTargetTexture(),
 			Context.Rect,
 			*Context.OutData,
@@ -123,7 +123,7 @@ bool FRenderTarget::ReadFloat16Pixels(FFloat16Color* OutImageData,ECubeFace Cube
 		ReadSurfaceFloatCommand,
 		FReadSurfaceFloatContext,Context,ReadSurfaceFloatContext,
 	{
-		RHIReadSurfaceFloatData(
+		RHICmdList.ReadSurfaceFloatData(
 			Context.SrcRenderTarget->GetRenderTargetTexture(),
 			Context.Rect,
 			*Context.OutData,
@@ -695,13 +695,13 @@ public:
 	}
 
 	// Begin FViewport interface
-	virtual void BeginRenderFrame() 
+	virtual void BeginRenderFrame(FRHICommandListImmediate& RHICmdList) override
 	{
 		check( IsInRenderingThread() );
-		RHISetRenderTarget( RenderTargetTextureRHI,  FTexture2DRHIRef() );
+		SetRenderTarget(RHICmdList,  RenderTargetTextureRHI,  FTexture2DRHIRef() );
 	};
 
-	void EndRenderFrame( bool bPresent, bool bLockToVsync )
+	virtual void EndRenderFrame(FRHICommandListImmediate& RHICmdList, bool bPresent, bool bLockToVsync) override
 	{
 		check( IsInRenderingThread() );
 	}
@@ -732,10 +732,10 @@ public:
 	}
 
 	// @todo UE4 DLL: Without these functions we get unresolved linker errors with FRenderResource
-	virtual void InitRHI(){}
-	virtual void ReleaseRHI(){}
-	virtual void InitResource(){ FViewport::InitResource(); }
-	virtual void ReleaseResource() { FViewport::ReleaseResource(); }
+	virtual void InitRHI() override{}
+	virtual void ReleaseRHI() override{}
+	virtual void InitResource() override{ FViewport::InitResource(); }
+	virtual void ReleaseResource() override { FViewport::ReleaseResource(); }
 	virtual FString GetFriendlyName() const { return FString(TEXT("FDummyViewport"));}
 	// End FRenderResource interface
 private:
@@ -817,7 +817,7 @@ void FViewport::HighResScreenshot()
 		FViewport*,Viewport,DummyViewport,
 		FIntPoint,InRestoreSize,RestoreSize,
 	{
-		Viewport->EndRenderFrame( false, false );
+		Viewport->EndRenderFrame(RHICmdList, false, false);
 		GetRendererModule().SceneRenderTargetsSetBufferSize(InRestoreSize.X, InRestoreSize.Y);
 	});
 
@@ -866,19 +866,18 @@ struct FEndDrawingCommandParams
  *
  * @param Parameters	Parameters passed from the gamethread to the renderthread command.
  */
-static void ViewportEndDrawing( FEndDrawingCommandParams Parameters )
+static void ViewportEndDrawing(FRHICommandListImmediate& RHICmdList, FEndDrawingCommandParams Parameters)
 {	
 	GInputLatencyTimer.RenderThreadTrigger = Parameters.bShouldTriggerTimerEvent;
-	Parameters.Viewport->EndRenderFrame( Parameters.bShouldPresent, Parameters.bLockToVsync );	
+	Parameters.Viewport->EndRenderFrame(RHICmdList, Parameters.bShouldPresent, Parameters.bLockToVsync);
 }
 
 /** Starts a new rendering frame. Called from the rendering thread. */
-void FViewport::BeginRenderFrame()
+void FViewport::BeginRenderFrame(FRHICommandListImmediate& RHICmdList)
 {
 	check( IsInRenderingThread() );
-
-	RHIBeginDrawingViewport( GetViewportRHI(), FTextureRHIRef() );
-	UpdateRenderTargetSurfaceRHIToCurrentBackBuffer( );
+	RHICmdList.BeginDrawingViewport(GetViewportRHI(), FTextureRHIRef());
+	UpdateRenderTargetSurfaceRHIToCurrentBackBuffer(RHICmdList);
 }
 
 /**
@@ -886,12 +885,12 @@ void FViewport::BeginRenderFrame()
  *	@param bPresent		Whether the frame should be presented to the screen
  *	@param bLockToVsync	Whether the GPU should block until VSYNC before presenting
  */
-void FViewport::EndRenderFrame( bool bPresent, bool bLockToVsync )
+void FViewport::EndRenderFrame(FRHICommandListImmediate& RHICmdList, bool bPresent, bool bLockToVsync)
 {
 	check( IsInRenderingThread() );
 
 	uint32 StartTime = FPlatformTime::Cycles();
-	RHIEndDrawingViewport( GetViewportRHI(), bPresent, bLockToVsync );
+	RHICmdList.EndDrawingViewport(GetViewportRHI(), bPresent, bLockToVsync);
 	uint32 EndTime = FPlatformTime::Cycles();
 
 	GRenderThreadIdle[ERenderThreadIdleTypes::WaitingForGPUPresent] += EndTime - StartTime;
@@ -961,7 +960,7 @@ void FViewport::EnqueueBeginRenderFrame()
 		BeginDrawingCommand,
 		FViewport*,Viewport,this,
 	{
-		Viewport->BeginRenderFrame();
+		Viewport->BeginRenderFrame(RHICmdList);
 	});
 }
 
@@ -1085,7 +1084,7 @@ void FViewport::Draw( bool bShouldPresent /*= true */)
 					EndDrawingCommand,
 					FEndDrawingCommandParams,Parameters,Params,
 				{
-					ViewportEndDrawing( Parameters );
+					ViewportEndDrawing(RHICmdList, Parameters);
 				});
 
 				GInputLatencyTimer.GameThreadTrigger = false;
@@ -1165,10 +1164,10 @@ const TArray<FColor>& FViewport::GetRawHitProxyData(FIntRect InRect)
 			FViewport*, Viewport, this,
 			{
 			// Set the hit proxy map's render target.
-			RHISetRenderTarget(Viewport->HitProxyMap.GetRenderTargetTexture(), FTextureRHIRef());
+			SetRenderTarget(RHICmdList, Viewport->HitProxyMap.GetRenderTargetTexture(), FTextureRHIRef());
 
 			// Clear the hit proxy map to white, which is overloaded to mean no hit proxy.
-			RHIClear(true, FLinearColor::White, false, 0, false, 0, FIntRect());
+			RHICmdList.Clear(true, FLinearColor::White, false, 0, false, 0, FIntRect());
 		});
 
 		// Let the viewport client draw its hit proxies.
@@ -1184,16 +1183,16 @@ const TArray<FColor>& FViewport::GetRawHitProxyData(FIntRect InRect)
 			FHitProxyMap*, HitProxyMap, &HitProxyMap,
 			{
 			// Copy (resolve) the rendered thumbnail from the render target to its texture
-			RHICopyToResolveTarget(HitProxyMap->GetRenderTargetTexture(), HitProxyMap->GetHitProxyTexture(), false, FResolveParams());
-			RHICopyToResolveTarget(HitProxyMap->GetRenderTargetTexture(), HitProxyMap->GetHitProxyCPUTexture(), false, FResolveParams());
+			RHICmdList.CopyToResolveTarget(HitProxyMap->GetRenderTargetTexture(), HitProxyMap->GetHitProxyTexture(), false, FResolveParams());
+			RHICmdList.CopyToResolveTarget(HitProxyMap->GetRenderTargetTexture(), HitProxyMap->GetHitProxyCPUTexture(), false, FResolveParams());
 		});
 
 		ENQUEUE_UNIQUE_RENDER_COMMAND_ONEPARAMETER(
 			EndDrawingCommand,
 			FViewport*, Viewport, this,
 			{
-			Viewport->EndRenderFrame(false, false);
-		});
+				Viewport->EndRenderFrame(RHICmdList, false, false);
+			});
 
 		// Cache the hit proxies for the next GetHitProxyMap call.
 		bHitProxiesCached = true;
@@ -1220,7 +1219,7 @@ const TArray<FColor>& FViewport::GetRawHitProxyData(FIntRect InRect)
 			ReadSurfaceCommand,
 			FReadSurfaceContext, Context, ReadSurfaceContext,
 			{
-			RHIReadSurfaceData(
+			RHICmdList.ReadSurfaceData(
 			Context.Viewport->HitProxyMap.GetHitProxyCPUTexture(),
 			Context.Rect,
 			*Context.OutData,
@@ -1419,7 +1418,9 @@ void FViewport::SetViewportClient( FViewportClient* InViewportClient )
 void FViewport::InitDynamicRHI()
 {
 	// Capture the viewport's back buffer surface for use through the FRenderTarget interface.
-	UpdateRenderTargetSurfaceRHIToCurrentBackBuffer();
+	FRHICommandListImmediate& RHICmdList = FRHICommandListExecutor::GetImmediateCommandList();
+
+	UpdateRenderTargetSurfaceRHIToCurrentBackBuffer(RHICmdList);
 
 	if(bRequiresHitProxyStorage)
 	{
@@ -1446,14 +1447,14 @@ void FViewport::InitRHI()
 
 	if(!IsValidRef(ViewportRHI))
 	{
-		ViewportRHI = RHICreateViewport(
+		FRHICommandListImmediate& RHICmdList = FRHICommandListExecutor::GetImmediateCommandList();
+		ViewportRHI = RHICmdList.CreateViewport(
 			GetWindow(),
 			SizeX,
 			SizeY,
 			IsFullscreen()
 			);
-		
-		UpdateRenderTargetSurfaceRHIToCurrentBackBuffer();
+		UpdateRenderTargetSurfaceRHIToCurrentBackBuffer(RHICmdList);
 	}
 }
 
@@ -1567,11 +1568,11 @@ bool FViewport::HasToggleFreezeCommand()
 /**
  * Update the render target surface RHI to the current back buffer 
  */
-void FViewport::UpdateRenderTargetSurfaceRHIToCurrentBackBuffer()
+void FViewport::UpdateRenderTargetSurfaceRHIToCurrentBackBuffer(FRHICommandListImmediate& RHICmdList)
 {
 	if(IsValidRef(ViewportRHI))
 	{
-		RenderTargetTextureRHI = RHIGetViewportBackBuffer(ViewportRHI);
+		RenderTargetTextureRHI = RHICmdList.GetViewportBackBuffer(ViewportRHI);
 	}
 }
 
diff --git a/Engine/Source/Runtime/Engine/Private/UserInterface/Canvas.cpp b/Engine/Source/Runtime/Engine/Private/UserInterface/Canvas.cpp
--- a/Engine/Source/Runtime/Engine/Private/UserInterface/Canvas.cpp
+++ b/Engine/Source/Runtime/Engine/Private/UserInterface/Canvas.cpp
@@ -175,8 +175,8 @@ bool FCanvasBatchedElementRenderItem::Render( const FCanvas* Canvas )
 		{
 			const bool bNeedsToSwitchVerticalAxis = IsES2Platform(GRHIShaderPlatform) && !IsPCPlatform(GRHIShaderPlatform);
 
-			//@todo-rco: RHIPacketList
-			FRHICommandList& RHICmdList = FRHICommandList::GetNullRef();
+			
+			FRHICommandListImmediate& RHICmdList = FRHICommandListExecutor::GetImmediateCommandList();
 
 			// draw batched items
 			Data->BatchedElements.Draw(
@@ -222,9 +222,7 @@ bool FCanvasBatchedElementRenderItem::Render( const FCanvas* Canvas )
 				FBatchedDrawParameters,Parameters,DrawParameters,
 			{
 				const bool bNeedsToSwitchVerticalAxis = IsES2Platform(GRHIShaderPlatform) && !IsPCPlatform(GRHIShaderPlatform);
-				//@todo-rco: RHIPacketList
-				FRHICommandList& RHICmdList = FRHICommandList::GetNullRef();
-
+				
 				// draw batched items
 				Parameters.RenderData->BatchedElements.Draw(
 					RHICmdList,
@@ -294,6 +292,7 @@ bool FCanvasTileRendererItem::Render( const FCanvas* Canvas )
 		{
 			const FRenderData::FTileInst& Tile = Data->Tiles[TileIdx];
 			FTileRenderer::DrawTile(
+				FRHICommandListExecutor::GetImmediateCommandList(),
 				*View, 
 				Data->MaterialRenderProxy, 
 				Tile.X, Tile.Y, Tile.SizeX, Tile.SizeY, 
@@ -333,6 +332,7 @@ bool FCanvasTileRendererItem::Render( const FCanvas* Canvas )
 			{
 				const FRenderData::FTileInst& Tile = Parameters.RenderData->Tiles[TileIdx];
 				FTileRenderer::DrawTile(
+					RHICmdList,
 					*Parameters.View, 
 					Parameters.RenderData->MaterialRenderProxy, 
 					Tile.X, Tile.Y, Tile.SizeX, Tile.SizeY, 
@@ -494,18 +494,19 @@ void FCanvas::Flush(bool bForce)
 	};
 	// sort the array of FCanvasSortElement entries so that higher sort keys render first (back-to-front)
 	SortedElements.Sort( FCompareFCanvasSortElement() );
-
 	if( GRenderingThread && IsInRenderingThread() )
 	{
+		FRHICommandListImmediate& RHICmdList = FRHICommandListExecutor::GetImmediateCommandList();
+
 		SCOPED_DRAW_EVENT(CanvasFlush, DEC_SCENE_ITEMS);
 		const FTexture2DRHIRef& RenderTargetTexture = RenderTarget->GetRenderTargetTexture();
 
 		check(IsValidRef(RenderTargetTexture));
 
 		// Set the RHI render target.
-		RHISetRenderTarget(RenderTargetTexture, FTextureRHIRef());
+		::SetRenderTarget(RHICmdList, RenderTargetTexture, FTextureRHIRef());
 		// disable depth test & writes
-		RHISetDepthStencilState(TStaticDepthStencilState<false,CF_Always>::GetRHI());
+		RHICmdList.SetDepthStencilState(TStaticDepthStencilState<false, CF_Always>::GetRHI());
 
 		if( ViewRect.Area() <= 0 )
 		{
@@ -513,7 +514,7 @@ void FCanvas::Flush(bool bForce)
 		}
 
 		// set viewport to RT size
-		RHISetViewport( ViewRect.Min.X,ViewRect.Min.Y, 0.0f, ViewRect.Max.X, ViewRect.Max.Y, 1.0f );	
+		RHICmdList.SetViewport( ViewRect.Min.X,ViewRect.Min.Y, 0.0f, ViewRect.Max.X, ViewRect.Max.Y, 1.0f );	
 	}
 	else 
 	{
@@ -539,13 +540,13 @@ void FCanvas::Flush(bool bForce)
 			SCOPED_DRAW_EVENT(CanvasFlush, DEC_SCENE_ITEMS);
 
 			// Set the RHI render target.
-			RHISetRenderTarget(Parameters.CanvasRenderTarget->GetRenderTargetTexture(), FTextureRHIRef());
+			::SetRenderTarget(RHICmdList, Parameters.CanvasRenderTarget->GetRenderTargetTexture(), FTextureRHIRef());
 			// disable depth test & writes
-			RHISetDepthStencilState(TStaticDepthStencilState<false,CF_Always>::GetRHI());
+			RHICmdList.SetDepthStencilState(TStaticDepthStencilState<false,CF_Always>::GetRHI());
 
 			const FIntRect& ViewportRect = Parameters.ViewRect;
 			// set viewport to RT size
-			RHISetViewport(ViewportRect.Min.X, ViewportRect.Min.Y, 0.0f, ViewportRect.Max.X, ViewportRect.Max.Y, 1.0f);
+			RHICmdList.SetViewport(ViewportRect.Min.X, ViewportRect.Min.Y, 0.0f, ViewportRect.Max.X, ViewportRect.Max.Y, 1.0f);
 		});
 	}
 
@@ -686,10 +687,10 @@ void FCanvas::Clear(const FLinearColor& Color)
 		{
 			if( CanvasRenderTarget )
 			{
-				RHISetRenderTarget(CanvasRenderTarget->GetRenderTargetTexture(),FTextureRHIRef());
-				RHISetViewport(0,0,0.0f,CanvasRenderTarget->GetSizeXY().X,CanvasRenderTarget->GetSizeXY().Y,1.0f);
+				::SetRenderTarget(RHICmdList, CanvasRenderTarget->GetRenderTargetTexture(),FTextureRHIRef());
+				RHICmdList.SetViewport(0,0,0.0f,CanvasRenderTarget->GetSizeXY().X,CanvasRenderTarget->GetSizeXY().Y,1.0f);
 			}
-			RHIClear(true,Color,false,0.0f,false,0, FIntRect());
+			RHICmdList.Clear(true,Color,false,0.0f,false,0, FIntRect());
 		});
 }
 
diff --git a/Engine/Source/Runtime/Engine/Private/VectorField.cpp b/Engine/Source/Runtime/Engine/Private/VectorField.cpp
--- a/Engine/Source/Runtime/Engine/Private/VectorField.cpp
+++ b/Engine/Source/Runtime/Engine/Private/VectorField.cpp
@@ -638,24 +638,24 @@ public:
 	/**
 	 * Set output buffer for this shader.
 	 */
-	void SetOutput( FUnorderedAccessViewRHIParamRef VolumeTextureUAV )
+	void SetOutput(FRHICommandList& RHICmdList, FUnorderedAccessViewRHIParamRef VolumeTextureUAV)
 	{
 		FComputeShaderRHIParamRef ComputeShaderRHI = GetComputeShader();
 		if ( OutVolumeTexture.IsBound() )
 		{
-			RHISetUAVParameter( ComputeShaderRHI, OutVolumeTexture.GetBaseIndex(), VolumeTextureUAV );
+			RHICmdList.SetUAVParameter(ComputeShaderRHI, OutVolumeTexture.GetBaseIndex(), VolumeTextureUAV);
 		}
 	}
 
 	/**
 	 * Unbinds any buffers that have been bound.
 	 */
-	void UnbindBuffers()
+	void UnbindBuffers(FRHICommandList& RHICmdList)
 	{
 		FComputeShaderRHIParamRef ComputeShaderRHI = GetComputeShader();
 		if ( OutVolumeTexture.IsBound() )
 		{
-			RHISetUAVParameter( ComputeShaderRHI, OutVolumeTexture.GetBaseIndex(), FUnorderedAccessViewRHIParamRef() );
+			RHICmdList.SetUAVParameter(ComputeShaderRHI, OutVolumeTexture.GetBaseIndex(), FUnorderedAccessViewRHIParamRef());
 		}
 	}
 
@@ -755,7 +755,7 @@ public:
 	 * Updates the vector field.
 	 * @param DeltaSeconds - Elapsed time since the last update.
 	 */
-	virtual void Update(float DeltaSeconds)
+	virtual void Update(FRHICommandListImmediate& RHICmdList, float DeltaSeconds) override
 	{
 		check(IsInRenderingThread());
 
@@ -811,21 +811,21 @@ public:
 				NoiseVolumeTextureRHI = AnimatedVectorField->NoiseField->Resource->VolumeTextureRHI;
 			}
 
-			RHISetComputeShader(CompositeCS->GetComputeShader());
-			CompositeCS->SetOutput(VolumeTextureUAV);
+			RHICmdList.SetComputeShader(CompositeCS->GetComputeShader());
+			CompositeCS->SetOutput(RHICmdList, VolumeTextureUAV);
 			/// ?
 			CompositeCS->SetParameters(
-				FRHICommandList::GetNullRef(),
+				RHICmdList,
 				UniformBuffer,
 				AnimatedVectorField->Texture->Resource->TextureRHI,
 				NoiseVolumeTextureRHI );
 			DispatchComputeShader(
-				FRHICommandList::GetNullRef(),
+				RHICmdList,
 				*CompositeCS,
 				SizeX / THREADS_PER_AXIS,
 				SizeY / THREADS_PER_AXIS,
 				SizeZ / THREADS_PER_AXIS );
-			CompositeCS->UnbindBuffers();
+			CompositeCS->UnbindBuffers(RHICmdList);
 		}
 	}
 
diff --git a/Engine/Source/Runtime/Engine/Private/VectorField.h b/Engine/Source/Runtime/Engine/Private/VectorField.h
--- a/Engine/Source/Runtime/Engine/Private/VectorField.h
+++ b/Engine/Source/Runtime/Engine/Private/VectorField.h
@@ -56,7 +56,7 @@ public:
 	 * Updates the vector field.
 	 * @param DeltaSeconds - Elapsed time since the last update.
 	 */
-	virtual void Update(float DeltaSeconds) {}
+	virtual void Update(FRHICommandListImmediate& RHICmdList, float DeltaSeconds) {}
 
 	/**
 	 * Resets the vector field simulation.
diff --git a/Engine/Source/Runtime/Engine/Private/VectorFieldVisualization.cpp b/Engine/Source/Runtime/Engine/Private/VectorFieldVisualization.cpp
--- a/Engine/Source/Runtime/Engine/Private/VectorFieldVisualization.cpp
+++ b/Engine/Source/Runtime/Engine/Private/VectorFieldVisualization.cpp
@@ -60,14 +60,14 @@ public:
 
 	FVertexDeclarationRHIRef VertexDeclarationRHI;
 
-	virtual void InitRHI()
+	virtual void InitRHI() override
 	{
 		FVertexDeclarationElementList Elements;
 		Elements.Add(FVertexElement(0, 0, VET_Float4, 0, sizeof(FVector4)));
 		VertexDeclarationRHI = RHICreateVertexDeclaration(Elements);
 	}
 
-	virtual void ReleaseRHI()
+	virtual void ReleaseRHI() override
 	{
 		VertexDeclarationRHI.SafeRelease();
 	}
@@ -83,7 +83,7 @@ class FDummyVertexBuffer : public FVertexBuffer
 {
 public:
 
-	virtual void InitRHI()
+	virtual void InitRHI() override
 	{
 		FRHIResourceCreateInfo CreateInfo;
 		VertexBufferRHI = RHICreateVertexBuffer(sizeof(FVector4) * 2, BUF_Static, CreateInfo);
diff --git a/Engine/Source/Runtime/Engine/Public/BatchedElements.h b/Engine/Source/Runtime/Engine/Public/BatchedElements.h
--- a/Engine/Source/Runtime/Engine/Public/BatchedElements.h
+++ b/Engine/Source/Runtime/Engine/Public/BatchedElements.h
@@ -61,7 +61,7 @@ public:
 	// Destructor.
 	virtual ~FSimpleElementVertexDeclaration() {}
 
-	virtual void InitRHI()
+	virtual void InitRHI() override
 	{
 		FVertexDeclarationElementList Elements;
 		uint16 Stride = sizeof(FSimpleElementVertex);
@@ -72,7 +72,7 @@ public:
 		VertexDeclarationRHI = RHICreateVertexDeclaration(Elements);
 	}
 
-	virtual void ReleaseRHI()
+	virtual void ReleaseRHI() override
 	{
 		VertexDeclarationRHI.SafeRelease();
 	}
@@ -91,7 +91,7 @@ class FBatchedElementParameters
 public:
 
 	/** Binds vertex and pixel shaders for this element */
-	virtual void BindShaders_RenderThread(FRHICommandList& RHICmdList, const FMatrix& InTransform, const float InGamma, const FMatrix& ColorWeights, const FTexture* Texture ) = 0;
+	virtual void BindShaders_RenderThread(FRHICommandListImmediate& RHICmdList, const FMatrix& InTransform, const float InGamma, const FMatrix& ColorWeights, const FTexture* Texture) = 0;
 
 };
 
@@ -184,7 +184,7 @@ public:
 	 * @param View			Optional FSceneView for shaders that need access to view constants
 	 * @param DepthTexture	DepthTexture for manual depth testing with editor compositing in the pixel shader
 	 */
-	bool Draw(FRHICommandList& RHICmdList, bool bNeedToSwitchVerticalAxis, const FMatrix& Transform, uint32 ViewportSizeX, uint32 ViewportSizeY, bool bHitTesting, float Gamma = 1.0f, const FSceneView* View = NULL, FTexture2DRHIRef DepthTexture = FTexture2DRHIRef()) const;
+	bool Draw(FRHICommandListImmediate& RHICmdList, bool bNeedToSwitchVerticalAxis, const FMatrix& Transform, uint32 ViewportSizeX, uint32 ViewportSizeY, bool bHitTesting, float Gamma = 1.0f, const FSceneView* View = NULL, FTexture2DRHIRef DepthTexture = FTexture2DRHIRef()) const;
 	
 	FORCEINLINE bool HasPrimsToDraw() const
 	{
@@ -218,7 +218,7 @@ private:
 	 * @param	CameraX		Local space normalized view direction X vector
 	 * @param	CameraY		Local space normalized view direction Y vector
 	 */
-	void DrawPointElements(FRHICommandList& RHICmdList, const FMatrix& Transform, const uint32 ViewportSizeX, const uint32 ViewportSizeY, const FVector& CameraX, const FVector& CameraY) const;
+	void DrawPointElements(FRHICommandListImmediate& RHICmdList, const FMatrix& Transform, const uint32 ViewportSizeX, const uint32 ViewportSizeY, const FVector& CameraX, const FVector& CameraY) const;
 
 	TArray<FSimpleElementVertex> LineVertices;
 
@@ -319,7 +319,7 @@ private:
 	 * Sets the appropriate vertex and pixel shader.
 	 */
 	void PrepareShaders(
-		FRHICommandList& RHICmdList,
+		FRHICommandListImmediate& RHICmdList,
 		ESimpleElementBlendMode BlendMode,
 		const FMatrix& Transform,
 		bool bSwitchVerticalAxis,
diff --git a/Engine/Source/Runtime/Engine/Public/Components.h b/Engine/Source/Runtime/Engine/Public/Components.h
--- a/Engine/Source/Runtime/Engine/Public/Components.h
+++ b/Engine/Source/Runtime/Engine/Public/Components.h
@@ -140,7 +140,7 @@ public:
 	}
 
 	// FRenderResource interface.
-	virtual void InitRHI();
+	virtual void InitRHI() override;
 	virtual FString GetFriendlyName() const { return TEXT("ColorOnly Static-mesh vertices"); }
 
 private:
diff --git a/Engine/Source/Runtime/Engine/Public/FXSystem.h b/Engine/Source/Runtime/Engine/Public/FXSystem.h
--- a/Engine/Source/Runtime/Engine/Public/FXSystem.h
+++ b/Engine/Source/Runtime/Engine/Public/FXSystem.h
@@ -173,17 +173,17 @@ public:
 	 * Notification from the renderer that it is about to draw FX belonging to
 	 * this system.
 	 */
-	virtual void PreRender() = 0;
+	virtual void PreRender(FRHICommandListImmediate& RHICmdList) = 0;
 
 	/**
 	 * Notification from the renderer that opaque primitives have rendered.
 	 */
-	virtual void PostRenderOpaque(const class FSceneView* CollisionView, FTexture2DRHIParamRef SceneDepthTexture, FTexture2DRHIParamRef GBufferATexture) = 0;
+	virtual void PostRenderOpaque(FRHICommandListImmediate& RHICmdList, const class FSceneView* CollisionView, FTexture2DRHIParamRef SceneDepthTexture, FTexture2DRHIParamRef GBufferATexture) = 0;
 
 	/**
 	 * Helper in case the data necessary for collision is not available.
 	 */
-	void PostRenderOpaque() { PostRenderOpaque(NULL,FTexture2DRHIParamRef(),FTexture2DRHIParamRef()); }
+	void PostRenderOpaque(FRHICommandListImmediate& RHICmdList) { PostRenderOpaque(RHICmdList, NULL, FTexture2DRHIParamRef(), FTexture2DRHIParamRef()); }
 
 protected:
 
diff --git a/Engine/Source/Runtime/Engine/Public/Landscape/LandscapeRender.h b/Engine/Source/Runtime/Engine/Public/Landscape/LandscapeRender.h
--- a/Engine/Source/Runtime/Engine/Public/Landscape/LandscapeRender.h
+++ b/Engine/Source/Runtime/Engine/Public/Landscape/LandscapeRender.h
@@ -175,7 +175,7 @@ public:
 	void Copy(const FLandscapeVertexFactory& Other);
 
 	// FRenderResource interface.
-	virtual void InitRHI();
+	virtual void InitRHI() override;
 
 	static bool SupportsTessellationShaders() { return true; }
 
@@ -253,7 +253,7 @@ public:
 	/** 
 	* Initialize the RHI for this rendering resource 
 	*/
-	virtual void InitRHI();
+	virtual void InitRHI() override;
 };
 
 
diff --git a/Engine/Source/Runtime/Engine/Public/LocalVertexFactory.h b/Engine/Source/Runtime/Engine/Public/LocalVertexFactory.h
--- a/Engine/Source/Runtime/Engine/Public/LocalVertexFactory.h
+++ b/Engine/Source/Runtime/Engine/Public/LocalVertexFactory.h
@@ -52,7 +52,7 @@ public:
 	void Copy(const FLocalVertexFactory& Other);
 
 	// FRenderResource interface.
-	virtual void InitRHI();
+	virtual void InitRHI() override;
 
 	static bool SupportsTessellationShaders() { return true; }
 
diff --git a/Engine/Source/Runtime/Engine/Public/MaterialShared.h b/Engine/Source/Runtime/Engine/Public/MaterialShared.h
--- a/Engine/Source/Runtime/Engine/Public/MaterialShared.h
+++ b/Engine/Source/Runtime/Engine/Public/MaterialShared.h
@@ -1268,8 +1268,8 @@ public:
 	bool IsHovered() const { return bHovered; }
 
 	// FRenderResource interface.
-	ENGINE_API virtual void InitDynamicRHI();
-	ENGINE_API virtual void ReleaseDynamicRHI();
+	ENGINE_API virtual void InitDynamicRHI() override;
+	ENGINE_API virtual void ReleaseDynamicRHI() override;
 
 	ENGINE_API static const TSet<FMaterialRenderProxy*>& GetMaterialRenderProxyMap() 
 	{
diff --git a/Engine/Source/Runtime/Engine/Public/RawIndexBuffer.h b/Engine/Source/Runtime/Engine/Public/RawIndexBuffer.h
--- a/Engine/Source/Runtime/Engine/Public/RawIndexBuffer.h
+++ b/Engine/Source/Runtime/Engine/Public/RawIndexBuffer.h
@@ -20,7 +20,7 @@ public:
 	void CacheOptimize();
 
 	// FRenderResource interface.
-	virtual void InitRHI();
+	virtual void InitRHI() override;
 
 	// Serialization.
 	friend FArchive& operator<<(FArchive& Ar,FRawIndexBuffer& I);
@@ -45,7 +45,7 @@ public:
 	void CacheOptimize();
 
 	// FRenderResource interface.
-	virtual void InitRHI();
+	virtual void InitRHI() override;
 
 	// Serialization.
 	friend FArchive& operator<<(FArchive& Ar,FRawIndexBuffer16or32& I);
@@ -221,7 +221,7 @@ public:
 	/**
 	* Create the index buffer RHI resource and initialize its data
 	*/
-	virtual void InitRHI()
+	virtual void InitRHI() override
 	{
 		uint32 Size = Indices.Num() * sizeof(INDEX_TYPE);
 		if(Indices.Num())
diff --git a/Engine/Source/Runtime/Engine/Public/SceneManagement.h b/Engine/Source/Runtime/Engine/Public/SceneManagement.h
--- a/Engine/Source/Runtime/Engine/Public/SceneManagement.h
+++ b/Engine/Source/Runtime/Engine/Public/SceneManagement.h
@@ -693,7 +693,7 @@ public:
 	// @param OutCascadeSettings can be 0
 	virtual FSphere GetShadowSplitBounds(const class FSceneView& View, int32 SplitIndex, FShadowCascadeSettings* OutCascadeSettings) const { return FSphere(FVector::ZeroVector, 0); }
 
-	virtual void SetScissorRect(const FSceneView& View) const
+	virtual void SetScissorRect(FRHICommandListImmediate& RHICmdList, const FSceneView& View) const
 	{
 	}
 
diff --git a/Engine/Source/Runtime/Engine/Public/SkeletalMeshTypes.h b/Engine/Source/Runtime/Engine/Public/SkeletalMeshTypes.h
--- a/Engine/Source/Runtime/Engine/Public/SkeletalMeshTypes.h
+++ b/Engine/Source/Runtime/Engine/Public/SkeletalMeshTypes.h
@@ -1176,7 +1176,7 @@ public:
 	/**
 	* Initialize the RHI resource for this vertex buffer
 	*/
-	virtual void InitRHI();
+	virtual void InitRHI() override;
 
 	/**
 	* @return text description for the resource type
@@ -1543,7 +1543,7 @@ public:
 	/**
 	 * Initialize the RHI resource for this vertex buffer
 	 */
-	virtual void InitRHI();
+	virtual void InitRHI() override;
 
 	/**
 	 * @return text description for the resource type
@@ -1656,7 +1656,7 @@ public:
 	/**
 	 * Initialize the RHI resource for this vertex buffer
 	 */
-	virtual void InitRHI();
+	virtual void InitRHI() override;
 
 	/**
 	 * @return text description for the resource type
diff --git a/Engine/Source/Runtime/Engine/Public/SkeletalRenderPublic.h b/Engine/Source/Runtime/Engine/Public/SkeletalRenderPublic.h
--- a/Engine/Source/Runtime/Engine/Public/SkeletalRenderPublic.h
+++ b/Engine/Source/Runtime/Engine/Public/SkeletalRenderPublic.h
@@ -44,7 +44,7 @@ public:
 	 * Called by the rendering thread to update the current dynamic data
 	 * @param	InDynamicData - data that was created by the game thread for use by the rendering thread
 	 */
-	virtual void UpdateDynamicData_RenderThread(class FDynamicSkelMeshObjectData* InDynamicData) = 0;
+	virtual void UpdateDynamicData_RenderThread(FRHICommandListImmediate& RHICmdList, class FDynamicSkelMeshObjectData* InDynamicData) = 0;
 
 	/**
 	 * @param	LODIndex - each LOD has its own vertex data
diff --git a/Engine/Source/Runtime/Engine/Public/Slate/SceneViewport.h b/Engine/Source/Runtime/Engine/Public/Slate/SceneViewport.h
--- a/Engine/Source/Runtime/Engine/Public/Slate/SceneViewport.h
+++ b/Engine/Source/Runtime/Engine/Public/Slate/SceneViewport.h
@@ -105,7 +105,7 @@ public:
 	virtual void EnqueueBeginRenderFrame() override;
 
 	/** Called when a frame starts to render */
-	virtual void BeginRenderFrame() override;
+	virtual void BeginRenderFrame(FRHICommandListImmediate& RHICmdList) override;
 
 	/** 
 	 * Called when a frame is done rendering
@@ -113,7 +113,7 @@ public:
 	 * @param bPresent	Not used in Slate viewports
 	 * @param bLockToVsync	Not used in Slate viewports
 	 */
-	virtual void EndRenderFrame( bool bPresent, bool bLockToVsync ) override;
+	virtual void EndRenderFrame(FRHICommandListImmediate& RHICmdList, bool bPresent, bool bLockToVsync) override;
 
 	/**
 	 * Ticks the viewport
diff --git a/Engine/Source/Runtime/Engine/Public/Slate/SlateTextures.h b/Engine/Source/Runtime/Engine/Public/Slate/SlateTextures.h
--- a/Engine/Source/Runtime/Engine/Public/Slate/SlateTextures.h
+++ b/Engine/Source/Runtime/Engine/Public/Slate/SlateTextures.h
@@ -19,10 +19,10 @@ public:
 	virtual uint32 GetHeight() const { return Height; }
 
 	/** FRenderResource Interface.  Called when render resources need to be initialized */
-	virtual void InitDynamicRHI();
+	virtual void InitDynamicRHI() override;
 
 	/** FRenderResource Interface.  Called when render resources need to be released */
-	virtual void ReleaseDynamicRHI();
+	virtual void ReleaseDynamicRHI() override;
 
 	/**
 	 * Resize the texture.  Can only be called on the render thread
@@ -91,7 +91,7 @@ public:
 	{
 	}
 
-	virtual void InitDynamicRHI() {}
+	virtual void InitDynamicRHI() override {}
 
 	/** 
 	 * Releases all dynamic RHI data
diff --git a/Engine/Source/Runtime/Engine/Public/StaticBoundShaderState.h b/Engine/Source/Runtime/Engine/Public/StaticBoundShaderState.h
--- a/Engine/Source/Runtime/Engine/Public/StaticBoundShaderState.h
+++ b/Engine/Source/Runtime/Engine/Public/StaticBoundShaderState.h
@@ -60,8 +60,9 @@ typedef TGlobalResource<FGlobalBoundShaderStateResource> FGlobalBoundShaderState
  * @param PixelShader				- the pixel shader to use in creating the new bound shader state
  * @param GeometryShader			- the geometry shader to use in creating the new bound shader state (0 if not used)
  */
+
 extern ENGINE_API void SetGlobalBoundShaderState(
-	FRHICommandList& RHICmdList,
+	FRHICommandListImmediate& RHICmdList,
 	FGlobalBoundShaderState& BoundShaderState,
 	FVertexDeclarationRHIParamRef VertexDeclaration,
 	FShader* VertexShader,
diff --git a/Engine/Source/Runtime/Engine/Public/StaticRHI.h b/Engine/Source/Runtime/Engine/Public/StaticRHI.h
--- a/Engine/Source/Runtime/Engine/Public/StaticRHI.h
+++ b/Engine/Source/Runtime/Engine/Public/StaticRHI.h
@@ -52,7 +52,7 @@ class STATIC_RHI_CLASS_NAME
 public:
 
 	// For Static RHI, we declare the various RHI functions as static methods
-	#define DEFINE_RHIMETHOD(Type,Name,ParameterTypesAndNames,ParameterNames,ReturnStatement,NullImplementation) static Type Name ParameterTypesAndNames
+	#define DEFINE_RHIMETHOD(Type,Name,ParameterTypesAndNames,ParameterNames,ReturnStatement,NullImplementation) static Type RHI##Name ParameterTypesAndNames
 	#include "RHIMethods.h"
 	#undef DEFINE_RHIMETHOD
 
@@ -186,13 +186,31 @@ private:
 };
 
 // Implement the statically bound RHI methods to simply call the static RHI.
+#define DEFINE_RHIMETHOD_CMDLIST(Type,Name,ParameterTypesAndNames,ParameterNames,ReturnStatement,NullImplementation) \
+	FORCEINLINE Type Name##_Internal ParameterTypesAndNames \
+	{ \
+		ReturnStatement GStaticRHI->RHI##Name ParameterNames; \
+	}
+#define DEFINE_RHIMETHOD_GLOBAL(Type,Name,ParameterTypesAndNames,ParameterNames,ReturnStatement,NullImplementation) \
+	FORCEINLINE Type RHI##Name ParameterTypesAndNames \
+	{ \
+		ReturnStatement GStaticRHI->RHI##Name ParameterNames; \
+	}
+#define DEFINE_RHIMETHOD_GLOBALFLUSH(Type,Name,ParameterTypesAndNames,ParameterNames,ReturnStatement,NullImplementation) \
+	FORCEINLINE Type Name##_Internal ParameterTypesAndNames \
+	{ \
+		ReturnStatement GStaticRHI->RHI##Name ParameterNames; \
+	}
 #define DEFINE_RHIMETHOD(Type,Name,ParameterTypesAndNames,ParameterNames,ReturnStatement,NullImplementation) \
-	FORCEINLINE Type Name ParameterTypesAndNames \
+	FORCEINLINE Type Name##_Internal ParameterTypesAndNames \
 	{ \
-		ReturnStatement GStaticRHI->Name ParameterNames; \
+		ReturnStatement GStaticRHI->RHI##Name ParameterNames; \
 	}
 #include "RHIMethods.h"
 #undef DEFINE_RHIMETHOD
+#undef DEFINE_RHIMETHOD_CMDLIST
+#undef DEFINE_RHIMETHOD_GLOBAL
+#undef DEFINE_RHIMETHOD_GLOBALFLUSH
 
 #endif	//#if USE_STATIC_RHI
 
diff --git a/Engine/Source/Runtime/Engine/Public/Texture.h b/Engine/Source/Runtime/Engine/Public/Texture.h
--- a/Engine/Source/Runtime/Engine/Public/Texture.h
+++ b/Engine/Source/Runtime/Engine/Public/Texture.h
@@ -223,11 +223,11 @@ public:
 	/**
 	 * Called when the resource is initialized. This is only called by the rendering thread.
 	 */
-	virtual void InitRHI();
+	virtual void InitRHI() override;
 	/**
 	 * Called when the resource is released. This is only called by the rendering thread.
 	 */
-	virtual void ReleaseRHI();
+	virtual void ReleaseRHI() override;
 
 	/** Returns the width of the texture in pixels. */
 	virtual uint32 GetSizeX() const;
@@ -404,10 +404,10 @@ public:
 	virtual uint32 GetSizeY() const;
 
 	/** Called when the resource is initialized. This is only called by the rendering thread. */
-	virtual void InitRHI();
+	virtual void InitRHI() override;
 
 	/** Called when the resource is released. This is only called by the rendering thread. */
-	virtual void ReleaseRHI();
+	virtual void ReleaseRHI() override;
 
 	/** Returns the Texture2DRHI, which can be used for locking/unlocking the mips. */
 	ENGINE_API FTexture2DRHIRef GetTexture2DRHI();
@@ -500,7 +500,7 @@ public:
 	/**
 	* Called when the resource is initialized. This is only called by the rendering thread.
 	*/
-	virtual void InitRHI();
+	virtual void InitRHI() override;
 
 	/** Returns the width of the texture in pixels. */
 	virtual uint32 GetSizeX() const
@@ -683,7 +683,7 @@ public:
 	 * Resources that need to initialize after a D3D device reset must implement this function.
 	 * This is only called by the rendering thread.
 	 */
-	virtual void InitDynamicRHI();
+	virtual void InitDynamicRHI() override;
 
 	/**
 	 * Releases the dynamic RHI resource and/or RHI render target resources used by this resource.
@@ -691,14 +691,14 @@ public:
 	 * Resources that need to release before a D3D device reset must implement this function.
 	 * This is only called by the rendering thread.
 	 */
-	virtual void ReleaseDynamicRHI();
+	virtual void ReleaseDynamicRHI() override;
 
 	// FDeferredClearResource interface
 
 	/**
 	 * Clear contents of the render target
 	 */
-	virtual void UpdateResource();
+	virtual void UpdateResource() override;
 
 	// FRenderTarget interface.
 	virtual FIntPoint GetSizeXY() const;
@@ -756,7 +756,7 @@ public:
 	 * Resources that need to initialize after a D3D device reset must implement this function.
 	 * This is only called by the rendering thread.
 	 */
-	virtual void InitDynamicRHI();
+	virtual void InitDynamicRHI() override;
 
 	/**
 	 * Releases the dynamic RHI resource and/or RHI render target resources used by this resource.
@@ -764,13 +764,13 @@ public:
 	 * Resources that need to release before a D3D device reset must implement this function.
 	 * This is only called by the rendering thread.
 	 */
-	virtual void ReleaseDynamicRHI();	
+	virtual void ReleaseDynamicRHI() override;
 
 	/**
 	 * Clear contents of the render target. Clears each face of the cube
 	 * This is only called by the rendering thread.
 	 */
-	virtual void UpdateResource();
+	virtual void UpdateResource() override;
 
 	// FRenderTarget interface.
 
diff --git a/Engine/Source/Runtime/Engine/Public/TileRendering.h b/Engine/Source/Runtime/Engine/Public/TileRendering.h
--- a/Engine/Source/Runtime/Engine/Public/TileRendering.h
+++ b/Engine/Source/Runtime/Engine/Public/TileRendering.h
@@ -15,7 +15,7 @@ public:
 	 * Draw a tile at the given location and size, using the given UVs
 	 * (UV = [0..1]
 	 */
-	ENGINE_API static void DrawTile(const class FSceneView& View, const FMaterialRenderProxy* MaterialRenderProxy, float X, float Y, float SizeX, float SizeY, float U, float V, float SizeU, float SizeV, bool bIsHitTesting=false, const FHitProxyId HitProxyId=FHitProxyId());
+	ENGINE_API static void DrawTile(FRHICommandListImmediate& RHICmdList, const class FSceneView& View, const FMaterialRenderProxy* MaterialRenderProxy, float X, float Y, float SizeX, float SizeY, float U, float V, float SizeU, float SizeV, bool bIsHitTesting = false, const FHitProxyId HitProxyId = FHitProxyId());
 
 private:
 
diff --git a/Engine/Source/Runtime/Engine/Public/UnrealClient.h b/Engine/Source/Runtime/Engine/Public/UnrealClient.h
--- a/Engine/Source/Runtime/Engine/Public/UnrealClient.h
+++ b/Engine/Source/Runtime/Engine/Public/UnrealClient.h
@@ -307,14 +307,14 @@ public:
 	/**
 	 *	Starts a new rendering frame. Called from the rendering thread.
 	 */
-	ENGINE_API virtual void	BeginRenderFrame();
+	ENGINE_API virtual void	BeginRenderFrame(FRHICommandListImmediate& RHICmdList);
 
 	/**
 	 *	Ends a rendering frame. Called from the rendering thread.
 	 *	@param bPresent		Whether the frame should be presented to the screen
 	 *	@param bLockToVsync	Whether the GPU should block until VSYNC before presenting
 	 */
-	ENGINE_API virtual void	EndRenderFrame( bool bPresent, bool bLockToVsync );
+	ENGINE_API virtual void	EndRenderFrame(FRHICommandListImmediate& RHICmdList, bool bPresent, bool bLockToVsync);
 
 	/**
 	 * @return whether or not this Controller has a keyboard available to be used
@@ -437,7 +437,7 @@ public:
 	/**
 	 * Update the render target surface RHI to the current back buffer 
 	 */
-	void UpdateRenderTargetSurfaceRHIToCurrentBackBuffer();
+	void UpdateRenderTargetSurfaceRHIToCurrentBackBuffer(FRHICommandListImmediate& RHICmdList);
 
 	/**
 	 * First chance for viewports to render custom stats text
@@ -591,10 +591,10 @@ protected:
 	/** Triggers the taking of a high res screen shot for this viewport. */
 	bool bTakeHighResScreenShot;
 	// FRenderResource interface.
-	ENGINE_API virtual void InitDynamicRHI();
-	ENGINE_API virtual void ReleaseDynamicRHI();
-	ENGINE_API virtual void InitRHI();
-	ENGINE_API virtual void ReleaseRHI();
+	ENGINE_API virtual void InitDynamicRHI() override;
+	ENGINE_API virtual void ReleaseDynamicRHI() override;
+	ENGINE_API virtual void InitRHI() override;
+	ENGINE_API virtual void ReleaseRHI() override;
 };
 
 // Shortcuts for checking the state of both left&right variations of control keys.
diff --git a/Engine/Source/Runtime/GameLiveStreaming/GameLiveStreaming.cpp b/Engine/Source/Runtime/GameLiveStreaming/GameLiveStreaming.cpp
--- a/Engine/Source/Runtime/GameLiveStreaming/GameLiveStreaming.cpp
+++ b/Engine/Source/Runtime/GameLiveStreaming/GameLiveStreaming.cpp
@@ -273,7 +273,7 @@ void FGameLiveStreaming::BroadcastGameVideoFrame()
 				ReadbackFromStagingBuffer,
 				FReadbackFromStagingBufferContext,Context,ReadbackFromStagingBufferContext,
 			{
-				RHIUnmapStagingSurface( Context.This->ReadbackTextures[ Context.This->ReadbackTextureIndex ] );
+				RHICmdList.UnmapStagingSurface(Context.This->ReadbackTextures[Context.This->ReadbackTextureIndex]);
 			});
 		}								
 	}
@@ -325,16 +325,14 @@ void FGameLiveStreaming::StartCopyingNextGameVideoFrame( const FViewportRHIRef&
 
 		const FSceneRenderTargetItem& DestRenderTarget = ResampleTexturePooledRenderTarget->GetRenderTargetItem();
 
-		RHISetRenderTarget( DestRenderTarget.TargetableTexture, FTextureRHIRef() );
-		RHISetViewport(0, 0, 0.0f, Context.ResizeTo.X, Context.ResizeTo.Y, 1.0f);
+		SetRenderTarget(RHICmdList, DestRenderTarget.TargetableTexture, FTextureRHIRef());
+		RHICmdList.SetViewport(0, 0, 0.0f, Context.ResizeTo.X, Context.ResizeTo.Y, 1.0f);
 
-		RHISetBlendState(TStaticBlendState<>::GetRHI());
-		RHISetRasterizerState(TStaticRasterizerState<>::GetRHI());
-		RHISetDepthStencilState(TStaticDepthStencilState<false,CF_Always>::GetRHI());
+		RHICmdList.SetBlendState(TStaticBlendState<>::GetRHI());
+		RHICmdList.SetRasterizerState(TStaticRasterizerState<>::GetRHI());
+		RHICmdList.SetDepthStencilState(TStaticDepthStencilState<false,CF_Always>::GetRHI());
 
-		FRHICommandList& RHICmdList = FRHICommandList::GetNullRef();
-
-		FTexture2DRHIRef ViewportBackBuffer = RHIGetViewportBackBuffer( Context.ViewportRHI );
+		FTexture2DRHIRef ViewportBackBuffer = RHICmdList.GetViewportBackBuffer(Context.ViewportRHI);
 
 		TShaderMapRef<FScreenVS> VertexShader(GetGlobalShaderMap());
 		TShaderMapRef<FScreenPS> PixelShader(GetGlobalShaderMap());
@@ -366,7 +364,7 @@ void FGameLiveStreaming::StartCopyingNextGameVideoFrame( const FViewportRHIRef&
 
 		// Asynchronously copy render target from GPU to CPU
 		const bool bKeepOriginalSurface = false;
-		RHICopyToResolveTarget(
+		RHICmdList.CopyToResolveTarget(
 			DestRenderTarget.TargetableTexture,
 			Context.This->ReadbackTextures[ Context.This->ReadbackTextureIndex ],
 			bKeepOriginalSurface,
@@ -390,7 +388,7 @@ void FGameLiveStreaming::StartCopyingNextGameVideoFrame( const FViewportRHIRef&
 		{
 			int32 UnusedWidth = 0;
 			int32 UnusedHeight = 0;
-			RHIMapStagingSurface( Context.This->ReadbackTextures[ Context.This->ReadbackTextureIndex ], Context.This->ReadbackBuffers[ Context.This->ReadbackBufferIndex ], UnusedWidth, UnusedHeight );
+			RHICmdList.MapStagingSurface(Context.This->ReadbackTextures[Context.This->ReadbackTextureIndex], Context.This->ReadbackBuffers[Context.This->ReadbackBufferIndex], UnusedWidth, UnusedHeight);
 
 			// Ping pong between readback textures
 			Context.This->ReadbackTextureIndex = ( Context.This->ReadbackTextureIndex + 1 ) % 2;
diff --git a/Engine/Source/Runtime/IOS/MetalRHI/Private/MetalCommands.cpp b/Engine/Source/Runtime/IOS/MetalRHI/Private/MetalCommands.cpp
--- a/Engine/Source/Runtime/IOS/MetalRHI/Private/MetalCommands.cpp
+++ b/Engine/Source/Runtime/IOS/MetalRHI/Private/MetalCommands.cpp
@@ -378,7 +378,7 @@ void FMetalDynamicRHI::RHIDrawPrimitive(uint32 PrimitiveType, uint32 BaseVertexI
 	}
 
 	// how many verts to render
-	uint32 NumVertices = RHIGetVertexCountForPrimitiveCount(NumPrimitives, PrimitiveType);
+	uint32 NumVertices = GetVertexCountForPrimitiveCount(NumPrimitives, PrimitiveType);
 
 	// finalize any pending state
 	FMetalManager::Get()->PrepareToDraw(NumVertices);
@@ -414,7 +414,7 @@ void FMetalDynamicRHI::RHIDrawIndexedPrimitive(FIndexBufferRHIParamRef IndexBuff
 	// finalize any pending state
 	FMetalManager::Get()->PrepareToDraw(NumVertices);
 	
-	uint32 NumIndices = RHIGetVertexCountForPrimitiveCount(NumPrimitives, PrimitiveType);
+	uint32 NumIndices = GetVertexCountForPrimitiveCount(NumPrimitives, PrimitiveType);
 #if !NO_DRAW
 	[FMetalManager::GetContext() drawIndexedPrimitives:TranslatePrimitiveType(PrimitiveType)
 											indexCount:NumIndices
@@ -471,7 +471,7 @@ void FMetalDynamicRHI::RHIEndDrawPrimitiveUP()
 	[FMetalManager::GetContext() setVertexBuffer:FMetalManager::Get()->GetRingBuffer() offset:GPendingVertexBufferOffset atIndex:UNREAL_TO_METAL_BUFFER_INDEX(0)];
 	
 	// how many to draw
-	uint32 NumVertices = RHIGetVertexCountForPrimitiveCount(GPendingNumPrimitives, GPendingPrimitiveType);
+	uint32 NumVertices = GetVertexCountForPrimitiveCount(GPendingNumPrimitives, GPendingPrimitiveType);
 
 	// last minute draw setup
 	FMetalManager::Get()->PrepareToDraw(0);
@@ -516,7 +516,7 @@ void FMetalDynamicRHI::RHIEndDrawIndexedPrimitiveUP()
 	[FMetalManager::GetContext() setVertexBuffer:FMetalManager::Get()->GetRingBuffer() offset:GPendingVertexBufferOffset atIndex:UNREAL_TO_METAL_BUFFER_INDEX(0)];
 
 	// how many to draw
-	uint32 NumIndices = RHIGetVertexCountForPrimitiveCount(GPendingNumPrimitives, GPendingPrimitiveType);
+	uint32 NumIndices = GetVertexCountForPrimitiveCount(GPendingNumPrimitives, GPendingPrimitiveType);
 
 	// last minute draw setup
 	FMetalManager::Get()->PrepareToDraw(0);
diff --git a/Engine/Source/Runtime/IOS/MetalRHI/Private/MetalGlobalUniformBuffer.h b/Engine/Source/Runtime/IOS/MetalRHI/Private/MetalGlobalUniformBuffer.h
--- a/Engine/Source/Runtime/IOS/MetalRHI/Private/MetalGlobalUniformBuffer.h
+++ b/Engine/Source/Runtime/IOS/MetalRHI/Private/MetalGlobalUniformBuffer.h
@@ -20,8 +20,8 @@ public:
 	~FMetalGlobalUniformBuffer();
 
 	// FRenderResource interface.
-	virtual void	InitDynamicRHI();
-	virtual void	ReleaseDynamicRHI();
+	virtual void	InitDynamicRHI() override;
+	virtual void	ReleaseDynamicRHI() override;
 
 	
 	void UpdateConstant(const void* Data, uint16 Offset, uint16 Size);
diff --git a/Engine/Source/Runtime/IOS/MetalRHI/Public/MetalRHI.h b/Engine/Source/Runtime/IOS/MetalRHI/Public/MetalRHI.h
--- a/Engine/Source/Runtime/IOS/MetalRHI/Public/MetalRHI.h
+++ b/Engine/Source/Runtime/IOS/MetalRHI/Public/MetalRHI.h
@@ -39,7 +39,7 @@ public:
 	virtual void Shutdown() {}
 
 	// The RHI methods are defined as virtual functions in URenderHardwareInterface.
-	#define DEFINE_RHIMETHOD(Type,Name,ParameterTypesAndNames,ParameterNames,ReturnStatement,NullImplementation) virtual Type Name ParameterTypesAndNames
+	#define DEFINE_RHIMETHOD(Type,Name,ParameterTypesAndNames,ParameterNames,ReturnStatement,NullImplementation) virtual Type RHI##Name ParameterTypesAndNames
 	#include "RHIMethods.h"
 	#undef DEFINE_RHIMETHOD
 
diff --git a/Engine/Source/Runtime/Launch/Private/LaunchEngineLoop.cpp b/Engine/Source/Runtime/Launch/Private/LaunchEngineLoop.cpp
--- a/Engine/Source/Runtime/Launch/Private/LaunchEngineLoop.cpp
+++ b/Engine/Source/Runtime/Launch/Private/LaunchEngineLoop.cpp
@@ -2008,7 +2008,7 @@ void FEngineLoop::Tick()
 			{
 				GFrameNumberRenderThread++;
 				GDynamicRHI->PushEvent(*FString::Printf(TEXT("Frame%d"),GFrameNumberRenderThread));
-				RHIBeginFrame();
+				RHICmdList.BeginFrame();
 			});
 
 		// Flush debug output which has been buffered by other threads.
@@ -2176,7 +2176,7 @@ void FEngineLoop::Tick()
 		ENQUEUE_UNIQUE_RENDER_COMMAND(
 			EndFrame,
 		{
-			RHIEndFrame();
+			RHICmdList.EndFrame();
 			GDynamicRHI->PopEvent();
 		});
 	} 
diff --git a/Engine/Source/Runtime/NullDrv/Public/NullRHI.h b/Engine/Source/Runtime/NullDrv/Public/NullRHI.h
--- a/Engine/Source/Runtime/NullDrv/Public/NullRHI.h
+++ b/Engine/Source/Runtime/NullDrv/Public/NullRHI.h
@@ -20,7 +20,7 @@ public:
 
 	// Implement the dynamic RHI interface using the null implementations defined in RHIMethods.h
 	#define DEFINE_RHIMETHOD(Type,Name,ParameterTypesAndNames,ParameterNames,ReturnStatement,NullImplementation) \
-		virtual Type Name ParameterTypesAndNames { NullImplementation; }
+	virtual Type RHI##Name ParameterTypesAndNames{ NullImplementation; }
 	#include "RHIMethods.h"
 	#undef DEFINE_RHIMETHOD
 
diff --git a/Engine/Source/Runtime/OpenGLDrv/Private/OpenGLCommands.cpp b/Engine/Source/Runtime/OpenGLDrv/Private/OpenGLCommands.cpp
--- a/Engine/Source/Runtime/OpenGLDrv/Private/OpenGLCommands.cpp
+++ b/Engine/Source/Runtime/OpenGLDrv/Private/OpenGLCommands.cpp
@@ -2256,7 +2256,7 @@ void FOpenGLDynamicRHI::RHIDrawPrimitive(uint32 PrimitiveType,uint32 BaseVertexI
 	SetupTexturesForDraw(ContextState);
 	CommitNonComputeShaderConstants();
 	CachedBindElementArrayBuffer(ContextState,0);
-	uint32 VertexCount = RHIGetVertexCountForPrimitiveCount(NumPrimitives,PrimitiveType);
+	uint32 VertexCount = GetVertexCountForPrimitiveCount(NumPrimitives,PrimitiveType);
 	SetupVertexArrays(ContextState, BaseVertexIndex, PendingState.Streams, VertexCount, NUM_OPENGL_VERTEX_STREAMS);
 
 	GLenum DrawMode = GL_TRIANGLES;
diff --git a/Engine/Source/Runtime/OpenGLDrv/Private/OpenGLDrv.cpp b/Engine/Source/Runtime/OpenGLDrv/Private/OpenGLDrv.cpp
--- a/Engine/Source/Runtime/OpenGLDrv/Private/OpenGLDrv.cpp
+++ b/Engine/Source/Runtime/OpenGLDrv/Private/OpenGLDrv.cpp
@@ -386,7 +386,7 @@ static FGlobalBoundShaderState LongGPUTaskBoundShaderState;
 
 void FOpenGLDynamicRHI::IssueLongGPUTask()
 {
-	FRHICommandList& RHICmdList = FRHICommandList::GetNullRef();
+	FRHICommandListImmediate& RHICmdList = FRHICommandListExecutor::GetRecursiveRHICommandList();
 	int32 LargestViewportIndex = INDEX_NONE;
 	int32 LargestViewportPixels = 0;
 
@@ -423,6 +423,7 @@ void FOpenGLDynamicRHI::IssueLongGPUTask()
 		Vertices[3].Set(  1.0f, -1.0f, 0, 1.0f );
 		DrawPrimitiveUP(RHICmdList, PT_TriangleStrip, 2, Vertices, sizeof(Vertices[0]));
 	}
+	RHICmdList.Flush(); // always call flush with GetRecursiveRHICommandList, recursive use of the RHI is hazardous
 }
 
 void FOpenGLDynamicRHI::InitializeStateResources()
diff --git a/Engine/Source/Runtime/OpenGLDrv/Private/OpenGLDrvPrivate.h b/Engine/Source/Runtime/OpenGLDrv/Private/OpenGLDrvPrivate.h
--- a/Engine/Source/Runtime/OpenGLDrv/Private/OpenGLDrvPrivate.h
+++ b/Engine/Source/Runtime/OpenGLDrv/Private/OpenGLDrvPrivate.h
@@ -463,13 +463,13 @@ class FVector4VertexDeclaration : public FRenderResource
 {
 public:
 	FVertexDeclarationRHIRef VertexDeclarationRHI;
-	virtual void InitRHI()
+	virtual void InitRHI() override
 	{
 		FVertexDeclarationElementList Elements;
 		Elements.Add(FVertexElement(0,0,VET_Float4,0,sizeof(FVector4)));
 		VertexDeclarationRHI = RHICreateVertexDeclaration(Elements);
 	}
-	virtual void ReleaseRHI()
+	virtual void ReleaseRHI() override
 	{
 		VertexDeclarationRHI.SafeRelease();
 	}
diff --git a/Engine/Source/Runtime/OpenGLDrv/Private/OpenGLViewport.cpp b/Engine/Source/Runtime/OpenGLDrv/Private/OpenGLViewport.cpp
--- a/Engine/Source/Runtime/OpenGLDrv/Private/OpenGLViewport.cpp
+++ b/Engine/Source/Runtime/OpenGLDrv/Private/OpenGLViewport.cpp
@@ -69,6 +69,8 @@ void FOpenGLDynamicRHI::RHITick( float DeltaTime )
 
 void FOpenGLDynamicRHI::RHIBeginDrawingViewport(FViewportRHIParamRef ViewportRHI, FTextureRHIParamRef RenderTarget)
 {
+	FRHICommandListImmediate& RHICmdList = FRHICommandListExecutor::GetRecursiveRHICommandList();
+
 	VERIFY_GL_SCOPE();
 
 	DYNAMIC_CAST_OPENGLRESOURCE(Viewport,Viewport);
@@ -92,9 +94,7 @@ void FOpenGLDynamicRHI::RHIBeginDrawingViewport(FViewportRHIParamRef ViewportRHI
 	{
 		GPUProfilingData.FrameTiming.InitResource();
 	}
-	//@todo-rco: RHIPacketList
-	FRHICommandList& RHICmdList = FRHICommandList::GetNullRef();
-
+	
 	// Set the render target and viewport.
 	if( RenderTarget )
 	{
@@ -104,6 +104,7 @@ void FOpenGLDynamicRHI::RHIBeginDrawingViewport(FViewportRHIParamRef ViewportRHI
 	{
 		SetRenderTarget(RHICmdList, DrawingViewport->GetBackBuffer(), FTextureRHIRef());
 	}
+	RHICmdList.Flush(); // always call flush with GetRecursiveRHICommandList, recursive use of the RHI is hazardous
 }
 
 void FOpenGLDynamicRHI::RHIEndDrawingViewport(FViewportRHIParamRef ViewportRHI,bool bPresent,bool bLockToVsync)
diff --git a/Engine/Source/Runtime/OpenGLDrv/Public/OpenGLDrv.h b/Engine/Source/Runtime/OpenGLDrv/Public/OpenGLDrv.h
--- a/Engine/Source/Runtime/OpenGLDrv/Public/OpenGLDrv.h
+++ b/Engine/Source/Runtime/OpenGLDrv/Public/OpenGLDrv.h
@@ -81,12 +81,12 @@ public:
 	/**
 	 * Initializes all OpenGL resources.
 	 */
-	virtual void InitDynamicRHI();
+	virtual void InitDynamicRHI() override;
 
 	/**
 	 * Releases all OpenGL resources.
 	 */
-	virtual void ReleaseDynamicRHI();
+	virtual void ReleaseDynamicRHI() override;
 
 
 private:
@@ -306,7 +306,7 @@ public:
 	virtual void PopEvent();
 
 
-	#define DEFINE_RHIMETHOD(Type,Name,ParameterTypesAndNames,ParameterNames,ReturnStatement,NullImplementation) virtual Type Name ParameterTypesAndNames
+	#define DEFINE_RHIMETHOD(Type,Name,ParameterTypesAndNames,ParameterNames,ReturnStatement,NullImplementation) virtual Type RHI##Name ParameterTypesAndNames
 	#include "RHIMethods.h"
 	#undef DEFINE_RHIMETHOD
 
diff --git a/Engine/Source/Runtime/OpenGLDrv/Public/OpenGLResources.h b/Engine/Source/Runtime/OpenGLDrv/Public/OpenGLResources.h
--- a/Engine/Source/Runtime/OpenGLDrv/Public/OpenGLResources.h
+++ b/Engine/Source/Runtime/OpenGLDrv/Public/OpenGLResources.h
@@ -1257,8 +1257,8 @@ public:
 	void WaitForCompletion();
 
 	// FRenderResource interface.
-	virtual void InitDynamicRHI();
-	virtual void ReleaseDynamicRHI();
+	virtual void InitDynamicRHI() override;
+	virtual void ReleaseDynamicRHI() override;
 
 private:
 	FOpenGLDynamicRHI* OpenGLRHI;
