Reorganize calls to GetWorldContexts

diff --git a/Engine/Source/Editor/Kismet/Private/SBlueprintEditorSelectedDebugObjectWidget.cpp b/Engine/Source/Editor/Kismet/Private/SBlueprintEditorSelectedDebugObjectWidget.cpp
--- a/Engine/Source/Editor/Kismet/Private/SBlueprintEditorSelectedDebugObjectWidget.cpp
+++ b/Engine/Source/Editor/Kismet/Private/SBlueprintEditorSelectedDebugObjectWidget.cpp
@@ -566,13 +566,12 @@ EVisibility SBlueprintEditorSelectedDebugObjectWidget::IsDebugWorldComboVisible(
 {
 	if (GEditor->PlayWorld != NULL)
 	{
-		const TArray<FWorldContext> &WorldContexts = GEngine->GetWorldContexts();
 		int32 LocalWorldCount = 0;
-		for (int32 i = 0; i < WorldContexts.Num() && LocalWorldCount <= 1; ++i)
+		for (const FWorldContext& Context : GEngine->GetWorldContexts())
 		{
-			if (WorldContexts[i].WorldType == EWorldType::PIE && WorldContexts[i].World() != NULL)
+			if (Context.WorldType == EWorldType::PIE && Context.World() != NULL)
 			{
-				LocalWorldCount++;
+				++LocalWorldCount;
 			}
 		}
 
diff --git a/Engine/Source/Editor/PhAT/Private/PhAT.cpp b/Engine/Source/Editor/PhAT/Private/PhAT.cpp
--- a/Engine/Source/Editor/PhAT/Private/PhAT.cpp
+++ b/Engine/Source/Editor/PhAT/Private/PhAT.cpp
@@ -3017,11 +3017,9 @@ bool FPhAT::CanStartSimulation() const
 
 bool FPhAT::IsPIERunning()
 {
-	TArray<FWorldContext> WorldContexts = GEngine->GetWorldContexts();
-
-	for (int32 i=0; i < WorldContexts.Num(); ++i)
+	for (const FWorldContext& Context : GEngine->GetWorldContexts())
 	{
-		if (WorldContexts[i].World()->IsPlayInEditor())
+		if (Context.World()->IsPlayInEditor())
 		{
 			return true;
 		}
diff --git a/Engine/Source/Editor/SceneOutliner/Private/SSceneOutliner.cpp b/Engine/Source/Editor/SceneOutliner/Private/SSceneOutliner.cpp
--- a/Engine/Source/Editor/SceneOutliner/Private/SSceneOutliner.cpp
+++ b/Engine/Source/Editor/SceneOutliner/Private/SSceneOutliner.cpp
@@ -688,17 +688,16 @@ namespace SceneOutliner
 		TGuardValue<bool> ReentrantGuard(bIsReentrant, true);
 	
 
-		const TArray<FWorldContext>& WorldContexts = GEngine->GetWorldContexts();
-		for (int32 WorldIndex=0; WorldIndex < WorldContexts.Num(); ++WorldIndex)
+		for (const FWorldContext& Context : GEngine->GetWorldContexts())
 		{
-			if (WorldContexts[WorldIndex].WorldType == EWorldType::PIE)
+			if (Context.WorldType == EWorldType::PIE)
 			{
-				RepresentingWorld = WorldContexts[WorldIndex].World();
+				RepresentingWorld = Context.World();
 				break;
 			}
-			else if (WorldContexts[WorldIndex].WorldType == EWorldType::Editor)
+			else if (Context.WorldType == EWorldType::Editor)
 			{
-				RepresentingWorld = WorldContexts[WorldIndex].World();
+				RepresentingWorld = Context.World();
 			}
 		}
 		
diff --git a/Engine/Source/Editor/UnrealEd/Private/AssetDeleteModel.cpp b/Engine/Source/Editor/UnrealEd/Private/AssetDeleteModel.cpp
--- a/Engine/Source/Editor/UnrealEd/Private/AssetDeleteModel.cpp
+++ b/Engine/Source/Editor/UnrealEd/Private/AssetDeleteModel.cpp
@@ -291,17 +291,16 @@ bool FAssetDeleteModel::GoToNextReferenceInLevel() const
 
 	UWorld* RepresentingWorld = NULL;
 
-	const TArray<FWorldContext>& WorldContexts = GEngine->GetWorldContexts();
-	for ( int32 WorldIndex=0; WorldIndex < WorldContexts.Num(); ++WorldIndex )
+	for (const FWorldContext& Context : GEngine->GetWorldContexts())
 	{
-		if ( WorldContexts[WorldIndex].WorldType == EWorldType::PIE )
+		if ( Context.WorldType == EWorldType::PIE )
 		{
-			RepresentingWorld = WorldContexts[WorldIndex].World();
+			RepresentingWorld = Context.World();
 			break;
 		}
-		else if ( WorldContexts[WorldIndex].WorldType == EWorldType::Editor )
+		else if ( Context.WorldType == EWorldType::Editor )
 		{
-			RepresentingWorld = WorldContexts[WorldIndex].World();
+			RepresentingWorld = Context.World();
 		}
 	}
 
diff --git a/Engine/Source/Runtime/Engine/Private/AI/Navigation/NavFilters/NavigationQueryFilter.cpp b/Engine/Source/Runtime/Engine/Private/AI/Navigation/NavFilters/NavigationQueryFilter.cpp
--- a/Engine/Source/Runtime/Engine/Private/AI/Navigation/NavFilters/NavigationQueryFilter.cpp
+++ b/Engine/Source/Runtime/Engine/Private/AI/Navigation/NavFilters/NavigationQueryFilter.cpp
@@ -241,10 +241,9 @@ void UNavigationQueryFilter::PostEditChangeProperty(struct FPropertyChangedEvent
 	Super::PostEditChangeProperty(PropertyChangedEvent);
 
 	// remove cached filter settings from existing NavigationSystems
-	const TArray<FWorldContext>& Worlds = GEngine->GetWorldContexts();
-	for (int32 i = 0; i < Worlds.Num(); i++)
+	for (const FWorldContext& Context : GEngine->GetWorldContexts())
 	{
-		UNavigationSystem* NavSys = UNavigationSystem::GetCurrent(Worlds[i].World());
+		UNavigationSystem* NavSys = UNavigationSystem::GetCurrent(Context.World());
 		if (NavSys)
 		{
 			NavSys->ResetCachedFilter(GetClass());
diff --git a/Engine/Source/Runtime/Engine/Private/AI/Navigation/NavigationSystem.cpp b/Engine/Source/Runtime/Engine/Private/AI/Navigation/NavigationSystem.cpp
--- a/Engine/Source/Runtime/Engine/Private/AI/Navigation/NavigationSystem.cpp
+++ b/Engine/Source/Runtime/Engine/Private/AI/Navigation/NavigationSystem.cpp
@@ -116,10 +116,9 @@ void FNavigationLockContext::LockUpdates()
 	}
 	else
 	{
-		const TArray<FWorldContext>& AllContexts = GEngine->GetWorldContexts();
-		for (int32 Idx = 0; Idx < AllContexts.Num(); Idx++)
+		for (const FWorldContext& Context : GEngine->GetWorldContexts())
 		{
-			UNavigationSystem* NavSys = UNavigationSystem::GetCurrent(AllContexts[Idx].World());
+			UNavigationSystem* NavSys = UNavigationSystem::GetCurrent(Context.World());
 			if (NavSys)
 			{
 				NavSys->BeginFakeComponentChanges();
@@ -142,10 +141,9 @@ void FNavigationLockContext::UnlockUpdates()
 	}
 	else
 	{
-		const TArray<FWorldContext>& AllContexts = GEngine->GetWorldContexts();
-		for (int32 Idx = 0; Idx < AllContexts.Num(); Idx++)
+		for (const FWorldContext& Context : GEngine->GetWorldContexts())
 		{
-			UNavigationSystem* NavSys = UNavigationSystem::GetCurrent(AllContexts[Idx].World());
+			UNavigationSystem* NavSys = UNavigationSystem::GetCurrent(Context.World());
 			if (NavSys)
 			{
 				NavSys->EndFakeComponentChanges();
diff --git a/Engine/Source/Runtime/Engine/Private/AI/Navigation/RecastNavMeshGenerator.cpp b/Engine/Source/Runtime/Engine/Private/AI/Navigation/RecastNavMeshGenerator.cpp
--- a/Engine/Source/Runtime/Engine/Private/AI/Navigation/RecastNavMeshGenerator.cpp
+++ b/Engine/Source/Runtime/Engine/Private/AI/Navigation/RecastNavMeshGenerator.cpp
@@ -366,12 +366,11 @@ static UWorld* FindEditorWorld()
 {
 	if (GEngine)
 	{
-		const TArray<FWorldContext>& Worlds = GEngine->GetWorldContexts();
-		for (int32 i = 0; i < Worlds.Num(); i++)
+		for (const FWorldContext& Context : GEngine->GetWorldContexts())
 		{
-			if (Worlds[i].WorldType == EWorldType::Editor)
+			if (Context.WorldType == EWorldType::Editor)
 			{
-				return Worlds[i].World();
+				return Context.World();
 			}
 		}
 	}
diff --git a/Engine/Source/Runtime/Engine/Private/EngineService.cpp b/Engine/Source/Runtime/Engine/Private/EngineService.cpp
--- a/Engine/Source/Runtime/Engine/Private/EngineService.cpp
+++ b/Engine/Source/Runtime/Engine/Private/EngineService.cpp
@@ -81,7 +81,6 @@ void FEngineService::SendPong( const IMessageContextRef& Context )
 			Message->InstanceType = TEXT("Other");
 		}
 
-		const TArray<FWorldContext>& WorldContexts = GEngine->GetWorldContexts();
 		FWorldContext const* ContextToUse = NULL;
 
 		// TODO: Should we be iterating here and sending a message for each context?
@@ -89,10 +88,8 @@ void FEngineService::SendPong( const IMessageContextRef& Context )
 		// We're going to look through the WorldContexts and pull any Game context we find
 		// If there isn't a Game context, we'll take the first PIE we find
 		// and if none of those we'll use an Editor
-		for (int32 WorldIndex = 0; WorldIndex < WorldContexts.Num(); ++WorldIndex)
+		for (const FWorldContext& WorldContext : GEngine->GetWorldContexts())
 		{
-			const FWorldContext& WorldContext = WorldContexts[WorldIndex];
-
 			if (WorldContext.WorldType == EWorldType::Game)
 			{
 				ContextToUse = &WorldContext;
diff --git a/Engine/Source/Runtime/Engine/Private/World.cpp b/Engine/Source/Runtime/Engine/Private/World.cpp
--- a/Engine/Source/Runtime/Engine/Private/World.cpp
+++ b/Engine/Source/Runtime/Engine/Private/World.cpp
@@ -1901,10 +1901,9 @@ void FLevelStreamingGCHelper::PrepareStreamedOutLevelsForGC()
 			UPackage* LevelPackage = Cast<UPackage>(Level->GetOutermost());
 			UE_LOG(LogStreaming, Log, TEXT("PrepareStreamedOutLevelsForGC called on '%s'"), *LevelPackage->GetName() );
 			
-			const TArray<FWorldContext>& WorldList = GEngine->GetWorldContexts();
-			for (int32 ContextIdx = 0; ContextIdx < WorldList.Num(); ++ContextIdx)
+			for (const FWorldContext& Context : GEngine->GetWorldContexts())
 			{
-				UWorld* World = WorldList[ContextIdx].World();
+				UWorld* World = Context.World();
 				if (World)
 				{
 					// This can never ever be called during tick; same goes for GC in general.
@@ -3183,15 +3182,11 @@ ENGINE_API const FString GetMapNameStatic()
 	FWorldContext const* ContextToUse = NULL;
 	if ( GEngine )
 	{
-		const TArray<FWorldContext>& WorldContexts = GEngine->GetWorldContexts();
-
 		// We're going to look through the WorldContexts and pull any Game context we find
 		// If there isn't a Game context, we'll take the first PIE we find
 		// and if none of those we'll use an Editor
-		for (int32 WorldIndex = 0; WorldIndex < WorldContexts.Num(); ++WorldIndex)
+		for (const FWorldContext& WorldContext : GEngine->GetWorldContexts())
 		{
-			const FWorldContext& WorldContext = WorldContexts[WorldIndex];
-
 			if (WorldContext.WorldType == EWorldType::Game)
 			{
 				ContextToUse = &WorldContext;
diff --git a/Engine/Source/Runtime/Launch/Private/LaunchEngineLoop.cpp b/Engine/Source/Runtime/Launch/Private/LaunchEngineLoop.cpp
--- a/Engine/Source/Runtime/Launch/Private/LaunchEngineLoop.cpp
+++ b/Engine/Source/Runtime/Launch/Private/LaunchEngineLoop.cpp
@@ -1975,11 +1975,9 @@ bool FEngineLoop::ShouldUseIdleMode() const
 	{
 		bIdleMode = true;
 
-		const TArray<FWorldContext>& WorldContexts = GEngine->GetWorldContexts();
-
-		for (int32 WorldIndex = 0; WorldIndex < WorldContexts.Num(); ++WorldIndex)
+		for (const FWorldContext& Context : GEngine->GetWorldContexts())
 		{
-			if (!WorldContexts[WorldIndex].World()->AreAlwaysLoadedLevelsLoaded())
+			if (!Context.World()->AreAlwaysLoadedLevelsLoaded())
 			{
 				bIdleMode = false;
 				break;
