#ue4 - Improve collision results when adjacent to or overlapping other geometry.

- Use a slightly inflated shape to compute MTD if possible (currently only for capsule).
- Prefer using inflated MTD against tri-mesh and heightfield before testing for overlapping triangles
- Prefer using inflated MTD against other geometry (convex, capsule, box, sphere) before using point distance check (which is fairly poor).

diff --git a/Engine/Source/Runtime/Engine/Private/Collision/CollisionConversions.cpp b/Engine/Source/Runtime/Engine/Private/Collision/CollisionConversions.cpp
--- a/Engine/Source/Runtime/Engine/Private/Collision/CollisionConversions.cpp
+++ b/Engine/Source/Runtime/Engine/Private/Collision/CollisionConversions.cpp
@@ -462,6 +462,32 @@ FVector FindBestOverlappingNormal(const PxGeometry& Geom, const PxTransform& Que
 	return BestPlaneNormal;
 }
 
+// Compute depenetration vector and distance if possible with a slightly larger geometry
+static bool ComputeInflatedMTD(const float MtdInflation, const PxLocationHit& PHit, FHitResult& OutResult, const PxTransform& QueryTM, const PxGeometry& Geom, const PxTransform& PShapeWorldPose)
+{
+	if (Geom.getType() == PxGeometryType::eCAPSULE)
+	{
+		const PxCapsuleGeometry* InCapsule = static_cast<const PxCapsuleGeometry*>(&Geom);
+		PxCapsuleGeometry InflatedCapsule(InCapsule->radius + MtdInflation, InCapsule->halfHeight + MtdInflation);
+
+		PxVec3 PxMtdNormal(0.f);
+		PxF32 PxMtdDepth = 0.f;
+		const PxGeometry& POtherGeom = PHit.shape->getGeometry().any();
+		const bool bMtdResult = PxGeometryQuery::computePenetration(PxMtdNormal, PxMtdDepth, InflatedCapsule, QueryTM, POtherGeom, PShapeWorldPose);
+		if (bMtdResult)
+		{
+			if (PxMtdNormal.isFinite())
+			{
+				OutResult.ImpactNormal = P2UVector(PxMtdNormal);
+				OutResult.PenetrationDepth = FMath::Max(FMath::Abs(PxMtdDepth) - MtdInflation, 0.f) + KINDA_SMALL_NUMBER;
+				return true;
+			}
+		}
+	}
+
+	return false;
+}
+
 
 /** Util to convert an overlapped shape into a sweep hit result, returns whether it was a blocking hit. */
 static bool ConvertOverlappedShapeToImpactHit(const PxLocationHit& PHit, const FVector& StartLoc, const FVector& EndLoc, FHitResult& OutResult, const PxGeometry& Geom, const PxTransform& QueryTM, const PxFilterData& QueryFilter, bool bReturnPhysMat)
@@ -497,7 +523,6 @@ static bool ConvertOverlappedShapeToImpactHit(const PxLocationHit& PHit, const F
 	const bool bValidMtdNormal = (PHit.flags & PxHitFlag::eNORMAL) && bFiniteNormal;
 	if (bValidMtdNormal)
 	{
-		OutResult.Normal = P2UVector(PHit.normal);
 		OutResult.ImpactNormal = P2UVector(PHit.normal);
 		OutResult.PenetrationDepth = FMath::Abs(PHit.distance);
 	}
@@ -509,14 +534,25 @@ static bool ConvertOverlappedShapeToImpactHit(const PxLocationHit& PHit, const F
 
 	// Zero-distance hits are often valid hits and we can extract the hit normal from the face index
 	// For invalid normals we can try other methods as well (get overlapping triangles).
+	static const float MtdInflation = 0.125f;
+
 	if (PHit.distance == 0.f || !bValidMtdNormal)
 	{
 		const PxTransform PShapeWorldPose = PxShapeExt::getGlobalPose(*PShape, *PActor); 
 		PxTriangleMeshGeometry PTriMeshGeom;
 		PxHeightFieldGeometry PHeightfieldGeom;
+
 		if (PShape->getTriangleMeshGeometry(PTriMeshGeom))
 		{
-			if (!FindGeomOpposingNormal(PHit, TraceDir, OutResult.ImpactNormal))
+			if (FindGeomOpposingNormal(PHit, TraceDir, OutResult.ImpactNormal))
+			{
+				// Success
+			}
+			else if (ComputeInflatedMTD(MtdInflation, PHit, OutResult, QueryTM, Geom, PShapeWorldPose))
+			{
+				// Success
+			}
+			else
 			{
 				PxU32 HitTris[64];
 				bool bOverflow = false;
@@ -529,7 +565,15 @@ static bool ConvertOverlappedShapeToImpactHit(const PxLocationHit& PHit, const F
 		}
 		else if (PShape->getHeightFieldGeometry(PHeightfieldGeom))
 		{
-			if (!FindGeomOpposingNormal(PHit, TraceDir, OutResult.ImpactNormal))
+			if (FindGeomOpposingNormal(PHit, TraceDir, OutResult.ImpactNormal))
+			{
+				// Success
+			}
+			else if (ComputeInflatedMTD(MtdInflation, PHit, OutResult, QueryTM, Geom, PShapeWorldPose))
+			{
+				// Success
+			}
+			else
 			{
 				PxU32 HitTris[64];
 				bool bOverflow = false;
@@ -542,24 +586,31 @@ static bool ConvertOverlappedShapeToImpactHit(const PxLocationHit& PHit, const F
 		}
 		else
 		{
-			// Non tri-mesh or heightfield
-			// Note: faceIndex seems to be unreliable for convex meshes in these cases, so not using FindGeomOpposingNormal() for them here.
-			PxGeometry& PGeom = PShape->getGeometry().any();
-			PxVec3 PClosestPoint;
-			const float Distance = PxGeometryQuery::pointDistance(QueryTM.p, PGeom, PShapeWorldPose, &PClosestPoint);
-
-			if (Distance < KINDA_SMALL_NUMBER)
+			// Not a tri-mesh or heightfield
+			if (ComputeInflatedMTD(MtdInflation, PHit, OutResult, QueryTM, Geom, PShapeWorldPose))
 			{
-				UE_LOG(LogCollision, Warning, TEXT("ConvertOverlappedShapeToImpactHit: Query origin inside shape, giving poor MTD."));
-				PClosestPoint = PxShapeExt::getWorldBounds(*PShape, *PActor).getCenter(); 
+				// Success
 			}
+			else
+			{
+				// MTD failed, use point distance. This is not ideal.
+				// Note: faceIndex seems to be unreliable for convex meshes in these cases, so not using FindGeomOpposingNormal() for them here.
+				PxGeometry& PGeom = PShape->getGeometry().any();
+				PxVec3 PClosestPoint;
+				const float Distance = PxGeometryQuery::pointDistance(QueryTM.p, PGeom, PShapeWorldPose, &PClosestPoint);
 
-			OutResult.ImpactNormal = (OutResult.Location - P2UVector(PClosestPoint)).SafeNormal();
-		}
+				if (Distance < KINDA_SMALL_NUMBER)
+				{
+					UE_LOG(LogCollision, Warning, TEXT("ConvertOverlappedShapeToImpactHit: Query origin inside shape, giving poor MTD."));
+					PClosestPoint = PxShapeExt::getWorldBounds(*PShape, *PActor).getCenter(); 
+				}
 
-		OutResult.Normal = OutResult.ImpactNormal;
+				OutResult.ImpactNormal = (OutResult.Location - P2UVector(PClosestPoint)).SafeNormal();
+			}
+		}
 	}
 
+	OutResult.Normal = OutResult.ImpactNormal;
 	return OutResult.bBlockingHit;
 }
 
