#ue4 - Remove restriction on CharacterMovementComponent when adjusting for penetrations and walking.

- Previously vertical adjustments were eliminated when walking, however now with the improved MTD vector these are safe and much higher quality.

diff --git a/Engine/Source/Runtime/Engine/Classes/GameFramework/CharacterMovementComponent.h b/Engine/Source/Runtime/Engine/Classes/GameFramework/CharacterMovementComponent.h
--- a/Engine/Source/Runtime/Engine/Classes/GameFramework/CharacterMovementComponent.h
+++ b/Engine/Source/Runtime/Engine/Classes/GameFramework/CharacterMovementComponent.h
@@ -1135,9 +1135,6 @@ protected:
 	/** Overridden to set bJustTeleported to true, so we don't make incorrect velocity calculations based on adjusted movement. */
 	virtual bool ResolvePenetration(const FVector& Adjustment, const FHitResult& Hit, const FRotator& NewRotation) override;
 
-	/** Overridden to avoid upward adjustments that can go too high when walking. */
-	virtual FVector GetPenetrationAdjustment(const FHitResult& Hit) const override;
-
 	/** Don't call for intermediate parts of move! */
 	virtual void HandleImpact(FHitResult const& Hit, float TimeSlice=0.f, const FVector& MoveDelta = FVector::ZeroVector) override;
 
diff --git a/Engine/Source/Runtime/Engine/Private/CharacterMovementComponent.cpp b/Engine/Source/Runtime/Engine/Private/CharacterMovementComponent.cpp
--- a/Engine/Source/Runtime/Engine/Private/CharacterMovementComponent.cpp
+++ b/Engine/Source/Runtime/Engine/Private/CharacterMovementComponent.cpp
@@ -1776,21 +1776,6 @@ bool UCharacterMovementComponent::ResolvePenetration(const FVector& Adjustment,
 }
 
 
-FVector UCharacterMovementComponent::GetPenetrationAdjustment(const FHitResult& Hit) const
-{
-	FVector Adjustment = Super::GetPenetrationAdjustment(Hit);
-
-	if (Adjustment.Z > 0.f && IsMovingOnGround())
-	{
-		// Don't allow upward adjustments to move us high off the ground.
-		// Floor height adjustments will take care of penetrations on the lower portion of the capsule.
-		Adjustment.Z = 0.f;
-	}
-
-	return Adjustment;
-}
-
-
 float UCharacterMovementComponent::SlideAlongSurface(const FVector& Delta, float Time, const FVector& InNormal, FHitResult &Hit, bool bHandleImpact)
 {
 	if (!Hit.bBlockingHit)
diff --git a/Engine/Source/Runtime/Engine/Private/PrimitiveComponent.cpp b/Engine/Source/Runtime/Engine/Private/PrimitiveComponent.cpp
--- a/Engine/Source/Runtime/Engine/Private/PrimitiveComponent.cpp
+++ b/Engine/Source/Runtime/Engine/Private/PrimitiveComponent.cpp
@@ -1175,7 +1175,7 @@ static bool ShouldIgnoreHitResult(const UWorld* InWorld, FHitResult const& TestH
 			{
 				UE_LOG(LogTemp, Log, TEXT("Overlapping %s Dir %s Dot %f Normal %s Depth %f"), *GetNameSafe(TestHit.Component.Get()), *MovementDir.ToString(), MoveDot, *TestHit.ImpactNormal.ToString(), TestHit.PenetrationDepth);
 				DrawDebugDirectionalArrow(InWorld, TestHit.TraceStart, TestHit.TraceStart + 30.f * TestHit.ImpactNormal, 5.f, bMovingOut ? FColor(64,128,255) : FColor(255,64,64), true, 4.f);
-				if (TestHit.PenetrationDepth > 0.f)
+				if (TestHit.PenetrationDepth > KINDA_SMALL_NUMBER)
 				{
 					DrawDebugDirectionalArrow(InWorld, TestHit.TraceStart, TestHit.TraceStart + TestHit.PenetrationDepth * TestHit.Normal, 5.f, FColor(64,255,64), true, 4.f);
 				}
