Async network object load support

diff --git a/Engine/Source/Runtime/CoreUObject/Public/UObject/CoreNet.h b/Engine/Source/Runtime/CoreUObject/Public/UObject/CoreNet.h
--- a/Engine/Source/Runtime/CoreUObject/Public/UObject/CoreNet.h
+++ b/Engine/Source/Runtime/CoreUObject/Public/UObject/CoreNet.h
@@ -142,7 +142,7 @@ class COREUOBJECT_API UPackageMap : public UObject
 	FNetworkGUID		GetLastUnmappedNetGUID() const { return LastUnmappedNetGUID; }
 
 	virtual void		LogDebugInfo( FOutputDevice & Ar) { }
-	virtual UObject *	GetObjectFromNetGUID( const FNetworkGUID & NetGUID ) { return NULL; }
+	virtual UObject *	GetObjectFromNetGUID( const FNetworkGUID & NetGUID, const bool bIgnoreMustBeMapped ) { return NULL; }
 
 protected:
 
diff --git a/Engine/Source/Runtime/Engine/Classes/Engine/PackageMapClient.h b/Engine/Source/Runtime/Engine/Classes/Engine/PackageMapClient.h
--- a/Engine/Source/Runtime/Engine/Classes/Engine/PackageMapClient.h
+++ b/Engine/Source/Runtime/Engine/Classes/Engine/PackageMapClient.h
@@ -61,13 +61,13 @@ public:
 	void			RegisterNetGUID_Server( const FNetworkGUID & NetGUID, const UObject * Object );
 	void			RegisterNetGUID_Client( const FNetworkGUID & NetGUID, const UObject * Object );
 	void			RegisterNetGUIDFromPath_Client( const FNetworkGUID & NetGUID, const FString & PathName, const FNetworkGUID & OuterGUID, const bool bNoLoad, const bool bIgnoreWhenMissing );
-	UObject *		GetObjectFromNetGUID( const FNetworkGUID & NetGUID );
+	UObject *		GetObjectFromNetGUID( const FNetworkGUID & NetGUID, const bool bIgnoreMustBeMapped );
 	bool			ShouldIgnoreWhenMissing( const FNetworkGUID & NetGUID );
 	bool			IsGUIDRegistered( const FNetworkGUID & NetGUID );
 	bool			IsGUIDLoaded( const FNetworkGUID & NetGUID );
 
-	UObject *		ResolvePath( const FString & PathName, UObject * ObjOuter, const bool bNoLoad );
-
+	void			AsyncPackageCallback( const FString & PackageName, UPackage * Package );
+	
 	TMap< FNetworkGUID, FNetGuidCacheObject >		ObjectLookup;
 	TMap< TWeakObjectPtr< UObject >, FNetworkGUID >	NetGUIDLookup;
 	int32											UniqueNetIDs[2];
@@ -77,6 +77,8 @@ public:
 
 	UNetDriver *									Driver;
 
+	TMap< FName, FNetworkGUID >						PendingAsyncPackages;
+
 #if !(UE_BUILD_SHIPPING || UE_BUILD_TEST)
 	// History for debugging entries in the guid cache
 	TMap<FNetworkGUID, FString>						History;
@@ -135,12 +137,9 @@ class UPackageMapClient : public UPackageMap
 	virtual bool PrintExportBatch() override;
 
 	virtual void		LogDebugInfo( FOutputDevice & Ar) override;
-	virtual UObject *	GetObjectFromNetGUID( const FNetworkGUID & NetGUID ) override;
+	virtual UObject *	GetObjectFromNetGUID( const FNetworkGUID & NetGUID, const bool bIgnoreMustBeMapped ) override;
 
-	TArray< FNetworkGUID > & GetMustBeMappedGuidsInLastBunch()
-	{
-		return MustBeMappedGuidsInLastBunch;
-	}
+	TArray< FNetworkGUID > & GetMustBeMappedGuidsInLastBunch() { return MustBeMappedGuidsInLastBunch; }
 
 protected:
 
@@ -148,7 +147,7 @@ protected:
 	void	ExportNetGUIDHeader();
 
 	void			InternalWriteObject( FArchive& Ar, FNetworkGUID NetGUID, const UObject * Object, FString ObjectPathName, UObject * ObjectOuter );	
-	FNetworkGUID	InternalLoadObject( FArchive & Ar, UObject *& Object, const bool bIsInnerLevel, int InternalLoadObjectRecursionCount );
+	FNetworkGUID	InternalLoadObject( FArchive & Ar, UObject *& Object, int InternalLoadObjectRecursionCount );
 
 	virtual UObject * ResolvePathAndAssignNetGUID( const FNetworkGUID & NetGUID, const FString & PathName ) override;
 
diff --git a/Engine/Source/Runtime/Engine/Private/DataChannel.cpp b/Engine/Source/Runtime/Engine/Private/DataChannel.cpp
--- a/Engine/Source/Runtime/Engine/Private/DataChannel.cpp
+++ b/Engine/Source/Runtime/Engine/Private/DataChannel.cpp
@@ -1477,21 +1477,19 @@ void UActorChannel::Tick()
 	Super::Tick();
 	
 	// Try to resolve any guids that are holding up the network stream on this channel
-	Connection->Driver->GuidCache->IsExportingNetGUIDBunch = true;		// Hack to make assert happy
 	for ( auto It = PendingGuidResolves.CreateIterator(); It; ++It )
 	{
-		if ( Connection->Driver->GuidCache->GetObjectFromNetGUID( *It ) != NULL )
+		if ( Connection->Driver->GuidCache->GetObjectFromNetGUID( *It, true ) != NULL )
 		{
 			// This guid is now resolved, we can remove it from the pending guid list
 			It.RemoveCurrent();
 		}
 	}
-	Connection->Driver->GuidCache->IsExportingNetGUIDBunch = false;
 
 	// If we don't have any pending guids to resolve, process any queued bunches now
 	if ( PendingGuidResolves.Num() == 0 && QueuedBunches.Num() > 0 )
 	{
-		UE_LOG( LogNet, Warning, TEXT( "UActorChannel::Tick: Purging queued bunches: ChIndex: %i, Actor: %s, Queued: %i" ), ChIndex, Actor != NULL ? *Actor->GetPathName() : TEXT( "NULL" ), QueuedBunches.Num() );
+		UE_LOG( LogNet, Log, TEXT( "UActorChannel::Tick: Flushing queued bunches: ChIndex: %i, Actor: %s, Queued: %i" ), ChIndex, Actor != NULL ? *Actor->GetPathName() : TEXT( "NULL" ), QueuedBunches.Num() );
 
 		for ( int32 i = 0; i < QueuedBunches.Num(); i++ )
 		{
diff --git a/Engine/Source/Runtime/Engine/Private/PackageMapClient.cpp b/Engine/Source/Runtime/Engine/Private/PackageMapClient.cpp
--- a/Engine/Source/Runtime/Engine/Private/PackageMapClient.cpp
+++ b/Engine/Source/Runtime/Engine/Private/PackageMapClient.cpp
@@ -95,7 +95,7 @@ bool UPackageMapClient::SerializeObject( FArchive& Ar, UClass* Class, UObject*&
 		// ----------------	
 		// Read NetGUID from stream and resolve object
 		// ----------------	
-		FNetworkGUID NetGUID = InternalLoadObject( Ar, Object, false, 0 );
+		FNetworkGUID NetGUID = InternalLoadObject( Ar, Object, 0 );
 
 		// Write out NetGUID to caller if necessary
 		if ( OutNetGUID )
@@ -407,9 +407,9 @@ struct FExportFlags
 	{
 		struct
 		{
-			uint8 bHasPath			: 1;
-			uint8 bIsLevel			: 1;
-			uint8 bIsPackage		: 1;
+			uint8 bHasPath		: 1;
+			uint8 bNoLoad		: 1;
+			uint8 bIsPackage	: 1;
 		};
 
 		uint8	Value;
@@ -421,32 +421,18 @@ struct FExportFlags
 	}
 };
 
-static bool CanClientLoadObject_Server( const UObject * Object, const FNetworkGUID & NetGUID )
+static bool CanClientLoadObject( const UObject * Object, const FNetworkGUID & NetGUID )
 {
 	if ( !NetGUID.IsValid() || NetGUID.IsDynamic() )
 	{
-		return false;
+		return false;		// We should never tell the client to load dynamic objects (actors or objects created during play for example)
 	}
 
-	const UPackage * Package = Cast< UPackage >( Object );
-
-	if ( Package != NULL )
+	// PackageMapClient can't load maps, we must wait for the client to load the map when ready
+	// These guids are special guids, where the guid and all child guids resolve once the map has been loaded
+	if ( Object != NULL && Object->GetOutermost()->ContainsMap() )
 	{
-		if ( Package->GetFName() != Package->FileName )
-		{
-			return false;		// Client can't load renamed packages
-		}
-	}
-
-	// We never allow the client to load levels on demand, we wait until the client has streamed the level in
-	while ( Object != NULL )
-	{
-		if ( Cast< const ULevel >( Object ) != NULL )
-		{
-			return false;
-		}
-
-		Object = Object->GetOuter();
+		return false;
 	}
 
 	// We can load everything else
@@ -458,7 +444,9 @@ void UPackageMapClient::InternalWriteObject( FArchive & Ar, FNetworkGUID NetGUID
 {
 	check( Ar.IsSaving() );
 
-	if ( CVarAllowAsyncLoading.GetValueOnGameThread() > 0 && IsNetGUIDAuthority() && !GuidCache->IsExportingNetGUIDBunch && CanClientLoadObject_Server( Object, NetGUID ) )
+	const bool bNoLoad = !CanClientLoadObject( Object, NetGUID );
+
+	if ( CVarAllowAsyncLoading.GetValueOnGameThread() > 0 && IsNetGUIDAuthority() && !GuidCache->IsExportingNetGUIDBunch && !bNoLoad )
 	{
 		// These are guids that must exist on the client in a package
 		// The client needs to know about these so it can determine if it has finished loading them
@@ -499,7 +487,7 @@ void UPackageMapClient::InternalWriteObject( FArchive & Ar, FNetworkGUID NetGUID
 			ExportFlags.bHasPath = ObjectPathName.IsEmpty() ? 0 : 1;
 		}
 
-		ExportFlags.bIsLevel	= ( Cast< const ULevel >( Object ) != NULL ) ? 1 : 0;
+		ExportFlags.bNoLoad		= bNoLoad ? 1 : 0;
 		ExportFlags.bIsPackage	= ( Cast< const UPackage >( Object ) != NULL ) ? 1 : 0;
 
 		Ar << ExportFlags.Value;
@@ -574,7 +562,7 @@ static void SanityCheckExport( const UObject * Object, const FNetworkGUID & NetG
 }
 
 /** Loads a UObject from an FArchive stream. Reads object path if there, and tries to load object if its not already loaded */
-FNetworkGUID UPackageMapClient::InternalLoadObject( FArchive & Ar, UObject *& Object, const bool bIsInnerLevel, const int InternalLoadObjectRecursionCount )
+FNetworkGUID UPackageMapClient::InternalLoadObject( FArchive & Ar, UObject *& Object, const int InternalLoadObjectRecursionCount )
 {
 	if ( InternalLoadObjectRecursionCount > INTERNAL_LOAD_OBJECT_RECURSION_LIMIT ) 
 	{
@@ -608,7 +596,7 @@ FNetworkGUID UPackageMapClient::InternalLoadObject( FArchive & Ar, UObject *& Ob
 	// ----------------	
 	if ( NetGUID.IsValid() && !NetGUID.IsDefault() )
 	{
-		Object = GetObjectFromNetGUID( NetGUID );
+		Object = GetObjectFromNetGUID( NetGUID, GuidCache->IsExportingNetGUIDBunch );
 
 		UE_CLOG( !bSuppressLogs, LogNetPackageMap, Log, TEXT( "InternalLoadObject loaded %s from NetGUID <%s>" ), Object ? *Object->GetFullName() : TEXT( "NULL" ), *NetGUID.ToString() );
 	}
@@ -637,7 +625,7 @@ FNetworkGUID UPackageMapClient::InternalLoadObject( FArchive & Ar, UObject *& Ob
 	{
 		UObject * ObjOuter = NULL;
 
-		FNetworkGUID OuterGUID = InternalLoadObject( Ar, ObjOuter, ExportFlags.bIsLevel || bIsInnerLevel, InternalLoadObjectRecursionCount + 1 );
+		FNetworkGUID OuterGUID = InternalLoadObject( Ar, ObjOuter, InternalLoadObjectRecursionCount + 1 );
 
 		FString PathName;
 		Ar << PathName;
@@ -666,7 +654,7 @@ FNetworkGUID UPackageMapClient::InternalLoadObject( FArchive & Ar, UObject *& Ob
 			// then export that to the client at the next available opportunity
 			check( IsNetGUIDAuthority() );
 
-			Object = GuidCache->ResolvePath( PathName, ObjOuter, false );
+			Object = StaticFindObject( UObject::StaticClass(), ObjOuter, *PathName, false );
 
 			if ( Object == NULL )
 			{
@@ -690,16 +678,17 @@ FNetworkGUID UPackageMapClient::InternalLoadObject( FArchive & Ar, UObject *& Ob
 			return NetGUID;
 		}
 
+		//
 		// At this point, only the client gets this far
-		const bool bIsForLevel			= ExportFlags.bIsLevel || bIsInnerLevel;			// Make note if this is for a level (or any of our inners are)
-		const bool bNoLoad				= bIsForLevel;
-		const bool bIgnoreWhenMissing	= bIsForLevel || GuidCache->ShouldIgnoreWhenMissing( OuterGUID );
+		//
+
+		const bool bIgnoreWhenMissing = ExportFlags.bNoLoad;
 
 		// Register this path and outer guid combo with the net guid
-		GuidCache->RegisterNetGUIDFromPath_Client( NetGUID, PathName, OuterGUID, bNoLoad, bIgnoreWhenMissing );
+		GuidCache->RegisterNetGUIDFromPath_Client( NetGUID, PathName, OuterGUID, ExportFlags.bNoLoad, bIgnoreWhenMissing );
 
 		// Try again now that we've registered the path
-		Object = GuidCache->GetObjectFromNetGUID( NetGUID );
+		Object = GuidCache->GetObjectFromNetGUID( NetGUID, GuidCache->IsExportingNetGUIDBunch );
 
 		if ( Object == NULL && !GuidCache->ShouldIgnoreWhenMissing( NetGUID ) )
 		{
@@ -725,28 +714,8 @@ FNetworkGUID UPackageMapClient::InternalLoadObject( FArchive & Ar, UObject *& Ob
 
 UObject * UPackageMapClient::ResolvePathAndAssignNetGUID( const FNetworkGUID & NetGUID, const FString & PathName )
 {
-	if ( NetGUID.IsDynamic() )
-	{
-		UE_LOG( LogNetPackageMap, Warning, TEXT( "UPackageMapClient::ResolvePathAndAssignNetGUID. Trying to load dynamic NetGUID by path. Path: %s, NetGUID: %s" ), *PathName, *NetGUID.ToString() );
-	}
-
-	if ( NetGUID.IsDefault() )
-	{
-		UE_LOG( LogNetPackageMap, Warning, TEXT( "UPackageMapClient::ResolvePathAndAssignNetGUID. Trying to load default NetGUID by path. Path: %s, NetGUID: %s" ), *PathName, *NetGUID.ToString() );
-		return NULL;
-	}
-
-	UObject * Object = GuidCache->ResolvePath( PathName, NULL, false );
-
-	if ( Object == NULL )
-	{
-		UE_LOG( LogNetPackageMap, Warning, TEXT( "UPackageMapClient::ResolvePathAndAssignNetGUID. Could not resolve path. Path: %s, NetGUID: %s" ), *PathName, *NetGUID.ToString() );
-		return NULL;
-	}
-
-	GuidCache->RegisterNetGUID_Client( NetGUID, Object );
-
-	return Object;
+	check( 0 );
+	return NULL;
 }
 
 //--------------------------------------------------------------------
@@ -936,7 +905,7 @@ void UPackageMapClient::ReceiveNetGUIDBunch( FInBunch &InBunch )
 	while( NumGUIDsRead < NumGUIDsInBunch )
 	{
 		UObject * Obj = NULL;
-		InternalLoadObject( InBunch, Obj, false, 0 );
+		InternalLoadObject( InBunch, Obj, 0 );
 
 		if ( InBunch.IsError() )
 		{
@@ -1317,9 +1286,9 @@ bool UPackageMapClient::PrintExportBatch()
 	return true;
 }
 
-UObject * UPackageMapClient::GetObjectFromNetGUID( const FNetworkGUID & NetGUID )
+UObject * UPackageMapClient::GetObjectFromNetGUID( const FNetworkGUID & NetGUID, const bool bIgnoreMustBeMapped )
 {
-	return GuidCache->GetObjectFromNetGUID( NetGUID );
+	return GuidCache->GetObjectFromNetGUID( NetGUID, bIgnoreMustBeMapped );
 }
 
 //----------------------------------------------------------------------------------------
@@ -1779,7 +1748,45 @@ void FNetGUIDCache::RegisterNetGUIDFromPath_Client( const FNetworkGUID & NetGUID
 	RegisterNetGUID_Internal( NetGUID, CacheObject );
 }
 
-UObject * FNetGUIDCache::GetObjectFromNetGUID( const FNetworkGUID & NetGUID )
+void FNetGUIDCache::AsyncPackageCallback( const FString & PackageName, UPackage * Package )
+{
+	check( Package == NULL || Package->IsFullyLoaded() );
+
+	const FName PackageFName( *PackageName );
+
+	FNetworkGUID NetGUID = PendingAsyncPackages.FindRef( PackageFName );
+
+	PendingAsyncPackages.Remove( PackageFName );
+
+	if ( !NetGUID.IsValid() )
+	{
+		UE_LOG( LogNetPackageMap, Warning, TEXT( "AsyncPackageCallback: Could not find package. Path: %s" ), *PackageName );
+		return;
+	}
+
+	FNetGuidCacheObject * CacheObject = ObjectLookup.Find( NetGUID );
+
+	if ( CacheObject == NULL )
+	{
+		UE_LOG( LogNetPackageMap, Warning, TEXT( "AsyncPackageCallback: Could not find net guid. Path: %s, NetGUID: %s" ), *PackageName, *NetGUID.ToString() );
+		return;
+	}
+
+	if ( !CacheObject->bIsPending )
+	{
+		UE_LOG( LogNetPackageMap, Warning, TEXT( "AsyncPackageCallback: Package wasn't pending. Path: %s, NetGUID: %s" ), *PackageName, *NetGUID.ToString() );
+	}
+
+	CacheObject->bIsPending = false;
+	
+	if ( Package == NULL )
+	{
+		CacheObject->bIsBroken = true;
+		UE_LOG( LogNetPackageMap, Warning, TEXT( "AsyncPackageCallback: Package FAILED to load. Path: %s, NetGUID: %s" ), *PackageName, *NetGUID.ToString() );
+	}
+}
+
+UObject * FNetGUIDCache::GetObjectFromNetGUID( const FNetworkGUID & NetGUID, const bool bIgnoreMustBeMapped )
 {
 	if ( !ensure( NetGUID.IsValid() ) )
 	{
@@ -1810,13 +1817,19 @@ UObject * FNetGUIDCache::GetObjectFromNetGUID( const FNetworkGUID & NetGUID )
 
 	if ( CacheObjectPtr->bIsBroken )
 	{
-		// Don't warn, assume we already have
+		// This object is broken, we know it won't load
+		return NULL;
+	}
+
+	if ( CacheObjectPtr->bIsPending )
+	{
+		// We're not done loading yet
 		return NULL;
 	}
 
 	if ( IsNetGUIDAuthority() )
 	{
-		UE_LOG( LogNetPackageMap, Warning, TEXT( "GetObjectFromNetGUID: Guid with no object on server. Path: %s, NetGUID: %s" ), *CacheObjectPtr->PathName.ToString(), *NetGUID.ToString() );
+		UE_LOG( LogNetPackageMap, Error, TEXT( "GetObjectFromNetGUID: Guid with no object on server. Path: %s, NetGUID: %s" ), *CacheObjectPtr->PathName.ToString(), *NetGUID.ToString() );
 		return NULL;
 	}
 
@@ -1830,7 +1843,7 @@ UObject * FNetGUIDCache::GetObjectFromNetGUID( const FNetworkGUID & NetGUID )
 	if ( CacheObjectPtr->PathName == NAME_None )
 	{
 		// Static guids should have a path
-		UE_LOG( LogNetPackageMap, Warning, TEXT( "GetObjectFromNetGUID: Static guid on client with no path. NetGUID: %s" ), *NetGUID.ToString() );
+		UE_LOG( LogNetPackageMap, Error, TEXT( "GetObjectFromNetGUID: Static guid on client with no path. NetGUID: %s" ), *NetGUID.ToString() );
 		return NULL;
 	}
 #else
@@ -1844,25 +1857,60 @@ UObject * FNetGUIDCache::GetObjectFromNetGUID( const FNetworkGUID & NetGUID )
 
 	// First, resolve the outer
 	UObject * ObjOuter = NULL;
-	
+
 	if ( CacheObjectPtr->OuterGUID.IsValid() )
 	{
-		// We should have an outer, resolve it
-		ObjOuter = GetObjectFromNetGUID( CacheObjectPtr->OuterGUID );
+		// If we get here, we depend on an outer to fully load, don't go further until we know we have a fully loaded outer
+		FNetGuidCacheObject * OuterCacheObject = ObjectLookup.Find( CacheObjectPtr->OuterGUID );
 
-		// If we can't resolve the outer, we need to stop
-		if ( ObjOuter == NULL )
+		if ( OuterCacheObject == NULL )
 		{
-			// If the outer is not pending, we need to warn
+			// Shouldn't be possible, but just in case...
+			UE_LOG( LogNetPackageMap, Error, TEXT( "GetObjectFromNetGUID: Outer not registered. Path: %s, OuterPath: %s, NetGUID: %s" ), *CacheObjectPtr->PathName.ToString(), OuterCacheObject ? *OuterCacheObject->PathName.ToString() : TEXT( "NULL!!" ), *NetGUID.ToString() );
+			CacheObjectPtr->bIsBroken = 1;	// Set this to 1 so that we don't keep spamming
+			return NULL;
+		}
+
+		// If outer is broken, we will never load, set outselves to broken as well and bail
+		if ( OuterCacheObject->bIsBroken )
+		{
+			UE_LOG( LogNetPackageMap, Error, TEXT( "GetObjectFromNetGUID: Outer is broken. Path: %s, Outer: %s, NetGUID: %s" ), *CacheObjectPtr->PathName.ToString(), *OuterCacheObject->PathName.ToString(), *NetGUID.ToString() );
+			CacheObjectPtr->bIsBroken = 1;	// Set this to 1 so that we don't keep spamming
+			return NULL;
+		}
+
+		// Try to resolve the outer
+		ObjOuter = GetObjectFromNetGUID( CacheObjectPtr->OuterGUID, bIgnoreMustBeMapped );
+
+		// If we can't resolve the outer (or the outer package is not fully loaded, and we are async loading), we need to bail
+		if ( ObjOuter == NULL || ( CVarAllowAsyncLoading.GetValueOnGameThread() > 0 && !ObjOuter->GetOutermost()->IsFullyLoaded() ) )
+		{
+			if ( !CacheObjectPtr->bNoLoad && CVarAllowAsyncLoading.GetValueOnGameThread() > 0 )
+			{
+				if ( !bIgnoreMustBeMapped )
+				{
+					// If we get here, we are trying to process a "must be mapped" guid when we shouldn't be (the asset hasn't loaded yet)
+					UE_LOG( LogNetPackageMap, Error, TEXT( "GetObjectFromNetGUID: Missing \"must be mapped\" asset. Path: %s, OuterPath: %s, NetGUID: %s" ), *CacheObjectPtr->PathName.ToString(), *OuterCacheObject->PathName.ToString(), *NetGUID.ToString() );
+				}
+
+				if ( ObjOuter != NULL && !ObjOuter->GetOutermost()->IsFullyLoaded() )
+				{
+					// Not done loading outer
+					return NULL;
+				}
+			}
+
+			// If the outer is missing, warn unless told to ignore
 			if ( !ShouldIgnoreWhenMissing( CacheObjectPtr->OuterGUID ) )
 			{
-				UE_LOG( LogNetPackageMap, Warning, TEXT( "GetObjectFromNetGUID: Failed to find outer. Path: %s, NetGUID: %s" ), *CacheObjectPtr->PathName.ToString(), *NetGUID.ToString() );
+				UE_LOG( LogNetPackageMap, Error, TEXT( "GetObjectFromNetGUID: Failed to find outer. Path: %s, Outer: %s, NetGUID: %s" ), *CacheObjectPtr->PathName.ToString(), *OuterCacheObject->PathName.ToString(), *NetGUID.ToString() );
 			}
 
 			return NULL;
 		}
 	}
 
+	// Simulate async loading if desired
 	if ( CVarAllowAsyncLoading.GetValueOnGameThread() > 0 && CVarSimulateAsyncLoading.GetValueOnGameThread() > 0 && !IsNetGUIDAuthority() && !CacheObjectPtr->OuterGUID.IsValid() )
 	{
 		if ( !CacheObjectPtr->bNoLoad && !CacheObjectPtr->bIsPending && CacheObjectPtr->InitialQueryTime <= 0.0f )
@@ -1878,32 +1926,72 @@ UObject * FNetGUIDCache::GetObjectFromNetGUID( const FNetworkGUID & NetGUID )
 				//UE_LOG( LogNetPackageMap, Warning, TEXT( "GetObjectFromNetGUID: FINISH simulating pending load. Path: %s, NetGUID: %s" ), *CacheObjectPtr->PathName.ToString(), *NetGUID.ToString() );
 				CacheObjectPtr->bIsPending = false;
 			}
-			else
+		}
+	}
+
+	check( !CacheObjectPtr->bIsPending );
+	check( ObjOuter == NULL || ObjOuter->GetOutermost()->IsFullyLoaded() );
+
+	// See if this object is in memory
+	Object = StaticFindObject( UObject::StaticClass(), ObjOuter, *CacheObjectPtr->PathName.ToString(), false );
+
+	// Assume this is a package if the outer is invalid and this is a static guid
+	const bool bIsPackage = NetGUID.IsStatic() && !CacheObjectPtr->OuterGUID.IsValid();
+
+	if ( Object == NULL && !CacheObjectPtr->bNoLoad )
+	{
+		if ( bIsPackage )
+		{
+			// Async load the package if:
+			//	1. We are actually a package
+			//	2. We aren't already pending
+			//	3. We're actually suppose to load (levels don't load here for example)
+			check( !PendingAsyncPackages.Contains( CacheObjectPtr->PathName ) );
+
+			if ( CVarAllowAsyncLoading.GetValueOnGameThread() > 0 && CVarSimulateAsyncLoading.GetValueOnGameThread() == 0 )
 			{
-				// We shouldn't get here unless we are trying to import
-				// The reason is because the actor channel code should be using the "must be mapped" list to pause the stream
-				// until these guids are resolved
-				check( IsExportingNetGUIDBunch );
+				PendingAsyncPackages.Add( CacheObjectPtr->PathName, NetGUID );
+				CacheObjectPtr->bIsPending = true;
+
+				LoadPackageAsync( CacheObjectPtr->PathName.ToString(), FLoadPackageAsyncDelegate::CreateRaw( this, &FNetGUIDCache::AsyncPackageCallback ) );
+
+				UE_LOG( LogNetPackageMap, Log, TEXT( "GetObjectFromNetGUID: Async loading package. Path: %s, NetGUID: %s" ), *CacheObjectPtr->PathName.ToString(), *NetGUID.ToString() );
+
+				// There is nothing else to do except wait on the delegate to tell us this package is done loading
 				return NULL;
 			}
+			else
+			{
+				// Async loading disabled
+				Object = LoadPackage( NULL, *CacheObjectPtr->PathName.ToString(), LOAD_None );
+			}
 		}
-	}
+		else
+		{
+			// If we have a package, but for some reason didn't find the object then do a blocking load as a last attempt
+			// This can happen for a few reasons:
+			//	1. The object was GC'd, but the package wasn't, so we need to reload
+			//	2. Someone else started async loading the outer package, and it's not fully loaded yet
+			Object = StaticLoadObject( UObject::StaticClass(), ObjOuter, *CacheObjectPtr->PathName.ToString(), false );
 
-	// Resolve the path of this object relative to outer
-	Object = ResolvePath( CacheObjectPtr->PathName.ToString(), ObjOuter, CacheObjectPtr->bNoLoad || CacheObjectPtr->bIsPending );
+			if ( CVarAllowAsyncLoading.GetValueOnGameThread() > 0 )
+			{
+				UE_LOG( LogNetPackageMap, Error, TEXT( "GetObjectFromNetGUID: Forced blocking load. Path: %s, NetGUID: %s" ), *CacheObjectPtr->PathName.ToString(), *NetGUID.ToString() );
+			}
+		}
+	}
 
 	if ( Object == NULL )
 	{
-		// If we can't find an object, and it's not pending, warn
-		if ( !CacheObjectPtr->bIgnoreWhenMissing && !CacheObjectPtr->bIsPending )
+		if ( !CacheObjectPtr->bIgnoreWhenMissing )
 		{
 			CacheObjectPtr->bIsBroken = 1;	// Set this to 1 so that we don't keep spamming
-			UE_LOG( LogNetPackageMap, Warning, TEXT( "GetObjectFromNetGUID: Failed to resolve path. Path: %s, Outer: %s, NetGUID: %s" ), *CacheObjectPtr->PathName.ToString(), ObjOuter != NULL ? *ObjOuter->GetPathName() : TEXT( "NULL" ), *NetGUID.ToString() );
+			UE_LOG( LogNetPackageMap, Error, TEXT( "GetObjectFromNetGUID: Failed to resolve path. Path: %s, Outer: %s, NetGUID: %s" ), *CacheObjectPtr->PathName.ToString(), ObjOuter != NULL ? *ObjOuter->GetPathName() : TEXT( "NULL" ), *NetGUID.ToString() );
 		}
 
 		return NULL;
 	}
-
+	
 	// Assign the resolved object to this guid
 	CacheObjectPtr->Object = Object;		
 
@@ -1938,7 +2026,7 @@ bool FNetGUIDCache::ShouldIgnoreWhenMissing( const FNetworkGUID & NetGUID )
 
 	if ( NetGUID.IsDynamic() )
 	{
-		return true;		// Ignore dynamic guids. Things like actors and sub-objects will come in when they are relevant to us, and we don't know when that is
+		return true;		// Ignore missing dynamic guids
 	}
 
 	const FNetGuidCacheObject * CacheObject = ObjectLookup.Find( NetGUID );
@@ -1948,9 +2036,15 @@ bool FNetGUIDCache::ShouldIgnoreWhenMissing( const FNetworkGUID & NetGUID )
 		return false;		// If we haven't been told about this static guid before, we need to warn
 	}
 
-	// Ignore warnings when we explicitly are told to, or if we're pending load (which could also be asynchronously loading)
-	// We will explicitly be told to ignore missing guids for levels for example, since the client is in control of when it will actually load levels
-	return CacheObject->bIgnoreWhenMissing || CacheObject->bIsPending;
+	const FNetGuidCacheObject * OutermostCacheObject = CacheObject;
+
+	while ( OutermostCacheObject->OuterGUID.IsValid() )
+	{
+		OutermostCacheObject = ObjectLookup.Find( OutermostCacheObject->OuterGUID );
+	}
+
+	// Ignore warnings when we explicitly are told to, or if the outer most package is pending
+	return CacheObject->bIgnoreWhenMissing || OutermostCacheObject->bIsPending;
 }
 
 bool FNetGUIDCache::IsGUIDRegistered( const FNetworkGUID & NetGUID )
@@ -1990,54 +2084,6 @@ bool FNetGUIDCache::IsGUIDLoaded( const FNetworkGUID & NetGUID )
 	return CacheObjectPtr->Object != NULL;
 }
 
-UObject * FNetGUIDCache::ResolvePath( const FString & PathName, UObject * ObjOuter, const bool bNoLoad )
-{
-	UObject * Object = StaticFindObject( UObject::StaticClass(), ObjOuter, *PathName, false );
-
-	if ( bNoLoad || Object != NULL )
-	{
-		return Object;
-	}
-
-	if ( IsNetGUIDAuthority() )
-	{
-		// The authority shouldn't be loading resources on demand, at least for now.
-		// This could be garbage or a security risk
-		// Another possibility is in dynamic property replication if the server reads the previously serialized state
-		// that has a now destroyed actor in it.
-		UE_LOG( LogNetPackageMap, Warning, TEXT( "FNetGUIDCache::ResolvePath: Server can't load objects from client. Path: %s" ), *PathName );
-		return NULL;
-	}
-
-	Object = StaticLoadObject( UObject::StaticClass(), ObjOuter, *PathName, NULL, LOAD_NoWarn );
-
-	if ( Object )
-	{
-		UE_LOG( LogNetPackageMap, Log, TEXT( "FNetGUIDCache::ResolvePath: StaticLoadObject. Found: %s" ), Object ? *Object->GetName() : TEXT( "NULL" ) );
-	}
-	else
-	{
-		//
-		// If we failed to load it as an object, try to load it as a package
-		//
-
-		UPackage * PackageOuter = Cast< UPackage >( ObjOuter );
-
-		// If we have an outer at this point, it only makes sense for it to be a package (or NULL)
-		if ( PackageOuter != NULL || ObjOuter == NULL )
-		{
-			Object = LoadPackage( PackageOuter, *PathName, LOAD_None );
-			UE_LOG( LogNetPackageMap, Log, TEXT( "FNetGUIDCache::ResolvePath: LoadPackage: %s. Found: %s" ), *PathName, Object ? *Object->GetName() : TEXT( "NULL" ) );
-		} 
-		else
-		{
-			UE_LOG( LogNetPackageMap, Warning, TEXT( "FNetGUIDCache::ResolvePath: Outer %s is not a package for object %s" ), ObjOuter ? *ObjOuter->GetName() : TEXT( "NULL" ), *PathName );
-		}
-	}
-
-	return Object;
-}
-
 //------------------------------------------------------
 // Debug command to see how many times we've exported each NetGUID
 // Used for measuring inefficiencies. Some duplication is unavoidable since we cannot garuntee atomicicity across multiple channels.
diff --git a/Engine/Source/Runtime/Engine/Private/RepLayout.cpp b/Engine/Source/Runtime/Engine/Private/RepLayout.cpp
--- a/Engine/Source/Runtime/Engine/Private/RepLayout.cpp
+++ b/Engine/Source/Runtime/Engine/Private/RepLayout.cpp
@@ -1347,7 +1347,7 @@ void FRepLayout::UpdateUnmappedObjects_r(
 
 		check( Cmd.Type == REPCMD_PropertyObject );
 
-		UObject * Object = PackageMap->GetObjectFromNetGUID( It.Value().Guid );
+		UObject * Object = PackageMap->GetObjectFromNetGUID( It.Value().Guid, false );
 
 		if ( Object != NULL )
 		{
