Compile fix

diff --git a/Engine/Source/Runtime/Engine/Private/PackageMapClient.cpp b/Engine/Source/Runtime/Engine/Private/PackageMapClient.cpp
--- a/Engine/Source/Runtime/Engine/Private/PackageMapClient.cpp
+++ b/Engine/Source/Runtime/Engine/Private/PackageMapClient.cpp
@@ -1972,7 +1972,7 @@ UObject * FNetGUIDCache::GetObjectFromNetGUID( const FNetworkGUID & NetGUID, con
 			// This can happen for a few reasons:
 			//	1. The object was GC'd, but the package wasn't, so we need to reload
 			//	2. Someone else started async loading the outer package, and it's not fully loaded yet
-			Object = StaticLoadObject( UObject::StaticClass(), ObjOuter, *CacheObjectPtr->PathName.ToString(), false );
+			Object = StaticLoadObject( UObject::StaticClass(), ObjOuter, *CacheObjectPtr->PathName.ToString(), NULL, LOAD_NoWarn );
 
 			if ( CVarAllowAsyncLoading.GetValueOnGameThread() > 0 )
 			{
