fixed some AI code style inconsistencies #UE4

diff --git a/Engine/Source/Runtime/Engine/Private/AI/Navigation/RecastNavMeshGenerator.cpp b/Engine/Source/Runtime/Engine/Private/AI/Navigation/RecastNavMeshGenerator.cpp
--- a/Engine/Source/Runtime/Engine/Private/AI/Navigation/RecastNavMeshGenerator.cpp
+++ b/Engine/Source/Runtime/Engine/Private/AI/Navigation/RecastNavMeshGenerator.cpp
@@ -2032,7 +2032,9 @@ void FRecastTileGenerator::AppendModifier(const FCompositeNavModifier& Modifier,
 void FRecastTileGenerator::AppendGeometry(const TNavStatArray<uint8>& RawCollisionCache)
 {
 	if (RawCollisionCache.Num() == 0)
+	{
 		return;
+	}
 
 	FRecastGeometryCache CollisionCache(RawCollisionCache.GetTypedData());
 
@@ -2067,7 +2069,9 @@ void FRecastTileGenerator::AppendGeometry(const TNavStatArray<uint8>& RawCollisi
 void FRecastTileGenerator::AppendGeometry(const TNavStatArray<FVector>& Verts, const TNavStatArray<int32>& Faces)
 {
 	if (Faces.Num() == 0 || Verts.Num() == 0)
+	{
 		return;
+	}
 
 #if CACHE_NAV_GENERATOR_DATA
 	const int32 FirstNewCoord = GeomCoords.Num();
@@ -3103,16 +3107,6 @@ FRecastNavMeshGenerator::FRecastNavMeshGenerator(class ARecastNavMesh* InDestNav
 
 	check(InDestNavMesh);
 	check(InDestNavMesh->GetWorld());
-
-//#if SEAMLESS_REBUILDING_ENABLED
-//	// register when everything is set up
-//	FSimpleDelegateGraphTask::CreateAndDispatchWhenReady(
-//		FSimpleDelegateGraphTask::FDelegate::CreateSP(this, &FRecastNavMeshGenerator::RegisterForCallbacks)
-//		, TEXT("Register for callback")
-//		, NULL
-//		, ENamedThreads::GameThread
-//		);
-//#endif // SEAMLESS_REBUILDING_ENABLED
 }
 
 FRecastNavMeshGenerator::~FRecastNavMeshGenerator()
