REGRESSION: FOLIAGE: Dynamic shadows do not show up unless the foliage is selected TTP# 340091

diff --git a/Engine/Source/Runtime/Engine/Private/InstancedStaticMesh.cpp b/Engine/Source/Runtime/Engine/Private/InstancedStaticMesh.cpp
--- a/Engine/Source/Runtime/Engine/Private/InstancedStaticMesh.cpp
+++ b/Engine/Source/Runtime/Engine/Private/InstancedStaticMesh.cpp
@@ -914,6 +914,9 @@ public:
 
 	/** Draw the scene proxy as a dynamic element */
 	virtual void DrawDynamicElements(FPrimitiveDrawInterface* PDI,const FSceneView* View) override;
+		
+	/** Sets up a shadow FMeshBatch for a specific LOD. */
+	virtual bool GetShadowMeshElement(int32 LODIndex, uint8 InDepthPriorityGroup, FMeshBatch& OutMeshElement) const override;
 
 	/** Sets up a FMeshBatch for a specific LOD and element. */
 	virtual bool GetMeshElement(int32 LODIndex,int32 ElementIndex,uint8 InDepthPriorityGroup,FMeshBatch& OutMeshElement, const bool bUseSelectedMaterial, const bool bUseHoveredMaterial) const override;
@@ -1025,6 +1028,18 @@ void FInstancedStaticMeshSceneProxy::DrawDynamicElements(FPrimitiveDrawInterface
 #endif
 }
 
+bool FInstancedStaticMeshSceneProxy::GetShadowMeshElement(int32 LODIndex, uint8 InDepthPriorityGroup, FMeshBatch& OutMeshElement) const
+{
+	if (LODIndex < InstancedRenderData.VertexFactories.Num() && FStaticMeshSceneProxy::GetShadowMeshElement(LODIndex, InDepthPriorityGroup, OutMeshElement))
+	{
+		OutMeshElement.Elements[0].NumInstances = InstancedRenderData.InstanceBuffer.GetNumInstances();
+		OutMeshElement.Elements[0].UserData = (void*)&UserData_AllInstances;
+		OutMeshElement.VertexFactory = &InstancedRenderData.VertexFactories[LODIndex];
+		return true;
+	}
+	return false;
+}
+
 
 /** Sets up a FMeshBatch for a specific LOD and element. */
 bool FInstancedStaticMeshSceneProxy::GetMeshElement(int32 LODIndex,int32 ElementIndex,uint8 InDepthPriorityGroup,FMeshBatch& OutMeshElement, const bool bUseSelectedMaterial, const bool bUseHoveredMaterial) const
