Remove superfluous headers, part 1 of n.

diff --git a/Engine/Source/Developer/LogVisualizer/Private/LogVisualizer.cpp b/Engine/Source/Developer/LogVisualizer/Private/LogVisualizer.cpp
--- a/Engine/Source/Developer/LogVisualizer/Private/LogVisualizer.cpp
+++ b/Engine/Source/Developer/LogVisualizer/Private/LogVisualizer.cpp
@@ -1,9 +1,10 @@
 // Copyright 1998-2014 Epic Games, Inc. All Rights Reserved.
 
 #include "LogVisualizerPCH.h"
-#include "CollisionDebugDrawingPublic.h"
 
 #if ENABLE_VISUAL_LOG
+#include "VisualLog.h"
+
 //////////////////////////////////////////////////////////////////////////
 
 void FLogVisualizer::SummonUI(UWorld* InWorld) 
diff --git a/Engine/Source/Developer/LogVisualizer/Private/LogVisualizer.h b/Engine/Source/Developer/LogVisualizer/Private/LogVisualizer.h
--- a/Engine/Source/Developer/LogVisualizer/Private/LogVisualizer.h
+++ b/Engine/Source/Developer/LogVisualizer/Private/LogVisualizer.h
@@ -2,10 +2,10 @@
 
 #pragma once
 
-#include "VisualLog.h"
-
 #if ENABLE_VISUAL_LOG
 
+struct FActorsVisLog;
+
 /** Actual implementation of LogVisualizer, private inside module */
 class FLogVisualizer : public ILogVisualizer
 {
diff --git a/Engine/Source/Developer/LogVisualizer/Private/SLogVisualizer.cpp b/Engine/Source/Developer/LogVisualizer/Private/SLogVisualizer.cpp
--- a/Engine/Source/Developer/LogVisualizer/Private/SLogVisualizer.cpp
+++ b/Engine/Source/Developer/LogVisualizer/Private/SLogVisualizer.cpp
@@ -1,6 +1,8 @@
 // Copyright 1998-2014 Epic Games, Inc. All Rights Reserved.
 
 #include "LogVisualizerPCH.h"
+#if ENABLE_VISUAL_LOG
+
 #include "SLogBar.h"
 #include "Debug/LogVisualizerCameraController.h"
 #include "Debug/ReporterGraph.h"
@@ -22,9 +24,8 @@
 #include "GameplayDebuggingComponent.h"
 
 #include "SFilterList.h"
-
-#if ENABLE_VISUAL_LOG
-
+#include "STimeLine.h"
+#include "Interfaces/IMainFrameModule.h"
 #define LOCTEXT_NAMESPACE "SLogVisualizer"
 
 const FName SLogVisualizer::NAME_LogName = TEXT("LogName");
diff --git a/Engine/Source/Developer/LogVisualizer/Private/SLogVisualizer.h b/Engine/Source/Developer/LogVisualizer/Private/SLogVisualizer.h
--- a/Engine/Source/Developer/LogVisualizer/Private/SLogVisualizer.h
+++ b/Engine/Source/Developer/LogVisualizer/Private/SLogVisualizer.h
@@ -2,10 +2,10 @@
 
 #pragma once
 
-#include "STimeline.h"
 #include "Debug/DebugDrawService.h"
 
 class SLogFilterList;
+class STimeline;
 
 struct FLogsListItem
 {
diff --git a/Engine/Source/Developer/LogVisualizer/Private/SLogsTableRow.cpp b/Engine/Source/Developer/LogVisualizer/Private/SLogsTableRow.cpp
--- a/Engine/Source/Developer/LogVisualizer/Private/SLogsTableRow.cpp
+++ b/Engine/Source/Developer/LogVisualizer/Private/SLogsTableRow.cpp
@@ -1,9 +1,12 @@
 // Copyright 1998-2014 Epic Games, Inc. All Rights Reserved.
 
 #include "LogVisualizerPCH.h"
-#include "SLogBar.h"
 
 #if ENABLE_VISUAL_LOG
+
+#include "SLogBar.h"
+#include "STimeLine.h"
+
 void SLogsTableRow::Construct(const FArguments& InArgs, const TSharedRef<STableViewBase>& InOwnerTableView)
 {
 	Item = InArgs._Item;
diff --git a/Engine/Source/Editor/UnrealEd/Classes/ActorFactories/ActorFactory.h b/Engine/Source/Editor/UnrealEd/Classes/ActorFactories/ActorFactory.h
--- a/Engine/Source/Editor/UnrealEd/Classes/ActorFactories/ActorFactory.h
+++ b/Engine/Source/Editor/UnrealEd/Classes/ActorFactories/ActorFactory.h
@@ -4,7 +4,10 @@
 #pragma once
 #include "ActorFactory.generated.h"
 
-#pragma once
+class ULevel;
+class UBlueprint;
+class FAssetData;
+class FRotator;
 
 UCLASS(collapsecategories, hidecategories=Object, editinlinenew, config=Editor, abstract, transient)
 class UNREALED_API UActorFactory : public UObject
diff --git a/Engine/Source/Editor/UnrealEd/Classes/ActorFactories/ActorFactoryAmbientSound.h b/Engine/Source/Editor/UnrealEd/Classes/ActorFactories/ActorFactoryAmbientSound.h
--- a/Engine/Source/Editor/UnrealEd/Classes/ActorFactories/ActorFactoryAmbientSound.h
+++ b/Engine/Source/Editor/UnrealEd/Classes/ActorFactories/ActorFactoryAmbientSound.h
@@ -7,6 +7,7 @@
  * The base class of all ambient sound types
  */
 
+#include "ActorFactories/ActorFactory.h"
 #include "ActorFactoryAmbientSound.generated.h"
 
 UCLASS(MinimalAPI, config=Editor, collapsecategories, hidecategories=Object)
diff --git a/Engine/Source/Editor/UnrealEd/Classes/ActorFactories/ActorFactoryAnimationAsset.h b/Engine/Source/Editor/UnrealEd/Classes/ActorFactories/ActorFactoryAnimationAsset.h
--- a/Engine/Source/Editor/UnrealEd/Classes/ActorFactories/ActorFactoryAnimationAsset.h
+++ b/Engine/Source/Editor/UnrealEd/Classes/ActorFactories/ActorFactoryAnimationAsset.h
@@ -2,8 +2,11 @@
 
 
 #pragma once
+#include "ActorFactories/ActorFactory.h"
 #include "ActorFactoryAnimationAsset.generated.h"
 
+class USkeletalMesh;
+
 UCLASS(MinimalAPI, config=Editor, hidecategories=Object)
 class UActorFactoryAnimationAsset : public UActorFactorySkeletalMesh
 {
diff --git a/Engine/Source/Editor/UnrealEd/Classes/ActorFactories/ActorFactoryBlueprint.h b/Engine/Source/Editor/UnrealEd/Classes/ActorFactories/ActorFactoryBlueprint.h
--- a/Engine/Source/Editor/UnrealEd/Classes/ActorFactories/ActorFactoryBlueprint.h
+++ b/Engine/Source/Editor/UnrealEd/Classes/ActorFactories/ActorFactoryBlueprint.h
@@ -2,6 +2,7 @@
 
 
 #pragma once
+#include "ActorFactories/ActorFactory.h"
 #include "ActorFactoryBlueprint.generated.h"
 
 UCLASS(MinimalAPI, config=Editor, collapsecategories, hidecategories=Object)
diff --git a/Engine/Source/Editor/UnrealEd/Classes/ActorFactories/ActorFactoryBoxVolume.h b/Engine/Source/Editor/UnrealEd/Classes/ActorFactories/ActorFactoryBoxVolume.h
--- a/Engine/Source/Editor/UnrealEd/Classes/ActorFactories/ActorFactoryBoxVolume.h
+++ b/Engine/Source/Editor/UnrealEd/Classes/ActorFactories/ActorFactoryBoxVolume.h
@@ -2,6 +2,7 @@
 
 
 #pragma once
+#include "ActorFactories/ActorFactory.h"
 #include "ActorFactoryBoxVolume.generated.h"
 
 UCLASS(MinimalAPI, config=Editor, collapsecategories, hidecategories=Object)
diff --git a/Engine/Source/Editor/UnrealEd/Classes/ActorFactories/ActorFactoryClass.h b/Engine/Source/Editor/UnrealEd/Classes/ActorFactories/ActorFactoryClass.h
--- a/Engine/Source/Editor/UnrealEd/Classes/ActorFactories/ActorFactoryClass.h
+++ b/Engine/Source/Editor/UnrealEd/Classes/ActorFactories/ActorFactoryClass.h
@@ -2,6 +2,7 @@
 
 
 #pragma once
+#include "ActorFactories/ActorFactory.h"
 #include "ActorFactoryClass.generated.h"
 
 UCLASS(MinimalAPI, config=Editor, collapsecategories, hidecategories=Object)
diff --git a/Engine/Source/Editor/UnrealEd/Classes/ActorFactories/ActorFactoryCylinderVolume.h b/Engine/Source/Editor/UnrealEd/Classes/ActorFactories/ActorFactoryCylinderVolume.h
--- a/Engine/Source/Editor/UnrealEd/Classes/ActorFactories/ActorFactoryCylinderVolume.h
+++ b/Engine/Source/Editor/UnrealEd/Classes/ActorFactories/ActorFactoryCylinderVolume.h
@@ -2,6 +2,7 @@
 
 
 #pragma once
+#include "ActorFactories/ActorFactory.h"
 #include "ActorFactoryCylinderVolume.generated.h"
 
 UCLASS(MinimalAPI, config=Editor, collapsecategories, hidecategories=Object)
diff --git a/Engine/Source/Editor/UnrealEd/Classes/ActorFactories/ActorFactoryDeferredDecal.h b/Engine/Source/Editor/UnrealEd/Classes/ActorFactories/ActorFactoryDeferredDecal.h
--- a/Engine/Source/Editor/UnrealEd/Classes/ActorFactories/ActorFactoryDeferredDecal.h
+++ b/Engine/Source/Editor/UnrealEd/Classes/ActorFactories/ActorFactoryDeferredDecal.h
@@ -2,6 +2,7 @@
 
 
 #pragma once
+#include "ActorFactories/ActorFactory.h"
 #include "ActorFactoryDeferredDecal.generated.h"
 
 UCLASS(MinimalAPI, config=Editor)
diff --git a/Engine/Source/Editor/UnrealEd/Classes/ActorFactories/ActorFactoryDestructible.h b/Engine/Source/Editor/UnrealEd/Classes/ActorFactories/ActorFactoryDestructible.h
--- a/Engine/Source/Editor/UnrealEd/Classes/ActorFactories/ActorFactoryDestructible.h
+++ b/Engine/Source/Editor/UnrealEd/Classes/ActorFactories/ActorFactoryDestructible.h
@@ -2,6 +2,7 @@
 
 
 #pragma once
+#include "ActorFactories/ActorFactory.h"
 #include "ActorFactoryDestructible.generated.h"
 
 UCLASS(MinimalAPI, config=Editor)
diff --git a/Engine/Source/Editor/UnrealEd/Classes/ActorFactories/ActorFactoryEmitter.h b/Engine/Source/Editor/UnrealEd/Classes/ActorFactories/ActorFactoryEmitter.h
--- a/Engine/Source/Editor/UnrealEd/Classes/ActorFactories/ActorFactoryEmitter.h
+++ b/Engine/Source/Editor/UnrealEd/Classes/ActorFactories/ActorFactoryEmitter.h
@@ -2,6 +2,7 @@
 
 
 #pragma once
+#include "ActorFactories/ActorFactory.h"
 #include "ActorFactoryEmitter.generated.h"
 
 UCLASS(MinimalAPI, config=Editor, collapsecategories, hidecategories=Object)
diff --git a/Engine/Source/Editor/UnrealEd/Classes/ActorFactories/ActorFactoryInteractiveFoliage.h b/Engine/Source/Editor/UnrealEd/Classes/ActorFactories/ActorFactoryInteractiveFoliage.h
--- a/Engine/Source/Editor/UnrealEd/Classes/ActorFactories/ActorFactoryInteractiveFoliage.h
+++ b/Engine/Source/Editor/UnrealEd/Classes/ActorFactories/ActorFactoryInteractiveFoliage.h
@@ -2,6 +2,7 @@
 
 
 #pragma once
+#include "ActorFactories/ActorFactoryStaticMesh.h"
 #include "ActorFactoryInteractiveFoliage.generated.h"
 
 UCLASS(MinimalAPI, config=Editor)
diff --git a/Engine/Source/Editor/UnrealEd/Classes/ActorFactories/ActorFactoryMatineeActor.h b/Engine/Source/Editor/UnrealEd/Classes/ActorFactories/ActorFactoryMatineeActor.h
--- a/Engine/Source/Editor/UnrealEd/Classes/ActorFactories/ActorFactoryMatineeActor.h
+++ b/Engine/Source/Editor/UnrealEd/Classes/ActorFactories/ActorFactoryMatineeActor.h
@@ -2,6 +2,7 @@
 
 
 #pragma once
+#include "ActorFactories/ActorFactory.h"
 #include "ActorFactoryMatineeActor.generated.h"
 
 UCLASS(MinimalAPI, config=Editor)
diff --git a/Engine/Source/Editor/UnrealEd/Classes/ActorFactories/ActorFactoryPhysicsAsset.h b/Engine/Source/Editor/UnrealEd/Classes/ActorFactories/ActorFactoryPhysicsAsset.h
--- a/Engine/Source/Editor/UnrealEd/Classes/ActorFactories/ActorFactoryPhysicsAsset.h
+++ b/Engine/Source/Editor/UnrealEd/Classes/ActorFactories/ActorFactoryPhysicsAsset.h
@@ -2,6 +2,7 @@
 
 
 #pragma once
+#include "ActorFactories/ActorFactory.h"
 #include "ActorFactoryPhysicsAsset.generated.h"
 
 UCLASS(MinimalAPI, config=Editor, collapsecategories, hidecategories=Object)
diff --git a/Engine/Source/Editor/UnrealEd/Classes/ActorFactories/ActorFactorySkeletalMesh.h b/Engine/Source/Editor/UnrealEd/Classes/ActorFactories/ActorFactorySkeletalMesh.h
--- a/Engine/Source/Editor/UnrealEd/Classes/ActorFactories/ActorFactorySkeletalMesh.h
+++ b/Engine/Source/Editor/UnrealEd/Classes/ActorFactories/ActorFactorySkeletalMesh.h
@@ -2,8 +2,11 @@
 
 
 #pragma once
+#include "ActorFactories/ActorFactory.h"
 #include "ActorFactorySkeletalMesh.generated.h"
 
+class USkeletalMesh;
+
 UCLASS(MinimalAPI, config=Editor)
 class UActorFactorySkeletalMesh : public UActorFactory
 {
diff --git a/Engine/Source/Editor/UnrealEd/Classes/ActorFactories/ActorFactorySphereVolume.h b/Engine/Source/Editor/UnrealEd/Classes/ActorFactories/ActorFactorySphereVolume.h
--- a/Engine/Source/Editor/UnrealEd/Classes/ActorFactories/ActorFactorySphereVolume.h
+++ b/Engine/Source/Editor/UnrealEd/Classes/ActorFactories/ActorFactorySphereVolume.h
@@ -2,6 +2,7 @@
 
 
 #pragma once
+#include "ActorFactories/ActorFactory.h"
 #include "ActorFactorySphereVolume.generated.h"
 
 UCLASS(MinimalAPI, config=Editor, collapsecategories, hidecategories=Object)
diff --git a/Engine/Source/Editor/UnrealEd/Classes/ActorFactories/ActorFactoryStaticMesh.h b/Engine/Source/Editor/UnrealEd/Classes/ActorFactories/ActorFactoryStaticMesh.h
--- a/Engine/Source/Editor/UnrealEd/Classes/ActorFactories/ActorFactoryStaticMesh.h
+++ b/Engine/Source/Editor/UnrealEd/Classes/ActorFactories/ActorFactoryStaticMesh.h
@@ -2,6 +2,7 @@
 
 
 #pragma once
+#include "ActorFactories/ActorFactory.h"
 #include "ActorFactoryStaticMesh.generated.h"
 
 UCLASS(MinimalAPI, config=Editor)
diff --git a/Engine/Source/Editor/UnrealEd/Classes/Animation/DebugSkelMeshComponent.h b/Engine/Source/Editor/UnrealEd/Classes/Animation/DebugSkelMeshComponent.h
--- a/Engine/Source/Editor/UnrealEd/Classes/Animation/DebugSkelMeshComponent.h
+++ b/Engine/Source/Editor/UnrealEd/Classes/Animation/DebugSkelMeshComponent.h
@@ -2,6 +2,7 @@
 
 
 #pragma once
+#include "Components/SkeletalMeshComponent.h"
 #include "DebugSkelMeshComponent.generated.h"
 
 USTRUCT()
diff --git a/Engine/Source/Editor/UnrealEd/Classes/Animation/EditorAnimBaseObj.h b/Engine/Source/Editor/UnrealEd/Classes/Animation/EditorAnimBaseObj.h
--- a/Engine/Source/Editor/UnrealEd/Classes/Animation/EditorAnimBaseObj.h
+++ b/Engine/Source/Editor/UnrealEd/Classes/Animation/EditorAnimBaseObj.h
@@ -10,6 +10,7 @@
 
 DECLARE_DELEGATE_TwoParams( FOnAnimObjectChange, class UObject*, bool)
 
+class UAnimSequenceBase;
 
 UCLASS(abstract, MinimalAPI)
 class UEditorAnimBaseObj: public UObject
diff --git a/Engine/Source/Editor/UnrealEd/Classes/Animation/EditorAnimCompositeSegment.h b/Engine/Source/Editor/UnrealEd/Classes/Animation/EditorAnimCompositeSegment.h
--- a/Engine/Source/Editor/UnrealEd/Classes/Animation/EditorAnimCompositeSegment.h
+++ b/Engine/Source/Editor/UnrealEd/Classes/Animation/EditorAnimCompositeSegment.h
@@ -6,6 +6,8 @@
  */
 
 #pragma once
+#include "Animation/EditorAnimBaseObj.h"
+#include "Animation/AnimCompositeBase.h"
 #include "EditorAnimCompositeSegment.generated.h"
 
 DECLARE_DELEGATE_OneParam( FOnAnimSegmentChanged, class UEditorAnimSegment*)
diff --git a/Engine/Source/Editor/UnrealEd/Classes/Animation/EditorAnimSegment.h b/Engine/Source/Editor/UnrealEd/Classes/Animation/EditorAnimSegment.h
--- a/Engine/Source/Editor/UnrealEd/Classes/Animation/EditorAnimSegment.h
+++ b/Engine/Source/Editor/UnrealEd/Classes/Animation/EditorAnimSegment.h
@@ -6,6 +6,8 @@
  */
 
 #pragma once
+#include "Animation/EditorAnimBaseObj.h"
+#include "Animation/AnimCompositeBase.h"
 #include "EditorAnimSegment.generated.h"
 
 DECLARE_DELEGATE_OneParam( FOnAnimSegmentChanged, class UEditorAnimSegment*)
diff --git a/Engine/Source/Editor/UnrealEd/Classes/Animation/EditorBranchPoint.h b/Engine/Source/Editor/UnrealEd/Classes/Animation/EditorBranchPoint.h
--- a/Engine/Source/Editor/UnrealEd/Classes/Animation/EditorBranchPoint.h
+++ b/Engine/Source/Editor/UnrealEd/Classes/Animation/EditorBranchPoint.h
@@ -6,6 +6,8 @@
  */
 
 #pragma once
+#include "Animation/EditorAnimBaseObj.h"
+#include "Animation/AnimMontage.h"
 #include "EditorBranchPoint.generated.h"
 
 
diff --git a/Engine/Source/Editor/UnrealEd/Classes/Animation/EditorCompositeSection.h b/Engine/Source/Editor/UnrealEd/Classes/Animation/EditorCompositeSection.h
--- a/Engine/Source/Editor/UnrealEd/Classes/Animation/EditorCompositeSection.h
+++ b/Engine/Source/Editor/UnrealEd/Classes/Animation/EditorCompositeSection.h
@@ -6,6 +6,8 @@
  */
 
 #pragma once
+#include "Animation/EditorAnimBaseObj.h"
+#include "Animation/AnimMontage.h"
 #include "EditorCompositeSection.generated.h"
 
 UCLASS(hidecategories=UObject, MinimalAPI, BlueprintType)
diff --git a/Engine/Source/Editor/UnrealEd/Classes/Animation/EditorParentPlayerListObj.h b/Engine/Source/Editor/UnrealEd/Classes/Animation/EditorParentPlayerListObj.h
--- a/Engine/Source/Editor/UnrealEd/Classes/Animation/EditorParentPlayerListObj.h
+++ b/Engine/Source/Editor/UnrealEd/Classes/Animation/EditorParentPlayerListObj.h
@@ -3,6 +3,9 @@
 #pragma once
 #include "EditorParentPlayerListObj.generated.h"
 
+struct FAnimParentNodeAssetOverride;
+class UAnimGraphNode_Base;
+
 UCLASS()
 class UNREALED_API UEditorParentPlayerListObj : public UObject
 {
diff --git a/Engine/Source/Editor/UnrealEd/Classes/AutoReimport/AutoReimportManager.h b/Engine/Source/Editor/UnrealEd/Classes/AutoReimport/AutoReimportManager.h
--- a/Engine/Source/Editor/UnrealEd/Classes/AutoReimport/AutoReimportManager.h
+++ b/Engine/Source/Editor/UnrealEd/Classes/AutoReimport/AutoReimportManager.h
@@ -1,6 +1,7 @@
 // Copyright 1998-2014 Epic Games, Inc. All Rights Reserved.
 
 #pragma once
+#include "TickableEditorObject.h"
 #include "AutoReimportManager.generated.h"
 
 /* Deals with auto reimporting of objects when the objects file on disk is modified*/
diff --git a/Engine/Source/Editor/UnrealEd/Classes/Commandlets/GatherTextFromAssetsCommandlet.h b/Engine/Source/Editor/UnrealEd/Classes/Commandlets/GatherTextFromAssetsCommandlet.h
--- a/Engine/Source/Editor/UnrealEd/Classes/Commandlets/GatherTextFromAssetsCommandlet.h
+++ b/Engine/Source/Editor/UnrealEd/Classes/Commandlets/GatherTextFromAssetsCommandlet.h
@@ -2,6 +2,7 @@
 
 #pragma once 
 #include "Commandlets/Commandlet.h"
+#include "Sound/DialogueTypes.h"
 #include "GatherTextFromAssetsCommandlet.generated.h"
 
 namespace EAssetTextGatherStatus
@@ -17,6 +18,7 @@ namespace EAssetTextGatherStatus
 }
 
 class FLocMetadataValue;
+class UDialogueWave;
 
 /**
  *	UGatherTextFromAssetsCommandlet: Localization commandlet that collects all text to be localized from the game assets.
diff --git a/Engine/Source/Editor/UnrealEd/Classes/CookOnTheSide/CookOnTheFlyServer.h b/Engine/Source/Editor/UnrealEd/Classes/CookOnTheSide/CookOnTheFlyServer.h
--- a/Engine/Source/Editor/UnrealEd/Classes/CookOnTheSide/CookOnTheFlyServer.h
+++ b/Engine/Source/Editor/UnrealEd/Classes/CookOnTheSide/CookOnTheFlyServer.h
@@ -6,6 +6,7 @@
 
 #pragma once
 
+#include "TickableEditorObject.h"
 #include "CookOnTheFlyServer.generated.h"
 
 
@@ -115,7 +116,7 @@ private:
 	};
 
 public:
-
+	virtual ~UCookOnTheFlyServer();
 
 	template<typename Type>
 	struct FThreadSafeQueue : public FUnsyncronizedQueue<Type, FCriticalSection, FScopeLock>
diff --git a/Engine/Source/Editor/UnrealEd/Classes/Editor/EditorEngine.h b/Engine/Source/Editor/UnrealEd/Classes/Editor/EditorEngine.h
--- a/Engine/Source/Editor/UnrealEd/Classes/Editor/EditorEngine.h
+++ b/Engine/Source/Editor/UnrealEd/Classes/Editor/EditorEngine.h
@@ -6,13 +6,15 @@
 #include "Transactor.h"
 #include "../Settings/LevelEditorPlaySettings.h"
 #include "../Settings/LevelEditorViewportSettings.h"
+#include "Engine/Brush.h"
+#include "Engine/Engine.h"
 #include "EditorEngine.generated.h"
 
 class FAssetData;
 class UAnimSequence;
 class USkeleton;
 class FPoly;
-
+class APlayerStart;
 //
 // Things to set in mapSetBrush.
 //
diff --git a/Engine/Source/Editor/UnrealEd/Classes/Editor/UnrealEdEngine.h b/Engine/Source/Editor/UnrealEd/Classes/Editor/UnrealEdEngine.h
--- a/Engine/Source/Editor/UnrealEd/Classes/Editor/UnrealEdEngine.h
+++ b/Engine/Source/Editor/UnrealEd/Classes/Editor/UnrealEdEngine.h
@@ -6,6 +6,8 @@
 #include "ComponentVisualizerManager.h"
 #include "UnrealEdEngine.generated.h"
 
+class IEngineLoop;
+
 UENUM()
 enum EPackageNotifyState
 {
diff --git a/Engine/Source/Editor/UnrealEd/Classes/Exporters/AnimSequenceExporterFBX.h b/Engine/Source/Editor/UnrealEd/Classes/Exporters/AnimSequenceExporterFBX.h
--- a/Engine/Source/Editor/UnrealEd/Classes/Exporters/AnimSequenceExporterFBX.h
+++ b/Engine/Source/Editor/UnrealEd/Classes/Exporters/AnimSequenceExporterFBX.h
@@ -5,6 +5,7 @@
 //=============================================================================
 
 #pragma once
+#include "Exporters/Exporter.h"
 #include "AnimSequenceExporterFBX.generated.h"
 
 UCLASS()
diff --git a/Engine/Source/Editor/UnrealEd/Classes/Factories/CurveFactory.h b/Engine/Source/Editor/UnrealEd/Classes/Factories/CurveFactory.h
--- a/Engine/Source/Editor/UnrealEd/Classes/Factories/CurveFactory.h
+++ b/Engine/Source/Editor/UnrealEd/Classes/Factories/CurveFactory.h
@@ -5,8 +5,11 @@
  */
 
 #pragma once
+#include "Factories/Factory.h"
 #include "CurveFactory.generated.h"
 
+class UCurveBase;
+
 /** Enum indicating the type of curve to make */
 UENUM()
 enum ECurveType
diff --git a/Engine/Source/Editor/UnrealEd/Classes/Factories/DataAssetFactory.h b/Engine/Source/Editor/UnrealEd/Classes/Factories/DataAssetFactory.h
--- a/Engine/Source/Editor/UnrealEd/Classes/Factories/DataAssetFactory.h
+++ b/Engine/Source/Editor/UnrealEd/Classes/Factories/DataAssetFactory.h
@@ -1,8 +1,11 @@
 // Copyright 1998-2014 Epic Games, Inc. All Rights Reserved.
 
 #pragma once
+#include "Factories/Factory.h"
 #include "DataAssetFactory.generated.h"
 
+class UDataAsset;
+
 UCLASS(hidecategories=Object, MinimalAPI)
 class UDataAssetFactory : public UFactory
 {
diff --git a/Engine/Source/Editor/UnrealEd/Classes/Factories/FbxAssetImportData.h b/Engine/Source/Editor/UnrealEd/Classes/Factories/FbxAssetImportData.h
--- a/Engine/Source/Editor/UnrealEd/Classes/Factories/FbxAssetImportData.h
+++ b/Engine/Source/Editor/UnrealEd/Classes/Factories/FbxAssetImportData.h
@@ -1,6 +1,7 @@
 // Copyright 1998-2014 Epic Games, Inc. All Rights Reserved.
 
 #pragma once
+#include "EditorFramework/AssetImportData.h"
 #include "FbxAssetImportData.generated.h"
 
 /**
diff --git a/Engine/Source/Editor/UnrealEd/Classes/MaterialEditor/MaterialEditorMeshComponent.h b/Engine/Source/Editor/UnrealEd/Classes/MaterialEditor/MaterialEditorMeshComponent.h
--- a/Engine/Source/Editor/UnrealEd/Classes/MaterialEditor/MaterialEditorMeshComponent.h
+++ b/Engine/Source/Editor/UnrealEd/Classes/MaterialEditor/MaterialEditorMeshComponent.h
@@ -2,6 +2,7 @@
 
 
 #pragma once
+#include "Components/StaticMeshComponent.h"
 #include "MaterialEditorMeshComponent.generated.h"
 
 UCLASS()
diff --git a/Engine/Source/Editor/UnrealEd/Classes/MaterialGraph/MaterialGraph.h b/Engine/Source/Editor/UnrealEd/Classes/MaterialGraph/MaterialGraph.h
--- a/Engine/Source/Editor/UnrealEd/Classes/MaterialGraph/MaterialGraph.h
+++ b/Engine/Source/Editor/UnrealEd/Classes/MaterialGraph/MaterialGraph.h
@@ -8,6 +8,8 @@ DECLARE_DELEGATE_RetVal( bool, FRealtimeStateGetter );
 DECLARE_DELEGATE( FSetMaterialDirty );
 DECLARE_DELEGATE_OneParam( FToggleExpressionCollapsed, UMaterialExpression* );
 
+class UMaterialGraphNode_Comment;
+class UMaterialExpressionComment;
 /**
  * A human-readable name - material expression input pair.
  */
diff --git a/Engine/Source/Editor/UnrealEd/Classes/TexAligner/TexAligner.h b/Engine/Source/Editor/UnrealEd/Classes/TexAligner/TexAligner.h
--- a/Engine/Source/Editor/UnrealEd/Classes/TexAligner/TexAligner.h
+++ b/Engine/Source/Editor/UnrealEd/Classes/TexAligner/TexAligner.h
@@ -10,6 +10,8 @@
 
 #include "TexAligner.generated.h"
 
+class FPoly;
+
 /** Alignment types. */
 UENUM()
 enum ETexAlign
diff --git a/Engine/Source/Editor/UnrealEd/Classes/ThumbnailRendering/SceneThumbnailInfo.h b/Engine/Source/Editor/UnrealEd/Classes/ThumbnailRendering/SceneThumbnailInfo.h
--- a/Engine/Source/Editor/UnrealEd/Classes/ThumbnailRendering/SceneThumbnailInfo.h
+++ b/Engine/Source/Editor/UnrealEd/Classes/ThumbnailRendering/SceneThumbnailInfo.h
@@ -5,6 +5,7 @@
  */
 
 #pragma once
+#include "EditorFramework/ThumbnailInfo.h"
 #include "SceneThumbnailInfo.generated.h"
 
 
diff --git a/Engine/Source/Editor/UnrealEd/Private/Animation/EditorAnimBaseObj.cpp b/Engine/Source/Editor/UnrealEd/Private/Animation/EditorAnimBaseObj.cpp
--- a/Engine/Source/Editor/UnrealEd/Private/Animation/EditorAnimBaseObj.cpp
+++ b/Engine/Source/Editor/UnrealEd/Private/Animation/EditorAnimBaseObj.cpp
@@ -5,7 +5,7 @@
 =============================================================================*/ 
 
 #include "UnrealEd.h"
-#include "AssetNotifications.h"
+#include "Animation/EditorAnimBaseObj.h"
 
 #define LOCTEXT_NAMESPACE "SSkeletonTree"
 
diff --git a/Engine/Source/Editor/UnrealEd/Private/ApexClothingUtils.cpp b/Engine/Source/Editor/UnrealEd/Private/ApexClothingUtils.cpp
--- a/Engine/Source/Editor/UnrealEd/Private/ApexClothingUtils.cpp
+++ b/Engine/Source/Editor/UnrealEd/Private/ApexClothingUtils.cpp
@@ -9,15 +9,13 @@ DEFINE_LOG_CATEGORY_STATIC(LogApexClothingUtils, Log, All);
 
 #if WITH_APEX
 
-#include "EditorPhysXSupport.h"
-// Utilities
-#include "NxParamUtils.h"
-
 #if WITH_APEX_CLOTHING
 	#include "NxClothingAssetAuthoring.h"
+	#include "NxClothingAsset.h"
 #endif // #if WITH_APEX_CLOTHING
 
-#include "Runtime/Engine/Private/PhysicsEngine/PhysXSupport.h"
+#include "NxParamUtils.h"
+#include "NxRenderMeshAsset.h"
 
 #endif // #if WITH_APEX
 
diff --git a/Engine/Source/Editor/UnrealEd/Private/ClassIconFinder.cpp b/Engine/Source/Editor/UnrealEd/Private/ClassIconFinder.cpp
--- a/Engine/Source/Editor/UnrealEd/Private/ClassIconFinder.cpp
+++ b/Engine/Source/Editor/UnrealEd/Private/ClassIconFinder.cpp
@@ -2,7 +2,6 @@
 
 
 #include "UnrealEd.h"
-#include "Slate.h"
 #include "ClassIconFinder.h"
 
 const FSlateBrush* FClassIconFinder::FindIconForActors(const TArray< TWeakObjectPtr<AActor> >& InActors, UClass*& CommonBaseClass)
diff --git a/Engine/Source/Editor/UnrealEd/Private/Commandlets/AutomatedLauncherCommandlet.cpp b/Engine/Source/Editor/UnrealEd/Private/Commandlets/AutomatedLauncherCommandlet.cpp
--- a/Engine/Source/Editor/UnrealEd/Private/Commandlets/AutomatedLauncherCommandlet.cpp
+++ b/Engine/Source/Editor/UnrealEd/Private/Commandlets/AutomatedLauncherCommandlet.cpp
@@ -1,9 +1,8 @@
 // Copyright 1998-2014 Epic Games, Inc. All Rights Reserved.
 
 #include "UnrealEd.h"
-
-#include "LauncherAutomatedService.h"
-
+#include "Interfaces/ILauncherAutomatedServiceProvider.h"
+#include "Interfaces/ILauncherAutomatedServiceModule.h"
 
 UAutomatedLauncherCommandlet::UAutomatedLauncherCommandlet(const class FPostConstructInitializeProperties& PCIP)
 	: Super(PCIP)
diff --git a/Engine/Source/Editor/UnrealEd/Private/Commandlets/GatherTextCommandlet.cpp b/Engine/Source/Editor/UnrealEd/Private/Commandlets/GatherTextCommandlet.cpp
--- a/Engine/Source/Editor/UnrealEd/Private/Commandlets/GatherTextCommandlet.cpp
+++ b/Engine/Source/Editor/UnrealEd/Private/Commandlets/GatherTextCommandlet.cpp
@@ -1,8 +1,6 @@
 ﻿// Copyright 1998-2014 Epic Games, Inc. All Rights Reserved.
 
 #include "UnrealEd.h"
-#include "ISourceControlModule.h"
-#include "Internationalization/InternationalizationManifest.h"
 
 DEFINE_LOG_CATEGORY_STATIC(LogGatherTextCommandlet, Log, All);
 
diff --git a/Engine/Source/Editor/UnrealEd/Private/Commandlets/GatherTextCommandletBase.cpp b/Engine/Source/Editor/UnrealEd/Private/Commandlets/GatherTextCommandletBase.cpp
--- a/Engine/Source/Editor/UnrealEd/Private/Commandlets/GatherTextCommandletBase.cpp
+++ b/Engine/Source/Editor/UnrealEd/Private/Commandlets/GatherTextCommandletBase.cpp
@@ -1,7 +1,6 @@
 // Copyright 1998-2014 Epic Games, Inc. All Rights Reserved.
 
 #include "UnrealEd.h"
-#include "Json.h"
 #include "InternationalizationManifestJsonSerializer.h"
 #include "Internationalization/InternationalizationMetadata.h"
 #include "InternationalizationMetadataJsonSerializer.h"
diff --git a/Engine/Source/Editor/UnrealEd/Private/Commandlets/GenerateGatherArchiveCommandlet.cpp b/Engine/Source/Editor/UnrealEd/Private/Commandlets/GenerateGatherArchiveCommandlet.cpp
--- a/Engine/Source/Editor/UnrealEd/Private/Commandlets/GenerateGatherArchiveCommandlet.cpp
+++ b/Engine/Source/Editor/UnrealEd/Private/Commandlets/GenerateGatherArchiveCommandlet.cpp
@@ -1,8 +1,6 @@
 // Copyright 1998-2014 Epic Games, Inc. All Rights Reserved.
 
 #include "UnrealEd.h"
-#include "ISourceControlModule.h"
-#include "Json.h"
 #include "Internationalization/InternationalizationArchive.h"
 #include "Internationalization/InternationalizationManifest.h"
 #include "Internationalization/InternationalizationMetadata.h"
diff --git a/Engine/Source/Editor/UnrealEd/Private/Commandlets/GenerateGatherManifestCommandlet.cpp b/Engine/Source/Editor/UnrealEd/Private/Commandlets/GenerateGatherManifestCommandlet.cpp
--- a/Engine/Source/Editor/UnrealEd/Private/Commandlets/GenerateGatherManifestCommandlet.cpp
+++ b/Engine/Source/Editor/UnrealEd/Private/Commandlets/GenerateGatherManifestCommandlet.cpp
@@ -1,8 +1,6 @@
 // Copyright 1998-2014 Epic Games, Inc. All Rights Reserved.
 
 #include "UnrealEd.h"
-#include "ISourceControlModule.h"
-#include "Json.h"
 #include "Internationalization/InternationalizationManifest.h"
 #include "InternationalizationManifestJsonSerializer.h"
 
diff --git a/Engine/Source/Editor/UnrealEd/Private/Commandlets/GenerateTextLocalizationReportCommandlet.cpp b/Engine/Source/Editor/UnrealEd/Private/Commandlets/GenerateTextLocalizationReportCommandlet.cpp
--- a/Engine/Source/Editor/UnrealEd/Private/Commandlets/GenerateTextLocalizationReportCommandlet.cpp
+++ b/Engine/Source/Editor/UnrealEd/Private/Commandlets/GenerateTextLocalizationReportCommandlet.cpp
@@ -1,7 +1,6 @@
 // Copyright 1998-2014 Epic Games, Inc. All Rights Reserved.
 
 #include "UnrealEd.h"
-#include "Json.h"
 #include "InternationalizationManifest.h"
 #include "InternationalizationManifestJsonSerializer.h"
 #include "InternationalizationArchive.h"
diff --git a/Engine/Source/Editor/UnrealEd/Private/Commandlets/GenerateTextLocalizationResourceCommandlet.cpp b/Engine/Source/Editor/UnrealEd/Private/Commandlets/GenerateTextLocalizationResourceCommandlet.cpp
--- a/Engine/Source/Editor/UnrealEd/Private/Commandlets/GenerateTextLocalizationResourceCommandlet.cpp
+++ b/Engine/Source/Editor/UnrealEd/Private/Commandlets/GenerateTextLocalizationResourceCommandlet.cpp
@@ -2,11 +2,9 @@
 
 #include "UnrealEd.h"
 #include "TextLocalizationResourceGenerator.h"
-#include "Json.h"
 #include "InternationalizationManifest.h"
-#include "InternationalizationManifestJsonSerializer.h"
-#include "InternationalizationArchive.h"
 #include "InternationalizationArchiveJsonSerializer.h"
+#include "Internationalization/InternationalizationManifestJsonSerializer.h"
 
 DEFINE_LOG_CATEGORY_STATIC(LogGenerateTextLocalizationResourceCommandlet, Log, All);
 
diff --git a/Engine/Source/Editor/UnrealEd/Private/ConvexDecompTool.cpp b/Engine/Source/Editor/UnrealEd/Private/ConvexDecompTool.cpp
--- a/Engine/Source/Editor/UnrealEd/Private/ConvexDecompTool.cpp
+++ b/Engine/Source/Editor/UnrealEd/Private/ConvexDecompTool.cpp
@@ -6,7 +6,6 @@
 
 #include "UnrealEd.h"
 #include "ThirdParty/HACD/HACD_1.0/public/HACD.h"
-#include "ConvexDecompTool.h"
 
 DEFINE_LOG_CATEGORY_STATIC(LogConvexDecompTool, Log, All);
 
@@ -26,7 +25,7 @@ public:
 	}
 };
 
-void DecomposeMeshToHulls(UBodySetup* InBodySetup, const TArray<FVector>& InVertices, const TArray<uint32>& InIndices, int32 InMaxHullCount, int32 InMaxHullVerts)
+UNREALED_API void DecomposeMeshToHulls(UBodySetup* InBodySetup, const TArray<FVector>& InVertices, const TArray<uint32>& InIndices, int32 InMaxHullCount, int32 InMaxHullVerts)
 {
 	check(InBodySetup != NULL);
 
diff --git a/Engine/Source/Editor/UnrealEd/Private/CookOnTheFlyServer.cpp b/Engine/Source/Editor/UnrealEd/Private/CookOnTheFlyServer.cpp
--- a/Engine/Source/Editor/UnrealEd/Private/CookOnTheFlyServer.cpp
+++ b/Engine/Source/Editor/UnrealEd/Private/CookOnTheFlyServer.cpp
@@ -70,7 +70,8 @@ UCookOnTheFlyServer::UCookOnTheFlyServer(const class FPostConstructInitializePro
 	bCookOnTheFly = true;
 }
 
-
+UCookOnTheFlyServer::~UCookOnTheFlyServer()
+{ }
 
 void UCookOnTheFlyServer::Tick(float DeltaTime)
 {
diff --git a/Engine/Source/Editor/UnrealEd/Private/Dialogs/Dialogs.cpp b/Engine/Source/Editor/UnrealEd/Private/Dialogs/Dialogs.cpp
--- a/Engine/Source/Editor/UnrealEd/Private/Dialogs/Dialogs.cpp
+++ b/Engine/Source/Editor/UnrealEd/Private/Dialogs/Dialogs.cpp
@@ -2,12 +2,11 @@
 
 
 #include "UnrealEd.h"
-#include "TexAlignTools.h"
-#include "ScopedTransaction.h"
 #include "Slate.h"
 #include "ObjectTools.h"
-#include "Editor/MainFrame/Public/MainFrame.h"
 #include "DesktopPlatformModule.h"
+#include "Interfaces/IMainFrameModule.h"
+
 DEFINE_LOG_CATEGORY_STATIC(LogDialogs, Log, All);
 
 #define LOCTEXT_NAMESPACE "Dialogs"
diff --git a/Engine/Source/Editor/UnrealEd/Private/Dialogs/DlgSoundWaveOptions.cpp b/Engine/Source/Editor/UnrealEd/Private/Dialogs/DlgSoundWaveOptions.cpp
--- a/Engine/Source/Editor/UnrealEd/Private/Dialogs/DlgSoundWaveOptions.cpp
+++ b/Engine/Source/Editor/UnrealEd/Private/Dialogs/DlgSoundWaveOptions.cpp
@@ -3,10 +3,8 @@
 
 #include "UnrealEd.h"
 
-#include "SoundDefinitions.h"
 #include "Dialogs/DlgSoundWaveOptions.h"
 #include "SoundPreviewThread.h"
-#include "Slate.h"
 
 /** 
  * Info used to setup the rows of the sound quality previewer
diff --git a/Engine/Source/Editor/UnrealEd/Private/Editor/EditorUserSettings.cpp b/Engine/Source/Editor/UnrealEd/Private/Editor/EditorUserSettings.cpp
--- a/Engine/Source/Editor/UnrealEd/Private/Editor/EditorUserSettings.cpp
+++ b/Engine/Source/Editor/UnrealEd/Private/Editor/EditorUserSettings.cpp
@@ -1,7 +1,6 @@
 // Copyright 1998-2014 Epic Games, Inc. All Rights Reserved.
 
 #include "UnrealEd.h"
-#include "MessageLog.h"
 #include "BlueprintPaletteFavorites.h"
 
 #define LOCTEXT_NAMESPACE "EditorUserSettings"
diff --git a/Engine/Source/Editor/UnrealEd/Private/EditorComponents.cpp b/Engine/Source/Editor/UnrealEd/Private/EditorComponents.cpp
--- a/Engine/Source/Editor/UnrealEd/Private/EditorComponents.cpp
+++ b/Engine/Source/Editor/UnrealEd/Private/EditorComponents.cpp
@@ -1,9 +1,6 @@
 // Copyright 1998-2014 Epic Games, Inc. All Rights Reserved.
 
-
 #include "UnrealEd.h"
-#include "MouseDeltaTracker.h"
-
 
 static TAutoConsoleVariable<int32> CVarEditorNewLevelGrid(
 	TEXT("r.Editor.NewLevelGrid"),
diff --git a/Engine/Source/Editor/UnrealEd/Private/EditorModeTools.cpp b/Engine/Source/Editor/UnrealEd/Private/EditorModeTools.cpp
--- a/Engine/Source/Editor/UnrealEd/Private/EditorModeTools.cpp
+++ b/Engine/Source/Editor/UnrealEd/Private/EditorModeTools.cpp
@@ -2,9 +2,7 @@
 
 
 #include "UnrealEd.h"
-#include "SurfaceIterators.h"
-#include "BSPOps.h"
-
+#include "EditorModeTools.h"
 
 FModeTool::FModeTool():
 	ID( MT_None ),
diff --git a/Engine/Source/Editor/UnrealEd/Private/EditorObject.cpp b/Engine/Source/Editor/UnrealEd/Private/EditorObject.cpp
--- a/Engine/Source/Editor/UnrealEd/Private/EditorObject.cpp
+++ b/Engine/Source/Editor/UnrealEd/Private/EditorObject.cpp
@@ -5,7 +5,6 @@
 =============================================================================*/
 
 #include "UnrealEd.h"
-#include "Factories.h"
 #include "BSPOps.h"
 
 #include "Foliage/InstancedFoliageActor.h"
diff --git a/Engine/Source/Editor/UnrealEd/Private/Factories/EditorMorphFactory.cpp b/Engine/Source/Editor/UnrealEd/Private/Factories/EditorMorphFactory.cpp
--- a/Engine/Source/Editor/UnrealEd/Private/Factories/EditorMorphFactory.cpp
+++ b/Engine/Source/Editor/UnrealEd/Private/Factories/EditorMorphFactory.cpp
@@ -6,7 +6,6 @@
 
 #include "UnrealEd.h"
 #include "Factories.h"
-#include "SkelImport.h"
 
 #define SET_MORPH_ERROR( ErrorType ) { if( Error ) *Error = ErrorType; }
 
diff --git a/Engine/Source/Editor/UnrealEd/Private/Factories/HDRLoader.cpp b/Engine/Source/Editor/UnrealEd/Private/Factories/HDRLoader.cpp
--- a/Engine/Source/Editor/UnrealEd/Private/Factories/HDRLoader.cpp
+++ b/Engine/Source/Editor/UnrealEd/Private/Factories/HDRLoader.cpp
@@ -4,7 +4,6 @@
 #include "UnrealEd.h"
 #include "HDRLoader.h"
 #include "DDSLoader.h"
-#include "Collision.h"
 
 DEFINE_LOG_CATEGORY_STATIC(LogHDRLoader, Log, All);
 
diff --git a/Engine/Source/Editor/UnrealEd/Private/Factories/SkeletalMeshImport.cpp b/Engine/Source/Editor/UnrealEd/Private/Factories/SkeletalMeshImport.cpp
--- a/Engine/Source/Editor/UnrealEd/Private/Factories/SkeletalMeshImport.cpp
+++ b/Engine/Source/Editor/UnrealEd/Private/Factories/SkeletalMeshImport.cpp
@@ -5,7 +5,6 @@
 =============================================================================*/
 
 #include "UnrealEd.h"
-#include "Factories.h"
 #include "SkelImport.h"
 #include "SkeletalMeshSorting.h"
 #include "../../../../Source/Runtime/Engine/Classes/PhysicsEngine/PhysicsAsset.h"
diff --git a/Engine/Source/Editor/UnrealEd/Private/Factories/TTFontImport.cpp b/Engine/Source/Editor/UnrealEd/Private/Factories/TTFontImport.cpp
--- a/Engine/Source/Editor/UnrealEd/Private/Factories/TTFontImport.cpp
+++ b/Engine/Source/Editor/UnrealEd/Private/Factories/TTFontImport.cpp
@@ -4,9 +4,8 @@
 	TTFontImport.cpp: True-type Font Importing
 =============================================================================*/
 #include "UnrealEd.h"
-#include "Factories.h"
-#include "Editor/MainFrame/Public/MainFrame.h"
 #include "DesktopPlatformModule.h"
+#include "Interfaces/IMainFrameModule.h"
 
 #ifndef WITH_FREETYPE
 	#define WITH_FREETYPE	0
@@ -1190,7 +1189,7 @@ bool UTrueTypeFontFactory::CreateFontTexture(
 			// like TrueType and OpenType; it won't work for raster fonts!
 			bool bUsingGlyphOutlines = false;
 			GLYPHMETRICS WinGlyphMetrics;
-			const MAT2 WinIdentityMatrix2x2 = { 0,1, 0,0, 0,0, 0,1 };
+			const MAT2 WinIdentityMatrix2x2 = { {0,1}, {0,0}, {0,0}, {0,1} };
 			int32 VerticalOffset = 0;
 			uint32 GGODataSize = 0;
 			if( !ImportOptions->Data.bEnableLegacyMode && ImportOptions->Data.bEnableAntialiasing )    // We only bother using GetGlyphOutline for AntiAliased fonts!
diff --git a/Engine/Source/Editor/UnrealEd/Private/Fbx/FbxLightImport.cpp b/Engine/Source/Editor/UnrealEd/Private/Fbx/FbxLightImport.cpp
--- a/Engine/Source/Editor/UnrealEd/Private/Fbx/FbxLightImport.cpp
+++ b/Engine/Source/Editor/UnrealEd/Private/Fbx/FbxLightImport.cpp
@@ -5,8 +5,6 @@
 =============================================================================*/
 
 #include "UnrealEd.h"
-#include "Factories.h"
-#include "Engine.h"
 #include "FbxImporter.h"
 
 DEFINE_LOG_CATEGORY_STATIC(LogFbxLightImport, Log, All);
diff --git a/Engine/Source/Editor/UnrealEd/Private/Fbx/FbxUtilsImport.cpp b/Engine/Source/Editor/UnrealEd/Private/Fbx/FbxUtilsImport.cpp
--- a/Engine/Source/Editor/UnrealEd/Private/Fbx/FbxUtilsImport.cpp
+++ b/Engine/Source/Editor/UnrealEd/Private/Fbx/FbxUtilsImport.cpp
@@ -29,9 +29,6 @@
 */
 
 #include "UnrealEd.h"
-#include "Factories.h"
-#include "Engine.h"
-#include "TextureLayout.h"
 #include "FbxImporter.h"
 
 using namespace UnFbx;
diff --git a/Engine/Source/Editor/UnrealEd/Private/Features/EditorFeatures.cpp b/Engine/Source/Editor/UnrealEd/Private/Features/EditorFeatures.cpp
--- a/Engine/Source/Editor/UnrealEd/Private/Features/EditorFeatures.cpp
+++ b/Engine/Source/Editor/UnrealEd/Private/Features/EditorFeatures.cpp
@@ -1,8 +1,6 @@
 // Copyright 1998-2014 Epic Games, Inc. All Rights Reserved.
 
 #include "UnrealEd.h"
-#include "EditorFeatures.h"
-
 
 namespace EditorFeatures
 {
diff --git a/Engine/Source/Editor/UnrealEd/Private/FeedbackContextEditor.cpp b/Engine/Source/Editor/UnrealEd/Private/FeedbackContextEditor.cpp
--- a/Engine/Source/Editor/UnrealEd/Private/FeedbackContextEditor.cpp
+++ b/Engine/Source/Editor/UnrealEd/Private/FeedbackContextEditor.cpp
@@ -4,7 +4,7 @@
 #include "UnrealEd.h"
 #include "FeedbackContextEditor.h"
 #include "Dialogs/SBuildProgress.h"
-#include "Editor/MainFrame/Public/MainFrame.h"
+#include "Interfaces/IMainFrameModule.h"
 
 /** Called to cancel the slow task activity */
 DECLARE_DELEGATE( FOnCancelClickedDelegate );
diff --git a/Engine/Source/Editor/UnrealEd/Private/GeomFitUtils.cpp b/Engine/Source/Editor/UnrealEd/Private/GeomFitUtils.cpp
--- a/Engine/Source/Editor/UnrealEd/Private/GeomFitUtils.cpp
+++ b/Engine/Source/Editor/UnrealEd/Private/GeomFitUtils.cpp
@@ -8,9 +8,8 @@
 #include "StaticMeshResources.h"
 #include "EditorSupportDelegates.h"
 #include "BSPOps.h"
-#include "../Private/GeomFitUtils.h"
 #include "RawMesh.h"
-#include "MeshUtilities.h"
+#include "GeomFitUtils.h"
 
 #define LOCAL_EPS (0.01f)
 static void AddVertexIfNotPresent(TArray<FVector> &vertices, FVector &newVertex)
diff --git a/Engine/Source/Editor/UnrealEd/Private/MRUFavoritesList.cpp b/Engine/Source/Editor/UnrealEd/Private/MRUFavoritesList.cpp
--- a/Engine/Source/Editor/UnrealEd/Private/MRUFavoritesList.cpp
+++ b/Engine/Source/Editor/UnrealEd/Private/MRUFavoritesList.cpp
@@ -3,7 +3,6 @@
 
 #include "UnrealEd.h"
 #include "MRUFavoritesList.h"
-#include "ConfigCacheIni.h"
 
 const FString FMainMRUFavoritesList::FAVORITES_INI_SECTION = TEXT("FavoriteFiles");
 
diff --git a/Engine/Source/Editor/UnrealEd/Private/MouseDeltaTracker.cpp b/Engine/Source/Editor/UnrealEd/Private/MouseDeltaTracker.cpp
--- a/Engine/Source/Editor/UnrealEd/Private/MouseDeltaTracker.cpp
+++ b/Engine/Source/Editor/UnrealEd/Private/MouseDeltaTracker.cpp
@@ -3,10 +3,8 @@
 
 #include "UnrealEd.h"
 #include "MouseDeltaTracker.h"
-#include "DragTool_BoxSelect.h"
-#include "DragTool_Measure.h"
-#include "DragTool_FrustumSelect.h"
 #include "SnappingUtils.h"
+#include "EditorDragTools.h"
 
 #define LOCTEXT_NAMESPACE "MouseDeltaTracker"
 
diff --git a/Engine/Source/Editor/UnrealEd/Private/NavigationBuildingNotification.cpp b/Engine/Source/Editor/UnrealEd/Private/NavigationBuildingNotification.cpp
--- a/Engine/Source/Editor/UnrealEd/Private/NavigationBuildingNotification.cpp
+++ b/Engine/Source/Editor/UnrealEd/Private/NavigationBuildingNotification.cpp
@@ -1,7 +1,6 @@
 // Copyright 1998-2014 Epic Games, Inc. All Rights Reserved.
 
 #include "UnrealEd.h"
-#include "MainFrame.h"
 #include "Kismet2/DebuggerCommands.h"
 #include "EditorBuildUtils.h"
 
diff --git a/Engine/Source/Editor/UnrealEd/Private/PackageAutoSaver.cpp b/Engine/Source/Editor/UnrealEd/Private/PackageAutoSaver.cpp
--- a/Engine/Source/Editor/UnrealEd/Private/PackageAutoSaver.cpp
+++ b/Engine/Source/Editor/UnrealEd/Private/PackageAutoSaver.cpp
@@ -3,7 +3,6 @@
 #include "UnrealEd.h"
 #include "PackageAutoSaver.h"
 #include "PackageRestore.h"
-#include "Json.h"
 
 namespace PackageAutoSaverJson
 {
diff --git a/Engine/Source/Editor/UnrealEd/Private/ReferenceInfoUtils.cpp b/Engine/Source/Editor/UnrealEd/Private/ReferenceInfoUtils.cpp
--- a/Engine/Source/Editor/UnrealEd/Private/ReferenceInfoUtils.cpp
+++ b/Engine/Source/Editor/UnrealEd/Private/ReferenceInfoUtils.cpp
@@ -1,9 +1,8 @@
 // Copyright 1998-2014 Epic Games, Inc. All Rights Reserved.
 
 #include "UnrealEd.h"
-#include "../Private/ReferenceInfoUtils.h"
 #include "AssetRegistryModule.h"
-#include "Editor/MainFrame/Public/MainFrame.h"
+#include "../Private/ReferenceInfoUtils.h"
 
 void ExecuteReferenceInfo(const TArray<FString>& Args, UWorld* InWorld )
 {
diff --git a/Engine/Source/Editor/UnrealEd/Private/ReferencedAssetsUtils.cpp b/Engine/Source/Editor/UnrealEd/Private/ReferencedAssetsUtils.cpp
--- a/Engine/Source/Editor/UnrealEd/Private/ReferencedAssetsUtils.cpp
+++ b/Engine/Source/Editor/UnrealEd/Private/ReferencedAssetsUtils.cpp
@@ -2,7 +2,6 @@
 
 
 #include "UnrealEd.h"
-#include "BusyCursor.h"
 #include "ReferencedAssetsUtils.h"
 #include "AssetRegistryModule.h"
 
@@ -183,8 +182,8 @@ FFindAssetsArchive::FFindAssetsArchive(
 	bool bIncludeDefaults/*=false*/,
 	bool bReverseReferenceGraph/*=false*/ ) 
 : StartObject(Search), AssetList(OutAssetList), CurrentReferenceGraph(ReferenceGraph)
-, bIncludeScriptRefs(bIncludeClasses), bIncludeDefaultRefs(bIncludeDefaults), MaxRecursionDepth(MaxRecursion), bUseReverseReferenceGraph(bReverseReferenceGraph)
-, CurrentDepth(0)
+, bIncludeScriptRefs(bIncludeClasses), bIncludeDefaultRefs(bIncludeDefaults), MaxRecursionDepth(MaxRecursion)
+, CurrentDepth(0), bUseReverseReferenceGraph(bReverseReferenceGraph)
 {
 	ArIsObjectReferenceCollector = true;
 	ArIsModifyingWeakAndStrongReferences = true; // While we are not modifying them, we want to follow weak references as well
diff --git a/Engine/Source/Editor/UnrealEd/Private/SEditorViewport.cpp b/Engine/Source/Editor/UnrealEd/Private/SEditorViewport.cpp
--- a/Engine/Source/Editor/UnrealEd/Private/SEditorViewport.cpp
+++ b/Engine/Source/Editor/UnrealEd/Private/SEditorViewport.cpp
@@ -5,7 +5,6 @@
 #include "EditorViewportClient.h"
 #include "SceneViewport.h"
 #include "EditorViewportCommands.h"
-#include "IDocumentation.h"
 
 #define LOCTEXT_NAMESPACE "EditorViewport"
 
diff --git a/Engine/Source/Editor/UnrealEd/Private/SEditorViewportToolBarButton.cpp b/Engine/Source/Editor/UnrealEd/Private/SEditorViewportToolBarButton.cpp
--- a/Engine/Source/Editor/UnrealEd/Private/SEditorViewportToolBarButton.cpp
+++ b/Engine/Source/Editor/UnrealEd/Private/SEditorViewportToolBarButton.cpp
@@ -3,7 +3,6 @@
 
 #include "UnrealEd.h"
 #include "SEditorViewportToolBarButton.h"
-#include "SViewportToolBar.h"
 
 
 void SEditorViewportToolBarButton::Construct( const FArguments& Declaration )
diff --git a/Engine/Source/Editor/UnrealEd/Private/STransformViewportToolbar.cpp b/Engine/Source/Editor/UnrealEd/Private/STransformViewportToolbar.cpp
--- a/Engine/Source/Editor/UnrealEd/Private/STransformViewportToolbar.cpp
+++ b/Engine/Source/Editor/UnrealEd/Private/STransformViewportToolbar.cpp
@@ -3,8 +3,8 @@
 
 #include "UnrealEd.h"
 #include "STransformViewportToolbar.h"
-#include "SLevelViewport.h"
 #include "EditorViewportCommands.h"
+#include "SEditorViewport.h"
 
 #define LOCTEXT_NAMESPACE "TransformToolBar"
 
diff --git a/Engine/Source/Editor/UnrealEd/Private/ShaderCompilingNotification.cpp b/Engine/Source/Editor/UnrealEd/Private/ShaderCompilingNotification.cpp
--- a/Engine/Source/Editor/UnrealEd/Private/ShaderCompilingNotification.cpp
+++ b/Engine/Source/Editor/UnrealEd/Private/ShaderCompilingNotification.cpp
@@ -1,8 +1,7 @@
 // Copyright 1998-2014 Epic Games, Inc. All Rights Reserved.
 
-
 #include "UnrealEd.h"
-#include "MainFrame.h"
+#include "TickableEditorObject.h"
 
 /** Notification class for asynchronous shader compiling. */
 class FShaderCompilingNotificationImpl
diff --git a/Engine/Source/Editor/UnrealEd/Private/SoundCueGraph.cpp b/Engine/Source/Editor/UnrealEd/Private/SoundCueGraph.cpp
--- a/Engine/Source/Editor/UnrealEd/Private/SoundCueGraph.cpp
+++ b/Engine/Source/Editor/UnrealEd/Private/SoundCueGraph.cpp
@@ -4,7 +4,7 @@
 // USoundCueGraph
 
 #include "UnrealEd.h"
-#include "SoundDefinitions.h"
+#include "SoundCueGraph/SoundCueGraph.h"
 
 USoundCueGraph::USoundCueGraph(const FPostConstructInitializeProperties& PCIP)
 	: Super(PCIP)
diff --git a/Engine/Source/Editor/UnrealEd/Private/SoundPreviewThread.cpp b/Engine/Source/Editor/UnrealEd/Private/SoundPreviewThread.cpp
--- a/Engine/Source/Editor/UnrealEd/Private/SoundPreviewThread.cpp
+++ b/Engine/Source/Editor/UnrealEd/Private/SoundPreviewThread.cpp
@@ -3,7 +3,6 @@
 
 #include "UnrealEd.h"
 #include "SoundPreviewThread.h"
-#include "SoundDefinitions.h"
 
 /**
  *
diff --git a/Engine/Source/Editor/UnrealEd/Private/SourceCodeNavigation.cpp b/Engine/Source/Editor/UnrealEd/Private/SourceCodeNavigation.cpp
--- a/Engine/Source/Editor/UnrealEd/Private/SourceCodeNavigation.cpp
+++ b/Engine/Source/Editor/UnrealEd/Private/SourceCodeNavigation.cpp
@@ -3,14 +3,11 @@
 
 #include "UnrealEd.h"
 #include "SourceCodeNavigation.h"
-#include "MainFrame.h"
 #include "ISourceCodeAccessModule.h"
 
 #if PLATFORM_WINDOWS
 #include "AllowWindowsPlatformTypes.h"
-	#include <DbgHelp.h>				
-	#include <TlHelp32.h>		
-	#include <psapi.h>
+	#include <DbgHelp.h>
 #include "HideWindowsPlatformTypes.h"
 #elif PLATFORM_MAC
 #include <mach-o/dyld.h>
diff --git a/Engine/Source/Editor/UnrealEd/Private/StaticMeshEdit.cpp b/Engine/Source/Editor/UnrealEd/Private/StaticMeshEdit.cpp
--- a/Engine/Source/Editor/UnrealEd/Private/StaticMeshEdit.cpp
+++ b/Engine/Source/Editor/UnrealEd/Private/StaticMeshEdit.cpp
@@ -6,11 +6,8 @@
 
 #include "UnrealEd.h"
 #include "StaticMeshResources.h"
-#include "Factories.h"
-#include "TextureLayout.h"
 #include "BSPOps.h"
 #include "RawMesh.h"
-#include "MeshUtilities.h"
 
 bool GBuildStaticMeshCollision = 1;
 
diff --git a/Engine/Source/Editor/UnrealEd/Private/ThumbnailManager.cpp b/Engine/Source/Editor/UnrealEd/Private/ThumbnailManager.cpp
--- a/Engine/Source/Editor/UnrealEd/Private/ThumbnailManager.cpp
+++ b/Engine/Source/Editor/UnrealEd/Private/ThumbnailManager.cpp
@@ -2,8 +2,6 @@
 
 #include "UnrealEd.h"
 
-#include "SoundDefinitions.h"
-#include "BlueprintUtilities.h"
 #include "ImageUtils.h"
 
 #include "ISourceControlModule.h"
diff --git a/Engine/Source/Editor/UnrealEd/Private/ThumbnailRendering/SlateBrushThumbnailRenderer.cpp b/Engine/Source/Editor/UnrealEd/Private/ThumbnailRendering/SlateBrushThumbnailRenderer.cpp
--- a/Engine/Source/Editor/UnrealEd/Private/ThumbnailRendering/SlateBrushThumbnailRenderer.cpp
+++ b/Engine/Source/Editor/UnrealEd/Private/ThumbnailRendering/SlateBrushThumbnailRenderer.cpp
@@ -3,7 +3,6 @@
 #include "UnrealEd.h"
 #include "Slate/SlateBrushAsset.h"
 // FPreviewScene derived helpers for rendering
-#include "ThumbnailHelpers.h"
 
 USlateBrushThumbnailRenderer::USlateBrushThumbnailRenderer(const class FPostConstructInitializeProperties& PCIP)
 	: Super(PCIP)
diff --git a/Engine/Source/Editor/UnrealEd/Private/Utils.cpp b/Engine/Source/Editor/UnrealEd/Private/Utils.cpp
--- a/Engine/Source/Editor/UnrealEd/Private/Utils.cpp
+++ b/Engine/Source/Editor/UnrealEd/Private/Utils.cpp
@@ -3,13 +3,6 @@
 
 #include "UnrealEd.h"
 
-#if PLATFORM_WINDOWS
-// Needed for showing balloon messages
-#include "AllowWindowsPlatformTypes.h"
-		#include <ShellAPI.h>
-#include "HideWindowsPlatformTypes.h"
-#endif
-
 DEFINE_LOG_CATEGORY(LogUtils);
 
 IMPLEMENT_HIT_PROXY(HWidgetUtilProxy,HHitProxy);
diff --git a/Engine/Source/Editor/UnrealEd/Public/EditorComponents.h b/Engine/Source/Editor/UnrealEd/Public/EditorComponents.h
--- a/Engine/Source/Editor/UnrealEd/Public/EditorComponents.h
+++ b/Engine/Source/Editor/UnrealEd/Public/EditorComponents.h
@@ -6,6 +6,14 @@
 
 #pragma once
 
+#include "UObject/GCObject.h"
+#include "Engine/EngineTypes.h"
+
+class FSceneView;
+class FPrimitiveDrawInterface;
+class FReferenceCollector;
+class UMaterial;
+class UMaterialInstanceDynamic;
 
 class FGridWidget : public FGCObject
 {
diff --git a/Engine/Source/Editor/UnrealEd/Public/TickableEditorObject.h b/Engine/Source/Editor/UnrealEd/Public/TickableEditorObject.h
--- a/Engine/Source/Editor/UnrealEd/Public/TickableEditorObject.h
+++ b/Engine/Source/Editor/UnrealEd/Public/TickableEditorObject.h
@@ -2,6 +2,8 @@
 
 #pragma once
 
+#include "Tickable.h"
+
 /**
  * This class provides common registration for gamethread editor only tickable objects. It is an
  * abstract base class requiring you to implement the Tick() method.
diff --git a/Engine/Source/Editor/UnrealEd/Public/UnrealEd.h b/Engine/Source/Editor/UnrealEd/Public/UnrealEd.h
--- a/Engine/Source/Editor/UnrealEd/Public/UnrealEd.h
+++ b/Engine/Source/Editor/UnrealEd/Public/UnrealEd.h
@@ -7,8 +7,6 @@
 
 #if WITH_EDITOR
 
-#include "Engine.h"
-
 #include "Slate.h"
 #include "EditorStyle.h"
 
diff --git a/Engine/Source/Runtime/AIModule/Classes/AIController.h b/Engine/Source/Runtime/AIModule/Classes/AIController.h
--- a/Engine/Source/Runtime/AIModule/Classes/AIController.h
+++ b/Engine/Source/Runtime/AIModule/Classes/AIController.h
@@ -240,7 +240,7 @@ public:
 	virtual void DisplayDebug(class UCanvas* Canvas, const FDebugDisplayInfo& DebugDisplay, float& YL, float& YPos) override;
 
 #if ENABLE_VISUAL_LOG
-	virtual void GrabDebugSnapshot(struct FVisLogEntry* Snapshot) const override;
+	virtual void GrabDebugSnapshot(struct FVisLogEntry* Snapshot) const;
 #endif
 
 	virtual void Reset() override;
diff --git a/Engine/Source/Runtime/AIModule/Classes/EnvironmentQuery/EnvQueryTypes.h b/Engine/Source/Runtime/AIModule/Classes/EnvironmentQuery/EnvQueryTypes.h
--- a/Engine/Source/Runtime/AIModule/Classes/EnvironmentQuery/EnvQueryTypes.h
+++ b/Engine/Source/Runtime/AIModule/Classes/EnvironmentQuery/EnvQueryTypes.h
@@ -3,6 +3,7 @@
 #pragma once
 
 #include "EnvironmentQuery/Items/EnvQueryItemType.h"
+#include "EnvironmentQuery/EnvQueryContext.h"
 #include "EnvQueryTypes.generated.h"
 
 DECLARE_LOG_CATEGORY_EXTERN(LogEQS, Log, All);
@@ -273,6 +274,8 @@ struct AIMODULE_API FEnvDirection
 {
 	GENERATED_USTRUCT_BODY()
 
+	~FEnvDirection();
+
 	/** line A: start context */
 	UPROPERTY(EditDefaultsOnly, Category=Direction)
 	TSubclassOf<class UEnvQueryContext> LineFrom;
diff --git a/Engine/Source/Runtime/AIModule/Private/AIController.cpp b/Engine/Source/Runtime/AIModule/Private/AIController.cpp
--- a/Engine/Source/Runtime/AIModule/Private/AIController.cpp
+++ b/Engine/Source/Runtime/AIModule/Private/AIController.cpp
@@ -5,7 +5,6 @@
 #include "DisplayDebugHelpers.h"
 #include "BrainComponent.h"
 #include "BehaviorTree/BehaviorTree.h"
-#include "BehaviorTree/BehaviorTreeManager.h"
 #include "BehaviorTree/BlackboardComponent.h"
 #include "Navigation/PathFollowingComponent.h"
 #include "Navigation/NavigationComponent.h"
diff --git a/Engine/Source/Runtime/AIModule/Private/BehaviorTree/BehaviorTreeManager.cpp b/Engine/Source/Runtime/AIModule/Private/BehaviorTree/BehaviorTreeManager.cpp
--- a/Engine/Source/Runtime/AIModule/Private/BehaviorTree/BehaviorTreeManager.cpp
+++ b/Engine/Source/Runtime/AIModule/Private/BehaviorTree/BehaviorTreeManager.cpp
@@ -2,7 +2,6 @@
 
 #include "AIModulePrivate.h"
 #include "BehaviorTree/Tasks/BTTask_RunBehavior.h"
-#include "BehaviorTree/BlackboardComponent.h"
 #include "BehaviorTree/BehaviorTreeManager.h"
 #if WITH_EDITOR
 #include "Kismet2/KismetEditorUtilities.h"
diff --git a/Engine/Source/Runtime/AIModule/Private/BehaviorTree/BlackboardComponent.cpp b/Engine/Source/Runtime/AIModule/Private/BehaviorTree/BlackboardComponent.cpp
--- a/Engine/Source/Runtime/AIModule/Private/BehaviorTree/BlackboardComponent.cpp
+++ b/Engine/Source/Runtime/AIModule/Private/BehaviorTree/BlackboardComponent.cpp
@@ -1,9 +1,7 @@
 // Copyright 1998-2014 Epic Games, Inc. All Rights Reserved.
 
 #include "AIModulePrivate.h"
-#include "EngineUtils.h"
 #include "BehaviorTree/Blackboard/BlackboardKeyAllTypes.h"
-#include "BehaviorTree/BlackboardComponent.h"
 
 UBlackboardComponent::UBlackboardComponent(const class FPostConstructInitializeProperties& PCIP) : Super(PCIP)
 {
diff --git a/Engine/Source/Runtime/AIModule/Private/BehaviorTree/BlueprintNodeHelpers.cpp b/Engine/Source/Runtime/AIModule/Private/BehaviorTree/BlueprintNodeHelpers.cpp
--- a/Engine/Source/Runtime/AIModule/Private/BehaviorTree/BlueprintNodeHelpers.cpp
+++ b/Engine/Source/Runtime/AIModule/Private/BehaviorTree/BlueprintNodeHelpers.cpp
@@ -1,7 +1,6 @@
 // Copyright 1998-2014 Epic Games, Inc. All Rights Reserved.
 
 #include "AIModulePrivate.h"
-#include "BlueprintNodeHelpers.h"
 
 namespace BlueprintNodeHelpers
 {
@@ -76,7 +75,7 @@ namespace BlueprintNodeHelpers
 		}
 	}
 
-	FString DescribeProperty(const UProperty* Prop, const uint8* PropertyAddr)
+	AIMODULE_API FString DescribeProperty(const UProperty* Prop, const uint8* PropertyAddr)
 	{
 		FString ExportedStringValue;
 		const UStructProperty* StructProp = Cast<const UStructProperty>(Prop);
diff --git a/Engine/Source/Runtime/AIModule/Private/BehaviorTree/Decorators/BTDecorator_BlueprintBase.cpp b/Engine/Source/Runtime/AIModule/Private/BehaviorTree/Decorators/BTDecorator_BlueprintBase.cpp
--- a/Engine/Source/Runtime/AIModule/Private/BehaviorTree/Decorators/BTDecorator_BlueprintBase.cpp
+++ b/Engine/Source/Runtime/AIModule/Private/BehaviorTree/Decorators/BTDecorator_BlueprintBase.cpp
@@ -4,7 +4,6 @@
 #include "BlueprintNodeHelpers.h"
 #include "BehaviorTree/Decorators/BTDecorator_BlueprintBase.h"
 #include "BehaviorTree/BlackboardComponent.h"
-#include "BehaviorTree/BlackboardData.h"
 
 UBTDecorator_BlueprintBase::UBTDecorator_BlueprintBase(const class FPostConstructInitializeProperties& PCIP) : Super(PCIP)
 {
diff --git a/Engine/Source/Runtime/AIModule/Private/BehaviorTree/Tasks/BTTask_MakeNoise.cpp b/Engine/Source/Runtime/AIModule/Private/BehaviorTree/Tasks/BTTask_MakeNoise.cpp
--- a/Engine/Source/Runtime/AIModule/Private/BehaviorTree/Tasks/BTTask_MakeNoise.cpp
+++ b/Engine/Source/Runtime/AIModule/Private/BehaviorTree/Tasks/BTTask_MakeNoise.cpp
@@ -1,8 +1,6 @@
 // Copyright 1998-2014 Epic Games, Inc. All Rights Reserved.
 
 #include "AIModulePrivate.h"
-#include "SoundDefinitions.h"
-#include "Sound/ReverbVolume.h"
 #include "BehaviorTree/Tasks/BTTask_MakeNoise.h"
 
 UBTTask_MakeNoise::UBTTask_MakeNoise(const class FPostConstructInitializeProperties& PCIP) 
diff --git a/Engine/Source/Runtime/AIModule/Private/BehaviorTree/Tasks/BTTask_MoveTo.cpp b/Engine/Source/Runtime/AIModule/Private/BehaviorTree/Tasks/BTTask_MoveTo.cpp
--- a/Engine/Source/Runtime/AIModule/Private/BehaviorTree/Tasks/BTTask_MoveTo.cpp
+++ b/Engine/Source/Runtime/AIModule/Private/BehaviorTree/Tasks/BTTask_MoveTo.cpp
@@ -2,7 +2,6 @@
 
 #include "AIModulePrivate.h"
 #include "BehaviorTree/Tasks/BTTask_MoveTo.h"
-#include "BehaviorTree/BehaviorTree.h"
 #include "BehaviorTree/BlackboardComponent.h"
 #include "BehaviorTree/Blackboard/BlackboardKeyType_Object.h"
 #include "BehaviorTree/Blackboard/BlackboardKeyType_Vector.h"
diff --git a/Engine/Source/Runtime/AIModule/Private/BehaviorTree/Tasks/BTTask_PlaySound.cpp b/Engine/Source/Runtime/AIModule/Private/BehaviorTree/Tasks/BTTask_PlaySound.cpp
--- a/Engine/Source/Runtime/AIModule/Private/BehaviorTree/Tasks/BTTask_PlaySound.cpp
+++ b/Engine/Source/Runtime/AIModule/Private/BehaviorTree/Tasks/BTTask_PlaySound.cpp
@@ -1,7 +1,6 @@
 // Copyright 1998-2014 Epic Games, Inc. All Rights Reserved.
 
 #include "AIModulePrivate.h"
-#include "SoundDefinitions.h"
 #include "Kismet/GameplayStatics.h"
 #include "Sound/SoundCue.h"
 #include "BehaviorTree/Tasks/BTTask_PlaySound.h"
diff --git a/Engine/Source/Runtime/AIModule/Private/EnvironmentQuery/Contexts/EnvQueryContext_BlueprintBase.cpp b/Engine/Source/Runtime/AIModule/Private/EnvironmentQuery/Contexts/EnvQueryContext_BlueprintBase.cpp
--- a/Engine/Source/Runtime/AIModule/Private/EnvironmentQuery/Contexts/EnvQueryContext_BlueprintBase.cpp
+++ b/Engine/Source/Runtime/AIModule/Private/EnvironmentQuery/Contexts/EnvQueryContext_BlueprintBase.cpp
@@ -1,7 +1,6 @@
 // Copyright 1998-2014 Epic Games, Inc. All Rights Reserved.
 
 #include "AIModulePrivate.h"
-#include "BlueprintNodeHelpers.h"
 #include "EnvironmentQuery/Contexts/EnvQueryContext_BlueprintBase.h"
 #include "EnvironmentQuery/Items/EnvQueryAllItemTypes.h"
 #include "EnvironmentQuery/EnvQueryTypes.h"
diff --git a/Engine/Source/Runtime/AIModule/Private/EnvironmentQuery/EQSTestingPawn.cpp b/Engine/Source/Runtime/AIModule/Private/EnvironmentQuery/EQSTestingPawn.cpp
--- a/Engine/Source/Runtime/AIModule/Private/EnvironmentQuery/EQSTestingPawn.cpp
+++ b/Engine/Source/Runtime/AIModule/Private/EnvironmentQuery/EQSTestingPawn.cpp
@@ -1,7 +1,6 @@
 // Copyright 1998-2014 Epic Games, Inc. All Rights Reserved.
 
 #include "AIModulePrivate.h"
-#include "GameplayDebuggingTypes.h"
 #include "EnvironmentQuery/EnvQuery.h"
 #include "EnvironmentQuery/EQSTestingPawn.h"
 #include "EnvironmentQuery/EnvQueryTypes.h"
@@ -13,7 +12,6 @@
 #endif // WITH_EDITORONLY_DATA
 
 #if WITH_EDITOR
-#include "Engine/Brush.h"
 #include "Editor/EditorEngine.h"
 extern UNREALED_API class UEditorEngine* GEditor;
 #endif // WITH_EDITOR
diff --git a/Engine/Source/Runtime/AIModule/Private/EnvironmentQuery/EnvQueryManager.cpp b/Engine/Source/Runtime/AIModule/Private/EnvironmentQuery/EnvQueryManager.cpp
--- a/Engine/Source/Runtime/AIModule/Private/EnvironmentQuery/EnvQueryManager.cpp
+++ b/Engine/Source/Runtime/AIModule/Private/EnvironmentQuery/EnvQueryManager.cpp
@@ -9,7 +9,6 @@
 #include "EnvironmentQuery/EQSTestingPawn.h"
 
 #if WITH_EDITOR
-#include "Engine/Brush.h"
 #include "Editor/EditorEngine.h"
 #include "EngineUtils.h"
 
diff --git a/Engine/Source/Runtime/AIModule/Private/EnvironmentQuery/EnvQueryTest.cpp b/Engine/Source/Runtime/AIModule/Private/EnvironmentQuery/EnvQueryTest.cpp
--- a/Engine/Source/Runtime/AIModule/Private/EnvironmentQuery/EnvQueryTest.cpp
+++ b/Engine/Source/Runtime/AIModule/Private/EnvironmentQuery/EnvQueryTest.cpp
@@ -1,8 +1,6 @@
 // Copyright 1998-2014 Epic Games, Inc. All Rights Reserved.
 
 #include "AIModulePrivate.h"
-#include "EnvironmentQuery/EnvQueryOption.h"
-#include "EnvironmentQuery/EnvQueryContext.h"
 #include "EnvironmentQuery/EnvQueryManager.h"
 #include "EnvironmentQuery/Contexts/EnvQueryContext_Item.h"
 #include "EnvironmentQuery/Items/EnvQueryItemType_VectorBase.h"
diff --git a/Engine/Source/Runtime/AIModule/Private/EnvironmentQuery/EnvQueryTypes.cpp b/Engine/Source/Runtime/AIModule/Private/EnvironmentQuery/EnvQueryTypes.cpp
--- a/Engine/Source/Runtime/AIModule/Private/EnvironmentQuery/EnvQueryTypes.cpp
+++ b/Engine/Source/Runtime/AIModule/Private/EnvironmentQuery/EnvQueryTypes.cpp
@@ -2,7 +2,6 @@
 
 #include "AIModulePrivate.h"
 #include "EnvironmentQuery/EnvQueryTypes.h"
-#include "EnvironmentQuery/EnvQueryContext.h"
 #include "EnvironmentQuery/Items/EnvQueryItemType.h"
 #include "EnvironmentQuery/Items/EnvQueryItemType_ActorBase.h"
 #include "EnvironmentQuery/Items/EnvQueryItemType_VectorBase.h"
@@ -92,6 +91,9 @@ FString UEnvQueryTypes::DescribeBoolParam(const FEnvBoolParam& Param)
 	return Param.Value ? TEXT("true") : TEXT("false");
 }
 
+FEnvDirection::~FEnvDirection()
+{ }
+
 FText FEnvDirection::ToText() const
 {
 	if(DirMode == EEnvDirection::TwoPoints)
diff --git a/Engine/Source/Runtime/AIModule/Private/Navigation/CrowdFollowingComponent.cpp b/Engine/Source/Runtime/AIModule/Private/Navigation/CrowdFollowingComponent.cpp
--- a/Engine/Source/Runtime/AIModule/Private/Navigation/CrowdFollowingComponent.cpp
+++ b/Engine/Source/Runtime/AIModule/Private/Navigation/CrowdFollowingComponent.cpp
@@ -4,7 +4,6 @@
 #include "Navigation/CrowdFollowingComponent.h"
 #include "Navigation/CrowdManager.h"
 #include "Navigation/NavigationComponent.h"
-#include "AI/Navigation/NavLinkCustomInterface.h"
 #include "GameFramework/CharacterMovementComponent.h"
 
 
diff --git a/Engine/Source/Runtime/AIModule/Private/Navigation/PathFollowingComponent.cpp b/Engine/Source/Runtime/AIModule/Private/Navigation/PathFollowingComponent.cpp
--- a/Engine/Source/Runtime/AIModule/Private/Navigation/PathFollowingComponent.cpp
+++ b/Engine/Source/Runtime/AIModule/Private/Navigation/PathFollowingComponent.cpp
@@ -10,7 +10,6 @@
 #include "GameFramework/NavMovementComponent.h"
 #include "GameFramework/Character.h"
 #include "Engine/Canvas.h"
-#include "TimerManager.h"
 
 #include "Navigation/PathFollowingComponent.h"
 
diff --git a/Engine/Source/Runtime/Core/Public/Misc/SecureHash.h b/Engine/Source/Runtime/Core/Public/Misc/SecureHash.h
--- a/Engine/Source/Runtime/Core/Public/Misc/SecureHash.h
+++ b/Engine/Source/Runtime/Core/Public/Misc/SecureHash.h
@@ -6,6 +6,8 @@
 
 #pragma once
 
+#include "Platform.h"
+
 /*-----------------------------------------------------------------------------
 	MD5 functions.
 -----------------------------------------------------------------------------*/
diff --git a/Engine/Source/Runtime/Engine/Classes/Components/LightComponent.h b/Engine/Source/Runtime/Engine/Classes/Components/LightComponent.h
--- a/Engine/Source/Runtime/Engine/Classes/Components/LightComponent.h
+++ b/Engine/Source/Runtime/Engine/Classes/Components/LightComponent.h
@@ -4,7 +4,7 @@
 
 #include "LightComponentBase.h"
 #include "SceneTypes.h"
-
+#include "EngineDefines.h"
 #include "LightComponent.generated.h"
 
 /** 
diff --git a/Engine/Source/Runtime/Engine/Classes/Components/PrimitiveComponent.h b/Engine/Source/Runtime/Engine/Classes/Components/PrimitiveComponent.h
--- a/Engine/Source/Runtime/Engine/Classes/Components/PrimitiveComponent.h
+++ b/Engine/Source/Runtime/Engine/Classes/Components/PrimitiveComponent.h
@@ -15,7 +15,7 @@
 #include "PrimitiveComponent.generated.h"
 
 class FPrimitiveSceneProxy;
-
+class AController;
 /** Information about a vertex of a primitive's triangle. */
 struct FPrimitiveTriangleVertex
 {
diff --git a/Engine/Source/Runtime/Engine/Classes/Engine/EngineTypes.h b/Engine/Source/Runtime/Engine/Classes/Engine/EngineTypes.h
--- a/Engine/Source/Runtime/Engine/Classes/Engine/EngineTypes.h
+++ b/Engine/Source/Runtime/Engine/Classes/Engine/EngineTypes.h
@@ -7,6 +7,7 @@
  *	use in multiple files where the enum can't be mapped to a specific file.
  */
 #include "NetSerialization.h"
+#include "UObject/NameTypes.h"
 #include "EngineTypes.generated.h"
 
 UENUM()
diff --git a/Engine/Source/Runtime/Engine/Classes/Engine/GameViewportClient.h b/Engine/Source/Runtime/Engine/Classes/Engine/GameViewportClient.h
--- a/Engine/Source/Runtime/Engine/Classes/Engine/GameViewportClient.h
+++ b/Engine/Source/Runtime/Engine/Classes/Engine/GameViewportClient.h
@@ -19,6 +19,8 @@
 #include "Engine/ScriptViewportClient.h"
 #include "GameViewportClient.generated.h"
 
+class UNetDriver;
+
 /**
  * Enum of the different splitscreen types
  */
diff --git a/Engine/Source/Runtime/Engine/Classes/Engine/ObjectLibrary.h b/Engine/Source/Runtime/Engine/Classes/Engine/ObjectLibrary.h
--- a/Engine/Source/Runtime/Engine/Classes/Engine/ObjectLibrary.h
+++ b/Engine/Source/Runtime/Engine/Classes/Engine/ObjectLibrary.h
@@ -3,6 +3,8 @@
 #pragma once
 #include "ObjectLibrary.generated.h"
 
+class FAssetData;
+
 /** Class that holds a library of Objects */
 UCLASS(MinimalAPI)
 class UObjectLibrary : public UObject
@@ -36,7 +38,7 @@ protected:
 	bool bIsFullyLoaded;
 
 	/** Asset data of objects that will belong in library, possibly not loaded yet */
-	TArray<class FAssetData>	AssetDataList;
+	TArray<FAssetData>	AssetDataList;
 
 public:
 	// Begin UObject Interface
diff --git a/Engine/Source/Runtime/Engine/Classes/Engine/StaticMesh.h b/Engine/Source/Runtime/Engine/Classes/Engine/StaticMesh.h
--- a/Engine/Source/Runtime/Engine/Classes/Engine/StaticMesh.h
+++ b/Engine/Source/Runtime/Engine/Classes/Engine/StaticMesh.h
@@ -10,6 +10,9 @@
 
 // Forward declarations
 class UFoliageType_InstancedStaticMesh;
+class FStaticMeshRenderData;
+class UNavCollision;
+struct FTriMeshCollisionData;
 
 /*-----------------------------------------------------------------------------
 	Legacy mesh optimization settings.
@@ -317,7 +320,7 @@ class UStaticMesh : public UObject, public IInterface_CollisionDataProvider, pub
 	GENERATED_UCLASS_BODY()
 
 	/** Pointer to the data used to render this static mesh. */
-	TScopedPointer<class FStaticMeshRenderData> RenderData;
+	TScopedPointer<FStaticMeshRenderData> RenderData;
 
 #if WITH_EDITORONLY_DATA
 	/** Imported raw mesh bulk data. */
@@ -458,7 +461,7 @@ protected:
 public:
 	/** Pre-build navigation collision */
 	UPROPERTY(VisibleAnywhere, transient, duplicatetransient, editinline, Category=Navigation)
-	class UNavCollision* NavCollision;
+	UNavCollision* NavCollision;
 	
 public:
 	/**
@@ -550,7 +553,7 @@ public:
 	ENGINE_API static void CheckLightMapUVs( UStaticMesh* InStaticMesh, TArray< FString >& InOutAssetsWithMissingUVSets, TArray< FString >& InOutAssetsWithBadUVSets, TArray< FString >& InOutAssetsWithValidUVSets, bool bInVerbose = true );
 
 	// Begin Interface_CollisionDataProvider Interface
-	ENGINE_API virtual bool GetPhysicsTriMeshData(struct FTriMeshCollisionData* CollisionData, bool InUseAllTriData) override;
+	ENGINE_API virtual bool GetPhysicsTriMeshData(FTriMeshCollisionData* CollisionData, bool InUseAllTriData) override;
 	ENGINE_API virtual bool ContainsPhysicsTriMeshData(bool InUseAllTriData) const override;
 	virtual bool WantsNegXTriMesh() override
 	{
diff --git a/Engine/Source/Runtime/Engine/Classes/Engine/World.h b/Engine/Source/Runtime/Engine/Classes/Engine/World.h
--- a/Engine/Source/Runtime/Engine/Classes/Engine/World.h
+++ b/Engine/Source/Runtime/Engine/Classes/Engine/World.h
@@ -24,6 +24,10 @@ class ABrush;
 class UModel;
 class APhysicsVolume;
 class UTexture2D;
+class AController;
+class APlayerController;
+class APawn;
+class ULevel;
 
 template<typename,typename> class TOctree;
 
@@ -34,7 +38,7 @@ template<typename,typename> class TOctree;
 typedef TArray<TAutoWeakObjectPtr<AController> >::TConstIterator FConstControllerIterator;
 typedef TArray<TAutoWeakObjectPtr<APlayerController> >::TConstIterator FConstPlayerControllerIterator;
 typedef TArray<TAutoWeakObjectPtr<APawn> >::TConstIterator FConstPawnIterator;	
-typedef TArray<class ULevel*>::TConstIterator FConstLevelIterator;
+typedef TArray<ULevel*>::TConstIterator FConstLevelIterator;
 
 DECLARE_LOG_CATEGORY_EXTERN(LogSpawn, Warning, All);
 
diff --git a/Engine/Source/Runtime/Engine/Classes/Exporters/Exporter.h b/Engine/Source/Runtime/Engine/Classes/Exporters/Exporter.h
--- a/Engine/Source/Runtime/Engine/Classes/Exporters/Exporter.h
+++ b/Engine/Source/Runtime/Engine/Classes/Exporters/Exporter.h
@@ -9,6 +9,8 @@
 #pragma once
 #include "Exporter.generated.h"
 
+class UActorComponent;
+
 UCLASS(abstract, transient, MinimalAPI)
 class UExporter : public UObject
 {
diff --git a/Engine/Source/Runtime/Engine/Classes/Sound/SoundBase.h b/Engine/Source/Runtime/Engine/Classes/Sound/SoundBase.h
--- a/Engine/Source/Runtime/Engine/Classes/Sound/SoundBase.h
+++ b/Engine/Source/Runtime/Engine/Classes/Sound/SoundBase.h
@@ -9,6 +9,7 @@
 #include "SoundBase.generated.h"
 
 struct FActiveSound;
+struct FAttenuationSettings;
 struct FSoundParseParameters;
 struct FWaveInstance;
 class  USoundClass;
diff --git a/Engine/Source/Runtime/Engine/Classes/Sound/SoundCue.h b/Engine/Source/Runtime/Engine/Classes/Sound/SoundCue.h
--- a/Engine/Source/Runtime/Engine/Classes/Sound/SoundCue.h
+++ b/Engine/Source/Runtime/Engine/Classes/Sound/SoundCue.h
@@ -3,6 +3,7 @@
 #pragma once
 
 #include "SoundClass.h"
+#include "Sound/SoundAttenuation.h"
 #include "Sound/SoundBase.h"
 
 #include "SoundCue.generated.h"
diff --git a/Engine/Source/Runtime/Engine/Public/CollisionQueryParams.h b/Engine/Source/Runtime/Engine/Public/CollisionQueryParams.h
--- a/Engine/Source/Runtime/Engine/Public/CollisionQueryParams.h
+++ b/Engine/Source/Runtime/Engine/Public/CollisionQueryParams.h
@@ -4,6 +4,8 @@
 
 #pragma once
 
+#include "GameFramework/Actor.h"
+
 /** Macro to convert ECollisionChannels to bit flag **/
 #define ECC_TO_BITFIELD(x)	(1<<(x))
 /** Macro to convert from CollisionResponseContainer to bit flag **/
diff --git a/Engine/Source/Runtime/Engine/Public/MaterialShared.h b/Engine/Source/Runtime/Engine/Public/MaterialShared.h
--- a/Engine/Source/Runtime/Engine/Public/MaterialShared.h
+++ b/Engine/Source/Runtime/Engine/Public/MaterialShared.h
@@ -25,6 +25,7 @@ class UMaterial;
 class UMaterialInstance;
 class UMaterialExpression;
 class UMaterialInterface;
+class UTexture;
 struct FExpressionInput;
 
 #define ME_CAPTION_HEIGHT		18
diff --git a/Engine/Source/Runtime/Engine/Public/Model.h b/Engine/Source/Runtime/Engine/Public/Model.h
--- a/Engine/Source/Runtime/Engine/Public/Model.h
+++ b/Engine/Source/Runtime/Engine/Public/Model.h
@@ -16,6 +16,8 @@ class ABrush;
 class UModel;
 class ULightComponent;
 class UPolys;
+class UMaterialInterface;
+class ULevel;
 
 //
 // One vertex associated with a Bsp node's polygon.  Contains a vertex index
diff --git a/Engine/Source/Runtime/Engine/Public/SkeletalMeshTypes.h b/Engine/Source/Runtime/Engine/Public/SkeletalMeshTypes.h
--- a/Engine/Source/Runtime/Engine/Public/SkeletalMeshTypes.h
+++ b/Engine/Source/Runtime/Engine/Public/SkeletalMeshTypes.h
@@ -14,6 +14,8 @@
 #include "GPUSkinPublicDefs.h"
 #include "Components/PrimitiveComponent.h"
 #include "PrimitiveSceneProxy.h"
+#include "Misc/OutputDevice.h"
+#include "EngineLogs.h"
 
 // Define that controls showing chart of distance factors for skel meshes during entire run of the game on exit.
 #define CHART_DISTANCE_FACTORS 0
@@ -125,17 +127,17 @@ struct FMeshBoneInfo
 	FMeshBoneInfo() : Name(NAME_None), ParentIndex(INDEX_NONE) {}
 
 	FMeshBoneInfo(const FName & InName, const FString& InExportName, int32 InParentIndex)
-	:	Name(InName)
+		:	Name(InName)
 #if WITH_EDITORONLY_DATA
-	, ExportName(InExportName)
+		,	ExportName(InExportName)
 #endif
-	,	ParentIndex(InParentIndex)
+		,	ParentIndex(InParentIndex)
 	{}
 
 	FMeshBoneInfo(const FMeshBoneInfo & Other)
 		:	Name(Other.Name)
 #if WITH_EDITORONLY_DATA
-		, ExportName(Other.ExportName)
+		,	ExportName(Other.ExportName)
 #endif
 		,	ParentIndex(Other.ParentIndex)
 	{}
diff --git a/Engine/Source/Runtime/Navmesh/Private/DebugUtils/DetourDebugDraw.cpp b/Engine/Source/Runtime/Navmesh/Private/DebugUtils/DetourDebugDraw.cpp
--- a/Engine/Source/Runtime/Navmesh/Private/DebugUtils/DetourDebugDraw.cpp
+++ b/Engine/Source/Runtime/Navmesh/Private/DebugUtils/DetourDebugDraw.cpp
@@ -20,7 +20,6 @@
 //
 
 #include "NavmeshModulePrivatePCH.h"
-#include <math.h>
 #include "DebugDraw.h"
 #include "DetourDebugDraw.h"
 #include "DetourNavMesh.h"
diff --git a/Engine/Source/Runtime/Navmesh/Private/DetourCrowd/DetourCrowd.cpp b/Engine/Source/Runtime/Navmesh/Private/DetourCrowd/DetourCrowd.cpp
--- a/Engine/Source/Runtime/Navmesh/Private/DetourCrowd/DetourCrowd.cpp
+++ b/Engine/Source/Runtime/Navmesh/Private/DetourCrowd/DetourCrowd.cpp
@@ -23,7 +23,6 @@
 #define _USE_MATH_DEFINES
 #include <math.h>
 #include <string.h>
-#include <float.h>
 #include <stdlib.h>
 #include <new>
 #include "DetourCrowd.h"
diff --git a/Engine/Source/Runtime/Navmesh/Private/DetourTileCache/DetourTileCacheRegion.cpp b/Engine/Source/Runtime/Navmesh/Private/DetourTileCache/DetourTileCacheRegion.cpp
--- a/Engine/Source/Runtime/Navmesh/Private/DetourTileCache/DetourTileCacheRegion.cpp
+++ b/Engine/Source/Runtime/Navmesh/Private/DetourTileCache/DetourTileCacheRegion.cpp
@@ -25,9 +25,6 @@
 #include "DetourAssert.h"
 #include "DetourTileCacheBuilder.h"
 #include <string.h>
-#include <float.h>
-#define _USE_MATH_DEFINES
-#include <math.h>
 
 static const unsigned short DT_BORDER_REG = 0x8000;
 
diff --git a/Engine/Source/Runtime/SlateRHIRenderer/Private/SlateRHIRenderingPolicy.cpp b/Engine/Source/Runtime/SlateRHIRenderer/Private/SlateRHIRenderingPolicy.cpp
--- a/Engine/Source/Runtime/SlateRHIRenderer/Private/SlateRHIRenderingPolicy.cpp
+++ b/Engine/Source/Runtime/SlateRHIRenderer/Private/SlateRHIRenderingPolicy.cpp
@@ -9,6 +9,7 @@
 #include "TileRendering.h"
 #include "PreviewScene.h"
 #include "EngineModule.h"
+#include "SlateMaterialShader.h"
 
 static EPrimitiveType GetRHIPrimitiveType( ESlateDrawPrimitive::Type SlateType )
 {
