Fixed comment for DrawColor in ULevelStreaming
#TTP 340071

diff --git a/Engine/Source/Runtime/Engine/Classes/Engine/LevelStreaming.h b/Engine/Source/Runtime/Engine/Classes/Engine/LevelStreaming.h
--- a/Engine/Source/Runtime/Engine/Classes/Engine/LevelStreaming.h
+++ b/Engine/Source/Runtime/Engine/Classes/Engine/LevelStreaming.h
@@ -125,7 +125,7 @@ class ULevelStreaming : public UObject
 	UPROPERTY(EditAnywhere, Category=LevelStreaming)
 	uint32 bDrawOnLevelStatusMap:1;
 
-	/** The level's color; used to make the level easily identifiable in the level browser, for actor level visulization, etc.	*/
+	/** The level color used for visualization. (Show -> Advanced -> Level Coloration) */
 	UPROPERTY(EditAnywhere, Category=LevelStreaming)
 	FColor DrawColor;
 
