#TTP 333841 -  CRASH: Occurs when running the "Static Mesh UVs" test. (UAnimInstance::InitializeAnimation)

#proj UE4

#summary Investigated issue and removed old fix up code from Landscape PostLoad that was causing the issue

diff --git a/Engine/Source/Runtime/Engine/Private/Landscape/Landscape.cpp b/Engine/Source/Runtime/Engine/Private/Landscape/Landscape.cpp
--- a/Engine/Source/Runtime/Engine/Private/Landscape/Landscape.cpp
+++ b/Engine/Source/Runtime/Engine/Private/Landscape/Landscape.cpp
@@ -389,11 +389,6 @@ void ULandscapeComponent::PostLoad()
 				WeightmapTextures[Idx]->ClearFlags(RF_Standalone);
 			}
 		}
-
-		if (GetLinkerUE4Version() < VER_UE4_FIXUP_TERRAIN_LAYER_NODES)
-		{
-			UpdateMaterialInstances();
-		}
 	}
 #endif
 
