Fix Mac link error seemingly caused by unnecessary use of 'class' specifier.

diff --git a/Engine/Source/Runtime/Engine/Classes/GameFramework/Pawn.h b/Engine/Source/Runtime/Engine/Classes/GameFramework/Pawn.h
--- a/Engine/Source/Runtime/Engine/Classes/GameFramework/Pawn.h
+++ b/Engine/Source/Runtime/Engine/Classes/GameFramework/Pawn.h
@@ -201,7 +201,7 @@ public:
 	/** Basically retrieved pawn's location on navmesh */
 	UFUNCTION(BlueprintCallable, Category="Pawn")
 	virtual FVector GetNavAgentLocation() const override { return GetActorLocation() - FVector(0.f, 0.f, BaseEyeHeight); }
-	virtual void GetMoveGoalReachTest(class AActor* MovingActor, const FVector& MoveOffset, FVector& GoalOffset, float& GoalRadius, float& GoalHalfHeight) const override;
+	virtual void GetMoveGoalReachTest(AActor* MovingActor, const FVector& MoveOffset, FVector& GoalOffset, float& GoalRadius, float& GoalHalfHeight) const override;
 	// End INavAgentInterface Interface
 
 	/** updates MovementComponent's parameters used by navigation system */
diff --git a/Engine/Source/Runtime/Engine/Private/Pawn.cpp b/Engine/Source/Runtime/Engine/Private/Pawn.cpp
--- a/Engine/Source/Runtime/Engine/Private/Pawn.cpp
+++ b/Engine/Source/Runtime/Engine/Private/Pawn.cpp
@@ -831,7 +831,7 @@ void APawn::DisableInput(class APlayerController* PlayerController)
 	}
 }
 
-void APawn::GetMoveGoalReachTest(class AActor* MovingActor, const FVector& MoveOffset, FVector& GoalOffset, float& GoalRadius, float& GoalHalfHeight) const 
+void APawn::GetMoveGoalReachTest(AActor* MovingActor, const FVector& MoveOffset, FVector& GoalOffset, float& GoalRadius, float& GoalHalfHeight) const 
 {
 	GoalOffset = FVector::ZeroVector;
 	GetSimpleCollisionCylinder(GoalRadius, GoalHalfHeight);
