#ue4 - Move UpdateFromCompressedFlags() from Character to CharacterMovementComponent.

diff --git a/Engine/Source/Runtime/Engine/Classes/GameFramework/Character.h b/Engine/Source/Runtime/Engine/Classes/GameFramework/Character.h
--- a/Engine/Source/Runtime/Engine/Classes/GameFramework/Character.h
+++ b/Engine/Source/Runtime/Engine/Classes/GameFramework/Character.h
@@ -507,8 +507,8 @@ public:
 	 */
 	virtual float GetJumpMaxHoldTime() const;
 
-	/** Unpack compressed flags from a saved move and set state accordingly. See FSavedMove_Character. */
-	virtual void UpdateFromCompressedFlags(uint8 Flags);
+	DEPRECATED(4.4, "UpdateFromCompressedFlags has moved to UCharacterMovementComponent")
+	virtual void UpdateFromCompressedFlags(uint8 Flags) { }
 
 public:
 	// Note: these network functions should be moved to the Character movement component. They are currently just wrappers that pass the calls to CharacterMovement.
diff --git a/Engine/Source/Runtime/Engine/Classes/GameFramework/CharacterMovementComponent.h b/Engine/Source/Runtime/Engine/Classes/GameFramework/CharacterMovementComponent.h
--- a/Engine/Source/Runtime/Engine/Classes/GameFramework/CharacterMovementComponent.h
+++ b/Engine/Source/Runtime/Engine/Classes/GameFramework/CharacterMovementComponent.h
@@ -1358,6 +1358,9 @@ protected:
 	/* Process a move at the given time stamp, given the compressed flags representing various events that occurred (ie jump). */
 	virtual void MoveAutonomous( float ClientTimeStamp, float DeltaTime, uint8 CompressedFlags, const FVector& NewAccel);
 
+	/** Unpack compressed flags from a saved move and set state accordingly. See FSavedMove_Character. */
+	virtual void UpdateFromCompressedFlags(uint8 Flags);
+
 public:
 
 	/** React to instantaneous change in position. Invalidates cached floor recomputes it if possible if there is a current movement base. */
diff --git a/Engine/Source/Runtime/Engine/Private/Character.cpp b/Engine/Source/Runtime/Engine/Private/Character.cpp
--- a/Engine/Source/Runtime/Engine/Private/Character.cpp
+++ b/Engine/Source/Runtime/Engine/Private/Character.cpp
@@ -1130,20 +1130,6 @@ void ACharacter::GetLifetimeReplicatedProps( TArray< FLifetimeProperty > & OutLi
 	DOREPLIFETIME_CONDITION( ACharacter, bIsCrouched,		COND_SimulatedOnly );
 }
 
-void ACharacter::UpdateFromCompressedFlags(uint8 Flags)
-{
-	const bool bWasJumping = bPressedJump;
-
-	bPressedJump = ((Flags & FSavedMove_Character::FLAG_JumpPressed) != 0);	
-	CharacterMovement->bWantsToCrouch = ((Flags & FSavedMove_Character::FLAG_WantsToCrouch) != 0);
-
-	// Reset JumpKeyHoldTime when player presses Jump key on server as well.
-	if (!bWasJumping && bPressedJump)
-	{
-		JumpKeyHoldTime = 0.0f;
-	}
-}
-
 bool ACharacter::IsPlayingRootMotion() const
 {
 	return (GetRootMotionAnimMontageInstance() != NULL);
diff --git a/Engine/Source/Runtime/Engine/Private/CharacterMovementComponent.cpp b/Engine/Source/Runtime/Engine/Private/CharacterMovementComponent.cpp
--- a/Engine/Source/Runtime/Engine/Private/CharacterMovementComponent.cpp
+++ b/Engine/Source/Runtime/Engine/Private/CharacterMovementComponent.cpp
@@ -5408,7 +5408,7 @@ void UCharacterMovementComponent::MoveAutonomous
 		return;
 	}
 
-	CharacterOwner->UpdateFromCompressedFlags(CompressedFlags);
+	UpdateFromCompressedFlags(CompressedFlags);
 	CharacterOwner->CheckJumpInput(DeltaTime);
 
 	Acceleration = ConstrainInputAcceleration(NewAccel);
@@ -6356,5 +6356,22 @@ uint8 FSavedMove_Character::GetCompressedFlags() const
 	return Result;
 }
 
+void UCharacterMovementComponent::UpdateFromCompressedFlags(uint8 Flags)
+{
+	if (!CharacterOwner)
+	{
+		return;
+	}
+
+	const bool bWasJumping = CharacterOwner->bPressedJump;
 
+	CharacterOwner->bPressedJump = ((Flags & FSavedMove_Character::FLAG_JumpPressed) != 0);	
+	bWantsToCrouch = ((Flags & FSavedMove_Character::FLAG_WantsToCrouch) != 0);
+
+	// Reset JumpKeyHoldTime when player presses Jump key on server as well.
+	if (!bWasJumping && CharacterOwner->bPressedJump)
+	{
+		CharacterOwner->JumpKeyHoldTime = 0.0f;
+	}
+}
 
