More case-sensitivity fixes for QAGame.

diff --git a/Engine/Source/Runtime/Engine/Classes/AI/Navigation/NavLinkProxy.h b/Engine/Source/Runtime/Engine/Classes/AI/Navigation/NavLinkProxy.h
--- a/Engine/Source/Runtime/Engine/Classes/AI/Navigation/NavLinkProxy.h
+++ b/Engine/Source/Runtime/Engine/Classes/AI/Navigation/NavLinkProxy.h
@@ -2,7 +2,7 @@
 
 #pragma once
 #include "NavigationTypes.h"
-#include "Ai/Navigation/NavRelevantActorInterface.h"
+#include "AI/Navigation/NavRelevantActorInterface.h"
 #include "NavLinkProxy.generated.h"
 
 DECLARE_DYNAMIC_MULTICAST_DELEGATE_TwoParams( FSmartLinkReachedSignature, class AActor*, MovingActor, const FVector&, DestinationPoint );
diff --git a/Engine/Source/Runtime/Engine/Classes/Engine/StaticMeshActor.h b/Engine/Source/Runtime/Engine/Classes/Engine/StaticMeshActor.h
--- a/Engine/Source/Runtime/Engine/Classes/Engine/StaticMeshActor.h
+++ b/Engine/Source/Runtime/Engine/Classes/Engine/StaticMeshActor.h
@@ -1,7 +1,7 @@
 // Copyright 1998-2014 Epic Games, Inc. All Rights Reserved.
 
 #pragma once
-#include "Ai/Navigation/NavRelevantActorInterface.h"
+#include "AI/Navigation/NavRelevantActorInterface.h"
 #include "StaticMeshActor.generated.h"
 
 /**
diff --git a/Engine/Source/Runtime/Engine/Private/AI/Navigation/RecastNavMeshGenerator.cpp b/Engine/Source/Runtime/Engine/Private/AI/Navigation/RecastNavMeshGenerator.cpp
--- a/Engine/Source/Runtime/Engine/Private/AI/Navigation/RecastNavMeshGenerator.cpp
+++ b/Engine/Source/Runtime/Engine/Private/AI/Navigation/RecastNavMeshGenerator.cpp
@@ -11,7 +11,7 @@
 #include "RecastNavMeshGenerator.h"
 #include "PImplRecastNavMesh.h"
 #include "SurfaceIterators.h"
-#include "Ai/Navigation/NavMeshBoundsVolume.h"
+#include "AI/Navigation/NavMeshBoundsVolume.h"
 
 // recast includes
 #include "Recast.h"
