#ue4 - NaN checks for MTD results from PhysX.

TTP 340254

diff --git a/Engine/Source/Runtime/Engine/Private/Collision/CollisionConversions.cpp b/Engine/Source/Runtime/Engine/Private/Collision/CollisionConversions.cpp
--- a/Engine/Source/Runtime/Engine/Private/Collision/CollisionConversions.cpp
+++ b/Engine/Source/Runtime/Engine/Private/Collision/CollisionConversions.cpp
@@ -515,28 +515,31 @@ static bool ConvertOverlappedShapeToImpactHit(const PxLocationHit& PHit, const F
 	OutResult.Location = P2UVector(QueryTM.p);
 	OutResult.ImpactPoint = OutResult.Location; // @todo not really sure of a better thing to do here...
 
-	// Fallback normal if we can't find it with MTD or otherwise.
-	OutResult.ImpactNormal = (PHit.flags & PxHitFlag::eNORMAL) ? P2UVector(PHit.normal) : FVector::UpVector;
+	const bool bFiniteNormal = PHit.normal.isFinite();
+	const bool bValidNormal = (PHit.flags & PxHitFlag::eNORMAL) && bFiniteNormal;
 
 	// Use MTD result if possible. We interpret the MTD vector as both the direction to move and the opposing normal.
-	const bool bFiniteNormal = PHit.normal.isFinite();
-	const bool bValidMtdNormal = (PHit.flags & PxHitFlag::eNORMAL) && bFiniteNormal;
-	if (bValidMtdNormal)
+	if (bValidNormal)
 	{
 		OutResult.ImpactNormal = P2UVector(PHit.normal);
 		OutResult.PenetrationDepth = FMath::Abs(PHit.distance);
 	}
-	else if (!bFiniteNormal)
+	else
 	{
-		UE_LOG(LogPhysics, Warning, TEXT("ConvertOverlappedShapeToImpactHit, MTD returned NaN :( normal: (X:%f, Y:%f, Z:%f)"), PHit.normal.x, PHit.normal.y, PHit.normal.z);
+		// Fallback normal if we can't find it with MTD or otherwise.
+		OutResult.ImpactNormal = FVector::UpVector;
 		OutResult.PenetrationDepth = 0.f;
+		if (!bFiniteNormal)
+		{
+			UE_LOG(LogPhysics, Warning, TEXT("ConvertOverlappedShapeToImpactHit, MTD returned NaN :( normal: (X:%f, Y:%f, Z:%f)"), PHit.normal.x, PHit.normal.y, PHit.normal.z);
+		}
 	}
 
 	// Zero-distance hits are often valid hits and we can extract the hit normal from the face index
 	// For invalid normals we can try other methods as well (get overlapping triangles).
 	static const float MtdInflation = 0.125f;
 
-	if (PHit.distance == 0.f || !bValidMtdNormal)
+	if (PHit.distance == 0.f || !bValidNormal)
 	{
 		const PxTransform PShapeWorldPose = PxShapeExt::getGlobalPose(*PShape, *PActor); 
 		PxTriangleMeshGeometry PTriMeshGeom;
