Fixed viewport resizing if GetHMDMonitorInfo() returns false.

diff --git a/Engine/Source/Runtime/Engine/Private/Slate/SceneViewport.cpp b/Engine/Source/Runtime/Engine/Private/Slate/SceneViewport.cpp
--- a/Engine/Source/Runtime/Engine/Private/Slate/SceneViewport.cpp
+++ b/Engine/Source/Runtime/Engine/Private/Slate/SceneViewport.cpp
@@ -827,10 +827,12 @@ void FSceneViewport::ResizeFrame(uint32 NewSizeX, uint32 NewSizeY, EWindowMode::
 						FSlateRect PreFullScreenRect = WindowToResize->GetRectInScreen();
 
 						IHeadMountedDisplay::MonitorInfo MonitorInfo;
-						GEngine->HMDDevice->GetHMDMonitorInfo(MonitorInfo);
-						NewSizeX = MonitorInfo.ResolutionX;
-						NewSizeY = MonitorInfo.ResolutionY;
-						WindowToResize->ReshapeWindow(FVector2D(MonitorInfo.DesktopX, MonitorInfo.DesktopY), FVector2D(MonitorInfo.ResolutionX, MonitorInfo.ResolutionY));
+						if (GEngine->HMDDevice->GetHMDMonitorInfo(MonitorInfo))
+						{
+							NewSizeX = MonitorInfo.ResolutionX;
+							NewSizeY = MonitorInfo.ResolutionY;
+							WindowToResize->ReshapeWindow(FVector2D(MonitorInfo.DesktopX, MonitorInfo.DesktopY), FVector2D(MonitorInfo.ResolutionX, MonitorInfo.ResolutionY));
+						}
 
 						GEngine->HMDDevice->PushPreFullScreenRect(PreFullScreenRect);
 					}
