Added Audio Streaming Manager and streaming audio playback

Audio Streaming Manager will keep track of SoundWaves that are marked for streaming and any sound sources that are using those waves.
Added functionality to all Sound Buffers to report which chunk of audio they are currently reading from, which the manager will use to request the current and next chunk be loaded in memory.
Added a struct to contain all of the additional information needed for a Streaming SoundWave and set off requests, which the manager will also keep track of to avoid increasing the memory size of all SoundWaves.
Tidied up the Opus Decompressor to reduce the amount of repeated code and share the basic functionality for streaming and non-streaming usage.
Set XAudio sounds up to use streaming data, with a function added to the audio device so that other platforms won't try to make use of it for the moment.

diff --git a/Engine/Source/Editor/WorldBrowser/Private/StreamingLevels/StreamingLevelCollectionModel.cpp b/Engine/Source/Editor/WorldBrowser/Private/StreamingLevels/StreamingLevelCollectionModel.cpp
--- a/Engine/Source/Editor/WorldBrowser/Private/StreamingLevels/StreamingLevelCollectionModel.cpp
+++ b/Engine/Source/Editor/WorldBrowser/Private/StreamingLevels/StreamingLevelCollectionModel.cpp
@@ -6,6 +6,7 @@
 #include "Editor/MainFrame/Public/MainFrame.h"
 #include "LevelEditor.h"
 #include "DesktopPlatformModule.h"
+#include "AssetData.h"
 
 #include "StreamingLevelCustomization.h"
 #include "StreamingLevelCollectionModel.h"
diff --git a/Engine/Source/Runtime/Engine/Classes/Sound/SoundWave.h b/Engine/Source/Runtime/Engine/Classes/Sound/SoundWave.h
--- a/Engine/Source/Runtime/Engine/Classes/Sound/SoundWave.h
+++ b/Engine/Source/Runtime/Engine/Classes/Sound/SoundWave.h
@@ -26,6 +26,7 @@ enum EDecompressionType
 	DTYPE_RealTime,
 	DTYPE_Procedural,
 	DTYPE_Xenon,
+	DTYPE_Streaming,
 	DTYPE_MAX,
 };
 
@@ -34,6 +35,9 @@ enum EDecompressionType
  */
 struct FStreamedAudioChunk
 {
+	/** Size of the chunk of data in bytes */
+	int32 DataSize;
+
 	/** Bulk data if stored in the package. */
 	FByteBulkData BulkData;
 
@@ -86,14 +90,13 @@ struct FStreamedAudioPlatformData
 	~FStreamedAudioPlatformData();
 
 	/**
-	 * Try to load audio chunks from the derived data cache.
-	 * @param FirstChunkToLoad - The first Chunk index to load.
-	 * @param OutChunkData -	Must point to an array of pointers with at least
-	 *						Audio.Chunks.Num() - FirstChunkToLoad + 1 entries. Upon
-	 *						return those pointers will contain audio chunk data.
-	 * @returns true if all requested chunks have been loaded.
+	 * Try to load audio chunk from the derived data cache.
+	 * @param ChunkIndex	The Chunk index to load.
+	 * @param OutChunkData	Address of pointer that will store chunk data - should
+	 *						either be NULL or have enough space for the chunk
+	 * @returns true if requested chunk has been loaded.
 	 */
-	bool TryLoadChunks(int32 FirstChunkToLoad, void** OutChunkData);
+	bool TryLoadChunk(int32 ChunkIndex, uint8** OutChunkData);
 
 	/** Serialization. */
 	void Serialize(FArchive& Ar, class USoundWave* Owner);
@@ -248,7 +251,7 @@ public:
 	/** The streaming derived data for this sound on this platform. */
 	FStreamedAudioPlatformData *RunningPlatformData;
 
-	/* cooked streaming platform data for this sound */
+	/** cooked streaming platform data for this sound */
 	TMap<FString, FStreamedAudioPlatformData*> CookedPlatformData;
 
 	// Begin UObject interface. 
@@ -390,7 +393,19 @@ public:
 	 * Blocks on async cache tasks and prepares platform data for use.
 	 */
 	void FinishCachePlatformData();
+
+	/**
+	 * Forces platform data to be rebuilt.
+	 */
+	void ForceRebuildPlatformData();
 #endif
+
+	/**
+	 * Get Chunk data for a specified chunk index.
+	 * @param ChunkIndex	The Chunk index to cache.
+	 * @param OutChunkData	Address of pointer that will store data.
+	 */
+	void GetChunkData(int32 ChunkIndex, uint8** OutChunkData);
 };
 
 
diff --git a/Engine/Source/Runtime/Engine/Private/Audio.cpp b/Engine/Source/Runtime/Engine/Private/Audio.cpp
--- a/Engine/Source/Runtime/Engine/Private/Audio.cpp
+++ b/Engine/Source/Runtime/Engine/Private/Audio.cpp
@@ -172,6 +172,7 @@ FString FSoundSource::Describe(bool bUseLongName)
 
 void FSoundSource::Stop( void )
 {
+	IStreamingManager::Get().GetAudioStreamingManager().RemoveStreamingSoundSource(this);
 	if( WaveInstance )
 	{
 		check( AudioDevice );
diff --git a/Engine/Source/Runtime/Engine/Private/AudioDecompress.cpp b/Engine/Source/Runtime/Engine/Private/AudioDecompress.cpp
--- a/Engine/Source/Runtime/Engine/Private/AudioDecompress.cpp
+++ b/Engine/Source/Runtime/Engine/Private/AudioDecompress.cpp
@@ -394,8 +394,12 @@ FOpusAudioInfo::FOpusAudioInfo(void)
 	, TrueSampleCount(0)
 	, CurrentSampleCount(0)
 	, NumChannels(0)
+	, SampleStride(0)
 	, LastPCMByteSize(0)
 	, LastPCMOffset(0)
+	, bStoringEndOfFile(false)
+	, StreamingSoundWave(NULL)
+	, CurrentChunkIndex(0)
 {
 }
 
@@ -409,6 +413,8 @@ FOpusAudioInfo::~FOpusAudioInfo(void)
 
 size_t FOpusAudioInfo::Read(void *ptr, uint32 size)
 {
+	check(SrcBufferOffset < SrcBufferDataSize);
+
 	size_t BytesToRead = FMath::Min(size, SrcBufferDataSize - SrcBufferOffset);
 	FMemory::Memcpy(ptr, SrcBufferData + SrcBufferOffset, BytesToRead);
 	SrcBufferOffset += BytesToRead;
@@ -436,10 +442,11 @@ bool FOpusAudioInfo::ReadCompressedInfo(const uint8* InSrcBufferData, uint32 InS
 	uint16 SerializedFrames = 0;
 	Read(&SerializedFrames, sizeof(uint16));
 	AudioDataOffset = SrcBufferOffset;
+	SampleStride = SAMPLE_SIZE * NumChannels;
 
 	OpusDecoderWrapper = new FOpusDecoderWrapper(NumChannels);
 	LastDecodedPCM.Empty();
-	LastDecodedPCM.AddUninitialized(OPUS_MAX_FRAME_SIZE * SAMPLE_SIZE * NumChannels);
+	LastDecodedPCM.AddUninitialized(OPUS_MAX_FRAME_SIZE * SampleStride);
 #if WITH_OPUS
 	if (OpusDecoderWrapper->DecError != OPUS_OK)
 	{
@@ -468,54 +475,44 @@ bool FOpusAudioInfo::ReadCompressedData(uint8* Destination, bool bLooping, uint3
 
 	SCOPE_CYCLE_COUNTER(STAT_OpusDecompressTime);
 
-	const uint32 kSampleStride = SAMPLE_SIZE * NumChannels;
-	uint32 RawPCMOffset = 0;
+	// Write out any PCM data that was decoded during the last request
+	uint32 RawPCMOffset = WriteFromDecodedPCM(Destination, BufferSize);
+
 	bool bLooped = false;
-	uint16 FrameSize = 0;
 
-	if (LastPCMOffset > 0 && LastPCMByteSize > 0)
+	if (bStoringEndOfFile && LastPCMByteSize > 0)
 	{
-		uint32 BytesToCopy = FMath::Min(BufferSize - RawPCMOffset, LastPCMByteSize - LastPCMOffset);
-		FMemory::Memcpy(Destination + RawPCMOffset, LastDecodedPCM.GetTypedData() + LastPCMOffset, BytesToCopy);
-		LastPCMOffset = (BytesToCopy < LastPCMByteSize - LastPCMOffset) ? LastPCMOffset + BytesToCopy : 0;
-		RawPCMOffset += BytesToCopy;
+		// delayed returning looped because we hadn't read the entire buffer
+		bLooped = true;
+		bStoringEndOfFile = false;
 	}
 
 	while (RawPCMOffset < BufferSize)
 	{
+		uint16 FrameSize = 0;
 		Read(&FrameSize, sizeof(uint16));
-#if WITH_OPUS
-		int32 DecodedSamples = opus_decode(OpusDecoderWrapper->Decoder, SrcBufferData + SrcBufferOffset, FrameSize, (int16*)(LastDecodedPCM.GetTypedData()), OPUS_MAX_FRAME_SIZE, 0);
-#else
-		int32 DecodedSamples = -1;
-#endif
-		SrcBufferOffset += FrameSize;
+		int32 DecodedSamples = DecompressToPCMBuffer(FrameSize);
 		if (DecodedSamples < 0)
 		{
 			LastPCMByteSize = 0;
-			LastPCMOffset = 0;
 			return false;
 		}
 		else
 		{
-			if (CurrentSampleCount + DecodedSamples > TrueSampleCount)
-			{
-				DecodedSamples = TrueSampleCount - CurrentSampleCount;
-				CurrentSampleCount = TrueSampleCount;
-			}
-			else
-			{
-				CurrentSampleCount += DecodedSamples;
-			}
-			LastPCMByteSize = DecodedSamples * kSampleStride;
-			uint32 BytesToCopy = FMath::Min(BufferSize - RawPCMOffset, LastPCMByteSize);
-			FMemory::Memcpy(Destination + RawPCMOffset, LastDecodedPCM.GetTypedData(), BytesToCopy);
-			RawPCMOffset += BytesToCopy;
-			LastPCMOffset = (BytesToCopy < LastPCMByteSize) ? BytesToCopy : 0;
-			
+			LastPCMByteSize = IncrementCurrentSampleCount(DecodedSamples) * SampleStride;
+			RawPCMOffset += WriteFromDecodedPCM(Destination + RawPCMOffset, BufferSize - RawPCMOffset);
+
 			if (SrcBufferOffset >= SrcBufferDataSize)
 			{
-				bLooped = true;
+				// check whether all decoded PCM was written
+				if (LastPCMByteSize == 0)
+				{
+					bLooped = true;
+				}
+				else
+				{
+					bStoringEndOfFile = true;
+				}
 				if (bLooping)
 				{
 					SrcBufferOffset = AudioDataOffset;
@@ -523,12 +520,7 @@ bool FOpusAudioInfo::ReadCompressedData(uint8* Destination, bool bLooping, uint3
 				}
 				else
 				{
-					int32 Count = (BufferSize - RawPCMOffset);
-					if (Count > 0)
-					{
-						FMemory::Memzero(Destination + RawPCMOffset, Count);
-						RawPCMOffset += Count;
-					}
+					RawPCMOffset += ZeroBuffer(Destination + RawPCMOffset, BufferSize - RawPCMOffset);
 				}
 			}
 		}
@@ -542,35 +534,203 @@ void FOpusAudioInfo::ExpandFile(uint8* DstBuffer, struct FSoundQualityInfo* Qual
 	check(OpusDecoderWrapper);
 	check(DstBuffer);
 	check(QualityInfo);
+	check(QualityInfo->SampleDataSize <= SrcBufferDataSize);
 
-	const uint32 kSampleStride = SAMPLE_SIZE * NumChannels;
-	uint32 RawPCMOffset = 0;
-	uint16 FrameSize = 0;
+	// Ensure we're at the start of the audio data
+	SrcBufferOffset = AudioDataOffset;
 
-	TArray<uint8> TempPCMData;
-	TempPCMData.AddUninitialized(OPUS_MAX_FRAME_SIZE * kSampleStride);
+	uint32 RawPCMOffset = 0;
 
 	while (RawPCMOffset < QualityInfo->SampleDataSize)
 	{
+		uint16 FrameSize = 0;
 		Read(&FrameSize, sizeof(uint16));
-#if WITH_OPUS
-		int32 DecodedSamples = opus_decode(OpusDecoderWrapper->Decoder, (const unsigned char*)(SrcBufferData + SrcBufferOffset), FrameSize, (int16*)(TempPCMData.GetTypedData()), OPUS_MAX_FRAME_SIZE, 0);
-#else
-		int32 DecodedSamples = -1;
-#endif
+		int32 DecodedSamples = DecompressToPCMBuffer(FrameSize);
+
+		if (DecodedSamples < 0)
+		{
+			RawPCMOffset += ZeroBuffer(DstBuffer + RawPCMOffset, QualityInfo->SampleDataSize - RawPCMOffset);
+		}
+		else
+		{
+			LastPCMByteSize = IncrementCurrentSampleCount(DecodedSamples) * SampleStride;
+			RawPCMOffset += WriteFromDecodedPCM(DstBuffer + RawPCMOffset, QualityInfo->SampleDataSize - RawPCMOffset);
+		}
+	}
+}
+
+bool FOpusAudioInfo::StreamCompressedInfo(USoundWave* Wave, struct FSoundQualityInfo* QualityInfo)
+{
+	StreamingSoundWave = Wave;
+
+	// Get the first chunk of audio data (should always be loaded)
+	CurrentChunkIndex = 0;
+	const uint8* FirstChunk = IStreamingManager::Get().GetAudioStreamingManager().GetLoadedChunk(StreamingSoundWave, CurrentChunkIndex);
+
+	if (FirstChunk)
+	{
+		return ReadCompressedInfo(FirstChunk, Wave->RunningPlatformData->Chunks[0].DataSize, QualityInfo);
+	}
+
+	return false;
+}
+
+bool FOpusAudioInfo::StreamCompressedData(uint8* Destination, bool bLooping, uint32 BufferSize)
+{
+	check(OpusDecoderWrapper);
+	check(Destination);
+
+	SCOPE_CYCLE_COUNTER(STAT_OpusDecompressTime);
+
+	// Write out any PCM data that was decoded during the last request
+	uint32 RawPCMOffset = WriteFromDecodedPCM(Destination, BufferSize);
+
+	// If next chunk wasn't loaded when last one finished reading, try to get it again now
+	if (SrcBufferData == NULL)
+	{
+		SrcBufferData = IStreamingManager::Get().GetAudioStreamingManager().GetLoadedChunk(StreamingSoundWave, CurrentChunkIndex);
+		if (SrcBufferData)
+		{
+			SrcBufferDataSize = StreamingSoundWave->RunningPlatformData->Chunks[CurrentChunkIndex].DataSize;
+			SrcBufferOffset = CurrentChunkIndex == 0 ? AudioDataOffset : 0;
+		}
+		else
+		{
+			// Still not loaded, should probably have some kind of error, zero current buffer
+			ZeroBuffer(Destination + RawPCMOffset, BufferSize - RawPCMOffset);
+			return false;
+		}
+	}
+
+	bool bLooped = false;
+
+	if (bStoringEndOfFile && LastPCMByteSize > 0)
+	{
+		// delayed returning looped because we hadn't read the entire buffer
+		bLooped = true;
+		bStoringEndOfFile = false;
+	}
+
+	while (RawPCMOffset < BufferSize)
+	{
+		uint16 FrameSize = 0;
+		Read(&FrameSize, sizeof(uint16));
+		int32 DecodedSamples = DecompressToPCMBuffer(FrameSize);
+
 		if (DecodedSamples < 0)
 		{
-			break;
+			LastPCMByteSize = 0;
+			return false;
 		}
 		else
 		{
-			uint32 DecodedBytes = DecodedSamples * kSampleStride;
-			uint32 BytesToCopy = FMath::Min(QualityInfo->SampleDataSize - RawPCMOffset, DecodedBytes);
-			FMemory::Memcpy( DstBuffer + RawPCMOffset, TempPCMData.GetTypedData(), BytesToCopy );
-			SrcBufferOffset += FrameSize;
-			RawPCMOffset += BytesToCopy;
+			LastPCMByteSize = IncrementCurrentSampleCount(DecodedSamples) * SampleStride;
+
+			RawPCMOffset += WriteFromDecodedPCM(Destination + RawPCMOffset, BufferSize - RawPCMOffset);
+
+			// Have we reached the end of buffer
+			if (SrcBufferOffset >= SrcBufferDataSize)
+			{
+				// Special case for the last chunk of audio
+				if (CurrentChunkIndex == StreamingSoundWave->RunningPlatformData->NumChunks - 1)
+				{
+					// check whether all decoded PCM was written
+					if (LastPCMByteSize == 0)
+					{
+						bLooped = true;
+					}
+					else
+					{
+						bStoringEndOfFile = true;
+					}
+					if (bLooping)
+					{
+						CurrentChunkIndex = 0;
+						SrcBufferOffset = AudioDataOffset;
+						CurrentSampleCount = 0;
+					}
+					else
+					{
+						RawPCMOffset += ZeroBuffer(Destination + RawPCMOffset, BufferSize - RawPCMOffset);
+					}
+				}
+				else
+				{
+					CurrentChunkIndex++;
+					SrcBufferOffset = 0;
+				}
+
+				SrcBufferData = IStreamingManager::Get().GetAudioStreamingManager().GetLoadedChunk(StreamingSoundWave, CurrentChunkIndex);
+				if (SrcBufferData)
+				{
+					SrcBufferDataSize = StreamingSoundWave->RunningPlatformData->Chunks[CurrentChunkIndex].DataSize;
+				}
+				else
+				{
+					SrcBufferDataSize = 0;
+					RawPCMOffset += ZeroBuffer(Destination + RawPCMOffset, BufferSize - RawPCMOffset);
+				}
+			}
+		}
+	}
+
+	return bLooped;
+}
+
+int32 FOpusAudioInfo::DecompressToPCMBuffer(uint16 FrameSize)
+{
+	check(SrcBufferOffset + FrameSize <= SrcBufferDataSize);
+
+	const uint8* SrcPtr = SrcBufferData + SrcBufferOffset;
+	SrcBufferOffset += FrameSize;
+	LastPCMOffset = 0;
+#if WITH_OPUS
+	return opus_decode(OpusDecoderWrapper->Decoder, SrcPtr, FrameSize, (int16*)(LastDecodedPCM.GetTypedData()), OPUS_MAX_FRAME_SIZE, 0);
+#else
+	return -1;
+#endif
+}
+
+uint32 FOpusAudioInfo::IncrementCurrentSampleCount(uint32 NewSamples)
+{
+	if (CurrentSampleCount + NewSamples > TrueSampleCount)
+	{
+		NewSamples = TrueSampleCount - CurrentSampleCount;
+		CurrentSampleCount = TrueSampleCount;
+	}
+	else
+	{
+		CurrentSampleCount += NewSamples;
+	}
+	return NewSamples;
+}
+
+uint32 FOpusAudioInfo::WriteFromDecodedPCM(uint8* Destination, uint32 BufferSize)
+{
+	uint32 BytesToCopy = FMath::Min(BufferSize, LastPCMByteSize - LastPCMOffset);
+	if (BytesToCopy > 0)
+	{
+		FMemory::Memcpy(Destination, LastDecodedPCM.GetTypedData() + LastPCMOffset, BytesToCopy);
+		LastPCMOffset += BytesToCopy;
+		if (LastPCMOffset >= LastPCMByteSize)
+		{
+			LastPCMOffset = 0;
+			LastPCMByteSize = 0;
 		}
 	}
+	return BytesToCopy;
+}
+
+uint32	FOpusAudioInfo::ZeroBuffer(uint8* Destination, uint32 BufferSize)
+{
+	check(Destination);
+
+	if (BufferSize > 0)
+	{
+		FMemory::Memzero(Destination, BufferSize);
+		return BufferSize;
+	}
+	return 0;
 }
 
 /**
diff --git a/Engine/Source/Runtime/Engine/Private/AudioDerivedData.cpp b/Engine/Source/Runtime/Engine/Private/AudioDerivedData.cpp
--- a/Engine/Source/Runtime/Engine/Private/AudioDerivedData.cpp
+++ b/Engine/Source/Runtime/Engine/Private/AudioDerivedData.cpp
@@ -14,7 +14,7 @@ DEFINE_LOG_CATEGORY_STATIC(LogAudioDerivedData, Log, All);
 
 // If you want to bump this version, generate a new guid using
 // VS->Tools->Create GUID and paste it here.
-#define STREAMEDAUDIO_DERIVEDDATA_VER		TEXT("B9C1C21AEBE443B28C3B5546C8D5D742")
+#define STREAMEDAUDIO_DERIVEDDATA_VER		TEXT("2F3D9C2B13284A2AB73726A0A37C10C1")
 
 /**
  * Computes the derived data key suffix for a SoundWave's Streamed Audio.
@@ -71,6 +71,7 @@ static void GetStreamedAudioDerivedDataKeyFromSuffix(const FString& KeySuffix, F
  */
 static void GetStreamedAudioDerivedChunkKey(
 	int32 ChunkIndex,
+	const FStreamedAudioChunk& Chunk,
 	const FString& KeySuffix,
 	FString& OutKey
 	)
@@ -78,7 +79,7 @@ static void GetStreamedAudioDerivedChunkKey(
 	OutKey = FDerivedDataCacheInterface::BuildCacheKey(
 		TEXT("STREAMEDAUDIO"),
 		STREAMEDAUDIO_DERIVEDDATA_VER,
-		*FString::Printf(TEXT("%s_CHUNK%u"), *KeySuffix, ChunkIndex)
+		*FString::Printf(TEXT("%s_CHUNK%u_%d"), *KeySuffix, ChunkIndex, Chunk.DataSize)
 		);
 }
 
@@ -166,7 +167,7 @@ static void PutDerivedDataInCache(
 	{
 		FString ChunkDerivedDataKey;
 		FStreamedAudioChunk& Chunk = DerivedData->Chunks[ChunkIndex];
-		GetStreamedAudioDerivedChunkKey(ChunkIndex, DerivedDataKeySuffix, ChunkDerivedDataKey);
+		GetStreamedAudioDerivedChunkKey(ChunkIndex, Chunk, DerivedDataKeySuffix, ChunkDerivedDataKey);
 
 		if (UE_LOG_ACTIVE(LogAudio,Verbose))
 		{
@@ -198,8 +199,9 @@ namespace EStreamedAudioCacheFlags
 		None			= 0x0,
 		Async			= 0x1,
 		ForceRebuild	= 0x2,
-		AllowAsyncBuild	= 0x4,
-		ForDDCBuild		= 0x8,
+		InlineChunks	= 0x4,
+		AllowAsyncBuild	= 0x8,
+		ForDDCBuild		= 0x10,
 	};
 };
 
@@ -252,6 +254,7 @@ class FStreamedAudioCacheDerivedDataWorker : public FNonAbandonableTask
 				for (int32 ChunkIndex = 0; ChunkIndex < ChunkBuffers.Num(); ++ChunkIndex)
 				{
 					FStreamedAudioChunk* NewChunk = new(DerivedData->Chunks) FStreamedAudioChunk();
+					NewChunk->DataSize = ChunkBuffers[ChunkIndex].Num();
 					NewChunk->BulkData.Lock(LOCK_READ_WRITE);
 					void* NewChunkData = NewChunk->BulkData.Realloc(ChunkBuffers[ChunkIndex].Num());
 					FMemory::Memcpy(NewChunkData, ChunkBuffers[ChunkIndex].GetTypedData(), ChunkBuffers[ChunkIndex].Num());
@@ -262,6 +265,7 @@ class FStreamedAudioCacheDerivedDataWorker : public FNonAbandonableTask
 			{
 				// Could not split so copy compressed data into a single chunk
 				FStreamedAudioChunk* NewChunk = new(DerivedData->Chunks) FStreamedAudioChunk();
+				NewChunk->DataSize = CompressedBuffer.Num();
 				NewChunk->BulkData.Lock(LOCK_READ_WRITE);
 				void* NewChunkData = NewChunk->BulkData.Realloc(CompressedBuffer.Num());
 				FMemory::Memcpy(NewChunkData, CompressedBuffer.GetTypedData(), CompressedBuffer.Num());
@@ -276,7 +280,8 @@ class FStreamedAudioCacheDerivedDataWorker : public FNonAbandonableTask
 		
 		if (DerivedData->Chunks.Num())
 		{
-			bSucceeded = true;
+			bool bInlineChunks = (CacheFlags & EStreamedAudioCacheFlags::InlineChunks) != 0;
+			bSucceeded = !bInlineChunks || DerivedData->TryInlineChunkData();
 		}
 		else
 		{
@@ -308,6 +313,7 @@ public:
 	{
 		TArray<uint8> RawDerivedData;
 		bool bForceRebuild = (CacheFlags & EStreamedAudioCacheFlags::ForceRebuild) != 0;
+		bool bInlineChunks = (CacheFlags & EStreamedAudioCacheFlags::InlineChunks) != 0;
 		bool bForDDC = (CacheFlags & EStreamedAudioCacheFlags::ForDDCBuild) != 0;
 		bool bAllowAsyncBuild = (CacheFlags & EStreamedAudioCacheFlags::AllowAsyncBuild) != 0;
 
@@ -318,7 +324,18 @@ public:
 			bSucceeded = true;
 			if (bForDDC)
 			{
-				bSucceeded = DerivedData->TryLoadChunks(0,NULL);
+				for (int32 Index = 0; Index < DerivedData->NumChunks; ++Index)
+				{
+					if (!DerivedData->TryLoadChunk(Index, NULL))
+					{
+						bSucceeded = false;
+						break;
+					}
+				}
+			}
+			else if (bInlineChunks)
+			{
+				bSucceeded = DerivedData->TryInlineChunkData();
 			}
 			else
 			{
@@ -430,6 +447,42 @@ static void BeginLoadDerivedChunks(TIndirectArray<FStreamedAudioChunk>& Chunks,
 	}
 }
 
+bool FStreamedAudioPlatformData::TryInlineChunkData()
+{
+	TArray<uint32> AsyncHandles;
+	TArray<uint8> TempData;
+	FDerivedDataCacheInterface& DDC = GetDerivedDataCacheRef();
+
+	BeginLoadDerivedChunks(Chunks, 0, AsyncHandles);
+	for (int32 ChunkIndex = 0; ChunkIndex < Chunks.Num(); ++ChunkIndex)
+	{
+		FStreamedAudioChunk& Chunk = Chunks[ChunkIndex];
+		if (Chunk.DerivedDataKey.IsEmpty() == false)
+		{
+			uint32 AsyncHandle = AsyncHandles[ChunkIndex];
+			DDC.WaitAsynchronousCompletion(AsyncHandle);
+			if (DDC.GetAsynchronousResults(AsyncHandle, TempData))
+			{
+				int32 ChunkSize = 0;
+				FMemoryReader Ar(TempData, /*bIsPersistent=*/ true);
+				Ar << ChunkSize;
+
+				Chunk.BulkData.Lock(LOCK_READ_WRITE);
+				void* ChunkData = Chunk.BulkData.Realloc(ChunkSize);
+				Ar.Serialize(ChunkData, ChunkSize);
+				Chunk.BulkData.Unlock();
+				Chunk.DerivedDataKey.Empty();
+			}
+			else
+			{
+				return false;
+			}
+			TempData.Reset();
+		}
+	}
+	return true;
+}
+
 #endif //WITH_EDITORONLY_DATA
 
 FStreamedAudioPlatformData::FStreamedAudioPlatformData()
@@ -451,72 +504,74 @@ FStreamedAudioPlatformData::~FStreamedAudioPlatformData()
 #endif
 }
 
-bool FStreamedAudioPlatformData::TryLoadChunks(int32 FirstChunkToLoad, void** OutChunkData)
+bool FStreamedAudioPlatformData::TryLoadChunk(int32 ChunkIndex, uint8** OutChunkData)
 {
-	int32 NumChunksCached = 0;
+	bool bCachedChunk = false;
+	FStreamedAudioChunk& Chunk = Chunks[ChunkIndex];
 
 #if WITH_EDITORONLY_DATA
-	TArray<uint8> TempData;
-	TArray<uint32> AsyncHandles;
+	// Begin async DDC retrieval
 	FDerivedDataCacheInterface& DDC = GetDerivedDataCacheRef();
-	BeginLoadDerivedChunks(Chunks, FirstChunkToLoad, AsyncHandles);
+	uint32 AsyncHandle = 0;
+	if (!Chunk.DerivedDataKey.IsEmpty())
+	{
+		AsyncHandle = DDC.GetAsynchronous(*Chunk.DerivedDataKey);
+	}
 #endif // #if WITH_EDITORONLY_DATA
 
-	// Load remaining chunks (if any) from bulk data.
-	for (int32 ChunkIndex = FirstChunkToLoad; ChunkIndex < Chunks.Num(); ++ChunkIndex)
+	// Load chunk from bulk data if available.
+	if (Chunk.BulkData.GetBulkDataSize() > 0)
 	{
-		FStreamedAudioChunk& Chunk = Chunks[ChunkIndex];
-		if (Chunk.BulkData.GetBulkDataSize() > 0)
+		if (OutChunkData)
 		{
-			if (OutChunkData)
+			if (*OutChunkData == NULL)
 			{
-				OutChunkData[ChunkIndex - FirstChunkToLoad] = FMemory::Malloc(Chunk.BulkData.GetBulkDataSize());
-				Chunk.BulkData.GetCopy(&OutChunkData[ChunkIndex - FirstChunkToLoad]);
+				*OutChunkData = static_cast<uint8*>(FMemory::Malloc(Chunk.BulkData.GetBulkDataSize()));
 			}
-			NumChunksCached++;
+			Chunk.BulkData.GetCopy((void**)OutChunkData);
 		}
+		bCachedChunk = true;
 	}
 
 #if WITH_EDITORONLY_DATA
-	// Wait for async DDC gets.
-	for (int32 ChunkIndex = FirstChunkToLoad; ChunkIndex < Chunks.Num(); ++ChunkIndex)
+	// Wait for async DDC to complete
+	if (Chunk.DerivedDataKey.IsEmpty() == false)
 	{
-		FStreamedAudioChunk& Chunk = Chunks[ChunkIndex];
-		if (Chunk.DerivedDataKey.IsEmpty() == false)
+		TArray<uint8> TempData;
+		DDC.WaitAsynchronousCompletion(AsyncHandle);
+		if (DDC.GetAsynchronousResults(AsyncHandle, TempData))
 		{
-			uint32 AsyncHandle = AsyncHandles[ChunkIndex];
-			DDC.WaitAsynchronousCompletion(AsyncHandle);
-			if (DDC.GetAsynchronousResults(AsyncHandle, TempData))
+			int32 ChunkSize = 0;
+			FMemoryReader Ar(TempData, /*bIsPersistent=*/ true);
+			Ar << ChunkSize;
+
+			if (ChunkSize != Chunk.DataSize)
 			{
-				FMemoryReader Ar(TempData, /*bIsPersistent=*/ true);
-				NumChunksCached++;
+				UE_LOG(LogAudio, Warning,
+					TEXT("Chunk %d of %s SoundWave has invalid data in the DDC. Got %d bytes, expected %d. Key=%s"),
+					ChunkIndex,
+					*AudioFormat.ToString(),
+					ChunkSize,
+					Chunk.DataSize,
+					*Chunk.DerivedDataKey
+					);
+			}
+
+			bCachedChunk = true;
 
-				if (OutChunkData)
+			if (OutChunkData)
+			{
+				if (*OutChunkData == NULL)
 				{
-					OutChunkData[ChunkIndex - FirstChunkToLoad] = FMemory::Malloc(TempData.Num());
-					Ar.Serialize(OutChunkData[ChunkIndex - FirstChunkToLoad], TempData.Num());
+					*OutChunkData = static_cast<uint8*>(FMemory::Malloc(ChunkSize));
 				}
+				Ar.Serialize(*OutChunkData, ChunkSize);
 			}
-			TempData.Reset();
 		}
 	}
 #endif // #if WITH_EDITORONLY_DATA
 
-	if (NumChunksCached != (Chunks.Num() - FirstChunkToLoad))
-	{
-		// Unable to cache all chunks. Release memory for those that were cached.
-		for (int32 ChunkIndex = FirstChunkToLoad; ChunkIndex < Chunks.Num(); ++ChunkIndex)
-		{
-			if (OutChunkData && OutChunkData[ChunkIndex - FirstChunkToLoad])
-			{
-				FMemory::Free(OutChunkData[ChunkIndex - FirstChunkToLoad]);
-				OutChunkData[ChunkIndex - FirstChunkToLoad] = NULL;
-			}
-		}
-		return false;
-	}
-
-	return true;
+	return bCachedChunk;
 }
 
 #if WITH_EDITORONLY_DATA
@@ -857,7 +912,7 @@ void USoundWave::SerializeCookedPlatformData(FArchive& Ar)
 			if (PlatformDataToSave == NULL)
 			{
 				PlatformDataToSave = new FStreamedAudioPlatformData();
-				PlatformDataToSave->Cache(*this, PlatformFormat, EStreamedAudioCacheFlags::Async);
+				PlatformDataToSave->Cache(*this, PlatformFormat, EStreamedAudioCacheFlags::InlineChunks | EStreamedAudioCacheFlags::Async);
 
 				CookedPlatformData.Add(DerivedDataKey, PlatformDataToSave);
 			}
@@ -919,7 +974,7 @@ void USoundWave::BeginCacheForCookedPlatformData(const ITargetPlatform *TargetPl
 		// Retrieve format to cache for targetplatform.
 		FName PlatformFormat = TargetPlatform->GetWaveFormat(this);
 
-		uint32 CacheFlags = EStreamedAudioCacheFlags::Async;
+		uint32 CacheFlags = EStreamedAudioCacheFlags::Async | EStreamedAudioCacheFlags::InlineChunks;
 
 		// If source data is resident in memory then allow the streamed audio to be built
 		// in a background thread.
@@ -970,4 +1025,16 @@ void USoundWave::FinishCachePlatformData()
 	check(RunningPlatformData->DerivedDataKey == DerivedDataKey);
 #endif
 }
+
+void USoundWave::ForceRebuildPlatformData()
+{
+	if (RunningPlatformData)
+	{
+		RunningPlatformData->Cache(
+			*this,
+			GetWaveFormatForRunningPlatform(*this),
+			EStreamedAudioCacheFlags::ForceRebuild
+			);
+	}
+}
 #endif //WITH_EDITORONLY_DATA
\ No newline at end of file
diff --git a/Engine/Source/Runtime/Engine/Private/AudioDevice.cpp b/Engine/Source/Runtime/Engine/Private/AudioDevice.cpp
--- a/Engine/Source/Runtime/Engine/Private/AudioDevice.cpp
+++ b/Engine/Source/Runtime/Engine/Private/AudioDevice.cpp
@@ -329,6 +329,13 @@ void FAudioDevice::GetSoundClassInfo( TMap<FName, FAudioClassInfo>& AudioClassIn
 			AudioClassInfo->NumRealTime++;
 			break;
 
+		case DTYPE_Streaming:
+			// Add these to real time count for now - eventually compressed data won't be loaded &
+			// might have a class info entry of their own
+			AudioClassInfo->SizeRealTime += SoundWave->GetCompressedDataSize(GetRuntimeFormat(SoundWave));
+			AudioClassInfo->NumRealTime++;
+			break;
+
 		case DTYPE_Setup:
 		case DTYPE_Invalid:
 		default:
@@ -1842,6 +1849,7 @@ void FAudioDevice::StartSources( TArray<FWaveInstance*>& WaveInstances, int32 Fi
 				// Try to initialize source.
 				if( Source->Init( WaveInstance ) )
 				{
+					IStreamingManager::Get().GetAudioStreamingManager().AddStreamingSoundSource(Source);
 					// Associate wave instance with it which is used earlier in this function.
 					WaveInstanceSourceMap.Add( WaveInstance, Source );
 					// Playback might be deferred till the end of the update function on certain implementations.
@@ -2283,7 +2291,11 @@ void FAudioDevice::Precache(USoundWave* SoundWave, bool bSynchronous, bool bTrac
 		const FSoundGroup& SoundGroup = GetDefault<USoundGroups>()->GetSoundGroup(SoundWave->SoundGroup);
 
 		// handle audio decompression
-		if (SupportsRealtimeDecompression() && 
+		if (SupportsStreaming() && SoundWave->IsStreaming())
+		{
+			SoundWave->DecompressionType = DTYPE_Streaming;
+		}
+		else if (SupportsRealtimeDecompression() && 
 			(bDisableAudioCaching || (!SoundGroup.bAlwaysDecompressOnLoad && SoundWave->Duration > SoundGroup.DecompressedDuration)))
 		{
 			// Store as compressed data and decompress in realtime
diff --git a/Engine/Source/Runtime/Engine/Private/AudioStreaming.cpp b/Engine/Source/Runtime/Engine/Private/AudioStreaming.cpp
new file mode 100644
--- /dev/null
+++ b/Engine/Source/Runtime/Engine/Private/AudioStreaming.cpp
@@ -0,0 +1,499 @@
+// Copyright 1998-2014 Epic Games, Inc. All Rights Reserved.
+
+/*=============================================================================
+AudioStreaming.cpp: Implementation of audio streaming classes.
+=============================================================================*/
+
+#include "EnginePrivate.h"
+#include "AudioStreaming.h"
+#include "Sound/SoundWave.h"
+#include "DerivedDataCacheInterface.h"
+
+/*------------------------------------------------------------------------------
+	Streaming chunks from the derived data cache.
+------------------------------------------------------------------------------*/
+
+#if WITH_EDITORONLY_DATA
+
+/** Initialization constructor. */
+FAsyncStreamDerivedChunkWorker::FAsyncStreamDerivedChunkWorker(
+	const FString& InDerivedDataKey,
+	void* InDestChunkData,
+	int32 InChunkSize,
+	FThreadSafeCounter* InThreadSafeCounter
+	)
+	: DerivedDataKey(InDerivedDataKey)
+	, DestChunkData(InDestChunkData)
+	, ExpectedChunkSize(InChunkSize)
+	, bRequestFailed(false)
+	, ThreadSafeCounter(InThreadSafeCounter)
+{
+}
+
+/** Retrieves the derived chunk from the derived data cache. */
+void FAsyncStreamDerivedChunkWorker::DoWork()
+{
+	TArray<uint8> DerivedChunkData;
+
+	if (GetDerivedDataCacheRef().GetSynchronous(*DerivedDataKey, DerivedChunkData))
+	{
+		FMemoryReader Ar(DerivedChunkData, true);
+		int32 ChunkSize = 0;
+		Ar << ChunkSize;
+		checkf(ChunkSize == ExpectedChunkSize, TEXT("ChunkSize(%d) != ExpectedSize(%d)"), ChunkSize, ExpectedChunkSize);
+		Ar.Serialize(DestChunkData, ChunkSize);
+	}
+	else
+	{
+		bRequestFailed = true;
+	}
+	FPlatformMisc::MemoryBarrier();
+	ThreadSafeCounter->Decrement();
+}
+
+#endif // #if WITH_EDITORONLY_DATA
+
+////////////////////////
+// FStreamingWaveData //
+////////////////////////
+
+FStreamingWaveData::FStreamingWaveData()
+	: SoundWave(NULL)
+{
+}
+
+FStreamingWaveData::~FStreamingWaveData()
+{
+	// Make sure there are no pending requests in flight.
+	while (UpdateStreamingStatus() == true)
+	{
+		// Give up timeslice.
+		FPlatformProcess::Sleep(0);
+	}
+
+	for (auto& LoadedChunk : LoadedChunks)
+	{
+		FreeLoadedChunk(LoadedChunk);
+	}
+}
+
+void FStreamingWaveData::Initialize(USoundWave* InSoundWave)
+{
+	SoundWave = InSoundWave;
+
+	// Always get the first chunk of data so we can play immediately
+	check(LoadedChunks.Num() == 0);
+	check(LoadedChunkIndices.Num() == 0);
+	FLoadedAudioChunk* FirstChunk = AddNewLoadedChunk(SoundWave->RunningPlatformData->Chunks[0].DataSize);
+	FirstChunk->Index = 0;
+	SoundWave->GetChunkData(0, &FirstChunk->Data);
+
+	// Set up the loaded/requested indices to be identical
+	LoadedChunkIndices.Add(0);
+	CurrentRequest.RequiredIndices.Add(0);
+}
+
+bool FStreamingWaveData::UpdateStreamingStatus()
+{
+	bool	bHasPendingRequestInFlight = true;
+	int32	RequestStatus = PendingChunkChangeRequestStatus.GetValue();
+	TArray<uint32> IndicesToLoad;
+	TArray<uint32> IndicesToFree;
+
+	if (!HasPendingRequests(IndicesToLoad, IndicesToFree))
+	{
+		check(RequestStatus == AudioState_ReadyFor_Requests);
+		bHasPendingRequestInFlight = false;
+	}
+	// Pending request in flight, though we might be able to finish it.
+	else
+	{
+		if (RequestStatus == AudioState_ReadyFor_Finalization)
+		{
+			bool bFailedRequests = false;
+#if WITH_EDITORONLY_DATA
+			bFailedRequests = FinishDDCRequests();
+#endif //WITH_EDITORONLY_DATA
+			PendingChunkChangeRequestStatus.Decrement();
+			bHasPendingRequestInFlight = false;
+			LoadedChunkIndices = CurrentRequest.RequiredIndices;
+		}
+		else if (RequestStatus == AudioState_ReadyFor_Requests)
+		{
+			BeginPendingRequests(IndicesToLoad, IndicesToFree);
+		}
+	}
+
+	return bHasPendingRequestInFlight;
+}
+
+void FStreamingWaveData::UpdateChunkRequests(FWaveRequest& InWaveRequest)
+{
+	// Might change this but ensures chunk 0 stays loaded for now
+	check(InWaveRequest.RequiredIndices.Contains(0));
+	check(PendingChunkChangeRequestStatus.GetValue() == AudioState_ReadyFor_Requests);
+
+	CurrentRequest = InWaveRequest;
+
+	// Clear last batch of request indices
+	IORequestIndices.Empty();
+}
+
+bool FStreamingWaveData::HasPendingRequests(TArray<uint32>& IndicesToLoad, TArray<uint32>& IndicesToFree) const
+{
+	IndicesToLoad.Empty();
+	IndicesToFree.Empty();
+
+	// Find indices that aren't loaded
+	for (auto NeededIndex : CurrentRequest.RequiredIndices)
+	{
+		if (!LoadedChunkIndices.Contains(NeededIndex))
+		{
+			IndicesToLoad.AddUnique(NeededIndex);
+		}
+	}
+
+	// Find indices that aren't needed anymore
+	for (auto CurrentIndex : LoadedChunkIndices)
+	{
+		if (!CurrentRequest.RequiredIndices.Contains(CurrentIndex))
+		{
+			IndicesToFree.AddUnique(CurrentIndex);
+		}
+	}
+
+	return IndicesToLoad.Num() > 0 || IndicesToFree.Num() > 0;
+}
+
+void FStreamingWaveData::BeginPendingRequests(const TArray<uint32>& IndicesToLoad, const TArray<uint32>& IndicesToFree)
+{
+	// Mark Chunks for removal in case they can be reused
+	TArray<uint32> FreeChunkIndices;
+	for (auto Index : IndicesToFree)
+	{
+		for (int32 ChunkIndex = 0; ChunkIndex < LoadedChunks.Num(); ++ChunkIndex)
+		{
+			if (LoadedChunks[ChunkIndex].Index == Index)
+			{
+				FreeChunkIndices.AddUnique(ChunkIndex);
+				break;
+			}
+		}
+	}
+
+	if (IndicesToLoad.Num() > 0)
+	{
+		PendingChunkChangeRequestStatus.Set(AudioState_InProgress_Loading);
+
+		// Set off all IO Requests
+		for (auto Index : IndicesToLoad)
+		{
+			const FStreamedAudioChunk& Chunk = SoundWave->RunningPlatformData->Chunks[Index];
+			int32 ChunkSize = Chunk.DataSize;
+
+			FLoadedAudioChunk* ChunkStorage = NULL;
+
+			for (auto FreeIndex : FreeChunkIndices)
+			{
+				if (LoadedChunks[FreeIndex].MemorySize >= ChunkSize)
+				{
+					FreeChunkIndices.Remove(FreeIndex);
+					ChunkStorage = &LoadedChunks[FreeIndex];
+					ChunkStorage->DataSize = ChunkSize;
+					ChunkStorage->Index = Index;
+					break;
+				}
+			}
+			if (ChunkStorage == NULL)
+			{
+				ChunkStorage = AddNewLoadedChunk(ChunkSize);
+				ChunkStorage->Index = Index;
+			}
+
+			// Pass the request on to the async io manager after increasing the request count. The request count 
+			// has been pre-incremented before fielding the update request so we don't have to worry about file
+			// I/O immediately completing and the game thread kicking off again before this function
+			// returns.
+			PendingChunkChangeRequestStatus.Increment();
+
+			EAsyncIOPriority AsyncIOPriority = CurrentRequest.bPrioritiseRequest ? AIOP_BelowNormal : AIOP_Low;
+
+			// Load and decompress async.
+#if WITH_EDITORONLY_DATA
+			if (Chunk.DerivedDataKey.IsEmpty() == false)
+			{
+				FAsyncStreamDerivedChunkTask* Task = new(PendingAsyncStreamDerivedChunkTasks)FAsyncStreamDerivedChunkTask(
+					Chunk.DerivedDataKey,
+					ChunkStorage->Data,
+					ChunkSize,
+					&PendingChunkChangeRequestStatus
+					);
+				Task->StartBackgroundTask();
+			}
+			else
+#endif // #if WITH_EDITORONLY_DATA
+			{
+				check(Chunk.BulkData.GetFilename().Len());
+				if (Chunk.BulkData.IsStoredCompressedOnDisk())
+				{
+					IORequestIndices.AddUnique(FIOSystem::Get().LoadCompressedData(
+						Chunk.BulkData.GetFilename(),						// filename
+						Chunk.BulkData.GetBulkDataOffsetInFile(),			// offset
+						Chunk.BulkData.GetBulkDataSizeOnDisk(),				// compressed size
+						Chunk.BulkData.GetBulkDataSize(),					// uncompressed size
+						ChunkStorage->Data,									// dest pointer
+						Chunk.BulkData.GetDecompressionFlags(),				// compressed data format
+						&PendingChunkChangeRequestStatus,					// counter to decrement
+						AsyncIOPriority										// priority
+						)
+						);
+				}
+				// Load async.
+				else
+				{
+					IORequestIndices.AddUnique(FIOSystem::Get().LoadData(
+						Chunk.BulkData.GetFilename(),						// filename
+						Chunk.BulkData.GetBulkDataOffsetInFile(),			// offset
+						Chunk.BulkData.GetBulkDataSize(),					// size
+						ChunkStorage->Data,									// dest pointer
+						&PendingChunkChangeRequestStatus,					// counter to decrement
+						AsyncIOPriority										// priority
+						)
+						);
+				}
+				check(IORequestIndices[IORequestIndices.Num() - 1]);
+			}
+		}
+
+		// Decrement the state to AudioState_InProgress_Loading + NumChunksCurrentLoading - 1.
+		PendingChunkChangeRequestStatus.Decrement();
+	}
+	else
+	{
+		// Skip straight to finalisation
+		PendingChunkChangeRequestStatus.Set(AudioState_ReadyFor_Finalization);
+	}
+
+	// Ensure indices are in order so we can step through backwards
+	FreeChunkIndices.Sort();
+
+	for (int32 FreeIndex = FreeChunkIndices.Num() - 1; FreeIndex >= 0; --FreeIndex)
+	{
+		FreeLoadedChunk(LoadedChunks[FreeChunkIndices[FreeIndex]]);
+		LoadedChunks.RemoveAt(FreeChunkIndices[FreeIndex]);
+	}
+}
+
+#if WITH_EDITORONLY_DATA
+bool FStreamingWaveData::FinishDDCRequests()
+{
+	bool bRequestFailed = false;
+	if (PendingAsyncStreamDerivedChunkTasks.Num())
+	{
+		for (int32 TaskIndex = 0; TaskIndex < PendingAsyncStreamDerivedChunkTasks.Num(); ++TaskIndex)
+		{
+			FAsyncStreamDerivedChunkTask& Task = PendingAsyncStreamDerivedChunkTasks[TaskIndex];
+			Task.EnsureCompletion();
+			bRequestFailed |= Task.GetTask().DidRequestFail();
+		}
+		PendingAsyncStreamDerivedChunkTasks.Empty();
+	}
+	return bRequestFailed;
+}
+#endif //WITH_EDITORONLY_DATA
+
+FLoadedAudioChunk* FStreamingWaveData::AddNewLoadedChunk(int32 ChunkSize)
+{
+	// TODO: Might update memory management stats here
+	int32 NewIndex = LoadedChunks.Num();
+	LoadedChunks.AddZeroed();
+	LoadedChunks[NewIndex].Data = static_cast<uint8*>(FMemory::Malloc(ChunkSize));
+	LoadedChunks[NewIndex].DataSize = LoadedChunks[NewIndex].MemorySize = ChunkSize;
+	return &LoadedChunks[NewIndex];
+}
+
+void FStreamingWaveData::FreeLoadedChunk(FLoadedAudioChunk& LoadedChunk)
+{
+	if (LoadedChunk.Data != NULL)
+	{
+		// TODO: Might update memory management stats here
+		FMemory::Free(LoadedChunk.Data);
+		LoadedChunk.Data = NULL;
+		LoadedChunk.DataSize = 0;
+		LoadedChunk.MemorySize = 0;
+		LoadedChunk.Index = 0;
+	}
+}
+
+////////////////////////////
+// FAudioStreamingManager //
+////////////////////////////
+
+FAudioStreamingManager::FAudioStreamingManager()
+{
+}
+
+FAudioStreamingManager::~FAudioStreamingManager()
+{
+}
+
+void FAudioStreamingManager::UpdateResourceStreaming(float DeltaTime, bool bProcessEverything /*= false*/)
+{
+	for (auto& WavePair : StreamingSoundWaves)
+	{
+		WavePair.Value.UpdateStreamingStatus();
+	}
+
+	for (auto Source : StreamingSoundSources)
+	{
+		USoundWave* Wave = Source->GetWaveInstance()->WaveData;
+		FStreamingWaveData* WaveData = StreamingSoundWaves.Find(Wave);
+
+		if (WaveData && WaveData->PendingChunkChangeRequestStatus.GetValue() == AudioState_ReadyFor_Requests)
+		{
+			// Request the chunk the source is using and the one after that
+			FWaveRequest& WaveRequest = GetWaveRequest(Wave);
+			uint32 SourceChunk = Source->GetBuffer()->GetCurrentChunkIndex();
+			WaveRequest.RequiredIndices.AddUnique(SourceChunk);
+			WaveRequest.RequiredIndices.AddUnique((SourceChunk+1)%Wave->RunningPlatformData->NumChunks);
+			if (Source->GetBuffer()->GetCurrentChunkOffset() > Wave->RunningPlatformData->Chunks[SourceChunk].DataSize / 2)
+			{
+				// already read over half, request is high priority
+				WaveRequest.bPrioritiseRequest = true;
+			}
+		}
+	}
+
+	for (auto WaveRequest : WaveRequests)
+	{
+		USoundWave* Wave = WaveRequest.Key;
+		FStreamingWaveData* WaveData = StreamingSoundWaves.Find(Wave);
+
+		if (WaveData && WaveData->PendingChunkChangeRequestStatus.GetValue() == AudioState_ReadyFor_Requests)
+		{
+			WaveData->UpdateChunkRequests(WaveRequest.Value);
+			WaveData->UpdateStreamingStatus();
+			WaveRequests.Remove(Wave);
+		}
+	}
+}
+
+int32 FAudioStreamingManager::BlockTillAllRequestsFinished(float TimeLimit /*= 0.0f*/, bool bLogResults /*= false*/)
+{
+	// Not sure yet whether this will work the same as textures - aside from just before destroying
+	return 0;
+}
+
+void FAudioStreamingManager::CancelForcedResources()
+{
+}
+
+void FAudioStreamingManager::NotifyLevelChange()
+{
+}
+
+void FAudioStreamingManager::SetDisregardWorldResourcesForFrames(int32 NumFrames)
+{
+}
+
+void FAudioStreamingManager::AddPreparedLevel(class ULevel* Level)
+{
+}
+
+void FAudioStreamingManager::RemoveLevel(class ULevel* Level)
+{
+}
+
+void FAudioStreamingManager::AddStreamingSoundWave(USoundWave* SoundWave)
+{
+	if (SoundWave->IsStreaming())
+	{
+		check(StreamingSoundWaves.Find(SoundWave) == NULL);
+
+		FStreamingWaveData& WaveData = StreamingSoundWaves.Add(SoundWave);
+		WaveData.Initialize(SoundWave);
+	}
+}
+
+void FAudioStreamingManager::RemoveStreamingSoundWave(USoundWave* SoundWave)
+{
+	FStreamingWaveData* WaveData = StreamingSoundWaves.Find(SoundWave);
+	if (WaveData)
+	{
+		StreamingSoundWaves.Remove(SoundWave);
+	}
+	WaveRequests.Remove(SoundWave);
+}
+
+bool FAudioStreamingManager::IsManagedStreamingSoundWave(const USoundWave* SoundWave) const
+{
+	return StreamingSoundWaves.Find(SoundWave) != NULL;
+}
+
+bool FAudioStreamingManager::IsStreamingInProgress(const USoundWave* SoundWave)
+{
+	FStreamingWaveData* WaveData = StreamingSoundWaves.Find(SoundWave);
+	if (WaveData)
+	{
+		return WaveData->UpdateStreamingStatus();
+	}
+	return false;
+}
+
+void FAudioStreamingManager::AddStreamingSoundSource(FSoundSource* SoundSource)
+{
+	const FWaveInstance* WaveInstance = SoundSource->GetWaveInstance();
+	if (WaveInstance && WaveInstance->WaveData && WaveInstance->WaveData->IsStreaming())
+	{
+		StreamingSoundSources.AddUnique(SoundSource);
+	}
+}
+
+void FAudioStreamingManager::RemoveStreamingSoundSource(FSoundSource* SoundSource)
+{
+	const FWaveInstance* WaveInstance = SoundSource->GetWaveInstance();
+	if (WaveInstance && WaveInstance->WaveData && WaveInstance->WaveData->IsStreaming())
+	{
+		// Make sure there is a request so that unused chunks
+		// can be cleared if this was the last playing instance
+		GetWaveRequest(WaveInstance->WaveData);
+		StreamingSoundSources.Remove(SoundSource);
+	}
+}
+
+bool FAudioStreamingManager::IsManagedStreamingSoundSource(const FSoundSource* SoundSource) const
+{
+	return StreamingSoundSources.FindByKey(SoundSource) != NULL;
+}
+
+const uint8* FAudioStreamingManager::GetLoadedChunk(const USoundWave* SoundWave, uint32 ChunkIndex) const
+{
+	const FStreamingWaveData* WaveData = StreamingSoundWaves.Find(SoundWave);
+	if (WaveData)
+	{
+		if (WaveData->LoadedChunkIndices.Contains(ChunkIndex))
+		{
+			for (int32 Index = 0; Index < WaveData->LoadedChunks.Num(); ++Index)
+			{
+				if (WaveData->LoadedChunks[Index].Index == ChunkIndex)
+				{
+					return WaveData->LoadedChunks[Index].Data;
+				}
+			}
+		}
+	}
+	return NULL;
+}
+
+FWaveRequest& FAudioStreamingManager::GetWaveRequest(USoundWave* SoundWave)
+{
+	FWaveRequest* WaveRequest = WaveRequests.Find(SoundWave);
+	if (!WaveRequest)
+	{
+		// Setup the new request so it always asks for chunk 0
+		WaveRequest = &WaveRequests.Add(SoundWave);
+		WaveRequest->RequiredIndices.AddUnique(0);
+		WaveRequest->bPrioritiseRequest = false;
+	}
+	return *WaveRequest;
+}
diff --git a/Engine/Source/Runtime/Engine/Private/ContentStreaming.cpp b/Engine/Source/Runtime/Engine/Private/ContentStreaming.cpp
--- a/Engine/Source/Runtime/Engine/Private/ContentStreaming.cpp
+++ b/Engine/Source/Runtime/Engine/Private/ContentStreaming.cpp
@@ -6,6 +6,7 @@
 
 #include "EnginePrivate.h"
 #include "GenericPlatformMemoryPoolStats.h"
+#include "AudioStreaming.h"
 
 DEFINE_LOG_CATEGORY_STATIC(LogContentStreaming, Log, All);
 
@@ -1583,6 +1584,9 @@ FStreamingManagerCollection::FStreamingManagerCollection()
 #endif
 
 	AddOrRemoveTextureStreamingManagerIfNeeded(true);
+
+	AudioStreamingManager = new FAudioStreamingManager();
+	AddStreamingManager( AudioStreamingManager );
 }
 
 /**
@@ -1772,6 +1776,12 @@ ITextureStreamingManager& FStreamingManagerCollection::GetTextureStreamingManage
 	return *TextureStreamingManager;
 }
 
+IAudioStreamingManager& FStreamingManagerCollection::GetAudioStreamingManager() const
+{
+	check(AudioStreamingManager);
+	return *AudioStreamingManager;
+}
+
 /** Don't stream world resources for the next NumFrames. */
 void FStreamingManagerCollection::SetDisregardWorldResourcesForFrames(int32 NumFrames )
 {
diff --git a/Engine/Source/Runtime/Engine/Private/SoundWave.cpp b/Engine/Source/Runtime/Engine/Private/SoundWave.cpp
--- a/Engine/Source/Runtime/Engine/Private/SoundWave.cpp
+++ b/Engine/Source/Runtime/Engine/Private/SoundWave.cpp
@@ -20,6 +20,7 @@ void FStreamedAudioChunk::Serialize(FArchive& Ar, UObject* Owner, int32 ChunkInd
 	Ar << bCooked;
 
 	BulkData.Serialize(Ar, Owner, ChunkIndex);
+	Ar << DataSize;
 
 #if WITH_EDITORONLY_DATA
 	if (!bCooked)
@@ -300,12 +301,13 @@ void USoundWave::PostLoad()
 		}
 	}
 
-#if WITH_EDITORONLY_DATA
 	if (IsStreaming())
 	{
+#if WITH_EDITORONLY_DATA
 		FinishCachePlatformData();
-	}
 #endif // #if WITH_EDITORONLY_DATA
+		IStreamingManager::Get().GetAudioStreamingManager().AddStreamingSoundWave(this);
+	}
 
 	INC_FLOAT_STAT_BY( STAT_AudioBufferTime, Duration );
 	INC_FLOAT_STAT_BY( STAT_AudioBufferTimeChannels, NumChannels * Duration );
@@ -398,6 +400,14 @@ void USoundWave::PostEditChangeProperty(FPropertyChangedEvent& PropertyChangedEv
 			FreeResources();
 			UpdatePlatformData();
 			MarkPackageDirty();
+			if (IsStreaming())
+			{
+				IStreamingManager::Get().GetAudioStreamingManager().AddStreamingSoundWave(this);
+			}
+			else
+			{
+				IStreamingManager::Get().GetAudioStreamingManager().RemoveStreamingSoundWave(this);
+			}
 		}
 	}
 }
@@ -479,8 +489,10 @@ FWaveInstance* USoundWave::HandleStart( FActiveSound& ActiveSound, const UPTRINT
 
 bool USoundWave::IsReadyForFinishDestroy()
 {
-	// Wait till vorbis decompression finishes before deleting resource.
-	return( ( AudioDecompressor == NULL ) || AudioDecompressor->IsDone() );
+	bool bIsStreamingInProgress = IStreamingManager::Get().GetAudioStreamingManager().IsStreamingInProgress(this);
+
+	// Wait till streaming and decompression finishes before deleting resource.
+	return( !bIsStreamingInProgress && (( AudioDecompressor == NULL ) || AudioDecompressor->IsDone()) );
 }
 
 
@@ -492,6 +504,8 @@ void USoundWave::FinishDestroy()
 
 	CleanupCachedRunningPlatformData();
 	CleanupCachedCookedPlatformData();
+
+	IStreamingManager::Get().GetAudioStreamingManager().RemoveStreamingSoundWave(this);
 }
 
 void USoundWave::Parse( FAudioDevice* AudioDevice, const UPTRINT NodeWaveInstanceHash, FActiveSound& ActiveSound, const FSoundParseParameters& ParseParams, TArray<FWaveInstance*>& WaveInstances )
@@ -646,7 +660,12 @@ void USoundWave::UpdatePlatformData()
 {
 	if (IsStreaming())
 	{
-		// TODO: Check we aren't already streaming before doing this
+		// Make sure there are no pending requests in flight.
+		while (IStreamingManager::Get().GetAudioStreamingManager().IsStreamingInProgress(this))
+		{
+			// Give up timeslice.
+			FPlatformProcess::Sleep(0);
+		}
 
 #if WITH_EDITORONLY_DATA
 		// Recache platform data if the source has changed.
@@ -654,3 +673,19 @@ void USoundWave::UpdatePlatformData()
 #endif
 	}
 }
+
+void USoundWave::GetChunkData(int32 ChunkIndex, uint8** OutChunkData)
+{
+	if (RunningPlatformData->TryLoadChunk(ChunkIndex, OutChunkData) == false)
+	{
+		// Unable to load chunks from the cache. Rebuild the sound and try again.
+		UE_LOG(LogAudio, Warning, TEXT("GetChunkData failed for %s"), *GetPathName());
+#if WITH_EDITORONLY_DATA
+		ForceRebuildPlatformData();
+		if (RunningPlatformData->TryLoadChunk(ChunkIndex, OutChunkData) == false)
+		{
+			UE_LOG(LogAudio, Error, TEXT("Failed to build sound %s."), *GetPathName());
+		}
+#endif // #if WITH_EDITORONLY_DATA
+	}
+}
diff --git a/Engine/Source/Runtime/Engine/Public/Audio.h b/Engine/Source/Runtime/Engine/Public/Audio.h
--- a/Engine/Source/Runtime/Engine/Public/Audio.h
+++ b/Engine/Source/Runtime/Engine/Public/Audio.h
@@ -302,6 +302,16 @@ public:
 	 */
 	FString GetChannelsDesc();
 
+	/**
+	 * Gets the chunk index that was last read from (for Streaming Manager requests)
+	 */
+	virtual int32 GetCurrentChunkIndex() const {return -1;}
+
+	/**
+	 * Gets the offset into the chunk that was last read to (for Streaming Manager priority)
+	 */
+	virtual int32 GetCurrentChunkOffset() const {return -1;}
+
 	/** Unique ID that ties this buffer to a USoundWave */
 	int32	ResourceID;
 	/** Cumulative channels from all streams */
@@ -433,6 +443,8 @@ public:
 	 */
 	ENGINE_API void SetHighFrequencyGain( void );
 
+	const FSoundBuffer* GetBuffer() const {return Buffer;}
+
 protected:
 	// Variables.	
 	class FAudioDevice*		AudioDevice;
diff --git a/Engine/Source/Runtime/Engine/Public/AudioDecompress.h b/Engine/Source/Runtime/Engine/Public/AudioDecompress.h
--- a/Engine/Source/Runtime/Engine/Public/AudioDecompress.h
+++ b/Engine/Source/Runtime/Engine/Public/AudioDecompress.h
@@ -79,6 +79,44 @@ public:
 	 * Gets the size of the source buffer originally passed to the info class (bytes)
 	 */
 	virtual uint32 GetSourceBufferSize() const = 0;
+
+	////////////////////////////////////////////////////////////////
+	// Following functions are optional if streaming is supported //
+	////////////////////////////////////////////////////////////////
+
+	/**
+	 * Whether this decompression class supports streaming decompression
+	 */
+	virtual bool SupportsStreaming() const {return false;}
+
+	/**
+	* Streams the header information of a compressed format
+	*
+	* @param	Wave			Wave that will be read from to retrieve necessary chunk
+	* @param	QualityInfo		Quality Info (to be filled out)
+	*/
+	virtual bool StreamCompressedInfo(USoundWave* Wave, struct FSoundQualityInfo* QualityInfo) {return false;}
+
+	/**
+	* Decompresses streamed data to raw PCM data.
+	*
+	* @param	Destination	where to place the decompressed sound
+	* @param	bLooping	whether to loop the sound by seeking to the start, or pad the buffer with zeroes
+	* @param	BufferSize	number of bytes of PCM data to create
+	*
+	* @return	bool		true if the end of the data was reached (for both single shot and looping sounds)
+	*/
+	virtual bool StreamCompressedData(uint8* Destination, bool bLooping, uint32 BufferSize) {return false;}
+
+	/**
+	 * Gets the chunk index that was last read from (for Streaming Manager requests)
+	 */
+	virtual int32 GetCurrentChunkIndex() const {return -1;}
+
+	/**
+	 * Gets the offset into the chunk that was last read to (for Streaming Manager priority)
+	 */
+	virtual int32 GetCurrentChunkOffset() const {return -1;}
 };
 
 #if WITH_OGGVORBIS
@@ -140,6 +178,8 @@ public:
 
 #define OPUS_ID_STRING "UE4OPUS"
 
+struct FOpusDecoderWrapper;
+
 /**
 * Helper class to parse opus data
 */
@@ -149,19 +189,60 @@ public:
 	ENGINE_API FOpusAudioInfo(void);
 	ENGINE_API virtual ~FOpusAudioInfo(void);
 
+	// ICompressedAudioInfo Interface
+	ENGINE_API virtual bool ReadCompressedInfo(const uint8* InSrcBufferData, uint32 InSrcBufferDataSize, struct FSoundQualityInfo* QualityInfo) override;
+	ENGINE_API virtual bool ReadCompressedData(uint8* Destination, bool bLooping, uint32 BufferSize) override;
+	ENGINE_API virtual void SeekToTime(const float SeekTime) override {};
+	ENGINE_API virtual void ExpandFile(uint8* DstBuffer, struct FSoundQualityInfo* QualityInfo) override;
+	ENGINE_API virtual void EnableHalfRate(bool HalfRate) override {};
+	virtual uint32 GetSourceBufferSize() const override { return SrcBufferDataSize;}
+
+	virtual bool SupportsStreaming() const override {return true;}
+	virtual bool StreamCompressedInfo(USoundWave* Wave, struct FSoundQualityInfo* QualityInfo) override;
+	virtual bool StreamCompressedData(uint8* Destination, bool bLooping, uint32 BufferSize) override;
+	virtual int32 GetCurrentChunkIndex() const override {return CurrentChunkIndex;}
+	virtual int32 GetCurrentChunkOffset() const override {return SrcBufferOffset;}
+	// End of ICompressedAudioInfo Interface
+
+protected:
 	/** Emulate read from memory functionality */
 	size_t			Read(void *ptr, uint32 size);
 
-	// ICompressedAudioInfo Interface
-	ENGINE_API virtual bool ReadCompressedInfo(const uint8* InSrcBufferData, uint32 InSrcBufferDataSize, struct FSoundQualityInfo* QualityInfo);
-	ENGINE_API virtual bool ReadCompressedData(uint8* Destination, bool bLooping, uint32 BufferSize);
-	ENGINE_API virtual void SeekToTime(const float SeekTime) {};
-	ENGINE_API virtual void ExpandFile(uint8* DstBuffer, struct FSoundQualityInfo* QualityInfo);
-	ENGINE_API virtual void EnableHalfRate(bool HalfRate) {};
-	virtual uint32 GetSourceBufferSize() const { return SrcBufferDataSize;}
-	// End of ICompressedAudioInfo Interface
+	/**
+	 * Decompresses a frame of Opus data to PCM buffer
+	 *
+	 * @param FrameSize Size of the frame in bytes
+	 * @return The amount of samples that were decompressed (< 0 indicates error)
+	 */
+	int32 DecompressToPCMBuffer(uint16 FrameSize);
 
-	struct FOpusDecoderWrapper* OpusDecoderWrapper;
+	/**
+	 * Adds to the count of samples that have currently been decoded
+	 *
+	 * @param NewSamples	How many samples have been decoded
+	 * @return How many samples were actually part of the true sample count
+	 */
+	uint32 IncrementCurrentSampleCount(uint32 NewSamples);
+
+	/**
+	 * Writes data from decoded PCM buffer, taking into account whether some PCM has been written before
+	 *
+	 * @param Destination	Where to place the decoded sound
+	 * @param BufferSize	Size of the destination buffer in bytes
+	 * @return				How many bytes were written
+	 */
+	uint32	WriteFromDecodedPCM(uint8* Destination, uint32 BufferSize);
+
+	/**
+	 * Zeroes the contents of a buffer
+	 *
+	 * @param Destination	Buffer to zero
+	 * @param BufferSize	Size of the destination buffer in bytes
+	 * @return				How many bytes were zeroed
+	 */
+	uint32	ZeroBuffer(uint8* Destination, uint32 BufferSize);
+
+	FOpusDecoderWrapper*	OpusDecoderWrapper;
 	const uint8*	SrcBufferData;
 	uint32			SrcBufferDataSize;
 	uint32			SrcBufferOffset;
@@ -170,10 +251,16 @@ public:
 	uint32			TrueSampleCount;
 	uint32			CurrentSampleCount;
 	uint8			NumChannels;
+	uint32			SampleStride;
 
 	TArray<uint8>	LastDecodedPCM;
 	uint32			LastPCMByteSize;
 	uint32			LastPCMOffset;
+	bool			bStoringEndOfFile;
+
+	// Streaming specific
+	USoundWave*		StreamingSoundWave;
+	int32			CurrentChunkIndex;
 };
 
 /**
diff --git a/Engine/Source/Runtime/Engine/Public/AudioDevice.h b/Engine/Source/Runtime/Engine/Public/AudioDevice.h
--- a/Engine/Source/Runtime/Engine/Public/AudioDevice.h
+++ b/Engine/Source/Runtime/Engine/Public/AudioDevice.h
@@ -487,6 +487,12 @@ public:
 		return false;	// assume no support by default
 	}
 
+	/** Whether this device supports streaming of sound waves (i.e. DTYPE_Streaming) */
+	virtual bool SupportsStreaming() const
+	{
+		return false;	// assume no support by default
+	}
+
 	/** Creates a Compressed audio info class suitable for decompressing this SoundWave */
 	virtual class ICompressedAudioInfo* CreateCompressedAudioInfo(USoundWave* SoundWave) { return NULL; }
 
diff --git a/Engine/Source/Runtime/Engine/Public/AudioStreaming.h b/Engine/Source/Runtime/Engine/Public/AudioStreaming.h
new file mode 100644
--- /dev/null
+++ b/Engine/Source/Runtime/Engine/Public/AudioStreaming.h
@@ -0,0 +1,229 @@
+// Copyright 1998-2014 Epic Games, Inc. All Rights Reserved.
+
+/*=============================================================================
+AudioStreaming.h: Definitions of classes used for audio streaming.
+=============================================================================*/
+
+#pragma once
+
+/** Lists possible states used by Thread-safe counter. */
+enum EAudioStreamingState
+{
+	// There are no pending requests/ all requests have been fulfilled.
+	AudioState_ReadyFor_Requests = 0,
+	// Initial request has completed and finalization needs to be kicked off.
+	AudioState_ReadyFor_Finalization = 1,
+	// We're currently loading in chunk data.
+	AudioState_InProgress_Loading = 2,
+	// ...
+	// States InProgress_Loading+N-1 means we're currently loading in N chunks
+	// ...
+};
+
+/**
+ * Async worker to stream audio chunks from the derived data cache.
+ */
+class FAsyncStreamDerivedChunkWorker : public FNonAbandonableTask
+{
+public:
+	/** Initialization constructor. */
+	FAsyncStreamDerivedChunkWorker(
+		const FString& InDerivedDataKey,
+		void* InDestChunkData,
+		int32 InChunkSize,
+		FThreadSafeCounter* InThreadSafeCounter
+		);
+	
+	/**
+	 * Retrieves the derived chunk from the derived data cache.
+	 */
+	void DoWork();
+
+	/**
+	 * Name of the async task.
+	 */
+	static const TCHAR* Name()
+	{
+		return TEXT("FAsyncStreamDerivedChunkTask");
+	}
+
+	/**
+	 * Returns true if the streaming mip request failed.
+	 */
+	bool DidRequestFail() const
+	{
+		return bRequestFailed;
+	}
+
+private:
+	/** Key for retrieving chunk data from the derived data cache. */
+	FString DerivedDataKey;
+	/** The location to which the chunk data should be copied. */
+	void* DestChunkData;
+	/** The size of the chunk in bytes. */
+	int32 ExpectedChunkSize;
+	/** true if the chunk data was not present in the derived data cache. */
+	bool bRequestFailed;
+	/** Thread-safe counter to decrement when data has been copied. */
+	FThreadSafeCounter* ThreadSafeCounter;
+};
+
+/** Async task to stream chunks from the derived data cache. */
+typedef FAsyncTask<FAsyncStreamDerivedChunkWorker> FAsyncStreamDerivedChunkTask;
+
+/**
+ * Contains a request to load chunks of a sound wave
+ */
+struct FWaveRequest
+{
+	TArray<uint32>	RequiredIndices;
+	bool			bPrioritiseRequest;
+};
+
+/**
+ * Stores info about an audio chunk once it's been loaded
+ */
+struct FLoadedAudioChunk
+{
+	uint8*	Data;
+	int32	DataSize;
+	int32	MemorySize;
+	uint32	Index;
+};
+
+/**
+ * Contains everything that will be needed by a SoundWave that's streaming in data
+ */
+struct FStreamingWaveData final
+{
+	FStreamingWaveData();
+	~FStreamingWaveData();
+
+	/**
+	 * Sets up the streaming wave data and loads the first chunk of audio for instant play
+	 *
+	 * @param SoundWave	The SoundWave we are managing
+	 */
+	void Initialize(USoundWave* SoundWave);
+
+	/**
+	 * Updates the streaming status of the sound wave and performs finalization when appropriate. The function returns
+	 * true while there are pending requests in flight and updating needs to continue.
+	 *
+	 * @return					true if there are requests in flight, false otherwise
+	 */
+	bool UpdateStreamingStatus();
+
+	/**
+	 * Tells the SoundWave which chunks are currently required so that it can start loading any needed
+	 *
+	 * @param InChunkIndices The Chunk Indices that are currently needed by all sources using this sound
+	 * @param bShouldPrioritizeAsyncIORequest Whether request should have higher priority than usual
+	 */
+	void UpdateChunkRequests(FWaveRequest& InWaveRequest);
+
+		/**
+	 * Checks whether the requested chunk indices differ from those loaded
+	 *
+	 * @param IndicesToLoad		List of chunk indices that should be loaded
+	 * @param IndicesToFree		List of chunk indices that should be freed
+	 * @return Whether any changes to loaded chunks are required
+	 */
+	bool HasPendingRequests(TArray<uint32>& IndicesToLoad, TArray<uint32>& IndicesToFree) const;
+
+	/**
+	 * Kicks off any pending requests
+	 */
+	void BeginPendingRequests(const TArray<uint32>& IndicesToLoad, const TArray<uint32>& IndicesToFree);
+
+#if WITH_EDITORONLY_DATA
+	/**
+	 * Finishes any Derived Data Cache requests that may be in progress
+	 *
+	 * @return Whether any of the requests failed.
+	 */
+	bool FinishDDCRequests();
+#endif //WITH_EDITORONLY_DATA
+
+private:
+	// Don't allow copy construction as it could free shared memory
+	FStreamingWaveData(const FStreamingWaveData& that);
+	FStreamingWaveData& operator=(FStreamingWaveData const&);
+
+	FLoadedAudioChunk* AddNewLoadedChunk(int32 ChunkSize);
+	void FreeLoadedChunk(FLoadedAudioChunk& LoadedChunk);
+
+public:
+	/** SoundWave this streaming data is for */
+	USoundWave* SoundWave;
+
+	/** Thread-safe counter indicating the audio streaming state. */
+	mutable FThreadSafeCounter	PendingChunkChangeRequestStatus;
+
+	/* Contains pointers to Chunks of audio data that have been streamed in */
+	TArray<FLoadedAudioChunk> LoadedChunks;
+
+	/** Potentially outstanding audio chunk I/O requests */
+	TArray<uint64>	IORequestIndices;
+
+	/** Indices of chunks that are currently loaded */
+	TArray<uint32>	LoadedChunkIndices;
+
+	/** Indices of chunks we want to have loaded */
+	FWaveRequest	CurrentRequest;
+
+#if WITH_EDITORONLY_DATA
+	/** Pending async derived data streaming tasks */
+	TIndirectArray<FAsyncStreamDerivedChunkTask> PendingAsyncStreamDerivedChunkTasks;
+#endif // #if WITH_EDITORONLY_DATA
+};
+
+/**
+* Streaming manager dealing with audio.
+*/
+struct FAudioStreamingManager : public IAudioStreamingManager
+{
+	/** Constructor, initializing all members */
+	FAudioStreamingManager();
+
+	virtual ~FAudioStreamingManager();
+
+	// IStreamingManager interface
+	virtual void UpdateResourceStreaming( float DeltaTime, bool bProcessEverything=false ) override;
+	virtual int32 BlockTillAllRequestsFinished( float TimeLimit = 0.0f, bool bLogResults = false ) override;
+	virtual void CancelForcedResources() override;
+	virtual void NotifyLevelChange() override;
+	virtual void SetDisregardWorldResourcesForFrames( int32 NumFrames ) override;
+	virtual void AddPreparedLevel( class ULevel* Level ) override;
+	virtual void RemoveLevel( class ULevel* Level ) override;
+	// End IStreamingManager interface
+
+	// IAudioStreamingManager interface
+	virtual void AddStreamingSoundWave(USoundWave* SoundWave) override;
+	virtual void RemoveStreamingSoundWave(USoundWave* SoundWave) override;
+	virtual bool IsManagedStreamingSoundWave(const USoundWave* SoundWave) const override;
+	virtual bool IsStreamingInProgress(const USoundWave* SoundWave) override;
+	virtual void AddStreamingSoundSource(FSoundSource* SoundSource) override;
+	virtual void RemoveStreamingSoundSource(FSoundSource* SoundSource) override;
+	virtual bool IsManagedStreamingSoundSource(const FSoundSource* SoundSource) const override;
+	virtual const uint8* GetLoadedChunk(const USoundWave* SoundWave, uint32 ChunkIndex) const override;
+	// End IAudioStreamingManager interface
+
+protected:
+	/**
+	 * Gets Wave request associated with a specific wave
+	 *
+	 * @param SoundWave		SoundWave we want request for
+	 * @return Existing or new request structure
+	 */
+	FWaveRequest& GetWaveRequest(USoundWave* SoundWave);
+
+	/** Sound Waves being managed. */
+	TMap<USoundWave*, FStreamingWaveData> StreamingSoundWaves;
+
+	/** Sound Sources being managed. */
+	TArray<FSoundSource*>	StreamingSoundSources;
+
+	/** Map of requests to make next time sound waves are ready */
+	TMap<USoundWave*, FWaveRequest> WaveRequests;
+};
diff --git a/Engine/Source/Runtime/Engine/Public/ContentStreaming.h b/Engine/Source/Runtime/Engine/Public/ContentStreaming.h
--- a/Engine/Source/Runtime/Engine/Public/ContentStreaming.h
+++ b/Engine/Source/Runtime/Engine/Public/ContentStreaming.h
@@ -475,6 +475,42 @@ struct ITextureStreamingManager : public IStreamingManager
 	virtual bool IsManagedStreamingTexture(const UTexture2D* Texture2D) = 0;
 };
 
+/**
+ * Interface to add functions specifically related to audio streaming
+ */
+struct IAudioStreamingManager : public IStreamingManager
+{
+	/** Adds a new Sound Wave to the streaming manager. */
+	virtual void AddStreamingSoundWave(USoundWave* SoundWave) = 0;
+
+	/** Removes a Sound Wave from the streaming manager. */
+	virtual void RemoveStreamingSoundWave(USoundWave* SoundWave) = 0;
+
+	/** Returns true if this is a Sound Wave that is managed by the streaming manager. */
+	virtual bool IsManagedStreamingSoundWave(const USoundWave* SoundWave) const = 0;
+
+	/** Returns true if this Sound Wave is currently streaming a chunk. */
+	virtual bool IsStreamingInProgress(const USoundWave* SoundWave) = 0;
+
+	/** Adds a new Sound Source to the streaming manager. */
+	virtual void AddStreamingSoundSource(FSoundSource* SoundSource) = 0;
+
+	/** Removes a Sound Source from the streaming manager. */
+	virtual void RemoveStreamingSoundSource(FSoundSource* SoundSource) = 0;
+
+	/** Returns true if this is a streaming Sound Source that is managed by the streaming manager. */
+	virtual bool IsManagedStreamingSoundSource(const FSoundSource* SoundSource) const = 0;
+
+	/**
+	 * Gets a pointer to a chunk of audio data
+	 *
+	 * @param SoundWave		SoundWave we want a chunk from
+	 * @param ChunkIndex	Index of the chunk we want
+	 * @return Either the desired chunk or NULL if it's not loaded
+	 */
+	virtual const uint8* GetLoadedChunk(const USoundWave* SoundWave, uint32 ChunkIndex) const = 0;
+};
+
 /**
  * Streaming manager collection, routing function calls to streaming managers that have been added
  * via AddStreamingManager.
@@ -553,6 +589,11 @@ struct FStreamingManagerCollection : public IStreamingManager
 	 */
 	ITextureStreamingManager& GetTextureStreamingManager() const;
 
+	/**
+	 * Gets a reference to the Audio Streaming Manager interface
+	 */
+	IAudioStreamingManager& GetAudioStreamingManager() const;
+
 	/**
 	 * Adds a streaming manager to the array of managers to route function calls to.
 	 *
@@ -649,6 +690,9 @@ protected:
 
 	/** The currently added texture streaming manager. Can be NULL*/
 	FStreamingManagerTexture* TextureStreamingManager;
+
+	/** The audio streaming manager, should always exist */
+	IAudioStreamingManager* AudioStreamingManager;
 };
 
 /*-----------------------------------------------------------------------------
diff --git a/Engine/Source/Runtime/Windows/XAudio2/Private/XAudio2Buffer.cpp b/Engine/Source/Runtime/Windows/XAudio2/Private/XAudio2Buffer.cpp
--- a/Engine/Source/Runtime/Windows/XAudio2/Private/XAudio2Buffer.cpp
+++ b/Engine/Source/Runtime/Windows/XAudio2/Private/XAudio2Buffer.cpp
@@ -133,6 +133,7 @@ FXAudio2SoundBuffer::~FXAudio2SoundBuffer( void )
 		break;
 
 	case SoundFormat_PCMRT:
+	case SoundFormat_Streaming:
 		// Buffers are freed as part of the ~FSoundSource
 		break;
 
@@ -196,11 +197,25 @@ int32 FXAudio2SoundBuffer::GetSize( void )
 	case SoundFormat_XWMA:
 		TotalSize = XWMA.XWMADataSize + XWMA.XWMASeekDataSize;
 		break;
+
+	case SoundFormat_Streaming:
+		TotalSize = ( MONO_PCM_BUFFER_SIZE * 2 * NumChannels );
+		break;
 	}
 
 	return( TotalSize );
 }
 
+int32 FXAudio2SoundBuffer::GetCurrentChunkIndex() const
+{
+	return DecompressionState->GetCurrentChunkIndex();
+}
+
+int32 FXAudio2SoundBuffer::GetCurrentChunkOffset() const
+{
+	return DecompressionState->GetCurrentChunkOffset();
+}
+
 /** 
  * Setup a WAVEFORMATEX structure
  */
@@ -300,7 +315,15 @@ void FXAudio2SoundBuffer::InitXWMA( USoundWave* Wave, FXMAInfo* XMAInfo )
  */
 bool FXAudio2SoundBuffer::ReadCompressedData( uint8* Destination, bool bLooping )
 {
-	return( DecompressionState->ReadCompressedData( Destination, bLooping, MONO_PCM_BUFFER_SIZE * NumChannels ) );
+	const uint32 kPCMBufferSize = MONO_PCM_BUFFER_SIZE * NumChannels;
+	if (SoundFormat == SoundFormat_Streaming)
+	{
+		return DecompressionState->StreamCompressedData( Destination, bLooping, kPCMBufferSize );
+	}
+	else
+	{
+		return( DecompressionState->ReadCompressedData( Destination, bLooping, kPCMBufferSize ) );
+	}
 }
 
 void FXAudio2SoundBuffer::Seek( const float SeekTime )
@@ -451,6 +474,41 @@ FXAudio2SoundBuffer* FXAudio2SoundBuffer::CreateNativeBuffer( FXAudio2Device* XA
 	return( Buffer );
 }
 
+FXAudio2SoundBuffer* FXAudio2SoundBuffer::CreateStreamingBuffer( FXAudio2Device* XAudio2Device, USoundWave* Wave )
+{
+	// Always create a new buffer for streaming sounds
+	FXAudio2SoundBuffer* Buffer = new FXAudio2SoundBuffer(XAudio2Device, SoundFormat_Streaming);
+
+	// Prime the first two buffers and prepare the decompression
+	FSoundQualityInfo QualityInfo = { 0 };
+
+	Buffer->DecompressionState = XAudio2Device->CreateCompressedAudioInfo(Wave);
+
+	if (Buffer->DecompressionState->StreamCompressedInfo(Wave, &QualityInfo))
+	{
+		// Refresh the wave data
+		Wave->SampleRate = QualityInfo.SampleRate;
+		Wave->NumChannels = QualityInfo.NumChannels;
+		Wave->RawPCMDataSize = QualityInfo.SampleDataSize;
+		Wave->Duration = QualityInfo.Duration;
+
+		// Clear out any dangling pointers
+		Buffer->PCM.PCMData = NULL;
+		Buffer->PCM.PCMDataSize = 0;
+
+		Buffer->InitWaveFormatEx(WAVE_FORMAT_PCM, Wave, false);
+	}
+	else
+	{
+		Wave->DecompressionType = DTYPE_Invalid;
+		Wave->NumChannels = 0;
+
+		Wave->RemoveAudioResource();
+	}
+
+	return(Buffer);
+}
+
 /**
  * Static function used to create a buffer.
  *
@@ -471,7 +529,7 @@ FXAudio2SoundBuffer* FXAudio2SoundBuffer::Init( FAudioDevice* AudioDevice, USoun
 
 	// Allow the precache to happen if necessary
 	EDecompressionType DecompressionType = Wave->DecompressionType;
-	if (bForceRealTime &&  DecompressionType != DTYPE_Setup )
+	if (bForceRealTime && DecompressionType != DTYPE_Setup && DecompressionType != DTYPE_Streaming)
 	{
 		DecompressionType = DTYPE_RealTime;
 	}
@@ -527,6 +585,11 @@ FXAudio2SoundBuffer* FXAudio2SoundBuffer::Init( FAudioDevice* AudioDevice, USoun
 		}
 		break;
 
+	case DTYPE_Streaming:
+		// Always create a new buffer for streaming sounds
+		Buffer = CreateStreamingBuffer( XAudio2Device, Wave );
+		break;
+
 	case DTYPE_Invalid:
 	default:
 		// Invalid will be set if the wave cannot be played
diff --git a/Engine/Source/Runtime/Windows/XAudio2/Private/XAudio2Source.cpp b/Engine/Source/Runtime/Windows/XAudio2/Private/XAudio2Source.cpp
--- a/Engine/Source/Runtime/Windows/XAudio2/Private/XAudio2Source.cpp
+++ b/Engine/Source/Runtime/Windows/XAudio2/Private/XAudio2Source.cpp
@@ -368,6 +368,7 @@ bool FXAudio2SoundSource::Init( FWaveInstance* InWaveInstance )
 				break;
 
 			case SoundFormat_PCMRT:
+			case SoundFormat_Streaming:
 				SubmitPCMRTBuffers();
 				break;
 
@@ -1081,7 +1082,7 @@ bool FXAudio2SoundSource::IsFinished( void )
 		XAUDIO2_VOICE_STATE SourceState;
 		Source->GetState( &SourceState );
 
-		const bool bIsRealTimeSource = XAudio2Buffer->SoundFormat == SoundFormat_PCMRT;
+		const bool bIsRealTimeSource = XAudio2Buffer->SoundFormat == SoundFormat_PCMRT || XAudio2Buffer->SoundFormat == SoundFormat_Streaming;
 
 		// If we have no queued buffers, we're either at the end of a sound, or starved
 		if( SourceState.BuffersQueued == 0 )
diff --git a/Engine/Source/Runtime/Windows/XAudio2/Private/XAudio2Support.h b/Engine/Source/Runtime/Windows/XAudio2/Private/XAudio2Support.h
--- a/Engine/Source/Runtime/Windows/XAudio2/Private/XAudio2Support.h
+++ b/Engine/Source/Runtime/Windows/XAudio2/Private/XAudio2Support.h
@@ -169,6 +169,15 @@ public:
 	 */
 	static FXAudio2SoundBuffer* CreateNativeBuffer( FXAudio2Device* XAudio2Device, USoundWave* Wave );
 
+	/**
+	 * Static function used to create an XAudio buffer and dynamically upload streaming data to.
+	 *
+	 * @param InWave		USoundWave to use as template and wave source
+	 * @param AudioDevice	audio device to attach created buffer to
+	 * @return FXAudio2SoundBuffer pointer if buffer creation succeeded, NULL otherwise
+	 */
+	static FXAudio2SoundBuffer* CreateStreamingBuffer( FXAudio2Device* XAudio2Device, USoundWave* Wave );
+
 	/**
 	 * Static function used to create a buffer.
 	 *
@@ -185,6 +194,9 @@ public:
 	 */
 	int32 GetSize( void );
 
+	virtual int32 GetCurrentChunkIndex() const override;
+	virtual int32 GetCurrentChunkOffset() const override;
+
 	/** Audio device this buffer is attached to	*/
 	FAudioDevice*				AudioDevice;
 
diff --git a/Engine/Source/Runtime/Windows/XAudio2/Public/XAudio2Device.h b/Engine/Source/Runtime/Windows/XAudio2/Public/XAudio2Device.h
--- a/Engine/Source/Runtime/Windows/XAudio2/Public/XAudio2Device.h
+++ b/Engine/Source/Runtime/Windows/XAudio2/Public/XAudio2Device.h
@@ -63,7 +63,8 @@ enum ESoundFormat
 	SoundFormat_PCMPreview,
 	SoundFormat_PCMRT,
 	SoundFormat_XMA2,
-	SoundFormat_XWMA
+	SoundFormat_XWMA,
+	SoundFormat_Streaming
 };
 
 class FXAudio2SoundBuffer;
@@ -115,6 +116,11 @@ class FXAudio2Device : public FAudioDevice
 		return true;
 	}
 
+	virtual bool SupportsStreaming() const override
+	{
+		return true;
+	}
+
 	virtual class ICompressedAudioInfo* CreateCompressedAudioInfo(USoundWave* SoundWave) override;
 
 	/** 
