fixed inconsitent navmesh generation on geometry from sublevels

diff --git a/Engine/Source/Runtime/Engine/Private/AI/Navigation/NavigationSystem.cpp b/Engine/Source/Runtime/Engine/Private/AI/Navigation/NavigationSystem.cpp
--- a/Engine/Source/Runtime/Engine/Private/AI/Navigation/NavigationSystem.cpp
+++ b/Engine/Source/Runtime/Engine/Private/AI/Navigation/NavigationSystem.cpp
@@ -291,9 +291,21 @@ void UNavigationSystem::PostInitProperties()
 
 #if WITH_NAVIGATION_GENERATOR
 		UWorld* World = GetWorld();
-		if (World && World->PersistentLevel)
+		if (World)
 		{
-			OnLevelAddedToWorld(World->PersistentLevel, World);
+			if (World->PersistentLevel)
+			{
+				OnLevelAddedToWorld(World->PersistentLevel, World);
+			}
+
+			for (int32 Idx = 0; Idx < World->StreamingLevels.Num(); Idx++)
+			{
+				ULevelStreaming* StreamingInfo = World->StreamingLevels[Idx];
+				if (StreamingInfo && StreamingInfo->GetLoadedLevel())
+				{
+					OnLevelAddedToWorld(StreamingInfo->GetLoadedLevel(), World);
+				}
+			}
 		}
 #endif
 
