Fix LandscapeLayerBlend material node when using non-weight blending. Previous content using LB_AlphaBlend will be automatically updated to use LB_WeightBlend, and LB_AlphaBlend now performs regular alpha blending for non-weight layers.

TTP#339892

diff --git a/Engine/Source/Runtime/Core/Public/UObject/ObjectVersion.h b/Engine/Source/Runtime/Core/Public/UObject/ObjectVersion.h
--- a/Engine/Source/Runtime/Core/Public/UObject/ObjectVersion.h
+++ b/Engine/Source/Runtime/Core/Public/UObject/ObjectVersion.h
@@ -572,6 +572,8 @@ enum EUnrealEngineObjectUE4Version
 	VER_UE4_STATIC_SHADOW_DEPTH_MAPS,
 	// Add RF_Transactional to data assets, fixing undo problems when editing them
 	VER_UE4_ADD_TRANSACTIONAL_TO_DATA_ASSETS,
+	// Change LB_AlphaBlend to LB_WeightBlend in ELandscapeLayerBlendType
+	VER_UE4_ADD_LB_WEIGHTBLEND,
 
 	// -----<new versions can be added before this line>-------------------------------------------------
 	// - this needs to be the last line (see note below)
diff --git a/Engine/Source/Runtime/Engine/Classes/Materials/MaterialExpressionLandscapeLayerBlend.h b/Engine/Source/Runtime/Engine/Classes/Materials/MaterialExpressionLandscapeLayerBlend.h
--- a/Engine/Source/Runtime/Engine/Classes/Materials/MaterialExpressionLandscapeLayerBlend.h
+++ b/Engine/Source/Runtime/Engine/Classes/Materials/MaterialExpressionLandscapeLayerBlend.h
@@ -8,6 +8,7 @@
 UENUM()
 enum ELandscapeLayerBlendType
 {
+	LB_WeightBlend,
 	LB_AlphaBlend,
 	LB_HeightBlend,
 	LB_MAX,
@@ -64,6 +65,7 @@ class UMaterialExpressionLandscapeLayerBlend : public UMaterialExpression
 
 
 	// Begin UObject Interface
+	void Serialize(FArchive& Ar) override;
 #if WITH_EDITOR
 	virtual void PostEditChangeProperty(FPropertyChangedEvent& PropertyChangedEvent) override;
 	// End UObject Interface
diff --git a/Engine/Source/Runtime/Engine/Private/Materials/MaterialExpressions.cpp b/Engine/Source/Runtime/Engine/Private/Materials/MaterialExpressions.cpp
--- a/Engine/Source/Runtime/Engine/Private/Materials/MaterialExpressions.cpp
+++ b/Engine/Source/Runtime/Engine/Private/Materials/MaterialExpressions.cpp
@@ -8573,6 +8573,23 @@ UMaterialExpressionLandscapeLayerBlend::UMaterialExpressionLandscapeLayerBlend(c
 	MenuCategories.Add(ConstructorStatics.NAME_Landscape);
 }
 
+void UMaterialExpressionLandscapeLayerBlend::Serialize(FArchive& Ar)
+{
+	Super::Serialize(Ar);
+
+	if (Ar.IsLoading() && Ar.UE4Ver() < VER_UE4_ADD_LB_WEIGHTBLEND)
+	{
+		// convert any LB_AlphaBlend entries to LB_WeightBlend
+		for (FLayerBlendInput& LayerInput : Layers)
+		{
+			if (LayerInput.BlendType == LB_AlphaBlend)
+			{
+				LayerInput.BlendType = LB_WeightBlend;
+			}
+		}
+	}
+}
+
 const TArray<FExpressionInput*> UMaterialExpressionLandscapeLayerBlend::GetInputs()
 {
 	TArray<FExpressionInput*> Result;
@@ -8660,35 +8677,39 @@ int32 UMaterialExpressionLandscapeLayerBlend::Compile(class FMaterialCompiler* C
 
 		FLayerBlendInput& Layer = Layers[LayerIdx];
 
-		// Height input
-		const int32 HeightCode = Layer.HeightInput.Expression ? Layer.HeightInput.Compile(Compiler, MultiplexIndex) : Compiler->Constant(Layer.ConstHeightInput);
-
-		const int32 WeightCode = Compiler->StaticTerrainLayerWeight(Layer.LayerName,Layer.PreviewWeight > 0.0f ? Compiler->Constant(Layer.PreviewWeight) : INDEX_NONE);
-		if( WeightCode != INDEX_NONE )
-		{
-			switch( Layer.BlendType )
-			{
-			case LB_AlphaBlend:
-				{
-					// Store the weight plus accumulate the sum of all weights so far
-					WeightCodes[LayerIdx] = WeightCode;
-					WeightSumCode = Compiler->Add(WeightSumCode, WeightCode);
-				}
-				break;
-			case LB_HeightBlend:
-				{
-					bNeedsRenormalize = true;
-
-					// Modify weight with height
-					int32 ModifiedWeightCode = Compiler->Clamp(
-						Compiler->Add(Compiler->Lerp(Compiler->Constant(-1.f), Compiler->Constant(1.f), WeightCode), HeightCode),
-						Compiler->Constant(0.f), Compiler->Constant(1.f) );
-
-					// Store the final weight plus accumulate the sum of all weights so far
-					WeightCodes[LayerIdx] = ModifiedWeightCode;
-					WeightSumCode = Compiler->Add(WeightSumCode, ModifiedWeightCode);
-				}
-				break;
+		// LB_AlphaBlend layers are blended last
+		if (Layer.BlendType != LB_AlphaBlend)
+		{
+		    // Height input
+		    const int32 HeightCode = Layer.HeightInput.Expression ? Layer.HeightInput.Compile(Compiler, MultiplexIndex) : Compiler->Constant(Layer.ConstHeightInput);
+    
+		    const int32 WeightCode = Compiler->StaticTerrainLayerWeight(Layer.LayerName,Layer.PreviewWeight > 0.0f ? Compiler->Constant(Layer.PreviewWeight) : INDEX_NONE);
+		    if( WeightCode != INDEX_NONE )
+		    {
+			    switch( Layer.BlendType )
+			    {
+			    case LB_WeightBlend:
+				    {
+					    // Store the weight plus accumulate the sum of all weights so far
+					    WeightCodes[LayerIdx] = WeightCode;
+					    WeightSumCode = Compiler->Add(WeightSumCode, WeightCode);
+				    }
+				    break;
+			    case LB_HeightBlend:
+				    {
+					    bNeedsRenormalize = true;
+    
+					    // Modify weight with height
+					    int32 ModifiedWeightCode = Compiler->Clamp(
+						    Compiler->Add(Compiler->Lerp(Compiler->Constant(-1.f), Compiler->Constant(1.f), WeightCode), HeightCode),
+						    Compiler->Constant(0.f), Compiler->Constant(1.f) );
+    
+					    // Store the final weight plus accumulate the sum of all weights so far
+					    WeightCodes[LayerIdx] = ModifiedWeightCode;
+					    WeightSumCode = Compiler->Add(WeightSumCode, ModifiedWeightCode);
+				    }
+				    break;
+			    }
 			}
 		}
 	}
@@ -8719,9 +8740,24 @@ int32 UMaterialExpressionLandscapeLayerBlend::Compile(class FMaterialCompiler* C
 		}
 	}
 
-	if( OutputCode != INDEX_NONE )
+	// Blend in LB_AlphaBlend layers
+	for (FLayerBlendInput& Layer : Layers)
+	{
+		if (Layer.BlendType == LB_AlphaBlend)
+		{
+			const int32 WeightCode = Compiler->StaticTerrainLayerWeight(Layer.LayerName, Layer.PreviewWeight > 0.0f ? Compiler->Constant(Layer.PreviewWeight) : INDEX_NONE);
+			if (WeightCode != INDEX_NONE)
+			{
+				const int32 LayerCode = Layer.LayerInput.Expression ? Layer.LayerInput.Compile(Compiler, MultiplexIndex) : Compiler->Constant3(Layer.ConstLayerInput.X, Layer.ConstLayerInput.Y, Layer.ConstLayerInput.Z);
+				// Blend in the layer using the alpha value
+				OutputCode = Compiler->Lerp(OutputCode, LayerCode, WeightCode);
+			}
+		}
+	}
+	
+	if (OutputCode != INDEX_NONE)
 	{
-		//We've definitely passed the reentrant check here so we're good to call IsResultMaterialAttributes().
+		// We've definitely passed the reentrant check here so we're good to call IsResultMaterialAttributes().
 		bool bFoundExpression = false;
 		bool bIsResultMaterialAttributes = false;
 		for (int32 LayerIdx = 0; LayerIdx < Layers.Num(); LayerIdx++)
