Stat Net now works correctly and no longer reports 'command not found'

#udn ConsoleCommand Stat NET showing as Command Not Found, can't use Stat NET

diff --git a/Engine/Source/Runtime/Core/Private/Stats/StatsCommand.cpp b/Engine/Source/Runtime/Core/Private/Stats/StatsCommand.cpp
--- a/Engine/Source/Runtime/Core/Private/Stats/StatsCommand.cpp
+++ b/Engine/Source/Runtime/Core/Private/Stats/StatsCommand.cpp
@@ -1380,6 +1380,9 @@ bool DirectStatsCommand(const TCHAR* Cmd, bool bBlockForCompletion /*= false*/,
 		else if ( FParse::Command( &TempCmd, TEXT( "hier" ) ) )
 		{
 		}
+		else if ( FParse::Command( &TempCmd, TEXT( "net" ) ) )
+		{
+		}
 		else
 #endif
 		{
diff --git a/Engine/Source/Runtime/Engine/Private/NetworkDriver.cpp b/Engine/Source/Runtime/Engine/Private/NetworkDriver.cpp
--- a/Engine/Source/Runtime/Engine/Private/NetworkDriver.cpp
+++ b/Engine/Source/Runtime/Engine/Private/NetworkDriver.cpp
@@ -170,123 +170,130 @@ void UNetDriver::TickFlush(float DeltaSeconds)
 #endif // WITH_SERVER_CODE
 	}
 
+#if STATS
 	// Update network stats
 	if ( Time - StatUpdateTime > StatPeriod && ( ClientConnections.Num() > 0 || ServerConnection != NULL ) )
 	{
-		float RealTime = Time - StatUpdateTime;
+		int32 Ping = 0;
+		int32 NumOpenChannels = 0;
+		int32 NumActorChannels = 0;
+		int32 NumDormantActors = 0;
+		int32 NumActorChannelsReadyDormant = 0;
+		int32 NumActors = 0;
+		int32 AckCount = 0;
+		int32 UnAckCount = 0;
+		int32 PendingCount = 0;
+		int32 NetSaturated = 0;
 
-		// Use the elapsed time to keep things scaled to one measured unit
-		InBytes = FMath::TruncToInt(InBytes / RealTime);
-		OutBytes = FMath::TruncToInt(OutBytes / RealTime);
+		if (FThreadStats::IsCollectingData())
+		{
+			float RealTime = Time - StatUpdateTime;
 
-		NetGUIDOutBytes = FMath::TruncToInt(NetGUIDOutBytes / RealTime);
-		NetGUIDInBytes = FMath::TruncToInt(NetGUIDInBytes / RealTime);
+			// Use the elapsed time to keep things scaled to one measured unit
+			InBytes = FMath::TruncToInt(InBytes / RealTime);
+			OutBytes = FMath::TruncToInt(OutBytes / RealTime);
 
-		// Save off for stats later
+			NetGUIDOutBytes = FMath::TruncToInt(NetGUIDOutBytes / RealTime);
+			NetGUIDInBytes = FMath::TruncToInt(NetGUIDInBytes / RealTime);
 
-		InBytesPerSecond = InBytes;
-		OutBytesPerSecond = OutBytes;
+			// Save off for stats later
 
-		InPackets = FMath::TruncToInt(InPackets / RealTime);
-		OutPackets = FMath::TruncToInt(OutPackets / RealTime);
-		InBunches = FMath::TruncToInt(InBunches / RealTime);
-		OutBunches = FMath::TruncToInt(OutBunches / RealTime);
-		OutPacketsLost = FMath::TruncToInt(100.f * OutPacketsLost / FMath::Max((float)OutPackets,1.f));
-		InPacketsLost = FMath::TruncToInt(100.f * InPacketsLost / FMath::Max((float)InPackets + InPacketsLost,1.f));
-		
-#if STATS
+			InBytesPerSecond = InBytes;
+			OutBytesPerSecond = OutBytes;
 
-		if (FThreadStats::IsCollectingData())
-		{
-			// Copy the net status values over
+			InPackets = FMath::TruncToInt(InPackets / RealTime);
+			OutPackets = FMath::TruncToInt(OutPackets / RealTime);
+			InBunches = FMath::TruncToInt(InBunches / RealTime);
+			OutBunches = FMath::TruncToInt(OutBunches / RealTime);
+			OutPacketsLost = FMath::TruncToInt(100.f * OutPacketsLost / FMath::Max((float)OutPackets,1.f));
+			InPacketsLost = FMath::TruncToInt(100.f * InPacketsLost / FMath::Max((float)InPackets + InPacketsLost,1.f));
+			
 			if (ServerConnection != NULL && ServerConnection->PlayerController != NULL && ServerConnection->PlayerController->PlayerState != NULL)
 			{
-				SET_DWORD_STAT(STAT_Ping, FMath::TruncToInt(1000.0f * ServerConnection->PlayerController->PlayerState->ExactPing));
-			}
-			else
-			{
-				SET_DWORD_STAT(STAT_Ping, 0);
+				Ping = FMath::TruncToInt(1000.0f * ServerConnection->PlayerController->PlayerState->ExactPing);
 			}
-			SET_DWORD_STAT(STAT_Channels, 0);
+
 			if (ServerConnection != NULL)
 			{
-				INC_DWORD_STAT_BY(STAT_Channels, ServerConnection->OpenChannels.Num());
+				NumOpenChannels = ServerConnection->OpenChannels.Num();
 			}
 			for (int32 i = 0; i < ClientConnections.Num(); i++)
 			{
-				INC_DWORD_STAT_BY(STAT_Channels, ClientConnections[i]->OpenChannels.Num());
+				NumOpenChannels += ClientConnections[i]->OpenChannels.Num();
 			}
-			SET_DWORD_STAT(STAT_OutLoss,OutPacketsLost);
-			SET_DWORD_STAT(STAT_InLoss,InPacketsLost);
-			SET_DWORD_STAT(STAT_InRate,InBytes);
-			SET_DWORD_STAT(STAT_OutRate,OutBytes);
-			SET_DWORD_STAT(STAT_InPackets,InPackets);
-			SET_DWORD_STAT(STAT_OutPackets,OutPackets);
-			SET_DWORD_STAT(STAT_InBunches,InBunches);
-			SET_DWORD_STAT(STAT_OutBunches,OutBunches);
-
-			SET_DWORD_STAT(STAT_NetGUIDInRate,NetGUIDInBytes);
-			SET_DWORD_STAT(STAT_NetGUIDOutRate,NetGUIDOutBytes);
 
 			// Use the elapsed time to keep things scaled to one measured unit
 			VoicePacketsSent = FMath::TruncToInt(VoicePacketsSent / RealTime);
-			SET_DWORD_STAT(STAT_VoicePacketsSent,VoicePacketsSent);
 			VoicePacketsRecv = FMath::TruncToInt(VoicePacketsRecv / RealTime);
-			SET_DWORD_STAT(STAT_VoicePacketsRecv,VoicePacketsRecv);
 			VoiceBytesSent = FMath::TruncToInt(VoiceBytesSent / RealTime);
-			SET_DWORD_STAT(STAT_VoiceBytesSent,VoiceBytesSent);
 			VoiceBytesRecv = FMath::TruncToInt(VoiceBytesRecv / RealTime);
-			SET_DWORD_STAT(STAT_VoiceBytesRecv,VoiceBytesRecv);
+
 			// Determine voice percentages
 			VoiceInPercent = (InBytes > 0) ? FMath::TruncToInt(100.f * (float)VoiceBytesRecv / (float)InBytes) : 0;
 			VoiceOutPercent = (OutBytes > 0) ? FMath::TruncToInt(100.f * (float)VoiceBytesSent / (float)OutBytes) : 0;
 
-			SET_DWORD_STAT(STAT_PercentInVoice,VoiceInPercent);
-			SET_DWORD_STAT(STAT_PercentOutVoice,VoiceOutPercent);
-
 			UNetConnection * Connection = (ServerConnection ? ServerConnection : (ClientConnections.Num() > 0 ? ClientConnections[0] : NULL));
 			if (Connection)
 			{
-				SET_DWORD_STAT(STAT_NumActorChannels,Connection->ActorChannels.Num());
-				SET_DWORD_STAT(STAT_NumDormantActors,Connection->DormantActors.Num());
-		
-				int32 NumReadyForDormancy = 0;
-			
+				NumActorChannels = Connection->ActorChannels.Num();
+				NumDormantActors = Connection->DormantActors.Num();
+
 				for (auto It = Connection->ActorChannels.CreateIterator(); It; ++It)
 				{
 					UActorChannel* Chan = It.Value();
 					if(Chan && Chan->ReadyForDormancy(true))
 					{
-						NumReadyForDormancy++;
+						NumActorChannelsReadyDormant++;
 					}
 				}
 
-				int32 ActorCount=0;
 				if (World)
 				{
 					for (FActorIterator It(World); It; ++It)
 					{
-						ActorCount++;
+						NumActors++;
 					}
 				}
 
-				SET_DWORD_STAT(STAT_NumActors,ActorCount);
-				SET_DWORD_STAT(STAT_NumActorChannelsReadyDormant,NumReadyForDormancy);				
-
-				int32 AckCount = 0;
-				int32 UnAckCount = 0;
-				int32 PendingCount = 0;
 				Connection->PackageMap->GetNetGUIDStats(AckCount, UnAckCount, PendingCount);
 
-				SET_DWORD_STAT(STAT_NumNetGUIDsAckd,AckCount);
-				SET_DWORD_STAT(STAT_NumNetGUIDsPending,PendingCount);
-				SET_DWORD_STAT(STAT_NumNetGUIDsUnAckd,UnAckCount);
-
-				SET_DWORD_STAT(STAT_NetSaturated, Connection->IsNetReady(false) ? 0 : 1 );
+				NetSaturated = Connection->IsNetReady(false) ? 0 : 1;
 			}
 		}
+
+		// Copy the net status values over
+		SET_DWORD_STAT(STAT_Ping, Ping);
+		SET_DWORD_STAT(STAT_Channels, NumOpenChannels);
+
+		SET_DWORD_STAT(STAT_OutLoss, OutPacketsLost);
+		SET_DWORD_STAT(STAT_InLoss, InPacketsLost);
+		SET_DWORD_STAT(STAT_InRate, InBytes);
+		SET_DWORD_STAT(STAT_OutRate, OutBytes);
+		SET_DWORD_STAT(STAT_InPackets, InPackets);
+		SET_DWORD_STAT(STAT_OutPackets, OutPackets);
+		SET_DWORD_STAT(STAT_InBunches, InBunches);
+		SET_DWORD_STAT(STAT_OutBunches, OutBunches);
+
+		SET_DWORD_STAT(STAT_NetGUIDInRate, NetGUIDInBytes);
+		SET_DWORD_STAT(STAT_NetGUIDOutRate, NetGUIDOutBytes);
+
+		SET_DWORD_STAT(STAT_VoicePacketsSent, VoicePacketsSent);
+		SET_DWORD_STAT(STAT_VoicePacketsRecv, VoicePacketsRecv);
+		SET_DWORD_STAT(STAT_VoiceBytesSent, VoiceBytesSent);
+		SET_DWORD_STAT(STAT_VoiceBytesRecv, VoiceBytesRecv);
+
+		SET_DWORD_STAT(STAT_PercentInVoice, VoiceInPercent);
+		SET_DWORD_STAT(STAT_PercentOutVoice, VoiceOutPercent);
+
+		SET_DWORD_STAT(STAT_NumActorChannels, NumActorChannels);
+		SET_DWORD_STAT(STAT_NumDormantActors, NumDormantActors);
+		SET_DWORD_STAT(STAT_NumActors, NumActors);
+		SET_DWORD_STAT(STAT_NumActorChannelsReadyDormant, NumActorChannelsReadyDormant);
+		SET_DWORD_STAT(STAT_NumNetGUIDsAckd, AckCount);
+		SET_DWORD_STAT(STAT_NumNetGUIDsPending, UnAckCount);
+		SET_DWORD_STAT(STAT_NumNetGUIDsUnAckd, PendingCount);
+		SET_DWORD_STAT(STAT_NetSaturated, NetSaturated);
 		
-#endif
 		// Reset everything
 		InBytes = 0;
 		OutBytes = 0;
@@ -306,6 +313,7 @@ void UNetDriver::TickFlush(float DeltaSeconds)
 		VoiceOutPercent = 0;
 		StatUpdateTime = Time;
 	}
+#endif
 
 	// Poll all sockets.
 	if( ServerConnection )
