Stats - stat unit, unitgraph etc, crashes when all viewports are closed (fixed, based on https://github.com/EpicGames/UnrealEngine/pull/251)

diff --git a/Engine/Source/Runtime/Engine/Private/UnrealEngine.cpp b/Engine/Source/Runtime/Engine/Private/UnrealEngine.cpp
--- a/Engine/Source/Runtime/Engine/Private/UnrealEngine.cpp
+++ b/Engine/Source/Runtime/Engine/Private/UnrealEngine.cpp
@@ -4973,7 +4973,7 @@ bool UEngine::HandleStatCommand( UWorld* World, FCommonViewportClient* ViewportC
 		{
 			if (EngineStat.ToggleFunc)
 			{
-				return (this->*(EngineStat.ToggleFunc))(World, ViewportClient, Temp);
+				return ViewportClient ? ( this->*(EngineStat.ToggleFunc) )(World, ViewportClient, Temp) : false;
 			}
 			return true;
 		}
@@ -9834,8 +9834,7 @@ bool UEngine::IsEngineStat(const FString& InName)
 
 void UEngine::SetEngineStat(UWorld* World, FCommonViewportClient* ViewportClient, const FString& InName, const bool bShow)
 {
-	check(ViewportClient);
-	if (IsEngineStat(InName) && ViewportClient->IsStatEnabled(*InName) != bShow)
+	if (ViewportClient && IsEngineStat(InName) && ViewportClient->IsStatEnabled(*InName) != bShow)
 	{
 		ExecEngineStat(World, ViewportClient, *InName);
 	}
@@ -9888,7 +9887,11 @@ int32 UEngine::RenderStatVersion(UWorld* World, FViewport* Viewport, FCanvas* Ca
 // DETAILED
 bool UEngine::ToggleStatDetailed(UWorld* World, FCommonViewportClient* ViewportClient, const TCHAR* Stream)
 {
-	check(ViewportClient);
+	if( ViewportClient == nullptr )
+	{
+		// Ignore if all Viewports are closed.
+		return false;
+	}
 
 	// Each of these stats should call "Detailed -Skip" when they themselves are disabled
 	static bool bSetup = false;
@@ -10009,8 +10012,13 @@ int32 UEngine::RenderStatSummary(UWorld* World, FViewport* Viewport, FCanvas* Ca
 // NAMEDEVENTS
 bool UEngine::ToggleStatNamedEvents(UWorld* World, FCommonViewportClient* ViewportClient, const TCHAR* Stream)
 {
+	if( ViewportClient == nullptr )
+	{
+		// Ignore if all Viewports are closed.
+		return false;
+	}
+
 	// Enable emission of named events and force enable cycle stats.
-	check(ViewportClient);
 	if (ViewportClient->IsStatEnabled(TEXT("NamedEvents")))
 	{
 		if (GCycleStatsShouldEmitNamedEvents == 0)
@@ -10251,7 +10259,12 @@ int32 UEngine::RenderStatLevelMap(UWorld* World, FViewport* Viewport, FCanvas* C
 // UNIT
 bool UEngine::ToggleStatUnit(UWorld* World, FCommonViewportClient* ViewportClient, const TCHAR* Stream)
 {
-	check(ViewportClient);
+	if( ViewportClient == nullptr )
+	{
+		// Ignore if all Viewports are closed.
+		return false;
+	}
+
 	const bool bShowUnitMaxTimes = ViewportClient->IsStatEnabled(TEXT("UnitMax"));
 	if (bShowUnitMaxTimes != false)
 	{
@@ -10288,7 +10301,11 @@ int32 UEngine::RenderStatUnit(UWorld* World, FViewport* Viewport, FCanvas* Canva
 #if !UE_BUILD_SHIPPING
 bool UEngine::ToggleStatUnitMax(UWorld* World, FCommonViewportClient* ViewportClient, const TCHAR* Stream)
 {
-	check(ViewportClient);
+	if( ViewportClient == nullptr )
+	{
+		// Ignore if all Viewports are closed.
+		return false;
+	}
 	const bool bShowUnitMaxTimes = ViewportClient->IsStatEnabled(TEXT("UnitMax"));
 	if (bShowUnitMaxTimes)
 	{
@@ -10313,7 +10330,11 @@ bool UEngine::ToggleStatUnitMax(UWorld* World, FCommonViewportClient* ViewportCl
 // UNITGRAPH
 bool UEngine::ToggleStatUnitGraph(UWorld* World, FCommonViewportClient* ViewportClient, const TCHAR* Stream)
 {
-	check(ViewportClient);
+	if( ViewportClient == nullptr )
+	{
+		// Ignore if all Viewports are closed.
+		return false;
+	}
 	const bool bShowUnitGraph = ViewportClient->IsStatEnabled(TEXT("UnitGraph"));
 	if (bShowUnitGraph)
 	{
@@ -10338,7 +10359,11 @@ bool UEngine::ToggleStatUnitGraph(UWorld* World, FCommonViewportClient* Viewport
 // UNITTIME
 bool UEngine::ToggleStatUnitTime(UWorld* World, FCommonViewportClient* ViewportClient, const TCHAR* Stream)
 {
-	check(ViewportClient);
+	if( ViewportClient == nullptr )
+	{
+		// Ignore if all Viewports are closed.
+		return false;
+	}
 	const bool bShowUnitTime = ViewportClient->IsStatEnabled(TEXT("UnitTime"));
 	if (bShowUnitTime)
 	{
@@ -10571,7 +10596,11 @@ int32 UEngine::RenderStatSoundCues(UWorld* World, FViewport* Viewport, FCanvas*
 // SOUNDS
 bool UEngine::ToggleStatSounds(UWorld* World, FCommonViewportClient* ViewportClient, const TCHAR* Stream)
 {
-	check(ViewportClient);
+	if( ViewportClient == nullptr )
+	{
+		// Ignore if all Viewports are closed.
+		return false;
+	}
 	const bool bHelp = Stream ? FCString::Stristr(Stream, TEXT("?")) != NULL : false;
 	if (bHelp)
 	{
