#ue4 - Move initialization of some CharacterMovementComponent values from Character to the component.

- derived classes of CMC trying to set those values in their constructor would have them overwritten by the Character initialization.

diff --git a/Engine/Source/Runtime/Engine/Private/Character.cpp b/Engine/Source/Runtime/Engine/Private/Character.cpp
--- a/Engine/Source/Runtime/Engine/Private/Character.cpp
+++ b/Engine/Source/Runtime/Engine/Private/Character.cpp
@@ -73,13 +73,7 @@ ACharacter::ACharacter(const class FPostConstructInitializeProperties& PCIP)
 	if (CharacterMovement)
 	{
 		CharacterMovement->UpdatedComponent = CapsuleComponent;
-		CharacterMovement->MaxStepHeight = 45.f;
 		CrouchedEyeHeight = CharacterMovement->CrouchedHalfHeight * 0.80f;
-
-		CharacterMovement->GetNavAgentProperties()->bCanJump = true;
-		CharacterMovement->GetNavAgentProperties()->bCanWalk = true;
-		CharacterMovement->GetNavAgentProperties()->bCanSwim = true;
-		CharacterMovement->SetJumpAllowed(true);
 	}
 
 	Mesh = PCIP.CreateOptionalDefaultSubobject<USkeletalMeshComponent>(this, ACharacter::MeshComponentName);
diff --git a/Engine/Source/Runtime/Engine/Private/CharacterMovementComponent.cpp b/Engine/Source/Runtime/Engine/Private/CharacterMovementComponent.cpp
--- a/Engine/Source/Runtime/Engine/Private/CharacterMovementComponent.cpp
+++ b/Engine/Source/Runtime/Engine/Private/CharacterMovementComponent.cpp
@@ -224,6 +224,7 @@ UCharacterMovementComponent::UCharacterMovementComponent(const class FPostConstr
 	// default character can jump and walk
 	NavAgentProps.bCanJump = true;
 	NavAgentProps.bCanWalk = true;
+	NavAgentProps.bCanSwim = true;
 	ResetMoveState();
 
 	ClientPredictionData = NULL;
