#ue4 - Add GetComponent() accessor to FHitResult and FOverlapResult

diff --git a/Engine/Source/Runtime/Engine/Classes/Engine/EngineTypes.h b/Engine/Source/Runtime/Engine/Classes/Engine/EngineTypes.h
--- a/Engine/Source/Runtime/Engine/Classes/Engine/EngineTypes.h
+++ b/Engine/Source/Runtime/Engine/Classes/Engine/EngineTypes.h
@@ -1394,6 +1394,9 @@ struct ENGINE_API FHitResult
 	/** Utility to return the Actor that owns the Component that was hit */
 	AActor* GetActor() const;
 
+	/** Utility to return the Component that was hit */
+	UPrimitiveComponent* GetComponent() const;
+
 	/** Return true if there was a blocking hit that was not caused by starting in penetration. */
 	FORCEINLINE bool IsValidBlockingHit() const
 	{
@@ -1456,6 +1459,9 @@ struct ENGINE_API FOverlapResult
 	/** Utility to return the Actor that owns the Component that was hit */
 	AActor* GetActor() const;
 
+	/** Utility to return the Component that was hit */
+	UPrimitiveComponent* GetComponent() const;
+
 	/** Indicates if this hit was requesting a block - if false, was requesting a touch instead */
 	UPROPERTY()
 	uint32 bBlockingHit:1;
diff --git a/Engine/Source/Runtime/Engine/Private/Collision/Collision.cpp b/Engine/Source/Runtime/Engine/Private/Collision/Collision.cpp
--- a/Engine/Source/Runtime/Engine/Private/Collision/Collision.cpp
+++ b/Engine/Source/Runtime/Engine/Private/Collision/Collision.cpp
@@ -12,7 +12,17 @@ AActor* FHitResult::GetActor() const
 	return Actor.Get();
 }
 
+UPrimitiveComponent* FHitResult::GetComponent() const
+{
+	return Component.Get();
+}
+
 AActor* FOverlapResult::GetActor() const
 {
 	return Actor.Get();
-}
\ No newline at end of file
+}
+
+UPrimitiveComponent* FOverlapResult::GetComponent() const
+{
+	return Component.Get();
+}
