Don't replay a matinee sound if it is dropped by the audio system

diff --git a/Engine/Source/Runtime/Engine/Classes/Matinee/InterpTrackSound.h b/Engine/Source/Runtime/Engine/Classes/Matinee/InterpTrackSound.h
--- a/Engine/Source/Runtime/Engine/Classes/Matinee/InterpTrackSound.h
+++ b/Engine/Source/Runtime/Engine/Classes/Matinee/InterpTrackSound.h
@@ -67,6 +67,9 @@ class UInterpTrackSound : public UInterpTrackVectorBase
 	UPROPERTY(EditAnywhere, Category=InterpTrackSound)
 	uint32 bAttach:1;
 
+	// True if the sound should have been playing at any point
+	uint32 bPlaying:1;
+
 	// Begin UObject Interface
 	virtual void PostLoad() override;
 	// End UObject Interface
diff --git a/Engine/Source/Runtime/Engine/Private/Interpolation.cpp b/Engine/Source/Runtime/Engine/Private/Interpolation.cpp
--- a/Engine/Source/Runtime/Engine/Private/Interpolation.cpp
+++ b/Engine/Source/Runtime/Engine/Private/Interpolation.cpp
@@ -7865,7 +7865,8 @@ void UInterpTrackSound::UpdateTrack(float NewPosition, UInterpTrackInst* TrInst,
 
 		// Check if we're in the audio range, and if we need to start playing the audio,
 		// either because it has never been played, or isn't currently playing.
-		bool bIsInRangeAndNeedsStart = NewPosition >= SoundTrackKey.Time && NewPosition <= ( SoundTrackKey.Time + SoundTrackKey.Sound->Duration );
+		// We only do this when we've jumped position.
+		bool bIsInRangeAndNeedsStart = !bPlaying && NewPosition >= SoundTrackKey.Time && NewPosition <= ( SoundTrackKey.Time + SoundTrackKey.Sound->Duration );
 		if ( bIsInRangeAndNeedsStart )
 		{
 			bIsInRangeAndNeedsStart = SoundInst->PlayAudioComp == NULL || !SoundInst->PlayAudioComp->IsPlaying();
@@ -7876,6 +7877,8 @@ void UInterpTrackSound::UpdateTrack(float NewPosition, UInterpTrackInst* TrInst,
 		// lets start it.
 		if ( StartSoundIndex != EndSoundIndex || bIsInRangeAndNeedsStart )
 		{
+			bPlaying = true;
+
 			USoundBase* NewSound = SoundTrackKey.Sound;
 
 			APawn* Speaker = NULL;
@@ -8061,6 +8064,7 @@ void UInterpTrackSound::PreviewStopPlayback(class UInterpTrackInst* TrInst)
 	{
 		SoundTrInst->PlayAudioComp->Stop();
 	}
+	bPlaying = false;
 }
 
 /*-----------------------------------------------------------------------------
