#ue4 - CharacterMovementComponent::IsWalkable() should never consider a trace starting in penetration as walkable.

diff --git a/Engine/Source/Runtime/Engine/Private/CharacterMovementComponent.cpp b/Engine/Source/Runtime/Engine/Private/CharacterMovementComponent.cpp
--- a/Engine/Source/Runtime/Engine/Private/CharacterMovementComponent.cpp
+++ b/Engine/Source/Runtime/Engine/Private/CharacterMovementComponent.cpp
@@ -3795,8 +3795,9 @@ void UCharacterMovementComponent::MoveSmooth(const FVector& InVelocity, const fl
 
 bool UCharacterMovementComponent::IsWalkable(const FHitResult& Hit) const
 {
-	if (!Hit.bBlockingHit)
+	if (!Hit.IsValidBlockingHit())
 	{
+		// No hit, or starting in penetration
 		return false;
 	}
 
