Add simple 2D physics constraint

diff --git a/Engine/Source/Runtime/Engine/Classes/PhysicsEngine/BodyInstance.h b/Engine/Source/Runtime/Engine/Classes/PhysicsEngine/BodyInstance.h
--- a/Engine/Source/Runtime/Engine/Classes/PhysicsEngine/BodyInstance.h
+++ b/Engine/Source/Runtime/Engine/Classes/PhysicsEngine/BodyInstance.h
@@ -4,6 +4,7 @@
 
 #include "PhysxUserData.h"
 #include "CollisionQueryParams.h"
+#include "ConstraintInstance.h"
 #include "BodyInstance.generated.h"
 
 
@@ -165,6 +166,12 @@ public:
 	UPROPERTY()
 	uint32 bUpdateMassWhenScaleChanges:1;
 
+	/** Constrain body to 2d plane with DOF as normal. (0,0,0) uses regular 3D */
+	UPROPERTY(EditAnywhere, AdvancedDisplay, Category = Physics)
+	FVector DOF;
+
+	FConstraintInstance DOFConstraint;
+
 protected:
 
 	/** Whether this instance of the object has its own custom walkability override setting. */
diff --git a/Engine/Source/Runtime/Engine/Private/PhysicsEngine/BodyInstance.cpp b/Engine/Source/Runtime/Engine/Private/PhysicsEngine/BodyInstance.cpp
--- a/Engine/Source/Runtime/Engine/Private/PhysicsEngine/BodyInstance.cpp
+++ b/Engine/Source/Runtime/Engine/Private/PhysicsEngine/BodyInstance.cpp
@@ -1096,7 +1096,7 @@ void FBodyInstance::InitBody(UBodySetup* Setup, const FTransform& Transform, UPr
 		}
 
 		// If we ever drive this body kinematically, we want to use its target for scene queries, so collision is updated right away, not on the next physics sim
-		PNewDynamic->setRigidDynamicFlag(PxRigidDynamicFlag::eUSE_KINEMATIC_TARGET_FOR_SCENE_QUERIES, true);
+		PNewDynamic->setRigidDynamicFlag(PxRigidDynamicFlag::eUSE_KINEMATIC_TARGET_FOR_SCENE_QUERIES, true);	
 	}
 
 	// Copy geom from template and scale
@@ -1223,6 +1223,40 @@ void FBodyInstance::InitBody(UBodySetup* Setup, const FTransform& Transform, UPr
 		int32 VelocityIterCount = FMath::Clamp(VelocitySolverIterationCount, 1, 255);
 		PNewDynamic->setSolverIterationCounts(PositionIterCount, VelocityIterCount);
 
+		//setup constraint based on DOF
+		if (!DOF.IsZero())
+		{
+			DOFConstraint.bSwingLimitSoft = false;
+			DOFConstraint.bTwistLimitSoft = false;
+			DOFConstraint.bLinearLimitSoft = false;
+			//set all rotation to free
+			DOFConstraint.AngularSwing1Motion = EAngularConstraintMotion::ACM_Locked;
+			DOFConstraint.AngularSwing2Motion = EAngularConstraintMotion::ACM_Locked;
+			DOFConstraint.AngularTwistMotion  = EAngularConstraintMotion::ACM_Free;
+
+			DOFConstraint.LinearXMotion = ELinearConstraintMotion::LCM_Locked;
+			DOFConstraint.LinearYMotion = ELinearConstraintMotion::LCM_Free;
+			DOFConstraint.LinearZMotion = ELinearConstraintMotion::LCM_Free;
+			
+			FVector Normal = DOF.SafeNormal();
+			FVector Sec;
+			FVector Garbage;
+			Normal.FindBestAxisVectors(Garbage, Sec);
+
+
+			FTransform TM = GetUnrealWorldTransform();
+
+			DOFConstraint.PriAxis1 = TM.InverseTransformVectorNoScale(Normal);
+			DOFConstraint.SecAxis1 = TM.InverseTransformVectorNoScale(Sec);
+
+			DOFConstraint.PriAxis2 = Normal;
+			DOFConstraint.SecAxis2 = Sec;
+			DOFConstraint.Pos2 = TM.GetLocation();
+
+			// Create constraint instance based on DOF
+			DOFConstraint.InitConstraint(OwnerComponent.Get(), this, nullptr, 1.f);
+		}
+
 		// wakeUp and putToSleep will issue warnings on kinematic actors
 		if (IsRigidDynamicNonKinematic(PNewDynamic))
 		{
@@ -2427,6 +2461,7 @@ void FBodyInstance::UpdateMassProperties()
 
 		float MassRatio = NewMass/OldMass;
 		PxVec3 InertiaTensor = PRigidDynamic->getMassSpaceInertiaTensor();
+
 		PRigidDynamic->setMassSpaceInertiaTensor(InertiaTensor * MassRatio);
 		PRigidDynamic->setMass(NewMass);
 
