ragdoll root body inherits DOF from skeletal mesh

diff --git a/Engine/Source/Runtime/Engine/Private/SkeletalMeshComponentPhysics.cpp b/Engine/Source/Runtime/Engine/Private/SkeletalMeshComponentPhysics.cpp
--- a/Engine/Source/Runtime/Engine/Private/SkeletalMeshComponentPhysics.cpp
+++ b/Engine/Source/Runtime/Engine/Private/SkeletalMeshComponentPhysics.cpp
@@ -823,6 +823,12 @@ void USkeletalMeshComponent::InitArticulated(FPhysScene* PhysScene)
 			// this is not true for all other BodyInstance, but for physics assets it is true. 
 			BodyInst->bSimulatePhysics = false;
 			BodyInst->InstanceBodyIndex = i; // Set body index 
+
+			if (i == RootBodyIndex)
+			{
+				BodyInst->DOF = BodyInstance.DOF;
+			}
+
 #if WITH_PHYSX
 			// Create physics body instance.
 			FTransform BoneTransform = GetBoneTransform( BoneIndex );
