UE4: Fixed a crash when saving a world asset thumbnail during autosave

diff --git a/Engine/Source/Runtime/Engine/Private/LevelActor.cpp b/Engine/Source/Runtime/Engine/Private/LevelActor.cpp
--- a/Engine/Source/Runtime/Engine/Private/LevelActor.cpp
+++ b/Engine/Source/Runtime/Engine/Private/LevelActor.cpp
@@ -972,7 +972,8 @@ void UWorld::SetMapNeedsLightingFullyRebuilt(int32 InNumLightingUnbuiltObjects)
 	static const auto AllowStaticLightingVar = IConsoleManager::Get().FindTConsoleVariableDataInt(TEXT("r.AllowStaticLighting"));
 	const bool bAllowStaticLighting = (!AllowStaticLightingVar || AllowStaticLightingVar->GetValueOnGameThread() != 0);
 
-	if (bAllowStaticLighting && !GetWorldSettings()->bForceNoPrecomputedLighting)
+	UWorld* WorldSettings = GetWorldSettings();
+	if (bAllowStaticLighting && WorldSettings && !WorldSettings->bForceNoPrecomputedLighting)
 	{
 		check(IsInGameThread());
 		if (NumLightingUnbuiltObjects != InNumLightingUnbuiltObjects && (NumLightingUnbuiltObjects == 0 || InNumLightingUnbuiltObjects == 0))
