UE4: Fix typo

diff --git a/Engine/Source/Runtime/Engine/Private/LevelActor.cpp b/Engine/Source/Runtime/Engine/Private/LevelActor.cpp
--- a/Engine/Source/Runtime/Engine/Private/LevelActor.cpp
+++ b/Engine/Source/Runtime/Engine/Private/LevelActor.cpp
@@ -972,7 +972,7 @@ void UWorld::SetMapNeedsLightingFullyRebuilt(int32 InNumLightingUnbuiltObjects)
 	static const auto AllowStaticLightingVar = IConsoleManager::Get().FindTConsoleVariableDataInt(TEXT("r.AllowStaticLighting"));
 	const bool bAllowStaticLighting = (!AllowStaticLightingVar || AllowStaticLightingVar->GetValueOnGameThread() != 0);
 
-	UWorld* WorldSettings = GetWorldSettings();
+	AWorldSettings* WorldSettings = GetWorldSettings();
 	if (bAllowStaticLighting && WorldSettings && !WorldSettings->bForceNoPrecomputedLighting)
 	{
 		check(IsInGameThread());
