Rendering: Make UMaterialExpression and several derived types extendable from outside of Engine
- Add ENGINE_API on several classes that are desirable to derive from
- Started to decouple graph node renaming from being a parameter (added CanRenameNode, GetEditableName, and SetEditableName)
(Integration of CL# 2124083 from 4.3 to main)

diff --git a/Engine/Source/Editor/MaterialEditor/Private/MaterialEditor.cpp b/Engine/Source/Editor/MaterialEditor/Private/MaterialEditor.cpp
--- a/Engine/Source/Editor/MaterialEditor/Private/MaterialEditor.cpp
+++ b/Engine/Source/Editor/MaterialEditor/Private/MaterialEditor.cpp
@@ -2560,7 +2560,10 @@ UMaterialExpression* FMaterialEditor::CreateNewMaterialExpression(UClass* NewExp
 			if( METextureBase )
 			{
 				FEditorDelegates::LoadSelectedAssetsIfNeeded.Broadcast();
-				METextureBase->Texture = GEditor->GetSelectedObjects()->GetTop<UTexture>();
+				if( UTexture* SelectedTexture = GEditor->GetSelectedObjects()->GetTop<UTexture>() )
+				{
+					METextureBase->Texture = SelectedTexture;
+				}
 				METextureBase->AutoSetSampleType();
 			}
 
@@ -2596,7 +2599,7 @@ UMaterialExpression* FMaterialEditor::CreateNewMaterialExpression(UClass* NewExp
 		NewExpression->UpdateParameterGuid(true, true);
 
 		UMaterialExpressionTextureSampleParameter* TextureParameterExpression = Cast<UMaterialExpressionTextureSampleParameter>( NewExpression );
-		if( TextureParameterExpression )
+		if( (TextureParameterExpression != nullptr) && TextureParameterExpression->CanRenameNode() )
 		{
 			// Change the parameter's name on creation to mirror the object's name; this avoids issues of having colliding parameter
 			// names and having the name left as "None"
diff --git a/Engine/Source/Editor/UnrealEd/Private/MaterialGraphNode.cpp b/Engine/Source/Editor/UnrealEd/Private/MaterialGraphNode.cpp
--- a/Engine/Source/Editor/UnrealEd/Private/MaterialGraphNode.cpp
+++ b/Engine/Source/Editor/UnrealEd/Private/MaterialGraphNode.cpp
@@ -580,7 +580,7 @@ void UMaterialGraphNode::PostPlacedNewNode()
 		NodeComment = MaterialExpression->Desc;
 		NodePosX = MaterialExpression->MaterialExpressionEditorX;
 		NodePosY = MaterialExpression->MaterialExpressionEditorY;
-		bCanRenameNode = UMaterial::IsParameter(MaterialExpression);
+		bCanRenameNode = MaterialExpression->CanRenameNode();
 	}
 }
 
@@ -639,38 +639,14 @@ void UMaterialGraphNode::GetPinHoverText(const UEdGraphPin& Pin, FString& HoverT
 
 FString UMaterialGraphNode::GetParameterName() const
 {
-	if (const UMaterialExpressionParameter* Parameter = Cast<const UMaterialExpressionParameter>(MaterialExpression))
-	{
-		return Parameter->ParameterName.ToString();
-	}
-	else if (const UMaterialExpressionTextureSampleParameter* TexParameter = Cast<const UMaterialExpressionTextureSampleParameter>(MaterialExpression))
-	{
-		return TexParameter->ParameterName.ToString();
-	}
-	else if (const UMaterialExpressionFontSampleParameter* FontParameter = Cast<const UMaterialExpressionFontSampleParameter>(MaterialExpression))
-	{
-		return FontParameter->ParameterName.ToString();
-	}
-	// Should have been able to get parameter name to edit
-	check(false);
-	return TEXT("");
+	return MaterialExpression->GetEditableName();
 }
 
 void UMaterialGraphNode::SetParameterName(const FString& NewName)
 {
-	if (UMaterialExpressionParameter* Parameter = Cast<UMaterialExpressionParameter>(MaterialExpression))
-	{
-		Parameter->ParameterName = *NewName;
-	}
-	else if (UMaterialExpressionTextureSampleParameter* TexParameter = Cast<UMaterialExpressionTextureSampleParameter>(MaterialExpression))
-	{
-		TexParameter->ParameterName = *NewName;
-	}
-	else if (UMaterialExpressionFontSampleParameter* FontParameter = Cast<UMaterialExpressionFontSampleParameter>(MaterialExpression))
-	{
-		FontParameter->ParameterName = *NewName;
-	}
+	MaterialExpression->SetEditableName(NewName);
 
+	//@TODO: Push into the SetEditableName interface
 	CastChecked<UMaterialGraph>(GetGraph())->Material->UpdateExpressionParameterName(MaterialExpression);
 }
 
diff --git a/Engine/Source/Runtime/Engine/Classes/Materials/Material.h b/Engine/Source/Runtime/Engine/Classes/Materials/Material.h
--- a/Engine/Source/Runtime/Engine/Classes/Materials/Material.h
+++ b/Engine/Source/Runtime/Engine/Classes/Materials/Material.h
@@ -1055,7 +1055,7 @@ public:
 	 *
 	 * @param	Expression	The expression node to inspect.
 	 */
-	ENGINE_API static bool IsParameter(UMaterialExpression* Expression);
+	ENGINE_API static bool IsParameter(const UMaterialExpression* Expression);
 
 	/**
 	 * Return whether the provided expression node is a dynamic parameter.
diff --git a/Engine/Source/Runtime/Engine/Classes/Materials/MaterialExpression.h b/Engine/Source/Runtime/Engine/Classes/Materials/MaterialExpression.h
--- a/Engine/Source/Runtime/Engine/Classes/Materials/MaterialExpression.h
+++ b/Engine/Source/Runtime/Engine/Classes/Materials/MaterialExpression.h
@@ -77,8 +77,8 @@ struct FExpressionOutput
 };
 #endif
 
-UCLASS(abstract, hidecategories=Object, MinimalAPI)
-class UMaterialExpression : public UObject
+UCLASS(abstract, hidecategories=Object)
+class ENGINE_API UMaterialExpression : public UObject
 {
 	GENERATED_UCLASS_BODY()
 
@@ -261,7 +261,7 @@ class UMaterialExpression : public UObject
 	 * Copy the SrcExpressions into the specified material, preserving internal references.
 	 * New material expressions are created within the specified material.
 	 */
-	ENGINE_API static void CopyMaterialExpressions(const TArray<class UMaterialExpression*>& SrcExpressions, const TArray<class UMaterialExpressionComment*>& SrcExpressionComments, 
+	static void CopyMaterialExpressions(const TArray<class UMaterialExpression*>& SrcExpressions, const TArray<class UMaterialExpressionComment*>& SrcExpressionComments, 
 		class UMaterial* Material, class UMaterialFunction* Function, TArray<class UMaterialExpression*>& OutNewExpressions, TArray<class UMaterialExpression*>& OutNewComments);
 
 	/**
@@ -272,22 +272,22 @@ class UMaterialExpression : public UObject
 	/**
 	 * Connects the specified output to the passed material for previewing. 
 	 */
-	ENGINE_API void ConnectToPreviewMaterial(UMaterial* Material, int32 OutputIndex);
+	void ConnectToPreviewMaterial(UMaterial* Material, int32 OutputIndex);
 
 	/**
 	 * Connects the specified input expression to the specified output of this expression.
 	 */
-	ENGINE_API void ConnectExpression(FExpressionInput* Input, int32 OutputIndex);
+	void ConnectExpression(FExpressionInput* Input, int32 OutputIndex);
 
 	/** 
 	* Generates a GUID for this expression if one doesn't already exist. 
 	*
 	* @param bForceGeneration	Whether we should generate a GUID even if it is already valid.
 	*/
-	ENGINE_API void UpdateParameterGuid(bool bForceGeneration, bool bAllowMarkingPackageDirty);
+	void UpdateParameterGuid(bool bForceGeneration, bool bAllowMarkingPackageDirty);
 
 	/** Callback to access derived classes' parameter expression id. */
-	ENGINE_API virtual FGuid& GetParameterExpressionId()
+	virtual FGuid& GetParameterExpressionId()
 	{
 		checkf(!bIsParameterExpression, TEXT("Expressions with bIsParameterExpression==true must implement their own GetParameterExpressionId!"));
 		static FGuid Dummy;
@@ -295,7 +295,7 @@ class UMaterialExpression : public UObject
 	}
 
 	/** Asserts if the expression is not contained by its Material or Function's expressions array. */
-	ENGINE_API void ValidateState();
+	void ValidateState();
 
 #if WITH_EDITOR
 	/** Returns the keywords that should be used when searching for this expression */
@@ -309,12 +309,29 @@ class UMaterialExpression : public UObject
 	 *
 	 * @return Whether a repeat was found while getting expressions
 	 */
-	ENGINE_API bool GetAllInputExpressions(TArray<UMaterialExpression*>& InputExpressions);
+	bool GetAllInputExpressions(TArray<UMaterialExpression*>& InputExpressions);
+
+	/**
+	 * Can this node be renamed?
+	 */
+	virtual bool CanRenameNode() const;
+
+	/**
+	 * Returns the current 'name' of the node (typically a parameter name).
+	 * Only valid to call on a node that previously returned CanRenameNode() = true.
+	 */
+	virtual FString GetEditableName() const;
+
+	/**
+	 * Sets the current 'name' of the node (typically a parameter name)
+	 * Only valid to call on a node that previously returned CanRenameNode() = true.
+	 */
+	virtual void SetEditableName(const FString& NewName);
 
 #endif // WITH_EDITOR
 
 	/** Checks whether any inputs to this expression create a loop */
-	ENGINE_API bool ContainsInputLoop();
+	bool ContainsInputLoop();
 
 protected:
 	/**
diff --git a/Engine/Source/Runtime/Engine/Classes/Materials/MaterialExpressionFontSampleParameter.h b/Engine/Source/Runtime/Engine/Classes/Materials/MaterialExpressionFontSampleParameter.h
--- a/Engine/Source/Runtime/Engine/Classes/Materials/MaterialExpressionFontSampleParameter.h
+++ b/Engine/Source/Runtime/Engine/Classes/Materials/MaterialExpressionFontSampleParameter.h
@@ -26,6 +26,11 @@ class UMaterialExpressionFontSampleParameter : public UMaterialExpressionFontSam
 	virtual int32 Compile(class FMaterialCompiler* Compiler, int32 OutputIndex, int32 MultiplexIndex) override;
 	virtual void GetCaption(TArray<FString>& OutCaptions) const override;
 	virtual bool MatchesSearchQuery( const TCHAR* SearchQuery ) override;
+#if WITH_EDITOR
+	virtual bool CanRenameNode() const override { return true; }
+	virtual FString GetEditableName() const override;
+	virtual void SetEditableName(const FString& NewName) override;
+#endif
 	// End UMaterialExpression Interface
 	
 	/**
diff --git a/Engine/Source/Runtime/Engine/Classes/Materials/MaterialExpressionParameter.h b/Engine/Source/Runtime/Engine/Classes/Materials/MaterialExpressionParameter.h
--- a/Engine/Source/Runtime/Engine/Classes/Materials/MaterialExpressionParameter.h
+++ b/Engine/Source/Runtime/Engine/Classes/Materials/MaterialExpressionParameter.h
@@ -24,6 +24,11 @@ class UMaterialExpressionParameter : public UMaterialExpression
 
 	// Begin UMaterialExpression Interface
 	virtual bool MatchesSearchQuery( const TCHAR* SearchQuery ) override;
+#if WITH_EDITOR
+	virtual bool CanRenameNode() const override { return true; }
+	virtual FString GetEditableName() const override;
+	virtual void SetEditableName(const FString& NewName) override;
+#endif
 	// End UMaterialExpression Interface
 
 	ENGINE_API virtual FGuid& GetParameterExpressionId() override
diff --git a/Engine/Source/Runtime/Engine/Classes/Materials/MaterialExpressionTextureBase.h b/Engine/Source/Runtime/Engine/Classes/Materials/MaterialExpressionTextureBase.h
--- a/Engine/Source/Runtime/Engine/Classes/Materials/MaterialExpressionTextureBase.h
+++ b/Engine/Source/Runtime/Engine/Classes/Materials/MaterialExpressionTextureBase.h
@@ -9,8 +9,8 @@
 #include "Materials/MaterialExpression.h"
 #include "MaterialExpressionTextureBase.generated.h"
 
-UCLASS(abstract, hidecategories=Object, MinimalAPI)
-class UMaterialExpressionTextureBase : public UMaterialExpression 
+UCLASS(abstract, hidecategories=Object)
+class ENGINE_API UMaterialExpressionTextureBase : public UMaterialExpression 
 {
 	GENERATED_UCLASS_BODY()
 
@@ -49,12 +49,12 @@ class UMaterialExpressionTextureBase : public UMaterialExpression
 	/**
 	 * Automatically determines and set the sampler type for the current texture.
 	 */
-	ENGINE_API void AutoSetSampleType();
+	void AutoSetSampleType();
 
 	/**
 	 * Returns the default sampler type for the specified texture.
 	 * @param Texture - The texture for which the default sampler type will be returned.
 	 * @returns the default sampler type for the specified texture.
 	 */
-	ENGINE_API static EMaterialSamplerType GetSamplerTypeForTexture( const UTexture* Texture );
+	static EMaterialSamplerType GetSamplerTypeForTexture( const UTexture* Texture );
 };
diff --git a/Engine/Source/Runtime/Engine/Classes/Materials/MaterialExpressionTextureSample.h b/Engine/Source/Runtime/Engine/Classes/Materials/MaterialExpressionTextureSample.h
--- a/Engine/Source/Runtime/Engine/Classes/Materials/MaterialExpressionTextureSample.h
+++ b/Engine/Source/Runtime/Engine/Classes/Materials/MaterialExpressionTextureSample.h
@@ -16,8 +16,8 @@ enum ETextureMipValueMode
 	TMVM_MAX,
 };
 
-UCLASS(collapsecategories, hidecategories=Object, MinimalAPI)
-class UMaterialExpressionTextureSample : public UMaterialExpressionTextureBase
+UCLASS(collapsecategories, hidecategories=Object)
+class ENGINE_API UMaterialExpressionTextureSample : public UMaterialExpressionTextureBase
 {
 	GENERATED_UCLASS_BODY()
 
diff --git a/Engine/Source/Runtime/Engine/Classes/Materials/MaterialExpressionTextureSampleParameter.h b/Engine/Source/Runtime/Engine/Classes/Materials/MaterialExpressionTextureSampleParameter.h
--- a/Engine/Source/Runtime/Engine/Classes/Materials/MaterialExpressionTextureSampleParameter.h
+++ b/Engine/Source/Runtime/Engine/Classes/Materials/MaterialExpressionTextureSampleParameter.h
@@ -5,8 +5,8 @@
 #include "Materials/MaterialExpressionTextureSample.h"
 #include "MaterialExpressionTextureSampleParameter.generated.h"
 
-UCLASS(collapsecategories, abstract, hidecategories=Object, MinimalAPI)
-class UMaterialExpressionTextureSampleParameter : public UMaterialExpressionTextureSample
+UCLASS(collapsecategories, abstract, hidecategories=Object)
+class ENGINE_API UMaterialExpressionTextureSampleParameter : public UMaterialExpressionTextureSample
 {
 	GENERATED_UCLASS_BODY()
 
@@ -24,6 +24,12 @@ class UMaterialExpressionTextureSampleParameter : public UMaterialExpressionText
 	// Begin UMaterialExpression Interface
 	virtual int32 Compile(class FMaterialCompiler* Compiler, int32 OutputIndex, int32 MultiplexIndex) override;
 	virtual void GetCaption(TArray<FString>& OutCaptions) const override;
+	virtual bool MatchesSearchQuery(const TCHAR* SearchQuery) override;
+#if WITH_EDITOR
+	virtual bool CanRenameNode() const override { return true; }
+	virtual FString GetEditableName() const override;
+	virtual void SetEditableName(const FString& NewName) override;
+#endif
 	// End UMaterialExpression Interface
 
 	/**
@@ -46,14 +52,10 @@ class UMaterialExpressionTextureSampleParameter : public UMaterialExpressionText
 	 */
 	virtual void SetDefaultTexture();
 
-	ENGINE_API virtual FGuid& GetParameterExpressionId() override
+	virtual FGuid& GetParameterExpressionId() override
 	{
 		return ExpressionGUID;
 	}
 
-	virtual bool MatchesSearchQuery( const TCHAR* SearchQuery );
 	void GetAllParameterNames(TArray<FName> &OutParameterNames, TArray<FGuid> &OutParameterIds);
 };
-
-
-
diff --git a/Engine/Source/Runtime/Engine/Classes/Materials/MaterialExpressionTextureSampleParameter2D.h b/Engine/Source/Runtime/Engine/Classes/Materials/MaterialExpressionTextureSampleParameter2D.h
--- a/Engine/Source/Runtime/Engine/Classes/Materials/MaterialExpressionTextureSampleParameter2D.h
+++ b/Engine/Source/Runtime/Engine/Classes/Materials/MaterialExpressionTextureSampleParameter2D.h
@@ -5,8 +5,8 @@
 #include "Materials/MaterialExpressionTextureSampleParameter.h"
 #include "MaterialExpressionTextureSampleParameter2D.generated.h"
 
-UCLASS(collapsecategories, hidecategories=Object, MinimalAPI)
-class UMaterialExpressionTextureSampleParameter2D : public UMaterialExpressionTextureSampleParameter
+UCLASS(collapsecategories, hidecategories=Object)
+class ENGINE_API UMaterialExpressionTextureSampleParameter2D : public UMaterialExpressionTextureSampleParameter
 {
 	GENERATED_UCLASS_BODY()
 
diff --git a/Engine/Source/Runtime/Engine/Private/Materials/Material.cpp b/Engine/Source/Runtime/Engine/Private/Materials/Material.cpp
--- a/Engine/Source/Runtime/Engine/Private/Materials/Material.cpp
+++ b/Engine/Source/Runtime/Engine/Private/Materials/Material.cpp
@@ -2350,7 +2350,7 @@ bool UMaterial::RemoveExpressionParameter(UMaterialExpression* Expression)
 }
 
 
-bool UMaterial::IsParameter(UMaterialExpression* Expression)
+bool UMaterial::IsParameter(const UMaterialExpression* Expression)
 {
 	bool bRet = false;
 
diff --git a/Engine/Source/Runtime/Engine/Private/Materials/MaterialExpressions.cpp b/Engine/Source/Runtime/Engine/Private/Materials/MaterialExpressions.cpp
--- a/Engine/Source/Runtime/Engine/Private/Materials/MaterialExpressions.cpp
+++ b/Engine/Source/Runtime/Engine/Private/Materials/MaterialExpressions.cpp
@@ -949,6 +949,25 @@ bool UMaterialExpression::GetAllInputExpressions(TArray<UMaterialExpression*>& I
 		return true;
 	}
 }
+
+bool UMaterialExpression::CanRenameNode() const
+{
+	return false;
+}
+
+FString UMaterialExpression::GetEditableName() const
+{
+	// This function is only safe to call in a class that has implemented CanRenameNode() to return true
+	check(false);
+	return TEXT("");
+}
+
+void UMaterialExpression::SetEditableName(const FString& NewName)
+{
+	// This function is only safe to call in a class that has implemented CanRenameNode() to return true
+	check(false);
+}
+
 #endif
 
 bool UMaterialExpression::ContainsInputLoop()
@@ -1718,6 +1737,19 @@ bool UMaterialExpressionTextureSampleParameter::MatchesSearchQuery( const TCHAR*
 	return Super::MatchesSearchQuery(SearchQuery);
 }
 
+#if WITH_EDITOR
+FString UMaterialExpressionTextureSampleParameter::GetEditableName() const
+{
+	return ParameterName.ToString();
+}
+
+void UMaterialExpressionTextureSampleParameter::SetEditableName(const FString& NewName)
+{
+	ParameterName = *NewName;
+}
+#endif
+
+
 //
 //  UMaterialExpressionTextureSampleParameterCube
 //
@@ -3535,6 +3567,18 @@ bool UMaterialExpressionParameter::MatchesSearchQuery( const TCHAR* SearchQuery
 	return Super::MatchesSearchQuery(SearchQuery);
 }
 
+#if WITH_EDITOR
+FString UMaterialExpressionParameter::GetEditableName() const
+{
+	return ParameterName.ToString();
+}
+
+void UMaterialExpressionParameter::SetEditableName(const FString& NewName)
+{
+	ParameterName = *NewName;
+}
+#endif
+
 void UMaterialExpressionParameter::GetAllParameterNames(TArray<FName> &OutParameterNames, TArray<FGuid> &OutParameterIds)
 {
 	int32 CurrentSize = OutParameterNames.Num();
@@ -5487,6 +5531,18 @@ bool UMaterialExpressionFontSampleParameter::MatchesSearchQuery( const TCHAR* Se
 	return Super::MatchesSearchQuery(SearchQuery);
 }
 
+#if WITH_EDITOR
+FString UMaterialExpressionFontSampleParameter::GetEditableName() const
+{
+	return ParameterName.ToString();
+}
+
+void UMaterialExpressionFontSampleParameter::SetEditableName(const FString& NewName)
+{
+	ParameterName = *NewName;
+}
+#endif
+
 void UMaterialExpressionFontSampleParameter::GetAllParameterNames(TArray<FName> &OutParameterNames, TArray<FGuid> &OutParameterIds)
 {
 	int32 CurrentSize = OutParameterNames.Num();
