StopActiveMovement will now affect path following

diff --git a/Engine/Config/BaseEngine.ini b/Engine/Config/BaseEngine.ini
--- a/Engine/Config/BaseEngine.ini
+++ b/Engine/Config/BaseEngine.ini
@@ -790,6 +790,7 @@ bUseStreamingPause=false
 +K2FieldRedirects=(OldFieldName="Character.Launch",NewFieldName="Character.LaunchCharacter")
 +K2FieldRedirects=(OldFieldName="Actor.SetActorLocation",NewFieldName="Actor.K2_SetActorLocation")
 +K2FieldRedirects=(OldFieldName="ProjectileMovementComponent.StopMovement",NewFieldName="ProjectileMovementComponent.StopSimulating")
++K2FieldRedirects=(OldFieldName="PawnMovementComponent.StopActiveMovement",NewFieldName="NavMovementComponent.StopActiveMovement")
 
 +K2FieldRedirects=(OldFieldName="Controller.K2_GetActorRotation",NewFieldName="Controller.GetControlRotation")
 +K2FieldRedirects=(OldFieldName="PlayerController.AddTurnInput",NewFieldName="PlayerController.AddYawInput")
diff --git a/Engine/Source/Runtime/AIModule/Private/Navigation/PathFollowingComponent.cpp b/Engine/Source/Runtime/AIModule/Private/Navigation/PathFollowingComponent.cpp
--- a/Engine/Source/Runtime/AIModule/Private/Navigation/PathFollowingComponent.cpp
+++ b/Engine/Source/Runtime/AIModule/Private/Navigation/PathFollowingComponent.cpp
@@ -259,7 +259,7 @@ void UPathFollowingComponent::AbortMove(const FString& Reason, FAIRequestID Requ
 
 		if (bResetVelocity && MovementComp && MovementComp->CanStopPathFollowing())
 		{
-			MovementComp->StopMovementImmediately();
+			MovementComp->StopMovementKeepPathing();
 		}
 
 		if (NavComp && FinishResult != EPathFollowingResult::Skipped)
@@ -276,11 +276,6 @@ void UPathFollowingComponent::AbortMove(const FString& Reason, FAIRequestID Requ
 
 		FAIMessage::Send(Cast<AController>(GetOwner()), Msg);
 	}
-	else
-	{
-		UE_VLOG(GetOwner(), LogPathFollowing, Log, TEXT("AbortMove FAILED due to RequestID(%u) != CurrentRequestId()%d")
-			, RequestID, CurrentRequestId);
-	}
 }
 
 void UPathFollowingComponent::PauseMove(FAIRequestID RequestID, bool bResetVelocity)
@@ -295,7 +290,7 @@ void UPathFollowingComponent::PauseMove(FAIRequestID RequestID, bool bResetVeloc
 	{
 		if (bResetVelocity && MovementComp && MovementComp->CanStopPathFollowing())
 		{
-			MovementComp->StopMovementImmediately();
+			MovementComp->StopMovementKeepPathing();
 		}
 
 		LocationWhenPaused = MovementComp ? MovementComp->GetActorFeetLocation() : FVector::ZeroVector;
@@ -381,7 +376,7 @@ void UPathFollowingComponent::OnPathFinished(EPathFollowingResult::Type Result)
 
 	if (MovementComp && MovementComp->CanStopPathFollowing())	
 	{
-		MovementComp->StopMovementImmediately();
+		MovementComp->StopMovementKeepPathing();
 	}
 
 	if (NavComp && Result != EPathFollowingResult::Skipped)
diff --git a/Engine/Source/Runtime/Engine/Classes/GameFramework/NavMovementComponent.h b/Engine/Source/Runtime/Engine/Classes/GameFramework/NavMovementComponent.h
--- a/Engine/Source/Runtime/Engine/Classes/GameFramework/NavMovementComponent.h
+++ b/Engine/Source/Runtime/Engine/Classes/GameFramework/NavMovementComponent.h
@@ -28,6 +28,10 @@ protected:
 	UPROPERTY(EditAnywhere, Category=MovementComponent)
 	uint32 bUpdateNavAgentWithOwnersCollision:1;
 
+private:
+	/** If set, StopActiveMovement call will abort current path following request */
+	uint32 bStopMovementAbortPaths:1;
+
 public:
 	/** Expresses runtime state of character's movement. Put all temporal changes to movement properties here */
 	UPROPERTY()
@@ -36,6 +40,17 @@ public:
 	/** associated path following component */
 	TWeakObjectPtr<class UPathFollowingComponent> PathFollowingComp;
 
+	/** Stops applying further movement (usually zeros acceleration). */
+	UFUNCTION(BlueprintCallable, Category="Pawn|Components|PawnMovement")
+	virtual void StopActiveMovement();
+
+	/** Stops movement immediately (reset velocity) but keeps following current path */
+	UFUNCTION(BlueprintCallable, Category="Components|Movement")
+	void StopMovementKeepPathing();
+
+	// Overridden to also call StopActiveMovement().
+	virtual void StopMovementImmediately() override;
+
 	FORCEINLINE bool ShouldUpdateNavAgentWithOwnersCollision() const { return bUpdateNavAgentWithOwnersCollision != 0; }
 	void UpdateNavAgent(class AActor* Owner);
 	void UpdateNavAgent(class UCapsuleComponent* CapsuleComponent);
@@ -134,3 +149,15 @@ inline bool UNavMovementComponent::IsFlying() const
 	return false;
 }
 
+inline void UNavMovementComponent::StopMovementKeepPathing()
+{
+	bStopMovementAbortPaths = false;
+	StopMovementImmediately();
+	bStopMovementAbortPaths = true;
+}
+
+inline void UNavMovementComponent::StopMovementImmediately()
+{
+	Super::StopMovementImmediately();
+	StopActiveMovement();
+}
diff --git a/Engine/Source/Runtime/Engine/Classes/GameFramework/PawnMovementComponent.h b/Engine/Source/Runtime/Engine/Classes/GameFramework/PawnMovementComponent.h
--- a/Engine/Source/Runtime/Engine/Classes/GameFramework/PawnMovementComponent.h
+++ b/Engine/Source/Runtime/Engine/Classes/GameFramework/PawnMovementComponent.h
@@ -16,13 +16,6 @@ class ENGINE_API UPawnMovementComponent : public UNavMovementComponent
 	// Overridden to only allow registration with components owned by a Pawn.
 	virtual void SetUpdatedComponent(class UPrimitiveComponent* NewUpdatedComponent) override;
 
-	// Overridden to also call StopActiveMovement().
-	virtual void StopMovementImmediately() override;
-
-	/** Stops applying further movement (usually zeros acceleration). */
-	UFUNCTION(BlueprintCallable, Category="Pawn|Components|PawnMovement")
-	virtual void StopActiveMovement();
-
 	/** Adds the given vector to the accumulated input in world space. Input vectors are usually between 0 and 1 in magnitude. 
 	 * They are accumulated during a frame then applied as acceleration during the movement update.
 	 * if bForce is true, it bypasses the IsMoveInputIgnored() check. */
@@ -62,16 +55,6 @@ protected:
 //////////////////////////////////////////////////////////////////////////
 // Inlines
 
-inline void UPawnMovementComponent::StopMovementImmediately()
-{
-	Super::StopMovementImmediately();
-	StopActiveMovement();
-}
-
-inline void UPawnMovementComponent::StopActiveMovement()
-{
-}
-
 inline FVector UPawnMovementComponent::GetInputVector() const
 {
 	return ControlInputVector;
diff --git a/Engine/Source/Runtime/Engine/Private/CharacterMovementComponent.cpp b/Engine/Source/Runtime/Engine/Private/CharacterMovementComponent.cpp
--- a/Engine/Source/Runtime/Engine/Private/CharacterMovementComponent.cpp
+++ b/Engine/Source/Runtime/Engine/Private/CharacterMovementComponent.cpp
@@ -633,7 +633,7 @@ void UCharacterMovementComponent::OnMovementModeChanged(EMovementMode PreviousMo
 		if (MovementMode == MOVE_None)
 		{
 			// Kill velocity and clear queued up events
-			StopMovementImmediately();
+			StopMovementKeepPathing();
 			CharacterOwner->ClearJumpInput();
 		}
 	}
@@ -3424,6 +3424,8 @@ void UCharacterMovementComponent::AdjustFloorHeight()
 
 void UCharacterMovementComponent::StopActiveMovement() 
 { 
+	Super::StopActiveMovement();
+
 	Acceleration = FVector::ZeroVector; 
 	bHasRequestedVelocity = false;
 	RequestedVelocity = FVector::ZeroVector;
diff --git a/Engine/Source/Runtime/Engine/Private/NavMovementComponent.cpp b/Engine/Source/Runtime/Engine/Private/NavMovementComponent.cpp
--- a/Engine/Source/Runtime/Engine/Private/NavMovementComponent.cpp
+++ b/Engine/Source/Runtime/Engine/Private/NavMovementComponent.cpp
@@ -1,6 +1,7 @@
 // Copyright 1998-2014 Epic Games, Inc. All Rights Reserved.
 
 #include "EnginePrivate.h"
+#include "Navigation/PathFollowingComponent.h"
 
 //----------------------------------------------------------------------//
 // FNavAgentProperties
@@ -17,6 +18,7 @@ void FNavAgentProperties::UpdateWithCollisionComponent(UShapeComponent* Collisio
 UNavMovementComponent::UNavMovementComponent(const class FPostConstructInitializeProperties& PCIP)
 	: Super(PCIP)
 	, bUpdateNavAgentWithOwnersCollision(true)
+	, bStopMovementAbortPaths(true)
 {
 }
 
@@ -35,6 +37,14 @@ bool UNavMovementComponent::CanStopPathFollowing() const
 	return true;
 }
 
+void UNavMovementComponent::StopActiveMovement()
+{
+	if (PathFollowingComp.IsValid() && bStopMovementAbortPaths)
+	{
+		PathFollowingComp->AbortMove("StopActiveMovement");
+	}
+}
+
 void UNavMovementComponent::UpdateNavAgent(AActor* Owner)
 {
 	ensure(Owner == NULL || Owner == GetOwner());
