Exposed UMaterialBillboardComponent via ENGINE_API

#ttp 337490 - Expose UMaterialBillboardComponent via ENGINE_API

diff --git a/Engine/Source/Runtime/Engine/Classes/Components/MaterialBillboardComponent.h b/Engine/Source/Runtime/Engine/Classes/Components/MaterialBillboardComponent.h
--- a/Engine/Source/Runtime/Engine/Classes/Components/MaterialBillboardComponent.h
+++ b/Engine/Source/Runtime/Engine/Classes/Components/MaterialBillboardComponent.h
@@ -51,7 +51,7 @@ struct FMaterialSpriteElement
  * A 2d material that will be rendered always facing the camera.
  */
 UCLASS(ClassGroup=Rendering, collapsecategories, hidecategories=Object, editinlinenew, meta=(BlueprintSpawnableComponent))
-class UMaterialBillboardComponent : public UPrimitiveComponent
+class ENGINE_API UMaterialBillboardComponent : public UPrimitiveComponent
 {
 	GENERATED_UCLASS_BODY()
 
