Exposed decal sort order value to blueprints.

#ttp 328998 - UE4: RENDERING: Decals: Need to be able to set sort order using blueprints on decals.

diff --git a/Engine/Source/Runtime/Engine/Classes/Components/DecalComponent.h b/Engine/Source/Runtime/Engine/Classes/Components/DecalComponent.h
--- a/Engine/Source/Runtime/Engine/Classes/Components/DecalComponent.h
+++ b/Engine/Source/Runtime/Engine/Classes/Components/DecalComponent.h
@@ -26,6 +26,10 @@ class UDecalComponent : public USceneComponent
 	UPROPERTY(EditAnywhere, BlueprintReadOnly, Category=Decal)
 	int32 SortOrder;
 
+	/** Sets the sort order for the decal component. Higher values draw later (on top). This will force the decal to reattach */
+	UFUNCTION(BlueprintCallable, Category = "Rendering|Components|Decal")
+	ENGINE_API void SetSortOrder(int32 Value);
+
 	/** setting decal material on decal component. This will force the decal to reattach */
 	UFUNCTION(BlueprintCallable, Category="Rendering|Components|Decal")
 	ENGINE_API void SetDecalMaterial(class UMaterialInterface* NewDecalMaterial);
diff --git a/Engine/Source/Runtime/Engine/Private/DecalComponent.cpp b/Engine/Source/Runtime/Engine/Private/DecalComponent.cpp
--- a/Engine/Source/Runtime/Engine/Private/DecalComponent.cpp
+++ b/Engine/Source/Runtime/Engine/Private/DecalComponent.cpp
@@ -152,6 +152,11 @@ UDecalComponent::UDecalComponent(const class FPostConstructInitializeProperties&
 {
 }
 
+void UDecalComponent::SetSortOrder(int32 Value)
+{
+	SortOrder = Value;
+	MarkRenderStateDirty();
+}
 
 void UDecalComponent::SetDecalMaterial(class UMaterialInterface* NewDecalMaterial)
 {
