Don't prompt the user to enable bRecomputeTangents if they (or the auto-importer) has already done it

diff --git a/Engine/Source/Runtime/Engine/Private/StaticMeshBuild.cpp b/Engine/Source/Runtime/Engine/Private/StaticMeshBuild.cpp
--- a/Engine/Source/Runtime/Engine/Private/StaticMeshBuild.cpp
+++ b/Engine/Source/Runtime/Engine/Private/StaticMeshBuild.cpp
@@ -113,9 +113,11 @@ void UStaticMesh::Build(bool bSilent)
 		// Warn the user if the new mesh has degenerate tangent bases.
 		if (HasBadTangents(this))
 		{
+			// Only suggest Recompute Tangents if the import hasn't already tried it
 			FFormatNamedArguments Arguments;
 			Arguments.Add( TEXT("Meshname"), FText::FromString(GetName()) );
-			const FText WarningMsg = FText::Format( LOCTEXT("MeshHasDegenerateTangents", "{Meshname} has degenerate tangent bases which will result in incorrect shading. Consider enabling Recompute Tangents in the mesh's Build Settings."), Arguments );
+			Arguments.Add( TEXT("Options"), SourceModels[0].BuildSettings.bRecomputeTangents ? FText::GetEmpty() : LOCTEXT("MeshRecomputeTangents", "Consider enabling Recompute Tangents in the mesh's Build Settings.") );
+			const FText WarningMsg = FText::Format( LOCTEXT("MeshHasDegenerateTangents", "{Meshname} has degenerate tangent bases which will result in incorrect shading. {Options}"), Arguments );
 			UE_LOG(LogStaticMesh,Warning,TEXT("%s"),*WarningMsg.ToString());
 			if (!bSilent)
 			{
