fixed navigation update after moving non relevant component

diff --git a/Engine/Source/Runtime/Engine/Classes/Components/PrimitiveComponent.h b/Engine/Source/Runtime/Engine/Classes/Components/PrimitiveComponent.h
--- a/Engine/Source/Runtime/Engine/Classes/Components/PrimitiveComponent.h
+++ b/Engine/Source/Runtime/Engine/Classes/Components/PrimitiveComponent.h
@@ -1427,11 +1427,14 @@ public:
 	bool IsNavigationRelevant(bool bSkipCollisionEnabledCheck = false) const;
 
 	/** Can this component potentially influence navigation */
-	bool CanEverAffectNavigation() const 
+	bool CanEverAffectNavigation() const
 	{
 		return bCanEverAffectNavigation;
 	}
 
+	/** check if navigation system needs to be updated after changing transform of this component (include attached components as well) */
+	bool ShouldUpdateNavigationOnTransformChange() const;
+
 	/** turn off navigation relevance, must be called before component is registered! */
 	void DisableNavigationRelevance();
 
diff --git a/Engine/Source/Runtime/Engine/Private/PrimitiveComponent.cpp b/Engine/Source/Runtime/Engine/Private/PrimitiveComponent.cpp
--- a/Engine/Source/Runtime/Engine/Private/PrimitiveComponent.cpp
+++ b/Engine/Source/Runtime/Engine/Private/PrimitiveComponent.cpp
@@ -1223,7 +1223,10 @@ void UPrimitiveComponent::SetRelativeScale3D(FVector NewScale3D)
 	Super::SetRelativeScale3D(NewScale3D);
 
 	AActor* Actor = GetOwner();
-	if (UNavigationSystem::ShouldUpdateNavOctreeOnPrimitiveComponentChange() && Actor && OldScale3D != RelativeScale3D && IsRegistered() && Actor->IsNavigationRelevant() && CanEverAffectNavigation() && IsNavigationRelevant() && World != NULL && World->IsGameWorld() && World->GetNetMode() < ENetMode::NM_Client)
+	if (OldScale3D != RelativeScale3D && UNavigationSystem::ShouldUpdateNavOctreeOnPrimitiveComponentChange() &&
+		Actor && IsRegistered() && Actor->IsNavigationRelevant() && 
+		ShouldUpdateNavigationOnTransformChange() &&
+		World != NULL && World->IsGameWorld() && World->GetNetMode() < ENetMode::NM_Client)
 	{
 		if (UNavigationSystem* NavSys = World->GetNavigationSystem())
 		{
@@ -1531,8 +1534,11 @@ bool UPrimitiveComponent::MoveComponent( const FVector& Delta, const FRotator& N
 		}
 	}
 
-	// Update Nav system if relevant.
-	if (bMoved && UNavigationSystem::ShouldUpdateNavOctreeOnPrimitiveComponentChange() && Actor && IsRegistered() && Actor->IsNavigationRelevant() && CanEverAffectNavigation() && IsNavigationRelevant() && World != NULL && World->IsGameWorld() && World->GetNetMode() < ENetMode::NM_Client)
+	// Update navigation system
+	if (bMoved && UNavigationSystem::ShouldUpdateNavOctreeOnPrimitiveComponentChange() &&
+		Actor && IsRegistered() && Actor->IsNavigationRelevant() &&
+		ShouldUpdateNavigationOnTransformChange() &&
+		World != NULL && World->IsGameWorld() && World->GetNetMode() < ENetMode::NM_Client)
 	{
 		if (UNavigationSystem* NavSys = World->GetNavigationSystem())
 		{
@@ -1599,6 +1605,22 @@ void UPrimitiveComponent::DisableNavigationRelevance()
 	bCanEverAffectNavigation = false;
 }
 
+bool UPrimitiveComponent::ShouldUpdateNavigationOnTransformChange() const
+{
+	bool bResult = CanEverAffectNavigation() && IsNavigationRelevant();
+	for (int32 Idx = 0; Idx < AttachChildren.Num() && !bResult; Idx++)
+	{
+		const UPrimitiveComponent* PrimComp = Cast<const UPrimitiveComponent>(AttachChildren[Idx]);
+		if (PrimComp)
+		{
+			bResult = PrimComp->ShouldUpdateNavigationOnTransformChange();
+		}
+	}
+
+	return bResult;
+}
+
+
 //////////////////////////////////////////////////////////////////////////
 // COLLISION
 
