Fixed ribbon spawn source bug.

Following it being moved to fix another issue, source location was being immediately cleared in PreSpawn.
Now passing the source location into PreSpawn as the InitialLocation.

#ttp 340187 - CRITICAL: Regression: CONTENTEXAMPLES: Ribbon trails are missing in the Effects map

diff --git a/Engine/Source/Runtime/Engine/Private/Particles/ParticleTrail2EmitterInstance.cpp b/Engine/Source/Runtime/Engine/Private/Particles/ParticleTrail2EmitterInstance.cpp
--- a/Engine/Source/Runtime/Engine/Private/Particles/ParticleTrail2EmitterInstance.cpp
+++ b/Engine/Source/Runtime/Engine/Private/Particles/ParticleTrail2EmitterInstance.cpp
@@ -1818,13 +1818,8 @@ bool FParticleRibbonEmitterInstance::Spawn_Source(float DeltaTime)
 
 				Component->ComponentToWorld = FTransform(CurrRotation, CurrPosition);
 
-				// Moved before the spawn module iteration so the modules have a chance to modify the position; this may change behavior slightly, but is necessary 
-				// for things like InitalLocation modules to not break in combination with SpawnPerUnit
-				Particle->Location = CurrPosition;
-				Particle->OldLocation = CurrPosition;
-
 				// Standard spawn setup
-				PreSpawn(Particle, Location, FVector::ZeroVector);
+				PreSpawn(Particle, CurrPosition, FVector::ZeroVector);
 				for (int32 SpawnModuleIdx = 0; SpawnModuleIdx < LODLevel->SpawnModules.Num(); SpawnModuleIdx++)
 				{
 					UParticleModule* SpawnModule = LODLevel->SpawnModules[SpawnModuleIdx];
