Always load engine runtime startup modules in the editor, also
- If you have dynamically-loaded "runtime" engine modules that need to be loaded at startup, add them to UGameEngine::LoadRuntimeEngineStartupModules()

diff --git a/Engine/Source/Editor/UnrealEd/Private/Editor.cpp b/Engine/Source/Editor/UnrealEd/Private/Editor.cpp
--- a/Engine/Source/Editor/UnrealEd/Private/Editor.cpp
+++ b/Engine/Source/Editor/UnrealEd/Private/Editor.cpp
@@ -496,6 +496,10 @@ void UEditorEngine::Init(IEngineLoop* InEngineLoop)
 	// Init transactioning.
 	Trans = CreateTrans();
 
+	// Load all of the runtime modules that the game needs.  The game is part of the editor, so we'll need these loaded.
+	UGameEngine::LoadRuntimeEngineStartupModules();
+
+
 	// Load all editor modules here
 	{
 		FModuleManager::Get().LoadModule(TEXT("Documentation"));
diff --git a/Engine/Source/Runtime/Engine/Classes/Engine/GameEngine.h b/Engine/Source/Runtime/Engine/Classes/Engine/GameEngine.h
--- a/Engine/Source/Runtime/Engine/Classes/Engine/GameEngine.h
+++ b/Engine/Source/Runtime/Engine/Classes/Engine/GameEngine.h
@@ -124,6 +124,10 @@ public:
 		return GameViewportWidget;
 	}
 
+	/** Loads engine startup modules that the game needs to run.  This will also be called when the editor starts up. */
+	static void LoadRuntimeEngineStartupModules();
+
+
 	/**
 	 * This is a global, parameterless function used by the online subsystem modules.
 	 * It should never be used in gamecode - instead use the appropriate world context function 
diff --git a/Engine/Source/Runtime/Engine/Private/GameEngine.cpp b/Engine/Source/Runtime/Engine/Private/GameEngine.cpp
--- a/Engine/Source/Runtime/Engine/Private/GameEngine.cpp
+++ b/Engine/Source/Runtime/Engine/Private/GameEngine.cpp
@@ -410,13 +410,8 @@ void UGameEngine::Init(IEngineLoop* InEngineLoop)
 	}
 #endif
 
-	// Load game modules
-	{
-		if( !IsRunningDedicatedServer() && !IsRunningCommandlet() )
-		{
-			FModuleManager::Get().LoadModule( TEXT("GameLiveStreaming") );
-		}
-	}
+	// Load all of the engine modules that we need at startup that are not editor-related
+	UGameEngine::LoadRuntimeEngineStartupModules();
 
 	// Load and apply user game settings
 	GetGameUserSettings()->LoadSettings();
@@ -563,6 +558,21 @@ void UGameEngine::FinishDestroy()
 	Super::FinishDestroy();
 }
 
+
+void UGameEngine::LoadRuntimeEngineStartupModules()
+{
+	// NOTE: These modules will be loaded when the game starts up, and also when the editor starts up.
+
+	// We only want live streaming support if we're actually in a game
+	if( !IsRunningDedicatedServer() && !IsRunningCommandlet() )
+	{
+		FModuleManager::Get().LoadModule( TEXT("GameLiveStreaming") );
+	}
+
+	// ... load other required engine runtime modules here (but NOT editor modules) ...
+}
+
+
 /*-----------------------------------------------------------------------------
 	Command line executor.
 -----------------------------------------------------------------------------*/
