Added clear coat material outputs

diff --git a/Engine/Shaders/BasePassPixelShader.usf b/Engine/Shaders/BasePassPixelShader.usf
--- a/Engine/Shaders/BasePassPixelShader.usf
+++ b/Engine/Shaders/BasePassPixelShader.usf
@@ -725,17 +725,16 @@ void Main(
 		#elif MATERIAL_SHADINGMODEL_CLEAR_COAT
 			GBuffer.ShadingModelID = SHADINGMODELID_CLEAR_COAT;
 
-			float ClearCoat = 1;
-			float ClearCoatRoughness = 0.1;
-			float Film = 1;
+			float ClearCoat				= GetMaterialClearCoat(MaterialParameters);
+			float ClearCoatRoughness	= GetMaterialClearCoatRoughness(MaterialParameters);
+			float Film = 1 * ClearCoat;
 			float MetalSpec = 0.9;
 
-			Film *= ClearCoat;
-
 			float NoV = saturate( dot( MaterialParameters.WorldNormal, MaterialParameters.CameraVector ) );
 
 			// Approximation of refraction's effect on EnvBRDF
 			float RefractionScale = saturate( (NoV - 1) * lerp( 1.3, 0.5, Roughness ) + 1 );
+			RefractionScale = lerp( 1, RefractionScale, ClearCoat );
 
 #if DIFFUSE_SPEC_INPUTS
 			// Absorption doesn't work correctly with nonmetals with diffuse/spec
diff --git a/Engine/Shaders/MaterialTemplate.usf b/Engine/Shaders/MaterialTemplate.usf
--- a/Engine/Shaders/MaterialTemplate.usf
+++ b/Engine/Shaders/MaterialTemplate.usf
@@ -1125,6 +1125,16 @@ half3 GetMaterialSubsurfaceColor(FMaterialPixelParameters Parameters)
 %s;
 }
 
+half GetMaterialClearCoat(FMaterialPixelParameters Parameters)
+{
+%s;
+}
+
+half GetMaterialClearCoatRoughness(FMaterialPixelParameters Parameters)
+{
+%s;
+}
+
 half GetMaterialAmbientOcclusion(FMaterialPixelParameters Parameters)
 {
 %s;
diff --git a/Engine/Source/Editor/MaterialEditor/Private/MaterialEditorUtilities.cpp b/Engine/Source/Editor/MaterialEditor/Private/MaterialEditorUtilities.cpp
--- a/Engine/Source/Editor/MaterialEditor/Private/MaterialEditorUtilities.cpp
+++ b/Engine/Source/Editor/MaterialEditor/Private/MaterialEditorUtilities.cpp
@@ -303,14 +303,16 @@ FExpressionInput* FMaterialEditorUtilities::GetMaterialInput(UMaterial* Material
 	case 10: ExpressionInput = &Material->WorldPositionOffset ; break;
 	case 11: ExpressionInput = &Material->WorldDisplacement ; break;
 	case 12: ExpressionInput = &Material->TessellationMultiplier ; break;
-	case 13: ExpressionInput = &Material->SubsurfaceColor ; break;		
-	case 14: ExpressionInput = &Material->AmbientOcclusion ; break;		
-	case 15: ExpressionInput = &Material->Refraction ; break;		
-	case 24: ExpressionInput = &Material->MaterialAttributes ; break;		
+	case 13: ExpressionInput = &Material->SubsurfaceColor ; break;
+	case 14: ExpressionInput = &Material->ClearCoat ; break;
+	case 15: ExpressionInput = &Material->ClearCoatRoughness ; break;
+	case 16: ExpressionInput = &Material->AmbientOcclusion ; break;
+	case 17: ExpressionInput = &Material->Refraction ; break;
+	case 26: ExpressionInput = &Material->MaterialAttributes ; break;
 	default: 
-		if (Index >= 16 && Index <= 23)
+		if (Index >= 18 && Index <= 25)
 		{
-			ExpressionInput = &Material->CustomizedUVs[Index - 16]; break;	
+			ExpressionInput = &Material->CustomizedUVs[Index - 18]; break;	
 		}
 		else
 		{
@@ -325,9 +327,9 @@ bool FMaterialEditorUtilities::IsInputVisible(UMaterial* Material, int32 Index)
 	static const auto UseDiffuseSpecularMaterialInputs = IConsoleManager::Get().FindTConsoleVariableDataInt(TEXT("r.UseDiffuseSpecularMaterialInputs"));
 
 	bool bVisibleUV = true;
-	if (Index >= 16 && Index <= 23)
+	if (Index >= 18 && Index <= 25)
 	{
-		bVisibleUV = Index - 16 < Material->NumCustomizedUVs;
+		bVisibleUV = Index - 18 < Material->NumCustomizedUVs;
 	}
 
 	switch( Index )
@@ -339,7 +341,7 @@ bool FMaterialEditorUtilities::IsInputVisible(UMaterial* Material, int32 Index)
 	case 3:
 	case 4:
 		return (UseDiffuseSpecularMaterialInputs->GetValueOnGameThread() == 0) && !Material->bUseMaterialAttributes;
-	case 24:
+	case 26:
 		return Material->bUseMaterialAttributes;
 	default:
 		return !Material->bUseMaterialAttributes && bVisibleUV;
diff --git a/Engine/Source/Editor/UnrealEd/Private/MaterialGraph.cpp b/Engine/Source/Editor/UnrealEd/Private/MaterialGraph.cpp
--- a/Engine/Source/Editor/UnrealEd/Private/MaterialGraph.cpp
+++ b/Engine/Source/Editor/UnrealEd/Private/MaterialGraph.cpp
@@ -44,6 +44,8 @@ void UMaterialGraph::RebuildGraph()
 		MaterialInputs.Add( FMaterialInputInfo( LOCTEXT("WorldDisplacement", "World Displacement"), &Material->WorldDisplacement, MP_WorldDisplacement ) );
 		MaterialInputs.Add( FMaterialInputInfo( LOCTEXT("TessellationMultiplier", "Tessellation Multiplier"), &Material->TessellationMultiplier, MP_TessellationMultiplier ) );
 		MaterialInputs.Add( FMaterialInputInfo( LOCTEXT("SubsurfaceColor", "Subsurface Color"), &Material->SubsurfaceColor, MP_SubsurfaceColor ) );
+		MaterialInputs.Add( FMaterialInputInfo( LOCTEXT("ClearCoat", "Clear Coat"), &Material->ClearCoat, MP_ClearCoat ) );
+		MaterialInputs.Add( FMaterialInputInfo( LOCTEXT("ClearCoatRoughness", "Clear Coat Roughness"), &Material->ClearCoatRoughness, MP_ClearCoatRoughness ) );
 		MaterialInputs.Add( FMaterialInputInfo( LOCTEXT("AmbientOcclusion", "Ambient Occlusion"), &Material->AmbientOcclusion, MP_AmbientOcclusion ) );
 		MaterialInputs.Add( FMaterialInputInfo( LOCTEXT("Refraction", "Refraction"), &Material->Refraction, MP_Refraction) );
 
diff --git a/Engine/Source/Runtime/Engine/Classes/Materials/Material.h b/Engine/Source/Runtime/Engine/Classes/Materials/Material.h
--- a/Engine/Source/Runtime/Engine/Classes/Materials/Material.h
+++ b/Engine/Source/Runtime/Engine/Classes/Materials/Material.h
@@ -358,6 +358,14 @@ public:
 	UPROPERTY()
 	FColorMaterialInput SubsurfaceColor;
 
+	/**  */
+	UPROPERTY()
+	FScalarMaterialInput ClearCoat;
+
+	/**  */
+	UPROPERTY()
+	FScalarMaterialInput ClearCoatRoughness;
+
 	/** output ambient occlusion to the GBuffer */
 	UPROPERTY()
 	FScalarMaterialInput AmbientOcclusion;
diff --git a/Engine/Source/Runtime/Engine/Classes/Materials/MaterialExpressionMakeMaterialAttributes.h b/Engine/Source/Runtime/Engine/Classes/Materials/MaterialExpressionMakeMaterialAttributes.h
--- a/Engine/Source/Runtime/Engine/Classes/Materials/MaterialExpressionMakeMaterialAttributes.h
+++ b/Engine/Source/Runtime/Engine/Classes/Materials/MaterialExpressionMakeMaterialAttributes.h
@@ -46,6 +46,12 @@ class UMaterialExpressionMakeMaterialAttributes : public UMaterialExpression
 	UPROPERTY()
 	FExpressionInput SubsurfaceColor;
 
+	UPROPERTY()
+	FExpressionInput ClearCoat;
+
+	UPROPERTY()
+	FExpressionInput ClearCoatRoughness;
+
 	UPROPERTY()
 	FExpressionInput AmbientOcclusion;
 
diff --git a/Engine/Source/Runtime/Engine/Private/Materials/HLSLMaterialTranslator.h b/Engine/Source/Runtime/Engine/Private/Materials/HLSLMaterialTranslator.h
--- a/Engine/Source/Runtime/Engine/Private/Materials/HLSLMaterialTranslator.h
+++ b/Engine/Source/Runtime/Engine/Private/Materials/HLSLMaterialTranslator.h
@@ -334,6 +334,8 @@ public:
 			}
 			Chunk[MP_TessellationMultiplier]		= ForceCast(Material->CompileProperty(MP_TessellationMultiplier,GetMaterialPropertyShaderFrequency(MP_TessellationMultiplier),this),MCT_Float1);
 			Chunk[MP_SubsurfaceColor]				= ForceCast(Material->CompileProperty(MP_SubsurfaceColor       ,GetMaterialPropertyShaderFrequency(MP_SubsurfaceColor),this),MCT_Float3);
+			Chunk[MP_ClearCoat]						= ForceCast(Material->CompileProperty(MP_ClearCoat		       ,GetMaterialPropertyShaderFrequency(MP_ClearCoat),this),MCT_Float1);
+			Chunk[MP_ClearCoatRoughness]			= ForceCast(Material->CompileProperty(MP_ClearCoatRoughness    ,GetMaterialPropertyShaderFrequency(MP_ClearCoatRoughness),this),MCT_Float1);
 			Chunk[MP_AmbientOcclusion]				= ForceCast(Material->CompileProperty(MP_AmbientOcclusion      ,GetMaterialPropertyShaderFrequency(MP_AmbientOcclusion),this),MCT_Float1);
 			Chunk[MP_Refraction]					= ForceCast(Material->CompileProperty(MP_Refraction            ,GetMaterialPropertyShaderFrequency(MP_Refraction),this),MCT_Float2);
 
@@ -635,6 +637,8 @@ public:
 		LazyPrintf.PushParam(*(TranslatedCodeChunkDefinitions[MP_WorldDisplacement] + TEXT("	return ") + TranslatedCodeChunks[MP_WorldDisplacement] + TEXT(";")));
 		LazyPrintf.PushParam(*(TranslatedCodeChunkDefinitions[MP_TessellationMultiplier] + TEXT("	return ") + TranslatedCodeChunks[MP_TessellationMultiplier] + TEXT(";")));
 		LazyPrintf.PushParam(*(TranslatedCodeChunkDefinitions[MP_SubsurfaceColor] + TEXT("	return ") + TranslatedCodeChunks[MP_SubsurfaceColor] + TEXT(";")));
+		LazyPrintf.PushParam(*(TranslatedCodeChunkDefinitions[MP_ClearCoat] + TEXT("	return ") + TranslatedCodeChunks[MP_ClearCoat] + TEXT(";")));
+		LazyPrintf.PushParam(*(TranslatedCodeChunkDefinitions[MP_ClearCoatRoughness] + TEXT("	return ") + TranslatedCodeChunks[MP_ClearCoatRoughness] + TEXT(";")));
 		LazyPrintf.PushParam(*(TranslatedCodeChunkDefinitions[MP_AmbientOcclusion] + TEXT("	return ") + TranslatedCodeChunks[MP_AmbientOcclusion] + TEXT(";")));
 		LazyPrintf.PushParam(*(TranslatedCodeChunkDefinitions[MP_Refraction] + TEXT("	return ") + TranslatedCodeChunks[MP_Refraction] + TEXT(";")));
 
diff --git a/Engine/Source/Runtime/Engine/Private/Materials/Material.cpp b/Engine/Source/Runtime/Engine/Private/Materials/Material.cpp
--- a/Engine/Source/Runtime/Engine/Private/Materials/Material.cpp
+++ b/Engine/Source/Runtime/Engine/Private/Materials/Material.cpp
@@ -86,31 +86,14 @@ int32 FMaterialResource::CompileProperty(EMaterialProperty Property,EShaderFrequ
 		{
 			return Compiler->ForceCast(MaterialInterface->CompileProperty(Compiler, MP_EmissiveColor),MCT_Float3);
 		}
-	case MP_Opacity: return MaterialInterface->CompileProperty(Compiler, MP_Opacity);
-	case MP_OpacityMask: return MaterialInterface->CompileProperty(Compiler, MP_OpacityMask);
 	case MP_DiffuseColor: 
 		return Compiler->Mul(Compiler->ForceCast(MaterialInterface->CompileProperty(Compiler, MP_DiffuseColor),MCT_Float3),Compiler->Sub(Compiler->Constant(1.0f),SelectionColorIndex));
-	case MP_SpecularColor: return MaterialInterface->CompileProperty(Compiler, MP_SpecularColor);
 	case MP_BaseColor: 
 		return Compiler->Mul(Compiler->ForceCast(MaterialInterface->CompileProperty(Compiler, MP_BaseColor),MCT_Float3),Compiler->Sub(Compiler->Constant(1.0f),SelectionColorIndex));
-	case MP_Metallic: return MaterialInterface->CompileProperty(Compiler, MP_Metallic);
-	case MP_Specular: return MaterialInterface->CompileProperty(Compiler, MP_Specular);
-	case MP_Roughness: return MaterialInterface->CompileProperty(Compiler, MP_Roughness);
-	case MP_Normal: return MaterialInterface->CompileProperty(Compiler, MP_Normal);
-	case MP_WorldPositionOffset: return MaterialInterface->CompileProperty(Compiler, MP_WorldPositionOffset);
-	case MP_WorldDisplacement: return MaterialInterface->CompileProperty(Compiler, MP_WorldDisplacement);
-	case MP_TessellationMultiplier: return MaterialInterface->CompileProperty(Compiler, MP_TessellationMultiplier);
-	case MP_SubsurfaceColor: return MaterialInterface->CompileProperty(Compiler, MP_SubsurfaceColor);
-	case MP_AmbientOcclusion: return MaterialInterface->CompileProperty(Compiler, MP_AmbientOcclusion);
-	case MP_Refraction: return MaterialInterface->CompileProperty(Compiler, MP_Refraction);
-	default:
-
-		if (Property >= MP_CustomizedUVs0 && Property <= MP_CustomizedUVs7)
-		{
-			return MaterialInterface->CompileProperty(Compiler, (EMaterialProperty)Property);
-		}
-
+	case MP_MaterialAttributes:
 		return INDEX_NONE;
+	default:
+		return MaterialInterface->CompileProperty(Compiler, Property);
 	};
 }
 
@@ -3007,6 +2990,10 @@ FExpressionInput* UMaterial::GetExpressionInputForProperty(EMaterialProperty InP
 		return &TessellationMultiplier;
 	case MP_SubsurfaceColor:
 		return &SubsurfaceColor;
+	case MP_ClearCoat:
+		return &ClearCoat;
+	case MP_ClearCoatRoughness:
+		return &ClearCoatRoughness;
 	case MP_AmbientOcclusion:
 		return &AmbientOcclusion;
 	case MP_Refraction:
@@ -3303,6 +3290,8 @@ int32 UMaterial::CompileProperty( FMaterialCompiler* Compiler, EMaterialProperty
 		case MP_Specular: Ret = Specular.Compile(Compiler,DefaultFloat); break;
 		case MP_Roughness: Ret = Roughness.Compile(Compiler,DefaultFloat); break;
 		case MP_TessellationMultiplier: Ret = TessellationMultiplier.Compile(Compiler,DefaultFloat); break;
+		case MP_ClearCoat: Ret = ClearCoat.Compile(Compiler,DefaultFloat); break;
+		case MP_ClearCoatRoughness: Ret = ClearCoatRoughness.Compile(Compiler,DefaultFloat); break;
 		case MP_AmbientOcclusion: Ret = AmbientOcclusion.Compile(Compiler,DefaultFloat); break;
 		case MP_Refraction: 
 			Ret = Compiler->AppendVector( 
@@ -3550,6 +3539,10 @@ bool UMaterial::IsPropertyActive(EMaterialProperty InProperty)const
 	case MP_SubsurfaceColor:
 		Active = ShadingModel == MSM_Subsurface || ShadingModel == MSM_PreintegratedSkin;
 		break;
+	case MP_ClearCoat:
+	case MP_ClearCoatRoughness:
+		Active = ShadingModel == MSM_ClearCoat;
+		break;
 	case MP_TessellationMultiplier:
 	case MP_WorldDisplacement:
 		Active = D3D11TessellationMode != MTM_NoTessellation;
diff --git a/Engine/Source/Runtime/Engine/Private/Materials/MaterialExpressions.cpp b/Engine/Source/Runtime/Engine/Private/Materials/MaterialExpressions.cpp
--- a/Engine/Source/Runtime/Engine/Private/Materials/MaterialExpressions.cpp
+++ b/Engine/Source/Runtime/Engine/Private/Materials/MaterialExpressions.cpp
@@ -3222,9 +3222,11 @@ int32 UMaterialExpressionMakeMaterialAttributes::Compile(class FMaterialCompiler
 		case 8: Ret = WorldPositionOffset.Compile(Compiler); Expression = WorldPositionOffset.Expression; break; 
 		case 9: Ret = WorldDisplacement.Compile(Compiler); Expression = WorldDisplacement.Expression; break; 
 		case 10: Ret = TessellationMultiplier.Compile(Compiler); Expression = TessellationMultiplier.Expression; break; 
-		case 11: Ret = SubsurfaceColor.Compile(Compiler); Expression = SubsurfaceColor.Expression; break; 
-		case 12: Ret = AmbientOcclusion.Compile(Compiler); Expression = AmbientOcclusion.Expression; break; 
-		case 13: Ret = Refraction.Compile(Compiler); Expression = Refraction.Expression; break; 
+		case 11: Ret = SubsurfaceColor.Compile(Compiler); Expression = SubsurfaceColor.Expression; break;
+		case 12: Ret = ClearCoat.Compile(Compiler); Expression = ClearCoat.Expression; break;
+		case 13: Ret = ClearCoatRoughness.Compile(Compiler); Expression = ClearCoatRoughness.Expression; break;
+		case 14: Ret = AmbientOcclusion.Compile(Compiler); Expression = AmbientOcclusion.Expression; break;
+		case 15: Ret = Refraction.Compile(Compiler); Expression = Refraction.Expression; break; 
 	};
 
 	const int32 CustomUVStart = 14;
@@ -3281,6 +3283,8 @@ UMaterialExpressionBreakMaterialAttributes::UMaterialExpressionBreakMaterialAttr
 	Outputs.Add(FExpressionOutput(TEXT("WorldDisplacement"), 1, 1, 1, 1, 0));
 	Outputs.Add(FExpressionOutput(TEXT("TessellationMultiplier"), 1, 1, 1, 1, 0));
 	Outputs.Add(FExpressionOutput(TEXT("SubsurfaceColor"), 1, 1, 1, 1, 0));
+	Outputs.Add(FExpressionOutput(TEXT("ClearCoat"), 1, 1, 1, 1, 0));
+	Outputs.Add(FExpressionOutput(TEXT("ClearCoatRoughness"), 1, 1, 1, 1, 0));
 	Outputs.Add(FExpressionOutput(TEXT("AmbientOcclusion"), 1, 1, 1, 1, 0));
 	Outputs.Add(FExpressionOutput(TEXT("Refraction"), 1, 1, 1, 0, 0));
 
diff --git a/Engine/Source/Runtime/Engine/Private/Materials/MaterialShared.cpp b/Engine/Source/Runtime/Engine/Private/Materials/MaterialShared.cpp
--- a/Engine/Source/Runtime/Engine/Private/Materials/MaterialShared.cpp
+++ b/Engine/Source/Runtime/Engine/Private/Materials/MaterialShared.cpp
@@ -247,6 +247,8 @@ int32 FMaterialAttributesInput::Compile(class FMaterialCompiler* Compiler, EMate
 		case MP_WorldDisplacement : Ret = Compiler->Constant3(DefaultVector.X,DefaultVector.Y,DefaultVector.Z); break;
 		case MP_TessellationMultiplier: Ret = Compiler->Constant(DefaultFloat); break;
 		case MP_SubsurfaceColor: Ret = Compiler->Constant3(DefaultColor.R,DefaultColor.G,DefaultColor.B); break;
+		case MP_ClearCoat: Ret = Compiler->Constant(DefaultFloat); break;
+		case MP_ClearCoatRoughness: Ret = Compiler->Constant(DefaultFloat); break;
 		case MP_AmbientOcclusion: Ret = Compiler->Constant(DefaultFloat); break;
 		case MP_Refraction: Ret = Compiler->Constant2(DefaultVector.X,DefaultVector.Y); break;
 		};
@@ -280,6 +282,8 @@ EMaterialValueType GetMaterialPropertyType(EMaterialProperty Property)
 	case MP_WorldDisplacement : return MCT_Float3;
 	case MP_TessellationMultiplier: return MCT_Float;
 	case MP_SubsurfaceColor: return MCT_Float3;
+	case MP_ClearCoat: return MCT_Float;
+	case MP_ClearCoatRoughness: return MCT_Float;
 	case MP_AmbientOcclusion: return MCT_Float;
 	case MP_Refraction: return MCT_Float;
 	case MP_MaterialAttributes: return MCT_MaterialAttributes;
@@ -2140,11 +2144,13 @@ EMaterialProperty GetMaterialPropertyFromInputOutputIndex(int32 Index)
 	case 9: return MP_WorldDisplacement;
 	case 10: return MP_TessellationMultiplier;
 	case 11: return MP_SubsurfaceColor;
-	case 12: return MP_AmbientOcclusion;
-	case 13: return MP_Refraction;
+	case 12: return MP_ClearCoat;
+	case 13: return MP_ClearCoatRoughness;
+	case 14: return MP_AmbientOcclusion;
+	case 15: return MP_Refraction;
 	};
 
-	const int32 UVStart = 14;
+	const int32 UVStart = 16;
 	const int32 UVEnd = UVStart + MP_CustomizedUVs7 - MP_CustomizedUVs0;
 
 	if (Index >= UVStart && Index <= UVEnd)
@@ -2176,14 +2182,16 @@ int32 GetInputOutputIndexFromMaterialProperty(EMaterialProperty Property)
 	case MP_WorldDisplacement: return 9;
 	case MP_TessellationMultiplier: return 10;
 	case MP_SubsurfaceColor: return 11;
-	case MP_AmbientOcclusion: return 12;
-	case MP_Refraction: return 13;
+	case MP_ClearCoat: return 12;
+	case MP_ClearCoatRoughness: return 13;
+	case MP_AmbientOcclusion: return 14;
+	case MP_Refraction: return 15;
 	case MP_MaterialAttributes: UE_LOG(LogMaterial, Fatal, TEXT("We should never need the IO index of the MaterialAttriubtes property."));
 	};
 
 	if (Property >= MP_CustomizedUVs0 && Property <= MP_CustomizedUVs7)
 	{
-		return 14 + Property - MP_CustomizedUVs0;
+		return 16 + Property - MP_CustomizedUVs0;
 	}
 
 	return -1;
@@ -2203,6 +2211,8 @@ void GetDefaultForMaterialProperty(EMaterialProperty Property, float& OutDefault
 	case MP_Specular:				OutDefaultFloat = 0.5f; break;
 	case MP_Roughness:				OutDefaultFloat = 0.5f; break;
 	case MP_TessellationMultiplier: OutDefaultFloat = 1.0f; break;
+	case MP_ClearCoat:				OutDefaultFloat = 1.0f; break;
+	case MP_ClearCoatRoughness:		OutDefaultFloat = 0.1f; break;
 	case MP_AmbientOcclusion:		OutDefaultFloat = 1.0f; break;
 	case MP_EmissiveColor:			OutDefaultColor = FColor(0,0,0); break;
 	case MP_DiffuseColor:			OutDefaultColor = FColor(0,0,0); break;
@@ -2232,6 +2242,8 @@ FString GetNameOfMaterialProperty(EMaterialProperty Property)
 	case MP_WorldDisplacement:		return TEXT("WorldDisplacement");
 	case MP_TessellationMultiplier: return TEXT("TessellationMultiplier");
 	case MP_SubsurfaceColor:		return TEXT("SubsurfaceColor");
+	case MP_ClearCoat:				return TEXT("ClearCoat");
+	case MP_ClearCoatRoughness:		return TEXT("ClearCoatRoughness");
 	case MP_AmbientOcclusion:		return TEXT("AmbientOcclusion");
 	case MP_Refraction:				return TEXT("Refraction");
 	};
@@ -2283,6 +2295,8 @@ int32 UMaterialInterface::CompileProperty( FMaterialCompiler* Compiler, EMateria
 		case MP_Specular:
 		case MP_Roughness:
 		case MP_TessellationMultiplier:
+		case MP_ClearCoat:
+		case MP_ClearCoatRoughness:
 		case MP_AmbientOcclusion:
 			{
 				return Compiler->Constant(DefaultFloat);
diff --git a/Engine/Source/Runtime/Engine/Public/SceneTypes.h b/Engine/Source/Runtime/Engine/Public/SceneTypes.h
--- a/Engine/Source/Runtime/Engine/Public/SceneTypes.h
+++ b/Engine/Source/Runtime/Engine/Public/SceneTypes.h
@@ -140,6 +140,8 @@ enum EMaterialProperty
 	MP_WorldDisplacement,
 	MP_TessellationMultiplier,
 	MP_SubsurfaceColor,
+	MP_ClearCoat,
+	MP_ClearCoatRoughness,
 	MP_AmbientOcclusion,
 	MP_Refraction,
 	MP_CustomizedUVs0,
