Grey out Metallic output when subsurface is used.

diff --git a/Engine/Source/Runtime/Engine/Private/Materials/Material.cpp b/Engine/Source/Runtime/Engine/Private/Materials/Material.cpp
--- a/Engine/Source/Runtime/Engine/Private/Materials/Material.cpp
+++ b/Engine/Source/Runtime/Engine/Private/Materials/Material.cpp
@@ -3527,12 +3527,14 @@ bool UMaterial::IsPropertyActive(EMaterialProperty InProperty)const
 	case MP_DiffuseColor:
 	case MP_SpecularColor:
 	case MP_BaseColor:
-	case MP_Metallic:
 	case MP_Specular:
 	case MP_Roughness:
 	case MP_AmbientOcclusion:
 		Active = ShadingModel != MSM_Unlit;
 		break;
+	case MP_Metallic:
+		Active = ShadingModel != MSM_Unlit && ShadingModel != MSM_Subsurface && ShadingModel != MSM_PreintegratedSkin;
+		break;
 	case MP_Normal:
 		Active = ShadingModel != MSM_Unlit || Refraction.IsConnected();
 		break;
