UE4 - more command list work, remove "either render thread or game thread" routines from canvas

diff --git a/Engine/Source/Editor/GraphEditor/Private/MaterialNodes/SGraphNodeMaterialBase.cpp b/Engine/Source/Editor/GraphEditor/Private/MaterialNodes/SGraphNodeMaterialBase.cpp
--- a/Engine/Source/Editor/GraphEditor/Private/MaterialNodes/SGraphNodeMaterialBase.cpp
+++ b/Engine/Source/Editor/GraphEditor/Private/MaterialNodes/SGraphNodeMaterialBase.cpp
@@ -196,7 +196,7 @@ void FPreviewElement::DrawRenderThread(FRHICommandListImmediate& RHICmdList, con
 			FCanvasTileItem TileItem( FVector2D::ZeroVector, ExpressionPreview , RenderTarget->GetSizeXY());
 			Canvas.DrawItem( TileItem );
 		}
-		Canvas.Flush(true);
+		Canvas.Flush_RenderThread(RHICmdList, true);
 	}
 	RenderTarget->ClearRenderTargetTexture();
 	RHICmdList.SetScissorRect(false, 0, 0, 0, 0);
diff --git a/Engine/Source/Editor/MeshPaint/Private/MeshPaintEdMode.cpp b/Engine/Source/Editor/MeshPaint/Private/MeshPaintEdMode.cpp
--- a/Engine/Source/Editor/MeshPaint/Private/MeshPaintEdMode.cpp
+++ b/Engine/Source/Editor/MeshPaint/Private/MeshPaintEdMode.cpp
@@ -2260,7 +2260,7 @@ void FEdModeMeshPaint::PaintTexture( const FMeshPaintParameters& InParams,
 
 	// Tell the rendering thread to draw any remaining batched elements
 	{
-		BrushPaintCanvas.Flush(true);
+		BrushPaintCanvas.Flush_GameThread(true);
 
 		TextureData->bIsPaintingTexture2DModified = true;
 	}
@@ -2282,7 +2282,7 @@ void FEdModeMeshPaint::PaintTexture( const FMeshPaintParameters& InParams,
 
 	if( bEnableSeamPainting )
 	{
-		BrushMaskCanvas->Flush(true);
+		BrushMaskCanvas->Flush_GameThread(true);
 
 		{
 			ENQUEUE_UNIQUE_RENDER_COMMAND_ONEPARAMETER(
@@ -2380,7 +2380,7 @@ void FEdModeMeshPaint::PaintTexture( const FMeshPaintParameters& InParams,
 			
 
 			// Tell the rendering thread to draw any remaining batched elements
-			Canvas3.Flush(true);
+			Canvas3.Flush_GameThread(true);
 
 		}
 
@@ -3293,7 +3293,7 @@ void FEdModeMeshPaint::CopyTextureToRenderTargetTexture( UTexture* SourceTexture
 		}
 
 		// Tell the rendering thread to draw any remaining batched elements
-		Canvas.Flush(true);
+		Canvas.Flush_GameThread(true);
 	}
 
 
@@ -3504,7 +3504,7 @@ bool FEdModeMeshPaint::GenerateSeamMask(UStaticMeshComponent* StaticMeshComponen
 		// And render it
 		Canvas.DrawItem( TriItem );
 		// Tell the rendering thread to draw any remaining batched elements
-		Canvas.Flush(true);
+		Canvas.Flush_GameThread(true);
 	}
 
 
diff --git a/Engine/Source/Editor/UnrealEd/Private/CubemapUnwrapUtils.cpp b/Engine/Source/Editor/UnrealEd/Private/CubemapUnwrapUtils.cpp
--- a/Engine/Source/Editor/UnrealEd/Private/CubemapUnwrapUtils.cpp
+++ b/Engine/Source/Editor/UnrealEd/Private/CubemapUnwrapUtils.cpp
@@ -46,7 +46,7 @@ namespace CubemapHelpers
 		FRenderTarget* RenderTarget = RenderTargetLongLat->GameThread_GetRenderTargetResource();
 
 		FCanvas* Canvas = new FCanvas(RenderTarget, NULL, 0, 0, 0);
-		Canvas->SetRenderTarget(RenderTarget);
+		Canvas->SetRenderTarget_GameThread(RenderTarget);
 
 		// Clear the render target to black
 		Canvas->Clear(FLinearColor(0, 0, 0, 0));
@@ -56,9 +56,9 @@ namespace CubemapHelpers
 		TileItem.BlendMode = SE_BLEND_Opaque;
 		Canvas->DrawItem(TileItem);
 
-		Canvas->Flush();
+		Canvas->Flush_GameThread();
 		FlushRenderingCommands();
-		Canvas->SetRenderTarget(NULL);
+		Canvas->SetRenderTarget_GameThread(NULL);
 		FlushRenderingCommands();
 		
 		int32 ImageBytes = CalculateImageBytes(LongLatDimensions.X, LongLatDimensions.Y, 0, TargetPixelFormat);
diff --git a/Engine/Source/Editor/UnrealEd/Private/EditorExporters.cpp b/Engine/Source/Editor/UnrealEd/Private/EditorExporters.cpp
--- a/Engine/Source/Editor/UnrealEd/Private/EditorExporters.cpp
+++ b/Engine/Source/Editor/UnrealEd/Private/EditorExporters.cpp
@@ -1310,15 +1310,15 @@ bool GenerateExportMaterialPropertyData(
 		}
 
 		// start drawing to the render target
-		Canvas->SetRenderTarget(RenderTarget->GetRenderTargetResource());
+		Canvas->SetRenderTarget_GameThread(RenderTarget->GetRenderTargetResource());
 		// Freeze time while capturing the material's inputs
 		Canvas->Clear(FLinearColor(0,0,0,0));
 		FCanvasTileItem TileItem( FVector2D( 0.0f, 0.0f ), MaterialProxy, FVector2D( InOutSizeX, InOutSizeY ) );
 		TileItem.bFreezeTime = true;
 		Canvas->DrawItem( TileItem );
-		Canvas->Flush();
+		Canvas->Flush_GameThread();
 		FlushRenderingCommands();
-		Canvas->SetRenderTarget(NULL);
+		Canvas->SetRenderTarget_GameThread(NULL);
 		FlushRenderingCommands();
 
 		// if PF_FloatRGB was used as render target format, the gamma conversion during rendering is deactivated
@@ -2602,9 +2602,9 @@ namespace MaterialExportUtils
 			FCanvasTileItem TileItem(FVector2D(0.0f, 0.0f), MaterialProxy, FVector2D(InRenderTarget->SizeX, InRenderTarget->SizeY));
 			TileItem.bFreezeTime = true;
 			Canvas.DrawItem( TileItem );
-			Canvas.Flush();
+			Canvas.Flush_GameThread();
 			FlushRenderingCommands();
-			Canvas.SetRenderTarget(NULL);
+			Canvas.SetRenderTarget_GameThread(NULL);
 			FlushRenderingCommands();
 		}
 				
diff --git a/Engine/Source/Editor/UnrealEd/Private/Lightmass/LightmassRender.cpp b/Engine/Source/Editor/UnrealEd/Private/Lightmass/LightmassRender.cpp
--- a/Engine/Source/Editor/UnrealEd/Private/Lightmass/LightmassRender.cpp
+++ b/Engine/Source/Editor/UnrealEd/Private/Lightmass/LightmassRender.cpp
@@ -930,7 +930,7 @@ bool FLightmassMaterialRenderer::GenerateMaterialPropertyData(
 			else
 			{
 				// At this point, we can't just return false at failure since we have some clean-up to do...
-				Canvas->SetRenderTarget(RenderTarget->GameThread_GetRenderTargetResource());
+				Canvas->SetRenderTarget_GameThread(RenderTarget->GameThread_GetRenderTargetResource());
 
 				// Clear the render target to black
 				// This is necessary because the below DrawTile doesn't write to the first column and first row
@@ -939,9 +939,9 @@ bool FLightmassMaterialRenderer::GenerateMaterialPropertyData(
 				FCanvasTileItem TileItem( FVector2D( 0.0f, 0.0f ), &MaterialProxy, FVector2D( InOutSizeX, InOutSizeY ) );
 				TileItem.bFreezeTime = true;
 				Canvas->DrawItem( TileItem );
-				Canvas->Flush();
+				Canvas->Flush_GameThread();
 				FlushRenderingCommands();
-				Canvas->SetRenderTarget(NULL);
+				Canvas->SetRenderTarget_GameThread(NULL);
 				FlushRenderingCommands();
 
 				// Read in the data
diff --git a/Engine/Source/Editor/UnrealEd/Private/ObjectTools.cpp b/Engine/Source/Editor/UnrealEd/Private/ObjectTools.cpp
--- a/Engine/Source/Editor/UnrealEd/Private/ObjectTools.cpp
+++ b/Engine/Source/Editor/UnrealEd/Private/ObjectTools.cpp
@@ -3459,7 +3459,7 @@ namespace ThumbnailTools
 		if( bUseGPUGeneratedThumbnail )
 		{
 			// Tell the rendering thread to draw any remaining batched elements
-			Canvas.Flush();
+			Canvas.Flush_GameThread();
 
 
 			{
diff --git a/Engine/Source/Runtime/Engine/Private/CanvasRenderTarget2D.cpp b/Engine/Source/Runtime/Engine/Private/CanvasRenderTarget2D.cpp
--- a/Engine/Source/Runtime/Engine/Private/CanvasRenderTarget2D.cpp
+++ b/Engine/Source/Runtime/Engine/Private/CanvasRenderTarget2D.cpp
@@ -62,7 +62,7 @@ void UCanvasRenderTarget2D::UpdateResource()
 
 	// Clean up and flush the rendering canvas.
 	Canvas->Canvas = nullptr;
-	RenderCanvas.Flush();
+	RenderCanvas.Flush_GameThread();
 
 	// Enqueue the rendering command to copy the freshly rendering texture resource back to the render target RHI 
 	// so that the texture is updated and available for rendering.
diff --git a/Engine/Source/Runtime/Engine/Private/GameViewportClient.cpp b/Engine/Source/Runtime/Engine/Private/GameViewportClient.cpp
--- a/Engine/Source/Runtime/Engine/Private/GameViewportClient.cpp
+++ b/Engine/Source/Runtime/Engine/Private/GameViewportClient.cpp
@@ -967,10 +967,10 @@ void UGameViewportClient::Draw(FViewport* InViewport, FCanvas* SceneCanvas)
 		}
 
 		//ensure canvas has been flushed before rendering UI
-		SceneCanvas->Flush();
+		SceneCanvas->Flush_GameThread();
 		if (DebugCanvas != NULL)
 		{
-			DebugCanvas->Flush();
+			DebugCanvas->Flush_GameThread();
 		}
 		// Allow the viewport to render additional stuff
 		PostRender(DebugCanvasObject);
diff --git a/Engine/Source/Runtime/Engine/Private/Slate/DebugCanvas.cpp b/Engine/Source/Runtime/Engine/Private/Slate/DebugCanvas.cpp
--- a/Engine/Source/Runtime/Engine/Private/Slate/DebugCanvas.cpp
+++ b/Engine/Source/Runtime/Engine/Private/Slate/DebugCanvas.cpp
@@ -172,7 +172,7 @@ void FDebugCanvasDrawer::DrawRenderThread(FRHICommandListImmediate& RHICmdList,
 		RenderTarget->SetRenderTargetTexture( *(FTexture2DRHIRef*)InWindowBackBuffer );
 		{
 			RenderThreadCanvas->Canvas.SetRenderTargetRect( RenderTarget->GetViewRect() );
-			RenderThreadCanvas->Canvas.Flush( true );
+			RenderThreadCanvas->Canvas.Flush_RenderThread(RHICmdList, true);
 		}
 		RenderTarget->ClearRenderTargetTexture();
 	}
diff --git a/Engine/Source/Runtime/Engine/Private/UnrealClient.cpp b/Engine/Source/Runtime/Engine/Private/UnrealClient.cpp
--- a/Engine/Source/Runtime/Engine/Private/UnrealClient.cpp
+++ b/Engine/Source/Runtime/Engine/Private/UnrealClient.cpp
@@ -807,7 +807,7 @@ void FViewport::HighResScreenshot()
 	{
 		ViewportClient->Draw(DummyViewport, &Canvas);
 	}
-	Canvas.Flush();
+	Canvas.Flush_GameThread();
 	ViewportClient->GetEngineShowFlags()->HighResScreenshotMask = MaskShowFlagBackup;
 	ViewportClient->GetEngineShowFlags()->MotionBlur = MotionBlurShowFlagBackup;
 	ViewportClient->ProcessScreenShots(DummyViewport);
@@ -822,7 +822,7 @@ void FViewport::HighResScreenshot()
 	});
 
 	// Draw the debug canvas;
-	DummyViewport->GetDebugCanvas()->Flush(true);
+	DummyViewport->GetDebugCanvas()->Flush_GameThread(true);
 
 	BeginReleaseResource(DummyViewport);
 	FlushRenderingCommands();
@@ -1072,7 +1072,7 @@ void FViewport::Draw( bool bShouldPresent /*= true */)
 				{
 					ViewportClient->Draw(this, &Canvas);
 				}
-				Canvas.Flush();
+				Canvas.Flush_GameThread();
 				ViewportClient->ProcessScreenShots(this);
 	
 				// Slate doesn't present immediately. Tag the viewport as requiring vsync so that it happens.
@@ -1175,7 +1175,7 @@ const TArray<FColor>& FViewport::GetRawHitProxyData(FIntRect InRect)
 		{
 			ViewportClient->Draw(this, &Canvas);
 		}
-		Canvas.Flush();
+		Canvas.Flush_GameThread();
 
 		//Resolve surface to texture.
 		ENQUEUE_UNIQUE_RENDER_COMMAND_ONEPARAMETER(
diff --git a/Engine/Source/Runtime/Engine/Private/UserInterface/Canvas.cpp b/Engine/Source/Runtime/Engine/Private/UserInterface/Canvas.cpp
--- a/Engine/Source/Runtime/Engine/Private/UserInterface/Canvas.cpp
+++ b/Engine/Source/Runtime/Engine/Private/UserInterface/Canvas.cpp
@@ -152,8 +152,50 @@ FMatrix FCanvas::CalcProjectionMatrix(uint32 ViewSizeX, uint32 ViewSizeY, float
 		);
 }
 
+bool FCanvasBatchedElementRenderItem::Render_RenderThread(FRHICommandListImmediate& RHICmdList, const FCanvas* Canvas)
+{
+	checkSlow(Data);
+	bool bDirty = false;
+	if (Data->BatchedElements.HasPrimsToDraw())
+	{
+		bDirty = true;
+
+		// current render target set for the canvas
+		const FRenderTarget* CanvasRenderTarget = Canvas->GetRenderTarget();
+		float Gamma = 1.0f / CanvasRenderTarget->GetDisplayGamma();
+		if (Data->Texture && Data->Texture->bIgnoreGammaConversions)
+		{
+			Gamma = 1.0f;
+		}
+
+		const bool bNeedsToSwitchVerticalAxis = IsES2Platform(GRHIShaderPlatform) && !IsPCPlatform(GRHIShaderPlatform);
+
+
+		// draw batched items
+		Data->BatchedElements.Draw(
+			RHICmdList,
+			bNeedsToSwitchVerticalAxis,
+			Data->Transform.GetMatrix(),
+			CanvasRenderTarget->GetSizeXY().X,
+			CanvasRenderTarget->GetSizeXY().Y,
+			Canvas->IsHitTesting(),
+			Gamma
+			);
 
-bool FCanvasBatchedElementRenderItem::Render( const FCanvas* Canvas )
+		if (Canvas->GetAllowedModes() & FCanvas::Allow_DeleteOnRender)
+		{
+			// delete data since we're done rendering it
+			delete Data;
+		}
+	}
+	if (Canvas->GetAllowedModes() & FCanvas::Allow_DeleteOnRender)
+	{
+		Data = NULL;
+	}
+	return bDirty;
+}
+
+bool FCanvasBatchedElementRenderItem::Render_GameThread(const FCanvas* Canvas )
 {	
 	checkSlow(Data);
 	bool bDirty=false;		
@@ -169,76 +211,47 @@ bool FCanvasBatchedElementRenderItem::Render( const FCanvas* Canvas )
 			Gamma = 1.0f;
 		}
 
-		// this allows us to use FCanvas operations from the rendering thread (ie, render subtitles
-		// on top of a movie that is rendered completely in rendering thread)
-		if (IsInRenderingThread())
+		// Render the batched elements.
+		struct FBatchedDrawParameters
+		{
+			FRenderData* RenderData;
+			uint32 bHitTesting : 1;
+			uint32 ViewportSizeX;
+			uint32 ViewportSizeY;
+			float DisplayGamma;
+			uint32 AllowedCanvasModes;
+		};
+		// all the parameters needed for rendering
+		FBatchedDrawParameters DrawParameters =
+		{
+			Data,
+			Canvas->IsHitTesting(),
+			(uint32)CanvasRenderTarget->GetSizeXY().X,
+			(uint32)CanvasRenderTarget->GetSizeXY().Y,
+			Gamma,
+			Canvas->GetAllowedModes()
+		};
+		ENQUEUE_UNIQUE_RENDER_COMMAND_ONEPARAMETER(
+			BatchedDrawCommand,
+			FBatchedDrawParameters,Parameters,DrawParameters,
 		{
 			const bool bNeedsToSwitchVerticalAxis = IsES2Platform(GRHIShaderPlatform) && !IsPCPlatform(GRHIShaderPlatform);
-
-			
-			FRHICommandListImmediate& RHICmdList = FRHICommandListExecutor::GetImmediateCommandList();
-
+				
 			// draw batched items
-			Data->BatchedElements.Draw(
+			Parameters.RenderData->BatchedElements.Draw(
 				RHICmdList,
 				bNeedsToSwitchVerticalAxis,
-				Data->Transform.GetMatrix(),
-				CanvasRenderTarget->GetSizeXY().X,
-				CanvasRenderTarget->GetSizeXY().Y,
-				Canvas->IsHitTesting(),
-				Gamma
+				Parameters.RenderData->Transform.GetMatrix(),
+				Parameters.ViewportSizeX,
+				Parameters.ViewportSizeY,
+				Parameters.bHitTesting,
+				Parameters.DisplayGamma
 				);
-
-			if( Canvas->GetAllowedModes() & FCanvas::Allow_DeleteOnRender )
+			if( Parameters.AllowedCanvasModes & FCanvas::Allow_DeleteOnRender )
 			{
-				// delete data since we're done rendering it
-				delete Data;
+				delete Parameters.RenderData;
 			}
-		}
-		else
-		{
-			// Render the batched elements.
-			struct FBatchedDrawParameters
-			{
-				FRenderData* RenderData;
-				uint32 bHitTesting : 1;
-				uint32 ViewportSizeX;
-				uint32 ViewportSizeY;
-				float DisplayGamma;
-				uint32 AllowedCanvasModes;
-			};
-			// all the parameters needed for rendering
-			FBatchedDrawParameters DrawParameters =
-			{
-				Data,
-				Canvas->IsHitTesting(),
-				(uint32)CanvasRenderTarget->GetSizeXY().X,
-				(uint32)CanvasRenderTarget->GetSizeXY().Y,
-				Gamma,
-				Canvas->GetAllowedModes()
-			};
-			ENQUEUE_UNIQUE_RENDER_COMMAND_ONEPARAMETER(
-				BatchedDrawCommand,
-				FBatchedDrawParameters,Parameters,DrawParameters,
-			{
-				const bool bNeedsToSwitchVerticalAxis = IsES2Platform(GRHIShaderPlatform) && !IsPCPlatform(GRHIShaderPlatform);
-				
-				// draw batched items
-				Parameters.RenderData->BatchedElements.Draw(
-					RHICmdList,
-					bNeedsToSwitchVerticalAxis,
-					Parameters.RenderData->Transform.GetMatrix(),
-					Parameters.ViewportSizeX,
-					Parameters.ViewportSizeY,
-					Parameters.bHitTesting,
-					Parameters.DisplayGamma
-					);
-				if( Parameters.AllowedCanvasModes & FCanvas::Allow_DeleteOnRender )
-				{
-					delete Parameters.RenderData;
-				}
-			});
-		}
+		});
 	}
 	if( Canvas->GetAllowedModes() & FCanvas::Allow_DeleteOnRender )
 	{
@@ -248,8 +261,7 @@ bool FCanvasBatchedElementRenderItem::Render( const FCanvas* Canvas )
 }
 
 
-
-bool FCanvasTileRendererItem::Render( const FCanvas* Canvas )
+bool FCanvasTileRendererItem::Render_RenderThread(FRHICommandListImmediate& RHICmdList, const FCanvas* Canvas)
 {
 	float CurrentRealTime = 0.f;
 	float CurrentWorldTime = 0.f;
@@ -285,78 +297,113 @@ bool FCanvasTileRendererItem::Render( const FCanvas* Canvas )
 
 	FSceneView* View = new FSceneView(ViewInitOptions);
 
-	// Render the batched elements.
-	if( IsInRenderingThread() )
+	for( int32 TileIdx=0; TileIdx < Data->Tiles.Num(); TileIdx++ )
 	{
-		for( int32 TileIdx=0; TileIdx < Data->Tiles.Num(); TileIdx++ )
+		const FRenderData::FTileInst& Tile = Data->Tiles[TileIdx];
+		FTileRenderer::DrawTile(
+			RHICmdList,
+			*View, 
+			Data->MaterialRenderProxy, 
+			Tile.X, Tile.Y, Tile.SizeX, Tile.SizeY, 
+			Tile.U, Tile.V, Tile.SizeU, Tile.SizeV,
+			Canvas->IsHitTesting(), Tile.HitProxyId
+			);
+	}
+
+	delete View->Family;
+	delete View;
+	if( Canvas->GetAllowedModes() & FCanvas::Allow_DeleteOnRender )
+	{
+		delete Data;
+	}
+	if( Canvas->GetAllowedModes() & FCanvas::Allow_DeleteOnRender )
+	{
+		Data = NULL;
+	}
+	return true;
+}
+
+bool FCanvasTileRendererItem::Render_GameThread(const FCanvas* Canvas)
+{
+	float CurrentRealTime = 0.f;
+	float CurrentWorldTime = 0.f;
+	float DeltaWorldTime = 0.f;
+
+	if (!bFreezeTime)
+	{
+		CurrentRealTime = Canvas->GetCurrentRealTime();
+		CurrentWorldTime = Canvas->GetCurrentWorldTime();
+		DeltaWorldTime = Canvas->GetCurrentDeltaWorldTime();
+	}
+
+	checkSlow(Data);
+	// current render target set for the canvas
+	const FRenderTarget* CanvasRenderTarget = Canvas->GetRenderTarget();
+	FSceneViewFamily* ViewFamily = new FSceneViewFamily(FSceneViewFamily::ConstructionValues(
+		CanvasRenderTarget,
+		NULL,
+		FEngineShowFlags(ESFIM_Game))
+		.SetWorldTimes(CurrentWorldTime, DeltaWorldTime, CurrentRealTime)
+		.SetGammaCorrection(CanvasRenderTarget->GetDisplayGamma()));
+
+	FIntRect ViewRect(FIntPoint(0, 0), CanvasRenderTarget->GetSizeXY());
+
+	// make a temporary view
+	FSceneViewInitOptions ViewInitOptions;
+	ViewInitOptions.ViewFamily = ViewFamily;
+	ViewInitOptions.SetViewRectangle(ViewRect);
+	ViewInitOptions.ViewMatrix = FMatrix::Identity;
+	ViewInitOptions.ProjectionMatrix = Data->Transform.GetMatrix();
+	ViewInitOptions.BackgroundColor = FLinearColor::Black;
+	ViewInitOptions.OverlayColor = FLinearColor::White;
+
+	FSceneView* View = new FSceneView(ViewInitOptions);
+
+	struct FDrawTileParameters
+	{
+		FSceneView* View;
+		FRenderData* RenderData;
+		uint32 bIsHitTesting : 1;
+		uint32 AllowedCanvasModes;
+	};
+	FDrawTileParameters DrawTileParameters =
+	{
+		View,
+		Data,
+		Canvas->IsHitTesting(),
+		Canvas->GetAllowedModes()
+	};
+	ENQUEUE_UNIQUE_RENDER_COMMAND_ONEPARAMETER(
+		DrawTileCommand,
+		FDrawTileParameters, Parameters, DrawTileParameters,
+		{
+		for (int32 TileIdx = 0; TileIdx < Parameters.RenderData->Tiles.Num(); TileIdx++)
 		{
-			const FRenderData::FTileInst& Tile = Data->Tiles[TileIdx];
+			const FRenderData::FTileInst& Tile = Parameters.RenderData->Tiles[TileIdx];
 			FTileRenderer::DrawTile(
-				FRHICommandListExecutor::GetImmediateCommandList(),
-				*View, 
-				Data->MaterialRenderProxy, 
-				Tile.X, Tile.Y, Tile.SizeX, Tile.SizeY, 
+				RHICmdList,
+				*Parameters.View,
+				Parameters.RenderData->MaterialRenderProxy,
+				Tile.X, Tile.Y, Tile.SizeX, Tile.SizeY,
 				Tile.U, Tile.V, Tile.SizeU, Tile.SizeV,
-				Canvas->IsHitTesting(), Tile.HitProxyId
+				Parameters.bIsHitTesting, Tile.HitProxyId
 				);
 		}
 
-		delete View->Family;
-		delete View;
-		if( Canvas->GetAllowedModes() & FCanvas::Allow_DeleteOnRender )
+		delete Parameters.View->Family;
+		delete Parameters.View;
+		if (Parameters.AllowedCanvasModes & FCanvas::Allow_DeleteOnRender)
 		{
-			delete Data;
+			delete Parameters.RenderData;
 		}
-	}
-	else
-	{
-		struct FDrawTileParameters
-		{
-			FSceneView* View;
-			FRenderData* RenderData;
-			uint32 bIsHitTesting : 1;
-			uint32 AllowedCanvasModes;
-		};
-		FDrawTileParameters DrawTileParameters =
-		{
-			View,
-			Data,
-			Canvas->IsHitTesting(),
-			Canvas->GetAllowedModes()
-		};
-		ENQUEUE_UNIQUE_RENDER_COMMAND_ONEPARAMETER(
-			DrawTileCommand,
-			FDrawTileParameters,Parameters,DrawTileParameters,
-		{
-			for( int32 TileIdx=0; TileIdx < Parameters.RenderData->Tiles.Num(); TileIdx++ )
-			{
-				const FRenderData::FTileInst& Tile = Parameters.RenderData->Tiles[TileIdx];
-				FTileRenderer::DrawTile(
-					RHICmdList,
-					*Parameters.View, 
-					Parameters.RenderData->MaterialRenderProxy, 
-					Tile.X, Tile.Y, Tile.SizeX, Tile.SizeY, 
-					Tile.U, Tile.V, Tile.SizeU, Tile.SizeV,
-					Parameters.bIsHitTesting, Tile.HitProxyId
-					);
-			}
-
-			delete Parameters.View->Family;
-			delete Parameters.View;
-			if( Parameters.AllowedCanvasModes & FCanvas::Allow_DeleteOnRender )
-			{
-				delete Parameters.RenderData;
-			}
-		});
-	}
-	if( Canvas->GetAllowedModes() & FCanvas::Allow_DeleteOnRender )
+	});
+	if (Canvas->GetAllowedModes() & FCanvas::Allow_DeleteOnRender)
 	{
 		Data = NULL;
 	}
 	return true;
 }
 
-
 FCanvas::FCanvasSortElement& FCanvas::GetSortElement(int32 DepthSortKey)
 {
 	// Optimization to resuse last index so that the more expensive Find() does not
@@ -465,8 +512,85 @@ FCanvas::~FCanvas()
 	}
 }
 
+void FCanvas::Flush_RenderThread(FRHICommandListImmediate& RHICmdList, bool bForce)
+{
+	SCOPE_CYCLE_COUNTER(STAT_Canvas_FlushTime);
+
+	if (!(AllowedModes&Allow_Flush) && !bForce)
+	{
+		return;
+	}
+
+	// current render target set for the canvas
+	check(RenderTarget);
 
-void FCanvas::Flush(bool bForce)
+	// no need to set the render target if we aren't going to draw anything to it!
+	if (SortedElements.Num() == 0)
+	{
+		return;
+	}
+
+	// FCanvasSortElement compare class
+	struct FCompareFCanvasSortElement
+	{
+		FORCEINLINE bool operator()(const FCanvasSortElement& A, const FCanvasSortElement& B) const
+		{
+			return B.DepthSortKey < A.DepthSortKey;
+		}
+	};
+	// sort the array of FCanvasSortElement entries so that higher sort keys render first (back-to-front)
+	SortedElements.Sort(FCompareFCanvasSortElement());
+
+	SCOPED_DRAW_EVENT(CanvasFlush, DEC_SCENE_ITEMS);
+	const FTexture2DRHIRef& RenderTargetTexture = RenderTarget->GetRenderTargetTexture();
+
+	check(IsValidRef(RenderTargetTexture));
+
+	// Set the RHI render target.
+	::SetRenderTarget(RHICmdList, RenderTargetTexture, FTextureRHIRef());
+	// disable depth test & writes
+	RHICmdList.SetDepthStencilState(TStaticDepthStencilState<false, CF_Always>::GetRHI());
+
+	if (ViewRect.Area() <= 0)
+	{
+		ViewRect = FIntRect(FIntPoint::ZeroValue, RenderTarget->GetSizeXY());
+	}
+
+	// set viewport to RT size
+	RHICmdList.SetViewport(ViewRect.Min.X, ViewRect.Min.Y, 0.0f, ViewRect.Max.X, ViewRect.Max.Y, 1.0f);
+
+	// iterate over the FCanvasSortElements in sorted order and render all the batched items for each entry
+	for (int32 Idx = 0; Idx < SortedElements.Num(); Idx++)
+	{
+		FCanvasSortElement& SortElement = SortedElements[Idx];
+		for (int32 BatchIdx = 0; BatchIdx < SortElement.RenderBatchArray.Num(); BatchIdx++)
+		{
+			FCanvasBaseRenderItem* RenderItem = SortElement.RenderBatchArray[BatchIdx];
+			if (RenderItem)
+			{
+				// mark current render target as dirty since we are drawing to it
+				bRenderTargetDirty |= RenderItem->Render_RenderThread(RHICmdList, this);
+				if (AllowedModes & Allow_DeleteOnRender)
+				{
+					delete RenderItem;
+				}
+			}
+		}
+		if (AllowedModes & Allow_DeleteOnRender)
+		{
+			SortElement.RenderBatchArray.Empty();
+		}
+	}
+	if (AllowedModes & Allow_DeleteOnRender)
+	{
+		// empty the array of FCanvasSortElement entries after finished with rendering	
+		SortedElements.Empty();
+		SortedElementLookupMap.Empty();
+		LastElementIndex = INDEX_NONE;
+	}
+}
+
+void FCanvas::Flush_GameThread(bool bForce)
 {
 	SCOPE_CYCLE_COUNTER(STAT_Canvas_FlushTime);
 
@@ -494,61 +618,36 @@ void FCanvas::Flush(bool bForce)
 	};
 	// sort the array of FCanvasSortElement entries so that higher sort keys render first (back-to-front)
 	SortedElements.Sort( FCompareFCanvasSortElement() );
-	if( GRenderingThread && IsInRenderingThread() )
+	if( ViewRect.Area() <= 0 )
 	{
-		FRHICommandListImmediate& RHICmdList = FRHICommandListExecutor::GetImmediateCommandList();
+		ViewRect = FIntRect( FIntPoint::ZeroValue, RenderTarget->GetSizeXY() );
+	}
 
+	struct FCanvasFlushParameters
+	{
+		FIntRect ViewRect;
+		const FRenderTarget* CanvasRenderTarget;
+	};
+	FCanvasFlushParameters FlushParameters =
+	{
+		ViewRect,
+		RenderTarget
+	};
+	ENQUEUE_UNIQUE_RENDER_COMMAND_ONEPARAMETER(
+		CanvasFlushSetupCommand,
+		FCanvasFlushParameters,Parameters,FlushParameters,
+	{
 		SCOPED_DRAW_EVENT(CanvasFlush, DEC_SCENE_ITEMS);
-		const FTexture2DRHIRef& RenderTargetTexture = RenderTarget->GetRenderTargetTexture();
-
-		check(IsValidRef(RenderTargetTexture));
 
 		// Set the RHI render target.
-		::SetRenderTarget(RHICmdList, RenderTargetTexture, FTextureRHIRef());
+		::SetRenderTarget(RHICmdList, Parameters.CanvasRenderTarget->GetRenderTargetTexture(), FTextureRHIRef());
 		// disable depth test & writes
-		RHICmdList.SetDepthStencilState(TStaticDepthStencilState<false, CF_Always>::GetRHI());
-
-		if( ViewRect.Area() <= 0 )
-		{
-			ViewRect = FIntRect( FIntPoint::ZeroValue, RenderTarget->GetSizeXY() );
-		}
+		RHICmdList.SetDepthStencilState(TStaticDepthStencilState<false,CF_Always>::GetRHI());
 
+		const FIntRect& ViewportRect = Parameters.ViewRect;
 		// set viewport to RT size
-		RHICmdList.SetViewport( ViewRect.Min.X,ViewRect.Min.Y, 0.0f, ViewRect.Max.X, ViewRect.Max.Y, 1.0f );	
-	}
-	else 
-	{
-		if( ViewRect.Area() <= 0 )
-		{
-			ViewRect = FIntRect( FIntPoint::ZeroValue, RenderTarget->GetSizeXY() );
-		}
-
-		struct FCanvasFlushParameters
-		{
-			FIntRect ViewRect;
-			const FRenderTarget* CanvasRenderTarget;
-		};
-		FCanvasFlushParameters FlushParameters =
-		{
-			ViewRect,
-			RenderTarget
-		};
-		ENQUEUE_UNIQUE_RENDER_COMMAND_ONEPARAMETER(
-			CanvasFlushSetupCommand,
-			FCanvasFlushParameters,Parameters,FlushParameters,
-		{
-			SCOPED_DRAW_EVENT(CanvasFlush, DEC_SCENE_ITEMS);
-
-			// Set the RHI render target.
-			::SetRenderTarget(RHICmdList, Parameters.CanvasRenderTarget->GetRenderTargetTexture(), FTextureRHIRef());
-			// disable depth test & writes
-			RHICmdList.SetDepthStencilState(TStaticDepthStencilState<false,CF_Always>::GetRHI());
-
-			const FIntRect& ViewportRect = Parameters.ViewRect;
-			// set viewport to RT size
-			RHICmdList.SetViewport(ViewportRect.Min.X, ViewportRect.Min.Y, 0.0f, ViewportRect.Max.X, ViewportRect.Max.Y, 1.0f);
-		});
-	}
+		RHICmdList.SetViewport(ViewportRect.Min.X, ViewportRect.Min.Y, 0.0f, ViewportRect.Max.X, ViewportRect.Max.Y, 1.0f);
+	});
 
 	// iterate over the FCanvasSortElements in sorted order and render all the batched items for each entry
 	for( int32 Idx=0; Idx < SortedElements.Num(); Idx++ )
@@ -560,7 +659,7 @@ void FCanvas::Flush(bool bForce)
 			if( RenderItem )
 			{
 				// mark current render target as dirty since we are drawing to it
-				bRenderTargetDirty |= RenderItem->Render(this);
+				bRenderTargetDirty |= RenderItem->Render_GameThread(this);
 				if( AllowedModes & Allow_DeleteOnRender )
 				{
 					delete RenderItem;
@@ -644,14 +743,14 @@ void FCanvas::CopyTransformStack(const FCanvas& Copy)
 	TransformStack = Copy.TransformStack;
 }
 
-void FCanvas::SetRenderTarget(FRenderTarget* NewRenderTarget)
+void FCanvas::SetRenderTarget_GameThread(FRenderTarget* NewRenderTarget)
 {
 	if( RenderTarget != NewRenderTarget )
 	{
 		// flush whenever we swap render targets
 		if( RenderTarget )
 		{
-			Flush();
+			Flush_GameThread();
 		}
 		// Change the current render target.
 		RenderTarget = NewRenderTarget;
@@ -1413,48 +1512,6 @@ void UCanvas::Deproject(FVector2D ScreenPos, /*out*/ FVector& WorldOrigin, /*out
 }
 
 
-
-
-void FCanvas::PushMaskRegion( float X, float Y, float SizeX, float SizeY )
-{
-	FMaskRegion NewMask(X, Y, SizeX, SizeY, TransformStack.Top().GetMatrix());
-	if ( !NewMask.IsEqual(GetCurrentMaskRegion()) )
-	{
-		Flush();
-	}
-
-	MaskRegionStack.Push(NewMask);
-}
-
-
-void FCanvas::PopMaskRegion()
-{
-	FMaskRegion NextMaskRegion = MaskRegionStack.Num() > 1 
-		? MaskRegionStack[MaskRegionStack.Num() - 2]
-		: FMaskRegion();
-
-	if ( !NextMaskRegion.IsEqual(GetCurrentMaskRegion()) )
-	{
-		Flush();
-	}
-
-	if ( MaskRegionStack.Num() > 0 )
-	{
-		MaskRegionStack.Pop();
-	}
-}
-
-FCanvas::FMaskRegion FCanvas::GetCurrentMaskRegion() const
-{
-	if ( MaskRegionStack.Num() > 0 )
-	{
-		return MaskRegionStack[MaskRegionStack.Num() - 1];
-	}
-
-	return FMaskRegion();
-}
-
-
 FFontRenderInfo UCanvas::CreateFontRenderInfo(bool bClipText, bool bEnableShadow, FLinearColor GlowColor, FVector2D GlowOuterRadius, FVector2D GlowInnerRadius)
 {
 	FFontRenderInfo Result;
diff --git a/Engine/Source/Runtime/Engine/Public/Canvas.h b/Engine/Source/Runtime/Engine/Public/Canvas.h
--- a/Engine/Source/Runtime/Engine/Public/Canvas.h
+++ b/Engine/Source/Runtime/Engine/Public/Canvas.h
@@ -74,9 +74,16 @@ public:
 
 	/** 
 	* Sends a message to the rendering thread to draw the batched elements. 
+	* @param RHICmdList - command list to use
 	* @param bForce - force the flush even if Allow_Flush is not enabled
 	*/
-	ENGINE_API void Flush(bool bForce=false);
+	ENGINE_API void Flush_RenderThread(FRHICommandListImmediate& RHICmdList, bool bForce = false);
+
+	/**
+	* Sends a message to the rendering thread to draw the batched elements.
+	* @param bForce - force the flush even if Allow_Flush is not enabled
+	*/
+	ENGINE_API void Flush_GameThread(bool bForce = false);
 
 	/**
 	 * Pushes a transform onto the canvas's transform stack, multiplying it with the current top of the stack.
@@ -177,28 +184,10 @@ public:
 	 **/
 	ENGINE_API void CopyTransformStack(const FCanvas& Copy);
 
-	/** 
-	 * Set the current masked region on the canvas
-	 * All rendering from this point on will be masked to this region.
-	 * The region being masked uses the current canvas transform
-	 *
-	 * @param X - x offset in canvas coords
-	 * @param Y - y offset in canvas coords
-	 * @param SizeX - x size in canvas coords
-	 * @param SizeY - y size in canvas coords
-	 */
-	ENGINE_API void PushMaskRegion(float X, float Y, float SizeX, float SizeY);
-
-	/**
-	 * Remove the current masking region; if other masking regions were previously pushed onto the stack,
-	 * the next one down will be activated.
-	 */
-	ENGINE_API void PopMaskRegion();
-
 	/**
 	 * Sets the render target which will be used for subsequent canvas primitives.
 	 */
-	ENGINE_API void SetRenderTarget(FRenderTarget* NewRenderTarget);	
+	ENGINE_API void SetRenderTarget_GameThread(FRenderTarget* NewRenderTarget);
 
 	/**
 	* Get the current render target for the canvas
@@ -370,42 +359,6 @@ private:
 	*/
 	void Construct();
 
-	/** 
-	* Region on the canvas that should be masked
-	*/
-	struct FMaskRegion
-	{
-		FMaskRegion(float InX=-1,float InY=-1,float InSizeX=-1,float InSizeY=-1, const FMatrix& InTransform=FMatrix::Identity) 
-			: X(InX), Y(InY), SizeX(InSizeX), SizeY(InSizeY), Transform(InTransform) 
-		{}
-		FORCEINLINE bool IsEqual(const FMaskRegion& R) const		
-		{ 
-			return(	FMath::Abs(X-R.X) < KINDA_SMALL_NUMBER && 
-					FMath::Abs(Y-R.Y) < KINDA_SMALL_NUMBER && 
-					FMath::Abs(SizeX-R.SizeX) < KINDA_SMALL_NUMBER && 
-					FMath::Abs(SizeY-R.SizeY) < KINDA_SMALL_NUMBER );
-		}
-		bool IsValid() const 
-		{ 
-			return X >= -DELTA && Y >= -DELTA && SizeX >= -DELTA && SizeY >= -DELTA; 
-		}
-		FORCEINLINE bool IsZero(float Tolerance=DELTA) const
-		{
-			//@todo - do we need to check tranform?
-			return FMath::Abs(X) < FMath::Abs(Tolerance)
-				&& FMath::Abs(Y) < FMath::Abs(Tolerance)
-				&& FMath::Abs(SizeX) < FMath::Abs(Tolerance)
-				&& FMath::Abs(SizeY) < FMath::Abs(Tolerance);
-		}
-		float X,Y,SizeX,SizeY;
-		FMatrix Transform;
-	};
-
-	/**
-	 * Stack of mask regions - top of stack (last element) is current canvas mask
-	 */
-	TArray<FMaskRegion> MaskRegionStack;
-
 public:	
 
 	/*
@@ -455,11 +408,6 @@ public:
 		Item.Draw( this, X, Y );
 	}
 
-	/**
-	 * Get the top-most canvas masking region from the stack.
-	 */
-	ENGINE_API FMaskRegion GetCurrentMaskRegion() const;
-
 	/**
 	* Clear the canvas
 	*
@@ -563,13 +511,21 @@ public:
 	virtual ~FCanvasBaseRenderItem()
 	{}
 
+	/**
+	* Renders the canvas item
+	*
+	* @param Canvas - canvas currently being rendered
+	* @param RHICmdList - command list to use
+	* @return true if anything rendered
+	*/
+	virtual bool Render_RenderThread(FRHICommandListImmediate& RHICmdList, const FCanvas* Canvas) = 0;
 	/**
 	* Renders the canvas item
 	*
 	* @param Canvas - canvas currently being rendered
 	* @return true if anything rendered
 	*/
-	virtual bool Render( const FCanvas* Canvas ) =0;
+	virtual bool Render_GameThread(const FCanvas* Canvas) = 0;
 	/**
 	* FCanvasBatchedElementRenderItem instance accessor
 	*
@@ -628,9 +584,18 @@ public:
 	* Iterates over all batched elements and draws them with their own transforms
 	*
 	* @param Canvas - canvas currently being rendered
+	* @param RHICmdList - command list to use
 	* @return true if anything rendered
 	*/
-	virtual bool Render( const FCanvas* Canvas );
+	virtual bool Render_RenderThread(FRHICommandListImmediate& RHICmdList, const FCanvas* Canvas) override;
+	/**
+	* Renders the canvas item.
+	* Iterates over all batched elements and draws them with their own transforms
+	*
+	* @param Canvas - canvas currently being rendered
+	* @return true if anything rendered
+	*/
+	virtual bool Render_GameThread(const FCanvas* Canvas) override;
 
 	/**
 	* Determine if this is a matching set by comparing texture,blendmode,elementype,transform. All must match
@@ -748,9 +713,19 @@ public:
 	* Iterates over each tile to be rendered and draws it with its own transforms
 	*
 	* @param Canvas - canvas currently being rendered
+	* @param RHICmdList - command list to use
+	* @return true if anything rendered
+	*/
+	virtual bool Render_RenderThread(FRHICommandListImmediate& RHICmdList, const FCanvas* Canvas) override;
+
+	/**
+	* Renders the canvas item.
+	* Iterates over each tile to be rendered and draws it with its own transforms
+	*
+	* @param Canvas - canvas currently being rendered
 	* @return true if anything rendered
 	*/
-	virtual bool Render( const FCanvas* Canvas );
+	virtual bool Render_GameThread(const FCanvas* Canvas) override;
 
 	/**
 	* Determine if this is a matching set by comparing material,transform. All must match
diff --git a/Engine/Source/Runtime/Renderer/Private/PostProcess/PostProcessGBufferHints.cpp b/Engine/Source/Runtime/Renderer/Private/PostProcess/PostProcessGBufferHints.cpp
--- a/Engine/Source/Runtime/Renderer/Private/PostProcess/PostProcessGBufferHints.cpp
+++ b/Engine/Source/Runtime/Renderer/Private/PostProcess/PostProcessGBufferHints.cpp
@@ -169,7 +169,7 @@ void FRCPassPostProcessGBufferHints::Process(FRenderingCompositePassContext& Con
 	Line = FString::Printf(TEXT("Red: Impossive material (this material emits more light than it receives)"));
 	Canvas.DrawShadowedString( X, Y += YStep, *Line, GetStatsFont(), FLinearColor(1, 0, 0));
 
-	Canvas.Flush();
+	Canvas.Flush_RenderThread(Context.RHICmdList);
 
 	Context.RHICmdList.CopyToResolveTarget(DestRenderTarget.TargetableTexture, DestRenderTarget.ShaderResourceTexture, false, FResolveParams());
 }
diff --git a/Engine/Source/Runtime/Renderer/Private/PostProcess/PostProcessTestImage.cpp b/Engine/Source/Runtime/Renderer/Private/PostProcess/PostProcessTestImage.cpp
--- a/Engine/Source/Runtime/Renderer/Private/PostProcess/PostProcessTestImage.cpp
+++ b/Engine/Source/Runtime/Renderer/Private/PostProcess/PostProcessTestImage.cpp
@@ -193,7 +193,7 @@ void FRCPassPostProcessTestImage::Process(FRenderingCompositePassContext& Contex
 		Line = FString::Printf(TEXT("   4 white pixel sized lines (only visible without overscan)"));
 		Canvas.DrawShadowedString( X, Y += YStep, *Line, GetStatsFont(), FLinearColor(1, 1, 1));
 
-		Canvas.Flush();
+		Canvas.Flush_RenderThread(Context.RHICmdList);
 	}
 
 
diff --git a/Engine/Source/Runtime/Renderer/Private/PostProcess/PostProcessVisualizeBuffer.cpp b/Engine/Source/Runtime/Renderer/Private/PostProcess/PostProcessVisualizeBuffer.cpp
--- a/Engine/Source/Runtime/Renderer/Private/PostProcess/PostProcessVisualizeBuffer.cpp
+++ b/Engine/Source/Runtime/Renderer/Private/PostProcess/PostProcessVisualizeBuffer.cpp
@@ -258,7 +258,7 @@ void FRCPassPostProcessVisualizeBuffer::Process(FRenderingCompositePassContext&
 	{
 		Canvas.DrawShadowedString(It->LocationX, It->LocationY, *It->Label, GetStatsFont(), LabelColor);
 	}
-	Canvas.Flush();
+	Canvas.Flush_RenderThread(Context.RHICmdList);
 
 
 	Context.RHICmdList.CopyToResolveTarget(DestRenderTarget.TargetableTexture, DestRenderTarget.ShaderResourceTexture, false, FResolveParams());
diff --git a/Engine/Source/Runtime/Renderer/Private/PostProcess/PostProcessVisualizeHDR.cpp b/Engine/Source/Runtime/Renderer/Private/PostProcess/PostProcessVisualizeHDR.cpp
--- a/Engine/Source/Runtime/Renderer/Private/PostProcess/PostProcessVisualizeHDR.cpp
+++ b/Engine/Source/Runtime/Renderer/Private/PostProcess/PostProcessVisualizeHDR.cpp
@@ -261,7 +261,7 @@ void FRCPassPostProcessVisualizeHDR::Process(FRenderingCompositePassContext& Con
 		*LogToString(View.FinalPostProcessSettings.HistogramLogMin), *LogToString(View.FinalPostProcessSettings.HistogramLogMax));
 	Canvas.DrawShadowedString( X + ColumnWidth, Y+= YStep, *Line, GetStatsFont(), FLinearColor(0.3f, 0.3f, 1));
 
-	Canvas.Flush();
+	Canvas.Flush_RenderThread(Context.RHICmdList);
 
 	Context.RHICmdList.CopyToResolveTarget(DestRenderTarget.TargetableTexture, DestRenderTarget.ShaderResourceTexture, false, FResolveParams());
 }
diff --git a/Engine/Source/Runtime/Renderer/Private/PostProcess/RenderTargetPool.cpp b/Engine/Source/Runtime/Renderer/Private/PostProcess/RenderTargetPool.cpp
--- a/Engine/Source/Runtime/Renderer/Private/PostProcess/RenderTargetPool.cpp
+++ b/Engine/Source/Runtime/Renderer/Private/PostProcess/RenderTargetPool.cpp
@@ -854,7 +854,7 @@ void FRenderTargetPool::PresentContent(FRHICommandListImmediate& RHICmdList, con
 				}
 			}
 
-			Canvas.Flush();
+			Canvas.Flush_RenderThread(RHICmdList);
 
 			CurrentEventRecordingTime = 0;
 			RenderTargetPoolEvents.Empty();
diff --git a/Engine/Source/Runtime/Renderer/Private/PostProcess/VisualizeTexture.cpp b/Engine/Source/Runtime/Renderer/Private/PostProcess/VisualizeTexture.cpp
--- a/Engine/Source/Runtime/Renderer/Private/PostProcess/VisualizeTexture.cpp
+++ b/Engine/Source/Runtime/Renderer/Private/PostProcess/VisualizeTexture.cpp
@@ -610,7 +610,7 @@ void FVisualizeTexture::PresentContent(FRHICommandListImmediate& RHICmdList, con
 		Canvas.DrawShadowedString( X, Y += YStep, TEXT("Blinking Blue: NAN or Inf"), GetStatsFont(), FLinearColor(0,0,1));
 	}
 
-	Canvas.Flush();
+	Canvas.Flush_RenderThread(RHICmdList);
 }
 
 
diff --git a/Engine/Source/Runtime/Renderer/Private/SceneRendering.cpp b/Engine/Source/Runtime/Renderer/Private/SceneRendering.cpp
--- a/Engine/Source/Runtime/Renderer/Private/SceneRendering.cpp
+++ b/Engine/Source/Runtime/Renderer/Private/SceneRendering.cpp
@@ -746,7 +746,7 @@ void FSceneRenderer::RenderFinish(FRHICommandListImmediate& RHICmdList)
 					Canvas.DrawShadowedText( 10, Y, Message, GetStatsFont(), FLinearColor(0.8,1.0,0.2,1.0));
 					Y += 14;
 				}
-				Canvas.Flush();
+				Canvas.Flush_RenderThread(RHICmdList);
 			}
 		}
 	}
@@ -987,7 +987,7 @@ static void RenderViewFamily_RenderThread(FRHICommandListImmediate& RHICmdList,
 void FRendererModule::BeginRenderingViewFamily(FCanvas* Canvas,const FSceneViewFamily* ViewFamily)
 {
 	// Flush the canvas first.
-	Canvas->Flush();
+	Canvas->Flush_GameThread();
 
 	// Increment FrameNumber before render the scene. Wrapping around is no problem.
 	// This is the only spot we change GFrameNumber, other places can only read.
diff --git a/Engine/Source/Runtime/StreamingPauseRendering/Private/StreamingPauseRendering.cpp b/Engine/Source/Runtime/StreamingPauseRendering/Private/StreamingPauseRendering.cpp
--- a/Engine/Source/Runtime/StreamingPauseRendering/Private/StreamingPauseRendering.cpp
+++ b/Engine/Source/Runtime/StreamingPauseRendering/Private/StreamingPauseRendering.cpp
@@ -81,7 +81,7 @@ void FStreamingPauseRenderingModule::BeginStreamingPause( FViewport* GameViewpor
 		{
 			ViewportClient->Draw(Viewport.Get(), &Canvas);
 		}
-		Canvas.Flush();
+		Canvas.Flush_GameThread();
 
 		//Don't need debug canvas I presume?
 		//Viewport->GetDebugCanvas()->Flush(true);
