Removing bLoopingInterpolation

diff --git a/Engine/Plugins/Experimental/Paper2D/Source/Paper2DEditor/Private/FlipbookEditor/FlipbookEditor.cpp b/Engine/Plugins/Experimental/Paper2D/Source/Paper2DEditor/Private/FlipbookEditor/FlipbookEditor.cpp
--- a/Engine/Plugins/Experimental/Paper2D/Source/Paper2DEditor/Private/FlipbookEditor/FlipbookEditor.cpp
+++ b/Engine/Plugins/Experimental/Paper2D/Source/Paper2DEditor/Private/FlipbookEditor/FlipbookEditor.cpp
@@ -235,8 +235,7 @@ TSharedRef<SDockTab> FFlipbookEditor::SpawnTab_Viewport(const FSpawnTabArgs& Arg
 		.ViewInputMax(this, &FFlipbookEditor::GetViewRangeMax)
 		.OnSetInputViewRange(this, &FFlipbookEditor::SetViewRange)
 		.bAllowZoom(true)
-		.IsRealtimeStreamingMode(false)
-		.bLastFrameIsFirstFrame(false);
+		.IsRealtimeStreamingMode(false);
 
 	return SNew(SDockTab)
 		.Label(LOCTEXT("ViewportTab_Title", "Viewport"))
diff --git a/Engine/Source/Editor/KismetWidgets/Private/SScrubControlPanel.cpp b/Engine/Source/Editor/KismetWidgets/Private/SScrubControlPanel.cpp
--- a/Engine/Source/Editor/KismetWidgets/Private/SScrubControlPanel.cpp
+++ b/Engine/Source/Editor/KismetWidgets/Private/SScrubControlPanel.cpp
@@ -59,7 +59,6 @@ void SScrubControlPanel::Construct( const SScrubControlPanel::FArguments& InArgs
 				.bAllowZoom(InArgs._bAllowZoom)
 				.DraggableBars(InArgs._DraggableBars)
 				.OnBarDrag(InArgs._OnBarDrag)
-				.bLastFrameIsFirstFrame(InArgs._bLastFrameIsFirstFrame)
 			]
 		]
 
diff --git a/Engine/Source/Editor/KismetWidgets/Private/SScrubWidget.cpp b/Engine/Source/Editor/KismetWidgets/Private/SScrubWidget.cpp
--- a/Engine/Source/Editor/KismetWidgets/Private/SScrubWidget.cpp
+++ b/Engine/Source/Editor/KismetWidgets/Private/SScrubWidget.cpp
@@ -9,19 +9,11 @@
 static const float MinStepLengh=15.f;
 
 /** This function is used by a few random widgets and is mostly arbitrary. It could be moved anywhere. */
-int32 SScrubWidget::GetDivider(float InputMinX, float InputMaxX, FVector2D WidgetSize, float SequenceLength, int32 NumFrames, bool bLastFrameIsFirstFrame)
+int32 SScrubWidget::GetDivider(float InputMinX, float InputMaxX, FVector2D WidgetSize, float SequenceLength, int32 NumFrames)
 {
 	FTrackScaleInfo TimeScaleInfo(InputMinX, InputMaxX, 0.f, 0.f, WidgetSize);
 
-	float TimePerKey = 0.0f;
-	if (bLastFrameIsFirstFrame)
-	{
-		TimePerKey = (NumFrames > 1) ? SequenceLength / (float)(NumFrames - 1) : 0.0f;
-	}
-	else
-	{
-		TimePerKey = (NumFrames > 0) ? SequenceLength / (float)(NumFrames) : 0.0f;
-	}
+	float TimePerKey = (NumFrames > 0) ? SequenceLength / (float)(NumFrames) : 0.0f;
 
 	float TotalWidgetWidth = TimeScaleInfo.WidgetSize.X;
 	float NumKeys = TimeScaleInfo.ViewInputRange / TimePerKey;
@@ -60,7 +52,6 @@ void SScrubWidget::Construct( const SScrubWidget::FArguments& InArgs )
 	DraggingBar = false;
 
 	bAllowZoom = InArgs._bAllowZoom;
-	bLastFrameIsFirstFrame = InArgs._bLastFrameIsFirstFrame;
 }
 
 int32 SScrubWidget::OnPaint( const FGeometry& AllottedGeometry, const FSlateRect& MyClippingRect, FSlateWindowElementList& OutDrawElements, int32 LayerId, const FWidgetStyle& InWidgetStyle, bool bParentEnabled ) const
@@ -89,20 +80,12 @@ int32 SScrubWidget::OnPaint( const FGeometry& AllottedGeometry, const FSlateRect
 	if ( NumOfKeys.Get() > 0 && SequenceLength.Get() > 0)
 	{
 		FTrackScaleInfo TimeScaleInfo(ViewInputMin.Get(), ViewInputMax.Get(), 0.f, 0.f, AllottedGeometry.Size);
-		int32 Divider = SScrubWidget::GetDivider( ViewInputMin.Get(), ViewInputMax.Get(), AllottedGeometry.Size, SequenceLength.Get(), NumOfKeys.Get(), bLastFrameIsFirstFrame );
+		int32 Divider = SScrubWidget::GetDivider( ViewInputMin.Get(), ViewInputMax.Get(), AllottedGeometry.Size, SequenceLength.Get(), NumOfKeys.Get() );
 		float HalfDivider = Divider/2.f;
 		
 		int32 TotalNumKeys = NumOfKeys.Get();
 
-		float TimePerKey = 0.0f;
-		if (bLastFrameIsFirstFrame)
-		{
-			TimePerKey = (TotalNumKeys > 1) ? SequenceLength.Get() / (float)(TotalNumKeys - 1) : 0.0f;
-		}
-		else
-		{
-			TimePerKey = (TotalNumKeys > 0) ? SequenceLength.Get() / (float)(TotalNumKeys) : 0.0f;
-		}
+		float TimePerKey = (TotalNumKeys > 0) ? SequenceLength.Get() / (float)(TotalNumKeys) : 0.0f;
 
 		for (float KeyVal = 0; KeyVal < TotalNumKeys; KeyVal += HalfDivider)
 		{
diff --git a/Engine/Source/Editor/KismetWidgets/Public/SScrubControlPanel.h b/Engine/Source/Editor/KismetWidgets/Public/SScrubControlPanel.h
--- a/Engine/Source/Editor/KismetWidgets/Public/SScrubControlPanel.h
+++ b/Engine/Source/Editor/KismetWidgets/Public/SScrubControlPanel.h
@@ -31,7 +31,6 @@ public:
 		, _ViewInputMax()
 		, _OnSetInputViewRange()
 		, _OnCropAnimSequence()
-		, _bLastFrameIsFirstFrame(true)
 	{}
 		SLATE_ATTRIBUTE( float, Value )
 		SLATE_ATTRIBUTE( uint32, NumOfKeys )
@@ -68,7 +67,6 @@ public:
 		/** Optional, additional values to draw on the timeline **/
 		SLATE_ATTRIBUTE( TArray<float>, DraggableBars )
 		SLATE_EVENT( FOnScrubBarDrag, OnBarDrag)
-		SLATE_ARGUMENT( bool, bLastFrameIsFirstFrame )
 	SLATE_END_ARGS()
 
 	/**
diff --git a/Engine/Source/Editor/KismetWidgets/Public/SScrubWidget.h b/Engine/Source/Editor/KismetWidgets/Public/SScrubWidget.h
--- a/Engine/Source/Editor/KismetWidgets/Public/SScrubWidget.h
+++ b/Engine/Source/Editor/KismetWidgets/Public/SScrubWidget.h
@@ -30,7 +30,6 @@ public:
 		, _ViewInputMin()
 		, _OnSetInputViewRange()
 		, _OnCropAnimSequence()
-		, _bLastFrameIsFirstFrame(true)
 		{}
 		
 		/** The value to display */
@@ -55,7 +54,6 @@ public:
 		/** Optional, additional values to draw on the timeline **/
 		SLATE_ATTRIBUTE( TArray<float>, DraggableBars )
 		SLATE_EVENT( FOnScrubBarDrag, OnBarDrag)
-		SLATE_ARGUMENT( bool, bLastFrameIsFirstFrame )
 	SLATE_END_ARGS()
 
 
@@ -75,7 +73,7 @@ public:
 	virtual FReply OnMouseWheel( const FGeometry& MyGeometry, const FPointerEvent& MouseEvent ) override;
 	virtual FCursorReply OnCursorQuery( const FGeometry& MyGeometry, const FPointerEvent& CursorEvent ) const override;
 
-	static int32 GetDivider(float InputMinX, float InputMaxX, FVector2D WidgetSize, float SequenceLength, int32 NumFrames, bool bLastFrameIsFirstFrame = true);
+	static int32 GetDivider(float InputMinX, float InputMaxX, FVector2D WidgetSize, float SequenceLength, int32 NumFrames);
 protected:
 
 	/**
@@ -128,7 +126,5 @@ private:
 	/** Bar Dragging*/
 	int32 DraggableBarIndex;
 	bool DraggingBar;
-
-	bool bLastFrameIsFirstFrame;
 };
 
diff --git a/Engine/Source/Editor/Persona/Private/Persona.cpp b/Engine/Source/Editor/Persona/Private/Persona.cpp
--- a/Engine/Source/Editor/Persona/Private/Persona.cpp
+++ b/Engine/Source/Editor/Persona/Private/Persona.cpp
@@ -54,6 +54,11 @@
 
 #include "Particles/ParticleSystemComponent.h"
 
+#include "DesktopPlatformModule.h"
+#include "MainFrame.h"
+#include "AnimationCompressionPanel.h"
+#include "FbxAnimUtils.h"
+
 #define LOCTEXT_NAMESPACE "FPersona"
 
 /////////////////////////////////////////////////////
@@ -777,7 +782,7 @@ void FPersona::InitPersona(const EToolkitMode::Type Mode, const TSharedPtr< clas
 			// View
 			MenuBuilder.BeginSection("Persona", LOCTEXT("PersonaEditorMenu", "Animation" ) );
 			{
-				MenuBuilder.AddMenuEntry( FPersonaCommands::Get().RecordAnimation );
+			
 			}
 			MenuBuilder.EndSection();
 		}
@@ -791,6 +796,16 @@ void FPersona::InitPersona(const EToolkitMode::Type Mode, const TSharedPtr< clas
 				MenuBuilder.AddMenuEntry( FPersonaCommands::Get().RemoveUnusedBones );
 			}
 			MenuBuilder.EndSection();
+
+			//  Animation menu
+			MenuBuilder.BeginSection("Persona", LOCTEXT("PersonaAssetMenuMenu", "Animation"));
+			{
+				MenuBuilder.AddMenuEntry(FPersonaCommands::Get().ApplyCompression);
+				MenuBuilder.AddMenuEntry(FPersonaCommands::Get().ExportToFBX);
+				MenuBuilder.AddMenuEntry(FPersonaCommands::Get().RecordAnimation );
+				MenuBuilder.AddMenuEntry(FPersonaCommands::Get().AddLoopingInterpolation);
+			}
+			MenuBuilder.EndSection();
 		}
 	};
 
@@ -1057,6 +1072,27 @@ void FPersona::CreateDefaultCommands()
 		FIsActionButtonVisible::CreateSP( this, &FPersona::IsRecordAvailable )
 		);
 
+	ToolkitCommands->MapAction(FPersonaCommands::Get().ApplyCompression,
+		FExecuteAction::CreateSP(this, &FPersona::OnApplyCompression),
+		FCanExecuteAction::CreateSP(this, &FPersona::HasValidAnimationSequencePlaying),
+		FIsActionChecked(),
+		FIsActionButtonVisible::CreateSP(this, &FPersona::IsInAnimationMode)
+		);
+
+	ToolkitCommands->MapAction(FPersonaCommands::Get().ExportToFBX,
+		FExecuteAction::CreateSP(this, &FPersona::OnExportToFBX),
+		FCanExecuteAction::CreateSP(this, &FPersona::HasValidAnimationSequencePlaying),
+		FIsActionChecked(),
+		FIsActionButtonVisible::CreateSP(this, &FPersona::IsInAnimationMode)
+		);
+
+	ToolkitCommands->MapAction(FPersonaCommands::Get().AddLoopingInterpolation,
+		FExecuteAction::CreateSP(this, &FPersona::OnAddLoopingInterpolation),
+		FCanExecuteAction::CreateSP(this, &FPersona::HasValidAnimationSequencePlaying),
+		FIsActionChecked(),
+		FIsActionButtonVisible::CreateSP(this, &FPersona::IsInAnimationMode)
+		);
+
 	ToolkitCommands->MapAction( FPersonaCommands::Get().ChangeSkeletonPreviewMesh,
 		FExecuteAction::CreateSP( this, &FPersona::ChangeSkeletonPreviewMesh ),
 		FCanExecuteAction::CreateSP( this, &FPersona::CanChangeSkeletonPreviewMesh )
@@ -2338,6 +2374,209 @@ void FPersona::RecordAnimation()
 	}
 }
 
+void FPersona::OnApplyCompression()
+{
+	UAnimSequence * AnimSequence = Cast<UAnimSequence> (GetPreviewAnimationAsset());
+
+	if (AnimSequence)
+	{
+		TArray<TWeakObjectPtr<UAnimSequence>> AnimSequences;
+		AnimSequences.Add(AnimSequence);
+		ApplyCompression(AnimSequences);
+	}
+}
+
+void FPersona::OnExportToFBX()
+{
+	UAnimSequence * AnimSequence = Cast<UAnimSequence>(GetPreviewAnimationAsset());
+
+	if(AnimSequence)
+	{
+		TArray<TWeakObjectPtr<UAnimSequence>> AnimSequences;
+		AnimSequences.Add(AnimSequence);
+		ExportToFBX(AnimSequences);
+	}
+}
+
+void FPersona::OnAddLoopingInterpolation()
+{
+	UAnimSequence * AnimSequence = Cast<UAnimSequence>(GetPreviewAnimationAsset());
+
+	if(AnimSequence)
+	{
+		TArray<TWeakObjectPtr<UAnimSequence>> AnimSequences;
+		AnimSequences.Add(AnimSequence);
+		AddLoopingInterpolation(AnimSequences);
+	}
+}
+
+void FPersona::ApplyCompression(TArray<TWeakObjectPtr<UAnimSequence>> & AnimSequences)
+{
+	FDlgAnimCompression AnimCompressionDialog(AnimSequences);
+	AnimCompressionDialog.ShowModal();
+}
+
+void FPersona::ExportToFBX(TArray<TWeakObjectPtr<UAnimSequence>> & AnimSequences)
+{
+	IDesktopPlatform* DesktopPlatform = FDesktopPlatformModule::Get();
+
+	if(DesktopPlatform)
+	{
+		USkeletalMesh * PreviewMesh = GetPreviewMeshComponent()->SkeletalMesh;
+		if(PreviewMesh == NULL)
+		{
+			FMessageDialog::Open(EAppMsgType::Ok, LOCTEXT("ExportToFBXExportMissingPreviewMesh", "ERROR: Missing preview mesh"));
+			return;
+		}
+
+		if(AnimSequences.Num() > 0)
+		{
+			//Get parent window for dialogs
+			IMainFrameModule& MainFrameModule = FModuleManager::LoadModuleChecked<IMainFrameModule>(TEXT("MainFrame"));
+			const TSharedPtr<SWindow>& MainFrameParentWindow = MainFrameModule.GetParentWindow();
+
+			void* ParentWindowWindowHandle = NULL;
+
+			if(MainFrameParentWindow.IsValid() && MainFrameParentWindow->GetNativeWindow().IsValid())
+			{
+				ParentWindowWindowHandle = MainFrameParentWindow->GetNativeWindow()->GetOSWindowHandle();
+			}
+
+			//Cache anim file names
+			TArray<FString> AnimFileNames;
+			for(auto Iter = AnimSequences.CreateIterator(); Iter; ++Iter)
+			{
+				AnimFileNames.Add(Iter->Get()->GetName() + TEXT(".fbx"));
+			}
+
+			IPlatformFile& PlatformFile = FPlatformFileManager::Get().GetPlatformFile();
+			FString DestinationFolder;
+
+			const FString Title = LOCTEXT("ExportFBXsToFolderTitle", "Choose a destination folder for the FBX file(s)").ToString();
+
+			if(AnimSequences.Num() > 1)
+			{
+				bool bFolderValid = false;
+				// More than one file, just ask for directory
+				while(!bFolderValid)
+				{
+					const bool bFolderSelected = DesktopPlatform->OpenDirectoryDialog(
+						ParentWindowWindowHandle,
+						Title,
+						FEditorDirectories::Get().GetLastDirectory(ELastDirectory::GENERIC_EXPORT),
+						DestinationFolder
+						);
+
+					if(!bFolderSelected)
+					{
+						// User canceled, return
+						return;
+					}
+
+					FEditorDirectories::Get().SetLastDirectory(ELastDirectory::GENERIC_EXPORT, DestinationFolder);
+					FPaths::NormalizeFilename(DestinationFolder);
+
+					//Check whether there are any fbx filename conflicts in this folder
+					for(auto Iter = AnimFileNames.CreateIterator(); Iter; ++Iter)
+					{
+						FString& AnimFileName = *Iter;
+						FString FullPath = DestinationFolder + "/" + AnimFileName;
+
+						bFolderValid = true;
+						if(PlatformFile.FileExists(*FullPath))
+						{
+							FFormatNamedArguments Args;
+							Args.Add(TEXT("DestinationFolder"), FText::FromString(DestinationFolder));
+							const FText DialogMessage = FText::Format(LOCTEXT("ExportToFBXFileOverwriteMessage", "Exporting to '{DestinationFolder}' will cause one or more existing FBX files to be overwritten. Would you like to continue?"), Args);
+							EAppReturnType::Type DialogReturn = FMessageDialog::Open(EAppMsgType::YesNo, DialogMessage);
+							bFolderValid = (EAppReturnType::Yes == DialogReturn);
+							break;
+						}
+					}
+				}
+			}
+			else
+			{
+				// One file only, ask for full filename.
+				// Can set bFolderValid from the SaveFileDialog call as the window will handle 
+				// duplicate files for us.
+				TArray<FString> TempDestinationNames;
+				bool bSave = DesktopPlatform->SaveFileDialog(
+					ParentWindowWindowHandle,
+					Title,
+					FEditorDirectories::Get().GetLastDirectory(ELastDirectory::GENERIC_EXPORT),
+					AnimSequences[0]->GetName(),
+					"FBX  |*.fbx",
+					EFileDialogFlags::None,
+					TempDestinationNames
+					);
+
+				if(!bSave)
+				{
+					// Canceled
+					return;
+				}
+				check(TempDestinationNames.Num() == 1);
+				check(AnimFileNames.Num() == 1);
+
+				DestinationFolder = FPaths::GetPath(TempDestinationNames[0]);
+				AnimFileNames[0] = FPaths::GetCleanFilename(TempDestinationNames[0]);
+
+				FEditorDirectories::Get().SetLastDirectory(ELastDirectory::GENERIC_EXPORT, DestinationFolder);
+			}
+
+			EAppReturnType::Type DialogReturn = FMessageDialog::Open(EAppMsgType::YesNo, LOCTEXT("ExportToFBXExportPreviewMeshToo", "Would you like to export the current skeletal mesh with the animation(s)?"));
+			bool bSaveSkeletalMesh = EAppReturnType::Yes == DialogReturn;
+
+			const bool bShowCancel = false;
+			const bool bShowProgressDialog = true;
+			GWarn->BeginSlowTask(LOCTEXT("ExportToFBXProgress", "Exporting Animation(s) to FBX"), bShowProgressDialog, bShowCancel);
+
+			// make sure to use PreviewMesh, when export inside of Persona
+			const int32 NumberOfAnimations = AnimSequences.Num();
+			for(int32 i = 0; i < NumberOfAnimations; ++i)
+			{
+				GWarn->UpdateProgress(i, NumberOfAnimations);
+
+				UAnimSequence* AnimSequence = AnimSequences[i].Get();
+
+				FString FileName = FString::Printf(TEXT("%s/%s"), *DestinationFolder, *AnimFileNames[i]);
+
+				FbxAnimUtils::ExportAnimFbx(*FileName, AnimSequence, PreviewMesh, bSaveSkeletalMesh);
+			}
+
+			GWarn->EndSlowTask();
+		}
+	}
+}
+
+void FPersona::AddLoopingInterpolation(TArray<TWeakObjectPtr<UAnimSequence>> & AnimSequences)
+{
+	FText WarningMessage = LOCTEXT("AddLoopiingInterpolation", "This will add extra frame at the end of the animation with the first frame to create a better looping interpolation. This action isn't undoable. Would you like to proceed?");
+
+	if(FMessageDialog::Open(EAppMsgType::YesNo, WarningMessage) == EAppReturnType::Yes)
+	{
+		for(auto Animation : AnimSequences)
+		{
+			// get first frame and add to the last frame and go through track
+			// now calculating old animated space bases
+			Animation->AddLoopingInterpolation();
+		}
+	}
+}
+
+bool FPersona::HasValidAnimationSequencePlaying() const
+{
+	UAnimSequence * AnimSequence = Cast<UAnimSequence> (GetAnimationAssetBeingEdited());
+	
+	return (AnimSequence != NULL);
+}
+
+bool FPersona::IsInAnimationMode() const
+{
+	return (GetCurrentMode() == FPersonaModes::AnimationEditMode);
+}
+
 void FPersona::ChangeSkeletonPreviewMesh()
 {
 	// Menu option cannot be called unless this is valid
@@ -2513,7 +2752,7 @@ bool FPersona::CanRemoveBones() const
 
 bool FPersona::IsRecordAvailable() const
 {
-	return (GetCurrentMode() == FPersonaModes::AnimBlueprintEditMode);
+	return (GetCurrentMode() == FPersonaModes::AnimBlueprintEditMode || GetCurrentMode() == FPersonaModes::AnimationEditMode );
 }
 
 bool FPersona::IsEditable(UEdGraph* InGraph) const
diff --git a/Engine/Source/Editor/Persona/Private/Persona.h b/Engine/Source/Editor/Persona/Private/Persona.h
--- a/Engine/Source/Editor/Persona/Private/Persona.h
+++ b/Engine/Source/Editor/Persona/Private/Persona.h
@@ -482,6 +482,13 @@ public:
 		OnGenericDelete.RemoveAll(Widget);
 	}
 
+	/** Apply Compression to list of animations */
+	void ApplyCompression(TArray<TWeakObjectPtr<UAnimSequence>> & AnimSequences);
+	/** Export to FBX files of the list of animations */
+	void ExportToFBX(TArray<TWeakObjectPtr<UAnimSequence>> & AnimSequences);
+	/** Add looping interpolation to the list of animations */
+	void AddLoopingInterpolation(TArray<TWeakObjectPtr<UAnimSequence>> & AnimSequences);
+
 protected:
 	/** Undo Action**/
 	void UndoAction();
@@ -535,6 +542,13 @@ private:
 	bool IsRecordAvailable() const;
 	bool IsAnimationBeingRecorded() const;
 
+	/** Animation menu functions **/
+	void OnApplyCompression();
+	void OnExportToFBX();
+	void OnAddLoopingInterpolation();
+	bool HasValidAnimationSequencePlaying() const;
+	bool IsInAnimationMode() const;
+
 	/** Change skeleton preview mesh functions */
 	void ChangeSkeletonPreviewMesh();
 	bool CanChangeSkeletonPreviewMesh() const;
diff --git a/Engine/Source/Editor/Persona/Private/PersonaCommands.cpp b/Engine/Source/Editor/Persona/Private/PersonaCommands.cpp
--- a/Engine/Source/Editor/Persona/Private/PersonaCommands.cpp
+++ b/Engine/Source/Editor/Persona/Private/PersonaCommands.cpp
@@ -6,10 +6,14 @@
 
 void FPersonaCommands::RegisterCommands()
 {
-	// this is going to show twice if we have blueprint
-	UI_COMMAND( RecordAnimation, "Export to Animation", "Export current graph to animation", EUserInterfaceActionType::Button, FInputGesture() );
-
+	// skeleton menu
 	UI_COMMAND( ChangeSkeletonPreviewMesh, "Set Preview Mesh as Default", "Changes the skeletons default preview mesh to the current open preview mesh. The skeleton will require saving after this action.", EUserInterfaceActionType::Button, FInputGesture() );
-
 	UI_COMMAND( RemoveUnusedBones, "Remove Unused Bones from Skeleton", "Removes any bones from the skeleton that are not used by any of its meshes. The skeleton and associated animations will require saving after this action.", EUserInterfaceActionType::Button, FInputGesture() );
+
+	// animation menu
+	UI_COMMAND( ApplyCompression, "Apply Compression", "Apply cvompression to current animation", EUserInterfaceActionType::Button, FInputGesture() );
+	UI_COMMAND( RecordAnimation, "Record to new Animation", "Create new animation from currently playing", EUserInterfaceActionType::Button, FInputGesture() );
+	UI_COMMAND( ExportToFBX, "Export to FBX", "Export current animation to FBX", EUserInterfaceActionType::Button, FInputGesture() );
+	UI_COMMAND( AddLoopingInterpolation, "Add Looping Interpolation", "Add an extra first frame at the end of the frame to create interpolation when looping", EUserInterfaceActionType::Button, FInputGesture() );
+
 }
diff --git a/Engine/Source/Editor/Persona/Private/PersonaCommands.h b/Engine/Source/Editor/Persona/Private/PersonaCommands.h
--- a/Engine/Source/Editor/Persona/Private/PersonaCommands.h
+++ b/Engine/Source/Editor/Persona/Private/PersonaCommands.h
@@ -14,11 +14,23 @@ public:
 	virtual void RegisterCommands() override;
 
 public:
- 	TSharedPtr<FUICommandInfo> RecordAnimation;
-
+	// skeleton menu options
 	// Command to allow users to set the skeletons preview mesh
 	TSharedPtr<FUICommandInfo> ChangeSkeletonPreviewMesh;
 
 	// Command to allow users to remove unused bones (not referenced by any skeletalmesh) from the skeleton
 	TSharedPtr<FUICommandInfo> RemoveUnusedBones;
+
+	// animation menu options
+	// record animation 
+ 	TSharedPtr<FUICommandInfo> RecordAnimation;
+
+	// apply compression
+	TSharedPtr<FUICommandInfo> ApplyCompression;
+
+	// export to FBX
+	TSharedPtr<FUICommandInfo> ExportToFBX;
+
+	// Add looping interpolation
+	TSharedPtr<FUICommandInfo> AddLoopingInterpolation;
 };
diff --git a/Engine/Source/Editor/Persona/Private/SAnimationSequenceBrowser.cpp b/Engine/Source/Editor/Persona/Private/SAnimationSequenceBrowser.cpp
--- a/Engine/Source/Editor/Persona/Private/SAnimationSequenceBrowser.cpp
+++ b/Engine/Source/Editor/Persona/Private/SAnimationSequenceBrowser.cpp
@@ -7,13 +7,9 @@
 #include "SAnimationSequenceBrowser.h"
 #include "Persona.h"
 #include "AssetRegistryModule.h"
-#include "AnimationCompressionPanel.h"
 #include "SSkeletonWidget.h"
 #include "Editor/ContentBrowser/Public/ContentBrowserModule.h"
-#include "MainFrame.h"
-#include "DesktopPlatformModule.h"
 #include "FeedbackContextEditor.h"
-#include "FbxAnimUtils.h"
 #include "EditorAnimUtils.h"
 #include "Editor/ContentBrowser/Public/FrontendFilterBase.h"
 
@@ -108,6 +104,43 @@ TSharedPtr<SWidget> SAnimationSequenceBrowser::OnGetAssetContextMenu(const TArra
 		}
 	}
 
+	if(bHasSelectedAnimSequence)
+	{
+		MenuBuilder.BeginSection("AnimationSequenceOptions", NSLOCTEXT("Docking", "TabAnimationHeading", "Animation"));
+		{
+			MenuBuilder.AddMenuEntry(
+				LOCTEXT("RunCompressionOnAnimations", "Apply Compression"),
+				LOCTEXT("RunCompressionOnAnimations_ToolTip", "Apply a compression scheme from the options given to the selected animations"),
+				FSlateIcon(),
+				FUIAction(
+				FExecuteAction::CreateSP(this, &SAnimationSequenceBrowser::OnApplyCompression, SelectedAssets),
+				FCanExecuteAction()
+				)
+				);
+
+			MenuBuilder.AddMenuEntry(
+				LOCTEXT("ExportAnimationsToFBX", "Export to FBX"),
+				LOCTEXT("ExportAnimationsToFBX_ToolTip", "Export Animation(s) To FBX"),
+				FSlateIcon(),
+				FUIAction(
+				FExecuteAction::CreateSP(this, &SAnimationSequenceBrowser::OnExportToFBX, SelectedAssets),
+				FCanExecuteAction()
+				)
+				);
+
+			MenuBuilder.AddMenuEntry(
+				LOCTEXT("AddLoopingInterpolation", "Add Looping Interpolation"),
+				LOCTEXT("AddLoopingInterpolation_ToolTip", "Add an extra frame at the end of the animation to create better looping"),
+				FSlateIcon(),
+				FUIAction(
+				FExecuteAction::CreateSP(this, &SAnimationSequenceBrowser::OnAddLoopingInterpolation, SelectedAssets),
+				FCanExecuteAction()
+				)
+				);
+		}
+		MenuBuilder.EndSection();
+	}
+
 	MenuBuilder.BeginSection("AnimationSequenceOptions", NSLOCTEXT("Docking", "TabOptionsHeading", "Options") );
 	{
 		MenuBuilder.AddMenuEntry(
@@ -129,29 +162,6 @@ TSharedPtr<SWidget> SAnimationSequenceBrowser::OnGetAssetContextMenu(const TArra
 				FCanExecuteAction::CreateSP( this, &SAnimationSequenceBrowser::CanSaveSelectedAssets, SelectedAssets)
 				)
 			);
-
-		if(bHasSelectedAnimSequence)
-		{
-			MenuBuilder.AddMenuEntry(
-				LOCTEXT("RunCompressionOnAnimations", "Apply Compression"),
-				LOCTEXT("RunCompressionOnAnimations_ToolTip", "Apply a compression scheme from the options given to the selected animations"),
-				FSlateIcon(),
-				FUIAction(
-				FExecuteAction::CreateSP( this, &SAnimationSequenceBrowser::OnApplyCompression, SelectedAssets ),
-				FCanExecuteAction()
-				)
-				);
-
-			MenuBuilder.AddMenuEntry(
-				LOCTEXT("ExportAnimationsToFBX", "Export to FBX"),
-				LOCTEXT("ExportAnimationsToFBX_ToolTip", "Export Animation(s) To FBX"),
-				FSlateIcon(),
-				FUIAction(
-				FExecuteAction::CreateSP( this, &SAnimationSequenceBrowser::OnExportToFBX, SelectedAssets ),
-				FCanExecuteAction()
-				)
-				);
-		}
 	}
 	MenuBuilder.EndSection();
 
@@ -224,8 +234,7 @@ void SAnimationSequenceBrowser::OnApplyCompression(TArray<FAssetData> SelectedAs
 			}
 		}
 
-		FDlgAnimCompression AnimCompressionDialog( AnimSequences );
-		AnimCompressionDialog.ShowModal();
+		PersonaPtr.Pin()->ApplyCompression(AnimSequences);
 	}
 }
 
@@ -233,146 +242,35 @@ void SAnimationSequenceBrowser::OnExportToFBX(TArray<FAssetData> SelectedAssets)
 {
 	if (SelectedAssets.Num() > 0)
 	{
-		IDesktopPlatform* DesktopPlatform = FDesktopPlatformModule::Get();
-
-		if(DesktopPlatform && PersonaPtr.IsValid())
+		TArray<TWeakObjectPtr<UAnimSequence>> AnimSequences;
+		for(auto Iter = SelectedAssets.CreateIterator(); Iter; ++Iter)
 		{
-			USkeletalMesh * PreviewMesh = PersonaPtr.Pin()->GetPreviewMeshComponent()->SkeletalMesh;
-			if (PreviewMesh == NULL)
+			if(UAnimSequence* AnimSequence = Cast<UAnimSequence>(Iter->GetAsset()))
 			{
-				FMessageDialog::Open(EAppMsgType::Ok, LOCTEXT("ExportToFBXExportMissingPreviewMesh", "ERROR: Missing preview mesh") );
-				return;
+				// we only shows anim sequence that belong to this skeleton
+				AnimSequences.Add(TWeakObjectPtr<UAnimSequence>(AnimSequence));
 			}
+		}
 
-			TArray<TWeakObjectPtr<UAnimSequence>> AnimSequences;
-			for(auto Iter = SelectedAssets.CreateIterator(); Iter; ++Iter)
-			{
-				if(UAnimSequence* AnimSequence = Cast<UAnimSequence>(Iter->GetAsset()) )
-				{
-					// we only shows anim sequence that belong to this skeleton
-					AnimSequences.Add( TWeakObjectPtr<UAnimSequence>(AnimSequence) );
-				}
-			}
+		PersonaPtr.Pin()->ExportToFBX(AnimSequences);
+	}
+}
 
-			if(AnimSequences.Num() > 0)
+void SAnimationSequenceBrowser::OnAddLoopingInterpolation(TArray<FAssetData> SelectedAssets)
+{
+	if(SelectedAssets.Num() > 0)
+	{
+		TArray<TWeakObjectPtr<UAnimSequence>> AnimSequences;
+		for(auto Iter = SelectedAssets.CreateIterator(); Iter; ++Iter)
+		{
+			if(UAnimSequence* AnimSequence = Cast<UAnimSequence>(Iter->GetAsset()))
 			{
-				//Get parent window for dialogs
-				IMainFrameModule& MainFrameModule = FModuleManager::LoadModuleChecked<IMainFrameModule>(TEXT("MainFrame"));
-				const TSharedPtr<SWindow>& MainFrameParentWindow = MainFrameModule.GetParentWindow();
-
-				void* ParentWindowWindowHandle = NULL;
-
-				if ( MainFrameParentWindow.IsValid() && MainFrameParentWindow->GetNativeWindow().IsValid() )
-				{
-					ParentWindowWindowHandle = MainFrameParentWindow->GetNativeWindow()->GetOSWindowHandle();
-				}
-
-				//Cache anim file names
-				TArray<FString> AnimFileNames;
-				for(auto Iter = AnimSequences.CreateIterator(); Iter; ++Iter)
-				{
-					AnimFileNames.Add( Iter->Get()->GetName() + TEXT(".fbx") );
-				}
-
-				IPlatformFile& PlatformFile = FPlatformFileManager::Get().GetPlatformFile();
-				FString DestinationFolder;
-
-				const FString Title = LOCTEXT("ExportFBXsToFolderTitle", "Choose a destination folder for the FBX file(s)").ToString();
-
-				if(SelectedAssets.Num() > 1)
-				{
-					bool bFolderValid = false;
-					// More than one file, just ask for directory
-					while(!bFolderValid)
-					{
-						const bool bFolderSelected = DesktopPlatform->OpenDirectoryDialog(
-							ParentWindowWindowHandle,
-							Title,
-							FEditorDirectories::Get().GetLastDirectory(ELastDirectory::GENERIC_EXPORT),
-							DestinationFolder
-							);
-
-						if ( !bFolderSelected )
-						{
-							// User canceled, return
-							return;
-						}
-
-						FEditorDirectories::Get().SetLastDirectory(ELastDirectory::GENERIC_EXPORT, DestinationFolder);
-						FPaths::NormalizeFilename(DestinationFolder);
-
-						//Check whether there are any fbx filename conflicts in this folder
-						for(auto Iter = AnimFileNames.CreateIterator(); Iter; ++Iter)
-						{
-							FString& AnimFileName = *Iter;
-							FString FullPath = DestinationFolder + "/" + AnimFileName;
-
-							bFolderValid = true;
-							if(PlatformFile.FileExists(*FullPath))
-							{
-								FFormatNamedArguments Args;
-								Args.Add( TEXT("DestinationFolder"), FText::FromString( DestinationFolder ) );
-								const FText DialogMessage = FText::Format( LOCTEXT("ExportToFBXFileOverwriteMessage", "Exporting to '{DestinationFolder}' will cause one or more existing FBX files to be overwritten. Would you like to continue?"), Args );
-								EAppReturnType::Type DialogReturn = FMessageDialog::Open(EAppMsgType::YesNo, DialogMessage);
-								bFolderValid = (EAppReturnType::Yes == DialogReturn);
-								break;
-							}
-						}
-					}
-				}
-				else
-				{
-					// One file only, ask for full filename.
-					// Can set bFolderValid from the SaveFileDialog call as the window will handle 
-					// duplicate files for us.
-					TArray<FString> TempDestinationNames;
-					bool bSave = DesktopPlatform->SaveFileDialog(
-						ParentWindowWindowHandle,
-						Title,
-						FEditorDirectories::Get().GetLastDirectory(ELastDirectory::GENERIC_EXPORT),
-						SelectedAssets[0].AssetName.ToString(),
-						"FBX  |*.fbx",
-						EFileDialogFlags::None,
-						TempDestinationNames
-						);
-
-					if(!bSave)
-					{
-						// Canceled
-						return;
-					}
-					check(TempDestinationNames.Num() == 1);
-					check(AnimFileNames.Num() == 1);
-
-					DestinationFolder = FPaths::GetPath(TempDestinationNames[0]);
-					AnimFileNames[0] = FPaths::GetCleanFilename(TempDestinationNames[0]);
-
-					FEditorDirectories::Get().SetLastDirectory(ELastDirectory::GENERIC_EXPORT, DestinationFolder);
-				}
-
-				EAppReturnType::Type DialogReturn = FMessageDialog::Open(EAppMsgType::YesNo, LOCTEXT("ExportToFBXExportPreviewMeshToo", "Would you like to export the current skeletal mesh with the animation(s)?") );
-				bool bSaveSkeletalMesh = EAppReturnType::Yes == DialogReturn;
-
-				const bool bShowCancel = false;
-				const bool bShowProgressDialog = true;
-				GWarn->BeginSlowTask( LOCTEXT("ExportToFBXProgress", "Exporting Animation(s) to FBX"), bShowProgressDialog, bShowCancel);
-
-				// make sure to use PreviewMesh, when export inside of Persona
-				const int32 NumberOfAnimations = AnimSequences.Num();
-				for(int32 i = 0; i < NumberOfAnimations; ++i)
-				{
-					GWarn->UpdateProgress(i,NumberOfAnimations);
-
-					UAnimSequence* AnimSequence = AnimSequences[i].Get();
-
-					FString FileName = FString::Printf(TEXT("%s/%s"), *DestinationFolder, *AnimFileNames[i]);
-
-					FbxAnimUtils::ExportAnimFbx( *FileName, AnimSequence, PreviewMesh, bSaveSkeletalMesh );
-				}
-
-				GWarn->EndSlowTask( );
+				// we only shows anim sequence that belong to this skeleton
+				AnimSequences.Add(TWeakObjectPtr<UAnimSequence>(AnimSequence));
 			}
 		}
+
+		PersonaPtr.Pin()->AddLoopingInterpolation(AnimSequences);
 	}
 }
 
diff --git a/Engine/Source/Editor/Persona/Private/SAnimationSequenceBrowser.h b/Engine/Source/Editor/Persona/Private/SAnimationSequenceBrowser.h
--- a/Engine/Source/Editor/Persona/Private/SAnimationSequenceBrowser.h
+++ b/Engine/Source/Editor/Persona/Private/SAnimationSequenceBrowser.h
@@ -43,6 +43,9 @@ public:
 	/** Delegate to handle Export FBX context menu option */
 	void OnExportToFBX(TArray<FAssetData> SelectedAssets);
 
+	/** Delegate to handle Add looping interpolation context menu option */
+	void OnAddLoopingInterpolation(TArray<FAssetData> SelectedAssets);
+
 	/** This will allow duplicate the current object, and remap to new skeleton 
 	 *	Only allowed for AnimSequence 
 	 */
diff --git a/Engine/Source/Editor/UnrealEd/Private/Animation/DebugSkelMeshComponent.cpp b/Engine/Source/Editor/UnrealEd/Private/Animation/DebugSkelMeshComponent.cpp
--- a/Engine/Source/Editor/UnrealEd/Private/Animation/DebugSkelMeshComponent.cpp
+++ b/Engine/Source/Editor/UnrealEd/Private/Animation/DebugSkelMeshComponent.cpp
@@ -403,7 +403,7 @@ void UDebugSkelMeshComponent::RefreshBoneTransforms(FActorComponentTickFunction*
 			if (Sequence->IsValidAdditive()) 
 			{ 
 				AdditiveBasePoses.AddUninitialized(PreviewInstance->RequiredBones.GetNumBones());
-				Sequence->GetAdditiveBasePose(AdditiveBasePoses, PreviewInstance->RequiredBones, FAnimExtractContext(PreviewInstance->CurrentTime, PreviewInstance->bLooping));
+				Sequence->GetAdditiveBasePose(AdditiveBasePoses, PreviewInstance->RequiredBones, FAnimExtractContext(PreviewInstance->CurrentTime));
 				
 				FA2CSPose CSPose;
 				CSPose.AllocateLocalPoses(AnimScriptInstance->RequiredBones, AdditiveBasePoses);
diff --git a/Engine/Source/Editor/UnrealEd/Private/Fbx/FbxAnimationExport.cpp b/Engine/Source/Editor/UnrealEd/Private/Fbx/FbxAnimationExport.cpp
--- a/Engine/Source/Editor/UnrealEd/Private/Fbx/FbxAnimationExport.cpp
+++ b/Engine/Source/Editor/UnrealEd/Private/Fbx/FbxAnimationExport.cpp
@@ -52,8 +52,7 @@ void FFbxExporter::ExportAnimSequenceToFbx(const UAnimSequence* AnimSeq,
 									 float AnimStartOffset,
 									 float AnimEndOffset,
 									 float AnimPlayRate,
-									 float StartTime,
-									 bool bLooping)
+									 float StartTime)
 {
 	USkeleton * Skeleton = AnimSeq->GetSkeleton();
 
@@ -102,7 +101,7 @@ void FFbxExporter::ExportAnimSequenceToFbx(const UAnimSequence* AnimSeq,
 		while (!bLastKey)
 		{
 			FTransform BoneAtom;
-			AnimSeq->GetBoneTransform(BoneAtom, BoneTrackIndex, AnimTime, bLooping, true);
+			AnimSeq->GetBoneTransform(BoneAtom, BoneTrackIndex, AnimTime, true);
 
 			FbxVector4 Translation = Converter.ConvertToFbxPos(BoneAtom.GetTranslation());
 
@@ -217,8 +216,7 @@ void FFbxExporter::ExportAnimSequence( const UAnimSequence* AnimSeq, USkeletalMe
 			0.f,		// AnimStartOffset
 			0.f,		// AnimEndOffset
 			1.f,		// AnimPlayRate
-			0.f,		// StartTime
-			false);		// bLooping
+			0.f);		// StartTime
 
 		CorrectAnimTrackInterpolation(BoneNodes, AnimLayer);
 	}
@@ -288,8 +286,7 @@ void FFbxExporter::ExportAnimSequencesAsSingle( USkeletalMesh* SkelMesh, const A
 			TrackKey.AnimStartOffset,
 			TrackKey.AnimEndOffset,
 			TrackKey.AnimPlayRate,
-			CurrentStartTime,
-			TrackKey.bLooping);
+			CurrentStartTime);
 	}
 
 	CorrectAnimTrackInterpolation(BoneNodes, AnimLayer);
diff --git a/Engine/Source/Editor/UnrealEd/Private/FbxExporter.h b/Engine/Source/Editor/UnrealEd/Private/FbxExporter.h
--- a/Engine/Source/Editor/UnrealEd/Private/FbxExporter.h
+++ b/Engine/Source/Editor/UnrealEd/Private/FbxExporter.h
@@ -244,7 +244,7 @@ private:
 	 * Add the given animation sequence as rotation and translation tracks to the given list of bone nodes
 	 */
 	void ExportAnimSequenceToFbx(const UAnimSequence* AnimSeq, const USkeletalMesh* SkelMesh, TArray<FbxNode*>& BoneNodes, FbxAnimLayer* AnimLayer,
-		float AnimStartOffset, float AnimEndOffset, float AnimPlayRate, float StartTime, bool bLooping);
+		float AnimStartOffset, float AnimEndOffset, float AnimPlayRate, float StartTime);
 
 	/** 
 	 * The curve code doesn't differentiate between angles and other data, so an interpolation from 179 to -179
diff --git a/Engine/Source/Runtime/Engine/Classes/Animation/AnimInstance.h b/Engine/Source/Runtime/Engine/Classes/Animation/AnimInstance.h
--- a/Engine/Source/Runtime/Engine/Classes/Animation/AnimInstance.h
+++ b/Engine/Source/Runtime/Engine/Classes/Animation/AnimInstance.h
@@ -282,7 +282,7 @@ public:
 
 	void ApplyAdditiveSequence(const struct FA2Pose& BasePose, const struct FA2Pose& AdditivePose, float Alpha, struct FA2Pose& Blended);
 
-	void BlendSpaceEvaluatePose(class UBlendSpaceBase* BlendSpace, TArray<FBlendSampleData>& BlendSampleDataCache, struct FA2Pose& Pose, bool bIsLooping);
+	void BlendSpaceEvaluatePose(class UBlendSpaceBase* BlendSpace, TArray<FBlendSampleData>& BlendSampleDataCache, struct FA2Pose& Pose);
 
 	// skeletal control related functions
 	void BlendRotationOffset(const struct FA2Pose& BasePose/* local space base pose */, struct FA2Pose const & RotationOffsetPose/* mesh space rotation only additive **/, float Alpha/*0 means no additive, 1 means whole additive */, struct FA2Pose& Pose /** local space blended pose **/);
diff --git a/Engine/Source/Runtime/Engine/Classes/Animation/AnimSequence.h b/Engine/Source/Runtime/Engine/Classes/Animation/AnimSequence.h
--- a/Engine/Source/Runtime/Engine/Classes/Animation/AnimSequence.h
+++ b/Engine/Source/Runtime/Engine/Classes/Animation/AnimSequence.h
@@ -312,12 +312,6 @@ class UAnimSequence : public UAnimSequenceBase
 	UPROPERTY(AssetRegistrySearchable)
 	int32 NumFrames;
 
-	/**
-	 * if true, enable interpolation between last and first frame when looping.
-	 */
-	UPROPERTY(EditAnywhere, Category=Animation)
-	uint32 bLoopingInterpolation:1;
-
 	/**
 	 * In the future, maybe keeping RawAnimSequenceTrack + TrackMap as one would be good idea to avoid inconsistent array size
 	 * TrackToSkeletonMapTable(i) should contains  track mapping data for RawAnimationData(i). 
@@ -585,10 +579,9 @@ public:
 	 * @param	OutAtom			[out] Output bone transform.
 	 * @param	TrackIndex		Index of track to interpolate.
 	 * @param	Time			Time on track to interpolate to.
-	 * @param	bLooping		true if the animation is looping.
 	 * @param	bUseRawData		If true, use raw animation data instead of compressed data.
 	 */
-	ENGINE_API void GetBoneTransform(FTransform& OutAtom, int32 TrackIndex, float Time, bool bLooping, bool bUseRawData) const;
+	ENGINE_API void GetBoneTransform(FTransform& OutAtom, int32 TrackIndex, float Time, bool bUseRawData) const;
 	
 	// End Transform related functions 
 
@@ -688,6 +681,11 @@ public:
 	 * After imported or any other change is made, call this to apply post process
 	 */
 	ENGINE_API void PostProcessSequence();
+	/** 
+	 * Insert extra frame of the first frame at the end of the frame so that it improves the interpolation when it loops
+	 * This increases framecount + time, so that it requires recompression
+	 */
+	ENGINE_API bool AddLoopingInterpolation();
 #endif
 
 private:
diff --git a/Engine/Source/Runtime/Engine/Classes/Animation/AnimSingleNodeInstance.h b/Engine/Source/Runtime/Engine/Classes/Animation/AnimSingleNodeInstance.h
--- a/Engine/Source/Runtime/Engine/Classes/Animation/AnimSingleNodeInstance.h
+++ b/Engine/Source/Runtime/Engine/Classes/Animation/AnimSingleNodeInstance.h
@@ -61,7 +61,7 @@ class ENGINE_API UAnimSingleNodeInstance : public UAnimInstance
 	virtual bool NativeEvaluateAnimation(FPoseContext& Output) override;
 protected:
 	virtual void Montage_Advance(float DeltaTime) override;
-	void InternalBlendSpaceEvaluatePose(class UBlendSpaceBase* BlendSpace, TArray<FBlendSampleData>& BlendSampleDataCache, struct FA2Pose& Pose, bool bIsLooping);
+	void InternalBlendSpaceEvaluatePose(class UBlendSpaceBase* BlendSpace, TArray<FBlendSampleData>& BlendSampleDataCache, struct FA2Pose& Pose);
 	// End UAnimInstance interface
 public:
 
diff --git a/Engine/Source/Runtime/Engine/Classes/Animation/AnimationAsset.h b/Engine/Source/Runtime/Engine/Classes/Animation/AnimationAsset.h
--- a/Engine/Source/Runtime/Engine/Classes/Animation/AnimationAsset.h
+++ b/Engine/Source/Runtime/Engine/Classes/Animation/AnimationAsset.h
@@ -105,10 +105,6 @@ struct FAnimExtractContext
 	UPROPERTY()
 	bool bExtractRootMotionRotation;
 
-	/** Are we looping this animation? */
-	UPROPERTY()
-	bool bLooping;
-
 	/** Position in animation to extract pose from */
 	UPROPERTY()
 	float CurrentTime;
@@ -120,25 +116,22 @@ struct FAnimExtractContext
 	FAnimExtractContext()
 		: bExtractRootMotionTranslation(false)
 		, bExtractRootMotionRotation(false)
-		, bLooping(false)
 		, CurrentTime(0.f)
 		, RootMotionRootLock(ERootMotionRootLock::RefPose)
 	{
 	}
 
-	FAnimExtractContext(float InCurrentTime, bool InbLooping)
+	FAnimExtractContext(float InCurrentTime)
 		: bExtractRootMotionTranslation(false)
 		, bExtractRootMotionRotation(false)
-		, bLooping(InbLooping)
 		, CurrentTime(InCurrentTime)
 		, RootMotionRootLock(ERootMotionRootLock::RefPose)
 	{
 	}
 
-	FAnimExtractContext(float InCurrentTime, bool InbLooping, bool InbExtractRootMotionTranslation, bool InbExtractRootMotionRotation, ERootMotionRootLock::Type InRootMotionRootLock)
+	FAnimExtractContext(float InCurrentTime, bool InbExtractRootMotionTranslation, bool InbExtractRootMotionRotation, ERootMotionRootLock::Type InRootMotionRootLock)
 		: bExtractRootMotionTranslation(InbExtractRootMotionTranslation)
 		, bExtractRootMotionRotation(InbExtractRootMotionRotation)
-		, bLooping(InbLooping)
 		, CurrentTime(InCurrentTime)
 		, RootMotionRootLock(InRootMotionRootLock)
 	{
diff --git a/Engine/Source/Runtime/Engine/Private/Animation/AnimCompress_RemoveLinearKeys.cpp b/Engine/Source/Runtime/Engine/Private/Animation/AnimCompress_RemoveLinearKeys.cpp
--- a/Engine/Source/Runtime/Engine/Private/Animation/AnimCompress_RemoveLinearKeys.cpp
+++ b/Engine/Source/Runtime/Engine/Private/Animation/AnimCompress_RemoveLinearKeys.cpp
@@ -328,7 +328,7 @@ void UAnimCompress_RemoveLinearKeys::UpdateWorldBoneTransformTable(
 			float Time = (float)FrameIndex * TimePerFrame;
 			FTransform LocalAtom;
 
-			AnimSeq->GetBoneTransform(LocalAtom, TrackIndex, Time, false, UseRaw);
+			AnimSeq->GetBoneTransform(LocalAtom, TrackIndex, Time, UseRaw);
 
 			FQuat Rot = LocalAtom.GetRotation();
 			LocalAtom.SetRotation(EnforceShortestArc(FQuat::Identity, Rot));
@@ -429,7 +429,7 @@ void UAnimCompress_RemoveLinearKeys::UpdateBoneAtomList(
 	{
 		float Time = (float)FrameIndex * TimePerFrame;
 		FTransform LocalAtom;
-		AnimSeq->GetBoneTransform(LocalAtom, TrackIndex, Time, false, false);
+		AnimSeq->GetBoneTransform(LocalAtom, TrackIndex, Time, false);
 
 		FQuat Rot = LocalAtom.GetRotation();
 		LocalAtom.SetRotation( EnforceShortestArc(FQuat::Identity, Rot) );
diff --git a/Engine/Source/Runtime/Engine/Private/Animation/AnimEncoding.cpp b/Engine/Source/Runtime/Engine/Private/Animation/AnimEncoding.cpp
--- a/Engine/Source/Runtime/Engine/Private/Animation/AnimEncoding.cpp
+++ b/Engine/Source/Runtime/Engine/Private/Animation/AnimEncoding.cpp
@@ -186,17 +186,15 @@ void PadMemoryReader(FMemoryReader* MemoryReader, uint8*& TrackData, const int32
  * @param	Seq				An Animation Sequence to extract the BoneAtom from.
  * @param	TrackIndex		The index of the track desired in the Animation Sequence.
  * @param	Time			The time (in seconds) to calculate the BoneAtom for.
- * @param	bLooping		true if the animation should be played in a cyclic manner.
  */
 void AnimationFormat_GetBoneAtom(	
 	FTransform& OutAtom,
 	const UAnimSequence& Seq,
 	int32 TrackIndex,
-	float Time,
-	bool bLooping)
+	float Time)
 {
 	checkSlow(Seq.RotationCodec != NULL);
-	((AnimEncoding*)Seq.RotationCodec)->GetBoneAtom(OutAtom, Seq, TrackIndex, Time, bLooping);
+	((AnimEncoding*)Seq.RotationCodec)->GetBoneAtom(OutAtom, Seq, TrackIndex, Time);
 }
 
 #if USE_ANIMATION_CODEC_BATCH_SOLVER
@@ -209,7 +207,6 @@ void AnimationFormat_GetBoneAtom(
  * @param	TranslationTracks	A BoneTrackArray element for each bone requesting translation data. 
  * @param	Seq					An Animation Sequence to extract the BoneAtom from.
  * @param	Time				The time (in seconds) to calculate the BoneAtom for.
- * @param	bLooping			true if the animation should be played in a cyclic manner.
  */
 void AnimationFormat_GetAnimationPose(	
 	FTransformArray& Atoms, 
@@ -217,22 +214,21 @@ void AnimationFormat_GetAnimationPose(
 	const BoneTrackArray& TranslationPairs,
 	const BoneTrackArray& ScalePairs,
 	const UAnimSequence& Seq,
-	float Time,
-	bool bLooping)
+	float Time)
 {
 	// decompress the translation component using the proper method
 	checkSlow(Seq.TranslationCodec != NULL);
-	((AnimEncoding*)Seq.TranslationCodec)->GetPoseTranslations(Atoms, TranslationPairs, Seq, Time, bLooping);
+	((AnimEncoding*)Seq.TranslationCodec)->GetPoseTranslations(Atoms, TranslationPairs, Seq, Time);
 
 	// decompress the rotation component using the proper method
 	checkSlow(Seq.RotationCodec != NULL);
-	((AnimEncoding*)Seq.RotationCodec)->GetPoseRotations(Atoms, RotationPairs, Seq, Time, bLooping);
+	((AnimEncoding*)Seq.RotationCodec)->GetPoseRotations(Atoms, RotationPairs, Seq, Time);
 
 	checkSlow(Seq.ScaleCodec != NULL);
 	// we allow scale key to be empty
 	if (Seq.CompressedScaleOffsets.IsValid())
 	{
-		((AnimEncoding*)Seq.ScaleCodec)->GetPoseScales(Atoms, ScalePairs, Seq, Time, bLooping);
+		((AnimEncoding*)Seq.ScaleCodec)->GetPoseScales(Atoms, ScalePairs, Seq, Time);
 	}
 }
 #endif
@@ -244,14 +240,12 @@ void AnimationFormat_GetAnimationPose(
  * @param	Seq				An Animation Sequence to extract the BoneAtom from.
  * @param	TrackIndex		The index of the track desired in the Animation Sequence.
  * @param	Time			The time (in seconds) to calculate the BoneAtom for.
- * @param	bLooping		true if the animation should be played in a cyclic manner.
  */
 void AnimEncodingLegacyBase::GetBoneAtom(
 	FTransform& OutAtom,
 	const UAnimSequence& Seq,
 	int32 TrackIndex,
-	float Time,
-	bool bLooping)
+	float Time)
 {
 	// Initialize to identity to set the scale and in case of a missing rotation or translation codec
 	OutAtom.SetIdentity();
@@ -269,11 +263,11 @@ void AnimEncodingLegacyBase::GetBoneAtom(
 
 	// decompress the translation component using the proper method
 	checkSlow(Seq.TranslationCodec != NULL);
-	((AnimEncodingLegacyBase*)Seq.TranslationCodec)->GetBoneAtomTranslation(OutAtom, Seq, TransStream, NumTransKeys, Time, RelativePos, bLooping);
+	((AnimEncodingLegacyBase*)Seq.TranslationCodec)->GetBoneAtomTranslation(OutAtom, Seq, TransStream, NumTransKeys, Time, RelativePos);
 
 	// decompress the rotation component using the proper method
 	checkSlow(Seq.RotationCodec != NULL);
-	((AnimEncodingLegacyBase*)Seq.RotationCodec)->GetBoneAtomRotation(OutAtom, Seq, RotStream, NumRotKeys, Time, RelativePos, bLooping);
+	((AnimEncodingLegacyBase*)Seq.RotationCodec)->GetBoneAtomRotation(OutAtom, Seq, RotStream, NumRotKeys, Time, RelativePos);
 
 	// we assume scale keys can be empty, so only extrace if we have valid keys
 	bool bHasValidScale = Seq.CompressedScaleOffsets.IsValid();
@@ -284,7 +278,7 @@ void AnimEncodingLegacyBase::GetBoneAtom(
 		const uint8* RESTRICT ScaleStream		= Seq.CompressedByteStream.GetTypedData()+ScaleKeyOffset;
 		// decompress the rotation component using the proper method
 		checkSlow(Seq.ScaleCodec != NULL);
-		((AnimEncodingLegacyBase*)Seq.ScaleCodec)->GetBoneAtomScale(OutAtom, Seq, ScaleStream, NumScaleKeys, Time, RelativePos, bLooping);
+		((AnimEncodingLegacyBase*)Seq.ScaleCodec)->GetBoneAtomScale(OutAtom, Seq, ScaleStream, NumScaleKeys, Time, RelativePos);
 	}
 }
 
diff --git a/Engine/Source/Runtime/Engine/Private/Animation/AnimEncoding_PerTrackCompression.cpp b/Engine/Source/Runtime/Engine/Private/Animation/AnimEncoding_PerTrackCompression.cpp
--- a/Engine/Source/Runtime/Engine/Private/Animation/AnimEncoding_PerTrackCompression.cpp
+++ b/Engine/Source/Runtime/Engine/Private/Animation/AnimEncoding_PerTrackCompression.cpp
@@ -451,14 +451,12 @@ void AEFPerTrackCompressionCodec::ByteSwapOut(
  * @param	Seq				An Animation Sequence to extract the BoneAtom from.
  * @param	TrackIndex		The index of the track desired in the Animation Sequence.
  * @param	Time			The time (in seconds) to calculate the BoneAtom for.
- * @param	bLooping		true if the animation should be played in a cyclic manner.
  */
 void AEFPerTrackCompressionCodec::GetBoneAtom(
 	FTransform& OutAtom,
 	const UAnimSequence& Seq,
 	int32 TrackIndex,
-	float Time,
-	bool bLooping)
+	float Time)
 {
 	// Initialize to identity to set the scale and in case of a missing rotation or translation codec
 	OutAtom.SetIdentity();
@@ -469,13 +467,13 @@ void AEFPerTrackCompressionCodec::GetBoneAtom(
 	const int32 RotKeysOffset = TrackData[1];
 	const float RelativePos = Time / (float)Seq.SequenceLength;
 
-	GetBoneAtomTranslation(OutAtom, Seq, TransKeysOffset, Time, RelativePos, bLooping);
-	GetBoneAtomRotation(OutAtom, Seq, RotKeysOffset, Time, RelativePos, bLooping);
+	GetBoneAtomTranslation(OutAtom, Seq, TransKeysOffset, Time, RelativePos);
+	GetBoneAtomRotation(OutAtom, Seq, RotKeysOffset, Time, RelativePos);
 	const bool bHasScaleData = Seq.CompressedScaleOffsets.IsValid();
 	if (bHasScaleData)
 	{
 		const int32 ScaleKeysOffset = Seq.CompressedScaleOffsets.GetOffsetData(TrackIndex, 0);
-		GetBoneAtomScale(OutAtom, Seq, ScaleKeysOffset, Time, RelativePos, bLooping);
+		GetBoneAtomScale(OutAtom, Seq, ScaleKeysOffset, Time, RelativePos);
 	}
 }
 	
@@ -489,8 +487,7 @@ void AEFPerTrackCompressionCodec::GetBoneAtomRotation(
 	const UAnimSequence& Seq,
 	int32 Offset,
 	float Time,
-	float RelativePos,
-	bool bLooping)
+	float RelativePos)
 {
 	if (Offset != INDEX_NONE)
 	{
@@ -515,12 +512,12 @@ void AEFPerTrackCompressionCodec::GetBoneAtomRotation(
 		{
 			if ((FormatFlags & 0x8) == 0)
 			{
-				Alpha = TimeToIndex(Seq, RelativePos, bLooping, NumKeys, Index0, Index1);
+				Alpha = TimeToIndex(Seq, RelativePos, NumKeys, Index0, Index1);
 			}
 			else
 			{
 				const uint8* RESTRICT FrameTable = Align(TrackData + FixedBytes + BytesPerKey * NumKeys, 4);
-				Alpha = TimeToIndex(Seq, FrameTable, RelativePos, bLooping, NumKeys, Index0, Index1);
+				Alpha = TimeToIndex(Seq, FrameTable, RelativePos, NumKeys, Index0, Index1);
 			}
 		}
 
@@ -581,8 +578,7 @@ void AEFPerTrackCompressionCodec::GetBoneAtomTranslation(
 	const UAnimSequence& Seq,
 	int32 Offset,
 	float Time,
-	float RelativePos,
-	bool bLooping)
+	float RelativePos)
 {
 	if (Offset != INDEX_NONE)
 	{
@@ -607,12 +603,12 @@ void AEFPerTrackCompressionCodec::GetBoneAtomTranslation(
 		{
 			if ((FormatFlags & 0x8) == 0)
 			{
-				Alpha = TimeToIndex(Seq, RelativePos, bLooping, NumKeys, Index0, Index1);
+				Alpha = TimeToIndex(Seq, RelativePos, NumKeys, Index0, Index1);
 			}
 			else
 			{
 				const uint8* RESTRICT FrameTable = Align(TrackData + FixedBytes + BytesPerKey * NumKeys, 4);
-				Alpha = TimeToIndex(Seq, FrameTable, RelativePos, bLooping, NumKeys, Index0, Index1);
+				Alpha = TimeToIndex(Seq, FrameTable, RelativePos, NumKeys, Index0, Index1);
 			}
 		}
 
@@ -667,8 +663,7 @@ void AEFPerTrackCompressionCodec::GetBoneAtomScale(
 	const UAnimSequence& Seq,
 	int32 Offset,
 	float Time,
-	float RelativePos,
-	bool bLooping)
+	float RelativePos)
 {
 	if (Offset != INDEX_NONE)
 	{
@@ -693,12 +688,12 @@ void AEFPerTrackCompressionCodec::GetBoneAtomScale(
 		{
 			if ((FormatFlags & 0x8) == 0)
 			{
-				Alpha = TimeToIndex(Seq, RelativePos, bLooping, NumKeys, Index0, Index1);
+				Alpha = TimeToIndex(Seq, RelativePos, NumKeys, Index0, Index1);
 			}
 			else
 			{
 				const uint8* RESTRICT FrameTable = Align(TrackData + FixedBytes + BytesPerKey * NumKeys, 4);
-				Alpha = TimeToIndex(Seq, FrameTable, RelativePos, bLooping, NumKeys, Index0, Index1);
+				Alpha = TimeToIndex(Seq, FrameTable, RelativePos, NumKeys, Index0, Index1);
 			}
 		}
 
@@ -756,15 +751,13 @@ void AEFPerTrackCompressionCodec::GetBoneAtomScale(
 * @param	DesiredPairs	Array of requested bone information
 * @param	Seq				The animation sequence to use.
 * @param	Time			Current time to solve for.
-* @param	bLooping		True when looping the stream in intended.
 * @return					None.
 */
 void AEFPerTrackCompressionCodec::GetPoseRotations(
 	FTransformArray& Atoms,
 	const BoneTrackArray& DesiredPairs,
 	const UAnimSequence& Seq,
-	float Time,
-	bool bLooping )
+	float Time)
 {
 	const int32 PairCount = DesiredPairs.Num();
 	const float RelativePos = Time / Seq.SequenceLength;
@@ -779,7 +772,7 @@ void AEFPerTrackCompressionCodec::GetPoseRotations(
 		const int32* RESTRICT TrackData = Seq.CompressedTrackOffsets.GetTypedData() + (TrackIndex * 2);
 		const int32 RotKeysOffset = *(TrackData + 1);
 
-		GetBoneAtomRotation( BoneAtom, Seq, RotKeysOffset, Time, RelativePos, bLooping );
+		GetBoneAtomRotation( BoneAtom, Seq, RotKeysOffset, Time, RelativePos );
 	}
 }
 
@@ -790,15 +783,13 @@ void AEFPerTrackCompressionCodec::GetPoseRotations(
 * @param	DesiredPairs	Array of requested bone information
 * @param	Seq				The animation sequence to use.
 * @param	Time			Current time to solve for.
-* @param	bLooping		True when looping the stream in intended.
 * @return					None.
 */
 void AEFPerTrackCompressionCodec::GetPoseTranslations(
 	FTransformArray& Atoms,
 	const BoneTrackArray& DesiredPairs,
 	const UAnimSequence& Seq,
-	float Time,
-	bool bLooping )
+	float Time)
 {
 	const int32 PairCount = DesiredPairs.Num();
 	const float RelativePos = Time / Seq.SequenceLength;
@@ -813,7 +804,7 @@ void AEFPerTrackCompressionCodec::GetPoseTranslations(
 		const int32* RESTRICT TrackData = Seq.CompressedTrackOffsets.GetTypedData() + (TrackIndex * 2);
 		const int32 PosKeysOffset = *(TrackData + 0);
 
-		GetBoneAtomTranslation( BoneAtom, Seq, PosKeysOffset, Time, RelativePos, bLooping );
+		GetBoneAtomTranslation( BoneAtom, Seq, PosKeysOffset, Time, RelativePos );
 	}
 }
 
@@ -824,15 +815,13 @@ void AEFPerTrackCompressionCodec::GetPoseTranslations(
 * @param	DesiredPairs	Array of requested bone information
 * @param	Seq				The animation sequence to use.
 * @param	Time			Current time to solve for.
-* @param	bLooping		True when looping the stream in intended.
 * @return					None.
 */
 void AEFPerTrackCompressionCodec::GetPoseScales(
 	FTransformArray& Atoms,
 	const BoneTrackArray& DesiredPairs,
 	const UAnimSequence& Seq,
-	float Time,
-	bool bLooping )
+	float Time)
 {
 	check( Seq.CompressedScaleOffsets.IsValid() );
 
@@ -848,7 +837,7 @@ void AEFPerTrackCompressionCodec::GetPoseScales(
 
 		const int32 ScaleKeysOffset = Seq.CompressedScaleOffsets.GetOffsetData( TrackIndex, 0 );
 
-		GetBoneAtomScale( BoneAtom, Seq, ScaleKeysOffset, Time, RelativePos, bLooping );
+		GetBoneAtomScale( BoneAtom, Seq, ScaleKeysOffset, Time, RelativePos );
 	}
 }
 
diff --git a/Engine/Source/Runtime/Engine/Private/Animation/AnimInstance.cpp b/Engine/Source/Runtime/Engine/Private/Animation/AnimInstance.cpp
--- a/Engine/Source/Runtime/Engine/Private/Animation/AnimInstance.cpp
+++ b/Engine/Source/Runtime/Engine/Private/Animation/AnimInstance.cpp
@@ -825,14 +825,13 @@ void UAnimInstance::DisplayDebug(class UCanvas* Canvas, const FDebugDisplayInfo&
 	}
 }
 
-void UAnimInstance::BlendSpaceEvaluatePose(class UBlendSpaceBase* BlendSpace, TArray<FBlendSampleData>& BlendSampleDataCache, struct FA2Pose& Pose, bool bIsLooping)
+void UAnimInstance::BlendSpaceEvaluatePose(class UBlendSpaceBase* BlendSpace, TArray<FBlendSampleData>& BlendSampleDataCache, struct FA2Pose& Pose)
 {
 	SCOPE_CYCLE_COUNTER(STAT_AnimNativeEvaluatePoses);
 
 	FAnimationRuntime::GetPoseFromBlendSpace(
 		BlendSpace,
 		BlendSampleDataCache, 
-		bIsLooping,
 		RequiredBones,
 		/*out*/ Pose.Bones);
 }
@@ -1252,7 +1251,7 @@ void UAnimInstance::SlotEvaluatePose(FName SlotNodeName, const FA2Pose & SourceP
 
 			// Extract pose from Track
 			UAnimMontage const * const MontageAsset = MontageInstance->Montage;
-			FAnimExtractContext ExtractionContext(MontageInstance->Position, false, MontageAsset->bEnableRootMotionTranslation, MontageAsset->bEnableRootMotionRotation, MontageAsset->RootMotionRootLock);
+			FAnimExtractContext ExtractionContext(MontageInstance->Position, MontageAsset->bEnableRootMotionTranslation, MontageAsset->bEnableRootMotionRotation, MontageAsset->RootMotionRootLock);
 			FAnimationRuntime::GetPoseFromAnimTrack(*AnimTrack, RequiredBones, NewPose.Pose.Bones, ExtractionContext);
 
 			TotalWeight += MontageInstance->Weight;
diff --git a/Engine/Source/Runtime/Engine/Private/Animation/AnimNode_BlendSpacePlayer.cpp b/Engine/Source/Runtime/Engine/Private/Animation/AnimNode_BlendSpacePlayer.cpp
--- a/Engine/Source/Runtime/Engine/Private/Animation/AnimNode_BlendSpacePlayer.cpp
+++ b/Engine/Source/Runtime/Engine/Private/Animation/AnimNode_BlendSpacePlayer.cpp
@@ -70,7 +70,7 @@ void FAnimNode_BlendSpacePlayer::Evaluate(FPoseContext& Output)
 {
 	if ((BlendSpace != NULL) && (Output.AnimInstance->CurrentSkeleton->IsCompatible(BlendSpace->GetSkeleton())))
 	{
-		Output.AnimInstance->BlendSpaceEvaluatePose(BlendSpace, BlendSampleDataCache, Output.Pose, bLoop);
+		Output.AnimInstance->BlendSpaceEvaluatePose(BlendSpace, BlendSampleDataCache, Output.Pose);
 	}
 	else
 	{
diff --git a/Engine/Source/Runtime/Engine/Private/Animation/AnimNode_SequenceEvaluator.cpp b/Engine/Source/Runtime/Engine/Private/Animation/AnimNode_SequenceEvaluator.cpp
--- a/Engine/Source/Runtime/Engine/Private/Animation/AnimNode_SequenceEvaluator.cpp
+++ b/Engine/Source/Runtime/Engine/Private/Animation/AnimNode_SequenceEvaluator.cpp
@@ -23,7 +23,7 @@ void FAnimNode_SequenceEvaluator::Evaluate(FPoseContext& Output)
 {
 	if ((Sequence != NULL) && (Output.AnimInstance->CurrentSkeleton->IsCompatible(Sequence->GetSkeleton())))
 	{
-		Output.AnimInstance->SequenceEvaluatePose(Sequence, Output.Pose, FAnimExtractContext(ExplicitTime, false));
+		Output.AnimInstance->SequenceEvaluatePose(Sequence, Output.Pose, FAnimExtractContext(ExplicitTime));
 	}
 	else
 	{
diff --git a/Engine/Source/Runtime/Engine/Private/Animation/AnimNode_SequencePlayer.cpp b/Engine/Source/Runtime/Engine/Private/Animation/AnimNode_SequencePlayer.cpp
--- a/Engine/Source/Runtime/Engine/Private/Animation/AnimNode_SequencePlayer.cpp
+++ b/Engine/Source/Runtime/Engine/Private/Animation/AnimNode_SequencePlayer.cpp
@@ -50,7 +50,7 @@ void FAnimNode_SequencePlayer::Evaluate(FPoseContext& Output)
 {
 	if ((Sequence != NULL) && (Output.AnimInstance->CurrentSkeleton->IsCompatible(Sequence->GetSkeleton())))
 	{
-		Output.AnimInstance->SequenceEvaluatePose(Sequence, Output.Pose, FAnimExtractContext(InternalTimeAccumulator, bLoopAnimation));
+		Output.AnimInstance->SequenceEvaluatePose(Sequence, Output.Pose, FAnimExtractContext(InternalTimeAccumulator));
 	}
 	else
 	{
diff --git a/Engine/Source/Runtime/Engine/Private/Animation/AnimSequence.cpp b/Engine/Source/Runtime/Engine/Private/Animation/AnimSequence.cpp
--- a/Engine/Source/Runtime/Engine/Private/Animation/AnimSequence.cpp
+++ b/Engine/Source/Runtime/Engine/Private/Animation/AnimSequence.cpp
@@ -100,7 +100,6 @@ UAnimSequence::UAnimSequence(const class FPostConstructInitializeProperties& PCI
 {
 
 	RateScale = 1.0;
-	bLoopingInterpolation = true;
 }
 
 SIZE_T UAnimSequence::GetResourceSize(EResourceSizeMode::Type Mode)
@@ -511,14 +510,14 @@ static int32 FindKeyIndex(float Time, const TimeArray& Times)
 	return FoundIndex;
 }
 
-void UAnimSequence::GetBoneTransform(FTransform& OutAtom, int32 TrackIndex, float Time, bool bLooping, bool bUseRawData) const
+void UAnimSequence::GetBoneTransform(FTransform& OutAtom, int32 TrackIndex, float Time, bool bUseRawData) const
 {
 	// If the caller didn't request that raw animation data be used . . .
 	if ( !bUseRawData )
 	{
 		if ( CompressedTrackOffsets.Num() > 0 )
 		{
-			AnimationFormat_GetBoneAtom( OutAtom, *this, TrackIndex, Time, bLooping );
+			AnimationFormat_GetBoneAtom( OutAtom, *this, TrackIndex, Time );
 			return;
 		}
 	}
@@ -544,7 +543,7 @@ void UAnimSequence::GetBoneTransform(FTransform& OutAtom, int32 TrackIndex, floa
 
 	int32 KeyIndex1, KeyIndex2;
 	float Alpha;
-	FAnimationRuntime::GetKeyIndicesFromTime(KeyIndex1, KeyIndex2, Alpha, Time, bLooping, NumFrames, SequenceLength);
+	FAnimationRuntime::GetKeyIndicesFromTime(KeyIndex1, KeyIndex2, Alpha, Time, NumFrames, SequenceLength);
 	// @Todo fix me: this change is not good, it has lots of branches. But we'd like to save memory for not saving scale if no scale change exists
 	const bool bHasScaleKey = (RawTrack.ScaleKeys.Num() > 0);
 	static const FVector DefaultScale3D = FVector(1.f);
@@ -615,7 +614,7 @@ void UAnimSequence::ExtractRootTrack(float Pos, FTransform & RootTransform, cons
 	if ((TrackToSkeletonMapTable.Num() > 0) && (TrackToSkeletonMapTable[0].BoneTreeIndex == 0) )
 	{
 		// if we do have root data, then return root data
-		GetBoneTransform(RootTransform, 0, Pos, false, false );
+		GetBoneTransform(RootTransform, 0, Pos, false );
 		return;
 	}
 
@@ -692,9 +691,6 @@ void UAnimSequence::ResetRootBoneForRootMotion(FTransformArrayA2 & BoneTransform
 
 void UAnimSequence::GetBonePose(FTransformArrayA2 & OutAtoms, const FBoneContainer & RequiredBones, const FAnimExtractContext & ExtractionContext) const
 {
-	// Allow 'bLoopingInterpolation' flag on the AnimSequence to disable the looping interpolation.
-	bool const bDoLoopingInterpolation = ExtractionContext.bLooping && bLoopingInterpolation;
-
 	USkeleton * MySkeleton = GetSkeleton();
 	if (!MySkeleton)
 	{
@@ -751,7 +747,7 @@ void UAnimSequence::GetBonePose(FTransformArrayA2 & OutAtoms, const FBoneContain
 			if( PoseBoneIndex != INDEX_NONE )
 			{
 				// extract animation
-				GetBoneTransform(OutAtoms[PoseBoneIndex], TrackIndex, ExtractionContext.CurrentTime, bDoLoopingInterpolation, true);
+				GetBoneTransform(OutAtoms[PoseBoneIndex], TrackIndex, ExtractionContext.CurrentTime, true);
 
 				if (!bDisableRetargeting)
 				{
@@ -821,8 +817,7 @@ void UAnimSequence::GetBonePose(FTransformArrayA2 & OutAtoms, const FBoneContain
 			RootAtom,
 			*this,
 			TrackIndex,
-			ExtractionContext.CurrentTime,
-			bDoLoopingInterpolation );
+			ExtractionContext.CurrentTime);
 
 		// @laurent - we should look into splitting rotation and translation tracks, so we don't have to process translation twice.
 		RetargetBoneTransform(RootAtom, 0, 0, RequiredBones);
@@ -835,8 +830,7 @@ void UAnimSequence::GetBonePose(FTransformArrayA2 & OutAtoms, const FBoneContain
 		TranslationPairs,
 		RotationScalePairs, 
 		*this,
-		ExtractionContext.CurrentTime,
-		bDoLoopingInterpolation);
+		ExtractionContext.CurrentTime);
 
 	// Once pose has been extracted, snap root bone back to first frame if we are extracting root motion.
 	if( ExtractionContext.bExtractRootMotionTranslation || ExtractionContext.bExtractRootMotionRotation )
@@ -1351,7 +1345,6 @@ bool UAnimSequence::CopyAnimSequenceProperties(UAnimSequence* SourceAnimSeq, UAn
 	DestAnimSeq->SequenceLength				= SourceAnimSeq->SequenceLength;
 	DestAnimSeq->NumFrames					= SourceAnimSeq->NumFrames;
 	DestAnimSeq->RateScale					= SourceAnimSeq->RateScale;
-	DestAnimSeq->bLoopingInterpolation		= SourceAnimSeq->bLoopingInterpolation;
 	DestAnimSeq->bDoNotOverrideCompression	= SourceAnimSeq->bDoNotOverrideCompression;
 
 	// Copy Compression Settings
@@ -1634,12 +1627,9 @@ void UAnimSequence::PostProcessSequence()
 	// initialize notify track
 	InitializeNotifyTrack();
 	//Make sure we dont have any notifies off the end of the sequence
-#if WITH_EDITOR
 	ClampNotifiesAtEndOfSequence();
-#endif
 }
 
-#if WITH_EDITOR
 void UAnimSequence::RemoveNaNTracks()
 {
 	bool bRecompress = false;
@@ -1686,7 +1676,6 @@ void UAnimSequence::RemoveNaNTracks()
 		FAnimationUtils::CompressAnimSequence(this, false, false);
 	}
 }
-#endif // WITH_EDITOR
 
 void UAnimSequence::RemoveTrack(int32 TrackIndex)
 {
@@ -1706,6 +1695,69 @@ bool UAnimSequence::GetAllAnimationSequencesReferred(TArray<UAnimSequence*>& Ani
 
 	return true;
 }
+
+bool UAnimSequence::AddLoopingInterpolation()
+{
+	int32 NumTracks = AnimationTrackNames.Num();
+	int32 NumKeys = NumFrames;
+	int32 SrcKeyIndex = 0; // this is the key index we copy animation from
+	int32 DestKeyIndex = NumFrames; // this is the key index we insert animation to
+	float Interval = (NumKeys)? SequenceLength/NumKeys : 0.f;
+
+	if(Interval > 0.f)
+	{
+		// 2d array of animated time [track index][time key]
+		TArray< TArray<FTransform> > AnimatedLocalSpaces;
+		AnimatedLocalSpaces.AddZeroed(NumTracks);
+
+		for(int32 TrackIndex=0; TrackIndex<NumTracks; ++TrackIndex)
+		{
+			// add one extra for looping interpolation
+			AnimatedLocalSpaces[TrackIndex].AddUninitialized(NumKeys + 1);
+
+			// fill up keys
+			for(int32 Key=0; Key<NumKeys; ++Key)
+			{
+				GetBoneTransform(AnimatedLocalSpaces[TrackIndex][Key], TrackIndex, Interval*Key, false);
+			}
+
+			// now add extra, source should be available all the time since we check interval already
+			AnimatedLocalSpaces[TrackIndex][DestKeyIndex] = AnimatedLocalSpaces[TrackIndex][SrcKeyIndex];
+		}
+
+		// added one more key
+		int32 NewNumKeys = NumKeys +1 ;
+
+		// now I need to calculate back to new animation data
+		TArray<struct FRawAnimSequenceTrack> NewRawAnimationData = RawAnimationData;
+		for(int32 TrackIndex=0; TrackIndex<NumTracks; ++TrackIndex)
+		{
+			auto & RawAnimation = NewRawAnimationData[TrackIndex];
+			RawAnimation.PosKeys.Empty(NewNumKeys);
+			RawAnimation.PosKeys.AddUninitialized(NewNumKeys);
+			RawAnimation.RotKeys.Empty(NewNumKeys);
+			RawAnimation.RotKeys.AddUninitialized(NewNumKeys);
+			RawAnimation.ScaleKeys.Empty(NewNumKeys);
+			RawAnimation.ScaleKeys.AddUninitialized(NewNumKeys);
+
+			for(int32 Key=0; Key<NewNumKeys; ++Key)
+			{
+				RawAnimation.PosKeys[Key] = AnimatedLocalSpaces[TrackIndex][Key].GetLocation();
+				RawAnimation.RotKeys[Key] = AnimatedLocalSpaces[TrackIndex][Key].GetRotation();
+				RawAnimation.ScaleKeys[Key] = AnimatedLocalSpaces[TrackIndex][Key].GetScale3D();
+			}
+		}
+
+		RawAnimationData = NewRawAnimationData;
+		NumFrames += 1;
+		SequenceLength += Interval;
+
+		PostProcessSequence();
+		return true;
+	}
+
+	return false;
+}
 #endif
 /*-----------------------------------------------------------------------------
 	AnimNotify & subclasses
diff --git a/Engine/Source/Runtime/Engine/Private/Animation/AnimSingleNodeInstance.cpp b/Engine/Source/Runtime/Engine/Private/Animation/AnimSingleNodeInstance.cpp
--- a/Engine/Source/Runtime/Engine/Private/Animation/AnimSingleNodeInstance.cpp
+++ b/Engine/Source/Runtime/Engine/Private/Animation/AnimSingleNodeInstance.cpp
@@ -295,7 +295,7 @@ bool UAnimSingleNodeInstance::NativeEvaluateAnimation(FPoseContext& Output)
 		//@TODO: animrefactor: Seems like more code duplication than we need
 		if (UBlendSpaceBase* BlendSpace = Cast<UBlendSpaceBase>(CurrentAsset))
 		{
-			InternalBlendSpaceEvaluatePose(BlendSpace, BlendSampleData, Output.Pose, bLooping);
+			InternalBlendSpaceEvaluatePose(BlendSpace, BlendSampleData, Output.Pose);
 		}
 		else if (UAnimSequence* Sequence = Cast<UAnimSequence>(CurrentAsset))
 		{
@@ -307,7 +307,7 @@ bool UAnimSingleNodeInstance::NativeEvaluateAnimation(FPoseContext& Output)
 				BasePose.Bones.AddUninitialized(Output.Pose.Bones.Num());
 				AdditivePose.Bones.AddUninitialized(Output.Pose.Bones.Num());
 
-				FAnimExtractContext ExtractionContext(CurrentTime, bLooping);
+				FAnimExtractContext ExtractionContext(CurrentTime);
 				Sequence->GetAdditiveBasePose(BasePose.Bones, RequiredBones, ExtractionContext);
 				Sequence->GetAnimationPose(AdditivePose.Bones, RequiredBones, ExtractionContext);
 				if (Sequence->AdditiveAnimType == AAT_LocalSpaceBase)
@@ -322,7 +322,7 @@ bool UAnimSingleNodeInstance::NativeEvaluateAnimation(FPoseContext& Output)
 			else
 			{
 				// if sekeltalmesh isn't there, we'll need to use skeleton
-				FAnimationRuntime::GetPoseFromSequence(Sequence, RequiredBones, Output.Pose.Bones, FAnimExtractContext(CurrentTime, bLooping));
+				FAnimationRuntime::GetPoseFromSequence(Sequence, RequiredBones, Output.Pose.Bones, FAnimExtractContext(CurrentTime));
 			}
 		}
 		else if (UAnimComposite* Composite = Cast<UAnimComposite>(CurrentAsset))
@@ -344,7 +344,7 @@ bool UAnimSingleNodeInstance::NativeEvaluateAnimation(FPoseContext& Output)
 				FAnimationRuntime::FillWithRefPose(BasePose.Bones, RequiredBones);
 				
 				//get the additive pose
-				FAnimationRuntime::GetPoseFromAnimTrack(AnimTrack, RequiredBones, AdditivePose.Bones, FAnimExtractContext(CurrentTime, false));
+				FAnimationRuntime::GetPoseFromAnimTrack(AnimTrack, RequiredBones, AdditivePose.Bones, FAnimExtractContext(CurrentTime));
 
 				// if additive, we should blend with source to make it fullbody
 				if (AdditiveAnimType == AAT_LocalSpaceBase)
@@ -359,7 +359,7 @@ bool UAnimSingleNodeInstance::NativeEvaluateAnimation(FPoseContext& Output)
 			else
 			{
 				//doesn't handle additive yet
-				FAnimationRuntime::GetPoseFromAnimTrack(AnimTrack, RequiredBones, Output.Pose.Bones, FAnimExtractContext(CurrentTime, false));
+				FAnimationRuntime::GetPoseFromAnimTrack(AnimTrack, RequiredBones, Output.Pose.Bones, FAnimExtractContext(CurrentTime));
 			}
 		}
 		else if (UAnimMontage* Montage = Cast<UAnimMontage>(CurrentAsset))
@@ -376,7 +376,7 @@ bool UAnimSingleNodeInstance::NativeEvaluateAnimation(FPoseContext& Output)
 					// if montage is additive, we need to have base pose for the slot pose evaluate
 					if (Montage->PreviewBasePose && Montage->SequenceLength > 0.f)
 					{
-						Montage->PreviewBasePose->GetBonePose(SourcePose.Bones, RequiredBones, FAnimExtractContext(CurrentTime, bLooping));
+						Montage->PreviewBasePose->GetBonePose(SourcePose.Bones, RequiredBones, FAnimExtractContext(CurrentTime));
 					}
 					else
 #endif // WITH_EDITORONLY_DATA
@@ -539,7 +539,7 @@ void UAnimSingleNodeInstance::SetBlendSpaceInput(const FVector& InBlendInput)
 	BlendSpaceInput = InBlendInput;
 }
 
-void UAnimSingleNodeInstance::InternalBlendSpaceEvaluatePose(class UBlendSpaceBase* BlendSpace, TArray<FBlendSampleData>& BlendSampleDataCache, struct FA2Pose& Pose, bool bIsLooping)
+void UAnimSingleNodeInstance::InternalBlendSpaceEvaluatePose(class UBlendSpaceBase* BlendSpace, TArray<FBlendSampleData>& BlendSampleDataCache, struct FA2Pose& Pose)
 {
 	USkeletalMeshComponent* Component = GetSkelMeshComponent();
 
@@ -553,7 +553,7 @@ void UAnimSingleNodeInstance::InternalBlendSpaceEvaluatePose(class UBlendSpaceBa
 #if WITH_EDITORONLY_DATA
 		if (BlendSpace->PreviewBasePose)
 		{
-			BlendSpace->PreviewBasePose->GetBonePose(/*out*/ BasePose.Bones, RequiredBones, FAnimExtractContext(PreviewPoseCurrentTime, bIsLooping));
+			BlendSpace->PreviewBasePose->GetBonePose(/*out*/ BasePose.Bones, RequiredBones, FAnimExtractContext(PreviewPoseCurrentTime));
 		}
 		else
 #endif // WITH_EDITORONLY_DATA
@@ -565,7 +565,6 @@ void UAnimSingleNodeInstance::InternalBlendSpaceEvaluatePose(class UBlendSpaceBa
 		FAnimationRuntime::GetPoseFromBlendSpace(
 			BlendSpace,
 			BlendSampleDataCache, 
-			bIsLooping,
 			RequiredBones,
 			/*out*/ AdditivePose.Bones);
 
@@ -584,7 +583,6 @@ void UAnimSingleNodeInstance::InternalBlendSpaceEvaluatePose(class UBlendSpaceBa
 		FAnimationRuntime::GetPoseFromBlendSpace(
 			BlendSpace,
 			BlendSampleDataCache, 
-			bIsLooping,
 			RequiredBones,
 			/*out*/ Pose.Bones);
 	}
diff --git a/Engine/Source/Runtime/Engine/Private/Animation/AnimationRuntime.cpp b/Engine/Source/Runtime/Engine/Private/Animation/AnimationRuntime.cpp
--- a/Engine/Source/Runtime/Engine/Private/Animation/AnimationRuntime.cpp
+++ b/Engine/Source/Runtime/Engine/Private/Animation/AnimationRuntime.cpp
@@ -425,7 +425,6 @@ void FAnimationRuntime::GetPoseFromAnimTrack(
 				// Copy passed in Extraction Context, but override position and looping parameters.
 				FAnimExtractContext SequenceExtractionContext(ExtractionContext);
 				SequenceExtractionContext.CurrentTime = PositionInAnim;
-				SequenceExtractionContext.bLooping = false;
 				FAnimationRuntime::GetPoseFromSequence(Sequence, RequiredBones, SourcePoses[NewIndex], SequenceExtractionContext);
 			}
 		}
@@ -546,7 +545,6 @@ void FAnimationRuntime::ApplyWeightToTransform(const FBoneContainer & RequiredBo
 void FAnimationRuntime::GetPoseFromBlendSpace(
 	UBlendSpaceBase * BlendSpace,
 	TArray<FBlendSampleData>& BlendSampleDataCache,
-	bool bLooping,
 	FBoneContainer & RequiredBones,
 	/*out*/ FTransformArrayA2& ResultAtoms)
 {
@@ -586,7 +584,7 @@ void FAnimationRuntime::GetPoseFromBlendSpace(
 				const float Time = FMath::Clamp<float>(BlendSampleDataCache[I].Time, 0.f, Sample.Animation->SequenceLength);
 
 				// first one always fills up the source one
-				FAnimationRuntime::GetPoseFromSequence(Sample.Animation, RequiredBones, ChildrenTransform[I], FAnimExtractContext(Time, bLooping));
+				FAnimationRuntime::GetPoseFromSequence(Sample.Animation, RequiredBones, ChildrenTransform[I], FAnimExtractContext(Time));
 			}
 			else
 			{
@@ -628,7 +626,7 @@ void FAnimationRuntime::GetPoseFromBlendSpace(
 
 
 /* from % from OutKeyIndex1, meaning (CurrentKeyIndex(float)-OutKeyIndex1)/(OutKeyIndex2-OutKeyIndex1) */
-void FAnimationRuntime::GetKeyIndicesFromTime(int32 & OutKeyIndex1, int32 & OutKeyIndex2, float& OutAlpha, const float Time, const bool bLooping, const int32 NumFrames, const float SequenceLength)
+void FAnimationRuntime::GetKeyIndicesFromTime(int32 & OutKeyIndex1, int32 & OutKeyIndex2, float& OutAlpha, const float Time, const int32 NumFrames, const float SequenceLength)
 {
 	// Check for 1-frame, before-first-frame and after-last-frame cases.
 	if( Time <= 0.f || NumFrames == 1 )
@@ -642,15 +640,14 @@ void FAnimationRuntime::GetKeyIndicesFromTime(int32 & OutKeyIndex1, int32 & OutK
 	const int32 LastIndex		= NumFrames - 1;
 	if( Time >= SequenceLength )
 	{
-		OutKeyIndex1 = bLooping ? 0 : LastIndex;
+		OutKeyIndex1 = LastIndex;
 		OutKeyIndex2 = (OutKeyIndex1 + 1) % (NumFrames);
 		OutAlpha = 0.f;
 		return;
 	}
 
 	// This assumes that all keys are equally spaced (ie. won't work if we have dropped unimportant frames etc).
-	// by default, for looping animation, the last frame has a duration, and interpolates back to the first one.
-	const int32 NumKeys = bLooping ? NumFrames : NumFrames - 1;
+	const int32 NumKeys = NumFrames - 1;
 	const float KeyPos = ((float)NumKeys * Time) / SequenceLength;
 
 	// Find the integer part (ensuring within range) and that gives us the 'starting' key index.
@@ -664,16 +661,7 @@ void FAnimationRuntime::GetKeyIndicesFromTime(int32 & OutKeyIndex1, int32 & OutK
 	// If we have gone over the end, do different things in case of looping
 	if( KeyIndex2 == NumFrames )
 	{
-		// If looping, interpolate between last and first frame
-		if( bLooping )
-		{
-			KeyIndex2 = 0;
-		}
-		// If not looping - hold the last frame.
-		else
-		{
-			KeyIndex2 = KeyIndex1;
-		}
+		KeyIndex2 = KeyIndex1;
 	}
 
 	OutKeyIndex1 = KeyIndex1;
diff --git a/Engine/Source/Runtime/Engine/Private/Animation/AnimationUtils.cpp b/Engine/Source/Runtime/Engine/Private/Animation/AnimationUtils.cpp
--- a/Engine/Source/Runtime/Engine/Private/Animation/AnimationUtils.cpp
+++ b/Engine/Source/Runtime/Engine/Private/Animation/AnimationUtils.cpp
@@ -259,8 +259,8 @@ void FAnimationUtils::ComputeCompressionError(const UAnimSequence* AnimSeq, cons
 				}
 				else
 				{
-					AnimSeq->GetBoneTransform(RawAtoms[BoneIndex], TrackIndex, Time, false, true);
-					AnimSeq->GetBoneTransform(NewAtoms[BoneIndex], TrackIndex, Time, false, false);
+					AnimSeq->GetBoneTransform(RawAtoms[BoneIndex], TrackIndex, Time, true);
+					AnimSeq->GetBoneTransform(NewAtoms[BoneIndex], TrackIndex, Time, false);
 
 
 					bool bSkipTranslationTrack = false;
@@ -1652,7 +1652,7 @@ void FAnimationUtils::TallyErrorsFromPerturbation(
 				}
 				else
 				{
-					AnimSeq->GetBoneTransform(RawAtoms[BoneIndex], TrackIndex, Time, false, true);
+					AnimSeq->GetBoneTransform(RawAtoms[BoneIndex], TrackIndex, Time, true);
 
 					NewAtomsT[BoneIndex] = RawAtoms[BoneIndex];
 					NewAtomsR[BoneIndex] = RawAtoms[BoneIndex];
diff --git a/Engine/Source/Runtime/Engine/Public/AnimEncoding.h b/Engine/Source/Runtime/Engine/Public/AnimEncoding.h
--- a/Engine/Source/Runtime/Engine/Public/AnimEncoding.h
+++ b/Engine/Source/Runtime/Engine/Public/AnimEncoding.h
@@ -61,13 +61,11 @@ typedef TArray< FTransform, TMemStackAllocator<> > FTransformArray;
  * @param	Seq				An Animation Sequence to extract the BoneAtom from.
  * @param	TrackIndex		The index of the track desired in the Animation Sequence.
  * @param	Time			The time (in seconds) to calculate the BoneAtom for.
- * @param	bLooping		true if the animation should be played in a cyclic manner.
  */
 void AnimationFormat_GetBoneAtom(	FTransform& OutAtom,
 									const UAnimSequence& Seq,
 									int32 TrackIndex,
-									float Time,
-									bool bLooping);
+									float Time);
 
 #if USE_ANIMATION_CODEC_BATCH_SOLVER
 
@@ -79,7 +77,6 @@ void AnimationFormat_GetBoneAtom(	FTransform& OutAtom,
  * @param	TranslationTracks	A BoneTrackArray element for each bone requesting translation data. 
  * @param	Seq					An Animation Sequence to extract the BoneAtom from.
  * @param	Time				The time (in seconds) to calculate the BoneAtom for.
- * @param	bLooping			true if the animation should be played in a cyclic manner.
  */
 void AnimationFormat_GetAnimationPose(	
 	FTransformArray& Atoms, 
@@ -87,8 +84,7 @@ void AnimationFormat_GetAnimationPose(
 	const BoneTrackArray& TranslationTracks,
 	const BoneTrackArray& ScaleTracks,
 	const UAnimSequence& Seq,
-	float Time,
-	bool bLooping);
+	float Time);
 
 #endif
 
@@ -188,14 +184,12 @@ public:
 	 * @param	Seq				An Animation Sequence to extract the BoneAtom from.
 	 * @param	TrackIndex		The index of the track desired in the Animation Sequence.
 	 * @param	Time			The time (in seconds) to calculate the BoneAtom for.
-	 * @param	bLooping		true if the animation should be played in a cyclic manner.
 	 */
 	virtual void GetBoneAtom(
 		FTransform& OutAtom,
 		const UAnimSequence& Seq,
 		int32 TrackIndex,
-		float Time,
-		bool bLooping) PURE_VIRTUAL(AnimEncoding::GetBoneAtom,);
+		float Time) PURE_VIRTUAL(AnimEncoding::GetBoneAtom,);
 
 #if USE_ANIMATION_CODEC_BATCH_SOLVER
 
@@ -206,15 +200,13 @@ public:
 	 * @param	DesiredPairs	Array of requested bone information
 	 * @param	Seq				The animation sequence to use.
 	 * @param	Time			Current time to solve for.
-	 * @param	bLooping		True when looping the stream in intended.
 	 * @return					None. 
 	 */
 	virtual void GetPoseRotations(	
 		FTransformArray& Atoms, 
 		const BoneTrackArray& DesiredPairs,
 		const UAnimSequence& Seq,
-		float Time,
-		bool bLooping) PURE_VIRTUAL(AnimEncoding::GetPoseRotations,);
+		float Time) PURE_VIRTUAL(AnimEncoding::GetPoseRotations,);
 
 	/**
 	 * Decompress all requested translation components from an Animation Sequence
@@ -223,15 +215,13 @@ public:
 	 * @param	DesiredPairs	Array of requested bone information
 	 * @param	Seq				The animation sequence to use.
 	 * @param	Time			Current time to solve for.
-	 * @param	bLooping		True when looping the stream in intended.
 	 * @return					None. 
 	 */
 	virtual void GetPoseTranslations(	
 		FTransformArray& Atoms, 
 		const BoneTrackArray& DesiredPairs,
 		const UAnimSequence& Seq,
-		float Time,
-		bool bLooping) PURE_VIRTUAL(AnimEncoding::GetPoseTranslations,);
+		float Time) PURE_VIRTUAL(AnimEncoding::GetPoseTranslations,);
 
 	/**
 	 * Decompress all requested translation components from an Animation Sequence
@@ -240,15 +230,13 @@ public:
 	 * @param	DesiredPairs	Array of requested bone information
 	 * @param	Seq				The animation sequence to use.
 	 * @param	Time			Current time to solve for.
-	 * @param	bLooping		True when looping the stream in intended.
 	 * @return					None. 
 	 */
 	virtual void GetPoseScales(	
 		FTransformArray& Atoms, 
 		const BoneTrackArray& DesiredPairs,
 		const UAnimSequence& Seq,
-		float Time,
-		bool bLooping) PURE_VIRTUAL(AnimEncoding::GetPoseScales,);
+		float Time) PURE_VIRTUAL(AnimEncoding::GetPoseScales,);
 #endif
 protected:
 
@@ -257,7 +245,6 @@ protected:
 	 *
 	 * @param	Seq				The UAnimSequence container.
 	 * @param	RelativePos		The relative position to solve in the range [0,1] inclusive.
-	 * @param	bLooping		true if the animation should be consider cyclic (last frame interpolates back to the start)
 	 * @param	NumKeys			The number of keys present in the track being solved.
 	 * @param	PosIndex0Out	Output value for the closest key index before the RelativePos specified.
 	 * @param	PosIndex1Out	Output value for the closest key index after the RelativePos specified.
@@ -266,7 +253,6 @@ protected:
 	static float TimeToIndex(
 		const UAnimSequence& Seq,
 		float RelativePos,
-		bool bLooping,
 		int32 NumKeys,
 		int32 &PosIndex0Out,
 		int32 &PosIndex1Out);
@@ -277,7 +263,6 @@ protected:
 	 * @param	Seq				The UAnimSequence container.
 	 * @param	FrameTable		The frame table containing a frame index for each key.
 	 * @param	RelativePos		The relative position to solve in the range [0,1] inclusive.
-	 * @param	bLooping		true if the animation should be consider cyclic (last frame interpolates back to the start)
 	 * @param	NumKeys			The number of keys present in the track being solved.
 	 * @param	PosIndex0Out	Output value for the closest key index before the RelativePos specified.
 	 * @param	PosIndex1Out	Output value for the closest key index after the RelativePos specified.
@@ -287,7 +272,6 @@ protected:
 		const UAnimSequence& Seq,
 		const uint8* FrameTable,
 		float RelativePos,
-		bool bLooping,
 		int32 NumKeys,
 		int32 &PosIndex0Out,
 		int32 &PosIndex1Out);
@@ -309,14 +293,12 @@ public:
 	 * @param	Seq				An Animation Sequence to extract the BoneAtom from.
 	 * @param	TrackIndex		The index of the track desired in the Animation Sequence.
 	 * @param	Time			The time (in seconds) to calculate the BoneAtom for.
-	 * @param	bLooping		true if the animation should be played in a cyclic manner.
 	 */
 	virtual void GetBoneAtom(
 		FTransform& OutAtom,
 		const UAnimSequence& Seq,
 		int32 TrackIndex,
-		float Time,
-		bool bLooping);
+		float Time);
 
 	/**
 	 * Decompress the Rotation component of a BoneAtom
@@ -327,7 +309,6 @@ public:
 	 * @param	NumKeys			The number of keys present in Stream.
 	 * @param	Time			Current time to solve for.
 	 * @param	RelativePos		Current position within the animation to solve for in the range [0.0,1.0].
-	 * @param	bLooping		True when looping the stream in intended.
 	 * @return					None. 
 	 */
 	virtual void GetBoneAtomRotation(	
@@ -336,8 +317,7 @@ public:
 		const uint8* RESTRICT Stream,
 		int32 NumKeys,
 		float Time,
-		float RelativePos,
-		bool bLooping) PURE_VIRTUAL(AnimEncoding::GetBoneAtomRotation,);
+		float RelativePos) PURE_VIRTUAL(AnimEncoding::GetBoneAtomRotation,);
 
 	/**
 	 * Decompress the Translation component of a BoneAtom
@@ -348,7 +328,6 @@ public:
 	 * @param	NumKeys			The number of keys present in Stream.
 	 * @param	Time			Current time to solve for.
 	 * @param	RelativePos		Current position within the animation to solve for in the range [0.0,1.0].
-	 * @param	bLooping		True when looping the stream in intended.
 	 * @return					None. 
 	 */
 	virtual void GetBoneAtomTranslation(	
@@ -357,8 +336,7 @@ public:
 		const uint8* RESTRICT Stream,
 		int32 NumKeys,
 		float Time,
-		float RelativePos,
-		bool bLooping) PURE_VIRTUAL(AnimEncoding::GetBoneAtomTranslation,);
+		float RelativePos) PURE_VIRTUAL(AnimEncoding::GetBoneAtomTranslation,);
 
 	/**
 	 * Decompress the Scale component of a BoneAtom
@@ -369,7 +347,6 @@ public:
 	 * @param	NumKeys			The number of keys present in Stream.
 	 * @param	Time			Current time to solve for.
 	 * @param	RelativePos		Current position within the animation to solve for in the range [0.0,1.0].
-	 * @param	bLooping		True when looping the stream in intended.
 	 * @return					None. 
 	 */
 	virtual void GetBoneAtomScale(	
@@ -378,8 +355,7 @@ public:
 		const uint8* RESTRICT Stream,
 		int32 NumKeys,
 		float Time,
-		float RelativePos,
-		bool bLooping) PURE_VIRTUAL(AnimEncoding::GetBoneAtomScale,);
+		float RelativePos) PURE_VIRTUAL(AnimEncoding::GetBoneAtomScale,);
 
 	/**
 	 * Handles Byte-swapping incoming animation data from a MemoryReader
@@ -512,7 +488,6 @@ public:
  *
  * @param	Seq				The UAnimSequence container.
  * @param	RelativePos		The relative position to solve in the range [0,1] inclusive.
- * @param	bLooping		true if the animation should be consider cyclic (last frame interpolates back to the start)
  * @param	NumKeys			The number of keys present in the track being solved.
  * @param	PosIndex0Out	Output value for the closest key index before the RelativePos specified.
  * @param	PosIndex1Out	Output value for the closest key index after the RelativePos specified.
@@ -521,7 +496,6 @@ public:
 FORCEINLINE float AnimEncoding::TimeToIndex(
 	const UAnimSequence& Seq,
 	float RelativePos,
-	bool bLooping,
 	int32 NumKeys,
 	int32 &PosIndex0Out,
 	int32 &PosIndex1Out)
@@ -529,7 +503,6 @@ FORCEINLINE float AnimEncoding::TimeToIndex(
 	static int32		NumKeysCache = 0; // this value is guaranteed to not be used for valid data
 	static float	TimeCache;
 	static float	SequenceLengthCache;
-	static bool	LoopingCache;
 	static int32		PosIndex0CacheOut; 
 	static int32		PosIndex1CacheOut; 
 	static float	AlphaCacheOut;
@@ -545,13 +518,11 @@ FORCEINLINE float AnimEncoding::TimeToIndex(
 	}
 	if (
 		NumKeysCache		!= NumKeys ||
-		LoopingCache		!= bLooping ||
 		SequenceLengthCache != SequenceLength ||
 		TimeCache			!= RelativePos
 		)
 	{
 		NumKeysCache		= NumKeys;
-		LoopingCache		= bLooping;
 		SequenceLengthCache = SequenceLength;
 		TimeCache			= RelativePos;
 		// Check for before-first-frame case.
@@ -563,74 +534,24 @@ FORCEINLINE float AnimEncoding::TimeToIndex(
 		}
 		else
 		{
-			if (!bLooping)
+			NumKeys -= 1; // never used without the minus one in this case
+			// Check for after-last-frame case.
+			if( RelativePos >= 1.0f )
 			{
-				NumKeys -= 1; // never used without the minus one in this case
-				// Check for after-last-frame case.
-				if( RelativePos >= 1.0f )
-				{
-					// If we're not looping, key n-1 is the final key.
-					PosIndex0CacheOut = NumKeys;
-					PosIndex1CacheOut = NumKeys;
-					AlphaCacheOut = 0.0f;
-				}
-				else
-				{
-					// For non-looping animation, the last frame is the ending frame, and has no duration.
-					const float KeyPos = RelativePos * float(NumKeys);
-					checkSlow(KeyPos >= 0.0f);
-					const float KeyPosFloor = floorf(KeyPos);
-					PosIndex0CacheOut = FMath::Min( FMath::TruncToInt(KeyPosFloor), NumKeys );
-					AlphaCacheOut = KeyPos - KeyPosFloor;
-					PosIndex1CacheOut = FMath::Min( PosIndex0CacheOut + 1, NumKeys );
-				}
+				// If we're not looping, key n-1 is the final key.
+				PosIndex0CacheOut = NumKeys;
+				PosIndex1CacheOut = NumKeys;
+				AlphaCacheOut = 0.0f;
 			}
-			else // we are looping
+			else
 			{
-				// Check for after-last-frame case.
-				if( RelativePos >= 1.0f )
-				{
-					// If we're looping, key 0 is the final key.
-					PosIndex0CacheOut = 0;
-					PosIndex1CacheOut = 0;
-					AlphaCacheOut = 0.0f;
-				}
-				else
-				{
-					// Work with animation total frames, to handle looping last->first
-					// Our track might have a different number of frames, and we handle that below.
-					int32 const NumFrames = Seq.NumFrames;
-
-					{
-						// For looping animation, the last frame has duration, and interpolates back to the first one.
-						const float KeyPos = RelativePos * float(NumFrames);
-						checkSlow(KeyPos >= 0.0f);
-						const float KeyPosFloor = floorf(KeyPos);
-						PosIndex0CacheOut = FMath::Min( FMath::TruncToInt(KeyPosFloor), NumFrames - 1 );
-						AlphaCacheOut = KeyPos - KeyPosFloor;
-						PosIndex1CacheOut = PosIndex0CacheOut + 1;
-					}
-
-					// Handle Looping
-					if( PosIndex1CacheOut == NumFrames )
-					{
-						PosIndex0CacheOut = NumKeys - 1;
-						PosIndex1CacheOut = 0;
-					}
-					// Non Looping! Special treatment if NumKeys is not the same as NumFrames...
-					else if( NumFrames != NumKeys )
-					{
-						// Since we're not looping first to last, chop off the last chunk
-						// And do a simple non looping interp between those keys.
-						float const AdjustedPosition = (RelativePos * float(NumFrames)) / float(NumFrames-1);
-						float const KeyPos = AdjustedPosition * float(NumKeys-1);
-						checkSlow(KeyPos >= 0.0f);
-						float const KeyPosFloor = floorf(KeyPos);
-						PosIndex0CacheOut = FMath::Min( FMath::TruncToInt(KeyPosFloor), NumKeys - 1 );
-						AlphaCacheOut = KeyPos - KeyPosFloor;
-						PosIndex1CacheOut = FMath::Min( PosIndex0CacheOut + 1, NumKeys - 1 );
-					}
-				}
+				// For non-looping animation, the last frame is the ending frame, and has no duration.
+				const float KeyPos = RelativePos * float(NumKeys);
+				checkSlow(KeyPos >= 0.0f);
+				const float KeyPosFloor = floorf(KeyPos);
+				PosIndex0CacheOut = FMath::Min( FMath::TruncToInt(KeyPosFloor), NumKeys );
+				AlphaCacheOut = KeyPos - KeyPosFloor;
+				PosIndex1CacheOut = FMath::Min( PosIndex0CacheOut + 1, NumKeys );
 			}
 		}
 	}
@@ -700,8 +621,7 @@ FORCEINLINE_DEBUGGABLE int32 FindLowKeyIndex(
  * @param	Seq				The UAnimSequence container.
  * @param	FrameTable		The frame table containing a frame index for each key.
  * @param	RelativePos		The relative position to solve in the range [0,1] inclusive.
- * @param	bLooping		true if the animation should be consider cyclic (last frame interpolates back to the start)
- * @param	NumKeys			The number of keys present in the track being solved.
+  * @param	NumKeys			The number of keys present in the track being solved.
  * @param	PosIndex0Out	Output value for the closest key index before the RelativePos specified.
  * @param	PosIndex1Out	Output value for the closest key index after the RelativePos specified.
  * @return	The rate at which to interpolate the two keys returned to obtain the final result.
@@ -710,7 +630,6 @@ FORCEINLINE float AnimEncoding::TimeToIndex(
 	const UAnimSequence& Seq,
 	const uint8* FrameTable,
 	float RelativePos,
-	bool bLooping,
 	int32 NumKeys,
 	int32 &PosIndex0Out,
 	int32 &PosIndex1Out)
@@ -723,11 +642,6 @@ FORCEINLINE float AnimEncoding::TimeToIndex(
 	
 	int32 TotalFrames = Seq.NumFrames-1;
 	int32 EndingKey = LastKey;
-	if (bLooping)
-	{
-		TotalFrames = Seq.NumFrames;
-		EndingKey = 0;
-	}
 
 	if (NumKeys < 2 || RelativePos <= 0.f)
 	{
diff --git a/Engine/Source/Runtime/Engine/Public/AnimEncoding_ConstantKeyLerp.h b/Engine/Source/Runtime/Engine/Public/AnimEncoding_ConstantKeyLerp.h
--- a/Engine/Source/Runtime/Engine/Public/AnimEncoding_ConstantKeyLerp.h
+++ b/Engine/Source/Runtime/Engine/Public/AnimEncoding_ConstantKeyLerp.h
@@ -118,7 +118,6 @@ public:
 	 * @param	NumKeys			The number of keys present in Stream.
 	 * @param	Time			Current time to solve for.
 	 * @param	RelativePos		Current position within the animation to solve for in the range [0.0,1.0].
-	 * @param	bLooping		True when looping the stream in intended.
 	 * @return					None. 
 	 */
 	void GetBoneAtomRotation(	
@@ -127,8 +126,7 @@ public:
 		const uint8* RESTRICT Stream,
 		int32 NumKeys,
 		float Time,
-		float RelativePos,
-		bool bLooping);
+		float RelativePos);
 
 	/**
 	 * Decompress the Translation component of a BoneAtom
@@ -138,7 +136,6 @@ public:
 	 * @param	NumKeys			The number of keys present in Stream.
 	 * @param	Time			Current time to solve for.
 	 * @param	RelativePos		Current position within the animation to solve for in the range [0.0,1.0].
-	 * @param	bLooping		True when looping the stream in intended.
 	 * @return					None. 
 	 */
 	void GetBoneAtomTranslation(	
@@ -147,8 +144,7 @@ public:
 		const uint8* RESTRICT Stream,
 		int32 NumKeys,
 		float Time,
-		float RelativePos,
-		bool bLooping);
+		float RelativePos);
 
 	/**
 	 * Decompress the Scale component of a BoneAtom
@@ -158,7 +154,6 @@ public:
 	 * @param	NumKeys			The number of keys present in Stream.
 	 * @param	Time			Current time to solve for.
 	 * @param	RelativePos		Current position within the animation to solve for in the range [0.0,1.0].
-	 * @param	bLooping		True when looping the stream in intended.
 	 * @return					None. 
 	 */
 	void GetBoneAtomScale(	
@@ -167,8 +162,7 @@ public:
 		const uint8* RESTRICT Stream,
 		int32 NumKeys,
 		float Time,
-		float RelativePos,
-		bool bLooping);
+		float RelativePos);
 #if USE_ANIMATION_CODEC_BATCH_SOLVER
 
 	/**
@@ -178,15 +172,13 @@ public:
 	 * @param	DesiredPairs	Array of requested bone information
 	 * @param	Seq				The animation sequence to use.
 	 * @param	Time			Current time to solve for.
-	 * @param	bLooping		True when looping the stream in intended.
 	 * @return					None. 
 	 */
 	void GetPoseRotations(	
 		FTransformArray& Atoms, 
 		const BoneTrackArray& DesiredPairs,
 		const UAnimSequence& Seq,
-		float RelativePos,
-		bool bLooping);
+		float RelativePos);
 
 	/**
 	 * Decompress all requested translation components from an Animation Sequence
@@ -195,15 +187,13 @@ public:
 	 * @param	DesiredPairs	Array of requested bone information
 	 * @param	Seq				The animation sequence to use.
 	 * @param	Time			Current time to solve for.
-	 * @param	bLooping		True when looping the stream in intended.
 	 * @return					None. 
 	 */
 	void GetPoseTranslations(	
 		FTransformArray& Atoms,
 		const BoneTrackArray& DesiredPairs,
 		const UAnimSequence& Seq,
-		float RelativePos,
-		bool bLooping);
+		float RelativePos);
 
 	/**
 	 * Decompress all requested Scale components from an Animation Sequence
@@ -212,15 +202,13 @@ public:
 	 * @param	DesiredPairs	Array of requested bone information
 	 * @param	Seq				The animation sequence to use.
 	 * @param	Time			Current time to solve for.
-	 * @param	bLooping		True when looping the stream in intended.
 	 * @return					None. 
 	 */
 	void GetPoseScales(	
 		FTransformArray& Atoms,
 		const BoneTrackArray& DesiredPairs,
 		const UAnimSequence& Seq,
-		float RelativePos,
-		bool bLooping);
+		float RelativePos);
 #endif
 
 };
@@ -235,7 +223,6 @@ public:
  * @param	NumKeys			The number of keys present in Stream.
  * @param	Time			Current time to solve for.
  * @param	RelativePos		Current position within the animation to solve for in the range [0.0,1.0].
- * @param	bLooping		True when looping the stream in intended.
  * @return					None. 
  */
 template<int32 FORMAT>
@@ -245,8 +232,7 @@ FORCEINLINE void AEFConstantKeyLerp<FORMAT>::GetBoneAtomRotation(
 	const uint8* RESTRICT RotStream,
 	int32 NumRotKeys,
 	float Time,
-	float RelativePos,
-	bool bLooping)
+	float RelativePos)
 {
 	if (NumRotKeys == 1)
 	{
@@ -259,7 +245,7 @@ FORCEINLINE void AEFConstantKeyLerp<FORMAT>::GetBoneAtomRotation(
 	{
 		int32 Index0;
 		int32 Index1;
-		float Alpha = TimeToIndex(Seq,RelativePos,bLooping,NumRotKeys,Index0,Index1);
+		float Alpha = TimeToIndex(Seq,RelativePos,NumRotKeys,Index0,Index1);
 
 		const int32 RotationStreamOffset = (FORMAT == ACF_IntervalFixed32NoW) ? (sizeof(float)*6) : 0; // offset past Min and Range data
 
@@ -303,7 +289,6 @@ FORCEINLINE void AEFConstantKeyLerp<FORMAT>::GetBoneAtomRotation(
  * @param	NumKeys			The number of keys present in Stream.
  * @param	Time			Current time to solve for.
  * @param	RelativePos		Current position within the animation to solve for in the range [0.0,1.0].
- * @param	bLooping		True when looping the stream in intended.
  * @return					None. 
  */
 template<int32 FORMAT>
@@ -313,12 +298,11 @@ FORCEINLINE void AEFConstantKeyLerp<FORMAT>::GetBoneAtomTranslation(
 	const uint8* RESTRICT TransStream,
 	int32 NumTransKeys,
 	float Time,
-	float RelativePos,
-	bool bLooping)
+	float RelativePos)
 {
 	int32 Index0;
 	int32 Index1;
-	float Alpha = TimeToIndex(Seq,RelativePos,bLooping,NumTransKeys,Index0,Index1);
+	float Alpha = TimeToIndex(Seq,RelativePos,NumTransKeys,Index0,Index1);
 
 	const int32 TransStreamOffset = ((FORMAT == ACF_IntervalFixed32NoW) && NumTransKeys > 1) ? (sizeof(float)*6) : 0; // offset past Min and Range data
 
@@ -350,7 +334,6 @@ FORCEINLINE void AEFConstantKeyLerp<FORMAT>::GetBoneAtomTranslation(
  * @param	NumKeys			The number of keys present in Stream.
  * @param	Time			Current time to solve for.
  * @param	RelativePos		Current position within the animation to solve for in the range [0.0,1.0].
- * @param	bLooping		True when looping the stream in intended.
  * @return					None. 
  */
 template<int32 FORMAT>
@@ -360,12 +343,11 @@ FORCEINLINE void AEFConstantKeyLerp<FORMAT>::GetBoneAtomScale(
 	const uint8* RESTRICT ScaleStream,
 	int32 NumScaleKeys,
 	float Time,
-	float RelativePos,
-	bool bLooping)
+	float RelativePos)
 {
 	int32 Index0;
 	int32 Index1;
-	float Alpha = TimeToIndex(Seq,RelativePos,bLooping,NumScaleKeys,Index0,Index1);
+	float Alpha = TimeToIndex(Seq,RelativePos,NumScaleKeys,Index0,Index1);
 
 	const int32 ScaleStreamOffset = ((FORMAT == ACF_IntervalFixed32NoW) && NumScaleKeys > 1) ? (sizeof(float)*6) : 0; // offset past Min and Range data
 
@@ -398,7 +380,6 @@ FORCEINLINE void AEFConstantKeyLerp<FORMAT>::GetBoneAtomScale(
  * @param	DesiredPairs	Array of requested bone information
  * @param	Seq				The animation sequence to use.
  * @param	Time			Current time to solve for.
- * @param	bLooping		True when looping the stream in intended.
  * @return					None. 
  */
 template<int32 FORMAT>
@@ -406,8 +387,7 @@ inline void AEFConstantKeyLerp<FORMAT>::GetPoseRotations(
 	FTransformArray& Atoms, 
 	const BoneTrackArray& DesiredPairs,
 	const UAnimSequence& Seq,
-	float Time,
-	bool bLooping)
+	float Time)
 {
 	const int32 PairCount = DesiredPairs.Num();
 	const float RelativePos = Time / (float)Seq.SequenceLength;
@@ -425,7 +405,7 @@ inline void AEFConstantKeyLerp<FORMAT>::GetPoseRotations(
 		const uint8* RESTRICT RotStream		= Seq.CompressedByteStream.GetTypedData()+RotKeysOffset;
 
 		// call the decoder directly (not through the vtable)
-		AEFConstantKeyLerp<FORMAT>::GetBoneAtomRotation(BoneAtom, Seq, RotStream, NumRotKeys, Time, RelativePos, bLooping);
+		AEFConstantKeyLerp<FORMAT>::GetBoneAtomRotation(BoneAtom, Seq, RotStream, NumRotKeys, Time, RelativePos);
 	}
 }
 
@@ -436,7 +416,6 @@ inline void AEFConstantKeyLerp<FORMAT>::GetPoseRotations(
  * @param	DesiredPairs	Array of requested bone information
  * @param	Seq				The animation sequence to use.
  * @param	Time			Current time to solve for.
- * @param	bLooping		True when looping the stream in intended.
  * @return					None. 
  */
 template<int32 FORMAT>
@@ -444,8 +423,7 @@ inline void AEFConstantKeyLerp<FORMAT>::GetPoseTranslations(
 	FTransformArray& Atoms, 
 	const BoneTrackArray& DesiredPairs,
 	const UAnimSequence& Seq,
-	float Time,
-	bool bLooping)
+	float Time)
 {
 	const int32 PairCount= DesiredPairs.Num();
 	const float RelativePos = Time / (float)Seq.SequenceLength;
@@ -479,7 +457,7 @@ inline void AEFConstantKeyLerp<FORMAT>::GetPoseTranslations(
 		const uint8* RESTRICT TransStream = Seq.CompressedByteStream.GetTypedData()+TransKeysOffset;
 
 		// call the decoder directly (not through the vtable)
-		AEFConstantKeyLerp<FORMAT>::GetBoneAtomTranslation(BoneAtom, Seq, TransStream, NumTransKeys, Time, RelativePos, bLooping);
+		AEFConstantKeyLerp<FORMAT>::GetBoneAtomTranslation(BoneAtom, Seq, TransStream, NumTransKeys, Time, RelativePos);
 	}
 }
 
@@ -490,7 +468,6 @@ inline void AEFConstantKeyLerp<FORMAT>::GetPoseTranslations(
  * @param	DesiredPairs	Array of requested bone information
  * @param	Seq				The animation sequence to use.
  * @param	Time			Current time to solve for.
- * @param	bLooping		True when looping the stream in intended.
  * @return					None. 
  */
 template<int32 FORMAT>
@@ -498,8 +475,7 @@ inline void AEFConstantKeyLerp<FORMAT>::GetPoseScales(
 	FTransformArray& Atoms, 
 	const BoneTrackArray& DesiredPairs,
 	const UAnimSequence& Seq,
-	float Time,
-	bool bLooping)
+	float Time)
 {
 	check (Seq.CompressedScaleOffsets.IsValid());
 
@@ -534,7 +510,7 @@ inline void AEFConstantKeyLerp<FORMAT>::GetPoseScales(
 		const uint8* RESTRICT ScaleStream = Seq.CompressedByteStream.GetTypedData()+ScaleKeysOffset;
 
 		// call the decoder directly (not through the vtable)
-		AEFConstantKeyLerp<FORMAT>::GetBoneAtomScale(BoneAtom, Seq, ScaleStream, NumScaleKeys, Time, RelativePos, bLooping);
+		AEFConstantKeyLerp<FORMAT>::GetBoneAtomScale(BoneAtom, Seq, ScaleStream, NumScaleKeys, Time, RelativePos);
 	}
 }
 #endif
diff --git a/Engine/Source/Runtime/Engine/Public/AnimEncoding_PerTrackCompression.h b/Engine/Source/Runtime/Engine/Public/AnimEncoding_PerTrackCompression.h
--- a/Engine/Source/Runtime/Engine/Public/AnimEncoding_PerTrackCompression.h
+++ b/Engine/Source/Runtime/Engine/Public/AnimEncoding_PerTrackCompression.h
@@ -45,14 +45,12 @@ public:
 	 * @param	Seq				An Animation Sequence to extract the BoneAtom from.
 	 * @param	TrackIndex		The index of the track desired in the Animation Sequence.
 	 * @param	Time			The time (in seconds) to calculate the BoneAtom for.
-	 * @param	bLooping		true if the animation should be played in a cyclic manner.
 	 */
 	virtual void GetBoneAtom(
 		FTransform& OutAtom,
 		const UAnimSequence& Seq,
 		int32 TrackIndex,
-		float Time,
-		bool bLooping);
+		float Time);
 
 #if USE_ANIMATION_CODEC_BATCH_SOLVER
 
@@ -63,15 +61,13 @@ public:
 	 * @param	DesiredPairs	Array of requested bone information
 	 * @param	Seq				The animation sequence to use.
 	 * @param	Time			Current time to solve for.
-	 * @param	bLooping		True when looping the stream in intended.
 	 * @return					None. 
 	 */
 	virtual void GetPoseRotations(	
 		FTransformArray& Atoms, 
 		const BoneTrackArray& DesiredPairs,
 		const UAnimSequence& Seq,
-		float Time,
-		bool bLooping);
+		float Time);
 
 	/**
 	 * Decompress all requested translation components from an Animation Sequence
@@ -80,15 +76,13 @@ public:
 	 * @param	DesiredPairs	Array of requested bone information
 	 * @param	Seq				The animation sequence to use.
 	 * @param	Time			Current time to solve for.
-	 * @param	bLooping		True when looping the stream in intended.
 	 * @return					None. 
 	 */
 	virtual void GetPoseTranslations(	
 		FTransformArray& Atoms, 
 		const BoneTrackArray& DesiredPairs,
 		const UAnimSequence& Seq,
-		float Time,
-		bool bLooping);
+		float Time);
 
 	/**
 	 * Decompress all requested Scale components from an Animation Sequence
@@ -97,15 +91,13 @@ public:
 	 * @param	DesiredPairs	Array of requested bone information
 	 * @param	Seq				The animation sequence to use.
 	 * @param	Time			Current time to solve for.
-	 * @param	bLooping		True when looping the stream in intended.
 	 * @return					None. 
 	 */
 	virtual void GetPoseScales(	
 		FTransformArray& Atoms, 
 		const BoneTrackArray& DesiredPairs,
 		const UAnimSequence& Seq,
-		float Time,
-		bool bLooping);
+		float Time);
 #endif
 
 protected:
@@ -135,7 +127,6 @@ protected:
 	 * @param	NumKeys			The number of keys present in Stream.
 	 * @param	Time			Current time to solve for.
 	 * @param	RelativePos		Current position within the animation to solve for in the range [0.0,1.0].
-	 * @param	bLooping		True when looping the stream in intended.
 	 * @return					None. 
 	 */
 	static void GetBoneAtomRotation(	
@@ -143,8 +134,7 @@ protected:
 		const UAnimSequence& Seq,
 		int32 Offset,
 		float Time,
-		float RelativePos,
-		bool bLooping);
+		float RelativePos);
 
 	/**
 	 * Decompress the Translation component of a BoneAtom
@@ -155,7 +145,6 @@ protected:
 	 * @param	NumKeys			The number of keys present in Stream.
 	 * @param	Time			Current time to solve for.
 	 * @param	RelativePos		Current position within the animation to solve for in the range [0.0,1.0].
-	 * @param	bLooping		True when looping the stream in intended.
 	 * @return					None. 
 	 */
 	static void GetBoneAtomTranslation(	
@@ -163,8 +152,7 @@ protected:
 		const UAnimSequence& Seq,
 		int32 Offset,
 		float Time,
-		float RelativePos,
-		bool bLooping);
+		float RelativePos);
 
 	/**
 	 * Decompress the Scale component of a BoneAtom
@@ -175,7 +163,6 @@ protected:
 	 * @param	NumKeys			The number of keys present in Stream.
 	 * @param	Time			Current time to solve for.
 	 * @param	RelativePos		Current position within the animation to solve for in the range [0.0,1.0].
-	 * @param	bLooping		True when looping the stream in intended.
 	 * @return					None. 
 	 */
 	static void GetBoneAtomScale(	
@@ -183,7 +170,6 @@ protected:
 		const UAnimSequence& Seq,
 		int32 Offset,
 		float Time,
-		float RelativePos,
-		bool bLooping);
+		float RelativePos);
 };
 
diff --git a/Engine/Source/Runtime/Engine/Public/AnimEncoding_VariableKeyLerp.h b/Engine/Source/Runtime/Engine/Public/AnimEncoding_VariableKeyLerp.h
--- a/Engine/Source/Runtime/Engine/Public/AnimEncoding_VariableKeyLerp.h
+++ b/Engine/Source/Runtime/Engine/Public/AnimEncoding_VariableKeyLerp.h
@@ -120,7 +120,6 @@ public:
 	 * @param	NumKeys			The number of keys present in Stream.
 	 * @param	Time			Current time to solve for.
 	 * @param	RelativePos		Current position within the animation to solve for in the range [0.0,1.0].
-	 * @param	bLooping		True when looping the stream in intended.
 	 * @return					None. 
 	 */
 	void GetBoneAtomRotation(	
@@ -129,8 +128,7 @@ public:
 		const uint8* RESTRICT Stream,
 		int32 NumKeys,
 		float Time,
-		float RelativePos,
-		bool bLooping);
+		float RelativePos);
 
 	/**
 	 * Decompress the Translation component of a BoneAtom
@@ -140,7 +138,6 @@ public:
 	 * @param	NumKeys			The number of keys present in Stream.
 	 * @param	Time			Current time to solve for.
 	 * @param	RelativePos		Current position within the animation to solve for in the range [0.0,1.0].
-	 * @param	bLooping		True when looping the stream in intended.
 	 * @return					None. 
 	 */
 	void GetBoneAtomTranslation(	
@@ -149,8 +146,7 @@ public:
 		const uint8* RESTRICT Stream,
 		int32 NumKeys,
 		float Time,
-		float RelativePos,
-		bool bLooping);
+		float RelativePos);
 
 	
 	/**
@@ -161,7 +157,6 @@ public:
 	 * @param	NumKeys			The number of keys present in Stream.
 	 * @param	Time			Current time to solve for.
 	 * @param	RelativePos		Current position within the animation to solve for in the range [0.0,1.0].
-	 * @param	bLooping		True when looping the stream in intended.
 	 * @return					None. 
 	 */
 	void GetBoneAtomScale(	
@@ -170,8 +165,7 @@ public:
 		const uint8* RESTRICT Stream,
 		int32 NumKeys,
 		float Time,
-		float RelativePos,
-		bool bLooping);
+		float RelativePos);
 
 #if USE_ANIMATION_CODEC_BATCH_SOLVER
 
@@ -182,15 +176,13 @@ public:
 	 * @param	DesiredPairs	Array of requested bone information
 	 * @param	Seq				The animation sequence to use.
 	 * @param	Time			Current time to solve for.
-	 * @param	bLooping		True when looping the stream in intended.
 	 * @return					None. 
 	 */
 	void GetPoseRotations(	
 		FTransformArray& Atoms, 
 		const BoneTrackArray& DesiredPairs,
 		const UAnimSequence& Seq,
-		float RelativePos,
-		bool bLooping);
+		float RelativePos);
 
 	/**
 	 * Decompress all requested translation components from an Animation Sequence
@@ -199,15 +191,13 @@ public:
 	 * @param	DesiredPairs	Array of requested bone information
 	 * @param	Seq				The animation sequence to use.
 	 * @param	Time			Current time to solve for.
-	 * @param	bLooping		True when looping the stream in intended.
 	 * @return					None. 
 	 */
 	void GetPoseTranslations(	
 		FTransformArray& Atoms,
 		const BoneTrackArray& DesiredPairs,
 		const UAnimSequence& Seq,
-		float RelativePos,
-		bool bLooping);
+		float RelativePos);
 
 	/**
 	 * Decompress all requested Scale components from an Animation Sequence
@@ -216,15 +206,13 @@ public:
 	 * @param	DesiredPairs	Array of requested bone information
 	 * @param	Seq				The animation sequence to use.
 	 * @param	Time			Current time to solve for.
-	 * @param	bLooping		True when looping the stream in intended.
 	 * @return					None. 
 	 */
 	void GetPoseScales(	
 		FTransformArray& Atoms,
 		const BoneTrackArray& DesiredPairs,
 		const UAnimSequence& Seq,
-		float RelativePos,
-		bool bLooping);
+		float RelativePos);
 #endif
 
 };
@@ -238,7 +226,6 @@ public:
  * @param	NumKeys			The number of keys present in Stream.
  * @param	Time			Current time to solve for.
  * @param	RelativePos		Current position within the animation to solve for in the range [0.0,1.0].
- * @param	bLooping		True when looping the stream in intended.
  * @return					None. 
  */
 template<int32 FORMAT>
@@ -248,8 +235,7 @@ FORCEINLINE_DEBUGGABLE void AEFVariableKeyLerp<FORMAT>::GetBoneAtomRotation(
 	const uint8* RESTRICT RotStream,
 	int32 NumRotKeys,
 	float Time,
-	float RelativePos,
-	bool bLooping)
+	float RelativePos)
 {
 	if (NumRotKeys == 1)
 	{
@@ -266,7 +252,7 @@ FORCEINLINE_DEBUGGABLE void AEFVariableKeyLerp<FORMAT>::GetBoneAtomRotation(
 
 		int32 Index0;
 		int32 Index1;
-		float Alpha = TimeToIndex(Seq,FrameTable,RelativePos,bLooping,NumRotKeys,Index0,Index1);
+		float Alpha = TimeToIndex(Seq,FrameTable,RelativePos,NumRotKeys,Index0,Index1);
 
 
 		if (Index0 != Index1)
@@ -309,7 +295,6 @@ FORCEINLINE_DEBUGGABLE void AEFVariableKeyLerp<FORMAT>::GetBoneAtomRotation(
  * @param	NumKeys			The number of keys present in Stream.
  * @param	Time			Current time to solve for.
  * @param	RelativePos		Current position within the animation to solve for in the range [0.0,1.0].
- * @param	bLooping		True when looping the stream in intended.
  * @return					None. 
  */
 template<int32 FORMAT>
@@ -319,15 +304,14 @@ FORCEINLINE_DEBUGGABLE void AEFVariableKeyLerp<FORMAT>::GetBoneAtomTranslation(
 	const uint8* RESTRICT TransStream,
 	int32 NumTransKeys,
 	float Time,
-	float RelativePos,
-	bool bLooping)
+	float RelativePos)
 {
 	const uint8* RESTRICT FrameTable= TransStream +(NumTransKeys*CompressedTranslationStrides[FORMAT]*CompressedTranslationNum[FORMAT]);
 	FrameTable= Align(FrameTable, 4);
 
 	int32 Index0;
 	int32 Index1;
-	float Alpha = TimeToIndex(Seq,FrameTable,RelativePos,bLooping,NumTransKeys,Index0,Index1);
+	float Alpha = TimeToIndex(Seq,FrameTable,RelativePos,NumTransKeys,Index0,Index1);
 	const int32 TransStreamOffset = ((FORMAT == ACF_IntervalFixed32NoW) && NumTransKeys > 1) ? (sizeof(float)*6) : 0; // offset past Min and Range data
 
 	if (Index0 != Index1)
@@ -358,7 +342,6 @@ FORCEINLINE_DEBUGGABLE void AEFVariableKeyLerp<FORMAT>::GetBoneAtomTranslation(
  * @param	NumKeys			The number of keys present in Stream.
  * @param	Time			Current time to solve for.
  * @param	RelativePos		Current position within the animation to solve for in the range [0.0,1.0].
- * @param	bLooping		True when looping the stream in intended.
  * @return					None. 
  */
 template<int32 FORMAT>
@@ -368,15 +351,14 @@ FORCEINLINE_DEBUGGABLE void AEFVariableKeyLerp<FORMAT>::GetBoneAtomScale(
 	const uint8* RESTRICT ScaleStream,
 	int32 NumScaleKeys,
 	float Time,
-	float RelativePos,
-	bool bLooping)
+	float RelativePos)
 {
 	const uint8* RESTRICT FrameTable= ScaleStream +(NumScaleKeys*CompressedScaleStrides[FORMAT]*CompressedScaleNum[FORMAT]);
 	FrameTable= Align(FrameTable, 4);
 
 	int32 Index0;
 	int32 Index1;
-	float Alpha = TimeToIndex(Seq,FrameTable,RelativePos,bLooping,NumScaleKeys,Index0,Index1);
+	float Alpha = TimeToIndex(Seq,FrameTable,RelativePos,NumScaleKeys,Index0,Index1);
 	const int32 ScaleStreamOffset = ((FORMAT == ACF_IntervalFixed32NoW) && NumScaleKeys > 1) ? (sizeof(float)*6) : 0; // offset past Min and Range data
 
 	if (Index0 != Index1)
@@ -408,7 +390,6 @@ FORCEINLINE_DEBUGGABLE void AEFVariableKeyLerp<FORMAT>::GetBoneAtomScale(
  * @param	DesiredPairs	Array of requested bone information
  * @param	Seq				The animation sequence to use.
  * @param	Time			Current time to solve for.
- * @param	bLooping		True when looping the stream in intended.
  * @return					None. 
  */
 template<int32 FORMAT>
@@ -416,8 +397,7 @@ FORCEINLINE_DEBUGGABLE void AEFVariableKeyLerp<FORMAT>::GetPoseRotations(
 	FTransformArray& Atoms, 
 	const BoneTrackArray& DesiredPairs,
 	const UAnimSequence& Seq,
-	float Time,
-	bool bLooping)
+	float Time)
 {
 	const int32 PairCount = DesiredPairs.Num();
 	const float RelativePos = Time / (float)Seq.SequenceLength;
@@ -435,7 +415,7 @@ FORCEINLINE_DEBUGGABLE void AEFVariableKeyLerp<FORMAT>::GetPoseRotations(
 		const uint8* RESTRICT RotStream		= Seq.CompressedByteStream.GetTypedData()+RotKeysOffset;
 
 		// call the decoder directly (not through the vtable)
-		AEFVariableKeyLerp<FORMAT>::GetBoneAtomRotation(BoneAtom, Seq, RotStream, NumRotKeys, Time, RelativePos, bLooping);
+		AEFVariableKeyLerp<FORMAT>::GetBoneAtomRotation(BoneAtom, Seq, RotStream, NumRotKeys, Time, RelativePos);
 	}
 }
 
@@ -446,7 +426,6 @@ FORCEINLINE_DEBUGGABLE void AEFVariableKeyLerp<FORMAT>::GetPoseRotations(
  * @param	DesiredPairs	Array of requested bone information
  * @param	Seq				The animation sequence to use.
  * @param	Time			Current time to solve for.
- * @param	bLooping		True when looping the stream in intended.
  * @return					None. 
  */
 template<int32 FORMAT>
@@ -454,8 +433,7 @@ FORCEINLINE_DEBUGGABLE void AEFVariableKeyLerp<FORMAT>::GetPoseTranslations(
 	FTransformArray& Atoms, 
 	const BoneTrackArray& DesiredPairs,
 	const UAnimSequence& Seq,
-	float Time,
-	bool bLooping)
+	float Time)
 {
 	const int32 PairCount= DesiredPairs.Num();
 	const float RelativePos = Time / (float)Seq.SequenceLength;
@@ -473,7 +451,7 @@ FORCEINLINE_DEBUGGABLE void AEFVariableKeyLerp<FORMAT>::GetPoseTranslations(
 		const uint8* RESTRICT TransStream = Seq.CompressedByteStream.GetTypedData()+TransKeysOffset;
 
 		// call the decoder directly (not through the vtable)
-		AEFVariableKeyLerp<FORMAT>::GetBoneAtomTranslation(BoneAtom, Seq, TransStream, NumTransKeys, Time, RelativePos, bLooping);
+		AEFVariableKeyLerp<FORMAT>::GetBoneAtomTranslation(BoneAtom, Seq, TransStream, NumTransKeys, Time, RelativePos);
 	}
 }
 
@@ -484,7 +462,6 @@ FORCEINLINE_DEBUGGABLE void AEFVariableKeyLerp<FORMAT>::GetPoseTranslations(
  * @param	DesiredPairs	Array of requested bone information
  * @param	Seq				The animation sequence to use.
  * @param	Time			Current time to solve for.
- * @param	bLooping		True when looping the stream in intended.
  * @return					None. 
  */
 template<int32 FORMAT>
@@ -492,8 +469,7 @@ FORCEINLINE_DEBUGGABLE void AEFVariableKeyLerp<FORMAT>::GetPoseScales(
 	FTransformArray& Atoms, 
 	const BoneTrackArray& DesiredPairs,
 	const UAnimSequence& Seq,
-	float Time,
-	bool bLooping)
+	float Time)
 {
 	check (Seq.CompressedScaleOffsets.IsValid());
 
@@ -512,7 +488,7 @@ FORCEINLINE_DEBUGGABLE void AEFVariableKeyLerp<FORMAT>::GetPoseScales(
 		const uint8* RESTRICT ScaleStream = Seq.CompressedByteStream.GetTypedData()+ScaleKeysOffset;
 
 		// call the decoder directly (not through the vtable)
-		AEFVariableKeyLerp<FORMAT>::GetBoneAtomScale(BoneAtom, Seq, ScaleStream, NumScaleKeys, Time, RelativePos, bLooping);
+		AEFVariableKeyLerp<FORMAT>::GetBoneAtomScale(BoneAtom, Seq, ScaleStream, NumScaleKeys, Time, RelativePos);
 	}
 }
 #endif
diff --git a/Engine/Source/Runtime/Engine/Public/AnimationRuntime.h b/Engine/Source/Runtime/Engine/Public/AnimationRuntime.h
--- a/Engine/Source/Runtime/Engine/Public/AnimationRuntime.h
+++ b/Engine/Source/Runtime/Engine/Public/AnimationRuntime.h
@@ -196,7 +196,6 @@ public:
 	static void GetPoseFromBlendSpace(
 		UBlendSpaceBase * BlendSpace,
 		TArray<FBlendSampleData>& BlendSampleDataCache,
-		bool bLooping,
 		FBoneContainer & RequiredBones,
 		/*out*/ FTransformArrayA2& ResultAtoms);
 
@@ -216,11 +215,11 @@ public:
 	static void ApplyWeightToTransform(const FBoneContainer & RequiredBones, /*inout*/ FTransformArrayA2& Atoms, float Weight);
 
 	/** 
-	 * Get Key Indices (start/end with alpha from start) with input parameter Time, bLooping, NumFrames
+	 * Get Key Indices (start/end with alpha from start) with input parameter Time, NumFrames
 	 * from % from StartKeyIndex, meaning (CurrentKeyIndex(float)-StartKeyIndex)/(EndKeyIndex-StartKeyIndex)
 	 * by this Start-End, it will be between 0-(NumFrames-1), not number of Pos/Rot key tracks 
 	 **/
-	static void GetKeyIndicesFromTime(int32 & OutKeyIndex1, int32 & OutKeyIndex2, float& OutAlpha, const float Time, const bool bLooping, const int32 NumFrames, const float SequenceLength);
+	static void GetKeyIndicesFromTime(int32 & OutKeyIndex1, int32 & OutKeyIndex2, float& OutAlpha, const float Time, const int32 NumFrames, const float SequenceLength);
 
 	/** 
 	 *	Utility for taking an array of bone indices and ensuring that all parents are present 
