UE4 - more command list work, RHIGetViewportBackBuffer is a globalflush rhi command now, which eliminated a few GetImmediateCOmmandList from the viewport code

diff --git a/Engine/Source/Runtime/Engine/Private/UnrealClient.cpp b/Engine/Source/Runtime/Engine/Private/UnrealClient.cpp
--- a/Engine/Source/Runtime/Engine/Private/UnrealClient.cpp
+++ b/Engine/Source/Runtime/Engine/Private/UnrealClient.cpp
@@ -877,7 +877,7 @@ void FViewport::BeginRenderFrame(FRHICommandListImmediate& RHICmdList)
 {
 	check( IsInRenderingThread() );
 	RHICmdList.BeginDrawingViewport(GetViewportRHI(), FTextureRHIRef());
-	UpdateRenderTargetSurfaceRHIToCurrentBackBuffer(RHICmdList);
+	UpdateRenderTargetSurfaceRHIToCurrentBackBuffer();
 }
 
 /**
@@ -1417,10 +1417,7 @@ void FViewport::SetViewportClient( FViewportClient* InViewportClient )
 
 void FViewport::InitDynamicRHI()
 {
-	// Capture the viewport's back buffer surface for use through the FRenderTarget interface.
-	FRHICommandListImmediate& RHICmdList = FRHICommandListExecutor::GetImmediateCommandList();
-
-	UpdateRenderTargetSurfaceRHIToCurrentBackBuffer(RHICmdList);
+	UpdateRenderTargetSurfaceRHIToCurrentBackBuffer();
 
 	if(bRequiresHitProxyStorage)
 	{
@@ -1447,14 +1444,13 @@ void FViewport::InitRHI()
 
 	if(!IsValidRef(ViewportRHI))
 	{
-		FRHICommandListImmediate& RHICmdList = FRHICommandListExecutor::GetImmediateCommandList();
-		ViewportRHI = RHICmdList.CreateViewport(
+		ViewportRHI = RHICreateViewport(
 			GetWindow(),
 			SizeX,
 			SizeY,
 			IsFullscreen()
 			);
-		UpdateRenderTargetSurfaceRHIToCurrentBackBuffer(RHICmdList);
+		UpdateRenderTargetSurfaceRHIToCurrentBackBuffer();
 	}
 }
 
@@ -1568,11 +1564,11 @@ bool FViewport::HasToggleFreezeCommand()
 /**
  * Update the render target surface RHI to the current back buffer 
  */
-void FViewport::UpdateRenderTargetSurfaceRHIToCurrentBackBuffer(FRHICommandListImmediate& RHICmdList)
+void FViewport::UpdateRenderTargetSurfaceRHIToCurrentBackBuffer()
 {
 	if(IsValidRef(ViewportRHI))
 	{
-		RenderTargetTextureRHI = RHICmdList.GetViewportBackBuffer(ViewportRHI);
+		RenderTargetTextureRHI = RHIGetViewportBackBuffer(ViewportRHI);
 	}
 }
 
diff --git a/Engine/Source/Runtime/Engine/Public/UnrealClient.h b/Engine/Source/Runtime/Engine/Public/UnrealClient.h
--- a/Engine/Source/Runtime/Engine/Public/UnrealClient.h
+++ b/Engine/Source/Runtime/Engine/Public/UnrealClient.h
@@ -437,7 +437,7 @@ public:
 	/**
 	 * Update the render target surface RHI to the current back buffer 
 	 */
-	void UpdateRenderTargetSurfaceRHIToCurrentBackBuffer(FRHICommandListImmediate& RHICmdList);
+	void UpdateRenderTargetSurfaceRHIToCurrentBackBuffer();
 
 	/**
 	 * First chance for viewports to render custom stats text
diff --git a/Engine/Source/Runtime/RHI/Public/RHIMethods.h b/Engine/Source/Runtime/RHI/Public/RHIMethods.h
--- a/Engine/Source/Runtime/RHI/Public/RHIMethods.h
+++ b/Engine/Source/Runtime/RHI/Public/RHIMethods.h
@@ -1229,7 +1229,7 @@ DEFINE_RHIMETHOD_0(
 	bool,IsDrawingViewport,
 	return,return false;
 	);
-DEFINE_RHIMETHOD_1(
+DEFINE_RHIMETHOD_GLOBALFLUSH_1(
 	FTexture2DRHIRef,GetViewportBackBuffer,
 	FViewportRHIParamRef,Viewport,
 	return,return new FRHITexture2D(1,1,1,1,PF_B8G8R8A8,TexCreate_RenderTargetable);
