Fixed so that trace functions can accept literal values

TTP 340351	UE4 - SphereOverlapActors won't use the default SpherePos as a literal value

diff --git a/Engine/Source/Runtime/Engine/Classes/Kismet/KismetSystemLibrary.h b/Engine/Source/Runtime/Engine/Classes/Kismet/KismetSystemLibrary.h
--- a/Engine/Source/Runtime/Engine/Classes/Kismet/KismetSystemLibrary.h
+++ b/Engine/Source/Runtime/Engine/Classes/Kismet/KismetSystemLibrary.h
@@ -388,10 +388,10 @@ class UKismetSystemLibrary : public UBlueprintFunctionLibrary
 	 * @return				true if there was an overlap that passed the filters, false otherwise.
 	 */
 	UFUNCTION(BlueprintCallable, Category="Collision", meta=(HidePin="WorldContextObject", DefaultToSelf="WorldContextObject", AutoCreateRefTerm="ActorsToIgnore", FriendlyName = "SphereOverlapActors"))
-	static ENGINE_API bool SphereOverlapActors_NEW(UObject* WorldContextObject, const FVector& SpherePos, float SphereRadius, const TArray<TEnumAsByte<EObjectTypeQuery> > & ObjectTypes, UClass* ActorClassFilter, const TArray<AActor*>& ActorsToIgnore, TArray<class AActor*>& OutActors);
+	static ENGINE_API bool SphereOverlapActors_NEW(UObject* WorldContextObject, const FVector SpherePos, float SphereRadius, const TArray<TEnumAsByte<EObjectTypeQuery> > & ObjectTypes, UClass* ActorClassFilter, const TArray<AActor*>& ActorsToIgnore, TArray<class AActor*>& OutActors);
 
 	UFUNCTION(BlueprintCallable, Category="Collision", meta=(DeprecatedFunction, DeprecationMessage = "Use new SphereOverlaActors", HidePin="WorldContextObject", DefaultToSelf="WorldContextObject", AutoCreateRefTerm="ActorsToIgnore"))
-	static ENGINE_API bool SphereOverlapActors_DEPRECATED(UObject* WorldContextObject, const FVector& SpherePos, float SphereRadius, EOverlapFilterOption Filter, UClass* ActorClassFilter, const TArray<AActor*>& ActorsToIgnore, TArray<class AActor*>& OutActors);
+	static ENGINE_API bool SphereOverlapActors_DEPRECATED(UObject* WorldContextObject, const FVector SpherePos, float SphereRadius, EOverlapFilterOption Filter, UClass* ActorClassFilter, const TArray<AActor*>& ActorsToIgnore, TArray<class AActor*>& OutActors);
 
 	/**
 	 * Returns an array of components that overlap the given sphere.
@@ -405,10 +405,10 @@ class UKismetSystemLibrary : public UBlueprintFunctionLibrary
 	 * @return				true if there was an overlap that passed the filters, false otherwise.
 	 */
 	UFUNCTION(BlueprintCallable, Category="Collision", meta=(HidePin="WorldContextObject", DefaultToSelf="WorldContextObject", AutoCreateRefTerm="ActorsToIgnore", FriendlyName="SphereOverlapComponents"))
-	static ENGINE_API bool SphereOverlapComponents_NEW(UObject* WorldContextObject, const FVector& SpherePos, float SphereRadius, const TArray<TEnumAsByte<EObjectTypeQuery> > & ObjectTypes, UClass* ComponentClassFilter, const TArray<AActor*>& ActorsToIgnore, TArray<class UPrimitiveComponent*>& OutComponents);
+	static ENGINE_API bool SphereOverlapComponents_NEW(UObject* WorldContextObject, const FVector SpherePos, float SphereRadius, const TArray<TEnumAsByte<EObjectTypeQuery> > & ObjectTypes, UClass* ComponentClassFilter, const TArray<AActor*>& ActorsToIgnore, TArray<class UPrimitiveComponent*>& OutComponents);
 
 	UFUNCTION(BlueprintCallable, Category="Collision", meta=(DeprecatedFunction, DeprecationMessage = "Use new SphereOverlapComponents", HidePin="WorldContextObject", DefaultToSelf="WorldContextObject", AutoCreateRefTerm="ActorsToIgnore"))
-	static ENGINE_API bool SphereOverlapComponents_DEPRECATED(UObject* WorldContextObject, const FVector& SpherePos, float SphereRadius, EOverlapFilterOption Filter, UClass* ComponentClassFilter, const TArray<AActor*>& ActorsToIgnore, TArray<class UPrimitiveComponent*>& OutComponents);
+	static ENGINE_API bool SphereOverlapComponents_DEPRECATED(UObject* WorldContextObject, const FVector SpherePos, float SphereRadius, EOverlapFilterOption Filter, UClass* ComponentClassFilter, const TArray<AActor*>& ActorsToIgnore, TArray<class UPrimitiveComponent*>& OutComponents);
 	
 
 	/**
@@ -423,10 +423,10 @@ class UKismetSystemLibrary : public UBlueprintFunctionLibrary
 	 * @return				true if there was an overlap that passed the filters, false otherwise.
 	 */
 	UFUNCTION(BlueprintCallable, Category="Collision", meta=(HidePin="WorldContextObject", DefaultToSelf="WorldContextObject", AutoCreateRefTerm="ActorsToIgnore", FriendlyName="BoxOverlapActors"))
-	static ENGINE_API bool BoxOverlapActors_NEW(UObject* WorldContextObject, const FVector& BoxPos, FVector BoxExtent, const TArray<TEnumAsByte<EObjectTypeQuery> > & ObjectTypes, UClass* ActorClassFilter, const TArray<AActor*>& ActorsToIgnore, TArray<class AActor*>& OutActors);
+	static ENGINE_API bool BoxOverlapActors_NEW(UObject* WorldContextObject, const FVector BoxPos, FVector BoxExtent, const TArray<TEnumAsByte<EObjectTypeQuery> > & ObjectTypes, UClass* ActorClassFilter, const TArray<AActor*>& ActorsToIgnore, TArray<class AActor*>& OutActors);
 
 	UFUNCTION(BlueprintCallable, Category="Collision", meta=(DeprecatedFunction, DeprecationMessage = "Use new BoxOverlapActors", HidePin="WorldContextObject", DefaultToSelf="WorldContextObject", AutoCreateRefTerm="ActorsToIgnore"))
-	static ENGINE_API bool BoxOverlapActors_DEPRECATED(UObject* WorldContextObject, const FVector& BoxPos, FVector BoxExtent, EOverlapFilterOption Filter, UClass* ActorClassFilter, const TArray<AActor*>& ActorsToIgnore, TArray<class AActor*>& OutActors);
+	static ENGINE_API bool BoxOverlapActors_DEPRECATED(UObject* WorldContextObject, const FVector BoxPos, FVector BoxExtent, EOverlapFilterOption Filter, UClass* ActorClassFilter, const TArray<AActor*>& ActorsToIgnore, TArray<class AActor*>& OutActors);
 
 	/**
 	 * Returns an array of components that overlap the given axis-aligned box.
@@ -440,10 +440,10 @@ class UKismetSystemLibrary : public UBlueprintFunctionLibrary
 	 * @return				true if there was an overlap that passed the filters, false otherwise.
 	 */
 	UFUNCTION(BlueprintCallable, Category="Collision", meta=(HidePin="WorldContextObject", DefaultToSelf="WorldContextObject", AutoCreateRefTerm="ActorsToIgnore", FriendlyName="BoxOverlapComponents"))
-	static ENGINE_API bool BoxOverlapComponents_NEW(UObject* WorldContextObject, const FVector& BoxPos, FVector Extent, const TArray<TEnumAsByte<EObjectTypeQuery> > & ObjectTypes, UClass* ComponentClassFilter, const TArray<AActor*>& ActorsToIgnore, TArray<class UPrimitiveComponent*>& OutComponents);
+	static ENGINE_API bool BoxOverlapComponents_NEW(UObject* WorldContextObject, const FVector BoxPos, FVector Extent, const TArray<TEnumAsByte<EObjectTypeQuery> > & ObjectTypes, UClass* ComponentClassFilter, const TArray<AActor*>& ActorsToIgnore, TArray<class UPrimitiveComponent*>& OutComponents);
 
 	UFUNCTION(BlueprintCallable, Category="Collision", meta=(DeprecatedFunction, DeprecationMessage = "Use new BoxOverlapComponents", HidePin="WorldContextObject", DefaultToSelf="WorldContextObject", AutoCreateRefTerm="ActorsToIgnore"))
-	static ENGINE_API bool BoxOverlapComponents_DEPRECATED(UObject* WorldContextObject, const FVector& BoxPos, FVector Extent, EOverlapFilterOption Filter, UClass* ComponentClassFilter, const TArray<AActor*>& ActorsToIgnore, TArray<class UPrimitiveComponent*>& OutComponents);
+	static ENGINE_API bool BoxOverlapComponents_DEPRECATED(UObject* WorldContextObject, const FVector BoxPos, FVector Extent, EOverlapFilterOption Filter, UClass* ComponentClassFilter, const TArray<AActor*>& ActorsToIgnore, TArray<class UPrimitiveComponent*>& OutComponents);
 
 
 	/**
@@ -459,10 +459,10 @@ class UKismetSystemLibrary : public UBlueprintFunctionLibrary
 	 * @return				true if there was an overlap that passed the filters, false otherwise.
 	 */
 	UFUNCTION(BlueprintCallable, Category="Collision", meta=(HidePin="WorldContextObject", DefaultToSelf="WorldContextObject", AutoCreateRefTerm="ActorsToIgnore", FriendlyName="CapsuleOverlapActors"))
-	static ENGINE_API bool CapsuleOverlapActors_NEW(UObject* WorldContextObject, const FVector& CapsulePos, float Radius, float HalfHeight, const TArray<TEnumAsByte<EObjectTypeQuery> > & ObjectTypes, UClass* ActorClassFilter, const TArray<AActor*>& ActorsToIgnore, TArray<class AActor*>& OutActors);
+	static ENGINE_API bool CapsuleOverlapActors_NEW(UObject* WorldContextObject, const FVector CapsulePos, float Radius, float HalfHeight, const TArray<TEnumAsByte<EObjectTypeQuery> > & ObjectTypes, UClass* ActorClassFilter, const TArray<AActor*>& ActorsToIgnore, TArray<class AActor*>& OutActors);
 
 	UFUNCTION(BlueprintCallable, Category="Collision", meta=(DeprecatedFunction, DeprecationMessage = "Use new CapsuleOverlapActors", HidePin="WorldContextObject", DefaultToSelf="WorldContextObject", AutoCreateRefTerm="ActorsToIgnore"))
-	static ENGINE_API bool CapsuleOverlapActors_DEPRECATED(UObject* WorldContextObject, const FVector& CapsulePos, float Radius, float HalfHeight, EOverlapFilterOption Filter, UClass* ActorClassFilter, const TArray<AActor*>& ActorsToIgnore, TArray<class AActor*>& OutActors);
+	static ENGINE_API bool CapsuleOverlapActors_DEPRECATED(UObject* WorldContextObject, const FVector CapsulePos, float Radius, float HalfHeight, EOverlapFilterOption Filter, UClass* ActorClassFilter, const TArray<AActor*>& ActorsToIgnore, TArray<class AActor*>& OutActors);
 	/**
 	 * Returns an array of components that overlap the given capsule.
 	 * @param WorldContext	World context
@@ -476,10 +476,10 @@ class UKismetSystemLibrary : public UBlueprintFunctionLibrary
 	 * @return				true if there was an overlap that passed the filters, false otherwise.
 	 */
 	UFUNCTION(BlueprintCallable, Category="Collision", meta=(HidePin="WorldContextObject", DefaultToSelf="WorldContextObject", AutoCreateRefTerm="ActorsToIgnore", FriendlyName="CapsuleOverlapComponents") )
-	static ENGINE_API bool CapsuleOverlapComponents_NEW(UObject* WorldContextObject, const FVector& CapsulePos, float Radius, float HalfHeight, const TArray<TEnumAsByte<EObjectTypeQuery> > & ObjectTypes, UClass* ComponentClassFilter, const TArray<AActor*>& ActorsToIgnore, TArray<class UPrimitiveComponent*>& OutComponents);
+	static ENGINE_API bool CapsuleOverlapComponents_NEW(UObject* WorldContextObject, const FVector CapsulePos, float Radius, float HalfHeight, const TArray<TEnumAsByte<EObjectTypeQuery> > & ObjectTypes, UClass* ComponentClassFilter, const TArray<AActor*>& ActorsToIgnore, TArray<class UPrimitiveComponent*>& OutComponents);
 
 	UFUNCTION(BlueprintCallable, Category="Collision", meta=(DeprecatedFunction, DeprecationMessage = "Use new CapsuleOverlapComponents", HidePin="WorldContextObject", DefaultToSelf="WorldContextObject", AutoCreateRefTerm="ActorsToIgnore") )
-	static ENGINE_API bool CapsuleOverlapComponents_DEPRECATED(UObject* WorldContextObject, const FVector& CapsulePos, float Radius, float HalfHeight, EOverlapFilterOption Filter, UClass* ComponentClassFilter, const TArray<AActor*>& ActorsToIgnore, TArray<class UPrimitiveComponent*>& OutComponents);
+	static ENGINE_API bool CapsuleOverlapComponents_DEPRECATED(UObject* WorldContextObject, const FVector CapsulePos, float Radius, float HalfHeight, EOverlapFilterOption Filter, UClass* ComponentClassFilter, const TArray<AActor*>& ActorsToIgnore, TArray<class UPrimitiveComponent*>& OutComponents);
 	
 
 	/**
@@ -493,10 +493,10 @@ class UKismetSystemLibrary : public UBlueprintFunctionLibrary
 	 * @return						true if there was an overlap that passed the filters, false otherwise.
 	 */
 	UFUNCTION(BlueprintCallable, Category="Collision", meta=(AutoCreateRefTerm="ActorsToIgnore", FriendlyName="ComponentOverlapActors"))
-	static ENGINE_API bool ComponentOverlapActors_NEW(UPrimitiveComponent* Component, const FTransform& ComponentTransform, const TArray<TEnumAsByte<EObjectTypeQuery> > & ObjectTypes, UClass* ActorClassFilter, const TArray<AActor*>& ActorsToIgnore, TArray<class AActor*>& OutActors);
+	static ENGINE_API bool ComponentOverlapActors_NEW(UPrimitiveComponent* Component, const FTransform ComponentTransform, const TArray<TEnumAsByte<EObjectTypeQuery> > & ObjectTypes, UClass* ActorClassFilter, const TArray<AActor*>& ActorsToIgnore, TArray<class AActor*>& OutActors);
 
 	UFUNCTION(BlueprintCallable, Category="Collision", meta=(DeprecatedFunction, DeprecationMessage = "Use new ComponentOverlapActors", AutoCreateRefTerm="ActorsToIgnore"))
-	static ENGINE_API bool ComponentOverlapActors_DEPRECATED(UPrimitiveComponent* Component, const FTransform& ComponentTransform, EOverlapFilterOption Filter, UClass* ActorClassFilter, const TArray<AActor*>& ActorsToIgnore, TArray<class AActor*>& OutActors);
+	static ENGINE_API bool ComponentOverlapActors_DEPRECATED(UPrimitiveComponent* Component, const FTransform ComponentTransform, EOverlapFilterOption Filter, UClass* ActorClassFilter, const TArray<AActor*>& ActorsToIgnore, TArray<class AActor*>& OutActors);
 	
 	/**
 	 * Returns an array of components that overlap the given component.
@@ -509,10 +509,10 @@ class UKismetSystemLibrary : public UBlueprintFunctionLibrary
 	 * @return						true if there was an overlap that passed the filters, false otherwise.
 	 */
 	UFUNCTION(BlueprintCallable, Category="Collision", meta=(AutoCreateRefTerm="ActorsToIgnore", FriendlyName="ComponentOverlapComponents"))
-	static ENGINE_API bool ComponentOverlapComponents_NEW(UPrimitiveComponent* Component, const FTransform& ComponentTransform, const TArray<TEnumAsByte<EObjectTypeQuery> > & ObjectTypes, UClass* ComponentClassFilter, const TArray<AActor*>& ActorsToIgnore, TArray<class UPrimitiveComponent*>& OutComponents);
+	static ENGINE_API bool ComponentOverlapComponents_NEW(UPrimitiveComponent* Component, const FTransform ComponentTransform, const TArray<TEnumAsByte<EObjectTypeQuery> > & ObjectTypes, UClass* ComponentClassFilter, const TArray<AActor*>& ActorsToIgnore, TArray<class UPrimitiveComponent*>& OutComponents);
 
 	UFUNCTION(BlueprintCallable, Category="Collision", meta=(DeprecatedFunction, DeprecationMessage = "Use new ComponentOverlapComponents", AutoCreateRefTerm="ActorsToIgnore"))
-	static ENGINE_API bool ComponentOverlapComponents_DEPRECATED(UPrimitiveComponent* Component, const FTransform& ComponentTransform, EOverlapFilterOption Filter, UClass* ComponentClassFilter, const TArray<AActor*>& ActorsToIgnore, TArray<class UPrimitiveComponent*>& OutComponents);
+	static ENGINE_API bool ComponentOverlapComponents_DEPRECATED(UPrimitiveComponent* Component, const FTransform ComponentTransform, EOverlapFilterOption Filter, UClass* ComponentClassFilter, const TArray<AActor*>& ActorsToIgnore, TArray<class UPrimitiveComponent*>& OutComponents);
 
 	/**
 	 * Does a collision trace along the given line and returns the first blocking hit encountered.
@@ -527,7 +527,7 @@ class UKismetSystemLibrary : public UBlueprintFunctionLibrary
 	 * @return				True if there was a hit, false otherwise.
 	 */
 	UFUNCTION(BlueprintCallable, Category="Collision", meta=(bIgnoreSelf="true", HidePin="WorldContextObject", DefaultToSelf="WorldContextObject", AutoCreateRefTerm="ActorsToIgnore", FriendlyName = "SingleLineTraceByChannel"))
-	static ENGINE_API bool LineTraceSingle_NEW(UObject* WorldContextObject, const FVector& Start, const FVector& End, ETraceTypeQuery TraceChannel, bool bTraceComplex, const TArray<AActor*>& ActorsToIgnore, EDrawDebugTrace::Type DrawDebugType, FHitResult& OutHit, bool bIgnoreSelf);
+	static ENGINE_API bool LineTraceSingle_NEW(UObject* WorldContextObject, const FVector Start, const FVector End, ETraceTypeQuery TraceChannel, bool bTraceComplex, const TArray<AActor*>& ActorsToIgnore, EDrawDebugTrace::Type DrawDebugType, FHitResult& OutHit, bool bIgnoreSelf);
 	
 	/**
 	 * Does a collision trace along the given line and returns all hits encountered up to and including the first blocking hit.
@@ -542,7 +542,7 @@ class UKismetSystemLibrary : public UBlueprintFunctionLibrary
 	 * @return				True if there was a hit, false otherwise.
 	 */
 	UFUNCTION(BlueprintCallable, Category="Collision", meta=(bIgnoreSelf="true", HidePin="WorldContextObject", DefaultToSelf="WorldContextObject", AutoCreateRefTerm="ActorsToIgnore", FriendlyName = "MultiLineTraceByChannel"))
-	static ENGINE_API bool LineTraceMulti_NEW(UObject* WorldContextObject, const FVector& Start, const FVector& End, ETraceTypeQuery TraceChannel, bool bTraceComplex, const TArray<AActor*>& ActorsToIgnore, EDrawDebugTrace::Type DrawDebugType, TArray<FHitResult>& OutHits, bool bIgnoreSelf);
+	static ENGINE_API bool LineTraceMulti_NEW(UObject* WorldContextObject, const FVector Start, const FVector End, ETraceTypeQuery TraceChannel, bool bTraceComplex, const TArray<AActor*>& ActorsToIgnore, EDrawDebugTrace::Type DrawDebugType, TArray<FHitResult>& OutHits, bool bIgnoreSelf);
 
 	/**
 	 * Sweeps a sphere along the given line and returns the first blocking hit encountered.
@@ -557,7 +557,7 @@ class UKismetSystemLibrary : public UBlueprintFunctionLibrary
 	 * @return				True if there was a hit, false otherwise.
 	 */
 	UFUNCTION(BlueprintCallable, Category="Collision", meta=(bIgnoreSelf="true", HidePin="WorldContextObject", DefaultToSelf="WorldContextObject", AutoCreateRefTerm="ActorsToIgnore", FriendlyName = "SingleSphereTraceByChannel"))
-	static ENGINE_API bool SphereTraceSingle_NEW(UObject* WorldContextObject, const FVector& Start, const FVector& End, float Radius, ETraceTypeQuery TraceChannel, bool bTraceComplex, const TArray<AActor*>& ActorsToIgnore, EDrawDebugTrace::Type DrawDebugType, FHitResult& OutHit, bool bIgnoreSelf);
+	static ENGINE_API bool SphereTraceSingle_NEW(UObject* WorldContextObject, const FVector Start, const FVector End, float Radius, ETraceTypeQuery TraceChannel, bool bTraceComplex, const TArray<AActor*>& ActorsToIgnore, EDrawDebugTrace::Type DrawDebugType, FHitResult& OutHit, bool bIgnoreSelf);
 
 	/**
 	 * Sweeps a sphere along the given line and returns all hits encountered up to and including the first blocking hit.
@@ -573,7 +573,7 @@ class UKismetSystemLibrary : public UBlueprintFunctionLibrary
 	 * @return				True if there was a hit, false otherwise.
 	 */
 	UFUNCTION(BlueprintCallable, Category="Collision", meta=(bIgnoreSelf="true", HidePin="WorldContextObject", DefaultToSelf="WorldContextObject", AutoCreateRefTerm="ActorsToIgnore", FriendlyName = "MultiSphereTraceByChannel"))
-	static ENGINE_API bool SphereTraceMulti_NEW(UObject* WorldContextObject, const FVector& Start, const FVector& End, float Radius, ETraceTypeQuery TraceChannel, bool bTraceComplex, const TArray<AActor*>& ActorsToIgnore, EDrawDebugTrace::Type DrawDebugType, TArray<FHitResult>& OutHits, bool bIgnoreSelf);
+	static ENGINE_API bool SphereTraceMulti_NEW(UObject* WorldContextObject, const FVector Start, const FVector End, float Radius, ETraceTypeQuery TraceChannel, bool bTraceComplex, const TArray<AActor*>& ActorsToIgnore, EDrawDebugTrace::Type DrawDebugType, TArray<FHitResult>& OutHits, bool bIgnoreSelf);
 
 	/**
 	* Sweeps a box along the given line and returns the first blocking hit encountered.
@@ -589,7 +589,7 @@ class UKismetSystemLibrary : public UBlueprintFunctionLibrary
 	* @return				True if there was a hit, false otherwise.
 	*/
 	UFUNCTION(BlueprintCallable, Category = "Collision", meta = (bIgnoreSelf = "true", HidePin = "WorldContextObject", DefaultToSelf = "WorldContextObject", AutoCreateRefTerm = "ActorsToIgnore", FriendlyName = "SingleBoxTraceByChannel"))
-	static ENGINE_API bool BoxTraceSingle(UObject* WorldContextObject, const FVector& Start, const FVector& End, const FVector & HalfSize, const FRotator & Orientation, ETraceTypeQuery TraceChannel, bool bTraceComplex, const TArray<AActor*>& ActorsToIgnore, EDrawDebugTrace::Type DrawDebugType, FHitResult& OutHit, bool bIgnoreSelf);
+	static ENGINE_API bool BoxTraceSingle(UObject* WorldContextObject, const FVector Start, const FVector End, const FVector & HalfSize, const FRotator & Orientation, ETraceTypeQuery TraceChannel, bool bTraceComplex, const TArray<AActor*>& ActorsToIgnore, EDrawDebugTrace::Type DrawDebugType, FHitResult& OutHit, bool bIgnoreSelf);
 
 	/**
 	* Sweeps a box along the given line and returns all hits encountered.
@@ -605,7 +605,7 @@ class UKismetSystemLibrary : public UBlueprintFunctionLibrary
 	* @return				True if there was a hit, false otherwise.
 	*/
 	UFUNCTION(BlueprintCallable, Category = "Collision", meta = (bIgnoreSelf = "true", HidePin = "WorldContextObject", DefaultToSelf = "WorldContextObject", AutoCreateRefTerm = "ActorsToIgnore", FriendlyName = "MultiBoxTraceByChannel"))
-	static ENGINE_API bool BoxTraceMulti(UObject* WorldContextObject, const FVector& Start, const FVector& End, const FVector & HalfSize, const FRotator & Orientation, ETraceTypeQuery TraceChannel, bool bTraceComplex, const TArray<AActor*>& ActorsToIgnore, EDrawDebugTrace::Type DrawDebugType, TArray<FHitResult>& OutHits, bool bIgnoreSelf);
+	static ENGINE_API bool BoxTraceMulti(UObject* WorldContextObject, const FVector Start, const FVector End, const FVector & HalfSize, const FRotator & Orientation, ETraceTypeQuery TraceChannel, bool bTraceComplex, const TArray<AActor*>& ActorsToIgnore, EDrawDebugTrace::Type DrawDebugType, TArray<FHitResult>& OutHits, bool bIgnoreSelf);
 
 	
 	/**
@@ -623,7 +623,7 @@ class UKismetSystemLibrary : public UBlueprintFunctionLibrary
 	 * @return				True if there was a hit, false otherwise.
 	 */
 	UFUNCTION(BlueprintCallable, Category="Collision", meta=(bIgnoreSelf="true", HidePin="WorldContextObject", DefaultToSelf="WorldContextObject", AutoCreateRefTerm="ActorsToIgnore", FriendlyName = "SingleCapsuleTraceByChannel"))
-	static ENGINE_API bool CapsuleTraceSingle_NEW(UObject* WorldContextObject, const FVector& Start, const FVector& End, float Radius, float HalfHeight, ETraceTypeQuery TraceChannel, bool bTraceComplex, const TArray<AActor*>& ActorsToIgnore, EDrawDebugTrace::Type DrawDebugType, FHitResult& OutHit, bool bIgnoreSelf);
+	static ENGINE_API bool CapsuleTraceSingle_NEW(UObject* WorldContextObject, const FVector Start, const FVector End, float Radius, float HalfHeight, ETraceTypeQuery TraceChannel, bool bTraceComplex, const TArray<AActor*>& ActorsToIgnore, EDrawDebugTrace::Type DrawDebugType, FHitResult& OutHit, bool bIgnoreSelf);
 
 	/**
 	 * Sweeps a capsule along the given line and returns all hits encountered up to and including the first blocking hit.
@@ -640,7 +640,7 @@ class UKismetSystemLibrary : public UBlueprintFunctionLibrary
 	 * @return				True if there was a hit, false otherwise.
 	 */
 	UFUNCTION(BlueprintCallable, Category="Collision", meta=(bIgnoreSelf="true", HidePin="WorldContextObject", DefaultToSelf="WorldContextObject", AutoCreateRefTerm="ActorsToIgnore", FriendlyName = "MultiCapsuleTraceByChannel"))
-	static ENGINE_API bool CapsuleTraceMulti_NEW(UObject* WorldContextObject, const FVector& Start, const FVector& End, float Radius, float HalfHeight, ETraceTypeQuery TraceChannel, bool bTraceComplex, const TArray<AActor*>& ActorsToIgnore, EDrawDebugTrace::Type DrawDebugType, TArray<FHitResult>& OutHits, bool bIgnoreSelf);
+	static ENGINE_API bool CapsuleTraceMulti_NEW(UObject* WorldContextObject, const FVector Start, const FVector End, float Radius, float HalfHeight, ETraceTypeQuery TraceChannel, bool bTraceComplex, const TArray<AActor*>& ActorsToIgnore, EDrawDebugTrace::Type DrawDebugType, TArray<FHitResult>& OutHits, bool bIgnoreSelf);
 
 	/**
 	 * Does a collision trace along the given line and returns the first hit encountered.
@@ -655,7 +655,7 @@ class UKismetSystemLibrary : public UBlueprintFunctionLibrary
 	 * @return				True if there was a hit, false otherwise.
 	 */
 	UFUNCTION(BlueprintCallable, Category="Collision", meta=(bIgnoreSelf="true", HidePin="WorldContextObject", DefaultToSelf="WorldContextObject", AutoCreateRefTerm="ActorsToIgnore", FriendlyName = "SingleLineTraceForObjects"))
-	static ENGINE_API bool LineTraceSingleForObjects(UObject* WorldContextObject, const FVector& Start, const FVector& End, const TArray<TEnumAsByte<EObjectTypeQuery> > & ObjectTypes, bool bTraceComplex, const TArray<AActor*>& ActorsToIgnore, EDrawDebugTrace::Type DrawDebugType, FHitResult& OutHit, bool bIgnoreSelf);
+	static ENGINE_API bool LineTraceSingleForObjects(UObject* WorldContextObject, const FVector Start, const FVector End, const TArray<TEnumAsByte<EObjectTypeQuery> > & ObjectTypes, bool bTraceComplex, const TArray<AActor*>& ActorsToIgnore, EDrawDebugTrace::Type DrawDebugType, FHitResult& OutHit, bool bIgnoreSelf);
 	
 	/**
 	 * Does a collision trace along the given line and returns all hits encountered.
@@ -670,7 +670,7 @@ class UKismetSystemLibrary : public UBlueprintFunctionLibrary
 	 * @return				True if there was a hit, false otherwise.
 	 */
 	UFUNCTION(BlueprintCallable, Category="Collision", meta=(bIgnoreSelf="true", HidePin="WorldContextObject", DefaultToSelf="WorldContextObject", AutoCreateRefTerm="ActorsToIgnore", FriendlyName = "MultiLineTraceForObjects"))
-	static ENGINE_API bool LineTraceMultiForObjects(UObject* WorldContextObject, const FVector& Start, const FVector& End, const TArray<TEnumAsByte<EObjectTypeQuery> > & ObjectTypes, bool bTraceComplex, const TArray<AActor*>& ActorsToIgnore, EDrawDebugTrace::Type DrawDebugType, TArray<FHitResult>& OutHits, bool bIgnoreSelf);
+	static ENGINE_API bool LineTraceMultiForObjects(UObject* WorldContextObject, const FVector Start, const FVector End, const TArray<TEnumAsByte<EObjectTypeQuery> > & ObjectTypes, bool bTraceComplex, const TArray<AActor*>& ActorsToIgnore, EDrawDebugTrace::Type DrawDebugType, TArray<FHitResult>& OutHits, bool bIgnoreSelf);
 
 	/**
 	 * Sweeps a sphere along the given line and returns the first hit encountered.
@@ -685,7 +685,7 @@ class UKismetSystemLibrary : public UBlueprintFunctionLibrary
 	 * @return				True if there was a hit, false otherwise.
 	 */
 	UFUNCTION(BlueprintCallable, Category="Collision", meta=(bIgnoreSelf="true", HidePin="WorldContextObject", DefaultToSelf="WorldContextObject", AutoCreateRefTerm="ActorsToIgnore", FriendlyName = "SingleSphereTraceForObjects"))
-	static ENGINE_API bool SphereTraceSingleForObjects(UObject* WorldContextObject, const FVector& Start, const FVector& End, float Radius, const TArray<TEnumAsByte<EObjectTypeQuery> > & ObjectTypes, bool bTraceComplex, const TArray<AActor*>& ActorsToIgnore, EDrawDebugTrace::Type DrawDebugType, FHitResult& OutHit, bool bIgnoreSelf);
+	static ENGINE_API bool SphereTraceSingleForObjects(UObject* WorldContextObject, const FVector Start, const FVector End, float Radius, const TArray<TEnumAsByte<EObjectTypeQuery> > & ObjectTypes, bool bTraceComplex, const TArray<AActor*>& ActorsToIgnore, EDrawDebugTrace::Type DrawDebugType, FHitResult& OutHit, bool bIgnoreSelf);
 
 	/**
 	 * Sweeps a sphere along the given line and returns all hits encountered.
@@ -701,7 +701,7 @@ class UKismetSystemLibrary : public UBlueprintFunctionLibrary
 	 * @return				True if there was a hit, false otherwise.
 	 */
 	UFUNCTION(BlueprintCallable, Category="Collision", meta=(bIgnoreSelf="true", HidePin="WorldContextObject", DefaultToSelf="WorldContextObject", AutoCreateRefTerm="ActorsToIgnore", FriendlyName = "MultiSphereTraceForObjects"))
-	static ENGINE_API bool SphereTraceMultiForObjects(UObject* WorldContextObject, const FVector& Start, const FVector& End, float Radius, const TArray<TEnumAsByte<EObjectTypeQuery> > & ObjectTypes, bool bTraceComplex, const TArray<AActor*>& ActorsToIgnore, EDrawDebugTrace::Type DrawDebugType, TArray<FHitResult>& OutHits, bool bIgnoreSelf);
+	static ENGINE_API bool SphereTraceMultiForObjects(UObject* WorldContextObject, const FVector Start, const FVector End, float Radius, const TArray<TEnumAsByte<EObjectTypeQuery> > & ObjectTypes, bool bTraceComplex, const TArray<AActor*>& ActorsToIgnore, EDrawDebugTrace::Type DrawDebugType, TArray<FHitResult>& OutHits, bool bIgnoreSelf);
 	
 
 	/**
@@ -718,7 +718,7 @@ class UKismetSystemLibrary : public UBlueprintFunctionLibrary
 	* @return				True if there was a hit, false otherwise.
 	*/
 	UFUNCTION(BlueprintCallable, Category = "Collision", meta = (bIgnoreSelf = "true", HidePin = "WorldContextObject", DefaultToSelf = "WorldContextObject", AutoCreateRefTerm = "ActorsToIgnore", FriendlyName = "SingleBoxTraceForObjects"))
-	static ENGINE_API bool BoxTraceSingleForObjects(UObject* WorldContextObject, const FVector& Start, const FVector& End, const FVector & HalfSize, const FRotator & Orientation, const TArray<TEnumAsByte<EObjectTypeQuery> > & ObjectTypes, bool bTraceComplex, const TArray<AActor*>& ActorsToIgnore, EDrawDebugTrace::Type DrawDebugType, FHitResult& OutHit, bool bIgnoreSelf);
+	static ENGINE_API bool BoxTraceSingleForObjects(UObject* WorldContextObject, const FVector Start, const FVector End, const FVector HalfSize, const FRotator Orientation, const TArray<TEnumAsByte<EObjectTypeQuery> > & ObjectTypes, bool bTraceComplex, const TArray<AActor*>& ActorsToIgnore, EDrawDebugTrace::Type DrawDebugType, FHitResult& OutHit, bool bIgnoreSelf);
 
 
 	/**
@@ -735,7 +735,7 @@ class UKismetSystemLibrary : public UBlueprintFunctionLibrary
 	* @return				True if there was a hit, false otherwise.
 	*/
 	UFUNCTION(BlueprintCallable, Category = "Collision", meta = (bIgnoreSelf = "true", HidePin = "WorldContextObject", DefaultToSelf = "WorldContextObject", AutoCreateRefTerm = "ActorsToIgnore", FriendlyName = "MultiBoxTraceForObjects"))
-	static ENGINE_API bool BoxTraceMultiForObjects(UObject* WorldContextObject, const FVector& Start, const FVector& End, const FVector & HalfSize, const FRotator & Orientation, const TArray<TEnumAsByte<EObjectTypeQuery> > & ObjectTypes, bool bTraceComplex, const TArray<AActor*>& ActorsToIgnore, EDrawDebugTrace::Type DrawDebugType, TArray<FHitResult>& OutHits, bool bIgnoreSelf);
+	static ENGINE_API bool BoxTraceMultiForObjects(UObject* WorldContextObject, const FVector Start, const FVector End, const FVector HalfSize, const FRotator Orientation, const TArray<TEnumAsByte<EObjectTypeQuery> > & ObjectTypes, bool bTraceComplex, const TArray<AActor*>& ActorsToIgnore, EDrawDebugTrace::Type DrawDebugType, TArray<FHitResult>& OutHits, bool bIgnoreSelf);
 
 	/**
 	 * Sweeps a capsule along the given line and returns the first hit encountered.
@@ -752,7 +752,7 @@ class UKismetSystemLibrary : public UBlueprintFunctionLibrary
 	 * @return				True if there was a hit, false otherwise.
 	 */
 	UFUNCTION(BlueprintCallable, Category="Collision", meta=(bIgnoreSelf="true", HidePin="WorldContextObject", DefaultToSelf="WorldContextObject", AutoCreateRefTerm="ActorsToIgnore", FriendlyName = "SingleCapsuleTraceForObjects"))
-	static ENGINE_API bool CapsuleTraceSingleForObjects(UObject* WorldContextObject, const FVector& Start, const FVector& End, float Radius, float HalfHeight, const TArray<TEnumAsByte<EObjectTypeQuery> > & ObjectTypes, bool bTraceComplex, const TArray<AActor*>& ActorsToIgnore, EDrawDebugTrace::Type DrawDebugType, FHitResult& OutHit, bool bIgnoreSelf);
+	static ENGINE_API bool CapsuleTraceSingleForObjects(UObject* WorldContextObject, const FVector Start, const FVector End, float Radius, float HalfHeight, const TArray<TEnumAsByte<EObjectTypeQuery> > & ObjectTypes, bool bTraceComplex, const TArray<AActor*>& ActorsToIgnore, EDrawDebugTrace::Type DrawDebugType, FHitResult& OutHit, bool bIgnoreSelf);
 
 	/**
 	 * Sweeps a capsule along the given line and returns all hits encountered.
@@ -769,7 +769,7 @@ class UKismetSystemLibrary : public UBlueprintFunctionLibrary
 	 * @return				True if there was a hit, false otherwise.
 	 */
 	UFUNCTION(BlueprintCallable, Category="Collision", meta=(bIgnoreSelf="true", HidePin="WorldContextObject", DefaultToSelf="WorldContextObject", AutoCreateRefTerm="ActorsToIgnore", FriendlyName = "MultiCapsuleTraceForObjects"))
-	static ENGINE_API bool CapsuleTraceMultiForObjects(UObject* WorldContextObject, const FVector& Start, const FVector& End, float Radius, float HalfHeight, const TArray<TEnumAsByte<EObjectTypeQuery> > & ObjectTypes, bool bTraceComplex, const TArray<AActor*>& ActorsToIgnore, EDrawDebugTrace::Type DrawDebugType, TArray<FHitResult>& OutHits, bool bIgnoreSelf);
+	static ENGINE_API bool CapsuleTraceMultiForObjects(UObject* WorldContextObject, const FVector Start, const FVector End, float Radius, float HalfHeight, const TArray<TEnumAsByte<EObjectTypeQuery> > & ObjectTypes, bool bTraceComplex, const TArray<AActor*>& ActorsToIgnore, EDrawDebugTrace::Type DrawDebugType, TArray<FHitResult>& OutHits, bool bIgnoreSelf);
 
 	/**
 	 * Does a collision trace along the given line and returns the first blocking hit encountered.
@@ -784,7 +784,7 @@ class UKismetSystemLibrary : public UBlueprintFunctionLibrary
 	 * @return				True if there was a hit, false otherwise.
 	 */
 	UFUNCTION(BlueprintCallable, Category="Collision", meta=(DeprecatedFunction, DeprecationMessage = "Use new SingleLineTraceByChannel", bIgnoreSelf="true", HidePin="WorldContextObject", DefaultToSelf="WorldContextObject", AutoCreateRefTerm="ActorsToIgnore", FriendlyName = "SingleLineTraceByChannelDeprecated"))
-	static ENGINE_API bool LineTraceSingle_DEPRECATED(UObject* WorldContextObject, const FVector& Start, const FVector& End, ECollisionChannel TraceChannel, bool bTraceComplex, const TArray<AActor*>& ActorsToIgnore, EDrawDebugTrace::Type DrawDebugType, FHitResult& OutHit, bool bIgnoreSelf);
+	static ENGINE_API bool LineTraceSingle_DEPRECATED(UObject* WorldContextObject, const FVector Start, const FVector End, ECollisionChannel TraceChannel, bool bTraceComplex, const TArray<AActor*>& ActorsToIgnore, EDrawDebugTrace::Type DrawDebugType, FHitResult& OutHit, bool bIgnoreSelf);
 	
 	/**
 	 * Does a collision trace along the given line and returns all hits encountered up to and including the first blocking hit.
@@ -799,7 +799,7 @@ class UKismetSystemLibrary : public UBlueprintFunctionLibrary
 	 * @return				True if there was a hit, false otherwise.
 	 */
 	UFUNCTION(BlueprintCallable, Category="Collision", meta=(DeprecatedFunction, DeprecationMessage = "Use new MultiLineTraceByChannel", bIgnoreSelf="true", HidePin="WorldContextObject", DefaultToSelf="WorldContextObject", AutoCreateRefTerm="ActorsToIgnore", FriendlyName = "MultiLineTraceByChannelDeprecated"))
-	static ENGINE_API bool LineTraceMulti_DEPRECATED(UObject* WorldContextObject, const FVector& Start, const FVector& End, ECollisionChannel TraceChannel, bool bTraceComplex, const TArray<AActor*>& ActorsToIgnore, EDrawDebugTrace::Type DrawDebugType, TArray<FHitResult>& OutHits, bool bIgnoreSelf);
+	static ENGINE_API bool LineTraceMulti_DEPRECATED(UObject* WorldContextObject, const FVector Start, const FVector End, ECollisionChannel TraceChannel, bool bTraceComplex, const TArray<AActor*>& ActorsToIgnore, EDrawDebugTrace::Type DrawDebugType, TArray<FHitResult>& OutHits, bool bIgnoreSelf);
 
 	/**
 	 * Sweeps a sphere along the given line and returns the first blocking hit encountered.
@@ -814,7 +814,7 @@ class UKismetSystemLibrary : public UBlueprintFunctionLibrary
 	 * @return				True if there was a hit, false otherwise.
 	 */
 	UFUNCTION(BlueprintCallable, Category="Collision", meta=(DeprecatedFunction, DeprecationMessage = "Use new SingleSphereTraceByChannel", bIgnoreSelf="true", HidePin="WorldContextObject", DefaultToSelf="WorldContextObject", AutoCreateRefTerm="ActorsToIgnore", FriendlyName = "SingleSphereTraceByChannelDeprecated"))
-	static ENGINE_API bool SphereTraceSingle_DEPRECATED(UObject* WorldContextObject, const FVector& Start, const FVector& End, float Radius, ECollisionChannel TraceChannel, bool bTraceComplex, const TArray<AActor*>& ActorsToIgnore, EDrawDebugTrace::Type DrawDebugType, FHitResult& OutHit, bool bIgnoreSelf);
+	static ENGINE_API bool SphereTraceSingle_DEPRECATED(UObject* WorldContextObject, const FVector Start, const FVector End, float Radius, ECollisionChannel TraceChannel, bool bTraceComplex, const TArray<AActor*>& ActorsToIgnore, EDrawDebugTrace::Type DrawDebugType, FHitResult& OutHit, bool bIgnoreSelf);
 
 	/**
 	 * Sweeps a sphere along the given line and returns all hits encountered up to and including the first blocking hit.
@@ -830,7 +830,7 @@ class UKismetSystemLibrary : public UBlueprintFunctionLibrary
 	 * @return				True if there was a hit, false otherwise.
 	 */
 	UFUNCTION(BlueprintCallable, Category="Collision", meta=(DeprecatedFunction, DeprecationMessage = "Use new MultiSphereTraceByChannel", bIgnoreSelf="true", HidePin="WorldContextObject", DefaultToSelf="WorldContextObject", AutoCreateRefTerm="ActorsToIgnore", FriendlyName = "MultiSphereTraceByChannelDeprecated"))
-	static ENGINE_API bool SphereTraceMulti_DEPRECATED(UObject* WorldContextObject, const FVector& Start, const FVector& End, float Radius, ECollisionChannel TraceChannel, bool bTraceComplex, const TArray<AActor*>& ActorsToIgnore, EDrawDebugTrace::Type DrawDebugType, TArray<FHitResult>& OutHits, bool bIgnoreSelf);
+	static ENGINE_API bool SphereTraceMulti_DEPRECATED(UObject* WorldContextObject, const FVector Start, const FVector End, float Radius, ECollisionChannel TraceChannel, bool bTraceComplex, const TArray<AActor*>& ActorsToIgnore, EDrawDebugTrace::Type DrawDebugType, TArray<FHitResult>& OutHits, bool bIgnoreSelf);
 	
 	/**
 	 * Sweeps a capsule along the given line and returns the first blocking hit encountered.
@@ -847,7 +847,7 @@ class UKismetSystemLibrary : public UBlueprintFunctionLibrary
 	 * @return				True if there was a hit, false otherwise.
 	 */
 	UFUNCTION(BlueprintCallable, Category="Collision", meta=(DeprecatedFunction, DeprecationMessage = "Use new SingleCapsuleTraceByChannel", bIgnoreSelf="true", HidePin="WorldContextObject", DefaultToSelf="WorldContextObject", AutoCreateRefTerm="ActorsToIgnore", FriendlyName = "SingleCapsuleTraceByChannelDeprecated"))
-	static ENGINE_API bool CapsuleTraceSingle_DEPRECATED(UObject* WorldContextObject, const FVector& Start, const FVector& End, float Radius, float HalfHeight, ECollisionChannel TraceChannel, bool bTraceComplex, const TArray<AActor*>& ActorsToIgnore, EDrawDebugTrace::Type DrawDebugType, FHitResult& OutHit, bool bIgnoreSelf);
+	static ENGINE_API bool CapsuleTraceSingle_DEPRECATED(UObject* WorldContextObject, const FVector Start, const FVector End, float Radius, float HalfHeight, ECollisionChannel TraceChannel, bool bTraceComplex, const TArray<AActor*>& ActorsToIgnore, EDrawDebugTrace::Type DrawDebugType, FHitResult& OutHit, bool bIgnoreSelf);
 
 	/**
 	 * Sweeps a capsule along the given line and returns all hits encountered up to and including the first blocking hit.
@@ -864,7 +864,7 @@ class UKismetSystemLibrary : public UBlueprintFunctionLibrary
 	 * @return				True if there was a hit, false otherwise.
 	 */
 	UFUNCTION(BlueprintCallable, Category="Collision", meta=(DeprecatedFunction, DeprecationMessage = "Use new MultiCapsuleTraceByChannel", bIgnoreSelf="true", HidePin="WorldContextObject", DefaultToSelf="WorldContextObject", AutoCreateRefTerm="ActorsToIgnore", FriendlyName = "MultiCapsuleTraceByChannelDeprecated"))
-	static ENGINE_API bool CapsuleTraceMulti_DEPRECATED(UObject* WorldContextObject, const FVector& Start, const FVector& End, float Radius, float HalfHeight, ECollisionChannel TraceChannel, bool bTraceComplex, const TArray<AActor*>& ActorsToIgnore, EDrawDebugTrace::Type DrawDebugType, TArray<FHitResult>& OutHits, bool bIgnoreSelf);
+	static ENGINE_API bool CapsuleTraceMulti_DEPRECATED(UObject* WorldContextObject, const FVector Start, const FVector End, float Radius, float HalfHeight, ECollisionChannel TraceChannel, bool bTraceComplex, const TArray<AActor*>& ActorsToIgnore, EDrawDebugTrace::Type DrawDebugType, TArray<FHitResult>& OutHits, bool bIgnoreSelf);
 
 	/**
 	 * Does a collision trace along the given line and returns the first hit encountered.
@@ -879,7 +879,7 @@ class UKismetSystemLibrary : public UBlueprintFunctionLibrary
 	 * @return				True if there was a hit, false otherwise.
 	 */
 	UFUNCTION(BlueprintCallable, Category="Collision", meta=(DeprecatedFunction, DeprecationMessage = "Use new SingleLineTraceForObjects", bIgnoreSelf="true", HidePin="WorldContextObject", DefaultToSelf="WorldContextObject", AutoCreateRefTerm="ActorsToIgnore", FriendlyName = "SingleLineTraceByObjectDeprecated"))
-	static ENGINE_API bool LineTraceSingleByObject_DEPRECATED(UObject* WorldContextObject, const FVector& Start, const FVector& End, const TArray<TEnumAsByte<ECollisionChannel> > & ObjectsToTrace, bool bTraceComplex, const TArray<AActor*>& ActorsToIgnore, EDrawDebugTrace::Type DrawDebugType, FHitResult& OutHit, bool bIgnoreSelf);
+	static ENGINE_API bool LineTraceSingleByObject_DEPRECATED(UObject* WorldContextObject, const FVector Start, const FVector End, const TArray<TEnumAsByte<ECollisionChannel> > & ObjectsToTrace, bool bTraceComplex, const TArray<AActor*>& ActorsToIgnore, EDrawDebugTrace::Type DrawDebugType, FHitResult& OutHit, bool bIgnoreSelf);
 	
 	/**
 	 * Does a collision trace along the given line and returns all hits encountered.
@@ -894,7 +894,7 @@ class UKismetSystemLibrary : public UBlueprintFunctionLibrary
 	 * @return				True if there was a hit, false otherwise.
 	 */
 	UFUNCTION(BlueprintCallable, Category="Collision", meta=(DeprecatedFunction, DeprecationMessage = "Use new MultiLineTraceForObjects", bIgnoreSelf="true", HidePin="WorldContextObject", DefaultToSelf="WorldContextObject", AutoCreateRefTerm="ActorsToIgnore", FriendlyName = "MultiLineTraceByObjectDeprecated"))
-	static ENGINE_API bool LineTraceMultiByObject_DEPRECATED(UObject* WorldContextObject, const FVector& Start, const FVector& End, const TArray<TEnumAsByte<ECollisionChannel> > & ObjectsToTrace, bool bTraceComplex, const TArray<AActor*>& ActorsToIgnore, EDrawDebugTrace::Type DrawDebugType, TArray<FHitResult>& OutHits, bool bIgnoreSelf);
+	static ENGINE_API bool LineTraceMultiByObject_DEPRECATED(UObject* WorldContextObject, const FVector Start, const FVector End, const TArray<TEnumAsByte<ECollisionChannel> > & ObjectsToTrace, bool bTraceComplex, const TArray<AActor*>& ActorsToIgnore, EDrawDebugTrace::Type DrawDebugType, TArray<FHitResult>& OutHits, bool bIgnoreSelf);
 
 	/**
 	 * Sweeps a sphere along the given line and returns the first hit encountered.
@@ -909,7 +909,7 @@ class UKismetSystemLibrary : public UBlueprintFunctionLibrary
 	 * @return				True if there was a hit, false otherwise.
 	 */
 	UFUNCTION(BlueprintCallable, Category="Collision", meta=(DeprecatedFunction, DeprecationMessage = "Use new SingleSphereTraceForObjects", bIgnoreSelf="true", HidePin="WorldContextObject", DefaultToSelf="WorldContextObject", AutoCreateRefTerm="ActorsToIgnore", FriendlyName = "SingleSphereTraceByObjectDeprecated"))
-	static ENGINE_API bool SphereTraceSingleByObject_DEPRECATED(UObject* WorldContextObject, const FVector& Start, const FVector& End, float Radius, const TArray<TEnumAsByte<ECollisionChannel> > & ObjectsToTrace, bool bTraceComplex, const TArray<AActor*>& ActorsToIgnore, EDrawDebugTrace::Type DrawDebugType, FHitResult& OutHit, bool bIgnoreSelf);
+	static ENGINE_API bool SphereTraceSingleByObject_DEPRECATED(UObject* WorldContextObject, const FVector Start, const FVector End, float Radius, const TArray<TEnumAsByte<ECollisionChannel> > & ObjectsToTrace, bool bTraceComplex, const TArray<AActor*>& ActorsToIgnore, EDrawDebugTrace::Type DrawDebugType, FHitResult& OutHit, bool bIgnoreSelf);
 
 	/**
 	 * Sweeps a sphere along the given line and returns all hits encountered.
@@ -925,7 +925,7 @@ class UKismetSystemLibrary : public UBlueprintFunctionLibrary
 	 * @return				True if there was a hit, false otherwise.
 	 */
 	UFUNCTION(BlueprintCallable, Category="Collision", meta=(DeprecatedFunction, DeprecationMessage = "Use new MultiSphereTraceForObjects", bIgnoreSelf="true", HidePin="WorldContextObject", DefaultToSelf="WorldContextObject", AutoCreateRefTerm="ActorsToIgnore", FriendlyName = "MultiSphereTraceByObjectDeprecated"))
-	static ENGINE_API bool SphereTraceMultiByObject_DEPRECATED(UObject* WorldContextObject, const FVector& Start, const FVector& End, float Radius, const TArray<TEnumAsByte<ECollisionChannel> > & ObjectsToTrace, bool bTraceComplex, const TArray<AActor*>& ActorsToIgnore, EDrawDebugTrace::Type DrawDebugType, TArray<FHitResult>& OutHits, bool bIgnoreSelf);
+	static ENGINE_API bool SphereTraceMultiByObject_DEPRECATED(UObject* WorldContextObject, const FVector Start, const FVector End, float Radius, const TArray<TEnumAsByte<ECollisionChannel> > & ObjectsToTrace, bool bTraceComplex, const TArray<AActor*>& ActorsToIgnore, EDrawDebugTrace::Type DrawDebugType, TArray<FHitResult>& OutHits, bool bIgnoreSelf);
 	
 	/**
 	 * Sweeps a capsule along the given line and returns the first hit encountered.
@@ -942,7 +942,7 @@ class UKismetSystemLibrary : public UBlueprintFunctionLibrary
 	 * @return				True if there was a hit, false otherwise.
 	 */
 	UFUNCTION(BlueprintCallable, Category="Collision", meta=(DeprecatedFunction, DeprecationMessage = "Use new SingleCapsuleTraceForObjects", bIgnoreSelf="true", HidePin="WorldContextObject", DefaultToSelf="WorldContextObject", AutoCreateRefTerm="ActorsToIgnore", FriendlyName = "SingleCapsuleTraceByObjectDeprecated"))
-	static ENGINE_API bool CapsuleTraceSingleByObject_DEPRECATED(UObject* WorldContextObject, const FVector& Start, const FVector& End, float Radius, float HalfHeight, const TArray<TEnumAsByte<ECollisionChannel> > & ObjectsToTrace, bool bTraceComplex, const TArray<AActor*>& ActorsToIgnore, EDrawDebugTrace::Type DrawDebugType, FHitResult& OutHit, bool bIgnoreSelf);
+	static ENGINE_API bool CapsuleTraceSingleByObject_DEPRECATED(UObject* WorldContextObject, const FVector Start, const FVector End, float Radius, float HalfHeight, const TArray<TEnumAsByte<ECollisionChannel> > & ObjectsToTrace, bool bTraceComplex, const TArray<AActor*>& ActorsToIgnore, EDrawDebugTrace::Type DrawDebugType, FHitResult& OutHit, bool bIgnoreSelf);
 
 	/**
 	 * Sweeps a capsule along the given line and returns all hits encountered.
@@ -959,7 +959,7 @@ class UKismetSystemLibrary : public UBlueprintFunctionLibrary
 	 * @return				True if there was a hit, false otherwise.
 	 */
 	UFUNCTION(BlueprintCallable, Category="Collision", meta=(DeprecatedFunction, DeprecationMessage = "Use new MultiCapsuleTraceForObjects", bIgnoreSelf="true", HidePin="WorldContextObject", DefaultToSelf="WorldContextObject", AutoCreateRefTerm="ActorsToIgnore", FriendlyName = "MultiCapsuleTraceByObjectDeprecated"))
-	static ENGINE_API bool CapsuleTraceMultiByObject_DEPRECATED(UObject* WorldContextObject, const FVector& Start, const FVector& End, float Radius, float HalfHeight, const TArray<TEnumAsByte<ECollisionChannel> > & ObjectsToTrace, bool bTraceComplex, const TArray<AActor*>& ActorsToIgnore, EDrawDebugTrace::Type DrawDebugType, TArray<FHitResult>& OutHits, bool bIgnoreSelf);
+	static ENGINE_API bool CapsuleTraceMultiByObject_DEPRECATED(UObject* WorldContextObject, const FVector Start, const FVector End, float Radius, float HalfHeight, const TArray<TEnumAsByte<ECollisionChannel> > & ObjectsToTrace, bool bTraceComplex, const TArray<AActor*>& ActorsToIgnore, EDrawDebugTrace::Type DrawDebugType, TArray<FHitResult>& OutHits, bool bIgnoreSelf);
 
 	/**
 	 * Returns an array of unique actors represented by the given list of components.
diff --git a/Engine/Source/Runtime/Engine/Private/KismetSystemLibrary.cpp b/Engine/Source/Runtime/Engine/Private/KismetSystemLibrary.cpp
--- a/Engine/Source/Runtime/Engine/Private/KismetSystemLibrary.cpp
+++ b/Engine/Source/Runtime/Engine/Private/KismetSystemLibrary.cpp
@@ -569,7 +569,7 @@ void UKismetSystemLibrary::GetActorListFromComponentList(const TArray<UPrimitive
 }
 
 
-bool UKismetSystemLibrary::SphereOverlapActors_DEPRECATED(UObject* WorldContextObject, const FVector& SpherePos, float SphereRadius, EOverlapFilterOption Filter, UClass* ActorClassFilter, const TArray<AActor*>& ActorsToIgnore, TArray<AActor*>& OutActors)
+bool UKismetSystemLibrary::SphereOverlapActors_DEPRECATED(UObject* WorldContextObject, const FVector SpherePos, float SphereRadius, EOverlapFilterOption Filter, UClass* ActorClassFilter, const TArray<AActor*>& ActorsToIgnore, TArray<AActor*>& OutActors)
 {
 	OutActors.Empty();
 
@@ -583,7 +583,7 @@ bool UKismetSystemLibrary::SphereOverlapActors_DEPRECATED(UObject* WorldContextO
 	return (OutActors.Num() > 0);
 }
 
-bool UKismetSystemLibrary::SphereOverlapComponents_DEPRECATED(UObject* WorldContextObject, const FVector& SpherePos, float SphereRadius, EOverlapFilterOption Filter, UClass* ComponentClassFilter, const TArray<AActor*>& ActorsToIgnore, TArray<UPrimitiveComponent*>& OutComponents)
+bool UKismetSystemLibrary::SphereOverlapComponents_DEPRECATED(UObject* WorldContextObject, const FVector SpherePos, float SphereRadius, EOverlapFilterOption Filter, UClass* ComponentClassFilter, const TArray<AActor*>& ActorsToIgnore, TArray<UPrimitiveComponent*>& OutComponents)
 {
 	OutComponents.Empty();
 
@@ -614,7 +614,7 @@ bool UKismetSystemLibrary::SphereOverlapComponents_DEPRECATED(UObject* WorldCont
 
 
 
-bool UKismetSystemLibrary::BoxOverlapActors_DEPRECATED(UObject* WorldContextObject, const FVector& BoxPos, FVector BoxExtent, EOverlapFilterOption Filter, UClass* ActorClassFilter, const TArray<AActor*>& ActorsToIgnore, TArray<AActor*>& OutActors)
+bool UKismetSystemLibrary::BoxOverlapActors_DEPRECATED(UObject* WorldContextObject, const FVector BoxPos, FVector BoxExtent, EOverlapFilterOption Filter, UClass* ActorClassFilter, const TArray<AActor*>& ActorsToIgnore, TArray<AActor*>& OutActors)
 {
 	OutActors.Empty();
 
@@ -628,7 +628,7 @@ bool UKismetSystemLibrary::BoxOverlapActors_DEPRECATED(UObject* WorldContextObje
 	return (OutActors.Num() > 0);
 }
 
-bool UKismetSystemLibrary::BoxOverlapComponents_DEPRECATED(UObject* WorldContextObject, const FVector& BoxPos, FVector BoxExtent, EOverlapFilterOption Filter, UClass* ComponentClassFilter, const TArray<AActor*>& ActorsToIgnore, TArray<UPrimitiveComponent*>& OutComponents)
+bool UKismetSystemLibrary::BoxOverlapComponents_DEPRECATED(UObject* WorldContextObject, const FVector BoxPos, FVector BoxExtent, EOverlapFilterOption Filter, UClass* ComponentClassFilter, const TArray<AActor*>& ActorsToIgnore, TArray<UPrimitiveComponent*>& OutComponents)
 {
 	OutComponents.Empty();
 
@@ -659,7 +659,7 @@ bool UKismetSystemLibrary::BoxOverlapComponents_DEPRECATED(UObject* WorldContext
 }
 
 
-bool UKismetSystemLibrary::CapsuleOverlapActors_DEPRECATED(UObject* WorldContextObject, const FVector& CapsulePos, float Radius, float HalfHeight, EOverlapFilterOption Filter, UClass* ActorClassFilter, const TArray<AActor*>& ActorsToIgnore, TArray<AActor*>& OutActors)
+bool UKismetSystemLibrary::CapsuleOverlapActors_DEPRECATED(UObject* WorldContextObject, const FVector CapsulePos, float Radius, float HalfHeight, EOverlapFilterOption Filter, UClass* ActorClassFilter, const TArray<AActor*>& ActorsToIgnore, TArray<AActor*>& OutActors)
 {
 	OutActors.Empty();
 
@@ -673,7 +673,7 @@ bool UKismetSystemLibrary::CapsuleOverlapActors_DEPRECATED(UObject* WorldContext
 	return (OutActors.Num() > 0);
 }
 
-bool UKismetSystemLibrary::CapsuleOverlapComponents_DEPRECATED(UObject* WorldContextObject, const FVector& CapsulePos, float Radius, float HalfHeight, EOverlapFilterOption Filter, UClass* ComponentClassFilter, const TArray<AActor*>& ActorsToIgnore, TArray<UPrimitiveComponent*>& OutComponents)
+bool UKismetSystemLibrary::CapsuleOverlapComponents_DEPRECATED(UObject* WorldContextObject, const FVector CapsulePos, float Radius, float HalfHeight, EOverlapFilterOption Filter, UClass* ComponentClassFilter, const TArray<AActor*>& ActorsToIgnore, TArray<UPrimitiveComponent*>& OutComponents)
 {
 	OutComponents.Empty();
 
@@ -704,7 +704,7 @@ bool UKismetSystemLibrary::CapsuleOverlapComponents_DEPRECATED(UObject* WorldCon
 }
 
 
-bool UKismetSystemLibrary::ComponentOverlapActors_DEPRECATED(UPrimitiveComponent* Component, FTransform const& ComponentTransform, EOverlapFilterOption Filter, UClass* ActorClassFilter, const TArray<AActor*>& ActorsToIgnore, TArray<AActor*>& OutActors)
+bool UKismetSystemLibrary::ComponentOverlapActors_DEPRECATED(UPrimitiveComponent* Component, const FTransform ComponentTransform, EOverlapFilterOption Filter, UClass* ActorClassFilter, const TArray<AActor*>& ActorsToIgnore, TArray<AActor*>& OutActors)
 {
 	OutActors.Empty();
 
@@ -718,7 +718,7 @@ bool UKismetSystemLibrary::ComponentOverlapActors_DEPRECATED(UPrimitiveComponent
 	return (OutActors.Num() > 0);
 }
 
-bool UKismetSystemLibrary::ComponentOverlapComponents_DEPRECATED(UPrimitiveComponent* Component, FTransform const& ComponentTransform, EOverlapFilterOption Filter, UClass* ComponentClassFilter, const TArray<AActor*>& ActorsToIgnore, TArray<UPrimitiveComponent*>& OutComponents)
+bool UKismetSystemLibrary::ComponentOverlapComponents_DEPRECATED(UPrimitiveComponent* Component, const FTransform ComponentTransform, EOverlapFilterOption Filter, UClass* ComponentClassFilter, const TArray<AActor*>& ActorsToIgnore, TArray<UPrimitiveComponent*>& OutComponents)
 {
 	OutComponents.Empty();
 
@@ -748,7 +748,7 @@ bool UKismetSystemLibrary::ComponentOverlapComponents_DEPRECATED(UPrimitiveCompo
 }
 
 
-bool UKismetSystemLibrary::SphereOverlapActors_NEW(UObject* WorldContextObject, const FVector& SpherePos, float SphereRadius, const TArray<TEnumAsByte<EObjectTypeQuery> > & ObjectTypes, UClass* ActorClassFilter, const TArray<AActor*>& ActorsToIgnore, TArray<AActor*>& OutActors)
+bool UKismetSystemLibrary::SphereOverlapActors_NEW(UObject* WorldContextObject, const FVector SpherePos, float SphereRadius, const TArray<TEnumAsByte<EObjectTypeQuery> > & ObjectTypes, UClass* ActorClassFilter, const TArray<AActor*>& ActorsToIgnore, TArray<AActor*>& OutActors)
 {
 	OutActors.Empty();
 
@@ -762,7 +762,7 @@ bool UKismetSystemLibrary::SphereOverlapActors_NEW(UObject* WorldContextObject,
 	return (OutActors.Num() > 0);
 }
 
-bool UKismetSystemLibrary::SphereOverlapComponents_NEW(UObject* WorldContextObject, const FVector& SpherePos, float SphereRadius, const TArray<TEnumAsByte<EObjectTypeQuery> > & ObjectTypes, UClass* ComponentClassFilter, const TArray<AActor*>& ActorsToIgnore, TArray<UPrimitiveComponent*>& OutComponents)
+bool UKismetSystemLibrary::SphereOverlapComponents_NEW(UObject* WorldContextObject, const FVector SpherePos, float SphereRadius, const TArray<TEnumAsByte<EObjectTypeQuery> > & ObjectTypes, UClass* ComponentClassFilter, const TArray<AActor*>& ActorsToIgnore, TArray<UPrimitiveComponent*>& OutComponents)
 {
 	OutComponents.Empty();
 
@@ -799,7 +799,7 @@ bool UKismetSystemLibrary::SphereOverlapComponents_NEW(UObject* WorldContextObje
 
 
 
-bool UKismetSystemLibrary::BoxOverlapActors_NEW(UObject* WorldContextObject, const FVector& BoxPos, FVector BoxExtent, const TArray<TEnumAsByte<EObjectTypeQuery> > & ObjectTypes, UClass* ActorClassFilter, const TArray<AActor*>& ActorsToIgnore, TArray<AActor*>& OutActors)
+bool UKismetSystemLibrary::BoxOverlapActors_NEW(UObject* WorldContextObject, const FVector BoxPos, FVector BoxExtent, const TArray<TEnumAsByte<EObjectTypeQuery> > & ObjectTypes, UClass* ActorClassFilter, const TArray<AActor*>& ActorsToIgnore, TArray<AActor*>& OutActors)
 {
 	OutActors.Empty();
 
@@ -813,7 +813,7 @@ bool UKismetSystemLibrary::BoxOverlapActors_NEW(UObject* WorldContextObject, con
 	return (OutActors.Num() > 0);
 }
 
-bool UKismetSystemLibrary::BoxOverlapComponents_NEW(UObject* WorldContextObject, const FVector& BoxPos, FVector BoxExtent, const TArray<TEnumAsByte<EObjectTypeQuery> > & ObjectTypes, UClass* ComponentClassFilter, const TArray<AActor*>& ActorsToIgnore, TArray<UPrimitiveComponent*>& OutComponents)
+bool UKismetSystemLibrary::BoxOverlapComponents_NEW(UObject* WorldContextObject, const FVector BoxPos, FVector BoxExtent, const TArray<TEnumAsByte<EObjectTypeQuery> > & ObjectTypes, UClass* ComponentClassFilter, const TArray<AActor*>& ActorsToIgnore, TArray<UPrimitiveComponent*>& OutComponents)
 {
 	OutComponents.Empty();
 
@@ -849,7 +849,7 @@ bool UKismetSystemLibrary::BoxOverlapComponents_NEW(UObject* WorldContextObject,
 }
 
 
-bool UKismetSystemLibrary::CapsuleOverlapActors_NEW(UObject* WorldContextObject, const FVector& CapsulePos, float Radius, float HalfHeight, const TArray<TEnumAsByte<EObjectTypeQuery> > & ObjectTypes, UClass* ActorClassFilter, const TArray<AActor*>& ActorsToIgnore, TArray<AActor*>& OutActors)
+bool UKismetSystemLibrary::CapsuleOverlapActors_NEW(UObject* WorldContextObject, const FVector CapsulePos, float Radius, float HalfHeight, const TArray<TEnumAsByte<EObjectTypeQuery> > & ObjectTypes, UClass* ActorClassFilter, const TArray<AActor*>& ActorsToIgnore, TArray<AActor*>& OutActors)
 {
 	OutActors.Empty();
 
@@ -863,7 +863,7 @@ bool UKismetSystemLibrary::CapsuleOverlapActors_NEW(UObject* WorldContextObject,
 	return (OutActors.Num() > 0);
 }
 
-bool UKismetSystemLibrary::CapsuleOverlapComponents_NEW(UObject* WorldContextObject, const FVector& CapsulePos, float Radius, float HalfHeight, const TArray<TEnumAsByte<EObjectTypeQuery> > & ObjectTypes, UClass* ComponentClassFilter, const TArray<AActor*>& ActorsToIgnore, TArray<UPrimitiveComponent*>& OutComponents)
+bool UKismetSystemLibrary::CapsuleOverlapComponents_NEW(UObject* WorldContextObject, const FVector CapsulePos, float Radius, float HalfHeight, const TArray<TEnumAsByte<EObjectTypeQuery> > & ObjectTypes, UClass* ComponentClassFilter, const TArray<AActor*>& ActorsToIgnore, TArray<UPrimitiveComponent*>& OutComponents)
 {
 	OutComponents.Empty();
 
@@ -899,7 +899,7 @@ bool UKismetSystemLibrary::CapsuleOverlapComponents_NEW(UObject* WorldContextObj
 }
 
 
-bool UKismetSystemLibrary::ComponentOverlapActors_NEW(UPrimitiveComponent* Component, FTransform const& ComponentTransform, const TArray<TEnumAsByte<EObjectTypeQuery> > & ObjectTypes, UClass* ActorClassFilter, const TArray<AActor*>& ActorsToIgnore, TArray<AActor*>& OutActors)
+bool UKismetSystemLibrary::ComponentOverlapActors_NEW(UPrimitiveComponent* Component, const FTransform ComponentTransform, const TArray<TEnumAsByte<EObjectTypeQuery> > & ObjectTypes, UClass* ActorClassFilter, const TArray<AActor*>& ActorsToIgnore, TArray<AActor*>& OutActors)
 {
 	OutActors.Empty();
 
@@ -913,7 +913,7 @@ bool UKismetSystemLibrary::ComponentOverlapActors_NEW(UPrimitiveComponent* Compo
 	return (OutActors.Num() > 0);
 }
 
-bool UKismetSystemLibrary::ComponentOverlapComponents_NEW(UPrimitiveComponent* Component, FTransform const& ComponentTransform, const TArray<TEnumAsByte<EObjectTypeQuery> > & ObjectTypes, UClass* ComponentClassFilter, const TArray<AActor*>& ActorsToIgnore, TArray<UPrimitiveComponent*>& OutComponents)
+bool UKismetSystemLibrary::ComponentOverlapComponents_NEW(UPrimitiveComponent* Component, const FTransform ComponentTransform, const TArray<TEnumAsByte<EObjectTypeQuery> > & ObjectTypes, UClass* ComponentClassFilter, const TArray<AActor*>& ActorsToIgnore, TArray<UPrimitiveComponent*>& OutComponents)
 {
 	OutComponents.Empty();
 
@@ -950,12 +950,12 @@ bool UKismetSystemLibrary::ComponentOverlapComponents_NEW(UPrimitiveComponent* C
 static const float KISMET_TRACE_DEBUG_DRAW_DURATION = 5.f;
 static const float KISMET_TRACE_DEBUG_IMPACTPOINT_SIZE = 16.f;
 
-bool UKismetSystemLibrary::LineTraceSingle_NEW(UObject* WorldContextObject, const FVector& Start, const FVector& End, ETraceTypeQuery TraceChannel, bool bTraceComplex, const TArray<AActor*>& ActorsToIgnore, EDrawDebugTrace::Type DrawDebugType, FHitResult& OutHit, bool bIgnoreSelf)
+bool UKismetSystemLibrary::LineTraceSingle_NEW(UObject* WorldContextObject, const FVector Start, const FVector End, ETraceTypeQuery TraceChannel, bool bTraceComplex, const TArray<AActor*>& ActorsToIgnore, EDrawDebugTrace::Type DrawDebugType, FHitResult& OutHit, bool bIgnoreSelf)
 {
 	return LineTraceSingle_DEPRECATED(WorldContextObject, Start, End, UEngineTypes::ConvertToCollisionChannel(TraceChannel), bTraceComplex, ActorsToIgnore, DrawDebugType, OutHit, bIgnoreSelf);
 }
 
-bool UKismetSystemLibrary::LineTraceSingle_DEPRECATED(UObject* WorldContextObject, const FVector& Start, const FVector& End, ECollisionChannel TraceChannel, bool bTraceComplex, const TArray<AActor*>& ActorsToIgnore, EDrawDebugTrace::Type DrawDebugType, FHitResult& OutHit, bool bIgnoreSelf)
+bool UKismetSystemLibrary::LineTraceSingle_DEPRECATED(UObject* WorldContextObject, const FVector Start, const FVector End, ECollisionChannel TraceChannel, bool bTraceComplex, const TArray<AActor*>& ActorsToIgnore, EDrawDebugTrace::Type DrawDebugType, FHitResult& OutHit, bool bIgnoreSelf)
 {
 	static const FName LineTraceSingleName(TEXT("LineTraceSingle"));
 
@@ -1013,12 +1013,12 @@ bool UKismetSystemLibrary::LineTraceSingle_DEPRECATED(UObject* WorldContextObjec
 	return bHit;
 }
 
-bool UKismetSystemLibrary::LineTraceMulti_NEW(UObject* WorldContextObject, const FVector& Start, const FVector& End, ETraceTypeQuery TraceChannel, bool bTraceComplex, const TArray<AActor*>& ActorsToIgnore, EDrawDebugTrace::Type DrawDebugType, TArray<FHitResult>& OutHits, bool bIgnoreSelf)
+bool UKismetSystemLibrary::LineTraceMulti_NEW(UObject* WorldContextObject, const FVector Start, const FVector End, ETraceTypeQuery TraceChannel, bool bTraceComplex, const TArray<AActor*>& ActorsToIgnore, EDrawDebugTrace::Type DrawDebugType, TArray<FHitResult>& OutHits, bool bIgnoreSelf)
 {
 	return LineTraceMulti_DEPRECATED(WorldContextObject, Start, End, UEngineTypes::ConvertToCollisionChannel(TraceChannel), bTraceComplex, ActorsToIgnore, DrawDebugType, OutHits, bIgnoreSelf);
 }
 
-bool UKismetSystemLibrary::LineTraceMulti_DEPRECATED(UObject* WorldContextObject, const FVector& Start, const FVector& End, ECollisionChannel TraceChannel, bool bTraceComplex, const TArray<AActor*>& ActorsToIgnore, EDrawDebugTrace::Type DrawDebugType, TArray<FHitResult>& OutHits, bool bIgnoreSelf)
+bool UKismetSystemLibrary::LineTraceMulti_DEPRECATED(UObject* WorldContextObject, const FVector Start, const FVector End, ECollisionChannel TraceChannel, bool bTraceComplex, const TArray<AActor*>& ActorsToIgnore, EDrawDebugTrace::Type DrawDebugType, TArray<FHitResult>& OutHits, bool bIgnoreSelf)
 {
 	static const FName LineTraceMultiName(TEXT("LineTraceMulti"));
 
@@ -1124,7 +1124,7 @@ void DrawDebugSweptBox(const UWorld* InWorld, FVector const& Start, FVector cons
 }
 
 
-bool UKismetSystemLibrary::BoxTraceSingle(UObject* WorldContextObject, const FVector& Start, const FVector& End, const FVector & HalfSize, const FRotator & Orientation, ETraceTypeQuery TraceChannel, bool bTraceComplex, const TArray<AActor*>& ActorsToIgnore, EDrawDebugTrace::Type DrawDebugType, FHitResult& OutHit, bool bIgnoreSelf)
+bool UKismetSystemLibrary::BoxTraceSingle(UObject* WorldContextObject, const FVector Start, const FVector End, const FVector & HalfSize, const FRotator & Orientation, ETraceTypeQuery TraceChannel, bool bTraceComplex, const TArray<AActor*>& ActorsToIgnore, EDrawDebugTrace::Type DrawDebugType, FHitResult& OutHit, bool bIgnoreSelf)
 {
 	static const FName BoxTraceSingleName(TEXT("BoxTraceSingle"));
 
@@ -1179,7 +1179,7 @@ bool UKismetSystemLibrary::BoxTraceSingle(UObject* WorldContextObject, const FVe
 	return bHit;
 }
 
-bool UKismetSystemLibrary::BoxTraceMulti(UObject* WorldContextObject, const FVector& Start, const FVector& End, const FVector & HalfSize, const FRotator & Orientation, ETraceTypeQuery TraceChannel, bool bTraceComplex, const TArray<AActor*>& ActorsToIgnore, EDrawDebugTrace::Type DrawDebugType, TArray<FHitResult>& OutHits, bool bIgnoreSelf)
+bool UKismetSystemLibrary::BoxTraceMulti(UObject* WorldContextObject, const FVector Start, const FVector End, const FVector & HalfSize, const FRotator & Orientation, ETraceTypeQuery TraceChannel, bool bTraceComplex, const TArray<AActor*>& ActorsToIgnore, EDrawDebugTrace::Type DrawDebugType, TArray<FHitResult>& OutHits, bool bIgnoreSelf)
 {
 	static const FName BoxTraceMultiName(TEXT("BoxTraceMulti"));
 
@@ -1241,12 +1241,12 @@ bool UKismetSystemLibrary::BoxTraceMulti(UObject* WorldContextObject, const FVec
 	return bHit;
 }
 
-bool UKismetSystemLibrary::SphereTraceSingle_NEW(UObject* WorldContextObject, const FVector& Start, const FVector& End, float Radius, ETraceTypeQuery TraceChannel, bool bTraceComplex, const TArray<AActor*>& ActorsToIgnore, EDrawDebugTrace::Type DrawDebugType, FHitResult& OutHit, bool bIgnoreSelf)
+bool UKismetSystemLibrary::SphereTraceSingle_NEW(UObject* WorldContextObject, const FVector Start, const FVector End, float Radius, ETraceTypeQuery TraceChannel, bool bTraceComplex, const TArray<AActor*>& ActorsToIgnore, EDrawDebugTrace::Type DrawDebugType, FHitResult& OutHit, bool bIgnoreSelf)
 {
 	return SphereTraceSingle_DEPRECATED(WorldContextObject, Start, End, Radius, UEngineTypes::ConvertToCollisionChannel(TraceChannel), bTraceComplex, ActorsToIgnore, DrawDebugType, OutHit, bIgnoreSelf);
 }
 
-bool UKismetSystemLibrary::SphereTraceSingle_DEPRECATED(UObject* WorldContextObject, const FVector& Start, const FVector& End, float Radius, ECollisionChannel TraceChannel, bool bTraceComplex, const TArray<AActor*>& ActorsToIgnore, EDrawDebugTrace::Type DrawDebugType, FHitResult& OutHit, bool bIgnoreSelf)
+bool UKismetSystemLibrary::SphereTraceSingle_DEPRECATED(UObject* WorldContextObject, const FVector Start, const FVector End, float Radius, ECollisionChannel TraceChannel, bool bTraceComplex, const TArray<AActor*>& ActorsToIgnore, EDrawDebugTrace::Type DrawDebugType, FHitResult& OutHit, bool bIgnoreSelf)
 {
 	static const FName SphereTraceSingleName(TEXT("SphereTraceSingle"));
 
@@ -1303,12 +1303,12 @@ bool UKismetSystemLibrary::SphereTraceSingle_DEPRECATED(UObject* WorldContextObj
 	return bHit;
 }
 
-bool UKismetSystemLibrary::SphereTraceMulti_NEW(UObject* WorldContextObject, const FVector& Start, const FVector& End, float Radius, ETraceTypeQuery TraceChannel, bool bTraceComplex, const TArray<AActor*>& ActorsToIgnore, EDrawDebugTrace::Type DrawDebugType, TArray<FHitResult>& OutHits, bool bIgnoreSelf)
+bool UKismetSystemLibrary::SphereTraceMulti_NEW(UObject* WorldContextObject, const FVector Start, const FVector End, float Radius, ETraceTypeQuery TraceChannel, bool bTraceComplex, const TArray<AActor*>& ActorsToIgnore, EDrawDebugTrace::Type DrawDebugType, TArray<FHitResult>& OutHits, bool bIgnoreSelf)
 {
 	return SphereTraceMulti_DEPRECATED(WorldContextObject, Start, End, Radius, UEngineTypes::ConvertToCollisionChannel(TraceChannel), bTraceComplex, ActorsToIgnore, DrawDebugType, OutHits, bIgnoreSelf);
 }
 
-bool UKismetSystemLibrary::SphereTraceMulti_DEPRECATED(UObject* WorldContextObject, const FVector& Start, const FVector& End, float Radius, ECollisionChannel TraceChannel, bool bTraceComplex, const TArray<AActor*>& ActorsToIgnore, EDrawDebugTrace::Type DrawDebugType, TArray<FHitResult>& OutHits, bool bIgnoreSelf)
+bool UKismetSystemLibrary::SphereTraceMulti_DEPRECATED(UObject* WorldContextObject, const FVector Start, const FVector End, float Radius, ECollisionChannel TraceChannel, bool bTraceComplex, const TArray<AActor*>& ActorsToIgnore, EDrawDebugTrace::Type DrawDebugType, TArray<FHitResult>& OutHits, bool bIgnoreSelf)
 {
 	static const FName SphereTraceMultiName(TEXT("SphereTraceMulti"));
 
@@ -1372,12 +1372,12 @@ bool UKismetSystemLibrary::SphereTraceMulti_DEPRECATED(UObject* WorldContextObje
 	return bHit;
 }
 
-bool UKismetSystemLibrary::CapsuleTraceSingle_NEW(UObject* WorldContextObject, const FVector& Start, const FVector& End, float Radius, float HalfHeight, ETraceTypeQuery TraceChannel, bool bTraceComplex, const TArray<AActor*>& ActorsToIgnore, EDrawDebugTrace::Type DrawDebugType, FHitResult& OutHit, bool bIgnoreSelf)
+bool UKismetSystemLibrary::CapsuleTraceSingle_NEW(UObject* WorldContextObject, const FVector Start, const FVector End, float Radius, float HalfHeight, ETraceTypeQuery TraceChannel, bool bTraceComplex, const TArray<AActor*>& ActorsToIgnore, EDrawDebugTrace::Type DrawDebugType, FHitResult& OutHit, bool bIgnoreSelf)
 {
 	return CapsuleTraceSingle_DEPRECATED(WorldContextObject, Start, End, Radius, HalfHeight, UEngineTypes::ConvertToCollisionChannel(TraceChannel), bTraceComplex, ActorsToIgnore, DrawDebugType, OutHit, bIgnoreSelf);
 }
 
-bool UKismetSystemLibrary::CapsuleTraceSingle_DEPRECATED(UObject* WorldContextObject, const FVector& Start, const FVector& End, float Radius, float HalfHeight, ECollisionChannel TraceChannel, bool bTraceComplex, const TArray<AActor*>& ActorsToIgnore, EDrawDebugTrace::Type DrawDebugType, FHitResult& OutHit, bool bIgnoreSelf)
+bool UKismetSystemLibrary::CapsuleTraceSingle_DEPRECATED(UObject* WorldContextObject, const FVector Start, const FVector End, float Radius, float HalfHeight, ECollisionChannel TraceChannel, bool bTraceComplex, const TArray<AActor*>& ActorsToIgnore, EDrawDebugTrace::Type DrawDebugType, FHitResult& OutHit, bool bIgnoreSelf)
 {
 	static const FName CapsuleTraceSingleName(TEXT("CapsuleTraceSingle"));
 
@@ -1440,12 +1440,12 @@ bool UKismetSystemLibrary::CapsuleTraceSingle_DEPRECATED(UObject* WorldContextOb
 	return bHit;
 }
 
-bool UKismetSystemLibrary::CapsuleTraceMulti_NEW(UObject* WorldContextObject, const FVector& Start, const FVector& End, float Radius, float HalfHeight, ETraceTypeQuery TraceChannel, bool bTraceComplex, const TArray<AActor*>& ActorsToIgnore, EDrawDebugTrace::Type DrawDebugType, TArray<FHitResult>& OutHits, bool bIgnoreSelf)
+bool UKismetSystemLibrary::CapsuleTraceMulti_NEW(UObject* WorldContextObject, const FVector Start, const FVector End, float Radius, float HalfHeight, ETraceTypeQuery TraceChannel, bool bTraceComplex, const TArray<AActor*>& ActorsToIgnore, EDrawDebugTrace::Type DrawDebugType, TArray<FHitResult>& OutHits, bool bIgnoreSelf)
 {
 	return CapsuleTraceMulti_DEPRECATED(WorldContextObject, Start, End, Radius, HalfHeight, UEngineTypes::ConvertToCollisionChannel(TraceChannel), bTraceComplex, ActorsToIgnore, DrawDebugType, OutHits, bIgnoreSelf);
 }
 
-bool UKismetSystemLibrary::CapsuleTraceMulti_DEPRECATED(UObject* WorldContextObject, const FVector& Start, const FVector& End, float Radius, float HalfHeight, ECollisionChannel TraceChannel, bool bTraceComplex, const TArray<AActor*>& ActorsToIgnore, EDrawDebugTrace::Type DrawDebugType, TArray<FHitResult>& OutHits, bool bIgnoreSelf)
+bool UKismetSystemLibrary::CapsuleTraceMulti_DEPRECATED(UObject* WorldContextObject, const FVector Start, const FVector End, float Radius, float HalfHeight, ECollisionChannel TraceChannel, bool bTraceComplex, const TArray<AActor*>& ActorsToIgnore, EDrawDebugTrace::Type DrawDebugType, TArray<FHitResult>& OutHits, bool bIgnoreSelf)
 {
 	static const FName CapsuleTraceMultiName(TEXT("CapsuleTraceMulti"));
 
@@ -1516,7 +1516,7 @@ bool UKismetSystemLibrary::CapsuleTraceMulti_DEPRECATED(UObject* WorldContextObj
 }
 
 /** Object Query functions **/
-bool UKismetSystemLibrary::LineTraceSingleForObjects(UObject* WorldContextObject, const FVector& Start, const FVector& End, const TArray<TEnumAsByte<EObjectTypeQuery> > & ObjectTypes, bool bTraceComplex, const TArray<AActor*>& ActorsToIgnore, EDrawDebugTrace::Type DrawDebugType, FHitResult& OutHit, bool bIgnoreSelf)
+bool UKismetSystemLibrary::LineTraceSingleForObjects(UObject* WorldContextObject, const FVector Start, const FVector End, const TArray<TEnumAsByte<EObjectTypeQuery> > & ObjectTypes, bool bTraceComplex, const TArray<AActor*>& ActorsToIgnore, EDrawDebugTrace::Type DrawDebugType, FHitResult& OutHit, bool bIgnoreSelf)
 {
 
 	TArray<TEnumAsByte<ECollisionChannel>> CollisionObjectTraces;
@@ -1530,7 +1530,7 @@ bool UKismetSystemLibrary::LineTraceSingleForObjects(UObject* WorldContextObject
 	return LineTraceSingleByObject_DEPRECATED(WorldContextObject, Start, End, CollisionObjectTraces, bTraceComplex, ActorsToIgnore, DrawDebugType, OutHit, bIgnoreSelf);
 }
 
-bool UKismetSystemLibrary::LineTraceSingleByObject_DEPRECATED(UObject* WorldContextObject, const FVector& Start, const FVector& End, const TArray<TEnumAsByte<ECollisionChannel> > & ObjectsToTrace, bool bTraceComplex, const TArray<AActor*>& ActorsToIgnore, EDrawDebugTrace::Type DrawDebugType, FHitResult& OutHit, bool bIgnoreSelf)
+bool UKismetSystemLibrary::LineTraceSingleByObject_DEPRECATED(UObject* WorldContextObject, const FVector Start, const FVector End, const TArray<TEnumAsByte<ECollisionChannel> > & ObjectsToTrace, bool bTraceComplex, const TArray<AActor*>& ActorsToIgnore, EDrawDebugTrace::Type DrawDebugType, FHitResult& OutHit, bool bIgnoreSelf)
 {
 	static const FName LineTraceSingleName(TEXT("LineTraceSingle"));
 
@@ -1608,7 +1608,7 @@ bool UKismetSystemLibrary::LineTraceSingleByObject_DEPRECATED(UObject* WorldCont
 	return bHit;
 }
 
-bool UKismetSystemLibrary::LineTraceMultiForObjects(UObject* WorldContextObject, const FVector& Start, const FVector& End, const TArray<TEnumAsByte<EObjectTypeQuery> > & ObjectTypes, bool bTraceComplex, const TArray<AActor*>& ActorsToIgnore, EDrawDebugTrace::Type DrawDebugType, TArray<FHitResult>& OutHits, bool bIgnoreSelf)
+bool UKismetSystemLibrary::LineTraceMultiForObjects(UObject* WorldContextObject, const FVector Start, const FVector End, const TArray<TEnumAsByte<EObjectTypeQuery> > & ObjectTypes, bool bTraceComplex, const TArray<AActor*>& ActorsToIgnore, EDrawDebugTrace::Type DrawDebugType, TArray<FHitResult>& OutHits, bool bIgnoreSelf)
 {
 	TArray<TEnumAsByte<ECollisionChannel>> CollisionObjectTraces;
 	CollisionObjectTraces.AddUninitialized(ObjectTypes.Num());
@@ -1621,7 +1621,7 @@ bool UKismetSystemLibrary::LineTraceMultiForObjects(UObject* WorldContextObject,
 	return LineTraceMultiByObject_DEPRECATED(WorldContextObject, Start, End, CollisionObjectTraces, bTraceComplex, ActorsToIgnore, DrawDebugType, OutHits, bIgnoreSelf);
 }
 
-bool UKismetSystemLibrary::LineTraceMultiByObject_DEPRECATED(UObject* WorldContextObject, const FVector& Start, const FVector& End, const TArray<TEnumAsByte<ECollisionChannel> > & ObjectsToTrace, bool bTraceComplex, const TArray<AActor*>& ActorsToIgnore, EDrawDebugTrace::Type DrawDebugType, TArray<FHitResult>& OutHits, bool bIgnoreSelf)
+bool UKismetSystemLibrary::LineTraceMultiByObject_DEPRECATED(UObject* WorldContextObject, const FVector Start, const FVector End, const TArray<TEnumAsByte<ECollisionChannel> > & ObjectsToTrace, bool bTraceComplex, const TArray<AActor*>& ActorsToIgnore, EDrawDebugTrace::Type DrawDebugType, TArray<FHitResult>& OutHits, bool bIgnoreSelf)
 {
 	static const FName LineTraceMultiName(TEXT("LineTraceMulti"));
 
@@ -1706,7 +1706,7 @@ bool UKismetSystemLibrary::LineTraceMultiByObject_DEPRECATED(UObject* WorldConte
 	return bHit;
 }
 
-bool UKismetSystemLibrary::SphereTraceSingleForObjects(UObject* WorldContextObject, const FVector& Start, const FVector& End, float Radius, const TArray<TEnumAsByte<EObjectTypeQuery> > & ObjectTypes, bool bTraceComplex, const TArray<AActor*>& ActorsToIgnore, EDrawDebugTrace::Type DrawDebugType, FHitResult& OutHit, bool bIgnoreSelf)
+bool UKismetSystemLibrary::SphereTraceSingleForObjects(UObject* WorldContextObject, const FVector Start, const FVector End, float Radius, const TArray<TEnumAsByte<EObjectTypeQuery> > & ObjectTypes, bool bTraceComplex, const TArray<AActor*>& ActorsToIgnore, EDrawDebugTrace::Type DrawDebugType, FHitResult& OutHit, bool bIgnoreSelf)
 {
 	TArray<TEnumAsByte<ECollisionChannel>> CollisionObjectTraces;
 	CollisionObjectTraces.AddUninitialized(ObjectTypes.Num());
@@ -1719,7 +1719,7 @@ bool UKismetSystemLibrary::SphereTraceSingleForObjects(UObject* WorldContextObje
 	return SphereTraceSingleByObject_DEPRECATED(WorldContextObject, Start, End, Radius, CollisionObjectTraces, bTraceComplex, ActorsToIgnore, DrawDebugType, OutHit, bIgnoreSelf);
 }
 
-bool UKismetSystemLibrary::SphereTraceSingleByObject_DEPRECATED(UObject* WorldContextObject, const FVector& Start, const FVector& End, float Radius, const TArray<TEnumAsByte<ECollisionChannel> > & ObjectsToTrace, bool bTraceComplex, const TArray<AActor*>& ActorsToIgnore, EDrawDebugTrace::Type DrawDebugType, FHitResult& OutHit, bool bIgnoreSelf)
+bool UKismetSystemLibrary::SphereTraceSingleByObject_DEPRECATED(UObject* WorldContextObject, const FVector Start, const FVector End, float Radius, const TArray<TEnumAsByte<ECollisionChannel> > & ObjectsToTrace, bool bTraceComplex, const TArray<AActor*>& ActorsToIgnore, EDrawDebugTrace::Type DrawDebugType, FHitResult& OutHit, bool bIgnoreSelf)
 {
 	static const FName SphereTraceSingleName(TEXT("SphereTraceSingle"));
 
@@ -1800,7 +1800,7 @@ bool UKismetSystemLibrary::SphereTraceSingleByObject_DEPRECATED(UObject* WorldCo
 	return bHit;
 }
 
-bool UKismetSystemLibrary::SphereTraceMultiForObjects(UObject* WorldContextObject, const FVector& Start, const FVector& End, float Radius, const TArray<TEnumAsByte<EObjectTypeQuery> > & ObjectTypes, bool bTraceComplex, const TArray<AActor*>& ActorsToIgnore, EDrawDebugTrace::Type DrawDebugType, TArray<FHitResult>& OutHits, bool bIgnoreSelf)
+bool UKismetSystemLibrary::SphereTraceMultiForObjects(UObject* WorldContextObject, const FVector Start, const FVector End, float Radius, const TArray<TEnumAsByte<EObjectTypeQuery> > & ObjectTypes, bool bTraceComplex, const TArray<AActor*>& ActorsToIgnore, EDrawDebugTrace::Type DrawDebugType, TArray<FHitResult>& OutHits, bool bIgnoreSelf)
 {
 	TArray<TEnumAsByte<ECollisionChannel>> CollisionObjectTraces;
 	CollisionObjectTraces.AddUninitialized(ObjectTypes.Num());
@@ -1813,7 +1813,7 @@ bool UKismetSystemLibrary::SphereTraceMultiForObjects(UObject* WorldContextObjec
 	return SphereTraceMultiByObject_DEPRECATED(WorldContextObject, Start, End, Radius, CollisionObjectTraces, bTraceComplex, ActorsToIgnore, DrawDebugType, OutHits, bIgnoreSelf);
 }
 
-bool UKismetSystemLibrary::SphereTraceMultiByObject_DEPRECATED(UObject* WorldContextObject, const FVector& Start, const FVector& End, float Radius, const TArray<TEnumAsByte<ECollisionChannel> > & ObjectsToTrace, bool bTraceComplex, const TArray<AActor*>& ActorsToIgnore, EDrawDebugTrace::Type DrawDebugType, TArray<FHitResult>& OutHits, bool bIgnoreSelf)
+bool UKismetSystemLibrary::SphereTraceMultiByObject_DEPRECATED(UObject* WorldContextObject, const FVector Start, const FVector End, float Radius, const TArray<TEnumAsByte<ECollisionChannel> > & ObjectsToTrace, bool bTraceComplex, const TArray<AActor*>& ActorsToIgnore, EDrawDebugTrace::Type DrawDebugType, TArray<FHitResult>& OutHits, bool bIgnoreSelf)
 {
 	static const FName SphereTraceMultiName(TEXT("SphereTraceMulti"));
 
@@ -1897,7 +1897,7 @@ bool UKismetSystemLibrary::SphereTraceMultiByObject_DEPRECATED(UObject* WorldCon
 	return bHit;
 }
 
-bool UKismetSystemLibrary::BoxTraceSingleForObjects(UObject* WorldContextObject, const FVector& Start, const FVector& End, const FVector & HalfSize, const FRotator & Orientation, const TArray<TEnumAsByte<EObjectTypeQuery> > & ObjectTypes, bool bTraceComplex, const TArray<AActor*>& ActorsToIgnore, EDrawDebugTrace::Type DrawDebugType, FHitResult& OutHit, bool bIgnoreSelf)
+bool UKismetSystemLibrary::BoxTraceSingleForObjects(UObject* WorldContextObject, const FVector Start, const FVector End, const FVector HalfSize, const FRotator Orientation, const TArray<TEnumAsByte<EObjectTypeQuery> > & ObjectTypes, bool bTraceComplex, const TArray<AActor*>& ActorsToIgnore, EDrawDebugTrace::Type DrawDebugType, FHitResult& OutHit, bool bIgnoreSelf)
 {
 	static const FName BoxTraceSingleName(TEXT("BoxTraceSingle"));
 
@@ -1980,7 +1980,7 @@ bool UKismetSystemLibrary::BoxTraceSingleForObjects(UObject* WorldContextObject,
 	return bHit;
 }
 
-bool UKismetSystemLibrary::BoxTraceMultiForObjects(UObject* WorldContextObject, const FVector& Start, const FVector& End, const FVector & HalfSize, const FRotator & Orientation, const TArray<TEnumAsByte<	EObjectTypeQuery> > & ObjectTypes, bool bTraceComplex, const TArray<AActor*>& ActorsToIgnore, EDrawDebugTrace::Type DrawDebugType, TArray<FHitResult>& OutHits, bool bIgnoreSelf)
+bool UKismetSystemLibrary::BoxTraceMultiForObjects(UObject* WorldContextObject, const FVector Start, const FVector End, const FVector HalfSize, const FRotator Orientation, const TArray<TEnumAsByte<	EObjectTypeQuery> > & ObjectTypes, bool bTraceComplex, const TArray<AActor*>& ActorsToIgnore, EDrawDebugTrace::Type DrawDebugType, TArray<FHitResult>& OutHits, bool bIgnoreSelf)
 {
 	static const FName BoxTraceMultiName(TEXT("BoxTraceMulti"));
 
@@ -2070,7 +2070,7 @@ bool UKismetSystemLibrary::BoxTraceMultiForObjects(UObject* WorldContextObject,
 	return bHit;
 }
 
-bool UKismetSystemLibrary::CapsuleTraceSingleForObjects(UObject* WorldContextObject, const FVector& Start, const FVector& End, float Radius, float HalfHeight, const TArray<TEnumAsByte<EObjectTypeQuery> > & ObjectTypes, bool bTraceComplex, const TArray<AActor*>& ActorsToIgnore, EDrawDebugTrace::Type DrawDebugType, FHitResult& OutHit, bool bIgnoreSelf)
+bool UKismetSystemLibrary::CapsuleTraceSingleForObjects(UObject* WorldContextObject, const FVector Start, const FVector End, float Radius, float HalfHeight, const TArray<TEnumAsByte<EObjectTypeQuery> > & ObjectTypes, bool bTraceComplex, const TArray<AActor*>& ActorsToIgnore, EDrawDebugTrace::Type DrawDebugType, FHitResult& OutHit, bool bIgnoreSelf)
 {
 	TArray<TEnumAsByte<ECollisionChannel>> CollisionObjectTraces;
 	CollisionObjectTraces.AddUninitialized(ObjectTypes.Num());
@@ -2083,7 +2083,7 @@ bool UKismetSystemLibrary::CapsuleTraceSingleForObjects(UObject* WorldContextObj
 	return CapsuleTraceSingleByObject_DEPRECATED(WorldContextObject, Start, End, Radius, HalfHeight, CollisionObjectTraces, bTraceComplex, ActorsToIgnore, DrawDebugType, OutHit, bIgnoreSelf);
 }
 
-bool UKismetSystemLibrary::CapsuleTraceSingleByObject_DEPRECATED(UObject* WorldContextObject, const FVector& Start, const FVector& End, float Radius, float HalfHeight, const TArray<TEnumAsByte<ECollisionChannel> > & ObjectsToTrace, bool bTraceComplex, const TArray<AActor*>& ActorsToIgnore, EDrawDebugTrace::Type DrawDebugType, FHitResult& OutHit, bool bIgnoreSelf)
+bool UKismetSystemLibrary::CapsuleTraceSingleByObject_DEPRECATED(UObject* WorldContextObject, const FVector Start, const FVector End, float Radius, float HalfHeight, const TArray<TEnumAsByte<ECollisionChannel> > & ObjectsToTrace, bool bTraceComplex, const TArray<AActor*>& ActorsToIgnore, EDrawDebugTrace::Type DrawDebugType, FHitResult& OutHit, bool bIgnoreSelf)
 {
 	static const FName CapsuleTraceSingleName(TEXT("CapsuleTraceSingle"));
 
@@ -2166,7 +2166,7 @@ bool UKismetSystemLibrary::CapsuleTraceSingleByObject_DEPRECATED(UObject* WorldC
 	return bHit;
 }
 
-bool UKismetSystemLibrary::CapsuleTraceMultiForObjects(UObject* WorldContextObject, const FVector& Start, const FVector& End, float Radius, float HalfHeight, const TArray<TEnumAsByte<EObjectTypeQuery> > & ObjectTypes, bool bTraceComplex, const TArray<AActor*>& ActorsToIgnore, EDrawDebugTrace::Type DrawDebugType, TArray<FHitResult>& OutHits, bool bIgnoreSelf)
+bool UKismetSystemLibrary::CapsuleTraceMultiForObjects(UObject* WorldContextObject, const FVector Start, const FVector End, float Radius, float HalfHeight, const TArray<TEnumAsByte<EObjectTypeQuery> > & ObjectTypes, bool bTraceComplex, const TArray<AActor*>& ActorsToIgnore, EDrawDebugTrace::Type DrawDebugType, TArray<FHitResult>& OutHits, bool bIgnoreSelf)
 {
 	TArray<TEnumAsByte<ECollisionChannel>> CollisionObjectTraces;
 	CollisionObjectTraces.AddUninitialized(ObjectTypes.Num());
@@ -2178,7 +2178,7 @@ bool UKismetSystemLibrary::CapsuleTraceMultiForObjects(UObject* WorldContextObje
 	return CapsuleTraceMultiByObject_DEPRECATED(WorldContextObject, Start, End, Radius, HalfHeight, CollisionObjectTraces, bTraceComplex, ActorsToIgnore, DrawDebugType, OutHits, bIgnoreSelf);
 }
 
-bool UKismetSystemLibrary::CapsuleTraceMultiByObject_DEPRECATED(UObject* WorldContextObject, const FVector& Start, const FVector& End, float Radius, float HalfHeight, const TArray<TEnumAsByte<ECollisionChannel> > & ObjectsToTrace, bool bTraceComplex, const TArray<AActor*>& ActorsToIgnore, EDrawDebugTrace::Type DrawDebugType, TArray<FHitResult>& OutHits, bool bIgnoreSelf)
+bool UKismetSystemLibrary::CapsuleTraceMultiByObject_DEPRECATED(UObject* WorldContextObject, const FVector Start, const FVector End, float Radius, float HalfHeight, const TArray<TEnumAsByte<ECollisionChannel> > & ObjectsToTrace, bool bTraceComplex, const TArray<AActor*>& ActorsToIgnore, EDrawDebugTrace::Type DrawDebugType, TArray<FHitResult>& OutHits, bool bIgnoreSelf)
 {
 	static const FName CapsuleTraceMultiName(TEXT("CapsuleTraceMulti"));
 
