Fix incorrectly passing angle to max angular velocity (should be radians)
TTP#340382

diff --git a/Engine/Source/Runtime/Engine/Private/PhysicsEngine/BodyInstance.cpp b/Engine/Source/Runtime/Engine/Private/PhysicsEngine/BodyInstance.cpp
--- a/Engine/Source/Runtime/Engine/Private/PhysicsEngine/BodyInstance.cpp
+++ b/Engine/Source/Runtime/Engine/Private/PhysicsEngine/BodyInstance.cpp
@@ -2662,7 +2662,7 @@ void FBodyInstance::SetMaxAngularVelocity(float NewMaxAngVel, bool bAddToCurrent
 			NewMaxAngVel += OldMaxAngVel;
 		}
 
-		PRigidDynamic->setMaxAngularVelocity(NewMaxAngVel);
+		PRigidDynamic->setMaxAngularVelocity(RadNewMaxAngVel);
 
 		MaxAngularVelocity = NewMaxAngVel;
 	}
