Removed clear from FSlateTextureRenderTarget2DResource::InitDynamicRHI

diff --git a/Engine/Source/Runtime/Engine/Private/Slate/SlateTextures.cpp b/Engine/Source/Runtime/Engine/Private/Slate/SlateTextures.cpp
--- a/Engine/Source/Runtime/Engine/Private/Slate/SlateTextures.cpp
+++ b/Engine/Source/Runtime/Engine/Private/Slate/SlateTextures.cpp
@@ -185,10 +185,6 @@ void FSlateTextureRenderTarget2DResource::InitDynamicRHI()
 			Texture2DRHI
 			);
 		TextureRHI = (FTextureRHIRef&)Texture2DRHI;
-
-		// make sure the texture target gets cleared
-		FRHICommandListImmediate& RHICmdList = FRHICommandListExecutor::GetImmediateCommandList();
-		UpdateDeferredResource(RHICmdList);
 	}
 
 	// Create the sampler state RHI resource.
