Fixed: FOLIAGE: Level transform does not work for foliage when loading the map in the editor TTP# 340396

diff --git a/Engine/Source/Runtime/Core/Public/UObject/ObjectVersion.h b/Engine/Source/Runtime/Core/Public/UObject/ObjectVersion.h
--- a/Engine/Source/Runtime/Core/Public/UObject/ObjectVersion.h
+++ b/Engine/Source/Runtime/Core/Public/UObject/ObjectVersion.h
@@ -574,6 +574,8 @@ enum EUnrealEngineObjectUE4Version
 	VER_UE4_ADD_TRANSACTIONAL_TO_DATA_ASSETS,
 	// Change LB_AlphaBlend to LB_WeightBlend in ELandscapeLayerBlendType
 	VER_UE4_ADD_LB_WEIGHTBLEND,
+	// Add root component to an foliage actor, all foliage cluster components will be attached to a root
+	VER_UE4_ADD_ROOTCOMPONENT_TO_FOLIAGEACTOR,
 
 	// -----<new versions can be added before this line>-------------------------------------------------
 	// - this needs to be the last line (see note below)
diff --git a/Engine/Source/Runtime/Engine/Classes/Foliage/InstancedFoliageActor.h b/Engine/Source/Runtime/Engine/Classes/Foliage/InstancedFoliageActor.h
--- a/Engine/Source/Runtime/Engine/Classes/Foliage/InstancedFoliageActor.h
+++ b/Engine/Source/Runtime/Engine/Classes/Foliage/InstancedFoliageActor.h
@@ -90,6 +90,9 @@ public:
 	// Returns the location for the widget
 	ENGINE_API FVector GetSelectionLocation();
 
+	// Transforms Editor specific data which is stored in world space
+	ENGINE_API void ApplyLevelTransform(const FTransform& LevelTransform);
+
 	/*
 	* Get the instanced foliage actor for the current streaming level.
 	*
diff --git a/Engine/Source/Runtime/Engine/Private/ActorEditor.cpp b/Engine/Source/Runtime/Engine/Private/ActorEditor.cpp
--- a/Engine/Source/Runtime/Engine/Private/ActorEditor.cpp
+++ b/Engine/Source/Runtime/Engine/Private/ActorEditor.cpp
@@ -84,7 +84,7 @@ void AActor::PostEditMove(bool bFinished)
 
 	if ( bFinished )
 	{
-		if ( GIsEditor && !GetWorld()->IsPlayInEditor() )
+		if ( GIsEditor && !GetWorld()->IsPlayInEditor() && !FLevelUtils::IsMovingLevel())
 		{
 			AInstancedFoliageActor* IFA = AInstancedFoliageActor::GetInstancedFoliageActorForLevel(GetLevel());
 			if (IFA)
diff --git a/Engine/Source/Runtime/Engine/Private/InstancedFoliage.cpp b/Engine/Source/Runtime/Engine/Private/InstancedFoliage.cpp
--- a/Engine/Source/Runtime/Engine/Private/InstancedFoliage.cpp
+++ b/Engine/Source/Runtime/Engine/Private/InstancedFoliage.cpp
@@ -298,8 +298,9 @@ void FFoliageMeshInfo::AddInstance(AInstancedFoliageActor* InIFA, UFoliageType*
 
 		BestCluster->ClusterComponent->BodyInstance.CopyBodyInstancePropertiesFrom(&InSettings->BodyInstance);
 
-		BestCluster->ClusterComponent->SetWorldTransform(InstanceToWorld);
+		BestCluster->ClusterComponent->AttachTo(InIFA->GetRootComponent());
 		BestCluster->ClusterComponent->RegisterComponent();
+		BestCluster->ClusterComponent->SetWorldTransform(InstanceToWorld);
 	}
 	else
 	{
@@ -330,7 +331,7 @@ void FFoliageMeshInfo::AddInstance(AInstancedFoliageActor* InIFA, UFoliageType*
 	BestCluster->ClusterComponent->CachedMaxDrawDistance = CullDistance;
 
 	CheckValid();
-		}
+}
 
 void FFoliageMeshInfo::RemoveInstances(AInstancedFoliageActor* InIFA, const TArray<int32>& InInstancesToRemove)
 {
@@ -704,6 +705,10 @@ void FFoliageMeshInfo::SelectInstances(AInstancedFoliageActor* InIFA, bool bSele
 AInstancedFoliageActor::AInstancedFoliageActor(const FPostConstructInitializeProperties& PCIP)
 	: Super(PCIP)
 {
+	TSubobjectPtr<USceneComponent> SceneComponent = PCIP.CreateDefaultSubobject<USceneComponent>(this, TEXT("RootComponent0"));
+	RootComponent = SceneComponent;
+	RootComponent->Mobility = EComponentMobility::Movable;
+	
 	SetActorEnableCollision(true);
 #if WITH_EDITORONLY_DATA
 	bListedInSceneOutliner = false;
@@ -1463,7 +1468,41 @@ void AInstancedFoliageActor::MapRebuild()
 	}
 }
 
-#endif
+void AInstancedFoliageActor::ApplyLevelTransform(const FTransform& LevelTransform)
+{
+	for (auto& MeshPair : FoliageMeshes)
+	{
+		FFoliageMeshInfo& MeshInfo = *MeshPair.Value;
+
+		MeshInfo.InstanceHash->Empty();
+		for (int32 InstanceIdx = 0; InstanceIdx < MeshInfo.Instances.Num(); InstanceIdx++)
+		{
+			FFoliageInstance& Instance = MeshInfo.Instances[InstanceIdx];
+			FTransform NewTransform = Instance.GetInstanceWorldTransform() * LevelTransform;
+			
+			Instance.Location		= NewTransform.GetLocation();
+			Instance.Rotation		= NewTransform.GetRotation().Rotator();
+			Instance.DrawScale3D	= NewTransform.GetScale3D();
+			// Rehash instance location
+			MeshInfo.InstanceHash->InsertInstance(Instance.Location, InstanceIdx);
+		}
+
+		for (auto It = MeshInfo.ComponentHash.CreateIterator(); It; ++It)
+		{
+			// We assume here that component we painted foliage on, was transformed as well
+			FFoliageComponentHashInfo& Info = It.Value();
+			Info.UpdateLocationFromActor(It.Key());
+		}
+
+		// Recalc cluster bounds
+		for (FFoliageInstanceCluster& Cluster : MeshInfo.InstanceClusters)
+		{
+			Cluster.Bounds = Cluster.ClusterComponent->CalcBounds(Cluster.ClusterComponent->ComponentToWorld);
+		}
+	}
+}
+
+#endif // WITH_EDITOR
 
 struct FFoliageMeshInfo_Old
 {
@@ -1538,6 +1577,17 @@ void AInstancedFoliageActor::Serialize(FArchive& Ar)
 			}
 		}
 	}
+
+	// Attach cluster components to a root component
+	if (Ar.UE4Ver() < VER_UE4_ADD_ROOTCOMPONENT_TO_FOLIAGEACTOR && GIsEditor && Ar.IsLoading())
+	{
+		TArray<UInstancedStaticMeshComponent*> ClusterComponents;
+		GetComponents(ClusterComponents);
+		for (UInstancedStaticMeshComponent* Component : ClusterComponents)
+		{
+			Component->AttachTo(GetRootComponent(), NAME_None, EAttachLocation::KeepWorldPosition);
+		}
+	}
 }
 
 void AInstancedFoliageActor::PostLoad()
@@ -1675,13 +1725,7 @@ void HInstancedStaticMeshInstance::AddReferencedObjects(FReferenceCollector& Col
 
 void AInstancedFoliageActor::ApplyWorldOffset(const FVector& InOffset, bool bWorldShift)
 {
-	TArray<UActorComponent*> Components;
-	GetComponents(Components);
-
-	for (UActorComponent* Component : Components)
-	{
-		Component->ApplyWorldOffset(InOffset, bWorldShift);
-	}
+	Super::ApplyWorldOffset(InOffset, bWorldShift);
 
 	if (GIsEditor)
 	{
diff --git a/Engine/Source/Runtime/Engine/Private/LevelUtils.cpp b/Engine/Source/Runtime/Engine/Private/LevelUtils.cpp
--- a/Engine/Source/Runtime/Engine/Private/LevelUtils.cpp
+++ b/Engine/Source/Runtime/Engine/Private/LevelUtils.cpp
@@ -4,6 +4,7 @@
 #include "EnginePrivate.h"
 #include "LevelUtils.h"
 #include "EditorSupportDelegates.h"
+#include "Foliage/InstancedFoliageActor.h"
 
 #if WITH_EDITOR
 #include "ScopedTransaction.h"
@@ -377,13 +378,21 @@ void FLevelUtils::ApplyLevelTransform( ULevel* Level, const FTransform& LevelTra
 				{
 					RootComponent->SetRelativeLocationAndRotation( LevelTransform.TransformPosition(RootComponent->RelativeLocation), (FTransform(RootComponent->RelativeRotation) * LevelTransform).Rotator());
 				}			
-			}	
-		}		
+			}
+		}
+
 #if WITH_EDITOR
 		if( bDoPostEditMove )
 		{
 			ApplyPostEditMove( Level );						
-		}	
+		}
+
+		// We need to transform foliage Editor data which is stored in world space
+		AInstancedFoliageActor* FoliageActor = AInstancedFoliageActor::GetInstancedFoliageActorForLevel(Level);
+		if (FoliageActor)
+		{
+			FoliageActor->ApplyLevelTransform(LevelTransform);
+		}
 #endif // WITH_EDITOR
 	}
 }
