Preventing full data from being cooked in streaming sounds

Added platform feature/property for audio streaming so that it's not used on platforms that don't support it yet
Prevent full data from being cooked if sound is streaming AND streaming is supported on the cooked platform
Prevent streaming chunks from being cooked if cooked platform doesn't support streaming
Prevent the Resource data from being initialized and full data size being added to stats if sound is streaming

diff --git a/Engine/Source/Developer/DeviceManager/Private/Widgets/Details/SDeviceDetails.cpp b/Engine/Source/Developer/DeviceManager/Private/Widgets/Details/SDeviceDetails.cpp
--- a/Engine/Source/Developer/DeviceManager/Private/Widgets/Details/SDeviceDetails.cpp
+++ b/Engine/Source/Developer/DeviceManager/Private/Widgets/Details/SDeviceDetails.cpp
@@ -129,6 +129,7 @@ void SDeviceDetails::HandleModelSelectedDeviceServiceChanged( )
 			const ITargetPlatform& TargetPlatform = TargetDevice->GetTargetPlatform();
 
 			// platform features
+			FeatureList.Add(MakeShareable(new FDeviceDetailsFeature(TEXT("AudioStreaming"), TargetPlatform.SupportsFeature(ETargetPlatformFeatures::AudioStreaming))));
 			FeatureList.Add(MakeShareable(new FDeviceDetailsFeature(TEXT("DistanceFieldShadows"), TargetPlatform.SupportsFeature(ETargetPlatformFeatures::DistanceFieldShadows))));
 			FeatureList.Add(MakeShareable(new FDeviceDetailsFeature(TEXT("GrayscaleSRGB"), TargetPlatform.SupportsFeature(ETargetPlatformFeatures::GrayscaleSRGB))));
 			FeatureList.Add(MakeShareable(new FDeviceDetailsFeature(TEXT("HighQualityLightmaps"), TargetPlatform.SupportsFeature(ETargetPlatformFeatures::HighQualityLightmaps))));
diff --git a/Engine/Source/Developer/TargetPlatform/Public/Common/TargetPlatformBase.h b/Engine/Source/Developer/TargetPlatform/Public/Common/TargetPlatformBase.h
--- a/Engine/Source/Developer/TargetPlatform/Public/Common/TargetPlatformBase.h
+++ b/Engine/Source/Developer/TargetPlatform/Public/Common/TargetPlatformBase.h
@@ -254,6 +254,9 @@ public:
 	{
 		switch (Feature)
 		{
+		case ETargetPlatformFeatures::AudioStreaming:
+			return TPlatformProperties::SupportsAudioStreaming();
+
 		case ETargetPlatformFeatures::DistanceFieldShadows:
 			return TPlatformProperties::SupportsDistanceFieldShadows();
 
diff --git a/Engine/Source/Developer/TargetPlatform/Public/Interfaces/ITargetPlatform.h b/Engine/Source/Developer/TargetPlatform/Public/Interfaces/ITargetPlatform.h
--- a/Engine/Source/Developer/TargetPlatform/Public/Interfaces/ITargetPlatform.h
+++ b/Engine/Source/Developer/TargetPlatform/Public/Interfaces/ITargetPlatform.h
@@ -47,6 +47,9 @@ namespace ETargetPlatformFeatures
 	 */
 	enum Type
 	{
+		/** Audio Streaming */
+		AudioStreaming,
+
 		/** Distance field shadows. */
 		DistanceFieldShadows,
 
diff --git a/Engine/Source/Runtime/Core/Public/GenericPlatform/GenericPlatformProperties.h b/Engine/Source/Runtime/Core/Public/GenericPlatform/GenericPlatformProperties.h
--- a/Engine/Source/Runtime/Core/Public/GenericPlatform/GenericPlatformProperties.h
+++ b/Engine/Source/Runtime/Core/Public/GenericPlatform/GenericPlatformProperties.h
@@ -200,6 +200,14 @@ struct FGenericPlatformProperties
 		return true;
 	}
 
+	/**
+	 * Whether this platform supports streaming audio
+	 */
+	static FORCEINLINE bool SupportsAudioStreaming()
+	{
+		return false;
+	}
+
 	static FORCEINLINE bool SupportsHighQualityLightmaps()
 	{
 		return true;
diff --git a/Engine/Source/Runtime/Core/Public/Windows/WindowsPlatformProperties.h b/Engine/Source/Runtime/Core/Public/Windows/WindowsPlatformProperties.h
--- a/Engine/Source/Runtime/Core/Public/Windows/WindowsPlatformProperties.h
+++ b/Engine/Source/Runtime/Core/Public/Windows/WindowsPlatformProperties.h
@@ -69,6 +69,11 @@ struct FWindowsPlatformProperties
 		return !HAS_EDITOR_DATA;
 	}
 
+	static FORCEINLINE bool SupportsAudioStreaming()
+	{
+		return true;
+	}
+
 	static FORCEINLINE bool SupportsGrayscaleSRGB()
 	{
 		return false; // Requires expand from G8 to RGBA
diff --git a/Engine/Source/Runtime/Engine/Private/AudioDevice.cpp b/Engine/Source/Runtime/Engine/Private/AudioDevice.cpp
--- a/Engine/Source/Runtime/Engine/Private/AudioDevice.cpp
+++ b/Engine/Source/Runtime/Engine/Private/AudioDevice.cpp
@@ -2291,7 +2291,7 @@ void FAudioDevice::Precache(USoundWave* SoundWave, bool bSynchronous, bool bTrac
 		const FSoundGroup& SoundGroup = GetDefault<USoundGroups>()->GetSoundGroup(SoundWave->SoundGroup);
 
 		// handle audio decompression
-		if (SupportsStreaming() && SoundWave->IsStreaming())
+		if (FPlatformProperties::SupportsAudioStreaming() && SoundWave->IsStreaming())
 		{
 			SoundWave->DecompressionType = DTYPE_Streaming;
 		}
diff --git a/Engine/Source/Runtime/Engine/Private/SoundWave.cpp b/Engine/Source/Runtime/Engine/Private/SoundWave.cpp
--- a/Engine/Source/Runtime/Engine/Private/SoundWave.cpp
+++ b/Engine/Source/Runtime/Engine/Private/SoundWave.cpp
@@ -71,7 +71,11 @@ SIZE_T USoundWave::GetResourceSize(EResourceSizeMode::Type Mode)
 
 	if (GEngine && GEngine->GetAudioDevice())
 	{
-		CalculatedResourceSize += GetCompressedDataSize(GEngine->GetAudioDevice()->GetRuntimeFormat(this));
+		// Don't add compressed data to size of streaming sounds
+		if (!FPlatformProperties::SupportsAudioStreaming() || !IsStreaming())
+		{
+			CalculatedResourceSize += GetCompressedDataSize(GEngine->GetAudioDevice()->GetRuntimeFormat(this));
+		}
 	}
 
 	return CalculatedResourceSize;
@@ -143,24 +147,30 @@ void USoundWave::Serialize( FArchive& Ar )
 
 	if (bCooked)
 	{
-		if (Ar.IsCooking())
+		// Only want to cook/load full data if we don't support streaming
+		if (!IsStreaming() ||
+			(Ar.IsLoading() && !FPlatformProperties::SupportsAudioStreaming()) ||
+			(Ar.IsCooking() && !Ar.CookingTarget()->SupportsFeature(ETargetPlatformFeatures::AudioStreaming)))
 		{
+			if (Ar.IsCooking())
+			{
 #if WITH_ENGINE
-			TArray<FName> ActualFormatsToSave;
-			if (!Ar.CookingTarget()->IsServerOnly())
+				TArray<FName> ActualFormatsToSave;
+				if (!Ar.CookingTarget()->IsServerOnly())
+				{
+					// for now we only support one format per wav
+					FName Format = Ar.CookingTarget()->GetWaveFormat(this);
+					GetCompressedData(Format); // Get the data from the DDC or build it
+
+					ActualFormatsToSave.Add(Format);
+				}
+				CompressedFormatData.Serialize(Ar, this, &ActualFormatsToSave);
+#endif
+			}
+			else
 			{
-				// for now we only support one format per wav
-				FName Format = Ar.CookingTarget()->GetWaveFormat(this);
-				GetCompressedData(Format); // Get the data from the DDC or build it
-
-				ActualFormatsToSave.Add(Format);
+				CompressedFormatData.Serialize(Ar, this);
 			}
-			CompressedFormatData.Serialize(Ar, this, &ActualFormatsToSave);
-#endif
-		}
-		else
-		{
-			CompressedFormatData.Serialize(Ar, this);
 		}
 	}
 	else
@@ -175,7 +185,12 @@ void USoundWave::Serialize( FArchive& Ar )
 	{
 		if (bCooked)
 		{
-			SerializeCookedPlatformData(Ar);
+			// only cook/load streaming data if it's supported
+			if ((Ar.IsLoading() && FPlatformProperties::SupportsAudioStreaming()) ||
+				(Ar.IsCooking() && Ar.CookingTarget()->SupportsFeature(ETargetPlatformFeatures::AudioStreaming)))
+			{
+				SerializeCookedPlatformData(Ar);
+			}
 		}
 
 #if WITH_EDITORONLY_DATA	
@@ -330,7 +345,7 @@ void USoundWave::InitAudioResource( FByteBulkData& CompressedData )
 
 bool USoundWave::InitAudioResource(FName Format)
 {
-	if( !ResourceSize )
+	if( !ResourceSize && (!FPlatformProperties::SupportsAudioStreaming() || !IsStreaming()) )
 	{
 		FByteBulkData* Bulk = GetCompressedData(Format);
 		if (Bulk)
diff --git a/Engine/Source/Runtime/Engine/Public/AudioDevice.h b/Engine/Source/Runtime/Engine/Public/AudioDevice.h
--- a/Engine/Source/Runtime/Engine/Public/AudioDevice.h
+++ b/Engine/Source/Runtime/Engine/Public/AudioDevice.h
@@ -487,12 +487,6 @@ public:
 		return false;	// assume no support by default
 	}
 
-	/** Whether this device supports streaming of sound waves (i.e. DTYPE_Streaming) */
-	virtual bool SupportsStreaming() const
-	{
-		return false;	// assume no support by default
-	}
-
 	/** Creates a Compressed audio info class suitable for decompressing this SoundWave */
 	virtual class ICompressedAudioInfo* CreateCompressedAudioInfo(USoundWave* SoundWave) { return NULL; }
 
diff --git a/Engine/Source/Runtime/Windows/XAudio2/Public/XAudio2Device.h b/Engine/Source/Runtime/Windows/XAudio2/Public/XAudio2Device.h
--- a/Engine/Source/Runtime/Windows/XAudio2/Public/XAudio2Device.h
+++ b/Engine/Source/Runtime/Windows/XAudio2/Public/XAudio2Device.h
@@ -116,11 +116,6 @@ class FXAudio2Device : public FAudioDevice
 		return true;
 	}
 
-	virtual bool SupportsStreaming() const override
-	{
-		return true;
-	}
-
 	virtual class ICompressedAudioInfo* CreateCompressedAudioInfo(USoundWave* SoundWave) override;
 
 	/** 
