1st pass of improved martial workflow with particles.
Added the ability to use a set of named materials for emitters. These can be overriden extrenally with MIDs from code or BPs.
Also added some error reporting to help stop confusion about where an emitter is getting it's material from.

#ttp 335126 - UE4: RENDERING: Improve MID -> Blueprint workflow

diff --git a/Engine/Source/Runtime/Engine/Classes/Particles/Material/ParticleModuleMeshMaterial.h b/Engine/Source/Runtime/Engine/Classes/Particles/Material/ParticleModuleMeshMaterial.h
--- a/Engine/Source/Runtime/Engine/Classes/Particles/Material/ParticleModuleMeshMaterial.h
+++ b/Engine/Source/Runtime/Engine/Classes/Particles/Material/ParticleModuleMeshMaterial.h
@@ -16,8 +16,12 @@ class UParticleModuleMeshMaterial : public UParticleModuleMaterialBase
 
 
 	//Begin UParticleModule Interface
-	virtual void	Spawn(FParticleEmitterInstance* Owner, int32 Offset, float SpawnTime, FBaseParticle* ParticleBase) override;
-	virtual uint32	RequiredBytesPerInstance(FParticleEmitterInstance* Owner = NULL) override;
+	virtual void	Spawn(FParticleEmitterInstance* Owner, int32 Offset, float SpawnTime, FBaseParticle* ParticleBase) OVERRIDE;
+	virtual uint32	RequiredBytesPerInstance(FParticleEmitterInstance* Owner = NULL) OVERRIDE;
+
+#if WITH_EDITOR
+	virtual bool	IsValidForLODLevel(UParticleLODLevel* LODLevel, FString& OutErrorString) OVERRIDE;
+#endif
 	//End UParticleModule Interface
 };
 
diff --git a/Engine/Source/Runtime/Engine/Classes/Particles/ParticleModuleRequired.h b/Engine/Source/Runtime/Engine/Classes/Particles/ParticleModuleRequired.h
--- a/Engine/Source/Runtime/Engine/Classes/Particles/ParticleModuleRequired.h
+++ b/Engine/Source/Runtime/Engine/Classes/Particles/ParticleModuleRequired.h
@@ -277,6 +277,13 @@ class UParticleModuleRequired : public UParticleModule
 	UPROPERTY(EditAnywhere, Category=Emitter)
 	uint32 bOrbitModuleAffectsVelocityAlignment:1;
 
+	/** 
+	*	Named material overrides for this emitter. 
+	*	Overrides this emitter's material(s) with those in the correspondingly named slot(s) of the owning system.
+	*/
+	UPROPERTY(EditAnywhere, Category = Materials)
+	TArray<FName> NamedMaterialOverrides;
+
 	/** Initializes the default values for this property */
 	void InitializeDefaults();
 
diff --git a/Engine/Source/Runtime/Engine/Classes/Particles/ParticleSystem.h b/Engine/Source/Runtime/Engine/Classes/Particles/ParticleSystem.h
--- a/Engine/Source/Runtime/Engine/Classes/Particles/ParticleSystem.h
+++ b/Engine/Source/Runtime/Engine/Classes/Particles/ParticleSystem.h
@@ -74,6 +74,24 @@ struct FLODSoloTrack
 
 };
 
+USTRUCT()
+struct FNamedEmitterMaterial
+{
+	GENERATED_USTRUCT_BODY()
+
+	FNamedEmitterMaterial()
+	: Name(NAME_None)
+	, Material(nullptr)
+	{
+	}
+
+	UPROPERTY(EditAnywhere, Category = NamedMaterial)
+		FName Name;
+
+	UPROPERTY(EditAnywhere, Category = NamedMaterial)
+	class UMaterialInterface* Material;
+};
+
 /**
  * A ParticleSystem is a complete particle effect that contains any number of ParticleEmitters. By allowing multiple emitters
  * in a system, the designer can create elaborate particle effects that are held in a single system. Once created using
@@ -297,6 +315,13 @@ public:
 	UPROPERTY(transient)
 	TArray<struct FLODSoloTrack> SoloTracking;
 
+	/** 
+	*	Array of named material slots for use by emitters of this system. 
+	*	Emitters can use these instead of their own materials by providing the name to the NamedMaterialOverrides property of their required module.
+	*	These materials can be overriden using CreateNamedDynamicMaterialInstance() on a ParticleSystemComponent.
+	*/
+	UPROPERTY(EditAnywhere, Category = Materials)
+	TArray<FNamedEmitterMaterial> NamedMaterialSlots;
 
 	// Begin UObject interface.
 #if WITH_EDITOR
diff --git a/Engine/Source/Runtime/Engine/Classes/Particles/ParticleSystemComponent.h b/Engine/Source/Runtime/Engine/Classes/Particles/ParticleSystemComponent.h
--- a/Engine/Source/Runtime/Engine/Classes/Particles/ParticleSystemComponent.h
+++ b/Engine/Source/Runtime/Engine/Classes/Particles/ParticleSystemComponent.h
@@ -786,6 +786,20 @@ public:
 	/** Decide if this psys can tick in any thread, and set bIsElligibleForAsyncTick */
 	void ComputeCanTickInAnyThread();
 
+	/**
+	* Creates a Dynamic Material Instance for the specified named material override, optionally from the supplied material.
+	* @param Name - The slot name of the material to replace.  If invalid, the material is unchanged and NULL is returned.
+	*/
+	UFUNCTION(BlueprintCallable, Category = "Rendering|Material")
+	virtual class UMaterialInstanceDynamic* CreateNamedDynamicMaterialInstance(FName Name, class UMaterialInterface* SourceMaterial = NULL);
+
+	/** Returns a named material. If this named material is not found, returns NULL. */
+	UFUNCTION(BlueprintCallable, Category = "Rendering|Material")
+	virtual class UMaterialInterface* GetNamedMaterial(FName Name) const;
+
+	/** Returns the index into the EmitterMaterials array for this named. If there are no named material slots or this material is not found, INDEX_NONE is returned. */
+	virtual int32 GetNamedMaterialIndex(FName Name) const;
+
 protected:
 
 	// @todo document
diff --git a/Engine/Source/Runtime/Engine/Private/Particles/ParticleComponents.cpp b/Engine/Source/Runtime/Engine/Private/Particles/ParticleComponents.cpp
--- a/Engine/Source/Runtime/Engine/Private/Particles/ParticleComponents.cpp
+++ b/Engine/Source/Runtime/Engine/Private/Particles/ParticleComponents.cpp
@@ -3969,11 +3969,14 @@ void UParticleSystemComponent::ComputeTickComponent_Concurrent()
 			{
 				Instance->Tick(DeltaTimeTick, bSuppressSpawning);
 
-				if (EmitterMaterials.IsValidIndex(EmitterIndex))
+				if (!Instance->Tick_MaterialOverrides())
 				{
-					if (EmitterMaterials[EmitterIndex])
+					if (EmitterMaterials.IsValidIndex(EmitterIndex))
 					{
-						Instance->CurrentMaterial = EmitterMaterials[EmitterIndex];
+						if (EmitterMaterials[EmitterIndex])
+						{
+							Instance->CurrentMaterial = EmitterMaterials[EmitterIndex];
+						}
 					}
 				}
 				TotalActiveParticles += Instance->ActiveParticles;
@@ -5845,6 +5848,70 @@ UParticleSystemReplay* UParticleSystemComponent::FindReplayClipForIDNumber( cons
 	return NULL;
 }
 
+UMaterialInstanceDynamic* UParticleSystemComponent::CreateNamedDynamicMaterialInstance(FName Name, class UMaterialInterface* SourceMaterial)
+{
+	int32 Index = GetNamedMaterialIndex(Name);
+	if (INDEX_NONE == Index)
+	{
+		UE_LOG(LogParticles, Warning, TEXT("CreateNamedDynamicMaterialInstance on %s: This material wasn't found. Check the particle system's named material slots in cascade."), *GetPathName(), *Name.ToString());
+		return NULL;
+	}
+
+	if (SourceMaterial)
+	{
+		SetMaterial(Index, SourceMaterial);
+	}
+
+	UMaterialInterface* MaterialInstance = GetMaterial(Index);
+	UMaterialInstanceDynamic* MID = Cast<UMaterialInstanceDynamic>(MaterialInstance);
+
+	if (MaterialInstance && !MID)
+	{
+		// Create and set the dynamic material instance.
+		MID = UMaterialInstanceDynamic::Create(MaterialInstance, this);
+		SetMaterial(Index, MID);
+	}
+	else if (!MaterialInstance)
+	{
+		UE_LOG(LogParticles, Warning, TEXT("CreateDynamicMaterialInstance on %s: Material index %d is invalid."), *GetPathName(), Index);
+	}
+
+	return MID;
+}
+
+UMaterialInterface* UParticleSystemComponent::GetNamedMaterial(FName Name) const
+{
+	int32 Index = GetNamedMaterialIndex(Name);
+	if (INDEX_NONE != Index)
+	{
+		if (EmitterMaterials.IsValidIndex(Index) && nullptr == EmitterMaterials[Index])
+		{
+			//Material has been overridden externally
+			return EmitterMaterials[Index];
+		}
+		else
+		{
+			//This slot hasn't been overridden so just used the default.
+			return Template->NamedMaterialSlots[Index].Material;
+		}
+	}
+	//Could not find this named materials slot.
+	return NULL;
+}
+
+int32 UParticleSystemComponent::GetNamedMaterialIndex(FName Name) const
+{
+	TArray<FNamedEmitterMaterial>& Slots = Template->NamedMaterialSlots;
+	for (int32 SlotIdx = 0; SlotIdx < Slots.Num(); ++SlotIdx)
+	{
+		if (Name == Slots[SlotIdx].Name)
+		{
+			return SlotIdx;
+		}
+	}
+	return INDEX_NONE;
+}
+
 UParticleSystemReplay::UParticleSystemReplay(const class FPostConstructInitializeProperties& PCIP)
 	: Super(PCIP)
 {
diff --git a/Engine/Source/Runtime/Engine/Private/Particles/ParticleEmitterInstances.cpp b/Engine/Source/Runtime/Engine/Private/Particles/ParticleEmitterInstances.cpp
--- a/Engine/Source/Runtime/Engine/Private/Particles/ParticleEmitterInstances.cpp
+++ b/Engine/Source/Runtime/Engine/Private/Particles/ParticleEmitterInstances.cpp
@@ -2539,6 +2539,39 @@ void FParticleEmitterInstance::ApplyWorldOffset(FVector InOffset, bool bWorldShi
 	}
 }
 
+bool FParticleEmitterInstance::Tick_MaterialOverrides()
+{
+	UParticleLODLevel* LODLevel = SpriteTemplate->GetCurrentLODLevel(this);
+	bool bOverridden = false;
+	if( LODLevel && LODLevel->RequiredModule && Component && Component->Template )
+	{
+	        TArray<FName>& NamedOverrides = LODLevel->RequiredModule->NamedMaterialOverrides;
+	        TArray<FNamedEmitterMaterial>& Slots = Component->Template->NamedMaterialSlots;
+	        TArray<UMaterialInterface*>& EmitterMaterials = Component->EmitterMaterials;
+	        if (NamedOverrides.Num() > 0)
+	        {
+		        //If we have named material overrides then get it's index into the emitter materials array.
+		        //Only check for [0] in in the named overrides as most emitter types only need one material. Mesh emitters might use more but they override this function.		
+		        for (int32 CheckIdx = 0; CheckIdx < Slots.Num(); ++CheckIdx)
+		        {
+			        if (NamedOverrides[0] == Slots[CheckIdx].Name)
+			        {
+				        //Default to the default material for that slot.
+				        CurrentMaterial = Slots[CheckIdx].Material;
+				        if (EmitterMaterials.IsValidIndex(CheckIdx) && nullptr != EmitterMaterials[CheckIdx])
+				        {
+					        //This material has been overridden externally, e.g. from a BP so use that one.
+					        CurrentMaterial = EmitterMaterials[CheckIdx];
+				        }
+        
+				        bOverridden = true;
+				        break;
+			        }
+		        }
+	        }
+	}
+	return bOverridden;
+}
 
 /*-----------------------------------------------------------------------------
 	ParticleSpriteEmitterInstance
@@ -2930,6 +2963,43 @@ void FParticleMeshEmitterInstance::Tick(float DeltaTime, bool bSuppressSpawning)
 	INC_DWORD_STAT_BY(STAT_MeshParticles, ActiveParticles);
 }
 
+bool FParticleMeshEmitterInstance::Tick_MaterialOverrides()
+{
+	UParticleLODLevel* LODLevel = SpriteTemplate->GetCurrentLODLevel(this);
+	bool bOverridden = false;
+	if( LODLevel && LODLevel->RequiredModule && Component && Component->Template )
+	{
+	        TArray<FName>& NamedOverrides = LODLevel->RequiredModule->NamedMaterialOverrides;
+	        TArray<FNamedEmitterMaterial>& Slots = Component->Template->NamedMaterialSlots;
+	        TArray<UMaterialInterface*>& EmitterMaterials = Component->EmitterMaterials;
+	        if (NamedOverrides.Num() > 0)
+	        {
+		        CurrentMaterials.SetNumZeroed(NamedOverrides.Num());
+		        for (int32 MaterialIdx = 0; MaterialIdx < NamedOverrides.Num(); ++MaterialIdx)
+		        {		
+			        //If we have named material overrides then get it's index into the emitter materials array.	
+			        for (int32 CheckIdx = 0; CheckIdx < Slots.Num(); ++CheckIdx)
+			        {
+				        if (NamedOverrides[MaterialIdx] == Slots[CheckIdx].Name)
+				        {
+					        //Default to the default material for that slot.
+					        CurrentMaterials[MaterialIdx] = Slots[CheckIdx].Material;
+					        if (EmitterMaterials.IsValidIndex(CheckIdx) && nullptr != EmitterMaterials[CheckIdx] )
+					        {
+						        //This material has been overridden externally, e.g. from a BP so use that one.
+						        CurrentMaterials[MaterialIdx] = EmitterMaterials[CheckIdx];
+					        }
+        
+					        bOverridden = true;
+					        break;
+				        }
+			        }
+		        }
+	        }
+        }
+	return bOverridden;
+}
+
 /**
  *	Update the bounding box for the emitter
  *
diff --git a/Engine/Source/Runtime/Engine/Private/Particles/ParticleModules.cpp b/Engine/Source/Runtime/Engine/Private/Particles/ParticleModules.cpp
--- a/Engine/Source/Runtime/Engine/Private/Particles/ParticleModules.cpp
+++ b/Engine/Source/Runtime/Engine/Private/Particles/ParticleModules.cpp
@@ -1142,14 +1142,22 @@ bool UParticleModuleRequired::IsValidForLODLevel(UParticleLODLevel* LODLevel, FS
 {
 	bool bValid = true;
 
-	if (LODLevel->TypeDataModule
-		&& LODLevel->TypeDataModule->IsA(UParticleModuleTypeDataGpu::StaticClass()))
+	if (LODLevel && LODLevel->TypeDataModule)
 	{
-		if (InterpolationMethod == PSUVIM_Random
-			|| InterpolationMethod == PSUVIM_Random_Blend)
+		if( LODLevel->TypeDataModule->IsA(UParticleModuleTypeDataGpu::StaticClass()))
 		{
-			OutErrorString = NSLOCTEXT("UnrealEd", "RandomSubUVForGPUEmitter", "Random SubUV interpolation is not supported with GPU particles.").ToString();
-			bValid = false;
+			if (InterpolationMethod == PSUVIM_Random || InterpolationMethod == PSUVIM_Random_Blend)
+			{
+				OutErrorString = NSLOCTEXT("UnrealEd", "RandomSubUVForGPUEmitter", "Random SubUV interpolation is not supported with GPU particles.").ToString();
+				bValid = false;
+			}
+		}
+		else if (UParticleModuleTypeDataMesh* MeshTypeData = Cast<UParticleModuleTypeDataMesh>(LODLevel->TypeDataModule))
+		{
+			if (NamedMaterialOverrides.Num() > 0 && MeshTypeData->bOverrideMaterial)
+			{
+				OutErrorString = NSLOCTEXT("UnrealEd", "NamedMaterialOverriedsAndOverrideMaterialError", "Cannot use Named Material Overrides when using OverrideMaterial in the mesh type data module.").ToString();
+			}
 		}
 	}
 
diff --git a/Engine/Source/Runtime/Engine/Private/Particles/ParticleModules_Material.cpp b/Engine/Source/Runtime/Engine/Private/Particles/ParticleModules_Material.cpp
--- a/Engine/Source/Runtime/Engine/Private/Particles/ParticleModules_Material.cpp
+++ b/Engine/Source/Runtime/Engine/Private/Particles/ParticleModules_Material.cpp
@@ -7,6 +7,9 @@
 #include "EnginePrivate.h"
 #include "ParticleDefinitions.h"
 #include "Particles/Material/ParticleModuleMeshMaterial.h"
+#include "Particles/ParticleLODLevel.h"
+#include "Particles/ParticleModuleRequired.h"
+#include "Particles/TypeData/ParticleModuleTypeDataMesh.h"
 
 UParticleModuleMaterialBase::UParticleModuleMaterialBase(const class FPostConstructInitializeProperties& PCIP)
 	: Super(PCIP)
@@ -45,3 +48,27 @@ uint32 UParticleModuleMeshMaterial::RequiredBytesPerInstance(FParticleEmitterIns
 	}
 	return 0;
 }
+
+#if WITH_EDITOR
+
+bool UParticleModuleMeshMaterial::IsValidForLODLevel(UParticleLODLevel* LODLevel, FString& OutErrorString)
+{
+	if (LODLevel->RequiredModule->NamedMaterialOverrides.Num() > 0)
+	{
+		OutErrorString = NSLOCTEXT("UnrealEd", "MeshMaterialsWithNamedMaterialsError", "Cannot use Mesh Materials Module when using Named Material Overrides in the required module.").ToString();
+		return false;
+	}
+
+	if (UParticleModuleTypeDataMesh* MeshTypeData = Cast<UParticleModuleTypeDataMesh>(LODLevel->TypeDataModule))
+	{
+		if (MeshTypeData->bOverrideMaterial)
+		{
+			OutErrorString = NSLOCTEXT("UnrealEd", "MeshMaterialsWithOverrideMaterialError", "Cannot use Mesh Materials Module when using OverrideMaterial in the mesh type data module.").ToString();
+			return false;
+		}
+	}
+
+	return true;
+}
+
+#endif
\ No newline at end of file
diff --git a/Engine/Source/Runtime/Engine/Public/ParticleEmitterInstances.h b/Engine/Source/Runtime/Engine/Public/ParticleEmitterInstances.h
--- a/Engine/Source/Runtime/Engine/Public/ParticleEmitterInstances.h
+++ b/Engine/Source/Runtime/Engine/Public/ParticleEmitterInstances.h
@@ -665,6 +665,12 @@ public:
 	* @param	InWidth				The width of the trail.
 	*/
 	virtual void SetTrailSourceData(FName InFirstSocketName, FName InSecondSocketName, ETrailWidthMode InWidthMode, float InWidth){}
+
+	/** 
+	Ticks the emitter's material overrides.
+	@return True if there were material overrides. Otherwise revert to default behaviour.
+	*/
+	virtual bool Tick_MaterialOverrides();
 protected:
 
 	/**
@@ -778,6 +784,7 @@ struct ENGINE_API FParticleMeshEmitterInstance : public FParticleEmitterInstance
 	virtual FDynamicEmitterDataBase* GetDynamicData(bool bSelected);
 	virtual bool IsDynamicDataRequired(UParticleLODLevel* CurrentLODLevel) override;
 
+	virtual bool Tick_MaterialOverrides();
 	/**
 	 *	Updates the dynamic data for the instance
 	 *
