Fixing CIS (win32). FTransform alignment issues when passed as fucntion parameters.

diff --git a/Engine/Source/Runtime/Engine/Classes/Kismet/KismetSystemLibrary.h b/Engine/Source/Runtime/Engine/Classes/Kismet/KismetSystemLibrary.h
--- a/Engine/Source/Runtime/Engine/Classes/Kismet/KismetSystemLibrary.h
+++ b/Engine/Source/Runtime/Engine/Classes/Kismet/KismetSystemLibrary.h
@@ -493,10 +493,10 @@ class UKismetSystemLibrary : public UBlueprintFunctionLibrary
 	 * @return						true if there was an overlap that passed the filters, false otherwise.
 	 */
 	UFUNCTION(BlueprintCallable, Category="Collision", meta=(AutoCreateRefTerm="ActorsToIgnore", FriendlyName="ComponentOverlapActors"))
-	static ENGINE_API bool ComponentOverlapActors_NEW(UPrimitiveComponent* Component, const FTransform ComponentTransform, const TArray<TEnumAsByte<EObjectTypeQuery> > & ObjectTypes, UClass* ActorClassFilter, const TArray<AActor*>& ActorsToIgnore, TArray<class AActor*>& OutActors);
+	static ENGINE_API bool ComponentOverlapActors_NEW(UPrimitiveComponent* Component, const FTransform& ComponentTransform, const TArray<TEnumAsByte<EObjectTypeQuery> > & ObjectTypes, UClass* ActorClassFilter, const TArray<AActor*>& ActorsToIgnore, TArray<class AActor*>& OutActors);
 
 	UFUNCTION(BlueprintCallable, Category="Collision", meta=(DeprecatedFunction, DeprecationMessage = "Use new ComponentOverlapActors", AutoCreateRefTerm="ActorsToIgnore"))
-	static ENGINE_API bool ComponentOverlapActors_DEPRECATED(UPrimitiveComponent* Component, const FTransform ComponentTransform, EOverlapFilterOption Filter, UClass* ActorClassFilter, const TArray<AActor*>& ActorsToIgnore, TArray<class AActor*>& OutActors);
+	static ENGINE_API bool ComponentOverlapActors_DEPRECATED(UPrimitiveComponent* Component, const FTransform& ComponentTransform, EOverlapFilterOption Filter, UClass* ActorClassFilter, const TArray<AActor*>& ActorsToIgnore, TArray<class AActor*>& OutActors);
 	
 	/**
 	 * Returns an array of components that overlap the given component.
@@ -509,10 +509,10 @@ class UKismetSystemLibrary : public UBlueprintFunctionLibrary
 	 * @return						true if there was an overlap that passed the filters, false otherwise.
 	 */
 	UFUNCTION(BlueprintCallable, Category="Collision", meta=(AutoCreateRefTerm="ActorsToIgnore", FriendlyName="ComponentOverlapComponents"))
-	static ENGINE_API bool ComponentOverlapComponents_NEW(UPrimitiveComponent* Component, const FTransform ComponentTransform, const TArray<TEnumAsByte<EObjectTypeQuery> > & ObjectTypes, UClass* ComponentClassFilter, const TArray<AActor*>& ActorsToIgnore, TArray<class UPrimitiveComponent*>& OutComponents);
+	static ENGINE_API bool ComponentOverlapComponents_NEW(UPrimitiveComponent* Component, const FTransform& ComponentTransform, const TArray<TEnumAsByte<EObjectTypeQuery> > & ObjectTypes, UClass* ComponentClassFilter, const TArray<AActor*>& ActorsToIgnore, TArray<class UPrimitiveComponent*>& OutComponents);
 
 	UFUNCTION(BlueprintCallable, Category="Collision", meta=(DeprecatedFunction, DeprecationMessage = "Use new ComponentOverlapComponents", AutoCreateRefTerm="ActorsToIgnore"))
-	static ENGINE_API bool ComponentOverlapComponents_DEPRECATED(UPrimitiveComponent* Component, const FTransform ComponentTransform, EOverlapFilterOption Filter, UClass* ComponentClassFilter, const TArray<AActor*>& ActorsToIgnore, TArray<class UPrimitiveComponent*>& OutComponents);
+	static ENGINE_API bool ComponentOverlapComponents_DEPRECATED(UPrimitiveComponent* Component, const FTransform& ComponentTransform, EOverlapFilterOption Filter, UClass* ComponentClassFilter, const TArray<AActor*>& ActorsToIgnore, TArray<class UPrimitiveComponent*>& OutComponents);
 
 	/**
 	 * Does a collision trace along the given line and returns the first blocking hit encountered.
diff --git a/Engine/Source/Runtime/Engine/Private/KismetSystemLibrary.cpp b/Engine/Source/Runtime/Engine/Private/KismetSystemLibrary.cpp
--- a/Engine/Source/Runtime/Engine/Private/KismetSystemLibrary.cpp
+++ b/Engine/Source/Runtime/Engine/Private/KismetSystemLibrary.cpp
@@ -704,7 +704,7 @@ bool UKismetSystemLibrary::CapsuleOverlapComponents_DEPRECATED(UObject* WorldCon
 }
 
 
-bool UKismetSystemLibrary::ComponentOverlapActors_DEPRECATED(UPrimitiveComponent* Component, const FTransform ComponentTransform, EOverlapFilterOption Filter, UClass* ActorClassFilter, const TArray<AActor*>& ActorsToIgnore, TArray<AActor*>& OutActors)
+bool UKismetSystemLibrary::ComponentOverlapActors_DEPRECATED(UPrimitiveComponent* Component, const FTransform& ComponentTransform, EOverlapFilterOption Filter, UClass* ActorClassFilter, const TArray<AActor*>& ActorsToIgnore, TArray<AActor*>& OutActors)
 {
 	OutActors.Empty();
 
@@ -718,7 +718,7 @@ bool UKismetSystemLibrary::ComponentOverlapActors_DEPRECATED(UPrimitiveComponent
 	return (OutActors.Num() > 0);
 }
 
-bool UKismetSystemLibrary::ComponentOverlapComponents_DEPRECATED(UPrimitiveComponent* Component, const FTransform ComponentTransform, EOverlapFilterOption Filter, UClass* ComponentClassFilter, const TArray<AActor*>& ActorsToIgnore, TArray<UPrimitiveComponent*>& OutComponents)
+bool UKismetSystemLibrary::ComponentOverlapComponents_DEPRECATED(UPrimitiveComponent* Component, const FTransform& ComponentTransform, EOverlapFilterOption Filter, UClass* ComponentClassFilter, const TArray<AActor*>& ActorsToIgnore, TArray<UPrimitiveComponent*>& OutComponents)
 {
 	OutComponents.Empty();
 
@@ -899,7 +899,7 @@ bool UKismetSystemLibrary::CapsuleOverlapComponents_NEW(UObject* WorldContextObj
 }
 
 
-bool UKismetSystemLibrary::ComponentOverlapActors_NEW(UPrimitiveComponent* Component, const FTransform ComponentTransform, const TArray<TEnumAsByte<EObjectTypeQuery> > & ObjectTypes, UClass* ActorClassFilter, const TArray<AActor*>& ActorsToIgnore, TArray<AActor*>& OutActors)
+bool UKismetSystemLibrary::ComponentOverlapActors_NEW(UPrimitiveComponent* Component, const FTransform& ComponentTransform, const TArray<TEnumAsByte<EObjectTypeQuery> > & ObjectTypes, UClass* ActorClassFilter, const TArray<AActor*>& ActorsToIgnore, TArray<AActor*>& OutActors)
 {
 	OutActors.Empty();
 
@@ -913,7 +913,7 @@ bool UKismetSystemLibrary::ComponentOverlapActors_NEW(UPrimitiveComponent* Compo
 	return (OutActors.Num() > 0);
 }
 
-bool UKismetSystemLibrary::ComponentOverlapComponents_NEW(UPrimitiveComponent* Component, const FTransform ComponentTransform, const TArray<TEnumAsByte<EObjectTypeQuery> > & ObjectTypes, UClass* ComponentClassFilter, const TArray<AActor*>& ActorsToIgnore, TArray<UPrimitiveComponent*>& OutComponents)
+bool UKismetSystemLibrary::ComponentOverlapComponents_NEW(UPrimitiveComponent* Component, const FTransform& ComponentTransform, const TArray<TEnumAsByte<EObjectTypeQuery> > & ObjectTypes, UClass* ComponentClassFilter, const TArray<AActor*>& ActorsToIgnore, TArray<UPrimitiveComponent*>& OutComponents)
 {
 	OutComponents.Empty();
 
