Build Warnings - Fixing some usages of OVERRIDE instead of override.

diff --git a/Engine/Source/Runtime/Engine/Classes/Particles/Material/ParticleModuleMeshMaterial.h b/Engine/Source/Runtime/Engine/Classes/Particles/Material/ParticleModuleMeshMaterial.h
--- a/Engine/Source/Runtime/Engine/Classes/Particles/Material/ParticleModuleMeshMaterial.h
+++ b/Engine/Source/Runtime/Engine/Classes/Particles/Material/ParticleModuleMeshMaterial.h
@@ -1,7 +1,7 @@
 // Copyright 1998-2014 Epic Games, Inc. All Rights Reserved.
 
-
 #pragma once
+
 #include "Particles/Material/ParticleModuleMaterialBase.h"
 #include "ParticleModuleMeshMaterial.generated.h"
 
@@ -16,14 +16,11 @@ class UParticleModuleMeshMaterial : public UParticleModuleMaterialBase
 
 
 	//Begin UParticleModule Interface
-	virtual void	Spawn(FParticleEmitterInstance* Owner, int32 Offset, float SpawnTime, FBaseParticle* ParticleBase) OVERRIDE;
-	virtual uint32	RequiredBytesPerInstance(FParticleEmitterInstance* Owner = NULL) OVERRIDE;
+	virtual void	Spawn(FParticleEmitterInstance* Owner, int32 Offset, float SpawnTime, FBaseParticle* ParticleBase) override;
+	virtual uint32	RequiredBytesPerInstance(FParticleEmitterInstance* Owner = NULL) override;
 
 #if WITH_EDITOR
-	virtual bool	IsValidForLODLevel(UParticleLODLevel* LODLevel, FString& OutErrorString) OVERRIDE;
+	virtual bool	IsValidForLODLevel(UParticleLODLevel* LODLevel, FString& OutErrorString) override;
 #endif
 	//End UParticleModule Interface
 };
-
-
-
