Added analytics for scalability settings

diff --git a/Engine/Source/Editor/UnrealEd/Private/PerformanceMonitor.cpp b/Engine/Source/Editor/UnrealEd/Private/PerformanceMonitor.cpp
--- a/Engine/Source/Editor/UnrealEd/Private/PerformanceMonitor.cpp
+++ b/Engine/Source/Editor/UnrealEd/Private/PerformanceMonitor.cpp
@@ -166,6 +166,9 @@ void FPerformanceMonitor::AutoApplyScalability()
 	Scalability::SetQualityLevels(BenchmarkLevels);
 	Scalability::SaveState(GEditorUserSettingsIni);
 
+	const bool bAutoApplied = true;
+	Scalability::RecordQualityLevelsAnalytics(bAutoApplied);
+
 	GEditor->DisableRealtimeViewports();
 
 	HidePerformanceWarning();
@@ -434,4 +437,4 @@ void FPerformanceMonitor::ShowScalabilityDialog()
 	}
 }
 
-#undef LOCTEXT_NAMESPACE
\ No newline at end of file
+#undef LOCTEXT_NAMESPACE
diff --git a/Engine/Source/Editor/UnrealEd/Private/SScalabilitySettings.cpp b/Engine/Source/Editor/UnrealEd/Private/SScalabilitySettings.cpp
--- a/Engine/Source/Editor/UnrealEd/Private/SScalabilitySettings.cpp
+++ b/Engine/Source/Editor/UnrealEd/Private/SScalabilitySettings.cpp
@@ -131,7 +131,7 @@ void SScalabilitySettings::Construct( const FArguments& InArgs )
 	auto TitleFont = FEditorStyle::GetFontStyle( FName( "Scalability.TitleFont" ) );
 	auto GroupFont = FEditorStyle::GetFontStyle( FName( "Scalability.GroupFont" ) );
 
-	CachedQualityLevels = Scalability::GetQualityLevels();
+	InitialQualityLevels = CachedQualityLevels = Scalability::GetQualityLevels();
 	static float Padding = 1.0f;
 	static float QualityColumnCoeff = 1.0f;
 	static float QualityLevelColumnCoeff = 0.2f;
@@ -203,4 +203,14 @@ void SScalabilitySettings::Construct( const FArguments& InArgs )
 		];
 }
 
-#undef LOCTEXT_NAMESPACE
\ No newline at end of file
+SScalabilitySettings::~SScalabilitySettings()
+{
+	// Record any changed quality levels
+	if( InitialQualityLevels != CachedQualityLevels )
+	{
+		const bool bAutoApplied = false;
+		Scalability::RecordQualityLevelsAnalytics(bAutoApplied);
+	}
+}
+
+#undef LOCTEXT_NAMESPACE
diff --git a/Engine/Source/Editor/UnrealEd/Public/SScalabilitySettings.h b/Engine/Source/Editor/UnrealEd/Public/SScalabilitySettings.h
--- a/Engine/Source/Editor/UnrealEd/Public/SScalabilitySettings.h
+++ b/Engine/Source/Editor/UnrealEd/Public/SScalabilitySettings.h
@@ -20,11 +20,9 @@ public:
 	// Widget construction
 	UNREALED_API void Construct( const FArguments& InArgs );
 
+	UNREALED_API ~SScalabilitySettings();
 private:
 
-	// Stores state of scalability settings at the point of opening the menu
-	Scalability::FQualityLevels CachedQualityLevels;
-
 	// Checks cached quality levels to see if the specified group is at the specified quality level
 	ESlateCheckBoxState::Type IsGroupQualityLevelSelected(const TCHAR* InGroupName, int32 InQualityLevel) const;
 
@@ -57,4 +55,12 @@ private:
 
 	// Create a gridslot for the group quality level with all the required formatting
 	SGridPanel::FSlot& MakeGridSlot(int32 InCol, int32 InRow, int32 InColSpan = 1, int32 InRowSpan = 1);
+
+private:
+	/* The state of scalability settings at the point of opening the menu*/
+	Scalability::FQualityLevels InitialQualityLevels;
+
+	/** The state of quality levels as they are changed in this widget */
+	Scalability::FQualityLevels CachedQualityLevels;
+
 };
\ No newline at end of file
diff --git a/Engine/Source/Runtime/Engine/Private/Scalability.cpp b/Engine/Source/Runtime/Engine/Private/Scalability.cpp
--- a/Engine/Source/Runtime/Engine/Private/Scalability.cpp
+++ b/Engine/Source/Runtime/Engine/Private/Scalability.cpp
@@ -3,6 +3,8 @@
 #include "EnginePrivate.h"
 #include "Scalability.h"
 #include "SynthBenchmark.h"
+#include "EngineAnalytics.h"
+#include "Runtime/Analytics/Analytics/Public/Interfaces/IAnalyticsProvider.h"
 
 static TAutoConsoleVariable<int32> CVarResolutionQuality(
 	TEXT("sg.ResolutionQuality"),
@@ -432,6 +434,29 @@ void SaveState(const FString& IniName)
 	GConfig->SetInt(Section, TEXT("sg.EffectsQuality"), State.EffectsQuality, IniName);
 }
 
+void RecordQualityLevelsAnalytics(bool bAutoApplied)
+{
+	if( FEngineAnalytics::IsAvailable() )
+	{
+		FQualityLevels State = GetQualityLevels();
+
+		TArray<FAnalyticsEventAttribute> Attributes;
+
+		Attributes.Add(FAnalyticsEventAttribute(TEXT("ResolutionQuality"), State.ResolutionQuality));
+		Attributes.Add(FAnalyticsEventAttribute(TEXT("ViewDistanceQuality"), State.ViewDistanceQuality));
+		Attributes.Add(FAnalyticsEventAttribute(TEXT("AntiAliasingQuality"), State.AntiAliasingQuality));
+		Attributes.Add(FAnalyticsEventAttribute(TEXT("ShadowQuality"), State.ShadowQuality));
+		Attributes.Add(FAnalyticsEventAttribute(TEXT("PostProcessQuality"), State.PostProcessQuality));
+		Attributes.Add(FAnalyticsEventAttribute(TEXT("TextureQuality"), State.TextureQuality));
+		Attributes.Add(FAnalyticsEventAttribute(TEXT("EffectsQuality"), State.EffectsQuality));
+		Attributes.Add(FAnalyticsEventAttribute(TEXT("AutoAppliedSettings"), bAutoApplied));
+
+		FEngineAnalytics::GetProvider().RecordEvent(TEXT("Editor.Performance.ScalabiltySettings"), Attributes);
+	}
+
+
+}
+
 }
 
 
diff --git a/Engine/Source/Runtime/Engine/Public/Scalability.h b/Engine/Source/Runtime/Engine/Public/Scalability.h
--- a/Engine/Source/Runtime/Engine/Public/Scalability.h
+++ b/Engine/Source/Runtime/Engine/Public/Scalability.h
@@ -27,6 +27,16 @@ namespace Scalability
 			SetDefaults();
 		}
 		
+		bool operator==(const FQualityLevels& Other ) const
+		{
+			return FMemory::Memcmp( this, &Other, sizeof(FQualityLevels) ) == 0;
+		}
+
+		bool operator!=(const FQualityLevels& Other ) const
+		{
+			return !(*this == Other);
+		}
+
 		// @param Value 0:low, 1:medium, 2:high, 3:epic
 		void SetFromSingleQualityLevel(int32 Value);
 
@@ -45,12 +55,19 @@ namespace Scalability
 	/**  */
 	ENGINE_API void InitScalabilitySystem();
 
-	/** @parma IniName e.g. GEditorUserSettingsIni or GGameUserSettingsIni */
+	/** @param IniName e.g. GEditorUserSettingsIni or GGameUserSettingsIni */
 	ENGINE_API void LoadState(const FString& IniName);
 	
-	/** @parma IniName e.g. GEditorUserSettingsIni or GGameUserSettingsIni */
+	/** @param IniName e.g. GEditorUserSettingsIni or GGameUserSettingsIni */
 	ENGINE_API void SaveState(const FString& IniName);
 
+	/**
+	 * Sends an analytic event with all quality level data
+	 *
+	 * @param bAutoApplied	Whether or not the quality levels were auto-applied (true) or applied by the user (false).
+	 */
+	ENGINE_API void RecordQualityLevelsAnalytics(bool bAutoApplied);
+
 	/** Run synthbenchmark and configure scalability based on results **/
 	ENGINE_API FQualityLevels BenchmarkQualityLevels();
 
@@ -62,4 +79,4 @@ namespace Scalability
 
 	/** Maximum single axis scale for render resolution */
 	static const int32 MaxResolutionScale = 100;
-}
\ No newline at end of file
+}
