Fix gpuparticle simulation reading garbage uniform buffers for offscreen emitters

diff --git a/Engine/Source/Runtime/Engine/Private/Particles/ParticleGpuSimulation.cpp b/Engine/Source/Runtime/Engine/Private/Particles/ParticleGpuSimulation.cpp
--- a/Engine/Source/Runtime/Engine/Private/Particles/ParticleGpuSimulation.cpp
+++ b/Engine/Source/Runtime/Engine/Private/Particles/ParticleGpuSimulation.cpp
@@ -1186,7 +1186,7 @@ struct FSimulationCommandGPU
 	/** Uniform buffer containing simulation parameters. */
 	FUniformBufferRHIParamRef UniformBuffer;
 	/** Uniform buffer containing per-frame simulation parameters. */
-	FUniformBufferRHIParamRef PerFrameUniformBuffer;
+	FUniformBufferRHIRef PerFrameUniformBuffer;
 	/** Parameters to sample the local vector field for this simulation. */
 	FVectorFieldUniformBufferRef VectorFieldsUniformBuffer;
 	/** Vector field volume textures for this simulation. */
@@ -2559,7 +2559,7 @@ struct FGPUSpriteDynamicEmitterData : FDynamicEmitterDataBase
 	 * Called to release render thread resources.
 	 */
 	virtual void ReleaseRenderThreadResources(const FParticleSystemSceneProxy* InOwnerProxy)
-	{
+	{		
 	}
 
 	/**
@@ -4147,6 +4147,11 @@ void FFXSystem::SimulateGPUParticles(
 
 			// Reset pending simulation time.
 			Simulation->PendingDeltaSeconds = 0.0f;
+
+			// PerFrameSimulationUniformBuffer was created with UniformBuffer_SingleFrame.  Release it so we use
+			// the default uniform buffer if we simulate again before creating a fresh buffer (i.e. when the Emitter goes offscreen)
+			// this way the simulation shader won't read garbage.
+			Simulation->PerFrameSimulationUniformBuffer.SafeRelease();			
 		}
 	}
 
