Enable engine analytics in source builds

diff --git a/Engine/Source/Runtime/Engine/Private/EngineAnalytics.cpp b/Engine/Source/Runtime/Engine/Private/EngineAnalytics.cpp
--- a/Engine/Source/Runtime/Engine/Private/EngineAnalytics.cpp
+++ b/Engine/Source/Runtime/Engine/Private/EngineAnalytics.cpp
@@ -2,6 +2,7 @@
 
 #include "EnginePrivate.h"
 #include "EngineAnalytics.h"
+#include "EngineBuildSettings.h"
 #include "Runtime/Analytics/Analytics/Public/Analytics.h"
 #include "Runtime/Analytics/Analytics/Public/Interfaces/IAnalyticsProvider.h"
 
@@ -97,18 +98,12 @@ private:
  * External code should bind this delegate if engine analytics are desired,
  * preferably in private code that won't be redistributed.
  */
-FAnalytics::FProviderConfigurationDelegate& GetEngineAnalyticsConfigDelegate()
+FAnalytics::FProviderConfigurationDelegate& GetEngineAnalyticsOverrideConfigDelegate()
 {
 	static FAnalytics::FProviderConfigurationDelegate Delegate;
 	return Delegate;
 }
 
-FSimpleDelegate& OnEngineAnalyticsStartupDelegate()
-{
-	static FSimpleDelegate Delegate;
-	return Delegate;
-}
-
 
 /**
  * On-demand construction of the singleton. 
@@ -133,13 +128,64 @@ void FEngineAnalytics::Initialize()
 		Analytics->SetUserID(FPlatformMisc::GetUniqueDeviceId());
 		Analytics->StartSession();
 #else
-		// Connect the engine analytics provider (if there is a configuration delegate installed)
-		if (GetEngineAnalyticsConfigDelegate().IsBound())
+
 		{
-		
+			// Setup some default engine analytics if there is nothing custom bound
+			FAnalytics::FProviderConfigurationDelegate DefaultEngineAnalyticsConfig;
+			DefaultEngineAnalyticsConfig.BindStatic( 
+				[]( const FString& KeyName, bool bIsValueRequired ) -> FString
+				{
+					static TMap<FString, FString> ConfigMap;
+					if (ConfigMap.Num() == 0)
+					{
+						ConfigMap.Add(TEXT("ProviderModuleName"), TEXT("AnalyticsET"));
+						ConfigMap.Add(TEXT("APIServerET"), TEXT("http://etsource.epicgames.com/ET2/"));
+
+						// We always use the "Release" analytics account unless we're running in analytics test mode (usually with
+						// a command-line parameter), or we're an internal Epic build
+						bool bUseReleaseAccount =
+							(FAnalytics::Get().GetBuildType() == FAnalytics::Development ||  
+							 FAnalytics::Get().GetBuildType() == FAnalytics::Release) &&
+							!FEngineBuildSettings::IsInternalBuild();	// Internal Epic build
+
+						if( FRocketSupport::IsRocket() )
+						{
+							const TCHAR* DevelopmentAccountAPIKeyET = TEXT("Rocket.Dev");
+							const TCHAR* ReleaseAccountAPIKeyET = TEXT("Rocket.Release");
+							ConfigMap.Add(TEXT("APIKeyET"), bUseReleaseAccount ? ReleaseAccountAPIKeyET : DevelopmentAccountAPIKeyET);
+						}
+						else if( FEngineBuildSettings::IsPerforceBuild() )
+						{
+							const TCHAR* DevelopmentAccountAPIKeyET = TEXT("Perforce.Dev");
+							const TCHAR* ReleaseAccountAPIKeyET = TEXT("Perforce.Release");
+							ConfigMap.Add(TEXT("APIKeyET"), bUseReleaseAccount ? ReleaseAccountAPIKeyET : DevelopmentAccountAPIKeyET);
+						}
+						else
+						{
+							const TCHAR* DevelopmentAccountAPIKeyET = TEXT("UnrealEngine.Dev");
+							const TCHAR* ReleaseAccountAPIKeyET = TEXT("UnrealEngine.Release");
+							ConfigMap.Add(TEXT("APIKeyET"), bUseReleaseAccount ? ReleaseAccountAPIKeyET : DevelopmentAccountAPIKeyET);
+						}
+					}
+
+					// Check for overrides
+					if( GetEngineAnalyticsOverrideConfigDelegate().IsBound() )
+					{
+						const FString OverrideValue = GetEngineAnalyticsOverrideConfigDelegate().Execute( KeyName, bIsValueRequired );
+						if( !OverrideValue.IsEmpty() )
+						{
+							return OverrideValue;
+						}
+					}
+
+					FString* ConfigValue = ConfigMap.Find(KeyName);
+					return ConfigValue != NULL ? *ConfigValue : TEXT("");
+				} );
+
+			// Connect the engine analytics provider (if there is a configuration delegate installed)
 			Analytics = FAnalytics::Get().CreateAnalyticsProvider(
-				FName(*GetEngineAnalyticsConfigDelegate().Execute(TEXT("ProviderModuleName"), true)), 
-				GetEngineAnalyticsConfigDelegate());
+				FName(*DefaultEngineAnalyticsConfig.Execute(TEXT("ProviderModuleName"), true)), 
+				DefaultEngineAnalyticsConfig);
 			if (Analytics.IsValid())
 			{
 				Analytics->SetUserID(FPlatformMisc::GetUniqueDeviceId());
@@ -149,12 +195,6 @@ void FEngineAnalytics::Initialize()
 #endif
 	}
 	bIsInitialized = true;
-
-	if( bIsInitialized )
-	{
-		// Allow the analytics configuration to do whatever it wants after starting up
-		OnEngineAnalyticsStartupDelegate().ExecuteIfBound();
-	}
 }
 
 
