UBlueprint::ChangeOwnerOfTemplates will use deterministic names.
ttp 338718 CRITICAL: Regression: CRASH: BlueprintOffice: CONTENTEXAMPLES: Occurs in cooked build when BP info is activated

diff --git a/Engine/Source/Runtime/Engine/Private/Blueprint.cpp b/Engine/Source/Runtime/Engine/Private/Blueprint.cpp
--- a/Engine/Source/Runtime/Engine/Private/Blueprint.cpp
+++ b/Engine/Source/Runtime/Engine/Private/Blueprint.cpp
@@ -11,6 +11,8 @@
 #include "Editor/UnrealEd/Public/Kismet2/KismetEditorUtilities.h"
 #include "Editor/UnrealEd/Public/Kismet2/CompilerResultsLog.h"
 #include "Editor/UnrealEd/Public/Kismet2/StructureEditorUtils.h"
+#include "Editor/UnrealEd/Public/Editor.h"
+#include "Crc.h"
 #endif
 
 DEFINE_LOG_CATEGORY(LogBlueprint);
@@ -913,6 +915,30 @@ void UBlueprint::Message_Error(const FString& MessageToLog)
 
 bool UBlueprint::ChangeOwnerOfTemplates()
 {
+	struct FUniqueNewNameHelper
+	{
+	private:
+		FString NewName;
+		bool isValid;
+	public:
+		FUniqueNewNameHelper(const FString& Name, UObject* Outer) : isValid(false)
+		{
+			const uint32 Hash = FCrc::StrCrc32<TCHAR>(*Name);
+			NewName = FString::Printf(TEXT("%s__%08X"), *Name, Hash);
+			isValid = IsUniqueObjectName(FName(*NewName), Outer);
+			if (!isValid)
+			{
+				check(Outer);
+				UE_LOG(LogBlueprint, Warning, TEXT("ChangeOwnerOfTemplates: Cannot generate a deterministic new name. Old name: %s New outer: %s"), *Name, *Outer->GetName());
+			}
+		}
+
+		const TCHAR* Get() const
+		{
+			return isValid ? *NewName : NULL;
+		}
+	};
+
 	UBlueprintGeneratedClass* BPGClass = Cast<UBlueprintGeneratedClass>(*GeneratedClass);
 	bool bIsStillStale = false;
 	if (BPGClass)
@@ -959,7 +985,7 @@ bool UBlueprint::ChangeOwnerOfTemplates()
 		{
 			if (Curve && (Curve->GetOuter() == this))
 			{
-				const bool bRenamed = Curve->Rename(NULL, BPGClass, REN_ForceNoResetLoaders|REN_DoNotDirty);
+				const bool bRenamed = Curve->Rename(FUniqueNewNameHelper(Curve->GetName(), BPGClass).Get(), BPGClass, REN_ForceNoResetLoaders | REN_DoNotDirty);
 				ensure(bRenamed);
 				bIsStillStale |= !bRenamed;
 			}
@@ -969,7 +995,7 @@ bool UBlueprint::ChangeOwnerOfTemplates()
 		{
 			if(SCS->GetOuter() == this)
 			{
-				const bool bRenamed = SCS->Rename(NULL, BPGClass, REN_ForceNoResetLoaders|REN_DoNotDirty);
+				const bool bRenamed = SCS->Rename(FUniqueNewNameHelper(SCS->GetName(), BPGClass).Get(), BPGClass, REN_ForceNoResetLoaders | REN_DoNotDirty);
 				ensure(bRenamed);
 				bIsStillStale |= !bRenamed;
 				bMigratedOwner = true;
@@ -981,7 +1007,7 @@ bool UBlueprint::ChangeOwnerOfTemplates()
 				UActorComponent* Component = SCSNode ? SCSNode->ComponentTemplate : NULL;
 				if (Component && Component->GetOuter() == this)
 				{
-					const bool bRenamed = Component->Rename(NULL, BPGClass, REN_ForceNoResetLoaders | REN_DoNotDirty);
+					const bool bRenamed = Component->Rename(FUniqueNewNameHelper(Component->GetName(), BPGClass).Get(), BPGClass, REN_ForceNoResetLoaders | REN_DoNotDirty);
 					ensure(bRenamed);
 					bIsStillStale |= !bRenamed;
 					bMigratedOwner = true;
