Customize locked axis mode

diff --git a/Engine/Source/Runtime/Engine/Classes/PhysicsEngine/BodyInstance.h b/Engine/Source/Runtime/Engine/Classes/PhysicsEngine/BodyInstance.h
--- a/Engine/Source/Runtime/Engine/Classes/PhysicsEngine/BodyInstance.h
+++ b/Engine/Source/Runtime/Engine/Classes/PhysicsEngine/BodyInstance.h
@@ -24,6 +24,26 @@ namespace physx
 }
 #endif // WITH_PHYSX
 
+UENUM()
+namespace ELockedAxis
+{
+	enum Type
+	{
+		/*Uses the default locked axis as specified in the project settings.*/
+		Default,
+		/*Lock movement along the x-axis*/
+		X,
+		/*Lock movement along the y-axis*/
+		Y,
+		/*Lock movement along the z-axis*/
+		Z,
+		/*Lock movement along custom axis*/
+		Custom,
+		/*No axis is locked.*/
+		None
+	};
+}
+
 USTRUCT()
 struct ENGINE_API FCollisionResponse
 {
@@ -166,9 +186,16 @@ public:
 	UPROPERTY()
 	uint32 bUpdateMassWhenScaleChanges:1;
 
-	/** Constrain body to 2d plane with DOF as normal. (0,0,0) uses regular 3D */
-	UPROPERTY(EditAnywhere, AdvancedDisplay, Category = Physics)
-	FVector DOF;
+	/** Locks physical movement along specified axis.*/
+	UPROPERTY(EditAnywhere, Category = Physics, meta=(DisplayName="Locked Axis"))
+	TEnumAsByte<ELockedAxis::Type> LockedAxisMode;
+	
+	/** Locks physical movement along custom axis. (0,0,0) indicates no lock*/
+	UPROPERTY(EditAnywhere, Category = Physics)
+	FVector CustomLockedAxis;
+
+	FVector GetLockedAxis() const;
+	void FBodyInstance::CreateDOFLock();
 
 	FConstraintInstance DOFConstraint;
 
diff --git a/Engine/Source/Runtime/Engine/Private/PhysicsEngine/BodyInstance.cpp b/Engine/Source/Runtime/Engine/Private/PhysicsEngine/BodyInstance.cpp
--- a/Engine/Source/Runtime/Engine/Private/PhysicsEngine/BodyInstance.cpp
+++ b/Engine/Source/Runtime/Engine/Private/PhysicsEngine/BodyInstance.cpp
@@ -505,6 +505,64 @@ void FBodyInstance::SetCollisionEnabled(ECollisionEnabled::Type NewType, bool bU
 	}
 }
 
+FVector FBodyInstance::GetLockedAxis() const
+{
+	ELockedAxis::Type MyLockedAxis = LockedAxisMode == ELockedAxis::Default ? ELockedAxis::None : LockedAxisMode;
+	switch (MyLockedAxis)
+	{
+	case ELockedAxis::None: return FVector::ZeroVector;
+	case ELockedAxis::X: return FVector(1, 0, 0);
+	case ELockedAxis::Y: return FVector(0, 1, 0);
+	case ELockedAxis::Z: return FVector(0, 0, 1);
+	case ELockedAxis::Custom: return CustomLockedAxis;
+	default:	check(0);	//unsupported locked axis type
+	}
+
+	return FVector::ZeroVector;
+}
+
+void FBodyInstance::CreateDOFLock()
+{
+	//TODO: This hasn't been tested with calling the function multiple times. It's really only expected to be used during InitBody
+	if (IsDynamic() == false) return;
+#if WITH_PHYSX
+	//setup constraint based on DOF
+	FVector LockedAxis = GetLockedAxis();
+	if (LockedAxis.IsNearlyZero() == false)
+	{
+		DOFConstraint.bSwingLimitSoft = false;
+		DOFConstraint.bTwistLimitSoft = false;
+		DOFConstraint.bLinearLimitSoft = false;
+		//set all rotation to free
+		DOFConstraint.AngularSwing1Motion = EAngularConstraintMotion::ACM_Locked;
+		DOFConstraint.AngularSwing2Motion = EAngularConstraintMotion::ACM_Locked;
+		DOFConstraint.AngularTwistMotion = EAngularConstraintMotion::ACM_Free;
+
+		DOFConstraint.LinearXMotion = ELinearConstraintMotion::LCM_Locked;
+		DOFConstraint.LinearYMotion = ELinearConstraintMotion::LCM_Free;
+		DOFConstraint.LinearZMotion = ELinearConstraintMotion::LCM_Free;
+
+		FVector Normal = LockedAxis.SafeNormal();
+		FVector Sec;
+		FVector Garbage;
+		Normal.FindBestAxisVectors(Garbage, Sec);
+
+
+		FTransform TM = GetUnrealWorldTransform();
+
+		DOFConstraint.PriAxis1 = TM.InverseTransformVectorNoScale(Normal);
+		DOFConstraint.SecAxis1 = TM.InverseTransformVectorNoScale(Sec);
+
+		DOFConstraint.PriAxis2 = Normal;
+		DOFConstraint.SecAxis2 = Sec;
+		DOFConstraint.Pos2 = TM.GetLocation();
+
+		// Create constraint instance based on DOF
+		DOFConstraint.InitConstraint(OwnerComponent.Get(), this, nullptr, 1.f);
+	}
+#endif
+}
+
 ECollisionEnabled::Type FBodyInstance::GetCollisionEnabled() const
 {
 	// Check actor override
@@ -1223,39 +1281,7 @@ void FBodyInstance::InitBody(UBodySetup* Setup, const FTransform& Transform, UPr
 		int32 VelocityIterCount = FMath::Clamp(VelocitySolverIterationCount, 1, 255);
 		PNewDynamic->setSolverIterationCounts(PositionIterCount, VelocityIterCount);
 
-		//setup constraint based on DOF
-		if (!DOF.IsZero())
-		{
-			DOFConstraint.bSwingLimitSoft = false;
-			DOFConstraint.bTwistLimitSoft = false;
-			DOFConstraint.bLinearLimitSoft = false;
-			//set all rotation to free
-			DOFConstraint.AngularSwing1Motion = EAngularConstraintMotion::ACM_Locked;
-			DOFConstraint.AngularSwing2Motion = EAngularConstraintMotion::ACM_Locked;
-			DOFConstraint.AngularTwistMotion  = EAngularConstraintMotion::ACM_Free;
-
-			DOFConstraint.LinearXMotion = ELinearConstraintMotion::LCM_Locked;
-			DOFConstraint.LinearYMotion = ELinearConstraintMotion::LCM_Free;
-			DOFConstraint.LinearZMotion = ELinearConstraintMotion::LCM_Free;
-			
-			FVector Normal = DOF.SafeNormal();
-			FVector Sec;
-			FVector Garbage;
-			Normal.FindBestAxisVectors(Garbage, Sec);
-
-
-			FTransform TM = GetUnrealWorldTransform();
-
-			DOFConstraint.PriAxis1 = TM.InverseTransformVectorNoScale(Normal);
-			DOFConstraint.SecAxis1 = TM.InverseTransformVectorNoScale(Sec);
-
-			DOFConstraint.PriAxis2 = Normal;
-			DOFConstraint.SecAxis2 = Sec;
-			DOFConstraint.Pos2 = TM.GetLocation();
-
-			// Create constraint instance based on DOF
-			DOFConstraint.InitConstraint(OwnerComponent.Get(), this, nullptr, 1.f);
-		}
+		CreateDOFLock();
 
 		// wakeUp and putToSleep will issue warnings on kinematic actors
 		if (IsRigidDynamicNonKinematic(PNewDynamic))
@@ -1909,10 +1935,13 @@ bool FBodyInstance::IsInstanceSimulatingPhysics(bool bIgnoreOwner)
 
 bool FBodyInstance::ShouldInstanceSimulatingPhysics(bool bIgnoreOwner)
 {
-	//If type is set to default inherit whatever the parent does
+	// if I'm simulating or owner is simulating
 	if (OwnerComponent != NULL && BodySetup.IsValid() && BodySetup.Get()->PhysicsType == PhysType_Default && !bIgnoreOwner)
 	{
+		// if derive from owner, and owner is simulating, this should simulate
 		return OwnerComponent->BodyInstance.bSimulatePhysics;
+
+		// or else, it should look its own setting
 	}
 
 	return bSimulatePhysics;
diff --git a/Engine/Source/Runtime/Engine/Private/SkeletalMeshComponentPhysics.cpp b/Engine/Source/Runtime/Engine/Private/SkeletalMeshComponentPhysics.cpp
--- a/Engine/Source/Runtime/Engine/Private/SkeletalMeshComponentPhysics.cpp
+++ b/Engine/Source/Runtime/Engine/Private/SkeletalMeshComponentPhysics.cpp
@@ -826,7 +826,8 @@ void USkeletalMeshComponent::InitArticulated(FPhysScene* PhysScene)
 
 			if (i == RootBodyIndex)
 			{
-				BodyInst->DOF = BodyInstance.DOF;
+				BodyInst->LockedAxisMode = BodyInstance.LockedAxisMode;
+				BodyInst->CustomLockedAxis = BodyInstance.CustomLockedAxis;
 			}
 
 #if WITH_PHYSX
