Expose locked axis physics to project settings

diff --git a/Engine/Source/Runtime/Engine/Classes/PhysicsEngine/PhysicsSettings.h b/Engine/Source/Runtime/Engine/Classes/PhysicsEngine/PhysicsSettings.h
--- a/Engine/Source/Runtime/Engine/Classes/PhysicsEngine/PhysicsSettings.h
+++ b/Engine/Source/Runtime/Engine/Classes/PhysicsEngine/PhysicsSettings.h
@@ -48,6 +48,23 @@ namespace EFrictionCombineMode
 	};
 }
 
+UENUM()
+namespace ESettingsLockedAxis
+{
+	enum Type
+	{
+		/*No axis is locked*/
+		None,
+		/*Lock movement along the x-axis*/
+		X,
+		/*Lock movement along the y-axis*/
+		Y,
+		/*Lock movement along the z-axis*/
+		Z
+	};
+}
+
+
 /**
  * Implements project settings for the physics sub-system.
  */
@@ -73,6 +90,11 @@ class ENGINE_API UPhysicsSettings
 	UPROPERTY(config, EditAnywhere, Category = Simulation)
 	bool bEnable2DPhysics;
 
+	/** Locks axis of movement for physical objects. (Useful for making 2D games in a 3D setting) */
+	UPROPERTY(config, EditAnywhere, Category = Simulation)
+	TEnumAsByte<ESettingsLockedAxis::Type> LockedAxis;
+	
+
 	/** Friction combine mode, controls how friction is computed for multiple materials. */
 	UPROPERTY(config, EditAnywhere, Category=Simulation)
 	TEnumAsByte<EFrictionCombineMode::Type> FrictionCombineMode;
diff --git a/Engine/Source/Runtime/Engine/Private/PhysicsEngine/BodyInstance.cpp b/Engine/Source/Runtime/Engine/Private/PhysicsEngine/BodyInstance.cpp
--- a/Engine/Source/Runtime/Engine/Private/PhysicsEngine/BodyInstance.cpp
+++ b/Engine/Source/Runtime/Engine/Private/PhysicsEngine/BodyInstance.cpp
@@ -507,7 +507,16 @@ void FBodyInstance::SetCollisionEnabled(ECollisionEnabled::Type NewType, bool bU
 
 FVector FBodyInstance::GetLockedAxis() const
 {
-	ELockedAxis::Type MyLockedAxis = LockedAxisMode == ELockedAxis::Default ? ELockedAxis::None : LockedAxisMode;
+	ELockedAxis::Type MyLockedAxis = LockedAxisMode;
+	if (MyLockedAxis == ELockedAxis::Default)
+	{
+		ESettingsLockedAxis::Type SettingLockedAxis = UPhysicsSettings::Get()->LockedAxis;
+		if (SettingLockedAxis == ESettingsLockedAxis::X) MyLockedAxis = ELockedAxis::X;
+		if (SettingLockedAxis == ESettingsLockedAxis::Y) MyLockedAxis = ELockedAxis::Y;
+		if (SettingLockedAxis == ESettingsLockedAxis::Z) MyLockedAxis = ELockedAxis::Z;
+		if (SettingLockedAxis == ESettingsLockedAxis::None) MyLockedAxis = ELockedAxis::None;
+	}
+
 	switch (MyLockedAxis)
 	{
 	case ELockedAxis::None: return FVector::ZeroVector;
