Fix build

diff --git a/Engine/Source/Runtime/Engine/Classes/PhysicsEngine/BodyInstance.h b/Engine/Source/Runtime/Engine/Classes/PhysicsEngine/BodyInstance.h
--- a/Engine/Source/Runtime/Engine/Classes/PhysicsEngine/BodyInstance.h
+++ b/Engine/Source/Runtime/Engine/Classes/PhysicsEngine/BodyInstance.h
@@ -195,7 +195,7 @@ public:
 	FVector CustomLockedAxis;
 
 	FVector GetLockedAxis() const;
-	void FBodyInstance::CreateDOFLock();
+	void CreateDOFLock();
 
 	FConstraintInstance DOFConstraint;
 
