Fixed navigation links' direction "RigthToLeft" not working #UE4

- was filed as TTP#340367
- merged over from dev branch

diff --git a/Engine/Source/Runtime/Engine/Private/AI/Navigation/NavigationSystemHelpers.cpp b/Engine/Source/Runtime/Engine/Private/AI/Navigation/NavigationSystemHelpers.cpp
--- a/Engine/Source/Runtime/Engine/Private/AI/Navigation/NavigationSystemHelpers.cpp
+++ b/Engine/Source/Runtime/Engine/Private/AI/Navigation/NavigationSystemHelpers.cpp
@@ -66,33 +66,29 @@ namespace NavigationHelper
 		const FTransform LocalToWorld = Actor->ActorToWorld();
 		FSimpleLinkNavModifier SimpleLink(NavLinks, LocalToWorld);
 
-		if (SimpleLink.HasFallDownLinks())
+		// adjust links
+		for (int32 LinkIndex = 0; LinkIndex < SimpleLink.Links.Num(); ++LinkIndex)
 		{
-			// adjust links
-			for (int32 LinkIndex = 0; LinkIndex < SimpleLink.Links.Num(); ++LinkIndex)
-			{
-				FNavigationLink& Link = SimpleLink.Links[LinkIndex];
-				if (Link.MaxFallDownLength > 0)
-				{
-					// this one needs adjusting
-					if (Link.Direction == ENavLinkDirection::RightToLeft)
-					{
-						Swap(Link.Left, Link.Right);
-					}
+			FNavigationLink& Link = SimpleLink.Links[LinkIndex];
 
-					const FVector WorldRight = LocalToWorld.TransformPosition(Link.Right);
-					//const FVector WorldLeft = LocalToWorld.TransformPosition(ModifiedLink.Left);
+			// this one needs adjusting
+			if (Link.Direction == ENavLinkDirection::RightToLeft)
+			{
+				Swap(Link.Left, Link.Right);
+			}
 
-					const float FallDownHeight = RawGeometryFall(Actor, WorldRight, Link.MaxFallDownLength);
+			if (Link.MaxFallDownLength > 0)
+			{
+				const FVector WorldRight = LocalToWorld.TransformPosition(Link.Right);
+				const float FallDownHeight = RawGeometryFall(Actor, WorldRight, Link.MaxFallDownLength);
 
-					if (FallDownHeight > 0.f)
-					{
-						// @todo maybe it's a good idea to clear ModifiedLink.MaxFallDownLength here
-						UE_VLOG_SEGMENT(Actor, LogNavigation, Log, WorldRight, WorldRight + FVector(0, 0, -FallDownHeight)
-							, FColor::Green, TEXT("FallDownHeight %d"), LinkIndex);
+				if (FallDownHeight > 0.f)
+				{
+					// @todo maybe it's a good idea to clear ModifiedLink.MaxFallDownLength here
+					UE_VLOG_SEGMENT(Actor, LogNavigation, Log, WorldRight, WorldRight + FVector(0, 0, -FallDownHeight)
+						, FColor::Green, TEXT("FallDownHeight %d"), LinkIndex);
 
-						Link.Right.Z -= FallDownHeight;
-					}
+					Link.Right.Z -= FallDownHeight;
 				}
 			}
 		}
@@ -105,43 +101,41 @@ namespace NavigationHelper
 		const FTransform LocalToWorld = Actor->ActorToWorld();
 		FSimpleLinkNavModifier SimpleLink(NavLinks, LocalToWorld);
 
-		if (SimpleLink.HasFallDownLinks())
+			// adjust links if needed
+		for (int32 LinkIndex = 0; LinkIndex < SimpleLink.SegmentLinks.Num(); ++LinkIndex)
 		{
-			// adjust links
-			for (int32 LinkIndex = 0; LinkIndex < SimpleLink.SegmentLinks.Num(); ++LinkIndex)
+			FNavigationSegmentLink& Link = SimpleLink.SegmentLinks[LinkIndex];
+
+			// this one needs adjusting
+			if (Link.Direction == ENavLinkDirection::RightToLeft)
+			{
+				Swap(Link.LeftStart, Link.RightStart);
+				Swap(Link.LeftEnd, Link.RightEnd);
+			}
+
+			if (Link.MaxFallDownLength > 0)
 			{
-				FNavigationSegmentLink& Link = SimpleLink.SegmentLinks[LinkIndex];
-				if (Link.MaxFallDownLength > 0)
+				const FVector WorldRightStart = LocalToWorld.TransformPosition(Link.RightStart);
+				const FVector WorldRightEnd = LocalToWorld.TransformPosition(Link.RightEnd);
+
+				const float FallDownHeightStart = RawGeometryFall(Actor, WorldRightStart, Link.MaxFallDownLength);
+				const float FallDownHeightEnd = RawGeometryFall(Actor, WorldRightEnd, Link.MaxFallDownLength);
+
+				if (FallDownHeightStart > 0.f)
 				{
-					// this one needs adjusting
-					if (Link.Direction == ENavLinkDirection::RightToLeft)
-					{
-						Swap(Link.LeftStart, Link.RightStart);
-						Swap(Link.LeftEnd, Link.RightEnd);
-					}
-
-					const FVector WorldRightStart = LocalToWorld.TransformPosition(Link.RightStart);
-					const FVector WorldRightEnd = LocalToWorld.TransformPosition(Link.RightEnd);
-
-					const float FallDownHeightStart = RawGeometryFall(Actor, WorldRightStart, Link.MaxFallDownLength);
-					const float FallDownHeightEnd = RawGeometryFall(Actor, WorldRightEnd, Link.MaxFallDownLength);
-
-					if (FallDownHeightStart > 0.f)
-					{
-						// @todo maybe it's a good idea to clear ModifiedLink.MaxFallDownLength here
-						UE_VLOG_SEGMENT(Actor, LogNavigation, Log, WorldRightStart, WorldRightStart + FVector(0, 0, -FallDownHeightStart)
-							, FColor::Green, TEXT("FallDownHeightStart %d"), LinkIndex);
-
-						Link.RightStart.Z -= FallDownHeightStart;
-					}
-					if (FallDownHeightEnd > 0.f)
-					{
-						// @todo maybe it's a good idea to clear ModifiedLink.MaxFallDownLength here
-						UE_VLOG_SEGMENT(Actor, LogNavigation, Log, WorldRightEnd, WorldRightEnd + FVector(0, 0, -FallDownHeightEnd)
-							, FColor::Green, TEXT("FallDownHeightEnd %d"), LinkIndex);
-
-						Link.RightEnd.Z -= FallDownHeightEnd;
-					}
+					// @todo maybe it's a good idea to clear ModifiedLink.MaxFallDownLength here
+					UE_VLOG_SEGMENT(Actor, LogNavigation, Log, WorldRightStart, WorldRightStart + FVector(0, 0, -FallDownHeightStart)
+						, FColor::Green, TEXT("FallDownHeightStart %d"), LinkIndex);
+
+					Link.RightStart.Z -= FallDownHeightStart;
+				}
+				if (FallDownHeightEnd > 0.f)
+				{
+					// @todo maybe it's a good idea to clear ModifiedLink.MaxFallDownLength here
+					UE_VLOG_SEGMENT(Actor, LogNavigation, Log, WorldRightEnd, WorldRightEnd + FVector(0, 0, -FallDownHeightEnd)
+						, FColor::Green, TEXT("FallDownHeightEnd %d"), LinkIndex);
+
+					Link.RightEnd.Z -= FallDownHeightEnd;
 				}
 			}
 		}
diff --git a/Engine/Source/Runtime/Engine/Private/AI/Navigation/RecastNavMeshGenerator.cpp b/Engine/Source/Runtime/Engine/Private/AI/Navigation/RecastNavMeshGenerator.cpp
--- a/Engine/Source/Runtime/Engine/Private/AI/Navigation/RecastNavMeshGenerator.cpp
+++ b/Engine/Source/Runtime/Engine/Private/AI/Navigation/RecastNavMeshGenerator.cpp
@@ -1282,16 +1282,9 @@ struct FOffMeshData
 			dtOffMeshLinkCreateParams NewInfo;
 			FMemory::MemZero(NewInfo);
 
-			if (Link.Direction == ENavLinkDirection::LeftToRight || Link.Direction == ENavLinkDirection::BothWays)
-			{
-				StoreUnrealPoint(NewInfo.vertsA0, LocalToWorld.TransformPosition(Link.Left));
-				StoreUnrealPoint(NewInfo.vertsB0, LocalToWorld.TransformPosition(Link.Right));
-			}
-			else
-			{
-				StoreUnrealPoint(NewInfo.vertsA0, LocalToWorld.TransformPosition(Link.Right));
-				StoreUnrealPoint(NewInfo.vertsB0, LocalToWorld.TransformPosition(Link.Left));
-			}
+			// not doing anything to link's points order - should be already ordered properly by link processor
+			StoreUnrealPoint(NewInfo.vertsA0, LocalToWorld.TransformPosition(Link.Left));
+			StoreUnrealPoint(NewInfo.vertsB0, LocalToWorld.TransformPosition(Link.Right));
 
 			NewInfo.type = DT_OFFMESH_CON_POINT | (Link.Direction == ENavLinkDirection::BothWays ? DT_OFFMESH_CON_BIDIR : 0);
 			NewInfo.snapRadius = Link.SnapRadius;
@@ -1326,20 +1319,11 @@ struct FOffMeshData
 			dtOffMeshLinkCreateParams NewInfo;
 			FMemory::MemZero(NewInfo);
 
-			if (Link.Direction == ENavLinkDirection::LeftToRight || Link.Direction == ENavLinkDirection::BothWays)
-			{
-				StoreUnrealPoint(NewInfo.vertsA0, LocalToWorld.TransformPosition(Link.LeftStart));
-				StoreUnrealPoint(NewInfo.vertsA1, LocalToWorld.TransformPosition(Link.LeftEnd));
-				StoreUnrealPoint(NewInfo.vertsB0, LocalToWorld.TransformPosition(Link.RightStart));
-				StoreUnrealPoint(NewInfo.vertsB1, LocalToWorld.TransformPosition(Link.RightEnd));
-			}
-			else
-			{
-				StoreUnrealPoint(NewInfo.vertsA0, LocalToWorld.TransformPosition(Link.RightStart));
-				StoreUnrealPoint(NewInfo.vertsA1, LocalToWorld.TransformPosition(Link.RightEnd));
-				StoreUnrealPoint(NewInfo.vertsB0, LocalToWorld.TransformPosition(Link.LeftStart));
-				StoreUnrealPoint(NewInfo.vertsB1, LocalToWorld.TransformPosition(Link.LeftEnd));
-			}
+			// not doing anything to link's points order - should be already ordered properly by link processor
+			StoreUnrealPoint(NewInfo.vertsA0, LocalToWorld.TransformPosition(Link.LeftStart));
+			StoreUnrealPoint(NewInfo.vertsA1, LocalToWorld.TransformPosition(Link.LeftEnd));
+			StoreUnrealPoint(NewInfo.vertsB0, LocalToWorld.TransformPosition(Link.RightStart));
+			StoreUnrealPoint(NewInfo.vertsB1, LocalToWorld.TransformPosition(Link.RightEnd));
 
 			NewInfo.type = DT_OFFMESH_CON_SEGMENT | (Link.Direction == ENavLinkDirection::BothWays ? DT_OFFMESH_CON_BIDIR : 0);
 			NewInfo.snapRadius = Link.SnapRadius;
