Make more classes able to be included directly
Rename Engine/Public/Texture.h to TextureResource.h
Move FTextureLODSettings into its own h/cpp files, change includes where relevant
Move FReferenceSkeleton into its own header (inc. FMeshBoneInfo) from SkeletalMeshTypes.h

diff --git a/Engine/Source/Developer/HTML5/HTML5TargetPlatform/Private/HTML5TargetPlatform.cpp b/Engine/Source/Developer/HTML5/HTML5TargetPlatform/Private/HTML5TargetPlatform.cpp
--- a/Engine/Source/Developer/HTML5/HTML5TargetPlatform/Private/HTML5TargetPlatform.cpp
+++ b/Engine/Source/Developer/HTML5/HTML5TargetPlatform/Private/HTML5TargetPlatform.cpp
@@ -7,7 +7,7 @@
 #include "HTML5TargetPlatformPrivatePCH.h"
 
 #if WITH_EDITOR
-#include "Texture.h"
+#include "TextureLODSettings.h"
 #endif 
 
  /* FHTML5TargetPlatform structors
diff --git a/Engine/Source/Developer/TargetPlatform/Public/Common/TargetPlatformBase.h b/Engine/Source/Developer/TargetPlatform/Public/Common/TargetPlatformBase.h
--- a/Engine/Source/Developer/TargetPlatform/Public/Common/TargetPlatformBase.h
+++ b/Engine/Source/Developer/TargetPlatform/Public/Common/TargetPlatformBase.h
@@ -40,7 +40,7 @@ public:
 	}
 
 
-#ifdef TEXTURE_H_INCLUDED // defined in Texture.h, this way we know if UTexture is available, needed for Clang
+#ifdef TEXTURERESOURCE_H_INCLUDED // defined in TextureResource.h, this way we know if UTexture is available, needed for Clang
 	FName GetDefaultTextureFormatName( const UTexture* Texture, const FConfigFile& EngineSettings ) const
 	{
 		FName TextureFormatName = NAME_None;
@@ -149,11 +149,11 @@ public:
 
 		return TextureFormatName;
 	}
-#else //TEXTURE_H_INCLUDED
+#else //TEXTURERESOURCE_H_INCLUDED
 
 	FName GetDefaultTextureFormatName( const UTexture* Texture, const FConfigFile& EngineSettings ) const;
 
-#endif //TEXTURE_H_INCLUDED
+#endif //TEXTURERESOURCE_H_INCLUDED
 #endif //WITH_ENGINE
 
 	virtual bool PackageBuild( const FString& InPackgeDirectory ) override
diff --git a/Engine/Source/Runtime/Engine/Classes/AI/Navigation/NavigationTypes.h b/Engine/Source/Runtime/Engine/Classes/AI/Navigation/NavigationTypes.h
--- a/Engine/Source/Runtime/Engine/Classes/AI/Navigation/NavigationTypes.h
+++ b/Engine/Source/Runtime/Engine/Classes/AI/Navigation/NavigationTypes.h
@@ -216,7 +216,6 @@ struct ENGINE_API FNavigationPath : public TSharedFromThis<FNavigationPath, ESPM
 	FORCEINLINE void GetObserver(FPathObserverDelegate& Observer) const { Observer = ObserverDelegate; }
 
 	FORCEINLINE void MarkReady() { bIsReady = true; }
-	FORCEINLINE void SetOwner(const class ANavigationData* const NewOwner) { Owner = NewOwner; }
 	FORCEINLINE void SetObserver(const FPathObserverDelegate& Observer) { ObserverDelegate = Observer; }
 	FORCEINLINE void SetIsPartial(const bool bPartial) { bIsPartial = bPartial; }
 	FORCEINLINE void SetSearchReachedLimit(const bool bLimited) { bReachedSearchLimit = bLimited; }
@@ -227,6 +226,8 @@ struct ENGINE_API FNavigationPath : public TSharedFromThis<FNavigationPath, ESPM
 		ObserverDelegate.ExecuteIfBound(this); 
 	}
 
+	void SetOwner(const class ANavigationData* const NewOwner);
+
 	virtual void DebugDraw(const class ANavigationData* NavData, FColor PathColor, class UCanvas* Canvas, bool bPersistent, const uint32 NextPathPointIndex = 0) const;
 #if ENABLE_VISUAL_LOG
 	virtual void DescribeSelfToVisLog(struct FVisLogEntry* Snapshot) const;
@@ -459,14 +460,7 @@ struct ENGINE_API FPathFindingQuery
 	/** additional flags passed to navigation data handling request */
 	int32 NavDataFlags;
 
-	FPathFindingQuery()
-		: NavData(NULL)
-		, Owner(NULL)
-		, StartLocation(FVector::ZeroVector)
-		, EndLocation(FVector::ZeroVector)
-		, NavDataFlags(0)
-	{
-	}
+	FPathFindingQuery();
 
 	FPathFindingQuery(const FPathFindingQuery& Source);
 
diff --git a/Engine/Source/Runtime/Engine/Classes/Camera/PlayerCameraManager.h b/Engine/Source/Runtime/Engine/Classes/Camera/PlayerCameraManager.h
--- a/Engine/Source/Runtime/Engine/Classes/Camera/PlayerCameraManager.h
+++ b/Engine/Source/Runtime/Engine/Classes/Camera/PlayerCameraManager.h
@@ -1,13 +1,15 @@
 // Copyright 1998-2014 Epic Games, Inc. All Rights Reserved.
 
 #pragma once
-
+#include "GameFramework/Actor.h"
 #include "CameraTypes.h"
 #include "Engine/Scene.h"
 #include "Engine/EngineBaseTypes.h"
 
 #include "PlayerCameraManager.generated.h"
 
+class APawn;
+
 UENUM()
 enum EViewTargetBlendFunction
 {
diff --git a/Engine/Source/Runtime/Engine/Classes/Components/InputComponent.h b/Engine/Source/Runtime/Engine/Classes/Components/InputComponent.h
--- a/Engine/Source/Runtime/Engine/Classes/Components/InputComponent.h
+++ b/Engine/Source/Runtime/Engine/Classes/Components/InputComponent.h
@@ -2,7 +2,8 @@
 
 #pragma once
 
-#include "ActorComponent.h"
+#include "Components/ActorComponent.h"
+#include "Runtime/InputCore/Classes/InputCoreTypes.h"
 #include "InputComponent.generated.h"
 
 template<class DelegateType, class DynamicDelegateType>
diff --git a/Engine/Source/Runtime/Engine/Classes/Components/LightComponent.h b/Engine/Source/Runtime/Engine/Classes/Components/LightComponent.h
--- a/Engine/Source/Runtime/Engine/Classes/Components/LightComponent.h
+++ b/Engine/Source/Runtime/Engine/Classes/Components/LightComponent.h
@@ -4,7 +4,7 @@
 
 #include "LightComponentBase.h"
 #include "SceneTypes.h"
-
+#include "EngineDefines.h"
 #include "LightComponent.generated.h"
 
 /** 
diff --git a/Engine/Source/Runtime/Engine/Classes/Components/PrimitiveComponent.h b/Engine/Source/Runtime/Engine/Classes/Components/PrimitiveComponent.h
--- a/Engine/Source/Runtime/Engine/Classes/Components/PrimitiveComponent.h
+++ b/Engine/Source/Runtime/Engine/Classes/Components/PrimitiveComponent.h
@@ -2,7 +2,8 @@
 
 #pragma once
 
-#include "StaticArray.h"
+#include "Runtime/Core/Public/Containers/StaticArray.h"
+#include "Runtime/InputCore/Classes/InputCoreTypes.h"
 #include "PhysicsEngine/BodyInstance.h"
 #include "Components/SceneComponent.h"
 #include "Components/LightComponent.h"
diff --git a/Engine/Source/Runtime/Engine/Classes/Components/SceneComponent.h b/Engine/Source/Runtime/Engine/Classes/Components/SceneComponent.h
--- a/Engine/Source/Runtime/Engine/Classes/Components/SceneComponent.h
+++ b/Engine/Source/Runtime/Engine/Classes/Components/SceneComponent.h
@@ -19,12 +19,7 @@ struct FOverlapInfo
 	{
 	}
 
-	FOverlapInfo(class UPrimitiveComponent* InComponent, int32 InBodyIndex = INDEX_NONE)
-		: bFromSweep(false)
-	{
-		OverlapInfo.Component = InComponent;
-		OverlapInfo.Item = InBodyIndex;
-	}
+	FOverlapInfo(class UPrimitiveComponent* InComponent, int32 InBodyIndex = INDEX_NONE);
 	
 	int32 GetBodyIndex() const { return OverlapInfo.Item;  }
 
diff --git a/Engine/Source/Runtime/Engine/Classes/Engine/BlendableInterface.h b/Engine/Source/Runtime/Engine/Classes/Engine/BlendableInterface.h
--- a/Engine/Source/Runtime/Engine/Classes/Engine/BlendableInterface.h
+++ b/Engine/Source/Runtime/Engine/Classes/Engine/BlendableInterface.h
@@ -4,6 +4,8 @@
 
 #include "BlendableInterface.generated.h"
 
+class UMaterialInstanceDynamic;
+
 /** Where to place a material node in the post processing graph */
 UENUM()
 enum EBlendableLocation
diff --git a/Engine/Source/Runtime/Engine/Classes/Engine/EngineTypes.h b/Engine/Source/Runtime/Engine/Classes/Engine/EngineTypes.h
--- a/Engine/Source/Runtime/Engine/Classes/Engine/EngineTypes.h
+++ b/Engine/Source/Runtime/Engine/Classes/Engine/EngineTypes.h
@@ -1367,16 +1367,7 @@ struct ENGINE_API FHitResult
 	}
 
 	/** Ctor for easily creating "fake" hits from limited data. */
-	FHitResult(class AActor* InActor, class UPrimitiveComponent* InComponent, FVector const& HitLoc, FVector const& HitNorm)
-	{
-		FMemory::Memzero(this, sizeof(FHitResult));
-		Location = HitLoc;
-		ImpactPoint = HitLoc;
-		Normal = HitNorm;
-		ImpactNormal = HitNorm;
-		Actor = InActor;
-		Component = InComponent;
-	}
+	FHitResult(class AActor* InActor, class UPrimitiveComponent* InComponent, FVector const& HitLoc, FVector const& HitNorm);
 
 	void Reset(float InTime = 1.f, bool bPreserveTraceData = true)
 	{
diff --git a/Engine/Source/Runtime/Engine/Classes/Engine/Texture.h b/Engine/Source/Runtime/Engine/Classes/Engine/Texture.h
--- a/Engine/Source/Runtime/Engine/Classes/Engine/Texture.h
+++ b/Engine/Source/Runtime/Engine/Classes/Engine/Texture.h
@@ -2,7 +2,8 @@
 
 #pragma once
 
-#include "TextureDefines.h"
+#include "Engine/TextureDefines.h"
+#include "TextureResource.h"
 #include "MaterialShared.h"
 #include "Texture.generated.h"
 
diff --git a/Engine/Source/Runtime/Engine/Classes/Engine/Texture2D.h b/Engine/Source/Runtime/Engine/Classes/Engine/Texture2D.h
--- a/Engine/Source/Runtime/Engine/Classes/Engine/Texture2D.h
+++ b/Engine/Source/Runtime/Engine/Classes/Engine/Texture2D.h
@@ -2,7 +2,7 @@
 
 #pragma once
 
-#include "Texture.h"
+#include "Engine/Texture.h"
 #include "Texture2D.generated.h"
 
 UCLASS(hidecategories=Object, MinimalAPI, BlueprintType)
diff --git a/Engine/Source/Runtime/Engine/Classes/GameFramework/Character.h b/Engine/Source/Runtime/Engine/Classes/GameFramework/Character.h
--- a/Engine/Source/Runtime/Engine/Classes/GameFramework/Character.h
+++ b/Engine/Source/Runtime/Engine/Classes/GameFramework/Character.h
@@ -1,10 +1,9 @@
 // Copyright 1998-2014 Epic Games, Inc. All Rights Reserved.
 
 #pragma once
-
-#include "Controller.h"
+#include "GameFramework/Pawn.h"
 #include "GameFramework/CharacterMovementComponent.h"
-
+#include "Components/PrimitiveComponent.h"
 #include "Character.generated.h"
 
 DECLARE_DYNAMIC_MULTICAST_DELEGATE_ThreeParams(FMovementModeChangedSignature, class ACharacter*, Character, EMovementMode, PrevMovementMode, uint8, PreviousCustomMode);
diff --git a/Engine/Source/Runtime/Engine/Classes/GameFramework/CharacterMovementComponent.h b/Engine/Source/Runtime/Engine/Classes/GameFramework/CharacterMovementComponent.h
--- a/Engine/Source/Runtime/Engine/Classes/GameFramework/CharacterMovementComponent.h
+++ b/Engine/Source/Runtime/Engine/Classes/GameFramework/CharacterMovementComponent.h
@@ -7,10 +7,17 @@
 
 #pragma once
 
-#include "AI/Navigation/NavigationAvoidanceTypes.h"
+#include "WorldCollision.h"
+#include "GameFramework/PawnMovementComponent.h"
 #include "Interfaces/NetworkPredictionInterface.h"
+#include "Engine/EngineTypes.h"
+#include "AI/Navigation/NavigationAvoidanceTypes.h"
 #include "CharacterMovementComponent.generated.h"
 
+struct FVector_NetQuantize100;
+struct FVector_NetQuantizeNormal;
+class FDebugDisplayInfo;
+
 /** Movement modes for Characters.  */
 UENUM(BlueprintType)
 enum EMovementMode
diff --git a/Engine/Source/Runtime/Engine/Classes/GameFramework/GameMode.h b/Engine/Source/Runtime/Engine/Classes/GameFramework/GameMode.h
--- a/Engine/Source/Runtime/Engine/Classes/GameFramework/GameMode.h
+++ b/Engine/Source/Runtime/Engine/Classes/GameFramework/GameMode.h
@@ -4,6 +4,8 @@
 #include "GameFramework/Info.h"
 #include "GameMode.generated.h"
 
+class FDebugDisplayInfo;
+
 //=============================================================================
 //  GameMode defines the rules and mechanics of the game.  It is only 
 //  instanced on the server and will never exist on the client.
diff --git a/Engine/Source/Runtime/Engine/Classes/GameFramework/HUD.h b/Engine/Source/Runtime/Engine/Classes/GameFramework/HUD.h
--- a/Engine/Source/Runtime/Engine/Classes/GameFramework/HUD.h
+++ b/Engine/Source/Runtime/Engine/Classes/GameFramework/HUD.h
@@ -1,9 +1,7 @@
 // Copyright 1998-2014 Epic Games, Inc. All Rights Reserved.
 
 #pragma once
-
-#include "Canvas.h"
-
+#include "GameFramework/Actor.h"
 #include "HUD.generated.h"
 
 /** List of actors and debug text to draw, @see AddDebugText(), RemoveDebugText(), and DrawDebugTextList() */
diff --git a/Engine/Source/Runtime/Engine/Classes/GameFramework/Info.h b/Engine/Source/Runtime/Engine/Classes/GameFramework/Info.h
--- a/Engine/Source/Runtime/Engine/Classes/GameFramework/Info.h
+++ b/Engine/Source/Runtime/Engine/Classes/GameFramework/Info.h
@@ -6,6 +6,7 @@
   */
 
 #pragma once
+#include "GameFramework/Actor.h"
 #include "Info.generated.h"
 
 UCLASS(abstract, hidecategories=(Input, Movement, Collision, Rendering, "Utilities|Transformation"), showcategories=("Input|MouseInput", "Input|TouchInput"), MinimalAPI, NotBlueprintable)
diff --git a/Engine/Source/Runtime/Engine/Classes/GameFramework/MovementComponent.h b/Engine/Source/Runtime/Engine/Classes/GameFramework/MovementComponent.h
--- a/Engine/Source/Runtime/Engine/Classes/GameFramework/MovementComponent.h
+++ b/Engine/Source/Runtime/Engine/Classes/GameFramework/MovementComponent.h
@@ -8,6 +8,10 @@
 #include "Components/SceneComponent.h"
 #include "MovementComponent.generated.h"
 
+struct FCollisionQueryParams;
+struct FCollisionResponseParams;
+struct FCollisionShape;
+
 UCLASS(ClassGroup=Movement, abstract, BlueprintType)
 class ENGINE_API UMovementComponent : public UActorComponent
 {
diff --git a/Engine/Source/Runtime/Engine/Classes/GameFramework/Pawn.h b/Engine/Source/Runtime/Engine/Classes/GameFramework/Pawn.h
--- a/Engine/Source/Runtime/Engine/Classes/GameFramework/Pawn.h
+++ b/Engine/Source/Runtime/Engine/Classes/GameFramework/Pawn.h
@@ -1,10 +1,8 @@
 // Copyright 1998-2014 Epic Games, Inc. All Rights Reserved.
 
 #pragma once
-
-#include "Controller.h"
+#include "GameFramework/Actor.h"
 #include "PawnMovementComponent.h"
-
 #include "Pawn.generated.h"
 
 ENGINE_API DECLARE_LOG_CATEGORY_EXTERN(LogDamage, Warning, All);
diff --git a/Engine/Source/Runtime/Engine/Classes/Materials/MaterialInterface.h b/Engine/Source/Runtime/Engine/Classes/Materials/MaterialInterface.h
--- a/Engine/Source/Runtime/Engine/Classes/Materials/MaterialInterface.h
+++ b/Engine/Source/Runtime/Engine/Classes/Materials/MaterialInterface.h
@@ -4,6 +4,7 @@
 #pragma once
 #include "Engine/BlendableInterface.h"
 #include "Runtime/RHI/Public/RHIDefinitions.h"
+#include "Runtime/RenderCore/Public/RenderCommandFence.h"
 #include "SceneTypes.h"
 #include "MaterialInterface.generated.h"
 
@@ -11,6 +12,7 @@ class UMaterial;
 class FMaterialResource;
 class FMaterialCompiler;
 struct FPrimitiveViewRelevance;
+class UTexture;
 
 UENUM()
 enum EMaterialUsage
diff --git a/Engine/Source/Runtime/Engine/Private/AI/Navigation/NavigationData.cpp b/Engine/Source/Runtime/Engine/Private/AI/Navigation/NavigationData.cpp
--- a/Engine/Source/Runtime/Engine/Private/AI/Navigation/NavigationData.cpp
+++ b/Engine/Source/Runtime/Engine/Private/AI/Navigation/NavigationData.cpp
@@ -8,6 +8,16 @@
 //----------------------------------------------------------------------//
 // FPathFindingQuery
 //----------------------------------------------------------------------//
+
+FPathFindingQuery::FPathFindingQuery()
+: NavData(NULL)
+, Owner(NULL)
+, StartLocation(FVector::ZeroVector)
+, EndLocation(FVector::ZeroVector)
+, NavDataFlags(0)
+{
+}
+
 FPathFindingQuery::FPathFindingQuery(const UObject* InOwner, const class ANavigationData* InNavData, const FVector& Start, const FVector& End, TSharedPtr<const FNavigationQueryFilter> SourceQueryFilter)
 : NavData(InNavData)
 , Owner(InOwner)
diff --git a/Engine/Source/Runtime/Engine/Private/AI/Navigation/NavigationPath.cpp b/Engine/Source/Runtime/Engine/Private/AI/Navigation/NavigationPath.cpp
--- a/Engine/Source/Runtime/Engine/Private/AI/Navigation/NavigationPath.cpp
+++ b/Engine/Source/Runtime/Engine/Private/AI/Navigation/NavigationPath.cpp
@@ -44,6 +44,12 @@ FNavigationPath::FNavigationPath(const TArray<FVector>& Points, AActor* InBase)
 	}
 }
 
+void FNavigationPath::SetOwner(const class ANavigationData* const NewOwner) 
+{ 
+	Owner = NewOwner; 
+}
+
+
 void FNavigationPath::DebugDraw(const ANavigationData* NavData, FColor PathColor, UCanvas* Canvas, bool bPersistent, const uint32 NextPathPointIndex) const
 {
 	static const FColor Grey(100,100,100);
diff --git a/Engine/Source/Runtime/Engine/Private/Animation/ReferenceSkeleton.cpp b/Engine/Source/Runtime/Engine/Private/Animation/ReferenceSkeleton.cpp
new file mode 100644
--- /dev/null
+++ b/Engine/Source/Runtime/Engine/Private/Animation/ReferenceSkeleton.cpp
@@ -0,0 +1,142 @@
+// Copyright 1998-2014 Epic Games, Inc. All Rights Reserved.
+
+#include "EnginePrivate.h"
+#include "ReferenceSkeleton.h"
+
+FArchive &operator<<(FArchive& Ar, FMeshBoneInfo& F)
+{
+	Ar << F.Name << F.ParentIndex;
+
+	if (Ar.IsLoading() && (Ar.UE4Ver() < VER_UE4_REFERENCE_SKELETON_REFACTOR))
+	{
+		FColor DummyColor = FColor::White;
+		Ar << DummyColor;
+	}
+
+#if WITH_EDITORONLY_DATA
+	if (Ar.UE4Ver() >= VER_UE4_STORE_BONE_EXPORT_NAMES)
+	{
+		if (!Ar.IsCooking())
+		{
+			Ar << F.ExportName;
+		}
+	}
+	else
+	{
+		F.ExportName = F.Name.ToString();
+	}
+#endif
+
+	return Ar;
+}
+
+//////////////////////////////////////////////////////////////////////////
+
+void FReferenceSkeleton::RemoveDuplicateBones(const UObject * Requester, TArray<FBoneIndexType> & DuplicateBones)
+{
+	const int32 NumBones = GetNum();
+	DuplicateBones.Empty();
+
+	TMap<FName, int32> BoneNameCheck;
+	bool bRemovedBones = false;
+	for (int32 BoneIndex = NumBones - 1; BoneIndex >= 0; BoneIndex--)
+	{
+		const FName & BoneName = GetBoneName(BoneIndex);
+		const int32 * FoundBoneIndexPtr = BoneNameCheck.Find(BoneName);
+
+		// Not a duplicate bone, track it.
+		if (FoundBoneIndexPtr == NULL)
+		{
+			BoneNameCheck.Add(BoneName, BoneIndex);
+		}
+		else
+		{
+			const int32 DuplicateBoneIndex = *FoundBoneIndexPtr;
+			DuplicateBones.Add(DuplicateBoneIndex);
+
+			UE_LOG(LogAnimation, Warning, TEXT("RemoveDuplicateBones: duplicate bone name (%s) detected for (%s)! Indices: %d and %d. Removing the latter."),
+				*BoneName.ToString(), *GetNameSafe(Requester), DuplicateBoneIndex, BoneIndex);
+
+			// Remove duplicate bone index, which was added later as a mistake.
+			RefBonePose.RemoveAt(DuplicateBoneIndex, 1);
+			RefBoneInfo.RemoveAt(DuplicateBoneIndex, 1);
+
+			// Now we need to fix all the parent indices that pointed to bones after this in the array
+			// These must be after this point in the array.
+			for (int32 j = DuplicateBoneIndex; j < GetNum(); j++)
+			{
+				if (GetParentIndex(j) >= DuplicateBoneIndex)
+				{
+					RefBoneInfo[j].ParentIndex -= 1;
+				}
+			}
+
+			// Update entry in case problem bones were added multiple times.
+			BoneNameCheck.Add(BoneName, BoneIndex);
+
+			// We need to make sure that any bone that has this old bone as a parent is fixed up
+			bRemovedBones = true;
+		}
+	}
+
+	// If we've removed bones, we need to rebuild our name table.
+	if (bRemovedBones || (NameToIndexMap.Num() == 0))
+	{
+		RebuildNameToIndexMap();
+	}
+
+	// Make sure our arrays are in sync.
+	checkSlow((RefBoneInfo.Num() == RefBonePose.Num()) && (RefBoneInfo.Num() == NameToIndexMap.Num()));
+
+	// Additionally normalize all quaternions to be safe.
+	for (int32 BoneIndex = 0; BoneIndex < GetNum(); BoneIndex++)
+	{
+		RefBonePose[BoneIndex].NormalizeRotation();
+	}
+}
+
+void FReferenceSkeleton::RebuildNameToIndexMap()
+{
+	// Start by clearing the current map.
+	NameToIndexMap.Empty();
+
+	// Then iterate over each bone, adding the name and bone index.
+	const int32 NumBones = GetNum();
+	for (int32 BoneIndex = 0; BoneIndex < NumBones; BoneIndex++)
+	{
+		const FName & BoneName = GetBoneName(BoneIndex);
+		if (BoneName != NAME_None)
+		{
+			NameToIndexMap.Add(BoneName, BoneIndex);
+		}
+		else
+		{
+			UE_LOG(LogAnimation, Warning, TEXT("RebuildNameToIndexMap: Bone with no name detected for index: %d"), BoneIndex);
+		}
+	}
+
+	// Make sure we don't have duplicate bone names. This would be very bad.
+	checkSlow(NameToIndexMap.Num() == NumBones);
+}
+
+FArchive & operator<<(FArchive & Ar, FReferenceSkeleton & F)
+{
+	Ar << F.RefBoneInfo;
+	Ar << F.RefBonePose;
+
+	if (Ar.UE4Ver() >= VER_UE4_REFERENCE_SKELETON_REFACTOR)
+	{
+		Ar << F.NameToIndexMap;
+	}
+
+	// Fix up any assets that don't have an INDEX_NONE parent for Bone[0]
+	if (Ar.IsLoading() && Ar.UE4Ver() < VER_UE4_FIXUP_ROOTBONE_PARENT)
+	{
+		if ((F.RefBoneInfo.Num() > 0) && (F.RefBoneInfo[0].ParentIndex != INDEX_NONE))
+		{
+			F.RefBoneInfo[0].ParentIndex = INDEX_NONE;
+		}
+	}
+
+	return Ar;
+}
\ No newline at end of file
diff --git a/Engine/Source/Runtime/Engine/Private/Collision/Collision.cpp b/Engine/Source/Runtime/Engine/Private/Collision/Collision.cpp
--- a/Engine/Source/Runtime/Engine/Private/Collision/Collision.cpp
+++ b/Engine/Source/Runtime/Engine/Private/Collision/Collision.cpp
@@ -6,6 +6,19 @@
 
 #include "EnginePrivate.h"
 
+//////////////////////////////////////////////////////////////////////////
+// FHitResult
+
+FHitResult::FHitResult(class AActor* InActor, class UPrimitiveComponent* InComponent, FVector const& HitLoc, FVector const& HitNorm)
+{
+	FMemory::Memzero(this, sizeof(FHitResult));
+	Location = HitLoc;
+	ImpactPoint = HitLoc;
+	Normal = HitNorm;
+	ImpactNormal = HitNorm;
+	Actor = InActor;
+	Component = InComponent;
+}
 
 AActor* FHitResult::GetActor() const
 {
@@ -17,6 +30,9 @@ UPrimitiveComponent* FHitResult::GetComponent() const
 	return Component.Get();
 }
 
+//////////////////////////////////////////////////////////////////////////
+// FOverlapResult
+
 AActor* FOverlapResult::GetActor() const
 {
 	return Actor.Get();
@@ -26,3 +42,55 @@ UPrimitiveComponent* FOverlapResult::GetComponent() const
 {
 	return Component.Get();
 }
+
+
+//////////////////////////////////////////////////////////////////////////
+// FCollisionQueryParams
+
+FCollisionQueryParams::FCollisionQueryParams(FName InTraceTag, bool bInTraceComplex, const AActor* InIgnoreActor)
+{
+	bTraceComplex = bInTraceComplex;
+	TraceTag = InTraceTag;
+	bTraceAsyncScene = false;
+	bFindInitialOverlaps = true;
+	bReturnFaceIndex = false;
+	bReturnPhysicalMaterial = false;
+
+	AddIgnoredActor(InIgnoreActor);
+	if (InIgnoreActor != NULL)
+	{
+		OwnerTag = InIgnoreActor->GetFName();
+	}
+}
+
+void FCollisionQueryParams::AddIgnoredActor(const AActor* InIgnoreActor)
+{
+	if (InIgnoreActor != NULL)
+	{
+		IgnoreActors.AddUnique(InIgnoreActor->GetUniqueID());
+	}
+}
+
+void FCollisionQueryParams::AddIgnoredActors(const TArray<AActor*>& InIgnoreActors)
+{
+	for (int32 Idx = 0; Idx < InIgnoreActors.Num(); ++Idx)
+	{
+		AActor const* const A = InIgnoreActors[Idx];
+		if (A)
+		{
+			IgnoreActors.Add(A->GetUniqueID());
+		}
+	}
+}
+
+void FCollisionQueryParams::AddIgnoredActors(const TArray<TWeakObjectPtr<AActor> >& InIgnoreActors)
+{
+	for (int32 Idx = 0; Idx < InIgnoreActors.Num(); ++Idx)
+	{
+		AActor const* const A = InIgnoreActors[Idx].Get();
+		if (A)
+		{
+			IgnoreActors.Add(A->GetUniqueID());
+		}
+	}
+}
\ No newline at end of file
diff --git a/Engine/Source/Runtime/Engine/Private/Collision/WorldCollision.cpp b/Engine/Source/Runtime/Engine/Private/Collision/WorldCollision.cpp
--- a/Engine/Source/Runtime/Engine/Private/Collision/WorldCollision.cpp
+++ b/Engine/Source/Runtime/Engine/Private/Collision/WorldCollision.cpp
@@ -49,6 +49,26 @@ FCollisionShape					FCollisionShape::LineShape;
 // default being the 0. That isn't invalid, but ObjectQuery param overrides this 
 ECollisionChannel DefaultCollisionChannel = (ECollisionChannel) 0;
 
+
+/* Set functions for each Shape type */
+void FBaseTraceDatum::Set(UWorld * World, const FCollisionShape & InCollisionShape, const FCollisionQueryParams & Param, const struct FCollisionResponseParams &InResponseParam, const struct FCollisionObjectQueryParams& InObjectQueryParam,
+	ECollisionChannel Channel, uint32 InUserData, bool bInIsMultiTrace, int32 FrameCounter)
+{
+	ensure(World);
+	CollisionParams.CollisionShape = InCollisionShape;
+	CollisionParams.CollisionQueryParam = Param;
+	CollisionParams.ResponseParam = InResponseParam;
+	CollisionParams.ObjectQueryParam = InObjectQueryParam;
+	TraceChannel = Channel;
+	UserData = InUserData;
+	bIsMultiTrace = bInIsMultiTrace;
+	FrameNumber = FrameCounter;
+	PhysWorld = World;
+}
+
+
+//////////////////////////////////////////////////////////////////////////
+
 bool UWorld::LineTraceTest(const FVector& Start, const FVector& End, ECollisionChannel TraceChannel, const struct FCollisionQueryParams& Params, const struct FCollisionResponseParams& ResponseParam) const
 {
 #if UE_WITH_PHYSICS
diff --git a/Engine/Source/Runtime/Engine/Private/SceneComponent.cpp b/Engine/Source/Runtime/Engine/Private/SceneComponent.cpp
--- a/Engine/Source/Runtime/Engine/Private/SceneComponent.cpp
+++ b/Engine/Source/Runtime/Engine/Private/SceneComponent.cpp
@@ -15,6 +15,13 @@
 
 DEFINE_LOG_CATEGORY_STATIC(LogSceneComponent, Log, All);
 
+FOverlapInfo::FOverlapInfo(UPrimitiveComponent* InComponent, int32 InBodyIndex)
+	: bFromSweep(false)
+{
+	OverlapInfo.Component = InComponent;
+	OverlapInfo.Item = InBodyIndex;
+}
+
 USceneComponent::USceneComponent(const class FPostConstructInitializeProperties& PCIP)
 	: Super(PCIP)
 {
diff --git a/Engine/Source/Runtime/Engine/Private/SkeletalMesh.cpp b/Engine/Source/Runtime/Engine/Private/SkeletalMesh.cpp
--- a/Engine/Source/Runtime/Engine/Private/SkeletalMesh.cpp
+++ b/Engine/Source/Runtime/Engine/Private/SkeletalMesh.cpp
@@ -667,54 +667,9 @@ FArchive& operator<<( FArchive& Ar, FSkeletalMeshVertexAPEXClothBuffer& VertexBu
 	return Ar;
 }
 
-FArchive &operator<<( FArchive& Ar, FMeshBoneInfo& F)
-{
-	Ar << F.Name << F.ParentIndex;
 
-	if( Ar.IsLoading() && (Ar.UE4Ver() < VER_UE4_REFERENCE_SKELETON_REFACTOR) )
-	{
-		FColor DummyColor = FColor::White;
-		Ar << DummyColor;
-	}
 
-#if WITH_EDITORONLY_DATA
-	if (Ar.UE4Ver() >= VER_UE4_STORE_BONE_EXPORT_NAMES)
-	{
-		if(!Ar.IsCooking())
-		{
-			Ar << F.ExportName;
-		}
-	}
-	else
-	{
-		F.ExportName = F.Name.ToString();
-	}
-#endif
 
-	return Ar;
-}
-
-FArchive & operator<<(FArchive & Ar, FReferenceSkeleton & F)
-{
-	Ar << F.RefBoneInfo;
-	Ar << F.RefBonePose;
-
-	if( Ar.UE4Ver() >= VER_UE4_REFERENCE_SKELETON_REFACTOR )
-	{
-		Ar << F.NameToIndexMap;
-	}
-
-	// Fix up any assets that don't have an INDEX_NONE parent for Bone[0]
-	if( Ar.IsLoading() && Ar.UE4Ver() < VER_UE4_FIXUP_ROOTBONE_PARENT )
-	{
-		if( (F.RefBoneInfo.Num() > 0) && (F.RefBoneInfo[0].ParentIndex != INDEX_NONE) )
-		{
-			F.RefBoneInfo[0].ParentIndex = INDEX_NONE;
-		}
-	}
-
-	return Ar;
-}
 
 
 /**
diff --git a/Engine/Source/Runtime/Engine/Private/Texture.cpp b/Engine/Source/Runtime/Engine/Private/Texture.cpp
--- a/Engine/Source/Runtime/Engine/Private/Texture.cpp
+++ b/Engine/Source/Runtime/Engine/Private/Texture.cpp
@@ -508,309 +508,7 @@ bool UTexture::ForceUpdateTextureStreaming()
 	return true;
 }
 
-void FTextureLODSettings::Initialize( const FString& IniFilename, const TCHAR* IniSection )
-{
-	// look up the file object
-	FConfigFile* ConfigFile = GConfig->FindConfigFile(IniFilename);
-	if (ConfigFile)
-	{
-		// pass to the other initialize function
-		Initialize(*ConfigFile, IniSection);
-	}
-}
-
-void FTextureLODSettings::Initialize(const FConfigFile& IniFile, const TCHAR* IniSection)
-{
-	// Read individual entries from a config file.
-#define GROUPREADENTRY(g) ReadEntry( g, TEXT(#g), IniFile, IniSection );
-	FOREACH_ENUM_TEXTUREGROUP(GROUPREADENTRY)
-#undef GROUPREADENTRY
-}
-
-/**
- * Returns the texture group names, sorted like enum.
- *
- * @return array of texture group names
- */
-TArray<FString> FTextureLODSettings::GetTextureGroupNames()
-{
-	TArray<FString> TextureGroupNames;
-
-#define GROUPNAMES(g) new(TextureGroupNames) FString(TEXT(#g));
-	FOREACH_ENUM_TEXTUREGROUP(GROUPNAMES)
-#undef GROUPNAMES
-
-	return TextureGroupNames;
-}
-
-void FTextureLODSettings::ReadEntry( int32 GroupId, const TCHAR* GroupName, const FConfigFile& IniFile, const TCHAR* IniSection )
-{
-	// Look for string in filename/ section.
-	FString Entry;
-	if (IniFile.GetString(IniSection, GroupName, Entry))
-	{
-		// Trim whitespace at the beginning.
-		Entry = Entry.Trim();
-		// Remove brackets.
-		Entry = Entry.Replace( TEXT("("), TEXT("") );
-		Entry = Entry.Replace( TEXT(")"), TEXT("") );
-		
-		// Parse minimum LOD mip count.
-		int32	MinLODSize = 0;
-		if( FParse::Value( *Entry, TEXT("MinLODSize="), MinLODSize ) )
-		{
-			TextureLODGroups[GroupId].MinLODMipCount = FMath::CeilLogTwo( MinLODSize );
-		}
-
-		// Parse maximum LOD mip count.
-		int32 MaxLODSize = 0;
-		if( FParse::Value( *Entry, TEXT("MaxLODSize="), MaxLODSize ) )
-		{
-			TextureLODGroups[GroupId].MaxLODMipCount = FMath::CeilLogTwo( MaxLODSize );
-		}
-
-		// Parse LOD bias.
-		int32 LODBias = 0;
-		if( FParse::Value( *Entry, TEXT("LODBias="), LODBias ) )
-		{
-			TextureLODGroups[GroupId].LODBias = LODBias;
-		}
-
-		// Parse min/map/mip filter names.
-		FName MinMagFilter = NAME_Aniso;
-		FParse::Value( *Entry, TEXT("MinMagFilter="), MinMagFilter );
-		FName MipFilter = NAME_Point;
-		FParse::Value( *Entry, TEXT("MipFilter="), MipFilter );
-
-		{
-			FString MipGenSettings;
-			FParse::Value( *Entry, TEXT("MipGenSettings="), MipGenSettings );
-			TextureLODGroups[GroupId].MipGenSettings = UTexture::GetMipGenSettingsFromString(*MipGenSettings, true);
-		}
-
-		// Convert into single filter enum. The code is layed out such that invalid input will 
-		// map to the default state of highest quality filtering.
-
-		// Linear filtering
-		if( MinMagFilter == NAME_Linear )
-		{
-			if( MipFilter == NAME_Point )
-			{
-				TextureLODGroups[GroupId].Filter = SF_Bilinear;
-			}
-			else
-			{
-				TextureLODGroups[GroupId].Filter = SF_Trilinear;
-			}
-		}
-		// Point. Don't even care about mip filter.
-		else if( MinMagFilter == NAME_Point )
-		{
-			TextureLODGroups[GroupId].Filter = SF_Point;
-		}
-		// Aniso or unknown.
-		else
-		{
-			if( MipFilter == NAME_Point )
-			{
-				TextureLODGroups[GroupId].Filter = SF_AnisotropicPoint;
-			}
-			else
-			{
-				TextureLODGroups[GroupId].Filter = SF_AnisotropicLinear;
-			}
-		}
-
-		// Parse NumStreamedMips
-		int32 NumStreamedMips = -1;
-		if( FParse::Value( *Entry, TEXT("NumStreamedMips="), NumStreamedMips ) )
-		{
-			TextureLODGroups[GroupId].NumStreamedMips = NumStreamedMips;
-		}
-	}
-}
-
-int32 FTextureLODSettings::CalculateLODBias( const UTexture* Texture, bool bIncTextureMips ) const
-{	
-	check( Texture );
-	TextureMipGenSettings MipGenSetting = TMGS_MAX;
-#if WITH_EDITORONLY_DATA
-	MipGenSetting = Texture->MipGenSettings;
-#endif // #if WITH_EDITORONLY_DATA
-	return CalculateLODBias(Texture->GetSurfaceWidth(), Texture->GetSurfaceHeight(), Texture->LODGroup, (bIncTextureMips ? Texture->LODBias : 0), (bIncTextureMips ? Texture->NumCinematicMipLevels : 0), MipGenSetting);
-}
-
-int32 FTextureLODSettings::CalculateLODBias( int32 Width, int32 Height, int32 LODGroup, int32 LODBias, int32 NumCinematicMipLevels, TextureMipGenSettings InMipGenSetting ) const
-{	
-	// Find LOD group.
-	const FTextureLODGroup& LODGroupInfo = TextureLODGroups[LODGroup];
-
-	// Test to see if we have no mip generation as in which case the LOD bias will be ignored
-	const TextureMipGenSettings FinalMipGenSetting = (InMipGenSetting == TMGS_FromTextureGroup) ? LODGroupInfo.MipGenSettings : InMipGenSetting;
-	if ( FinalMipGenSetting == TMGS_NoMipmaps )
-	{
-		return 0;
-	}
-
-	// Calculate maximum number of miplevels.
-	int32 TextureMaxLOD	= FMath::CeilLogTwo( FMath::TruncToInt( FMath::Max( Width, Height ) ) );
-
-	// Calculate LOD bias.
-	int32 UsedLODBias	= LODGroupInfo.LODBias + LODBias + NumCinematicMipLevels;
-	int32 MinLOD		= LODGroupInfo.MinLODMipCount;
-	int32 MaxLOD		= LODGroupInfo.MaxLODMipCount;
-	int32 WantedMaxLOD	= FMath::Clamp( TextureMaxLOD - UsedLODBias, MinLOD, MaxLOD );
-	WantedMaxLOD		= FMath::Clamp( WantedMaxLOD, 0, TextureMaxLOD );
-	UsedLODBias			= TextureMaxLOD - WantedMaxLOD;
-
-	return UsedLODBias;
-}
-
-/** 
-* Useful for stats in the editor.
-*/
-void FTextureLODSettings::ComputeInGameMaxResolution(int32 LODBias, UTexture &Texture, uint32 &OutSizeX, uint32 &OutSizeY) const
-{
-	uint32 ImportedSizeX = FMath::TruncToInt(Texture.GetSurfaceWidth());
-	uint32 ImportedSizeY = FMath::TruncToInt(Texture.GetSurfaceHeight());
-	
-	const FTextureLODGroup& LODGroup = GetTextureLODGroup((TextureGroup)Texture.LODGroup);
-
-	uint32 SourceLOD = FMath::Max(FMath::CeilLogTwo(ImportedSizeX), FMath::CeilLogTwo(ImportedSizeY));
-	uint32 MinLOD = FMath::Max(uint32(UTexture2D::GetMinTextureResidentMipCount() - 1), (uint32)LODGroup.MinLODMipCount);
-	uint32 MaxLOD = FMath::Min(uint32(GMaxTextureMipCount - 1), (uint32)LODGroup.MaxLODMipCount);
-	uint32 GameLOD = FMath::Min(SourceLOD, FMath::Clamp(SourceLOD - LODBias, MinLOD, MaxLOD));
-
-	uint32 DeltaLOD = SourceLOD - GameLOD;
-
-	OutSizeX = ImportedSizeX >> DeltaLOD;
-	OutSizeY = ImportedSizeY >> DeltaLOD;
-}
-
-/**
-* TextureLODGroups access with bounds check
-*
-* @param   GroupIndex      usually from Texture.LODGroup
-* @return                  A handle to the indexed LOD group. 
-*/
-const FTextureLODSettings::FTextureLODGroup& FTextureLODSettings::GetTextureLODGroup(TextureGroup GroupIndex) const
-{
-	if((uint32)GroupIndex >= TEXTUREGROUP_MAX)
-	{
-		// to prevent crash
-		GroupIndex = (TextureGroup)0;
-	}
-//	check((uint32)GroupIndex < TEXTUREGROUP_MAX);
-	return TextureLODGroups[GroupIndex];
-}
-
-#if WITH_EDITORONLY_DATA
-void FTextureLODSettings::GetMipGenSettings( const UTexture& Texture, TextureMipGenSettings& OutMipGenSettings, float& OutSharpen, uint32& OutKernelSize, bool& bOutDownsampleWithAverage, bool& bOutSharpenWithoutColorShift, bool &bOutBorderColorBlack ) const
-{
-	TextureMipGenSettings Setting = (TextureMipGenSettings)Texture.MipGenSettings;
-
-	bOutBorderColorBlack = false;
-
-	// avoiding the color shift assumes we deal with colors which is not true for normalmaps
-	// or we blur where it's good to blur the color as well
-	bOutSharpenWithoutColorShift = !Texture.IsNormalMap();
-
-	bOutDownsampleWithAverage = true;
-
-	// inherit from texture group
-	if(Setting == TMGS_FromTextureGroup)
-	{
-		const FTextureLODGroup& LODGroup = TextureLODGroups[Texture.LODGroup];
-
-		Setting = LODGroup.MipGenSettings;
-	}
-	OutMipGenSettings = Setting;
-
-	// ------------
 
-	// default:
-	OutSharpen = 0;
-	OutKernelSize = 2;
-
-	if(Setting >= TMGS_Sharpen0 && Setting <= TMGS_Sharpen10)
-	{
-		// 0 .. 2.0f
-		OutSharpen = ((int32)Setting - (int32)TMGS_Sharpen0) * 0.2f;
-		OutKernelSize = 8;
-	}
-	else if(Setting >= TMGS_Blur1 && Setting <= TMGS_Blur5)
-	{
-		int32 BlurFactor = ((int32)Setting + 1 - (int32)TMGS_Blur1);
-		OutSharpen = -BlurFactor * 2;
-		OutKernelSize = 2 + 2 * BlurFactor;
-		bOutDownsampleWithAverage = false;
-		bOutSharpenWithoutColorShift = false;
-		bOutBorderColorBlack = true;
-	}
-}
-#endif // #if WITH_EDITORONLY_DATA
-
-/**
- * Will return the LODBias for a passed in LODGroup
- *
- * @param	InLODGroup		The LOD Group ID 
- * @return	LODBias
- */
-int32 FTextureLODSettings::GetTextureLODGroupLODBias( int32 InLODGroup ) const
-{
-	int32 Retval = 0;
-
-	const FTextureLODGroup& LODGroup = TextureLODGroups[InLODGroup]; 
-
-	Retval = LODGroup.LODBias;
-
-	return Retval;
-}
-
-/**
- * Returns the LODGroup setting for number of streaming mip-levels.
- * -1 means that all mip-levels are allowed to stream.
- *
- * @param	InLODGroup		The LOD Group ID 
- * @return	Number of streaming mip-levels for textures in the specified LODGroup
- */
-int32 FTextureLODSettings::GetNumStreamedMips( int32 InLODGroup ) const
-{
-	int32 Retval = 0;
-
-	const FTextureLODGroup& LODGroup = TextureLODGroups[InLODGroup]; 
-
-	Retval = LODGroup.NumStreamedMips;
-
-	return Retval;
-}
-
-/**
- * Returns the filter state that should be used for the passed in texture, taking
- * into account other system settings.
- *
- * @param	Texture		Texture to retrieve filter state for, must not be 0
- * @return	Filter sampler state for passed in texture
- */
-ESamplerFilter FTextureLODSettings::GetSamplerFilter( const UTexture* Texture ) const
-{
-	// Default to point filtering.
-	ESamplerFilter Filter = SF_Point;
-
-	switch(Texture->Filter)
-	{
-		case TF_Nearest: Filter = SF_Point; break;
-		case TF_Bilinear: Filter = SF_Bilinear; break;
-		case TF_Trilinear: Filter = SF_Trilinear; break;
-
-		// TF_Default
-		default:
-			// Use LOD group value to find proper filter setting.
-			Filter = TextureLODGroups[Texture->LODGroup].Filter;
-	}
-
-	return Filter;
-}
 
 /*------------------------------------------------------------------------------
 	Texture source data.
diff --git a/Engine/Source/Runtime/Engine/Private/TextureLODSettings.cpp b/Engine/Source/Runtime/Engine/Private/TextureLODSettings.cpp
new file mode 100644
--- /dev/null
+++ b/Engine/Source/Runtime/Engine/Private/TextureLODSettings.cpp
@@ -0,0 +1,308 @@
+// Copyright 1998-2014 Epic Games, Inc. All Rights Reserved.
+
+#include "EnginePrivate.h"
+#include "TextureLODSettings.h"
+
+void FTextureLODSettings::Initialize( const FString& IniFilename, const TCHAR* IniSection )
+{
+	// look up the file object
+	FConfigFile* ConfigFile = GConfig->FindConfigFile(IniFilename);
+	if (ConfigFile)
+	{
+		// pass to the other initialize function
+		Initialize(*ConfigFile, IniSection);
+	}
+}
+
+void FTextureLODSettings::Initialize(const FConfigFile& IniFile, const TCHAR* IniSection)
+{
+	// Read individual entries from a config file.
+#define GROUPREADENTRY(g) ReadEntry( g, TEXT(#g), IniFile, IniSection );
+	FOREACH_ENUM_TEXTUREGROUP(GROUPREADENTRY)
+#undef GROUPREADENTRY
+}
+
+/**
+ * Returns the texture group names, sorted like enum.
+ *
+ * @return array of texture group names
+ */
+TArray<FString> FTextureLODSettings::GetTextureGroupNames()
+{
+	TArray<FString> TextureGroupNames;
+
+#define GROUPNAMES(g) new(TextureGroupNames) FString(TEXT(#g));
+	FOREACH_ENUM_TEXTUREGROUP(GROUPNAMES)
+#undef GROUPNAMES
+
+	return TextureGroupNames;
+}
+
+void FTextureLODSettings::ReadEntry( int32 GroupId, const TCHAR* GroupName, const FConfigFile& IniFile, const TCHAR* IniSection )
+{
+	// Look for string in filename/ section.
+	FString Entry;
+	if (IniFile.GetString(IniSection, GroupName, Entry))
+	{
+		// Trim whitespace at the beginning.
+		Entry = Entry.Trim();
+		// Remove brackets.
+		Entry = Entry.Replace( TEXT("("), TEXT("") );
+		Entry = Entry.Replace( TEXT(")"), TEXT("") );
+		
+		// Parse minimum LOD mip count.
+		int32	MinLODSize = 0;
+		if( FParse::Value( *Entry, TEXT("MinLODSize="), MinLODSize ) )
+		{
+			TextureLODGroups[GroupId].MinLODMipCount = FMath::CeilLogTwo( MinLODSize );
+		}
+
+		// Parse maximum LOD mip count.
+		int32 MaxLODSize = 0;
+		if( FParse::Value( *Entry, TEXT("MaxLODSize="), MaxLODSize ) )
+		{
+			TextureLODGroups[GroupId].MaxLODMipCount = FMath::CeilLogTwo( MaxLODSize );
+		}
+
+		// Parse LOD bias.
+		int32 LODBias = 0;
+		if( FParse::Value( *Entry, TEXT("LODBias="), LODBias ) )
+		{
+			TextureLODGroups[GroupId].LODBias = LODBias;
+		}
+
+		// Parse min/map/mip filter names.
+		FName MinMagFilter = NAME_Aniso;
+		FParse::Value( *Entry, TEXT("MinMagFilter="), MinMagFilter );
+		FName MipFilter = NAME_Point;
+		FParse::Value( *Entry, TEXT("MipFilter="), MipFilter );
+
+		{
+			FString MipGenSettings;
+			FParse::Value( *Entry, TEXT("MipGenSettings="), MipGenSettings );
+			TextureLODGroups[GroupId].MipGenSettings = UTexture::GetMipGenSettingsFromString(*MipGenSettings, true);
+		}
+
+		// Convert into single filter enum. The code is layed out such that invalid input will 
+		// map to the default state of highest quality filtering.
+
+		// Linear filtering
+		if( MinMagFilter == NAME_Linear )
+		{
+			if( MipFilter == NAME_Point )
+			{
+				TextureLODGroups[GroupId].Filter = SF_Bilinear;
+			}
+			else
+			{
+				TextureLODGroups[GroupId].Filter = SF_Trilinear;
+			}
+		}
+		// Point. Don't even care about mip filter.
+		else if( MinMagFilter == NAME_Point )
+		{
+			TextureLODGroups[GroupId].Filter = SF_Point;
+		}
+		// Aniso or unknown.
+		else
+		{
+			if( MipFilter == NAME_Point )
+			{
+				TextureLODGroups[GroupId].Filter = SF_AnisotropicPoint;
+			}
+			else
+			{
+				TextureLODGroups[GroupId].Filter = SF_AnisotropicLinear;
+			}
+		}
+
+		// Parse NumStreamedMips
+		int32 NumStreamedMips = -1;
+		if( FParse::Value( *Entry, TEXT("NumStreamedMips="), NumStreamedMips ) )
+		{
+			TextureLODGroups[GroupId].NumStreamedMips = NumStreamedMips;
+		}
+	}
+}
+
+int32 FTextureLODSettings::CalculateLODBias( const UTexture* Texture, bool bIncTextureMips ) const
+{	
+	check( Texture );
+	TextureMipGenSettings MipGenSetting = TMGS_MAX;
+#if WITH_EDITORONLY_DATA
+	MipGenSetting = Texture->MipGenSettings;
+#endif // #if WITH_EDITORONLY_DATA
+	return CalculateLODBias(Texture->GetSurfaceWidth(), Texture->GetSurfaceHeight(), Texture->LODGroup, (bIncTextureMips ? Texture->LODBias : 0), (bIncTextureMips ? Texture->NumCinematicMipLevels : 0), MipGenSetting);
+}
+
+int32 FTextureLODSettings::CalculateLODBias( int32 Width, int32 Height, int32 LODGroup, int32 LODBias, int32 NumCinematicMipLevels, TextureMipGenSettings InMipGenSetting ) const
+{	
+	// Find LOD group.
+	const FTextureLODGroup& LODGroupInfo = TextureLODGroups[LODGroup];
+
+	// Test to see if we have no mip generation as in which case the LOD bias will be ignored
+	const TextureMipGenSettings FinalMipGenSetting = (InMipGenSetting == TMGS_FromTextureGroup) ? LODGroupInfo.MipGenSettings : InMipGenSetting;
+	if ( FinalMipGenSetting == TMGS_NoMipmaps )
+	{
+		return 0;
+	}
+
+	// Calculate maximum number of miplevels.
+	int32 TextureMaxLOD	= FMath::CeilLogTwo( FMath::TruncToInt( FMath::Max( Width, Height ) ) );
+
+	// Calculate LOD bias.
+	int32 UsedLODBias	= LODGroupInfo.LODBias + LODBias + NumCinematicMipLevels;
+	int32 MinLOD		= LODGroupInfo.MinLODMipCount;
+	int32 MaxLOD		= LODGroupInfo.MaxLODMipCount;
+	int32 WantedMaxLOD	= FMath::Clamp( TextureMaxLOD - UsedLODBias, MinLOD, MaxLOD );
+	WantedMaxLOD		= FMath::Clamp( WantedMaxLOD, 0, TextureMaxLOD );
+	UsedLODBias			= TextureMaxLOD - WantedMaxLOD;
+
+	return UsedLODBias;
+}
+
+/** 
+* Useful for stats in the editor.
+*/
+void FTextureLODSettings::ComputeInGameMaxResolution(int32 LODBias, UTexture &Texture, uint32 &OutSizeX, uint32 &OutSizeY) const
+{
+	uint32 ImportedSizeX = FMath::TruncToInt(Texture.GetSurfaceWidth());
+	uint32 ImportedSizeY = FMath::TruncToInt(Texture.GetSurfaceHeight());
+	
+	const FTextureLODGroup& LODGroup = GetTextureLODGroup((TextureGroup)Texture.LODGroup);
+
+	uint32 SourceLOD = FMath::Max(FMath::CeilLogTwo(ImportedSizeX), FMath::CeilLogTwo(ImportedSizeY));
+	uint32 MinLOD = FMath::Max(uint32(UTexture2D::GetMinTextureResidentMipCount() - 1), (uint32)LODGroup.MinLODMipCount);
+	uint32 MaxLOD = FMath::Min(uint32(GMaxTextureMipCount - 1), (uint32)LODGroup.MaxLODMipCount);
+	uint32 GameLOD = FMath::Min(SourceLOD, FMath::Clamp(SourceLOD - LODBias, MinLOD, MaxLOD));
+
+	uint32 DeltaLOD = SourceLOD - GameLOD;
+
+	OutSizeX = ImportedSizeX >> DeltaLOD;
+	OutSizeY = ImportedSizeY >> DeltaLOD;
+}
+
+/**
+* TextureLODGroups access with bounds check
+*
+* @param   GroupIndex      usually from Texture.LODGroup
+* @return                  A handle to the indexed LOD group. 
+*/
+const FTextureLODSettings::FTextureLODGroup& FTextureLODSettings::GetTextureLODGroup(TextureGroup GroupIndex) const
+{
+	if((uint32)GroupIndex >= TEXTUREGROUP_MAX)
+	{
+		// to prevent crash
+		GroupIndex = (TextureGroup)0;
+	}
+//	check((uint32)GroupIndex < TEXTUREGROUP_MAX);
+	return TextureLODGroups[GroupIndex];
+}
+
+#if WITH_EDITORONLY_DATA
+void FTextureLODSettings::GetMipGenSettings( const UTexture& Texture, TextureMipGenSettings& OutMipGenSettings, float& OutSharpen, uint32& OutKernelSize, bool& bOutDownsampleWithAverage, bool& bOutSharpenWithoutColorShift, bool &bOutBorderColorBlack ) const
+{
+	TextureMipGenSettings Setting = (TextureMipGenSettings)Texture.MipGenSettings;
+
+	bOutBorderColorBlack = false;
+
+	// avoiding the color shift assumes we deal with colors which is not true for normalmaps
+	// or we blur where it's good to blur the color as well
+	bOutSharpenWithoutColorShift = !Texture.IsNormalMap();
+
+	bOutDownsampleWithAverage = true;
+
+	// inherit from texture group
+	if(Setting == TMGS_FromTextureGroup)
+	{
+		const FTextureLODGroup& LODGroup = TextureLODGroups[Texture.LODGroup];
+
+		Setting = LODGroup.MipGenSettings;
+	}
+	OutMipGenSettings = Setting;
+
+	// ------------
+
+	// default:
+	OutSharpen = 0;
+	OutKernelSize = 2;
+
+	if(Setting >= TMGS_Sharpen0 && Setting <= TMGS_Sharpen10)
+	{
+		// 0 .. 2.0f
+		OutSharpen = ((int32)Setting - (int32)TMGS_Sharpen0) * 0.2f;
+		OutKernelSize = 8;
+	}
+	else if(Setting >= TMGS_Blur1 && Setting <= TMGS_Blur5)
+	{
+		int32 BlurFactor = ((int32)Setting + 1 - (int32)TMGS_Blur1);
+		OutSharpen = -BlurFactor * 2;
+		OutKernelSize = 2 + 2 * BlurFactor;
+		bOutDownsampleWithAverage = false;
+		bOutSharpenWithoutColorShift = false;
+		bOutBorderColorBlack = true;
+	}
+}
+#endif // #if WITH_EDITORONLY_DATA
+
+/**
+ * Will return the LODBias for a passed in LODGroup
+ *
+ * @param	InLODGroup		The LOD Group ID 
+ * @return	LODBias
+ */
+int32 FTextureLODSettings::GetTextureLODGroupLODBias( int32 InLODGroup ) const
+{
+	int32 Retval = 0;
+
+	const FTextureLODGroup& LODGroup = TextureLODGroups[InLODGroup]; 
+
+	Retval = LODGroup.LODBias;
+
+	return Retval;
+}
+
+/**
+ * Returns the LODGroup setting for number of streaming mip-levels.
+ * -1 means that all mip-levels are allowed to stream.
+ *
+ * @param	InLODGroup		The LOD Group ID 
+ * @return	Number of streaming mip-levels for textures in the specified LODGroup
+ */
+int32 FTextureLODSettings::GetNumStreamedMips( int32 InLODGroup ) const
+{
+	int32 Retval = 0;
+
+	const FTextureLODGroup& LODGroup = TextureLODGroups[InLODGroup]; 
+
+	Retval = LODGroup.NumStreamedMips;
+
+	return Retval;
+}
+
+/**
+ * Returns the filter state that should be used for the passed in texture, taking
+ * into account other system settings.
+ *
+ * @param	Texture		Texture to retrieve filter state for, must not be 0
+ * @return	Filter sampler state for passed in texture
+ */
+ESamplerFilter FTextureLODSettings::GetSamplerFilter( const UTexture* Texture ) const
+{
+	// Default to point filtering.
+	ESamplerFilter Filter = SF_Point;
+
+	switch(Texture->Filter)
+	{
+		case TF_Nearest: Filter = SF_Point; break;
+		case TF_Bilinear: Filter = SF_Bilinear; break;
+		case TF_Trilinear: Filter = SF_Trilinear; break;
+
+		// TF_Default
+		default:
+			// Use LOD group value to find proper filter setting.
+			Filter = TextureLODGroups[Texture->LODGroup].Filter;
+	}
+
+	return Filter;
+}
\ No newline at end of file
diff --git a/Engine/Source/Runtime/Engine/Public/BatchedElements.h b/Engine/Source/Runtime/Engine/Public/BatchedElements.h
--- a/Engine/Source/Runtime/Engine/Public/BatchedElements.h
+++ b/Engine/Source/Runtime/Engine/Public/BatchedElements.h
@@ -8,6 +8,7 @@
 #define _INC_BATCHEDELEMENTS
 
 #include "StaticBoundShaderState.h"
+#include "HitProxies.h"
 
 /** Blend modes supported for simple element rendering */
 enum ESimpleElementBlendMode
diff --git a/Engine/Source/Runtime/Engine/Public/CollisionQueryParams.h b/Engine/Source/Runtime/Engine/Public/CollisionQueryParams.h
--- a/Engine/Source/Runtime/Engine/Public/CollisionQueryParams.h
+++ b/Engine/Source/Runtime/Engine/Public/CollisionQueryParams.h
@@ -4,6 +4,8 @@
 
 #pragma once
 
+#include "Engine/EngineTypes.h"
+
 /** Macro to convert ECollisionChannels to bit flag **/
 #define ECC_TO_BITFIELD(x)	(1<<(x))
 /** Macro to convert from CollisionResponseContainer to bit flag **/
@@ -47,58 +49,18 @@ struct ENGINE_API FCollisionQueryParams
 		bReturnPhysicalMaterial = false;
 	}
 
-	FCollisionQueryParams(FName InTraceTag, bool bInTraceComplex=false, const AActor* InIgnoreActor=NULL)
-	{
-		bTraceComplex = bInTraceComplex;
-		TraceTag = InTraceTag;
-		bTraceAsyncScene = false;
-		bFindInitialOverlaps = true;
-		bReturnFaceIndex = false;
-		bReturnPhysicalMaterial = false;
-
-		AddIgnoredActor(InIgnoreActor);
-		if(InIgnoreActor != NULL)
-		{
-			OwnerTag = InIgnoreActor->GetFName();
-		}
-	}
+	FCollisionQueryParams(FName InTraceTag, bool bInTraceComplex=false, const AActor* InIgnoreActor=NULL);
 
 	// Utils
 
 	/** Add an actor for this trace to ignore */
-	void AddIgnoredActor(const AActor* InIgnoreActor)
-	{
-		if(InIgnoreActor != NULL)
-		{
-			IgnoreActors.AddUnique(InIgnoreActor->GetUniqueID());
-		}
-	}
+	void AddIgnoredActor(const AActor* InIgnoreActor);
 
 	/** Add a collection of actors for this trace to ignore */
-	void AddIgnoredActors(const TArray<AActor*>& InIgnoreActors)
-	{
-		for (int32 Idx=0; Idx<InIgnoreActors.Num(); ++Idx)
-		{
-			AActor const* const A = InIgnoreActors[Idx];
-			if (A)
-			{
-				IgnoreActors.Add(A->GetUniqueID());
-			}
-		}
-	}
+	void AddIgnoredActors(const TArray<AActor*>& InIgnoreActors);
 
 	/** Variant that uses an array of TWeakObjectPtrs */
-	void AddIgnoredActors(const TArray<TWeakObjectPtr<AActor> >& InIgnoreActors)
-	{
-		for (int32 Idx=0; Idx<InIgnoreActors.Num(); ++Idx)
-		{
-			AActor const* const A = InIgnoreActors[Idx].Get();
-			if (A)
-			{
-				IgnoreActors.Add(A->GetUniqueID());
-			}
-		}
-	}
+	void AddIgnoredActors(const TArray<TWeakObjectPtr<AActor> >& InIgnoreActors);
 
 	FString ToString() const
 	{
diff --git a/Engine/Source/Runtime/Engine/Public/ReferenceSkeleton.h b/Engine/Source/Runtime/Engine/Public/ReferenceSkeleton.h
new file mode 100644
--- /dev/null
+++ b/Engine/Source/Runtime/Engine/Public/ReferenceSkeleton.h
@@ -0,0 +1,294 @@
+// Copyright 1998-2014 Epic Games, Inc. All Rights Reserved.
+
+#pragma once
+
+
+// This contains Reference-skeleton related info
+// Bone transform is saved as FTransform array
+struct FMeshBoneInfo
+{
+	// Bone's name.
+	FName Name;
+
+#if WITH_EDITORONLY_DATA
+	// Name used for export (this should be exact as FName may mess with case) 
+	FString ExportName;
+#endif
+
+	// 0/NULL if this is the root bone. 
+	int32 ParentIndex;
+
+	FMeshBoneInfo() : Name(NAME_None), ParentIndex(INDEX_NONE) {}
+
+	FMeshBoneInfo(const FName & InName, const FString& InExportName, int32 InParentIndex)
+		: Name(InName)
+#if WITH_EDITORONLY_DATA
+		, ExportName(InExportName)
+#endif
+		, ParentIndex(InParentIndex)
+	{}
+
+	FMeshBoneInfo(const FMeshBoneInfo & Other)
+		: Name(Other.Name)
+#if WITH_EDITORONLY_DATA
+		, ExportName(Other.ExportName)
+#endif
+		, ParentIndex(Other.ParentIndex)
+	{}
+
+	bool operator==(const FMeshBoneInfo& B) const
+	{
+		return(Name == B.Name);
+	}
+
+	friend FArchive &operator<<(FArchive& Ar, FMeshBoneInfo& F);
+};
+
+/** Reference Skeleton **/
+struct FReferenceSkeleton
+{
+private:
+	/** Reference bone related info to be serialized **/
+	TArray<FMeshBoneInfo>	RefBoneInfo;
+	/** Reference bone transform **/
+	TArray<FTransform>		RefBonePose;
+	/** TMap to look up bone index from bone name. */
+	TMap<FName, int32>		NameToIndexMap;
+
+	/** Removes the specified bone, so long as it has no children. Returns whether we removed the bone or not */
+	bool RemoveIndividualBone(int32 BoneIndex, TArray<int32>& OutBonesRemoved)
+	{
+		bool bRemoveThisBone = true;
+
+		// Make sure we have no children
+		for(int32 CurrBoneIndex=BoneIndex+1; CurrBoneIndex < GetNum(); CurrBoneIndex++)
+		{
+			if( RefBoneInfo[CurrBoneIndex].ParentIndex == BoneIndex )
+			{
+				bRemoveThisBone = false;
+				break;
+			}
+		}
+
+		if(bRemoveThisBone)
+		{
+			// Update parent indices of bones further through the array
+			for(int32 CurrBoneIndex=BoneIndex+1; CurrBoneIndex < GetNum(); CurrBoneIndex++)
+			{
+				FMeshBoneInfo& Bone = RefBoneInfo[CurrBoneIndex];
+				if( Bone.ParentIndex > BoneIndex )
+				{
+					Bone.ParentIndex -= 1;
+				}
+			}
+
+			OutBonesRemoved.Add(BoneIndex);
+			RefBonePose.RemoveAt(BoneIndex, 1);
+			RefBoneInfo.RemoveAt(BoneIndex, 1);
+		}
+		return bRemoveThisBone;
+	}
+
+public:
+	void Allocate(int32 Size)
+	{
+		NameToIndexMap.Empty(Size);
+		RefBoneInfo.Empty(Size);
+		RefBonePose.Empty(Size);
+	}
+
+	/** Returns number of bones in Skeleton. */
+	int32 GetNum() const
+	{
+		return RefBoneInfo.Num();
+	}
+
+	/** Accessor to private data. Const so it can't be changed recklessly. */
+	const TArray<FMeshBoneInfo> & GetRefBoneInfo() const
+	{
+		return RefBoneInfo;
+	}
+
+	/** Accessor to private data. Const so it can't be changed recklessly. */
+	const TArray<FTransform> & GetRefBonePose() const
+	{
+		return RefBonePose;
+	}
+
+	void UpdateRefPoseTransform(const int32 & BoneIndex, const FTransform & BonePose)
+	{
+		RefBonePose[BoneIndex] = BonePose;
+	}
+
+	/** Add a new bone. 
+	 * BoneName must not already exist! ParentIndex must be valid. */
+	void Add(const FMeshBoneInfo & BoneInfo, const FTransform & BonePose)
+	{
+		// Adding a bone that already exists is illegal
+		check(FindBoneIndex(BoneInfo.Name) == INDEX_NONE);
+	
+		// Make sure our arrays are in sync.
+		checkSlow( (RefBoneInfo.Num() == RefBonePose.Num()) && (RefBoneInfo.Num() == NameToIndexMap.Num()) );
+
+		const int32 BoneIndex = RefBoneInfo.Add(BoneInfo);
+		RefBonePose.Add(BonePose);
+		NameToIndexMap.Add(BoneInfo.Name, BoneIndex);
+
+		// Normalize Quaternion to be safe.
+		RefBonePose[BoneIndex].NormalizeRotation();
+
+		// Parent must be valid. Either INDEX_NONE for Root, or before us.
+		check( ((BoneIndex == 0) && (BoneInfo.ParentIndex == INDEX_NONE)) || ((BoneIndex > 0) && RefBoneInfo.IsValidIndex(BoneInfo.ParentIndex)) );
+	}
+
+	/** Insert a new bone
+	* BoneName must not already exist! ParentIndex must be valid. */
+	void Insert(int32 BoneIndex, const FMeshBoneInfo & BoneInfo, const FTransform & BonePose)
+	{
+		// Make sure our arrays are in sync.
+		checkSlow( (RefBoneInfo.Num() == RefBonePose.Num()) && (RefBoneInfo.Num() == NameToIndexMap.Num()) );
+
+		// Inserting a bone that already exists is illegal
+		check(FindBoneIndex(BoneInfo.Name) == INDEX_NONE);
+
+		// Parent must be valid. Either INDEX_NONE for Root, or before us.
+		check( ((BoneIndex == 0) && (BoneInfo.ParentIndex == INDEX_NONE)) || ((BoneIndex > 0) && RefBoneInfo.IsValidIndex(BoneInfo.ParentIndex)) );
+
+		// Make sure our bone transform is valid.
+		check( BonePose.GetRotation().IsNormalized() );
+
+		RefBoneInfo.Insert(BoneInfo, BoneIndex);
+		RefBonePose.Insert(BonePose, BoneIndex);
+		NameToIndexMap.Add(BoneInfo.Name, BoneIndex);
+
+		// Normalize Quaternion to be safe.
+		RefBonePose[BoneIndex].NormalizeRotation();
+
+		// Now we need to fix all the parent indices that pointed to bones after this in the array
+		// These must be after this point in the array.
+		for(int32 j=BoneIndex+1; j<GetNum(); j++)
+		{
+			if( GetParentIndex(j) >= BoneIndex )
+			{
+				RefBoneInfo[j].ParentIndex += 1;
+			}
+		}
+	}
+
+	void Empty()
+	{
+		RefBoneInfo.Empty();
+		RefBonePose.Empty();
+		NameToIndexMap.Empty();
+	}
+
+	/** Find Bone Index from BoneName. Precache as much as possible in speed critical sections! */
+	int32 FindBoneIndex(const FName & BoneName) const
+	{
+		checkSlow(RefBoneInfo.Num() == NameToIndexMap.Num());
+		int32 BoneIndex = INDEX_NONE;
+		if( BoneName != NAME_None )
+		{
+			const int32* IndexPtr = NameToIndexMap.Find(BoneName);
+			if( IndexPtr )
+			{
+				BoneIndex = *IndexPtr;
+			}
+		}
+		return BoneIndex;
+	}
+
+	FName GetBoneName(const int32 & BoneIndex) const
+	{
+		return RefBoneInfo[BoneIndex].Name;
+	}
+
+	int32 GetParentIndex(const int32 & BoneIndex) const
+	{
+		// Parent must be valid. Either INDEX_NONE for Root, or before us.
+		checkSlow( ((BoneIndex == 0) && (RefBoneInfo[BoneIndex].ParentIndex == INDEX_NONE)) || ((BoneIndex > 0) && RefBoneInfo.IsValidIndex(RefBoneInfo[BoneIndex].ParentIndex)) );
+		return RefBoneInfo[BoneIndex].ParentIndex;
+	}
+
+	bool IsValidIndex(int32 Index) const
+	{
+		return (RefBoneInfo.IsValidIndex(Index));
+	}
+
+	/** 
+	 * Returns # of Depth from BoneIndex to ParentBoneIndex
+	 * This will return 0 if BoneIndex == ParentBoneIndex;
+	 * This will return -1 if BoneIndex isn't child of ParentBoneIndex
+	 */
+	int32 GetDepthBetweenBones(const int32 & BoneIndex, const int32 & ParentBoneIndex) const
+	{
+		if (BoneIndex >= ParentBoneIndex)
+		{
+			int32 CurBoneIndex = BoneIndex;
+			int32 Depth = 0;
+
+			do
+			{
+				// if same return;
+				if (CurBoneIndex == ParentBoneIndex)
+				{
+					return Depth;
+				}
+
+				CurBoneIndex = RefBoneInfo[CurBoneIndex].ParentIndex;
+				++Depth;
+
+			} while (CurBoneIndex!=INDEX_NONE);
+		}
+
+		return INDEX_NONE;
+	}
+
+	bool BoneIsChildOf(const int32 & ChildBoneIndex, const int32 & ParentBoneIndex) const
+	{
+		// Bones are in strictly increasing order.
+		// So child must have an index greater than his parent.
+		if( ChildBoneIndex > ParentBoneIndex )
+		{
+			int32 BoneIndex = GetParentIndex(ChildBoneIndex);
+			do
+			{
+				if( BoneIndex == ParentBoneIndex )
+				{
+					return true;
+				}
+				BoneIndex = GetParentIndex(BoneIndex);
+
+			} while (BoneIndex != INDEX_NONE);
+		}
+
+		return false;
+	}
+
+	void RemoveDuplicateBones(const UObject * Requester, TArray<FBoneIndexType> & DuplicateBones);
+
+
+	/** Removes the supplied bones from the skeleton, unless they have children that aren't also going to be removed */
+	TArray<int32> RemoveBonesByName(const TArray<FName>& BonesToRemove)
+	{
+		TArray<int32> BonesRemoved;
+
+		const int32 NumBones = GetNum();
+		for(int32 BoneIndex=NumBones-1; BoneIndex>=0; BoneIndex--)
+		{
+			FMeshBoneInfo& Bone = RefBoneInfo[BoneIndex];
+
+			if(BonesToRemove.Contains(Bone.Name))
+			{
+				RemoveIndividualBone(BoneIndex, BonesRemoved);
+			}
+		}
+		RebuildNameToIndexMap();
+		return BonesRemoved;
+	}
+
+	void RebuildNameToIndexMap();
+
+
+	friend FArchive & operator<<(FArchive & Ar, FReferenceSkeleton & F);
+};
\ No newline at end of file
diff --git a/Engine/Source/Runtime/Engine/Public/SkeletalMeshTypes.h b/Engine/Source/Runtime/Engine/Public/SkeletalMeshTypes.h
--- a/Engine/Source/Runtime/Engine/Public/SkeletalMeshTypes.h
+++ b/Engine/Source/Runtime/Engine/Public/SkeletalMeshTypes.h
@@ -15,10 +15,13 @@
 #include "Components/PrimitiveComponent.h"
 #include "PrimitiveSceneProxy.h"
 
+typedef uint16 FBoneIndexType;
+
+#include "ReferenceSkeleton.h"
+
 // Define that controls showing chart of distance factors for skel meshes during entire run of the game on exit.
 #define CHART_DISTANCE_FACTORS 0
 
-typedef uint16 FBoneIndexType;
 
 class FRawStaticIndexBuffer16or32Interface;
 
@@ -107,376 +110,9 @@ struct VJointPos
 
 template <> struct TIsPODType<VJointPos> { enum { Value = true }; };
 
-// This contains Reference-skeleton related info
-// Bone transform is saved as FTransform array
-struct FMeshBoneInfo
-{
-	// Bone's name.
-	FName Name;			
-
-#if WITH_EDITORONLY_DATA
-	// Name used for export (this should be exact as FName may mess with case) 
-	FString ExportName;
-#endif
-
-	// 0/NULL if this is the root bone. 
-	int32 ParentIndex;	
-
-	FMeshBoneInfo() : Name(NAME_None), ParentIndex(INDEX_NONE) {}
-
-	FMeshBoneInfo(const FName & InName, const FString& InExportName, int32 InParentIndex)
-	:	Name(InName)
-#if WITH_EDITORONLY_DATA
-	, ExportName(InExportName)
-#endif
-	,	ParentIndex(InParentIndex)
-	{}
-
-	FMeshBoneInfo(const FMeshBoneInfo & Other)
-		:	Name(Other.Name)
-#if WITH_EDITORONLY_DATA
-		, ExportName(Other.ExportName)
-#endif
-		,	ParentIndex(Other.ParentIndex)
-	{}
-
-	bool operator==( const FMeshBoneInfo& B ) const
-	{
-		return( Name == B.Name );
-	}
-
-	friend FArchive &operator<<( FArchive& Ar, FMeshBoneInfo& F);
-};
-
-/** Reference Skeleton **/
-struct FReferenceSkeleton
-{
-private:
-	/** Reference bone related info to be serialized **/
-	TArray<FMeshBoneInfo>	RefBoneInfo;
-	/** Reference bone transform **/
-	TArray<FTransform>		RefBonePose;
-	/** TMap to look up bone index from bone name. */
-	TMap<FName, int32>		NameToIndexMap;
-
-	/** Removes the specified bone, so long as it has no children. Returns whether we removed the bone or not */
-	bool RemoveIndividualBone(int32 BoneIndex, TArray<int32>& OutBonesRemoved)
-	{
-		bool bRemoveThisBone = true;
-
-		// Make sure we have no children
-		for(int32 CurrBoneIndex=BoneIndex+1; CurrBoneIndex < GetNum(); CurrBoneIndex++)
-		{
-			if( RefBoneInfo[CurrBoneIndex].ParentIndex == BoneIndex )
-			{
-				bRemoveThisBone = false;
-				break;
-			}
-		}
-
-		if(bRemoveThisBone)
-		{
-			// Update parent indices of bones further through the array
-			for(int32 CurrBoneIndex=BoneIndex+1; CurrBoneIndex < GetNum(); CurrBoneIndex++)
-			{
-				FMeshBoneInfo& Bone = RefBoneInfo[CurrBoneIndex];
-				if( Bone.ParentIndex > BoneIndex )
-				{
-					Bone.ParentIndex -= 1;
-				}
-			}
-
-			OutBonesRemoved.Add(BoneIndex);
-			RefBonePose.RemoveAt(BoneIndex, 1);
-			RefBoneInfo.RemoveAt(BoneIndex, 1);
-		}
-		return bRemoveThisBone;
-	}
-
-public:
-	void Allocate(int32 Size)
-	{
-		NameToIndexMap.Empty(Size);
-		RefBoneInfo.Empty(Size);
-		RefBonePose.Empty(Size);
-	}
-
-	/** Returns number of bones in Skeleton. */
-	int32 GetNum() const
-	{
-		return RefBoneInfo.Num();
-	}
-
-	/** Accessor to private data. Const so it can't be changed recklessly. */
-	const TArray<FMeshBoneInfo> & GetRefBoneInfo() const
-	{
-		return RefBoneInfo;
-	}
-
-	/** Accessor to private data. Const so it can't be changed recklessly. */
-	const TArray<FTransform> & GetRefBonePose() const
-	{
-		return RefBonePose;
-	}
-
-	void UpdateRefPoseTransform(const int32 & BoneIndex, const FTransform & BonePose)
-	{
-		RefBonePose[BoneIndex] = BonePose;
-	}
-
-	/** Add a new bone. 
-	 * BoneName must not already exist! ParentIndex must be valid. */
-	void Add(const FMeshBoneInfo & BoneInfo, const FTransform & BonePose)
-	{
-		// Adding a bone that already exists is illegal
-		check(FindBoneIndex(BoneInfo.Name) == INDEX_NONE);
-	
-		// Make sure our arrays are in sync.
-		checkSlow( (RefBoneInfo.Num() == RefBonePose.Num()) && (RefBoneInfo.Num() == NameToIndexMap.Num()) );
-
-		const int32 BoneIndex = RefBoneInfo.Add(BoneInfo);
-		RefBonePose.Add(BonePose);
-		NameToIndexMap.Add(BoneInfo.Name, BoneIndex);
-
-		// Normalize Quaternion to be safe.
-		RefBonePose[BoneIndex].NormalizeRotation();
-
-		// Parent must be valid. Either INDEX_NONE for Root, or before us.
-		check( ((BoneIndex == 0) && (BoneInfo.ParentIndex == INDEX_NONE)) || ((BoneIndex > 0) && RefBoneInfo.IsValidIndex(BoneInfo.ParentIndex)) );
-	}
-
-	/** Insert a new bone
-	* BoneName must not already exist! ParentIndex must be valid. */
-	void Insert(int32 BoneIndex, const FMeshBoneInfo & BoneInfo, const FTransform & BonePose)
-	{
-		// Make sure our arrays are in sync.
-		checkSlow( (RefBoneInfo.Num() == RefBonePose.Num()) && (RefBoneInfo.Num() == NameToIndexMap.Num()) );
-
-		// Inserting a bone that already exists is illegal
-		check(FindBoneIndex(BoneInfo.Name) == INDEX_NONE);
-
-		// Parent must be valid. Either INDEX_NONE for Root, or before us.
-		check( ((BoneIndex == 0) && (BoneInfo.ParentIndex == INDEX_NONE)) || ((BoneIndex > 0) && RefBoneInfo.IsValidIndex(BoneInfo.ParentIndex)) );
-
-		// Make sure our bone transform is valid.
-		check( BonePose.GetRotation().IsNormalized() );
-
-		RefBoneInfo.Insert(BoneInfo, BoneIndex);
-		RefBonePose.Insert(BonePose, BoneIndex);
-		NameToIndexMap.Add(BoneInfo.Name, BoneIndex);
-
-		// Normalize Quaternion to be safe.
-		RefBonePose[BoneIndex].NormalizeRotation();
-
-		// Now we need to fix all the parent indices that pointed to bones after this in the array
-		// These must be after this point in the array.
-		for(int32 j=BoneIndex+1; j<GetNum(); j++)
-		{
-			if( GetParentIndex(j) >= BoneIndex )
-			{
-				RefBoneInfo[j].ParentIndex += 1;
-			}
-		}
-	}
-
-	void Empty()
-	{
-		RefBoneInfo.Empty();
-		RefBonePose.Empty();
-		NameToIndexMap.Empty();
-	}
-
-	/** Find Bone Index from BoneName. Precache as much as possible in speed critical sections! */
-	int32 FindBoneIndex(const FName & BoneName) const
-	{
-		checkSlow(RefBoneInfo.Num() == NameToIndexMap.Num());
-		int32 BoneIndex = INDEX_NONE;
-		if( BoneName != NAME_None )
-		{
-			const int32* IndexPtr = NameToIndexMap.Find(BoneName);
-			if( IndexPtr )
-			{
-				BoneIndex = *IndexPtr;
-			}
-		}
-		return BoneIndex;
-	}
-
-	FName GetBoneName(const int32 & BoneIndex) const
-	{
-		return RefBoneInfo[BoneIndex].Name;
-	}
-
-	int32 GetParentIndex(const int32 & BoneIndex) const
-	{
-		// Parent must be valid. Either INDEX_NONE for Root, or before us.
-		checkSlow( ((BoneIndex == 0) && (RefBoneInfo[BoneIndex].ParentIndex == INDEX_NONE)) || ((BoneIndex > 0) && RefBoneInfo.IsValidIndex(RefBoneInfo[BoneIndex].ParentIndex)) );
-		return RefBoneInfo[BoneIndex].ParentIndex;
-	}
-
-	bool IsValidIndex(int32 Index) const
-	{
-		return (RefBoneInfo.IsValidIndex(Index));
-	}
-
-	/** 
-	 * Returns # of Depth from BoneIndex to ParentBoneIndex
-	 * This will return 0 if BoneIndex == ParentBoneIndex;
-	 * This will return -1 if BoneIndex isn't child of ParentBoneIndex
-	 */
-	int32 GetDepthBetweenBones(const int32 & BoneIndex, const int32 & ParentBoneIndex) const
-	{
-		if (BoneIndex >= ParentBoneIndex)
-		{
-			int32 CurBoneIndex = BoneIndex;
-			int32 Depth = 0;
-
-			do
-			{
-				// if same return;
-				if (CurBoneIndex == ParentBoneIndex)
-				{
-					return Depth;
-				}
-
-				CurBoneIndex = RefBoneInfo[CurBoneIndex].ParentIndex;
-				++Depth;
 
-			} while (CurBoneIndex!=INDEX_NONE);
-		}
 
-		return INDEX_NONE;
-	}
-
-	bool BoneIsChildOf(const int32 & ChildBoneIndex, const int32 & ParentBoneIndex) const
-	{
-		// Bones are in strictly increasing order.
-		// So child must have an index greater than his parent.
-		if( ChildBoneIndex > ParentBoneIndex )
-		{
-			int32 BoneIndex = GetParentIndex(ChildBoneIndex);
-			do
-			{
-				if( BoneIndex == ParentBoneIndex )
-				{
-					return true;
-				}
-				BoneIndex = GetParentIndex(BoneIndex);
-
-			} while (BoneIndex != INDEX_NONE);
-		}
-
-		return false;
-	}
-
-	void RemoveDuplicateBones(const UObject * Requester, TArray<FBoneIndexType> & DuplicateBones)
-	{
-		const int32 NumBones = GetNum();
-		DuplicateBones.Empty();
 
-		TMap<FName, int32> BoneNameCheck;
-		bool bRemovedBones = false;
-		for(int32 BoneIndex=NumBones-1; BoneIndex>=0; BoneIndex--)
-		{
-			const FName & BoneName = GetBoneName(BoneIndex);
-			const int32 * FoundBoneIndexPtr = BoneNameCheck.Find(BoneName);
-
-			// Not a duplicate bone, track it.
-			if( FoundBoneIndexPtr == NULL )
-			{
-				BoneNameCheck.Add(BoneName, BoneIndex);
-			}
-			else 
-			{
-				const int32 DuplicateBoneIndex = *FoundBoneIndexPtr;
-				DuplicateBones.Add(DuplicateBoneIndex);
-
-				UE_LOG(LogAnimation, Warning, TEXT("RemoveDuplicateBones: duplicate bone name (%s) detected for (%s)! Indices: %d and %d. Removing the latter."), 
-					*BoneName.ToString(), *GetNameSafe(Requester), DuplicateBoneIndex, BoneIndex);
-
-				// Remove duplicate bone index, which was added later as a mistake.
-				RefBonePose.RemoveAt(DuplicateBoneIndex, 1);
-				RefBoneInfo.RemoveAt(DuplicateBoneIndex, 1);
-
-				// Now we need to fix all the parent indices that pointed to bones after this in the array
-				// These must be after this point in the array.
-				for(int32 j=DuplicateBoneIndex; j<GetNum(); j++)
-				{
-					if( GetParentIndex(j) >= DuplicateBoneIndex )
-					{
-						RefBoneInfo[j].ParentIndex -= 1;
-					}
-				}
-
-				// Update entry in case problem bones were added multiple times.
-				BoneNameCheck.Add(BoneName, BoneIndex);
-
-				// We need to make sure that any bone that has this old bone as a parent is fixed up
-				bRemovedBones = true;
-			}
-		}
-
-		// If we've removed bones, we need to rebuild our name table.
-		if( bRemovedBones || (NameToIndexMap.Num() == 0) )
-		{
-			RebuildNameToIndexMap();
-		}
-
-		// Make sure our arrays are in sync.
-		checkSlow( (RefBoneInfo.Num() == RefBonePose.Num()) && (RefBoneInfo.Num() == NameToIndexMap.Num()) );
-
-		// Additionally normalize all quaternions to be safe.
-		for(int32 BoneIndex=0; BoneIndex<GetNum(); BoneIndex++)
-		{
-			RefBonePose[BoneIndex].NormalizeRotation();
-		}
-	}
-
-	/** Removes the supplied bones from the skeleton, unless they have children that aren't also going to be removed */
-	TArray<int32> RemoveBonesByName(const TArray<FName>& BonesToRemove)
-	{
-		TArray<int32> BonesRemoved;
-
-		const int32 NumBones = GetNum();
-		for(int32 BoneIndex=NumBones-1; BoneIndex>=0; BoneIndex--)
-		{
-			FMeshBoneInfo& Bone = RefBoneInfo[BoneIndex];
-
-			if(BonesToRemove.Contains(Bone.Name))
-			{
-				RemoveIndividualBone(BoneIndex, BonesRemoved);
-			}
-		}
-		RebuildNameToIndexMap();
-		return BonesRemoved;
-	}
-
-	void RebuildNameToIndexMap()
-	{
-		// Start by clearing the current map.
-		NameToIndexMap.Empty();
-
-		// Then iterate over each bone, adding the name and bone index.
-		const int32 NumBones = GetNum();
-		for(int32 BoneIndex=0; BoneIndex<NumBones; BoneIndex++)
-		{
-			const FName & BoneName = GetBoneName(BoneIndex);
-			if( BoneName != NAME_None )
-			{
-				NameToIndexMap.Add(BoneName, BoneIndex);
-			}
-			else
-			{
-				UE_LOG(LogAnimation, Warning, TEXT("RebuildNameToIndexMap: Bone with no name detected for index: %d"), BoneIndex);
-			}
-		}
-
-		// Make sure we don't have duplicate bone names. This would be very bad.
-		checkSlow(NameToIndexMap.Num() == NumBones);
-	}
-
-	friend FArchive & operator<<(FArchive & Ar, FReferenceSkeleton & F);
-};
 
 // Textured triangle.
 struct VTriangle
diff --git a/Engine/Source/Runtime/Engine/Public/SystemSettings.h b/Engine/Source/Runtime/Engine/Public/SystemSettings.h
--- a/Engine/Source/Runtime/Engine/Public/SystemSettings.h
+++ b/Engine/Source/Runtime/Engine/Public/SystemSettings.h
@@ -6,7 +6,7 @@
 
 #pragma once
 
-#include "Texture.h"
+#include "TextureLODSettings.h"
 #include "ShowFlags.h"
 
 /*-----------------------------------------------------------------------------
diff --git a/Engine/Source/Runtime/Engine/Public/Texture.h b/Engine/Source/Runtime/Engine/Public/Texture.h
deleted file mode 100644
diff --git a/Engine/Source/Runtime/Engine/Public/TextureLODSettings.h b/Engine/Source/Runtime/Engine/Public/TextureLODSettings.h
new file mode 100644
--- /dev/null
+++ b/Engine/Source/Runtime/Engine/Public/TextureLODSettings.h
@@ -0,0 +1,139 @@
+// Copyright 1998-2014 Epic Games, Inc. All Rights Reserved.
+
+#pragma once
+
+/**
+ * Structure containing all information related to an LOD group and providing helper functions to calculate
+ * the LOD bias of a given group.
+ */
+struct FTextureLODSettings
+{
+	/**
+	 * Initializes LOD settings by reading them from the passed in filename/ section.
+	 *
+	 * @param	IniFilename		Filename of ini to read from.
+	 * @param	IniSection		Section in ini to look for settings
+	 */
+	ENGINE_API void Initialize( const FString& IniFilename, const TCHAR* IniSection );
+
+	/**
+	 * Initializes LOD settings by reading them from the passed in filename/ section.
+	 *
+	 * @param	IniFile			Preloaded ini file object to load from
+	 * @param	IniSection		Section in ini to look for settings
+	 */
+	ENGINE_API void Initialize(const FConfigFile& IniFile, const TCHAR* IniSection);
+
+	/**
+	 * Calculates and returns the LOD bias based on texture LOD group, LOD bias and maximum size.
+	 *
+	 * @param	Texture			Texture object to calculate LOD bias for.
+	 * @param	bIncTextureMips	If true, takes the textures LOD & Cinematic mips into consideration
+	 * @return	LOD bias
+	 */
+	ENGINE_API int32 CalculateLODBias(const UTexture* Texture, bool bIncTextureMips = true) const;
+
+	/**
+	 * Calculates and returns the LOD bias based on the information provided.
+	 *
+	 * @param	Width						Width of the texture
+	 * @param	Height						Height of the texture
+	 * @param	LODGroup					Which LOD group the texture belongs to
+	 * @param	LODBias						LOD bias to include in the calculation
+	 * @param	NumCinematicMipLevels		The texture cinematic mip levels to include in the calculation
+	 * @param	MipGenSetting				Mip generation setting
+	 * @return	LOD bias
+	 */
+	ENGINE_API int32 CalculateLODBias( int32 Width, int32 Height, int32 LODGroup, int32 LODBias, int32 NumCinematicMipLevels, TextureMipGenSettings MipGenSetting ) const;
+
+	/** 
+	* Useful for stats in the editor.
+	*
+	* @param LODBias			Default LOD at which the texture renders. Platform dependent, call FTextureLODSettings::CalculateLODBias(Texture)
+	*/
+	ENGINE_API void ComputeInGameMaxResolution(int32 LODBias, UTexture& Texture, uint32& OutSizeX, uint32& OutSizeY) const;
+
+#if WITH_EDITORONLY_DATA
+	void GetMipGenSettings( const UTexture& Texture, TextureMipGenSettings& OutMipGenSettings, float& OutSharpen, uint32& OutKernelSize, bool& bOutDownsampleWithAverage, bool& bOutSharpenWithoutColorShift, bool &bOutBorderColorBlack ) const;
+#endif // #if WITH_EDITORONLY_DATA
+
+	/**
+	 * Will return the LODBias for a passed in LODGroup
+	 *
+	 * @param	InLODGroup		The LOD Group ID 
+	 * @return	LODBias
+	 */
+	int32 GetTextureLODGroupLODBias( int32 InLODGroup ) const;
+
+	/**
+	 * Returns the LODGroup setting for number of streaming mip-levels.
+	 * -1 means that all mip-levels are allowed to stream.
+	 *
+	 * @param	InLODGroup		The LOD Group ID 
+	 * @return	Number of streaming mip-levels for textures in the specified LODGroup
+	 */
+	int32 GetNumStreamedMips( int32 InLODGroup ) const;
+
+	/**
+	 * Returns the filter state that should be used for the passed in texture, taking
+	 * into account other system settings.
+	 *
+	 * @param	Texture		Texture to retrieve filter state for, must not be 0
+	 * @return	Filter sampler state for passed in texture
+	 */
+	ESamplerFilter GetSamplerFilter( const UTexture* Texture ) const;
+
+	/**
+	 * Returns the texture group names, sorted like enum.
+	 *
+	 * @return array of texture group names
+	 */
+	ENGINE_API static TArray<FString> GetTextureGroupNames();
+
+	/** LOD settings for a single texture group. */
+	struct FTextureLODGroup 
+	{
+		FTextureLODGroup()
+		:	MinLODMipCount(0)
+		,	MaxLODMipCount(12)
+		,	LODBias(0) 
+		,	Filter(SF_AnisotropicPoint)
+		,	NumStreamedMips(-1)
+		,	MipGenSettings(TMGS_SimpleAverage)
+		{}
+		/** Minimum LOD mip count below which the code won't bias.						*/
+		int32 MinLODMipCount;
+		/** Maximum LOD mip count. Bias will be adjusted so texture won't go above.		*/
+		int32 MaxLODMipCount;
+		/** Group LOD bias.																*/
+		int32 LODBias;
+		/** Sampler filter state.														*/
+		ESamplerFilter Filter;
+		/** Number of mip-levels that can be streamed. -1 means all mips can stream.	*/
+		int32 NumStreamedMips;
+		/** Defines how the the mip-map generation works, e.g. sharpening				*/
+		TextureMipGenSettings MipGenSettings;
+	};
+
+protected:
+	/**
+	 * Reads a single entry and parses it into the group array.
+	 *
+	 * @param	GroupId			Id/ enum of group to parse
+	 * @param	GroupName		Name of group to look for in ini
+	 * @param	IniFile			Config file to read from.
+	 * @param	IniSection		Section in ini to look for settings
+	 */
+	void ReadEntry( int32 GroupId, const TCHAR* GroupName, const FConfigFile& IniFile, const TCHAR* IniSection );
+
+	/**
+	 * TextureLODGroups access with bounds check
+	 *
+	 * @param   GroupIndex      usually from Texture.LODGroup
+	 * @return                  A handle to the indexed LOD group. 
+	 */
+	const FTextureLODGroup& GetTextureLODGroup(TextureGroup GroupIndex) const;
+
+	/** Array of LOD settings with entries per group. */
+	FTextureLODGroup TextureLODGroups[TEXTUREGROUP_MAX];
+};
diff --git a/Engine/Source/Runtime/Engine/Public/TextureResource.h b/Engine/Source/Runtime/Engine/Public/TextureResource.h
new file mode 100644
--- /dev/null
+++ b/Engine/Source/Runtime/Engine/Public/TextureResource.h
@@ -0,0 +1,883 @@
+// Copyright 1998-2014 Epic Games, Inc. All Rights Reserved.
+
+#pragma once
+
+/*=============================================================================
+	Texture.h: Unreal texture related classes.
+=============================================================================*/
+
+#include "Engine/TextureDefines.h"
+#include "UnrealClient.h"
+
+class UTexture2D;
+
+/** Maximum number of slices in texture source art. */
+#define MAX_TEXTURE_SOURCE_SLICES 6
+
+/** Thread-safe counter indicating the texture streaming state. The definitions below are mirrored in Texture2D.h */
+enum ETextureStreamingState
+{
+	// The renderer hasn't created the resource yet.
+	TexState_InProgress_Initialization	= -1,
+	// There are no pending requests/ all requests have been fulfilled.
+	TexState_ReadyFor_Requests			= 0,
+	// Finalization has been kicked off and is in progress.
+	TexState_InProgress_Finalization	= 1,
+	// Initial request has completed and finalization needs to be kicked off.
+	TexState_ReadyFor_Finalization		= 2,
+	// Mip data is in the process of being uploaded to the GPU.
+	TexState_InProgress_Upload			= 3,
+	// Mip data has been loaded in to system memory and is ready to be transferred to the GPU.
+	TexState_ReadyFor_Upload			= 4,
+	// We're currently loading in mip data.
+	TexState_InProgress_Loading			= 5,
+	// ...
+	// States 2+N means we're currently loading in N mips
+	// ...
+	// Memory has been allocated and we're ready to start loading in mips.
+	TexState_ReadyFor_Loading			= 100,
+	// We're currently allocating/preparing memory for the new mip count.
+	TexState_InProgress_Allocation		= 101,
+	// The RHI is asynchronously allocating/preparing memory for the new mip count.
+	TexState_InProgress_AsyncAllocation = 102
+};
+
+/**
+ * Async worker to stream mips from the derived data cache.
+ */
+class FAsyncStreamDerivedMipWorker : public FNonAbandonableTask
+{
+public:
+	/** Initialization constructor. */
+	FAsyncStreamDerivedMipWorker(
+		const FString& InDerivedDataKey,
+		void* InDestMipData,
+		int32 InMipSize,
+		FThreadSafeCounter* InThreadSafeCounter
+		);
+	
+	/**
+	 * Retrieves the derived mip from the derived data cache.
+	 */
+	void DoWork();
+
+	/**
+	 * Name of the async task.
+	 */
+	static const TCHAR* Name()
+	{
+		return TEXT("FAsyncStreamDerivedMipTask");
+	}
+
+	/**
+	 * Returns true if the streaming mip request failed.
+	 */
+	bool DidRequestFail() const
+	{
+		return bRequestFailed;
+	}
+
+private:
+	/** Key for retrieving mip data from the derived data cache. */
+	FString DerivedDataKey;
+	/** The location to which the mip data should be copied. */
+	void* DestMipData;
+	/** The size of the mip in bytes. */
+	int32 ExpectedMipSize;
+	/** true if the mip data was not present in the derived data cache. */
+	bool bRequestFailed;
+	/** Thread-safe counter to decrement when data has been copied. */
+	FThreadSafeCounter* ThreadSafeCounter;
+};
+
+/** Async task to stream mips from the derived data cache. */
+typedef FAsyncTask<FAsyncStreamDerivedMipWorker> FAsyncStreamDerivedMipTask;
+
+/**
+ * Task to create a texture asynchronously.
+ */
+class FCreateTextureTask : public FNonAbandonableTask
+{
+public:
+	/** Task arguments. */
+	struct FArguments
+	{
+		/** Width of the texture. */
+		uint32 SizeX;
+		/** Height of the texture. */
+		uint32 SizeY;
+		/** Format of the texture. */
+		EPixelFormat Format;
+		/** The number of mips. */
+		uint32 NumMips;
+		/** Texture creation flags. */
+		uint32 Flags;
+		/** Initial mip data. */
+		void** MipData;
+		/** The number of provided mips. */
+		uint32 NumNewMips;
+		/** Pointer to a reference where the new texture will be stored. */
+		FTexture2DRHIRef* TextureRefPtr;
+		/** Thread safe counter to decrement when complete. */
+		FThreadSafeCounter* ThreadSafeCounter;
+	};
+
+	/** Initialization constructor. */
+	FCreateTextureTask(FArguments InArgs)
+		: Args(InArgs)
+	{
+		check(Args.TextureRefPtr);
+		check(Args.ThreadSafeCounter);
+	}
+
+	static const TCHAR* Name()
+	{
+		return TEXT("FCreateTextureTask");
+	}
+
+	/** Creates the texture. */
+	void DoWork();
+
+private:
+	/** Task arguments. */
+	FArguments Args;
+};
+
+/** Async task to create a texture. */
+typedef FAsyncTask<FCreateTextureTask> FAsyncCreateTextureTask;
+
+/**
+ * A 2D texture mip-map.
+ */
+struct FTexture2DMipMap
+{
+	/** Width of the mip-map. */
+	int32 SizeX;
+	/** Height of the mip-map. */
+	int32 SizeY;
+	/** Bulk data if stored in the package. */
+	FByteBulkData BulkData;
+
+	/** Default constructor. */
+	FTexture2DMipMap()
+		: SizeX(0)
+		, SizeY(0)
+	{
+	}
+
+	/** Serialization. */
+	ENGINE_API void Serialize(FArchive& Ar, UObject* Owner, int32 MipIndex);
+
+#if WITH_EDITORONLY_DATA
+	/** Key if stored in the derived data cache. */
+	FString DerivedDataKey;
+
+	/**
+	 * Place mip-map data in the derived data cache associated with the provided
+	 * key.
+	 */
+	void StoreInDerivedDataCache(const FString& InDerivedDataKey);
+#endif // #if WITH_EDITORONLY_DATA
+};
+
+/** 
+ * The rendering resource which represents a texture.
+ */
+class FTextureResource : public FTexture
+{
+public:
+
+	FTextureResource()
+	{}
+	virtual ~FTextureResource() {}
+
+#if STATS
+	/* The Stat_ FName corresponding to each TEXTUREGROUP */
+	static FName TextureGroupStatFNames[TEXTUREGROUP_MAX];
+#endif
+};
+
+/**
+ * FTextureResource implementation for streamable 2D textures.
+ */
+class FTexture2DResource : public FTextureResource
+{
+public:
+	/**
+	 * Minimal initialization constructor.
+	 *
+	 * @param InOwner			UTexture2D which this FTexture2DResource represents.
+	 * @param InitialMipCount	Initial number of miplevels to upload to card
+	 * @param InFilename		Filename to read data from
+ 	 */
+	FTexture2DResource( UTexture2D* InOwner, int32 InitialMipCount );
+
+	/**
+	 * Destructor, freeing MipData in the case of resource being destroyed without ever 
+	 * having been initialized by the rendering thread via InitRHI.
+	 */
+	virtual ~FTexture2DResource();
+
+	// FRenderResource interface.
+
+	/**
+	 * Called when the resource is initialized. This is only called by the rendering thread.
+	 */
+	virtual void InitRHI() override;
+	/**
+	 * Called when the resource is released. This is only called by the rendering thread.
+	 */
+	virtual void ReleaseRHI() override;
+
+	/** Returns the width of the texture in pixels. */
+	virtual uint32 GetSizeX() const;
+
+	/** Returns the height of the texture in pixels. */
+	virtual uint32 GetSizeY() const;
+
+	/**
+	 * Called from the game thread to kick off a change in ResidentMips after modifying RequestedMips.
+	 * @param bShouldPrioritizeAsyncIORequest	- Whether the Async I/O request should have higher priority
+	 */
+	void BeginUpdateMipCount( bool bShouldPrioritizeAsyncIORequest );
+	/**
+	 * Called from the game thread to kick off async I/O to load in new mips.
+	 */
+	void BeginLoadMipData();
+	/**
+	 * Called from the game thread to kick off uploading mip data to the GPU.
+	 */
+	void BeginUploadMipData();
+	/**
+	 * Called from the game thread to kick off finalization of mip change.
+	 */
+	void BeginFinalizeMipCount();
+	/**
+	 * Called from the game thread to kick off cancelation of async operations for request.
+	 */
+	void BeginCancelUpdate();
+
+	/** 
+	 * Accessor
+	 * @return Texture2DRHI
+	 */
+	FTexture2DRHIRef GetTexture2DRHI()
+	{
+		return Texture2DRHI;
+	}
+
+	bool DidUpdateMipCountFail() const
+	{
+		return NumFailedReallocs > 0 || bDerivedDataStreamRequestFailed;
+	}
+
+	bool DidDerivedDataRequestFail() const
+	{
+		return bDerivedDataStreamRequestFailed;
+	}
+
+	/**
+	 *	Tries to reallocate the texture for a new mip count.
+	 *	@param OldMipCount	- The old mip count we're currently using.
+	 *	@param NewMipCount	- The new mip count to use.
+	 */
+	bool TryReallocate( int32 OldMipCount, int32 NewMipCount );
+
+	virtual FString GetFriendlyName() const;
+
+	//Returns the raw data for a particular mip level
+	void* GetRawMipData( uint32 MipIndex);
+
+	//Returns the current first mip (always valid)
+	int32 GetCurrentFirstMip() const
+	{
+		return CurrentFirstMip;
+	}
+
+private:
+	/** Texture streaming command classes that need to be friends in order to call Update/FinalizeMipCount.	*/
+	friend class FUpdateMipCountCommand;
+	friend class FFinalinzeMipCountCommand;
+	friend class FCancelUpdateCommand;
+	friend class FUploadMipDataCommand;
+	friend class UTexture2D;
+
+	/** The UTexture2D which this resource represents.														*/
+	const UTexture2D*	Owner;
+	/** Resource memory allocated by the owner for serialize bulk mip data into								*/
+	FTexture2DResourceMem* ResourceMem;
+	
+	/** First miplevel used in the texture being streamed in, which is IntermediateTextureRHI when it is valid. */
+	int32	PendingFirstMip;
+
+	/** First mip level used in Texture2DRHI. This is always correct as long as Texture2DRHI is allocated, regardless of streaming status. */
+	int32 CurrentFirstMip;
+
+	/** Pending async create texture task, if any.															*/
+	TScopedPointer<FAsyncCreateTextureTask> AsyncCreateTextureTask;
+	
+	/** Local copy/ cache of mip data between creation and first call to InitRHI.							*/
+	void*				MipData[MAX_TEXTURE_MIP_COUNT];
+	/** Potentially outstanding texture I/O requests.														*/
+	uint64				IORequestIndices[MAX_TEXTURE_MIP_COUNT];
+	/** Number of file I/O requests for current request														*/
+	int32					IORequestCount;
+
+#if WITH_EDITORONLY_DATA
+	/** Pending async derived data streaming tasks															*/
+	TIndirectArray<FAsyncStreamDerivedMipTask> PendingAsyncStreamDerivedMipTasks;
+#endif // #if WITH_EDITORONLY_DATA
+
+	/** 2D texture version of TextureRHI which is used to lock the 2D texture during mip transitions.		*/
+	FTexture2DRHIRef	Texture2DRHI;
+	/** Intermediate texture used to fulfill mip change requests. Swapped in FinalizeMipCount.				*/
+	FTexture2DRHIRef	IntermediateTextureRHI;
+	/** Whether the intermediate texture is being created asynchronously.									*/
+	uint32			bUsingAsyncCreation:1;
+	/** Whether the current stream request is prioritized higher than normal.	*/
+	uint32			bPrioritizedIORequest:1;
+	/** Whether the last mip streaming request failed.														*/
+	uint32			bDerivedDataStreamRequestFailed:1;
+	/** Number of times UpdateMipCount has failed to reallocate memory.										*/
+	int32					NumFailedReallocs;
+
+#if STATS
+	/** Cached texture size for stats.																		*/
+	int32					TextureSize;
+	/** Cached intermediate texture size for stats.															*/
+	int32					IntermediateTextureSize;
+	/** The FName of the LODGroup-specific stat																*/
+	FName					LODGroupStatName;
+#endif
+
+	/**
+	 * Writes the data for a single mip-level into a destination buffer.
+	 * @param MipIndex	The index of the mip-level to read.
+	 * @param Dest		The address of the destination buffer to receive the mip-level's data.
+	 * @param DestPitch	Number of bytes per row
+	 */
+	void GetData( uint32 MipIndex,void* Dest,uint32 DestPitch );
+
+	/**
+	 * Called from the rendering thread to perform the work to kick off a change in ResidentMips.
+	 */
+	void UpdateMipCount();
+	/**
+	 * Called from the rendering thread to start async I/O to load in new mips.
+	 */
+	void LoadMipData();
+	/**
+	 * Called from the rendering thread to unlock textures and make the data visible to the GPU.
+	 */
+	void UploadMipData();
+	/**
+	 * Called from the rendering thread to finalize a mip change.
+	 */
+	void FinalizeMipCount();
+	/**
+	 * Called from the rendering thread to cancel async operations for request.
+	 */
+	void CancelUpdate();
+
+	/** Create RHI sampler states. */
+	void CreateSamplerStates(float MipMapBias);
+
+	/** Returns the default mip map bias for this texture. */
+	int32 GetDefaultMipMapBias() const;
+
+	// releases and recreates sampler state objects.
+	// used when updating mip map bias offset
+	void RefreshSamplerStates();
+};
+
+/** A dynamic 2D texture resource. */
+class FTexture2DDynamicResource : public FTextureResource
+{
+public:
+	/** Initialization constructor. */
+	FTexture2DDynamicResource(class UTexture2DDynamic* InOwner);
+
+	/** Returns the width of the texture in pixels. */
+	virtual uint32 GetSizeX() const;
+
+	/** Returns the height of the texture in pixels. */
+	virtual uint32 GetSizeY() const;
+
+	/** Called when the resource is initialized. This is only called by the rendering thread. */
+	virtual void InitRHI() override;
+
+	/** Called when the resource is released. This is only called by the rendering thread. */
+	virtual void ReleaseRHI() override;
+
+	/** Returns the Texture2DRHI, which can be used for locking/unlocking the mips. */
+	ENGINE_API FTexture2DRHIRef GetTexture2DRHI();
+
+private:
+	/** The owner of this resource. */
+	class UTexture2DDynamic* Owner;
+	/** Texture2D reference, used for locking/unlocking the mips. */
+	FTexture2DRHIRef Texture2DRHI;
+};
+
+/** Stores information about a mip map, used by FTexture2DArrayResource to mirror game thread data. */
+class FMipMapDataEntry
+{
+public:
+	uint32 SizeX;
+	uint32 SizeY;
+	TArray<uint8> Data;
+};
+
+/** Stores information about a single texture in FTexture2DArrayResource. */
+class FTextureArrayDataEntry
+{
+public:
+
+	FTextureArrayDataEntry() : 
+		NumRefs(0)
+	{}
+
+	/** Number of FTexture2DArrayResource::AddTexture2D calls that specified this texture. */
+	int32 NumRefs;
+
+	/** Mip maps of the texture. */
+	TArray<FMipMapDataEntry, TInlineAllocator<MAX_TEXTURE_MIP_COUNT> > MipData;
+};
+
+/** 
+ * Stores information about a UTexture2D so the rendering thread can access it, 
+ * Even though the UTexture2D may have changed by the time the rendering thread gets around to it.
+ */
+class FIncomingTextureArrayDataEntry : public FTextureArrayDataEntry
+{
+public:
+
+	FIncomingTextureArrayDataEntry() {}
+
+	FIncomingTextureArrayDataEntry(UTexture2D* InTexture);
+
+	int32 SizeX;
+	int32 SizeY;
+	int32 NumMips;
+	TextureGroup LODGroup;
+	EPixelFormat Format;
+	ESamplerFilter Filter;
+	bool bSRGB;
+};
+
+/** Represents a 2D Texture Array to the renderer. */
+class FTexture2DArrayResource : public FTextureResource
+{
+public:
+
+	FTexture2DArrayResource() :
+		SizeX(0),
+		bDirty(false),
+		bPreventingReallocation(false)
+	{}
+
+	// Rendering thread functions
+
+	/** 
+	 * Adds a texture to the texture array.  
+	 * This is called on the rendering thread, so it must not dereference NewTexture.
+	 */
+	void AddTexture2D(UTexture2D* NewTexture, const FIncomingTextureArrayDataEntry* InEntry);
+
+	/** Removes a texture from the texture array, and potentially removes the CachedData entry if the last ref was removed. */
+	void RemoveTexture2D(const UTexture2D* NewTexture);
+
+	/** Updates a CachedData entry (if one exists for this texture), with a new texture. */
+	void UpdateTexture2D(UTexture2D* NewTexture, const FIncomingTextureArrayDataEntry* InEntry);
+
+	/** Initializes the texture array resource if needed, and re-initializes if the texture array has been made dirty since the last init. */
+	void UpdateResource();
+
+	/** Returns the index of a given texture in the texture array. */
+	int32 GetTextureIndex(const UTexture2D* Texture) const;
+	int32 GetNumValidTextures() const;
+
+	/**
+	* Called when the resource is initialized. This is only called by the rendering thread.
+	*/
+	virtual void InitRHI() override;
+
+	/** Returns the width of the texture in pixels. */
+	virtual uint32 GetSizeX() const
+	{
+		return SizeX;
+	}
+
+	/** Returns the height of the texture in pixels. */
+	virtual uint32 GetSizeY() const
+	{
+		return SizeY;
+	}
+
+	/** Prevents reallocation from removals of the texture array until EndPreventReallocation is called. */
+	void BeginPreventReallocation();
+
+	/** Restores the ability to reallocate the texture array. */
+	void EndPreventReallocation();
+
+private:
+
+	/** Texture data, has to persist past the first InitRHI call, because more textures may be added later. */
+	TMap<const UTexture2D*, FTextureArrayDataEntry> CachedData;
+	uint32 SizeX;
+	uint32 SizeY;
+	uint32 NumMips;
+	TextureGroup LODGroup;
+	EPixelFormat Format;
+	ESamplerFilter Filter;
+
+	bool bSRGB;
+	bool bDirty;
+	bool bPreventingReallocation;
+
+	/** Copies data from DataEntry into Dest, taking stride into account. */
+	void GetData(const FTextureArrayDataEntry& DataEntry, int32 MipIndex, void* Dest, uint32 DestPitch);
+};
+
+/**
+ * FDeferredUpdateResource for resources that need to be updated after scene rendering has begun
+ * (should only be used on the rendering thread)
+ */
+class FDeferredUpdateResource
+{
+public:
+	/**
+	 * Constructor, initializing UpdateListLink.
+	 */
+	FDeferredUpdateResource()
+		:	UpdateListLink(NULL)
+		,	bOnlyUpdateOnce(false)
+	{}
+
+	/**
+	 * Iterate over the global list of resources that need to
+	 * be updated and call UpdateResource on each one.
+	 */
+	ENGINE_API static void UpdateResources(FRHICommandListImmediate& RHICmdList);
+
+	/** 
+	 * This is reset after all viewports have been rendered
+	 */
+	static void ResetNeedsUpdate()
+	{
+		bNeedsUpdate = true;
+	}
+
+	// FDeferredUpdateResource interface
+
+protected:
+
+	/**
+	* Updates the resource
+	*/
+	virtual void UpdateDeferredResource(FRHICommandListImmediate& RHICmdList) = 0;
+
+	/**
+	 * Add this resource to deferred update list
+	 * @param OnlyUpdateOnce - flag this resource for a single update if true
+	 */
+	void AddToDeferredUpdateList( bool OnlyUpdateOnce );
+
+	/**
+	 * Remove this resource from deferred update list
+	 */
+	void RemoveFromDeferredUpdateList();
+
+private:
+	/** 
+	 * Resources can be added to this list if they need a deferred update during scene rendering.
+	 * @return global list of resource that need to be updated. 
+	 */
+	static TLinkedList<FDeferredUpdateResource*>*& GetUpdateList();
+	/** This resource's link in the global list of resources needing clears. */
+	TLinkedList<FDeferredUpdateResource*> UpdateListLink;
+	/** if true then UpdateResources needs to be called */
+	ENGINE_API static bool bNeedsUpdate;
+	/** if true then remove this resource from the update list after a single update */
+	bool bOnlyUpdateOnce;
+};
+
+/**
+ * FTextureResource type for render target textures.
+ */
+class FTextureRenderTargetResource : public FTextureResource, public FRenderTarget, public FDeferredUpdateResource
+{
+public:
+	/**
+	 * Constructor, initializing ClearLink.
+	 */
+	FTextureRenderTargetResource()
+	{}
+
+	/** 
+	 * Return true if a render target of the given format is allowed
+	 * for creation
+	 */
+	ENGINE_API static bool IsSupportedFormat( EPixelFormat Format );
+
+	// FTextureRenderTargetResource interface
+	
+	virtual class FTextureRenderTarget2DResource* GetTextureRenderTarget2DResource()
+	{
+		return NULL;
+	}
+	virtual void ClampSize(int32 SizeX,int32 SizeY) {}
+
+	// FRenderTarget interface.
+	virtual FIntPoint GetSizeXY() const = 0;
+
+	/** 
+	 * Render target resource should be sampled in linear color space
+	 *
+	 * @return display gamma expected for rendering to this render target 
+	 */
+	virtual float GetDisplayGamma() const;
+};
+
+/**
+ * FTextureResource type for 2D render target textures.
+ */
+class FTextureRenderTarget2DResource : public FTextureRenderTargetResource
+{
+public:
+	
+	/** 
+	 * Constructor
+	 * @param InOwner - 2d texture object to create a resource for
+	 */
+	FTextureRenderTarget2DResource(const class UTextureRenderTarget2D* InOwner);
+
+	FORCEINLINE FLinearColor GetClearColor()
+	{
+		return ClearColor;
+	}
+
+	// FTextureRenderTargetResource interface
+
+	/** 
+	 * 2D texture RT resource interface 
+	 */
+	virtual class FTextureRenderTarget2DResource* GetTextureRenderTarget2DResource()
+	{
+		return this;
+	}
+
+	/**
+	 * Clamp size of the render target resource to max values
+	 *
+	 * @param MaxSizeX max allowed width
+	 * @param MaxSizeY max allowed height
+	 */
+	virtual void ClampSize(int32 SizeX,int32 SizeY);
+	
+	// FRenderResource interface.
+
+	/**
+	 * Initializes the dynamic RHI resource and/or RHI render target used by this resource.
+	 * Called when the resource is initialized, or when reseting all RHI resources.
+	 * Resources that need to initialize after a D3D device reset must implement this function.
+	 * This is only called by the rendering thread.
+	 */
+	virtual void InitDynamicRHI() override;
+
+	/**
+	 * Releases the dynamic RHI resource and/or RHI render target resources used by this resource.
+	 * Called when the resource is released, or when reseting all RHI resources.
+	 * Resources that need to release before a D3D device reset must implement this function.
+	 * This is only called by the rendering thread.
+	 */
+	virtual void ReleaseDynamicRHI() override;
+
+	// FDeferredClearResource interface
+
+	// FRenderTarget interface.
+	virtual FIntPoint GetSizeXY() const;
+
+	/** 
+	 * Render target resource should be sampled in linear color space
+	 *
+	 * @return display gamma expected for rendering to this render target 
+	 */
+	virtual float GetDisplayGamma() const;
+
+	/** 
+	 * @return TextureRHI for rendering 
+	 */
+	FTexture2DRHIRef GetTextureRHI() { return Texture2DRHI; }
+protected:
+	/**
+	* Clear contents of the render target
+	*/
+	friend class UTextureRenderTarget2D;
+	virtual void UpdateDeferredResource(FRHICommandListImmediate& RHICmdList) override;
+
+private:
+	/** The UTextureRenderTarget2D which this resource represents. */
+	const class UTextureRenderTarget2D* Owner;
+	/** Texture resource used for rendering with and resolving to */
+	FTexture2DRHIRef Texture2DRHI;
+	/** the color the texture is cleared to */
+	FLinearColor ClearColor;
+	EPixelFormat Format;
+	int32 TargetSizeX,TargetSizeY;
+};
+
+/**
+ * FTextureResource type for cube render target textures.
+ */
+class FTextureRenderTargetCubeResource : public FTextureRenderTargetResource
+{
+public:
+
+	/** 
+	 * Constructor
+	 * @param InOwner - cube texture object to create a resource for
+	 */
+	FTextureRenderTargetCubeResource(const class UTextureRenderTargetCube* InOwner)
+		:	Owner(InOwner)
+	{
+	}
+
+	/** 
+	 * Cube texture RT resource interface 
+	 */
+	virtual class FTextureRenderTargetCubeResource* GetTextureRenderTargetCubeResource()
+	{
+		return this;
+	}
+
+	/**
+	 * Initializes the dynamic RHI resource and/or RHI render target used by this resource.
+	 * Called when the resource is initialized, or when reseting all RHI resources.
+	 * Resources that need to initialize after a D3D device reset must implement this function.
+	 * This is only called by the rendering thread.
+	 */
+	virtual void InitDynamicRHI() override;
+
+	/**
+	 * Releases the dynamic RHI resource and/or RHI render target resources used by this resource.
+	 * Called when the resource is released, or when reseting all RHI resources.
+	 * Resources that need to release before a D3D device reset must implement this function.
+	 * This is only called by the rendering thread.
+	 */
+	virtual void ReleaseDynamicRHI() override;
+
+	// FRenderTarget interface.
+
+	/**
+	 * @return dimensions of the target
+	 */
+	virtual FIntPoint GetSizeXY() const;
+
+	/** 
+	 * @return TextureRHI for rendering 
+	 */
+	FTextureCubeRHIRef GetTextureRHI() { return TextureCubeRHI; }
+
+	/** 
+	* Render target resource should be sampled in linear color space
+	*
+	* @return display gamma expected for rendering to this render target 
+	*/
+	float GetDisplayGamma() const;
+
+	/**
+	* Copy the texels of a single face of the cube into an array.
+	* @param OutImageData - float16 values will be stored in this array.
+	* @param InFlags - read flags. ensure cubeface member has been set.
+	* @param InRect - Rectangle of texels to copy.
+	* @return true if the read succeeded.
+	*/
+	ENGINE_API bool ReadPixels(TArray< FColor >& OutImageData, FReadSurfaceDataFlags InFlags, FIntRect InRect = FIntRect(0, 0, 0, 0));
+
+	/**
+	* Copy the texels of a single face of the cube into an array.
+	* @param OutImageData - float16 values will be stored in this array.
+	* @param InFlags - read flags. ensure cubeface member has been set.
+	* @param InRect - Rectangle of texels to copy.
+	* @return true if the read succeeded.
+	*/
+	ENGINE_API bool ReadPixels(TArray<FFloat16Color>& OutImageData, FReadSurfaceDataFlags InFlags, FIntRect InRect = FIntRect(0, 0, 0, 0));
+
+protected:
+	/**
+	* Clear contents of the render target. Clears each face of the cube
+	* This is only called by the rendering thread.
+	*/
+	virtual void UpdateDeferredResource(FRHICommandListImmediate& RHICmdList) override;
+
+private:
+	/** The UTextureRenderTargetCube which this resource represents. */
+	const class UTextureRenderTargetCube* Owner;
+	/** Texture resource used for rendering with and resolving to */
+	FTextureCubeRHIRef TextureCubeRHI;
+	/** Target surfaces for each cube face */
+	FTexture2DRHIRef CubeFaceSurfaceRHI;
+
+	/** Represents the current render target (from one of the cube faces)*/
+	FTextureCubeRHIRef RenderTargetCubeRHI;
+
+	/** Face currently used for target surface */
+	ECubeFace CurrentTargetFace;
+};
+
+/**
+ * FTextureResource type for movie textures.
+ */
+class FTextureMovieResource : public FTextureResource, public FRenderTarget, public FDeferredUpdateResource
+{
+public:
+
+	/** 
+	 * Constructor
+	 * @param InOwner - movie texture object to create a resource for
+	 */
+	FTextureMovieResource(const class UTextureMovie* InOwner)
+		:	Owner(InOwner)
+	{
+	}
+
+	// FRenderResource interface.
+
+	/**
+	 * Initializes the dynamic RHI resource and/or RHI render target used by this resource.
+	 * Called when the resource is initialized, or when reseting all RHI resources.
+	 * Resources that need to initialize after a D3D device reset must implement this function.
+	 * This is only called by the rendering thread.
+	 */
+	virtual void InitDynamicRHI();
+
+	/**
+	 * Releases the dynamic RHI resource and/or RHI render target resources used by this resource.
+	 * Called when the resource is released, or when reseting all RHI resources.
+	 * Resources that need to release before a D3D device reset must implement this function.
+	 * This is only called by the rendering thread.
+	 */
+	virtual void ReleaseDynamicRHI();	
+
+	// FDeferredClearResource interface
+
+	// FRenderTarget interface.
+	virtual FIntPoint GetSizeXY() const;
+
+protected:
+
+	/**
+	* Decodes the next frame of the movie stream and renders the result to this movie texture target
+	*/
+	virtual void UpdateDeferredResource(FRHICommandListImmediate& RHICmdList) override;
+
+
+private:
+	/** The UTextureRenderTarget2D which this resource represents. */
+	const class UTextureMovie* Owner;
+	/** Texture resource used for rendering with and resolving to */
+	FTexture2DRHIRef Texture2DRHI;
+};
+
+#define TEXTURERESOURCE_H_INCLUDED 1 // needed by TargetPlatform, so TTargetPlatformBase template knows if it can use UTexture in GetDefaultTextureFormatName, or rather just declare it
diff --git a/Engine/Source/Runtime/Engine/Public/WorldCollision.h b/Engine/Source/Runtime/Engine/Public/WorldCollision.h
--- a/Engine/Source/Runtime/Engine/Public/WorldCollision.h
+++ b/Engine/Source/Runtime/Engine/Public/WorldCollision.h
@@ -6,9 +6,9 @@
 
 #pragma once 
 
-#include "UniquePtr.h"
+#include "Runtime/Core/Public/Templates/UniquePtr.h"
+#include "Runtime/Core/Public/Async/TaskGraphInterfaces.h"
 #include "CollisionQueryParams.h"
-//#include "Engine/World.h"
 
 /**
  * Async Trace Data Structures 
@@ -244,19 +244,7 @@ struct FBaseTraceDatum
 
 	/* Set functions for each Shape type */
 	void Set(UWorld * World, const FCollisionShape & InCollisionShape, const FCollisionQueryParams & Param, const struct FCollisionResponseParams &InResponseParam, const struct FCollisionObjectQueryParams& InObjectQueryParam, 
-		ECollisionChannel Channel, uint32 InUserData, bool bInIsMultiTrace, int32 FrameCounter)
-	{
-		ensure (World);
-		CollisionParams.CollisionShape = InCollisionShape;
-		CollisionParams.CollisionQueryParam = Param;
-		CollisionParams.ResponseParam = InResponseParam;
-		CollisionParams.ObjectQueryParam = InObjectQueryParam;
-		TraceChannel = Channel;
-		UserData = InUserData;
-		bIsMultiTrace = bInIsMultiTrace;
-		FrameNumber = FrameCounter;
-		PhysWorld = World;
-	}
+		ECollisionChannel Channel, uint32 InUserData, bool bInIsMultiTrace, int32 FrameCounter);
 };
 
 struct FTraceDatum;
