Splitting format specific classes out of AudioDecompress.cpp/h

New files added to contain FVorbisAudioInfo and FOpusAudioInfo, added all necessary includes
Added a few more checks for whether the target platform supports audio streaming where appropriate
Moved #define to .h file as it was being used in two .cpp files (working because of unity build)

diff --git a/Engine/Source/Developer/AudioFormatOgg/Private/AudioFormatOgg.cpp b/Engine/Source/Developer/AudioFormatOgg/Private/AudioFormatOgg.cpp
--- a/Engine/Source/Developer/AudioFormatOgg/Private/AudioFormatOgg.cpp
+++ b/Engine/Source/Developer/AudioFormatOgg/Private/AudioFormatOgg.cpp
@@ -4,7 +4,7 @@
 #include "ModuleInterface.h"
 #include "ModuleManager.h"
 #include "TargetPlatform.h"
-#include "AudioDecompress.h"
+#include "VorbisAudioInfo.h"
 
 
 #if WITH_OGGVORBIS
diff --git a/Engine/Source/Developer/AudioFormatOpus/Private/AudioFormatOpus.cpp b/Engine/Source/Developer/AudioFormatOpus/Private/AudioFormatOpus.cpp
--- a/Engine/Source/Developer/AudioFormatOpus/Private/AudioFormatOpus.cpp
+++ b/Engine/Source/Developer/AudioFormatOpus/Private/AudioFormatOpus.cpp
@@ -4,7 +4,7 @@
 #include "ModuleInterface.h"
 #include "ModuleManager.h"
 #include "TargetPlatform.h"
-#include "AudioDecompress.h"
+#include "OpusAudioInfo.h"
 #include "AudioFormatOpus.h"
 
 // Need to define this so that resampler.h compiles - probably a way around this somehow
diff --git a/Engine/Source/Runtime/Android/AndroidAudio/Private/AndroidAudioDevice.cpp b/Engine/Source/Runtime/Android/AndroidAudio/Private/AndroidAudioDevice.cpp
--- a/Engine/Source/Runtime/Android/AndroidAudio/Private/AndroidAudioDevice.cpp
+++ b/Engine/Source/Runtime/Android/AndroidAudio/Private/AndroidAudioDevice.cpp
@@ -6,7 +6,7 @@
 ------------------------------------------------------------------------------------*/
 
 #include "AndroidAudioDevice.h"
-#include "AudioDecompress.h"
+#include "VorbisAudioInfo.h"
 #include "AudioEffect.h"
 #include "Engine.h"
 
diff --git a/Engine/Source/Runtime/Engine/Private/AudioDecompress.cpp b/Engine/Source/Runtime/Engine/Private/AudioDecompress.cpp
--- a/Engine/Source/Runtime/Engine/Private/AudioDecompress.cpp
+++ b/Engine/Source/Runtime/Engine/Private/AudioDecompress.cpp
@@ -5,733 +5,7 @@
 #include "AudioDecompress.h"
 #include "AudioDevice.h"
 #include "Sound/SoundWave.h"
-
-#if WITH_OGGVORBIS
-
-// hack to get ogg types right for HTML5. 
-#if PLATFORM_HTML5 && !PLATFORM_HTML5_WIN32
-#define _WIN32 
-#define __MINGW32__
-#endif 
-	#pragma pack(push, 8)
-	#include "ogg/ogg.h"
-	#include "vorbis/vorbisenc.h"
-	#include "vorbis/vorbisfile.h"
-	#pragma pack(pop)
-#endif
-
-#if PLATFORM_HTML5 && !PLATFORM_HTML5_WIN32
-#undef  _WIN32 
-#undef __MINGW32__
-#endif 
-
-#include "TargetPlatform.h"
-
-#if PLATFORM_LITTLE_ENDIAN
-#define VORBIS_BYTE_ORDER 0
-#else
-#define VORBIS_BYTE_ORDER 1
-#endif
-
-#define WITH_OPUS (PLATFORM_WINDOWS || PLATFORM_MAC || PLATFORM_LINUX)
-
-#if WITH_OPUS
-#include "opus.h"
-#endif
-
-/*------------------------------------------------------------------------------------
-	FVorbisFileWrapper. Hides Vorbis structs from public headers.
-------------------------------------------------------------------------------------*/
-struct FVorbisFileWrapper
-{
-	FVorbisFileWrapper()
-	{
-#if WITH_OGGVORBIS
-		FMemory::Memzero( &vf, sizeof( OggVorbis_File ) );
-#endif
-	}
-
-	~FVorbisFileWrapper()
-	{
-#if WITH_OGGVORBIS
-		ov_clear( &vf ); 
-#endif
-	}
-
-#if WITH_OGGVORBIS
-	/** Ogg vorbis decompression state */
-	OggVorbis_File	vf;
-#endif
-};
-
-#if WITH_OGGVORBIS
-/*------------------------------------------------------------------------------------
-	FVorbisAudioInfo.
-------------------------------------------------------------------------------------*/
-FVorbisAudioInfo::FVorbisAudioInfo( void ) 
-	: VFWrapper(new FVorbisFileWrapper())
-	,	SrcBufferData(NULL)
-	,	SrcBufferDataSize(0)
-	,	BufferOffset(0)
-{ 
-}
-
-FVorbisAudioInfo::~FVorbisAudioInfo( void ) 
-{ 
-	check( VFWrapper != NULL );
-	delete VFWrapper;
-}
-
-/** Emulate read from memory functionality */
-size_t FVorbisAudioInfo::Read( void *Ptr, uint32 Size )
-{
-	size_t BytesToRead = FMath::Min( Size, SrcBufferDataSize - BufferOffset );
-	FMemory::Memcpy( Ptr, SrcBufferData + BufferOffset, BytesToRead );
-	BufferOffset += BytesToRead;
-	return( BytesToRead );
-}
-
-static size_t OggRead( void *ptr, size_t size, size_t nmemb, void *datasource )
-{
-	FVorbisAudioInfo* OggInfo = ( FVorbisAudioInfo* )datasource;
-	return( OggInfo->Read( ptr, size * nmemb ) );
-}
-
-int FVorbisAudioInfo::Seek( uint32 offset, int whence )
-{
-	switch( whence )
-	{
-	case SEEK_SET:
-		BufferOffset = offset;
-		break;
-
-	case SEEK_CUR:
-		BufferOffset += offset;
-		break;
-
-	case SEEK_END:
-		BufferOffset = SrcBufferDataSize - offset;
-		break;
-	}
-
-	return( BufferOffset );
-}
-
-static int OggSeek( void *datasource, ogg_int64_t offset, int whence )
-{
-	FVorbisAudioInfo* OggInfo = ( FVorbisAudioInfo* )datasource;
-	return( OggInfo->Seek( offset, whence ) );
-}
-
-int FVorbisAudioInfo::Close( void )
-{
-	return( 0 );
-}
-
-static int OggClose( void *datasource )
-{
-	FVorbisAudioInfo* OggInfo = ( FVorbisAudioInfo* )datasource;
-	return( OggInfo->Close() );
-}
-
-long FVorbisAudioInfo::Tell( void )
-{
-	return( BufferOffset );
-}
-
-static long OggTell( void *datasource )
-{
-	FVorbisAudioInfo *OggInfo = ( FVorbisAudioInfo* )datasource;
-	return( OggInfo->Tell() );
-}
-
-/** 
- * Reads the header information of an ogg vorbis file
- * 
- * @param	Resource		Info about vorbis data
- */
-bool FVorbisAudioInfo::ReadCompressedInfo( const uint8* InSrcBufferData, uint32 InSrcBufferDataSize, FSoundQualityInfo* QualityInfo )
-{
-	SCOPE_CYCLE_COUNTER( STAT_VorbisPrepareDecompressionTime );
-
-	check( VFWrapper != NULL );
-	ov_callbacks		Callbacks;
-
-	SrcBufferData = InSrcBufferData;
-	SrcBufferDataSize = InSrcBufferDataSize;
-	BufferOffset = 0;
-
-	Callbacks.read_func = OggRead;
-	Callbacks.seek_func = OggSeek;
-	Callbacks.close_func = OggClose;
-	Callbacks.tell_func = OggTell;
-
-	// Set up the read from memory variables
-	if (ov_open_callbacks(this, &VFWrapper->vf, NULL, 0, Callbacks) < 0)
-	{
-		return( false );
-	}
-
-	if( QualityInfo )
-	{
-		// The compression could have resampled the source to make loopable
-		vorbis_info* vi = ov_info( &VFWrapper->vf, -1 );
-		QualityInfo->SampleRate = vi->rate;
-		QualityInfo->NumChannels = vi->channels;
-		ogg_int64_t PCMTotal = ov_pcm_total( &VFWrapper->vf, -1 );
-		if (PCMTotal >= 0)
-		{
-			QualityInfo->SampleDataSize = PCMTotal * QualityInfo->NumChannels * sizeof( int16 );
-			QualityInfo->Duration = ( float )ov_time_total( &VFWrapper->vf, -1 );
-		}
-		else if (PCMTotal == OV_EINVAL)
-		{
-			// indicates an error or that the bitstream is non-seekable
-			QualityInfo->SampleDataSize = 0;
-			QualityInfo->Duration = 0.0f;
-		}
-	}
-
-	return( true );
-}
-
-
-/** 
- * Decompress an entire ogg vorbis data file to a TArray
- */
-void FVorbisAudioInfo::ExpandFile( uint8* DstBuffer, FSoundQualityInfo* QualityInfo )
-{
-	uint32		TotalBytesRead, BytesToRead;
-
-	check( VFWrapper != NULL );
-	check( DstBuffer );
-	check( QualityInfo );
-
-	// A zero buffer size means decompress the entire ogg vorbis stream to PCM.
-	TotalBytesRead = 0;
-	BytesToRead = QualityInfo->SampleDataSize;
-
-	char* Destination = ( char* )DstBuffer;
-	while( TotalBytesRead < BytesToRead )
-	{
-		long BytesRead = ov_read( &VFWrapper->vf, Destination, BytesToRead - TotalBytesRead, 0, 2, 1, NULL );
-
-		if (BytesRead < 0)
-		{
-			// indicates an error - fill remainder of buffer with zero
-			FMemory::Memzero(Destination, BytesToRead - TotalBytesRead);
-			return;
-		}
-
-		TotalBytesRead += BytesRead;
-		Destination += BytesRead;
-	}
-}
-
-
-
-/** 
- * Decompresses ogg vorbis data to raw PCM data. 
- * 
- * @param	PCMData		where to place the decompressed sound
- * @param	bLooping	whether to loop the wav by seeking to the start, or pad the buffer with zeroes
- * @param	BufferSize	number of bytes of PCM data to create. A value of 0 means decompress the entire sound.
- *
- * @return	bool		true if the end of the data was reached (for both single shot and looping sounds)
- */
-bool FVorbisAudioInfo::ReadCompressedData( uint8* InDestination, bool bLooping, uint32 BufferSize )
-{
-	bool		bLooped;
-	uint32		TotalBytesRead;
-
-	SCOPE_CYCLE_COUNTER( STAT_VorbisDecompressTime );
-
-	bLooped = false;
-
-	// Work out number of samples to read
-	TotalBytesRead = 0;
-	char* Destination = ( char* )InDestination;
-
-	check( VFWrapper != NULL );
-
-	while( TotalBytesRead < BufferSize )
-	{
-		long BytesRead = ov_read( &VFWrapper->vf, Destination, BufferSize - TotalBytesRead, 0, 2, 1, NULL );
-		if( !BytesRead )
-		{
-			// We've reached the end
-			bLooped = true;
-			if( bLooping )
-			{
-				int Result = ov_pcm_seek_page( &VFWrapper->vf, 0 );
-				if (Result < 0)
-				{
-					// indicates an error - fill remainder of buffer with zero
-					FMemory::Memzero(Destination, BufferSize - TotalBytesRead);
-					return true;
-				}
-			}
-			else
-			{
-				int32 Count = ( BufferSize - TotalBytesRead );
-				FMemory::Memzero( Destination, Count );
-
-				BytesRead += BufferSize - TotalBytesRead;
-			}
-		}
-		else if (BytesRead < 0)
-		{
-			// indicates an error - fill remainder of buffer with zero
-			FMemory::Memzero(Destination, BufferSize - TotalBytesRead);
-			return false;
-		}
-
-		TotalBytesRead += BytesRead;
-		Destination += BytesRead;
-	}
-
-	return( bLooped );
-}
-
-void FVorbisAudioInfo::SeekToTime( const float SeekTime )
-{
-	const float TargetTime = FMath::Min(SeekTime, ( float )ov_time_total( &VFWrapper->vf, -1 ));
-	ov_time_seek( &VFWrapper->vf, TargetTime );
-}
-
-void FVorbisAudioInfo::EnableHalfRate( bool HalfRate )
-{
-	ov_halfrate( &VFWrapper->vf, int32(HalfRate));
-}
-
-void LoadVorbisLibraries()
-{
-	static bool bIsIntialized = false;
-	if (!bIsIntialized)
-	{
-		bIsIntialized = true;
-#if PLATFORM_WINDOWS  && WITH_OGGVORBIS
-		//@todo if ogg is every ported to another platform, then use the platform abstraction to load these DLLs
-		// Load the Ogg dlls
-		#if _MSC_VER >= 1800
-			FString VSVersion = TEXT("VS2013/");
-		#else
-		FString VSVersion = TEXT("VS2012/");
-		#endif
-
-		FString PlatformString = TEXT("Win32");
-		FString DLLNameStub = TEXT(".dll");
-#if PLATFORM_64BITS
-		PlatformString = TEXT("Win64");
-		DLLNameStub = TEXT("_64.dll");
-#endif
-
-		FString RootOggPath = FPaths::EngineDir() / TEXT("Binaries/ThirdParty/Ogg/") / PlatformString / VSVersion;
-		FString RootVorbisPath = FPaths::EngineDir() / TEXT("Binaries/ThirdParty/Vorbis/") / PlatformString / VSVersion;
-
-		FString DLLToLoad = RootOggPath + TEXT("libogg") + DLLNameStub;
-		verifyf(LoadLibraryW(*DLLToLoad), TEXT("Failed to load DLL %s"), *DLLToLoad);
-		// Load the Vorbis dlls
-		DLLToLoad = RootVorbisPath + TEXT("libvorbis") + DLLNameStub;
-		verifyf(LoadLibraryW(*DLLToLoad), TEXT("Failed to load DLL %s"), *DLLToLoad);
-		DLLToLoad = RootVorbisPath + TEXT("libvorbisfile") + DLLNameStub;
-		verifyf(LoadLibraryW(*DLLToLoad), TEXT("Failed to load DLL %s"), *DLLToLoad);
-#endif	//PLATFORM_WINDOWS
-	}
-}
-
-#endif		// WITH_OGGVORBIS
-
-/*------------------------------------------------------------------------------------
-FOpusDecoderWrapper
-------------------------------------------------------------------------------------*/
-#define USE_UE4_MEM_ALLOC 1
-#define OPUS_SAMPLE_RATE 48000
-#define OPUS_MAX_FRAME_SIZE 5760
-#define SAMPLE_SIZE			( ( uint32 )sizeof( short ) )
-
-struct FOpusDecoderWrapper
-{
-	FOpusDecoderWrapper(uint8 NumChannels)
-	{
-#if WITH_OPUS
-#if USE_UE4_MEM_ALLOC
-		int32 DecSize = opus_decoder_get_size(NumChannels);
-		Decoder = (OpusDecoder*)FMemory::Malloc(DecSize);
-		DecError = opus_decoder_init(Decoder, OPUS_SAMPLE_RATE, NumChannels);
-#else
-		Decoder = opus_decoder_create(OPUS_SAMPLE_RATE, NumChannels, &DecError);
-#endif
-#endif
-	}
-
-	~FOpusDecoderWrapper()
-	{
-#if WITH_OPUS
-#if USE_UE4_MEM_ALLOC
-		FMemory::Free(Decoder);
-#else
-		opus_encoder_destroy(Decoder);
-#endif
-#endif
-	}
-
-#if WITH_OPUS
-	OpusDecoder* Decoder;
-#endif
-	int32 DecError;
-};
-
-/*------------------------------------------------------------------------------------
-FOpusAudioInfo.
-------------------------------------------------------------------------------------*/
-FOpusAudioInfo::FOpusAudioInfo(void)
-	: OpusDecoderWrapper(NULL)
-	, SrcBufferData(NULL)
-	, SrcBufferDataSize(0)
-	, SrcBufferOffset(0)
-	, AudioDataOffset(0)
-	, TrueSampleCount(0)
-	, CurrentSampleCount(0)
-	, NumChannels(0)
-	, SampleStride(0)
-	, LastPCMByteSize(0)
-	, LastPCMOffset(0)
-	, bStoringEndOfFile(false)
-	, StreamingSoundWave(NULL)
-	, CurrentChunkIndex(0)
-{
-}
-
-FOpusAudioInfo::~FOpusAudioInfo(void)
-{
-	if (OpusDecoderWrapper != NULL)
-	{
-		delete OpusDecoderWrapper;
-	}
-}
-
-size_t FOpusAudioInfo::Read(void *ptr, uint32 size)
-{
-	check(SrcBufferOffset < SrcBufferDataSize);
-
-	size_t BytesToRead = FMath::Min(size, SrcBufferDataSize - SrcBufferOffset);
-	FMemory::Memcpy(ptr, SrcBufferData + SrcBufferOffset, BytesToRead);
-	SrcBufferOffset += BytesToRead;
-	return(BytesToRead);
-}
-
-bool FOpusAudioInfo::ReadCompressedInfo(const uint8* InSrcBufferData, uint32 InSrcBufferDataSize, struct FSoundQualityInfo* QualityInfo)
-{
-	SrcBufferData = InSrcBufferData;
-	SrcBufferDataSize = InSrcBufferDataSize;
-	SrcBufferOffset = 0;
-	CurrentSampleCount = 0;
-
-	// Read Identifier, True Sample Count, Number of channels and Frames to Encode first
-	if (FCStringAnsi::Strcmp((char*)InSrcBufferData, OPUS_ID_STRING) != 0)
-	{
-		return false;
-	}
-	else
-	{
-		SrcBufferOffset += FCStringAnsi::Strlen(OPUS_ID_STRING) + 1;
-	}
-	Read(&TrueSampleCount, sizeof(uint32));
-	Read(&NumChannels, sizeof(uint8));
-	uint16 SerializedFrames = 0;
-	Read(&SerializedFrames, sizeof(uint16));
-	AudioDataOffset = SrcBufferOffset;
-	SampleStride = SAMPLE_SIZE * NumChannels;
-
-	OpusDecoderWrapper = new FOpusDecoderWrapper(NumChannels);
-	LastDecodedPCM.Empty();
-	LastDecodedPCM.AddUninitialized(OPUS_MAX_FRAME_SIZE * SampleStride);
-#if WITH_OPUS
-	if (OpusDecoderWrapper->DecError != OPUS_OK)
-	{
-		delete OpusDecoderWrapper;
-		return false;
-	}
-#else
-	return false;
-#endif
-
-	if (QualityInfo)
-	{
-		QualityInfo->SampleRate = OPUS_SAMPLE_RATE;
-		QualityInfo->NumChannels = NumChannels;
-		QualityInfo->SampleDataSize = TrueSampleCount * QualityInfo->NumChannels * SAMPLE_SIZE;
-		QualityInfo->Duration = (float)TrueSampleCount / QualityInfo->SampleRate;
-	}
-
-	return true;
-}
-
-bool FOpusAudioInfo::ReadCompressedData(uint8* Destination, bool bLooping, uint32 BufferSize /*= 0*/)
-{
-	check(OpusDecoderWrapper);
-	check(Destination);
-
-	SCOPE_CYCLE_COUNTER(STAT_OpusDecompressTime);
-
-	// Write out any PCM data that was decoded during the last request
-	uint32 RawPCMOffset = WriteFromDecodedPCM(Destination, BufferSize);
-
-	bool bLooped = false;
-
-	if (bStoringEndOfFile && LastPCMByteSize > 0)
-	{
-		// delayed returning looped because we hadn't read the entire buffer
-		bLooped = true;
-		bStoringEndOfFile = false;
-	}
-
-	while (RawPCMOffset < BufferSize)
-	{
-		uint16 FrameSize = 0;
-		Read(&FrameSize, sizeof(uint16));
-		int32 DecodedSamples = DecompressToPCMBuffer(FrameSize);
-		if (DecodedSamples < 0)
-		{
-			LastPCMByteSize = 0;
-			return false;
-		}
-		else
-		{
-			LastPCMByteSize = IncrementCurrentSampleCount(DecodedSamples) * SampleStride;
-			RawPCMOffset += WriteFromDecodedPCM(Destination + RawPCMOffset, BufferSize - RawPCMOffset);
-
-			if (SrcBufferOffset >= SrcBufferDataSize)
-			{
-				// check whether all decoded PCM was written
-				if (LastPCMByteSize == 0)
-				{
-					bLooped = true;
-				}
-				else
-				{
-					bStoringEndOfFile = true;
-				}
-				if (bLooping)
-				{
-					SrcBufferOffset = AudioDataOffset;
-					CurrentSampleCount = 0;
-				}
-				else
-				{
-					RawPCMOffset += ZeroBuffer(Destination + RawPCMOffset, BufferSize - RawPCMOffset);
-				}
-			}
-		}
-	}
-
-	return bLooped;
-}
-
-void FOpusAudioInfo::ExpandFile(uint8* DstBuffer, struct FSoundQualityInfo* QualityInfo)
-{
-	check(OpusDecoderWrapper);
-	check(DstBuffer);
-	check(QualityInfo);
-	check(QualityInfo->SampleDataSize <= SrcBufferDataSize);
-
-	// Ensure we're at the start of the audio data
-	SrcBufferOffset = AudioDataOffset;
-
-	uint32 RawPCMOffset = 0;
-
-	while (RawPCMOffset < QualityInfo->SampleDataSize)
-	{
-		uint16 FrameSize = 0;
-		Read(&FrameSize, sizeof(uint16));
-		int32 DecodedSamples = DecompressToPCMBuffer(FrameSize);
-
-		if (DecodedSamples < 0)
-		{
-			RawPCMOffset += ZeroBuffer(DstBuffer + RawPCMOffset, QualityInfo->SampleDataSize - RawPCMOffset);
-		}
-		else
-		{
-			LastPCMByteSize = IncrementCurrentSampleCount(DecodedSamples) * SampleStride;
-			RawPCMOffset += WriteFromDecodedPCM(DstBuffer + RawPCMOffset, QualityInfo->SampleDataSize - RawPCMOffset);
-		}
-	}
-}
-
-bool FOpusAudioInfo::StreamCompressedInfo(USoundWave* Wave, struct FSoundQualityInfo* QualityInfo)
-{
-	StreamingSoundWave = Wave;
-
-	// Get the first chunk of audio data (should always be loaded)
-	CurrentChunkIndex = 0;
-	const uint8* FirstChunk = IStreamingManager::Get().GetAudioStreamingManager().GetLoadedChunk(StreamingSoundWave, CurrentChunkIndex);
-
-	if (FirstChunk)
-	{
-		return ReadCompressedInfo(FirstChunk, Wave->RunningPlatformData->Chunks[0].DataSize, QualityInfo);
-	}
-
-	return false;
-}
-
-bool FOpusAudioInfo::StreamCompressedData(uint8* Destination, bool bLooping, uint32 BufferSize)
-{
-	check(OpusDecoderWrapper);
-	check(Destination);
-
-	SCOPE_CYCLE_COUNTER(STAT_OpusDecompressTime);
-
-	// Write out any PCM data that was decoded during the last request
-	uint32 RawPCMOffset = WriteFromDecodedPCM(Destination, BufferSize);
-
-	// If next chunk wasn't loaded when last one finished reading, try to get it again now
-	if (SrcBufferData == NULL)
-	{
-		SrcBufferData = IStreamingManager::Get().GetAudioStreamingManager().GetLoadedChunk(StreamingSoundWave, CurrentChunkIndex);
-		if (SrcBufferData)
-		{
-			SrcBufferDataSize = StreamingSoundWave->RunningPlatformData->Chunks[CurrentChunkIndex].DataSize;
-			SrcBufferOffset = CurrentChunkIndex == 0 ? AudioDataOffset : 0;
-		}
-		else
-		{
-			// Still not loaded, should probably have some kind of error, zero current buffer
-			ZeroBuffer(Destination + RawPCMOffset, BufferSize - RawPCMOffset);
-			return false;
-		}
-	}
-
-	bool bLooped = false;
-
-	if (bStoringEndOfFile && LastPCMByteSize > 0)
-	{
-		// delayed returning looped because we hadn't read the entire buffer
-		bLooped = true;
-		bStoringEndOfFile = false;
-	}
-
-	while (RawPCMOffset < BufferSize)
-	{
-		uint16 FrameSize = 0;
-		Read(&FrameSize, sizeof(uint16));
-		int32 DecodedSamples = DecompressToPCMBuffer(FrameSize);
-
-		if (DecodedSamples < 0)
-		{
-			LastPCMByteSize = 0;
-			return false;
-		}
-		else
-		{
-			LastPCMByteSize = IncrementCurrentSampleCount(DecodedSamples) * SampleStride;
-
-			RawPCMOffset += WriteFromDecodedPCM(Destination + RawPCMOffset, BufferSize - RawPCMOffset);
-
-			// Have we reached the end of buffer
-			if (SrcBufferOffset >= SrcBufferDataSize)
-			{
-				// Special case for the last chunk of audio
-				if (CurrentChunkIndex == StreamingSoundWave->RunningPlatformData->NumChunks - 1)
-				{
-					// check whether all decoded PCM was written
-					if (LastPCMByteSize == 0)
-					{
-						bLooped = true;
-					}
-					else
-					{
-						bStoringEndOfFile = true;
-					}
-					if (bLooping)
-					{
-						CurrentChunkIndex = 0;
-						SrcBufferOffset = AudioDataOffset;
-						CurrentSampleCount = 0;
-					}
-					else
-					{
-						RawPCMOffset += ZeroBuffer(Destination + RawPCMOffset, BufferSize - RawPCMOffset);
-					}
-				}
-				else
-				{
-					CurrentChunkIndex++;
-					SrcBufferOffset = 0;
-				}
-
-				SrcBufferData = IStreamingManager::Get().GetAudioStreamingManager().GetLoadedChunk(StreamingSoundWave, CurrentChunkIndex);
-				if (SrcBufferData)
-				{
-					SrcBufferDataSize = StreamingSoundWave->RunningPlatformData->Chunks[CurrentChunkIndex].DataSize;
-				}
-				else
-				{
-					SrcBufferDataSize = 0;
-					RawPCMOffset += ZeroBuffer(Destination + RawPCMOffset, BufferSize - RawPCMOffset);
-				}
-			}
-		}
-	}
-
-	return bLooped;
-}
-
-int32 FOpusAudioInfo::DecompressToPCMBuffer(uint16 FrameSize)
-{
-	check(SrcBufferOffset + FrameSize <= SrcBufferDataSize);
-
-	const uint8* SrcPtr = SrcBufferData + SrcBufferOffset;
-	SrcBufferOffset += FrameSize;
-	LastPCMOffset = 0;
-#if WITH_OPUS
-	return opus_decode(OpusDecoderWrapper->Decoder, SrcPtr, FrameSize, (int16*)(LastDecodedPCM.GetTypedData()), OPUS_MAX_FRAME_SIZE, 0);
-#else
-	return -1;
-#endif
-}
-
-uint32 FOpusAudioInfo::IncrementCurrentSampleCount(uint32 NewSamples)
-{
-	if (CurrentSampleCount + NewSamples > TrueSampleCount)
-	{
-		NewSamples = TrueSampleCount - CurrentSampleCount;
-		CurrentSampleCount = TrueSampleCount;
-	}
-	else
-	{
-		CurrentSampleCount += NewSamples;
-	}
-	return NewSamples;
-}
-
-uint32 FOpusAudioInfo::WriteFromDecodedPCM(uint8* Destination, uint32 BufferSize)
-{
-	uint32 BytesToCopy = FMath::Min(BufferSize, LastPCMByteSize - LastPCMOffset);
-	if (BytesToCopy > 0)
-	{
-		FMemory::Memcpy(Destination, LastDecodedPCM.GetTypedData() + LastPCMOffset, BytesToCopy);
-		LastPCMOffset += BytesToCopy;
-		if (LastPCMOffset >= LastPCMByteSize)
-		{
-			LastPCMOffset = 0;
-			LastPCMByteSize = 0;
-		}
-	}
-	return BytesToCopy;
-}
-
-uint32	FOpusAudioInfo::ZeroBuffer(uint8* Destination, uint32 BufferSize)
-{
-	check(Destination);
-
-	if (BufferSize > 0)
-	{
-		FMemory::Memzero(Destination, BufferSize);
-		return BufferSize;
-	}
-	return 0;
-}
+#include "IAudioFormat.h"
 
 /**
  * Worker for decompression on a separate thread
diff --git a/Engine/Source/Runtime/Engine/Private/AudioDerivedData.cpp b/Engine/Source/Runtime/Engine/Private/AudioDerivedData.cpp
--- a/Engine/Source/Runtime/Engine/Private/AudioDerivedData.cpp
+++ b/Engine/Source/Runtime/Engine/Private/AudioDerivedData.cpp
@@ -969,7 +969,7 @@ void USoundWave::BeginCachePlatformData()
 
 void USoundWave::BeginCacheForCookedPlatformData(const ITargetPlatform *TargetPlatform)
 {
-	if (IsStreaming())
+	if (TargetPlatform->SupportsFeature(ETargetPlatformFeatures::AudioStreaming) && IsStreaming())
 	{
 		// Retrieve format to cache for targetplatform.
 		FName PlatformFormat = TargetPlatform->GetWaveFormat(this);
diff --git a/Engine/Source/Runtime/Engine/Private/AudioStreaming.cpp b/Engine/Source/Runtime/Engine/Private/AudioStreaming.cpp
--- a/Engine/Source/Runtime/Engine/Private/AudioStreaming.cpp
+++ b/Engine/Source/Runtime/Engine/Private/AudioStreaming.cpp
@@ -353,13 +353,16 @@ void FAudioStreamingManager::UpdateResourceStreaming(float DeltaTime, bool bProc
 		{
 			// Request the chunk the source is using and the one after that
 			FWaveRequest& WaveRequest = GetWaveRequest(Wave);
-			uint32 SourceChunk = Source->GetBuffer()->GetCurrentChunkIndex();
-			WaveRequest.RequiredIndices.AddUnique(SourceChunk);
-			WaveRequest.RequiredIndices.AddUnique((SourceChunk+1)%Wave->RunningPlatformData->NumChunks);
-			if (Source->GetBuffer()->GetCurrentChunkOffset() > Wave->RunningPlatformData->Chunks[SourceChunk].DataSize / 2)
+			int32 SourceChunk = Source->GetBuffer()->GetCurrentChunkIndex();
+			if (SourceChunk >= 0 && SourceChunk < Wave->RunningPlatformData->NumChunks)
 			{
-				// already read over half, request is high priority
-				WaveRequest.bPrioritiseRequest = true;
+				WaveRequest.RequiredIndices.AddUnique(SourceChunk);
+				WaveRequest.RequiredIndices.AddUnique((SourceChunk+1)%Wave->RunningPlatformData->NumChunks);
+				if (Source->GetBuffer()->GetCurrentChunkOffset() > Wave->RunningPlatformData->Chunks[SourceChunk].DataSize / 2)
+				{
+					// already read over half, request is high priority
+					WaveRequest.bPrioritiseRequest = true;
+				}
 			}
 		}
 	}
@@ -406,7 +409,7 @@ void FAudioStreamingManager::RemoveLevel(class ULevel* Level)
 
 void FAudioStreamingManager::AddStreamingSoundWave(USoundWave* SoundWave)
 {
-	if (SoundWave->IsStreaming())
+	if (FPlatformProperties::SupportsAudioStreaming() && SoundWave->IsStreaming())
 	{
 		check(StreamingSoundWaves.Find(SoundWave) == NULL);
 
@@ -442,10 +445,13 @@ bool FAudioStreamingManager::IsStreamingInProgress(const USoundWave* SoundWave)
 
 void FAudioStreamingManager::AddStreamingSoundSource(FSoundSource* SoundSource)
 {
-	const FWaveInstance* WaveInstance = SoundSource->GetWaveInstance();
-	if (WaveInstance && WaveInstance->WaveData && WaveInstance->WaveData->IsStreaming())
+	if (FPlatformProperties::SupportsAudioStreaming())
 	{
-		StreamingSoundSources.AddUnique(SoundSource);
+		const FWaveInstance* WaveInstance = SoundSource->GetWaveInstance();
+		if (WaveInstance && WaveInstance->WaveData && WaveInstance->WaveData->IsStreaming())
+		{
+			StreamingSoundSources.AddUnique(SoundSource);
+		}
 	}
 }
 
diff --git a/Engine/Source/Runtime/Engine/Private/OpusAudioInfo.cpp b/Engine/Source/Runtime/Engine/Private/OpusAudioInfo.cpp
new file mode 100644
--- /dev/null
+++ b/Engine/Source/Runtime/Engine/Private/OpusAudioInfo.cpp
@@ -0,0 +1,404 @@
+// Copyright 1998-2014 Epic Games, Inc. All Rights Reserved.
+
+
+#include "EnginePrivate.h"
+#include "OpusAudioInfo.h"
+#include "IAudioFormat.h"
+#include "Sound/SoundWave.h"
+
+#define WITH_OPUS (PLATFORM_WINDOWS || PLATFORM_MAC || PLATFORM_LINUX)
+
+#if WITH_OPUS
+#include "opus.h"
+#endif
+
+/*------------------------------------------------------------------------------------
+FOpusDecoderWrapper
+------------------------------------------------------------------------------------*/
+#define USE_UE4_MEM_ALLOC 1
+#define OPUS_SAMPLE_RATE 48000
+#define OPUS_MAX_FRAME_SIZE 5760
+#define SAMPLE_SIZE			( ( uint32 )sizeof( short ) )
+
+struct FOpusDecoderWrapper
+{
+	FOpusDecoderWrapper(uint8 NumChannels)
+	{
+#if WITH_OPUS
+#if USE_UE4_MEM_ALLOC
+		int32 DecSize = opus_decoder_get_size(NumChannels);
+		Decoder = (OpusDecoder*)FMemory::Malloc(DecSize);
+		DecError = opus_decoder_init(Decoder, OPUS_SAMPLE_RATE, NumChannels);
+#else
+		Decoder = opus_decoder_create(OPUS_SAMPLE_RATE, NumChannels, &DecError);
+#endif
+#endif
+	}
+
+	~FOpusDecoderWrapper()
+	{
+#if WITH_OPUS
+#if USE_UE4_MEM_ALLOC
+		FMemory::Free(Decoder);
+#else
+		opus_encoder_destroy(Decoder);
+#endif
+#endif
+	}
+
+#if WITH_OPUS
+	OpusDecoder* Decoder;
+#endif
+	int32 DecError;
+};
+
+/*------------------------------------------------------------------------------------
+FOpusAudioInfo.
+------------------------------------------------------------------------------------*/
+FOpusAudioInfo::FOpusAudioInfo(void)
+	: OpusDecoderWrapper(NULL)
+	, SrcBufferData(NULL)
+	, SrcBufferDataSize(0)
+	, SrcBufferOffset(0)
+	, AudioDataOffset(0)
+	, TrueSampleCount(0)
+	, CurrentSampleCount(0)
+	, NumChannels(0)
+	, SampleStride(0)
+	, LastPCMByteSize(0)
+	, LastPCMOffset(0)
+	, bStoringEndOfFile(false)
+	, StreamingSoundWave(NULL)
+	, CurrentChunkIndex(0)
+{
+}
+
+FOpusAudioInfo::~FOpusAudioInfo(void)
+{
+	if (OpusDecoderWrapper != NULL)
+	{
+		delete OpusDecoderWrapper;
+	}
+}
+
+size_t FOpusAudioInfo::Read(void *ptr, uint32 size)
+{
+	check(SrcBufferOffset < SrcBufferDataSize);
+
+	size_t BytesToRead = FMath::Min(size, SrcBufferDataSize - SrcBufferOffset);
+	FMemory::Memcpy(ptr, SrcBufferData + SrcBufferOffset, BytesToRead);
+	SrcBufferOffset += BytesToRead;
+	return(BytesToRead);
+}
+
+bool FOpusAudioInfo::ReadCompressedInfo(const uint8* InSrcBufferData, uint32 InSrcBufferDataSize, struct FSoundQualityInfo* QualityInfo)
+{
+	SrcBufferData = InSrcBufferData;
+	SrcBufferDataSize = InSrcBufferDataSize;
+	SrcBufferOffset = 0;
+	CurrentSampleCount = 0;
+
+	// Read Identifier, True Sample Count, Number of channels and Frames to Encode first
+	if (FCStringAnsi::Strcmp((char*)InSrcBufferData, OPUS_ID_STRING) != 0)
+	{
+		return false;
+	}
+	else
+	{
+		SrcBufferOffset += FCStringAnsi::Strlen(OPUS_ID_STRING) + 1;
+	}
+	Read(&TrueSampleCount, sizeof(uint32));
+	Read(&NumChannels, sizeof(uint8));
+	uint16 SerializedFrames = 0;
+	Read(&SerializedFrames, sizeof(uint16));
+	AudioDataOffset = SrcBufferOffset;
+	SampleStride = SAMPLE_SIZE * NumChannels;
+
+	OpusDecoderWrapper = new FOpusDecoderWrapper(NumChannels);
+	LastDecodedPCM.Empty();
+	LastDecodedPCM.AddUninitialized(OPUS_MAX_FRAME_SIZE * SampleStride);
+#if WITH_OPUS
+	if (OpusDecoderWrapper->DecError != OPUS_OK)
+	{
+		delete OpusDecoderWrapper;
+		return false;
+	}
+#else
+	return false;
+#endif
+
+	if (QualityInfo)
+	{
+		QualityInfo->SampleRate = OPUS_SAMPLE_RATE;
+		QualityInfo->NumChannels = NumChannels;
+		QualityInfo->SampleDataSize = TrueSampleCount * QualityInfo->NumChannels * SAMPLE_SIZE;
+		QualityInfo->Duration = (float)TrueSampleCount / QualityInfo->SampleRate;
+	}
+
+	return true;
+}
+
+bool FOpusAudioInfo::ReadCompressedData(uint8* Destination, bool bLooping, uint32 BufferSize /*= 0*/)
+{
+	check(OpusDecoderWrapper);
+	check(Destination);
+
+	SCOPE_CYCLE_COUNTER(STAT_OpusDecompressTime);
+
+	// Write out any PCM data that was decoded during the last request
+	uint32 RawPCMOffset = WriteFromDecodedPCM(Destination, BufferSize);
+
+	bool bLooped = false;
+
+	if (bStoringEndOfFile && LastPCMByteSize > 0)
+	{
+		// delayed returning looped because we hadn't read the entire buffer
+		bLooped = true;
+		bStoringEndOfFile = false;
+	}
+
+	while (RawPCMOffset < BufferSize)
+	{
+		uint16 FrameSize = 0;
+		Read(&FrameSize, sizeof(uint16));
+		int32 DecodedSamples = DecompressToPCMBuffer(FrameSize);
+		if (DecodedSamples < 0)
+		{
+			LastPCMByteSize = 0;
+			return false;
+		}
+		else
+		{
+			LastPCMByteSize = IncrementCurrentSampleCount(DecodedSamples) * SampleStride;
+			RawPCMOffset += WriteFromDecodedPCM(Destination + RawPCMOffset, BufferSize - RawPCMOffset);
+
+			if (SrcBufferOffset >= SrcBufferDataSize)
+			{
+				// check whether all decoded PCM was written
+				if (LastPCMByteSize == 0)
+				{
+					bLooped = true;
+				}
+				else
+				{
+					bStoringEndOfFile = true;
+				}
+				if (bLooping)
+				{
+					SrcBufferOffset = AudioDataOffset;
+					CurrentSampleCount = 0;
+				}
+				else
+				{
+					RawPCMOffset += ZeroBuffer(Destination + RawPCMOffset, BufferSize - RawPCMOffset);
+				}
+			}
+		}
+	}
+
+	return bLooped;
+}
+
+void FOpusAudioInfo::ExpandFile(uint8* DstBuffer, struct FSoundQualityInfo* QualityInfo)
+{
+	check(OpusDecoderWrapper);
+	check(DstBuffer);
+	check(QualityInfo);
+	check(QualityInfo->SampleDataSize <= SrcBufferDataSize);
+
+	// Ensure we're at the start of the audio data
+	SrcBufferOffset = AudioDataOffset;
+
+	uint32 RawPCMOffset = 0;
+
+	while (RawPCMOffset < QualityInfo->SampleDataSize)
+	{
+		uint16 FrameSize = 0;
+		Read(&FrameSize, sizeof(uint16));
+		int32 DecodedSamples = DecompressToPCMBuffer(FrameSize);
+
+		if (DecodedSamples < 0)
+		{
+			RawPCMOffset += ZeroBuffer(DstBuffer + RawPCMOffset, QualityInfo->SampleDataSize - RawPCMOffset);
+		}
+		else
+		{
+			LastPCMByteSize = IncrementCurrentSampleCount(DecodedSamples) * SampleStride;
+			RawPCMOffset += WriteFromDecodedPCM(DstBuffer + RawPCMOffset, QualityInfo->SampleDataSize - RawPCMOffset);
+		}
+	}
+}
+
+bool FOpusAudioInfo::StreamCompressedInfo(USoundWave* Wave, struct FSoundQualityInfo* QualityInfo)
+{
+	StreamingSoundWave = Wave;
+
+	// Get the first chunk of audio data (should always be loaded)
+	CurrentChunkIndex = 0;
+	const uint8* FirstChunk = IStreamingManager::Get().GetAudioStreamingManager().GetLoadedChunk(StreamingSoundWave, CurrentChunkIndex);
+
+	if (FirstChunk)
+	{
+		return ReadCompressedInfo(FirstChunk, Wave->RunningPlatformData->Chunks[0].DataSize, QualityInfo);
+	}
+
+	return false;
+}
+
+bool FOpusAudioInfo::StreamCompressedData(uint8* Destination, bool bLooping, uint32 BufferSize)
+{
+	check(OpusDecoderWrapper);
+	check(Destination);
+
+	SCOPE_CYCLE_COUNTER(STAT_OpusDecompressTime);
+
+	// Write out any PCM data that was decoded during the last request
+	uint32 RawPCMOffset = WriteFromDecodedPCM(Destination, BufferSize);
+
+	// If next chunk wasn't loaded when last one finished reading, try to get it again now
+	if (SrcBufferData == NULL)
+	{
+		SrcBufferData = IStreamingManager::Get().GetAudioStreamingManager().GetLoadedChunk(StreamingSoundWave, CurrentChunkIndex);
+		if (SrcBufferData)
+		{
+			SrcBufferDataSize = StreamingSoundWave->RunningPlatformData->Chunks[CurrentChunkIndex].DataSize;
+			SrcBufferOffset = CurrentChunkIndex == 0 ? AudioDataOffset : 0;
+		}
+		else
+		{
+			// Still not loaded, should probably have some kind of error, zero current buffer
+			ZeroBuffer(Destination + RawPCMOffset, BufferSize - RawPCMOffset);
+			return false;
+		}
+	}
+
+	bool bLooped = false;
+
+	if (bStoringEndOfFile && LastPCMByteSize > 0)
+	{
+		// delayed returning looped because we hadn't read the entire buffer
+		bLooped = true;
+		bStoringEndOfFile = false;
+	}
+
+	while (RawPCMOffset < BufferSize)
+	{
+		uint16 FrameSize = 0;
+		Read(&FrameSize, sizeof(uint16));
+		int32 DecodedSamples = DecompressToPCMBuffer(FrameSize);
+
+		if (DecodedSamples < 0)
+		{
+			LastPCMByteSize = 0;
+			return false;
+		}
+		else
+		{
+			LastPCMByteSize = IncrementCurrentSampleCount(DecodedSamples) * SampleStride;
+
+			RawPCMOffset += WriteFromDecodedPCM(Destination + RawPCMOffset, BufferSize - RawPCMOffset);
+
+			// Have we reached the end of buffer
+			if (SrcBufferOffset >= SrcBufferDataSize)
+			{
+				// Special case for the last chunk of audio
+				if (CurrentChunkIndex == StreamingSoundWave->RunningPlatformData->NumChunks - 1)
+				{
+					// check whether all decoded PCM was written
+					if (LastPCMByteSize == 0)
+					{
+						bLooped = true;
+					}
+					else
+					{
+						bStoringEndOfFile = true;
+					}
+					if (bLooping)
+					{
+						CurrentChunkIndex = 0;
+						SrcBufferOffset = AudioDataOffset;
+						CurrentSampleCount = 0;
+					}
+					else
+					{
+						RawPCMOffset += ZeroBuffer(Destination + RawPCMOffset, BufferSize - RawPCMOffset);
+					}
+				}
+				else
+				{
+					CurrentChunkIndex++;
+					SrcBufferOffset = 0;
+				}
+
+				SrcBufferData = IStreamingManager::Get().GetAudioStreamingManager().GetLoadedChunk(StreamingSoundWave, CurrentChunkIndex);
+				if (SrcBufferData)
+				{
+					SrcBufferDataSize = StreamingSoundWave->RunningPlatformData->Chunks[CurrentChunkIndex].DataSize;
+				}
+				else
+				{
+					SrcBufferDataSize = 0;
+					RawPCMOffset += ZeroBuffer(Destination + RawPCMOffset, BufferSize - RawPCMOffset);
+				}
+			}
+		}
+	}
+
+	return bLooped;
+}
+
+int32 FOpusAudioInfo::DecompressToPCMBuffer(uint16 FrameSize)
+{
+	check(SrcBufferOffset + FrameSize <= SrcBufferDataSize);
+
+	const uint8* SrcPtr = SrcBufferData + SrcBufferOffset;
+	SrcBufferOffset += FrameSize;
+	LastPCMOffset = 0;
+#if WITH_OPUS
+	return opus_decode(OpusDecoderWrapper->Decoder, SrcPtr, FrameSize, (int16*)(LastDecodedPCM.GetTypedData()), OPUS_MAX_FRAME_SIZE, 0);
+#else
+	return -1;
+#endif
+}
+
+uint32 FOpusAudioInfo::IncrementCurrentSampleCount(uint32 NewSamples)
+{
+	if (CurrentSampleCount + NewSamples > TrueSampleCount)
+	{
+		NewSamples = TrueSampleCount - CurrentSampleCount;
+		CurrentSampleCount = TrueSampleCount;
+	}
+	else
+	{
+		CurrentSampleCount += NewSamples;
+	}
+	return NewSamples;
+}
+
+uint32 FOpusAudioInfo::WriteFromDecodedPCM(uint8* Destination, uint32 BufferSize)
+{
+	uint32 BytesToCopy = FMath::Min(BufferSize, LastPCMByteSize - LastPCMOffset);
+	if (BytesToCopy > 0)
+	{
+		FMemory::Memcpy(Destination, LastDecodedPCM.GetTypedData() + LastPCMOffset, BytesToCopy);
+		LastPCMOffset += BytesToCopy;
+		if (LastPCMOffset >= LastPCMByteSize)
+		{
+			LastPCMOffset = 0;
+			LastPCMByteSize = 0;
+		}
+	}
+	return BytesToCopy;
+}
+
+uint32	FOpusAudioInfo::ZeroBuffer(uint8* Destination, uint32 BufferSize)
+{
+	check(Destination);
+
+	if (BufferSize > 0)
+	{
+		FMemory::Memzero(Destination, BufferSize);
+		return BufferSize;
+	}
+	return 0;
+}
diff --git a/Engine/Source/Runtime/Engine/Private/VorbisAudioInfo.cpp b/Engine/Source/Runtime/Engine/Private/VorbisAudioInfo.cpp
new file mode 100644
--- /dev/null
+++ b/Engine/Source/Runtime/Engine/Private/VorbisAudioInfo.cpp
@@ -0,0 +1,335 @@
+// Copyright 1998-2014 Epic Games, Inc. All Rights Reserved.
+
+
+#include "EnginePrivate.h"
+#include "VorbisAudioInfo.h"
+#include "IAudioFormat.h"
+
+#if WITH_OGGVORBIS
+
+// hack to get ogg types right for HTML5. 
+#if PLATFORM_HTML5 && !PLATFORM_HTML5_WIN32
+#define _WIN32 
+#define __MINGW32__
+#endif 
+#pragma pack(push, 8)
+#include "ogg/ogg.h"
+#include "vorbis/vorbisenc.h"
+#include "vorbis/vorbisfile.h"
+#pragma pack(pop)
+#endif
+
+#if PLATFORM_HTML5 && !PLATFORM_HTML5_WIN32
+#undef  _WIN32 
+#undef __MINGW32__
+#endif 
+
+#if PLATFORM_LITTLE_ENDIAN
+#define VORBIS_BYTE_ORDER 0
+#else
+#define VORBIS_BYTE_ORDER 1
+#endif
+
+/*------------------------------------------------------------------------------------
+	FVorbisFileWrapper. Hides Vorbis structs from public headers.
+------------------------------------------------------------------------------------*/
+struct FVorbisFileWrapper
+{
+	FVorbisFileWrapper()
+	{
+#if WITH_OGGVORBIS
+		FMemory::Memzero( &vf, sizeof( OggVorbis_File ) );
+#endif
+	}
+
+	~FVorbisFileWrapper()
+	{
+#if WITH_OGGVORBIS
+		ov_clear( &vf ); 
+#endif
+	}
+
+#if WITH_OGGVORBIS
+	/** Ogg vorbis decompression state */
+	OggVorbis_File	vf;
+#endif
+};
+
+#if WITH_OGGVORBIS
+/*------------------------------------------------------------------------------------
+	FVorbisAudioInfo.
+------------------------------------------------------------------------------------*/
+FVorbisAudioInfo::FVorbisAudioInfo( void ) 
+	: VFWrapper(new FVorbisFileWrapper())
+	,	SrcBufferData(NULL)
+	,	SrcBufferDataSize(0)
+	,	BufferOffset(0)
+{ 
+}
+
+FVorbisAudioInfo::~FVorbisAudioInfo( void ) 
+{ 
+	check( VFWrapper != NULL );
+	delete VFWrapper;
+}
+
+/** Emulate read from memory functionality */
+size_t FVorbisAudioInfo::Read( void *Ptr, uint32 Size )
+{
+	size_t BytesToRead = FMath::Min( Size, SrcBufferDataSize - BufferOffset );
+	FMemory::Memcpy( Ptr, SrcBufferData + BufferOffset, BytesToRead );
+	BufferOffset += BytesToRead;
+	return( BytesToRead );
+}
+
+static size_t OggRead( void *ptr, size_t size, size_t nmemb, void *datasource )
+{
+	FVorbisAudioInfo* OggInfo = ( FVorbisAudioInfo* )datasource;
+	return( OggInfo->Read( ptr, size * nmemb ) );
+}
+
+int FVorbisAudioInfo::Seek( uint32 offset, int whence )
+{
+	switch( whence )
+	{
+	case SEEK_SET:
+		BufferOffset = offset;
+		break;
+
+	case SEEK_CUR:
+		BufferOffset += offset;
+		break;
+
+	case SEEK_END:
+		BufferOffset = SrcBufferDataSize - offset;
+		break;
+	}
+
+	return( BufferOffset );
+}
+
+static int OggSeek( void *datasource, ogg_int64_t offset, int whence )
+{
+	FVorbisAudioInfo* OggInfo = ( FVorbisAudioInfo* )datasource;
+	return( OggInfo->Seek( offset, whence ) );
+}
+
+int FVorbisAudioInfo::Close( void )
+{
+	return( 0 );
+}
+
+static int OggClose( void *datasource )
+{
+	FVorbisAudioInfo* OggInfo = ( FVorbisAudioInfo* )datasource;
+	return( OggInfo->Close() );
+}
+
+long FVorbisAudioInfo::Tell( void )
+{
+	return( BufferOffset );
+}
+
+static long OggTell( void *datasource )
+{
+	FVorbisAudioInfo *OggInfo = ( FVorbisAudioInfo* )datasource;
+	return( OggInfo->Tell() );
+}
+
+/** 
+ * Reads the header information of an ogg vorbis file
+ * 
+ * @param	Resource		Info about vorbis data
+ */
+bool FVorbisAudioInfo::ReadCompressedInfo( const uint8* InSrcBufferData, uint32 InSrcBufferDataSize, FSoundQualityInfo* QualityInfo )
+{
+	SCOPE_CYCLE_COUNTER( STAT_VorbisPrepareDecompressionTime );
+
+	check( VFWrapper != NULL );
+	ov_callbacks		Callbacks;
+
+	SrcBufferData = InSrcBufferData;
+	SrcBufferDataSize = InSrcBufferDataSize;
+	BufferOffset = 0;
+
+	Callbacks.read_func = OggRead;
+	Callbacks.seek_func = OggSeek;
+	Callbacks.close_func = OggClose;
+	Callbacks.tell_func = OggTell;
+
+	// Set up the read from memory variables
+	if (ov_open_callbacks(this, &VFWrapper->vf, NULL, 0, Callbacks) < 0)
+	{
+		return( false );
+	}
+
+	if( QualityInfo )
+	{
+		// The compression could have resampled the source to make loopable
+		vorbis_info* vi = ov_info( &VFWrapper->vf, -1 );
+		QualityInfo->SampleRate = vi->rate;
+		QualityInfo->NumChannels = vi->channels;
+		ogg_int64_t PCMTotal = ov_pcm_total( &VFWrapper->vf, -1 );
+		if (PCMTotal >= 0)
+		{
+			QualityInfo->SampleDataSize = PCMTotal * QualityInfo->NumChannels * sizeof( int16 );
+			QualityInfo->Duration = ( float )ov_time_total( &VFWrapper->vf, -1 );
+		}
+		else if (PCMTotal == OV_EINVAL)
+		{
+			// indicates an error or that the bitstream is non-seekable
+			QualityInfo->SampleDataSize = 0;
+			QualityInfo->Duration = 0.0f;
+		}
+	}
+
+	return( true );
+}
+
+
+/** 
+ * Decompress an entire ogg vorbis data file to a TArray
+ */
+void FVorbisAudioInfo::ExpandFile( uint8* DstBuffer, FSoundQualityInfo* QualityInfo )
+{
+	uint32		TotalBytesRead, BytesToRead;
+
+	check( VFWrapper != NULL );
+	check( DstBuffer );
+	check( QualityInfo );
+
+	// A zero buffer size means decompress the entire ogg vorbis stream to PCM.
+	TotalBytesRead = 0;
+	BytesToRead = QualityInfo->SampleDataSize;
+
+	char* Destination = ( char* )DstBuffer;
+	while( TotalBytesRead < BytesToRead )
+	{
+		long BytesRead = ov_read( &VFWrapper->vf, Destination, BytesToRead - TotalBytesRead, 0, 2, 1, NULL );
+
+		if (BytesRead < 0)
+		{
+			// indicates an error - fill remainder of buffer with zero
+			FMemory::Memzero(Destination, BytesToRead - TotalBytesRead);
+			return;
+		}
+
+		TotalBytesRead += BytesRead;
+		Destination += BytesRead;
+	}
+}
+
+
+
+/** 
+ * Decompresses ogg vorbis data to raw PCM data. 
+ * 
+ * @param	PCMData		where to place the decompressed sound
+ * @param	bLooping	whether to loop the wav by seeking to the start, or pad the buffer with zeroes
+ * @param	BufferSize	number of bytes of PCM data to create. A value of 0 means decompress the entire sound.
+ *
+ * @return	bool		true if the end of the data was reached (for both single shot and looping sounds)
+ */
+bool FVorbisAudioInfo::ReadCompressedData( uint8* InDestination, bool bLooping, uint32 BufferSize )
+{
+	bool		bLooped;
+	uint32		TotalBytesRead;
+
+	SCOPE_CYCLE_COUNTER( STAT_VorbisDecompressTime );
+
+	bLooped = false;
+
+	// Work out number of samples to read
+	TotalBytesRead = 0;
+	char* Destination = ( char* )InDestination;
+
+	check( VFWrapper != NULL );
+
+	while( TotalBytesRead < BufferSize )
+	{
+		long BytesRead = ov_read( &VFWrapper->vf, Destination, BufferSize - TotalBytesRead, 0, 2, 1, NULL );
+		if( !BytesRead )
+		{
+			// We've reached the end
+			bLooped = true;
+			if( bLooping )
+			{
+				int Result = ov_pcm_seek_page( &VFWrapper->vf, 0 );
+				if (Result < 0)
+				{
+					// indicates an error - fill remainder of buffer with zero
+					FMemory::Memzero(Destination, BufferSize - TotalBytesRead);
+					return true;
+				}
+			}
+			else
+			{
+				int32 Count = ( BufferSize - TotalBytesRead );
+				FMemory::Memzero( Destination, Count );
+
+				BytesRead += BufferSize - TotalBytesRead;
+			}
+		}
+		else if (BytesRead < 0)
+		{
+			// indicates an error - fill remainder of buffer with zero
+			FMemory::Memzero(Destination, BufferSize - TotalBytesRead);
+			return false;
+		}
+
+		TotalBytesRead += BytesRead;
+		Destination += BytesRead;
+	}
+
+	return( bLooped );
+}
+
+void FVorbisAudioInfo::SeekToTime( const float SeekTime )
+{
+	const float TargetTime = FMath::Min(SeekTime, ( float )ov_time_total( &VFWrapper->vf, -1 ));
+	ov_time_seek( &VFWrapper->vf, TargetTime );
+}
+
+void FVorbisAudioInfo::EnableHalfRate( bool HalfRate )
+{
+	ov_halfrate( &VFWrapper->vf, int32(HalfRate));
+}
+
+void LoadVorbisLibraries()
+{
+	static bool bIsIntialized = false;
+	if (!bIsIntialized)
+	{
+		bIsIntialized = true;
+#if PLATFORM_WINDOWS  && WITH_OGGVORBIS
+		//@todo if ogg is every ported to another platform, then use the platform abstraction to load these DLLs
+		// Load the Ogg dlls
+		#if _MSC_VER >= 1800
+			FString VSVersion = TEXT("VS2013/");
+		#else
+		FString VSVersion = TEXT("VS2012/");
+		#endif
+
+		FString PlatformString = TEXT("Win32");
+		FString DLLNameStub = TEXT(".dll");
+#if PLATFORM_64BITS
+		PlatformString = TEXT("Win64");
+		DLLNameStub = TEXT("_64.dll");
+#endif
+
+		FString RootOggPath = FPaths::EngineDir() / TEXT("Binaries/ThirdParty/Ogg/") / PlatformString / VSVersion;
+		FString RootVorbisPath = FPaths::EngineDir() / TEXT("Binaries/ThirdParty/Vorbis/") / PlatformString / VSVersion;
+
+		FString DLLToLoad = RootOggPath + TEXT("libogg") + DLLNameStub;
+		verifyf(LoadLibraryW(*DLLToLoad), TEXT("Failed to load DLL %s"), *DLLToLoad);
+		// Load the Vorbis dlls
+		DLLToLoad = RootVorbisPath + TEXT("libvorbis") + DLLNameStub;
+		verifyf(LoadLibraryW(*DLLToLoad), TEXT("Failed to load DLL %s"), *DLLToLoad);
+		DLLToLoad = RootVorbisPath + TEXT("libvorbisfile") + DLLNameStub;
+		verifyf(LoadLibraryW(*DLLToLoad), TEXT("Failed to load DLL %s"), *DLLToLoad);
+#endif	//PLATFORM_WINDOWS
+	}
+}
+
+#endif		// WITH_OGGVORBIS
+
diff --git a/Engine/Source/Runtime/Engine/Public/AudioDecompress.h b/Engine/Source/Runtime/Engine/Public/AudioDecompress.h
--- a/Engine/Source/Runtime/Engine/Public/AudioDecompress.h
+++ b/Engine/Source/Runtime/Engine/Public/AudioDecompress.h
@@ -6,27 +6,11 @@
 
 #pragma once
 
-/**
- * Whether to use OggVorbis audio format.
- **/
-#ifndef WITH_OGGVORBIS
-	#if PLATFORM_DESKTOP
-		#define WITH_OGGVORBIS 1
-	#else
-		#define WITH_OGGVORBIS 0
-	#endif
-#endif
-
 // 186ms of 44.1KHz data
 // 372ms of 22KHz data
 #define MONO_PCM_BUFFER_SAMPLES		8192
 #define MONO_PCM_BUFFER_SIZE		( MONO_PCM_BUFFER_SAMPLES * sizeof( int16 ) )
 
-/**
- * Loads voribs dlls
-*/
-ENGINE_API void LoadVorbisLibraries();
-
 /**
  * Interface class to decompress various types of audio data
  */
@@ -119,150 +103,6 @@ public:
 	virtual int32 GetCurrentChunkOffset() const {return -1;}
 };
 
-#if WITH_OGGVORBIS
-/** 
- * Helper class to parse ogg vorbis data
- */
-class FVorbisAudioInfo : public ICompressedAudioInfo
-{
-public:
-	ENGINE_API FVorbisAudioInfo( void );
-	ENGINE_API virtual ~FVorbisAudioInfo( void );
-
-	/** Emulate read from memory functionality */
-	size_t			Read( void *ptr, uint32 size );
-	int				Seek( uint32 offset, int whence );
-	int				Close( void );
-	long			Tell( void );
-
-	/** 
-	 * Reads the header information of an ogg vorbis file
-	 * 
-	 * @param	Resource		Info about vorbis data
-	 */
-	ENGINE_API virtual bool ReadCompressedInfo( const uint8* InSrcBufferData, uint32 InSrcBufferDataSize, struct FSoundQualityInfo* QualityInfo );
-
-	/** 
-	 * Decompresses ogg data to raw PCM data. 
-	 * 
-	 * @param	Destination	where to place the decompressed sound
-	 * @param	bLooping	whether to loop the sound by seeking to the start, or pad the buffer with zeroes
-	 * @param	BufferSize	number of bytes of PCM data to create
-	 *
-	 * @return	bool		true if the end of the data was reached (for both single shot and looping sounds)
-	 */
-	ENGINE_API virtual bool ReadCompressedData( uint8* Destination, bool bLooping, uint32 BufferSize );
-
-	ENGINE_API virtual void SeekToTime( const float SeekTime );
-
-	/** 
-	 * Decompress an entire ogg data file to a TArray
-	 */
-	ENGINE_API virtual void ExpandFile( uint8* DstBuffer, struct FSoundQualityInfo* QualityInfo );
-
-	/** 
-	 * Sets ogg to decode to half-rate
-	 * 
-	 * @param	Resource		Info about vorbis data
-	 */
-	ENGINE_API virtual void EnableHalfRate( bool HalfRate );
-
-	virtual uint32 GetSourceBufferSize() const { return SrcBufferDataSize;}
-
-	struct FVorbisFileWrapper* VFWrapper;
-	const uint8*		SrcBufferData;
-	uint32			SrcBufferDataSize;
-	uint32			BufferOffset;
-};
-#endif
-
-#define OPUS_ID_STRING "UE4OPUS"
-
-struct FOpusDecoderWrapper;
-
-/**
-* Helper class to parse opus data
-*/
-class FOpusAudioInfo : public ICompressedAudioInfo
-{
-public:
-	ENGINE_API FOpusAudioInfo(void);
-	ENGINE_API virtual ~FOpusAudioInfo(void);
-
-	// ICompressedAudioInfo Interface
-	ENGINE_API virtual bool ReadCompressedInfo(const uint8* InSrcBufferData, uint32 InSrcBufferDataSize, struct FSoundQualityInfo* QualityInfo) override;
-	ENGINE_API virtual bool ReadCompressedData(uint8* Destination, bool bLooping, uint32 BufferSize) override;
-	ENGINE_API virtual void SeekToTime(const float SeekTime) override {};
-	ENGINE_API virtual void ExpandFile(uint8* DstBuffer, struct FSoundQualityInfo* QualityInfo) override;
-	ENGINE_API virtual void EnableHalfRate(bool HalfRate) override {};
-	virtual uint32 GetSourceBufferSize() const override { return SrcBufferDataSize;}
-
-	virtual bool SupportsStreaming() const override {return true;}
-	virtual bool StreamCompressedInfo(USoundWave* Wave, struct FSoundQualityInfo* QualityInfo) override;
-	virtual bool StreamCompressedData(uint8* Destination, bool bLooping, uint32 BufferSize) override;
-	virtual int32 GetCurrentChunkIndex() const override {return CurrentChunkIndex;}
-	virtual int32 GetCurrentChunkOffset() const override {return SrcBufferOffset;}
-	// End of ICompressedAudioInfo Interface
-
-protected:
-	/** Emulate read from memory functionality */
-	size_t			Read(void *ptr, uint32 size);
-
-	/**
-	 * Decompresses a frame of Opus data to PCM buffer
-	 *
-	 * @param FrameSize Size of the frame in bytes
-	 * @return The amount of samples that were decompressed (< 0 indicates error)
-	 */
-	int32 DecompressToPCMBuffer(uint16 FrameSize);
-
-	/**
-	 * Adds to the count of samples that have currently been decoded
-	 *
-	 * @param NewSamples	How many samples have been decoded
-	 * @return How many samples were actually part of the true sample count
-	 */
-	uint32 IncrementCurrentSampleCount(uint32 NewSamples);
-
-	/**
-	 * Writes data from decoded PCM buffer, taking into account whether some PCM has been written before
-	 *
-	 * @param Destination	Where to place the decoded sound
-	 * @param BufferSize	Size of the destination buffer in bytes
-	 * @return				How many bytes were written
-	 */
-	uint32	WriteFromDecodedPCM(uint8* Destination, uint32 BufferSize);
-
-	/**
-	 * Zeroes the contents of a buffer
-	 *
-	 * @param Destination	Buffer to zero
-	 * @param BufferSize	Size of the destination buffer in bytes
-	 * @return				How many bytes were zeroed
-	 */
-	uint32	ZeroBuffer(uint8* Destination, uint32 BufferSize);
-
-	FOpusDecoderWrapper*	OpusDecoderWrapper;
-	const uint8*	SrcBufferData;
-	uint32			SrcBufferDataSize;
-	uint32			SrcBufferOffset;
-	uint32			AudioDataOffset;
-
-	uint32			TrueSampleCount;
-	uint32			CurrentSampleCount;
-	uint8			NumChannels;
-	uint32			SampleStride;
-
-	TArray<uint8>	LastDecodedPCM;
-	uint32			LastPCMByteSize;
-	uint32			LastPCMOffset;
-	bool			bStoringEndOfFile;
-
-	// Streaming specific
-	USoundWave*		StreamingSoundWave;
-	int32			CurrentChunkIndex;
-};
-
 /**
  * Asynchronous audio decompression
  */
diff --git a/Engine/Source/Runtime/Engine/Public/OpusAudioInfo.h b/Engine/Source/Runtime/Engine/Public/OpusAudioInfo.h
new file mode 100644
--- /dev/null
+++ b/Engine/Source/Runtime/Engine/Public/OpusAudioInfo.h
@@ -0,0 +1,96 @@
+// Copyright 1998-2014 Epic Games, Inc. All Rights Reserved.
+
+/*=============================================================================
+	OpusAudioInfo.h: Unreal audio opus decompression interface object.
+=============================================================================*/
+
+#pragma once
+
+#include "AudioDecompress.h"
+
+#define OPUS_ID_STRING "UE4OPUS"
+
+struct FOpusDecoderWrapper;
+
+/**
+* Helper class to parse opus data
+*/
+class FOpusAudioInfo : public ICompressedAudioInfo
+{
+public:
+	ENGINE_API FOpusAudioInfo(void);
+	ENGINE_API virtual ~FOpusAudioInfo(void);
+
+	// ICompressedAudioInfo Interface
+	ENGINE_API virtual bool ReadCompressedInfo(const uint8* InSrcBufferData, uint32 InSrcBufferDataSize, struct FSoundQualityInfo* QualityInfo) override;
+	ENGINE_API virtual bool ReadCompressedData(uint8* Destination, bool bLooping, uint32 BufferSize) override;
+	ENGINE_API virtual void SeekToTime(const float SeekTime) override {};
+	ENGINE_API virtual void ExpandFile(uint8* DstBuffer, struct FSoundQualityInfo* QualityInfo) override;
+	ENGINE_API virtual void EnableHalfRate(bool HalfRate) override {};
+	virtual uint32 GetSourceBufferSize() const override { return SrcBufferDataSize;}
+
+	virtual bool SupportsStreaming() const override {return true;}
+	virtual bool StreamCompressedInfo(USoundWave* Wave, struct FSoundQualityInfo* QualityInfo) override;
+	virtual bool StreamCompressedData(uint8* Destination, bool bLooping, uint32 BufferSize) override;
+	virtual int32 GetCurrentChunkIndex() const override {return CurrentChunkIndex;}
+	virtual int32 GetCurrentChunkOffset() const override {return SrcBufferOffset;}
+	// End of ICompressedAudioInfo Interface
+
+protected:
+	/** Emulate read from memory functionality */
+	size_t			Read(void *ptr, uint32 size);
+
+	/**
+	 * Decompresses a frame of Opus data to PCM buffer
+	 *
+	 * @param FrameSize Size of the frame in bytes
+	 * @return The amount of samples that were decompressed (< 0 indicates error)
+	 */
+	int32 DecompressToPCMBuffer(uint16 FrameSize);
+
+	/**
+	 * Adds to the count of samples that have currently been decoded
+	 *
+	 * @param NewSamples	How many samples have been decoded
+	 * @return How many samples were actually part of the true sample count
+	 */
+	uint32 IncrementCurrentSampleCount(uint32 NewSamples);
+
+	/**
+	 * Writes data from decoded PCM buffer, taking into account whether some PCM has been written before
+	 *
+	 * @param Destination	Where to place the decoded sound
+	 * @param BufferSize	Size of the destination buffer in bytes
+	 * @return				How many bytes were written
+	 */
+	uint32	WriteFromDecodedPCM(uint8* Destination, uint32 BufferSize);
+
+	/**
+	 * Zeroes the contents of a buffer
+	 *
+	 * @param Destination	Buffer to zero
+	 * @param BufferSize	Size of the destination buffer in bytes
+	 * @return				How many bytes were zeroed
+	 */
+	uint32	ZeroBuffer(uint8* Destination, uint32 BufferSize);
+
+	FOpusDecoderWrapper*	OpusDecoderWrapper;
+	const uint8*	SrcBufferData;
+	uint32			SrcBufferDataSize;
+	uint32			SrcBufferOffset;
+	uint32			AudioDataOffset;
+
+	uint32			TrueSampleCount;
+	uint32			CurrentSampleCount;
+	uint8			NumChannels;
+	uint32			SampleStride;
+
+	TArray<uint8>	LastDecodedPCM;
+	uint32			LastPCMByteSize;
+	uint32			LastPCMOffset;
+	bool			bStoringEndOfFile;
+
+	// Streaming specific
+	USoundWave*		StreamingSoundWave;
+	int32			CurrentChunkIndex;
+};
diff --git a/Engine/Source/Runtime/Engine/Public/VorbisAudioInfo.h b/Engine/Source/Runtime/Engine/Public/VorbisAudioInfo.h
new file mode 100644
--- /dev/null
+++ b/Engine/Source/Runtime/Engine/Public/VorbisAudioInfo.h
@@ -0,0 +1,82 @@
+// Copyright 1998-2014 Epic Games, Inc. All Rights Reserved.
+
+/*=============================================================================
+	VorbisAudioInfo.h: Unreal audio vorbis decompression interface object.
+=============================================================================*/
+
+#pragma once
+
+#include "AudioDecompress.h"
+
+/**
+ * Whether to use OggVorbis audio format.
+ **/
+#ifndef WITH_OGGVORBIS
+	#if PLATFORM_DESKTOP
+		#define WITH_OGGVORBIS 1
+	#else
+		#define WITH_OGGVORBIS 0
+	#endif
+#endif
+
+/**
+ * Loads vorbis dlls
+*/
+ENGINE_API void LoadVorbisLibraries();
+
+#if WITH_OGGVORBIS
+/** 
+ * Helper class to parse ogg vorbis data
+ */
+class FVorbisAudioInfo : public ICompressedAudioInfo
+{
+public:
+	ENGINE_API FVorbisAudioInfo( void );
+	ENGINE_API virtual ~FVorbisAudioInfo( void );
+
+	/** Emulate read from memory functionality */
+	size_t			Read( void *ptr, uint32 size );
+	int				Seek( uint32 offset, int whence );
+	int				Close( void );
+	long			Tell( void );
+
+	/** 
+	 * Reads the header information of an ogg vorbis file
+	 * 
+	 * @param	Resource		Info about vorbis data
+	 */
+	ENGINE_API virtual bool ReadCompressedInfo( const uint8* InSrcBufferData, uint32 InSrcBufferDataSize, struct FSoundQualityInfo* QualityInfo );
+
+	/** 
+	 * Decompresses ogg data to raw PCM data. 
+	 * 
+	 * @param	Destination	where to place the decompressed sound
+	 * @param	bLooping	whether to loop the sound by seeking to the start, or pad the buffer with zeroes
+	 * @param	BufferSize	number of bytes of PCM data to create
+	 *
+	 * @return	bool		true if the end of the data was reached (for both single shot and looping sounds)
+	 */
+	ENGINE_API virtual bool ReadCompressedData( uint8* Destination, bool bLooping, uint32 BufferSize );
+
+	ENGINE_API virtual void SeekToTime( const float SeekTime );
+
+	/** 
+	 * Decompress an entire ogg data file to a TArray
+	 */
+	ENGINE_API virtual void ExpandFile( uint8* DstBuffer, struct FSoundQualityInfo* QualityInfo );
+
+	/** 
+	 * Sets ogg to decode to half-rate
+	 * 
+	 * @param	Resource		Info about vorbis data
+	 */
+	ENGINE_API virtual void EnableHalfRate( bool HalfRate );
+
+	virtual uint32 GetSourceBufferSize() const { return SrcBufferDataSize;}
+
+	struct FVorbisFileWrapper* VFWrapper;
+	const uint8*		SrcBufferData;
+	uint32			SrcBufferDataSize;
+	uint32			BufferOffset;
+};
+#endif
diff --git a/Engine/Source/Runtime/HTML5/HTML5Audio/Private/HTML5AudioDevice.cpp b/Engine/Source/Runtime/HTML5/HTML5Audio/Private/HTML5AudioDevice.cpp
--- a/Engine/Source/Runtime/HTML5/HTML5Audio/Private/HTML5AudioDevice.cpp
+++ b/Engine/Source/Runtime/HTML5/HTML5Audio/Private/HTML5AudioDevice.cpp
@@ -11,16 +11,12 @@
 ------------------------------------------------------------------------------------*/
 
 #include "HTML5AudioDevice.h"
-#include "AudioDecompress.h"
+#include "VorbisAudioInfo.h"
 #include "AudioEffect.h"
 #include "Engine.h"
 
 DEFINE_LOG_CATEGORY(LogALAudio);
 
-//     2 UU == 1"
-// <=> 1 UU == 0.0127 m
-#define AUDIO_DISTANCE_FACTOR ( 0.0127f )
-
 ALCdevice* FALAudioDevice::HardwareDevice = NULL;
 ALCcontext*	FALAudioDevice::SoundContext = NULL;
 
diff --git a/Engine/Source/Runtime/HTML5/HTML5Audio/Public/HTML5AudioDevice.h b/Engine/Source/Runtime/HTML5/HTML5Audio/Public/HTML5AudioDevice.h
--- a/Engine/Source/Runtime/HTML5/HTML5Audio/Public/HTML5AudioDevice.h
+++ b/Engine/Source/Runtime/HTML5/HTML5Audio/Public/HTML5AudioDevice.h
@@ -29,6 +29,10 @@ class FALAudioDevice;
 
 DECLARE_LOG_CATEGORY_EXTERN(LogALAudio, Log, All);
 
+//     2 UU == 1"
+// <=> 1 UU == 0.0127 m
+#define AUDIO_DISTANCE_FACTOR ( 0.0127f )
+
 /**
  * OpenAL implementation of FSoundBuffer, containing the wave data and format information.
  */
diff --git a/Engine/Source/Runtime/Linux/ALAudio/Private/ALAudioDevice.cpp b/Engine/Source/Runtime/Linux/ALAudio/Private/ALAudioDevice.cpp
--- a/Engine/Source/Runtime/Linux/ALAudio/Private/ALAudioDevice.cpp
+++ b/Engine/Source/Runtime/Linux/ALAudio/Private/ALAudioDevice.cpp
@@ -10,16 +10,12 @@
 ------------------------------------------------------------------------------------*/
 
 #include "ALAudioDevice.h"
-#include "AudioDecompress.h"
+#include "VorbisAudioInfo.h"
 #include "AudioEffect.h"
 #include "Engine.h"
 
 DEFINE_LOG_CATEGORY(LogALAudio);
 
-//     2 UU == 1"
-// <=> 1 UU == 0.0127 m
-#define AUDIO_DISTANCE_FACTOR ( 0.0127f )
-
 ALCdevice* FALAudioDevice::HardwareDevice = NULL;
 ALCcontext*	FALAudioDevice::SoundContext = NULL;
 
diff --git a/Engine/Source/Runtime/Linux/ALAudio/Public/ALAudioDevice.h b/Engine/Source/Runtime/Linux/ALAudio/Public/ALAudioDevice.h
--- a/Engine/Source/Runtime/Linux/ALAudio/Public/ALAudioDevice.h
+++ b/Engine/Source/Runtime/Linux/ALAudio/Public/ALAudioDevice.h
@@ -28,6 +28,10 @@ class FALAudioDevice;
 
 DECLARE_LOG_CATEGORY_EXTERN(LogALAudio, Log, All);
 
+//     2 UU == 1"
+// <=> 1 UU == 0.0127 m
+#define AUDIO_DISTANCE_FACTOR ( 0.0127f )
+
 /**
  * OpenAL implementation of FSoundBuffer, containing the wave data and format information.
  */
diff --git a/Engine/Source/Runtime/Mac/CoreAudio/Private/CoreAudioDevice.cpp b/Engine/Source/Runtime/Mac/CoreAudio/Private/CoreAudioDevice.cpp
--- a/Engine/Source/Runtime/Mac/CoreAudio/Private/CoreAudioDevice.cpp
+++ b/Engine/Source/Runtime/Mac/CoreAudio/Private/CoreAudioDevice.cpp
@@ -9,7 +9,7 @@
  ------------------------------------------------------------------------------------*/
 
 #include "CoreAudioDevice.h"
-#include "AudioDecompress.h"
+#include "VorbisAudioInfo.h"
 #include "AudioEffect.h"
 #include "CoreAudioEffects.h"
 #include "Engine.h"
diff --git a/Engine/Source/Runtime/Windows/XAudio2/Private/XAudio2Device.cpp b/Engine/Source/Runtime/Windows/XAudio2/Private/XAudio2Device.cpp
--- a/Engine/Source/Runtime/Windows/XAudio2/Private/XAudio2Device.cpp
+++ b/Engine/Source/Runtime/Windows/XAudio2/Private/XAudio2Device.cpp
@@ -13,7 +13,8 @@
 
 #include "XAudio2Device.h"
 #include "AudioEffect.h"
-#include "AudioDecompress.h"
+#include "OpusAudioInfo.h"
+#include "VorbisAudioInfo.h"
 #include "XAudio2Effects.h"
 #include "Engine.h"
 #include "AllowWindowsPlatformTypes.h"
