Removed ability to duplicate entries in StreamingVolumes array, null entries can still be duplicated

diff --git a/Engine/Source/Runtime/Engine/Classes/Engine/LevelStreaming.h b/Engine/Source/Runtime/Engine/Classes/Engine/LevelStreaming.h
--- a/Engine/Source/Runtime/Engine/Classes/Engine/LevelStreaming.h
+++ b/Engine/Source/Runtime/Engine/Classes/Engine/LevelStreaming.h
@@ -130,7 +130,7 @@ class ULevelStreaming : public UObject
 	FColor DrawColor;
 
 	/** The level streaming volumes bound to this level.																		*/
-	UPROPERTY(EditAnywhere, Category=LevelStreaming, meta=(DisplayName = "Streaming Volumes"))
+	UPROPERTY(EditAnywhere, Category=LevelStreaming, meta=(DisplayName = "Streaming Volumes", NoElementDuplicate))
 	TArray<class ALevelStreamingVolume*> EditorStreamingVolumes;
 
 	/** Cooldown time in seconds between volume-based unload requests.  Used in preventing spurious unload requests.			*/
diff --git a/Engine/Source/Runtime/Engine/Private/LevelStreaming.cpp b/Engine/Source/Runtime/Engine/Private/LevelStreaming.cpp
--- a/Engine/Source/Runtime/Engine/Private/LevelStreaming.cpp
+++ b/Engine/Source/Runtime/Engine/Private/LevelStreaming.cpp
@@ -798,12 +798,14 @@ void ULevelStreaming::RemoveStreamingVolumeDuplicates()
 	for (int32 VolumeIdx = EditorStreamingVolumes.Num()-1; VolumeIdx >= 0; VolumeIdx--)
 	{
 		ALevelStreamingVolume* Volume = EditorStreamingVolumes[VolumeIdx];
-		
-		int32 DuplicateIdx = EditorStreamingVolumes.Find(Volume);
-		check(DuplicateIdx != INDEX_NONE);
-		if (DuplicateIdx != VolumeIdx)
+		if (Volume) // Allow duplicate null entries, for array editor convenience
 		{
-			EditorStreamingVolumes.RemoveAt(VolumeIdx);
+			int32 DuplicateIdx = EditorStreamingVolumes.Find(Volume);
+			check(DuplicateIdx != INDEX_NONE);
+			if (DuplicateIdx != VolumeIdx)
+			{
+				EditorStreamingVolumes.RemoveAt(VolumeIdx);
+			}
 		}
 	}
 }
