Changed Streaming Levels references in ALevelStreamingVolume to store package names instead of ULevelStreamingObjects

diff --git a/Engine/Source/Editor/Levels/Private/LevelCollectionViewModel.cpp b/Engine/Source/Editor/Levels/Private/LevelCollectionViewModel.cpp
--- a/Engine/Source/Editor/Levels/Private/LevelCollectionViewModel.cpp
+++ b/Engine/Source/Editor/Levels/Private/LevelCollectionViewModel.cpp
@@ -824,7 +824,7 @@ void FLevelCollectionViewModel::ClearStreamingLevelVolumes()
 					if ( LevelStreamingVolume )
 					{
 						LevelStreamingVolume->Modify();
-						LevelStreamingVolume->StreamingLevels.Remove( StreamingLevel );
+						LevelStreamingVolume->StreamingLevelNames.Remove( StreamingLevel->PackageName );
 					}
 				}
 
@@ -1795,7 +1795,7 @@ void FLevelCollectionViewModel::AddStreamingLevelVolumes()
 
 					// Associate the level to the volume.
 					LevelStreamingVolume->Modify();
-					LevelStreamingVolume->StreamingLevels.AddUnique( StreamingLevel );
+					LevelStreamingVolume->StreamingLevelNames.AddUnique( StreamingLevel->PackageName );
 
 					// Associate the volume to the level.
 					StreamingLevel->EditorStreamingVolumes.AddUnique( LevelStreamingVolume );
diff --git a/Engine/Source/Editor/UnrealEd/Private/EditorLevelUtils.cpp b/Engine/Source/Editor/UnrealEd/Private/EditorLevelUtils.cpp
--- a/Engine/Source/Editor/UnrealEd/Private/EditorLevelUtils.cpp
+++ b/Engine/Source/Editor/UnrealEd/Private/EditorLevelUtils.cpp
@@ -328,7 +328,7 @@ namespace EditorLevelUtils
 				if ( LevelStreamingVolume )
 				{
 					LevelStreamingVolume->Modify();
-					LevelStreamingVolume->StreamingLevels.Remove( InLevelStreaming );
+					LevelStreamingVolume->StreamingLevelNames.Remove( InLevelStreaming->PackageName );
 				}
 			}
 
diff --git a/Engine/Source/Runtime/Engine/Classes/Engine/LevelStreamingVolume.h b/Engine/Source/Runtime/Engine/Classes/Engine/LevelStreamingVolume.h
--- a/Engine/Source/Runtime/Engine/Classes/Engine/LevelStreamingVolume.h
+++ b/Engine/Source/Runtime/Engine/Classes/Engine/LevelStreamingVolume.h
@@ -27,9 +27,9 @@ class ALevelStreamingVolume : public AVolume
 {
 	GENERATED_UCLASS_BODY()
 
-	/** Levels affected by this level streaming volume. */
+	/** Levels names affected by this level streaming volume. */
 	UPROPERTY(Category=LevelStreamingVolume, VisibleAnywhere, BlueprintReadOnly)
-	TArray<class ULevelStreaming*> StreamingLevels;
+	TArray<FName> StreamingLevelNames;
 
 	/** If true, this streaming volume should only be used for editor streaming level previs. */
 	UPROPERTY(EditAnywhere, BlueprintReadWrite, Category=LevelStreamingVolume)
@@ -47,11 +47,17 @@ class ALevelStreamingVolume : public AVolume
 	UPROPERTY(EditAnywhere, BlueprintReadWrite, Category=LevelStreamingVolume)
 	TEnumAsByte<enum EStreamingVolumeUsage> StreamingUsage;
 
-
+	// Begin UObject interface.
+	void PostLoad() override;
+	// End UObject interface.
+	
 #if WITH_EDITOR
 	// Begin AActor interface.
 	virtual void CheckForErrors() override;
 	// End AActor interface.
+	
+	/** Updates list of streaming levels that are referenced by this streaming volume */
+	void UpdateStreamingLevelsRefs();
 #endif
 };
 
diff --git a/Engine/Source/Runtime/Engine/Private/LevelStreaming.cpp b/Engine/Source/Runtime/Engine/Private/LevelStreaming.cpp
--- a/Engine/Source/Runtime/Engine/Private/LevelStreaming.cpp
+++ b/Engine/Source/Runtime/Engine/Private/LevelStreaming.cpp
@@ -765,6 +765,12 @@ void ULevelStreaming::PostEditChangeChainProperty(FPropertyChangedChainEvent& Pr
 			if (PropertyName == GET_MEMBER_NAME_CHECKED(ULevelStreaming, EditorStreamingVolumes))
 			{
 				RemoveStreamingVolumeDuplicates();
+
+				// Update levels references in each streaming volume 
+				for (TActorIterator<ALevelStreamingVolume> It(GetWorld()); It; ++It)
+				{
+					(*It)->UpdateStreamingLevelsRefs();
+				}
 			}
 
 			else if( PropertyName == TEXT( "DrawColor" ) )
diff --git a/Engine/Source/Runtime/Engine/Private/LevelStreamingVolume.cpp b/Engine/Source/Runtime/Engine/Private/LevelStreamingVolume.cpp
--- a/Engine/Source/Runtime/Engine/Private/LevelStreamingVolume.cpp
+++ b/Engine/Source/Runtime/Engine/Private/LevelStreamingVolume.cpp
@@ -20,3 +20,33 @@ ALevelStreamingVolume::ALevelStreamingVolume(const class FPostConstructInitializ
 	StreamingUsage = SVB_LoadingAndVisibility;
 }
 
+void ALevelStreamingVolume::PostLoad()
+{
+	Super::PostLoad();
+
+	if (GIsEditor)
+	{
+		UpdateStreamingLevelsRefs();
+	}
+}
+
+#if WITH_EDITOR
+
+void ALevelStreamingVolume::UpdateStreamingLevelsRefs()
+{
+	StreamingLevelNames.Reset();
+	
+	UWorld* OwningWorld = GetWorld();
+	if (OwningWorld)
+	{
+		for (ULevelStreaming* LevelStreaming : OwningWorld->StreamingLevels)
+		{
+			if (LevelStreaming && LevelStreaming->EditorStreamingVolumes.Find(this) != INDEX_NONE)
+			{
+				StreamingLevelNames.Add(LevelStreaming->PackageName);
+			}
+		}
+	}
+}
+
+#endif// WITH_EDITOR
\ No newline at end of file
