fix for non-editor build

diff --git a/Engine/Source/Runtime/Engine/Private/LevelStreamingVolume.cpp b/Engine/Source/Runtime/Engine/Private/LevelStreamingVolume.cpp
--- a/Engine/Source/Runtime/Engine/Private/LevelStreamingVolume.cpp
+++ b/Engine/Source/Runtime/Engine/Private/LevelStreamingVolume.cpp
@@ -24,10 +24,12 @@ void ALevelStreamingVolume::PostLoad()
 {
 	Super::PostLoad();
 
+#if WITH_EDITOR
 	if (GIsEditor)
 	{
 		UpdateStreamingLevelsRefs();
 	}
+#endif//WITH_EDITOR
 }
 
 #if WITH_EDITOR
