Fixed crash on maps that have foliage and use level streaming transform in standalone game

diff --git a/Engine/Source/Runtime/Engine/Private/InstancedFoliage.cpp b/Engine/Source/Runtime/Engine/Private/InstancedFoliage.cpp
--- a/Engine/Source/Runtime/Engine/Private/InstancedFoliage.cpp
+++ b/Engine/Source/Runtime/Engine/Private/InstancedFoliage.cpp
@@ -1470,34 +1470,39 @@ void AInstancedFoliageActor::MapRebuild()
 
 void AInstancedFoliageActor::ApplyLevelTransform(const FTransform& LevelTransform)
 {
-	for (auto& MeshPair : FoliageMeshes)
+	// Apply transform to foliage editor only data
+	if (GIsEditor)
 	{
-		FFoliageMeshInfo& MeshInfo = *MeshPair.Value;
-
-		MeshInfo.InstanceHash->Empty();
-		for (int32 InstanceIdx = 0; InstanceIdx < MeshInfo.Instances.Num(); InstanceIdx++)
+		for (auto& MeshPair : FoliageMeshes)
 		{
-			FFoliageInstance& Instance = MeshInfo.Instances[InstanceIdx];
-			FTransform NewTransform = Instance.GetInstanceWorldTransform() * LevelTransform;
+			FFoliageMeshInfo& MeshInfo = *MeshPair.Value;
+		
+			MeshInfo.InstanceHash->Empty();
+			for (int32 InstanceIdx = 0; InstanceIdx < MeshInfo.Instances.Num(); InstanceIdx++)
+			{
+				FFoliageInstance& Instance = MeshInfo.Instances[InstanceIdx];
+				FTransform NewTransform = Instance.GetInstanceWorldTransform() * LevelTransform;
 			
-			Instance.Location		= NewTransform.GetLocation();
-			Instance.Rotation		= NewTransform.GetRotation().Rotator();
-			Instance.DrawScale3D	= NewTransform.GetScale3D();
-			// Rehash instance location
-			MeshInfo.InstanceHash->InsertInstance(Instance.Location, InstanceIdx);
-		}
+				Instance.Location		= NewTransform.GetLocation();
+				Instance.Rotation		= NewTransform.GetRotation().Rotator();
+				Instance.DrawScale3D	= NewTransform.GetScale3D();
+			
+				// Rehash instance location
+				MeshInfo.InstanceHash->InsertInstance(Instance.Location, InstanceIdx);
+			}
 
-		for (auto It = MeshInfo.ComponentHash.CreateIterator(); It; ++It)
-		{
-			// We assume here that component we painted foliage on, was transformed as well
-			FFoliageComponentHashInfo& Info = It.Value();
-			Info.UpdateLocationFromActor(It.Key());
-		}
+			for (auto It = MeshInfo.ComponentHash.CreateIterator(); It; ++It)
+			{
+				// We assume here that component we painted foliage on, was transformed as well
+				FFoliageComponentHashInfo& Info = It.Value();
+				Info.UpdateLocationFromActor(It.Key());
+			}
 
-		// Recalc cluster bounds
-		for (FFoliageInstanceCluster& Cluster : MeshInfo.InstanceClusters)
-		{
-			Cluster.Bounds = Cluster.ClusterComponent->CalcBounds(Cluster.ClusterComponent->ComponentToWorld);
+			// Recalc cluster bounds
+			for (FFoliageInstanceCluster& Cluster : MeshInfo.InstanceClusters)
+			{
+				Cluster.Bounds = Cluster.ClusterComponent->CalcBounds(Cluster.ClusterComponent->ComponentToWorld);
+			}
 		}
 	}
 }
