Check if the used canvas is null.

ttp334434 CRASH: EDITOR: Crash when trying to simulate a map with breakpoint set in level blueprint

diff --git a/Engine/Source/Runtime/Engine/Private/GameViewportClient.cpp b/Engine/Source/Runtime/Engine/Private/GameViewportClient.cpp
--- a/Engine/Source/Runtime/Engine/Private/GameViewportClient.cpp
+++ b/Engine/Source/Runtime/Engine/Private/GameViewportClient.cpp
@@ -613,17 +613,19 @@ void UGameViewportClient::Draw(FViewport* InViewport, FCanvas* SceneCanvas)
 	DebugCanvasObject->Canvas = DebugCanvas;	
 	DebugCanvasObject->Init(InViewport->GetSizeXY().X, InViewport->GetSizeXY().Y, NULL);
 
-	if (GEngine->HMDDevice.IsValid() && GEngine->IsStereoscopic3D())
+	const bool bScaledToRenderTarget = GEngine->HMDDevice.IsValid() && GEngine->IsStereoscopic3D();
+	if (bScaledToRenderTarget)
 	{
 		// Allow HMD to modify screen settings
 		GEngine->HMDDevice->UpdateScreenSettings(Viewport);
-		DebugCanvas->SetScaledToRenderTarget(true);
-		SceneCanvas->SetScaledToRenderTarget(true);
 	}
-	else
+	if (DebugCanvas)
+	{
+		DebugCanvas->SetScaledToRenderTarget(bScaledToRenderTarget);
+	}
+	if (SceneCanvas)
 	{
-		DebugCanvas->SetScaledToRenderTarget(false);
-		SceneCanvas->SetScaledToRenderTarget(false);
+		SceneCanvas->SetScaledToRenderTarget(bScaledToRenderTarget);
 	}
 
 	bool bUIDisableWorldRendering = false;
diff --git a/Engine/Source/Runtime/Engine/Private/UserInterface/CanvasItem.cpp b/Engine/Source/Runtime/Engine/Private/UserInterface/CanvasItem.cpp
--- a/Engine/Source/Runtime/Engine/Private/UserInterface/CanvasItem.cpp
+++ b/Engine/Source/Runtime/Engine/Private/UserInterface/CanvasItem.cpp
@@ -673,7 +673,7 @@ void FCanvasTextItem::Draw( class FCanvas* InCanvas )
 {	
 	SCOPE_CYCLE_COUNTER(STAT_Canvas_TextItemTime);
 
-	if(Font == NULL || Text.IsEmpty() )
+	if (InCanvas == NULL || Font == NULL || Text.IsEmpty())
 	{
 		return;
 	}
