Preview root motion in Persona

diff --git a/Engine/Source/Editor/Persona/Private/AnimViewportShowCommands.cpp b/Engine/Source/Editor/Persona/Private/AnimViewportShowCommands.cpp
--- a/Engine/Source/Editor/Persona/Private/AnimViewportShowCommands.cpp
+++ b/Engine/Source/Editor/Persona/Private/AnimViewportShowCommands.cpp
@@ -13,6 +13,7 @@ void FAnimViewportShowCommands::RegisterCommands()
 
 	UI_COMMAND( MuteAudio, "Mute Audio", "Mutes audio from the preview", EUserInterfaceActionType::ToggleButton, FInputGesture() );
 
+	UI_COMMAND(ProcessRootMotion, "Process Root Motion", "Moves preview based on animation root motion", EUserInterfaceActionType::ToggleButton, FInputGesture());
 
 	UI_COMMAND( ShowReferencePose, "Reference Pose", "Show reference pose on preview mesh", EUserInterfaceActionType::ToggleButton, FInputGesture() );
 	UI_COMMAND( ShowBound, "Bound", "Show bound on preview mesh", EUserInterfaceActionType::ToggleButton, FInputGesture() );
diff --git a/Engine/Source/Editor/Persona/Private/AnimViewportShowCommands.h b/Engine/Source/Editor/Persona/Private/AnimViewportShowCommands.h
--- a/Engine/Source/Editor/Persona/Private/AnimViewportShowCommands.h
+++ b/Engine/Source/Editor/Persona/Private/AnimViewportShowCommands.h
@@ -36,6 +36,9 @@ public:
 	/** Option to mute audio in the viewport */
 	TSharedPtr< FUICommandInfo > MuteAudio;
 
+	/** Option to show root motion in viewport */
+	TSharedPtr< FUICommandInfo > ProcessRootMotion;
+
 	/** Show reference pose on preview mesh */
 	TSharedPtr< FUICommandInfo > ShowReferencePose;
 	
diff --git a/Engine/Source/Editor/Persona/Private/AnimationEditorViewportClient.cpp b/Engine/Source/Editor/Persona/Private/AnimationEditorViewportClient.cpp
--- a/Engine/Source/Editor/Persona/Private/AnimationEditorViewportClient.cpp
+++ b/Engine/Source/Editor/Persona/Private/AnimationEditorViewportClient.cpp
@@ -592,6 +592,14 @@ void FAnimationViewportClient::Tick(float DeltaSeconds)
 	{
 		PreviewScene->GetWorld()->Tick(LEVELTICK_All, DeltaSeconds*AnimationPlaybackScale);
 	}
+
+	UDebugSkelMeshComponent* PreviewComp = PreviewSkelMeshComp.Get();
+	if (PreviewComp)
+	{
+		// Handle updating the preview component to represent the effects of root motion	
+		FBoxSphereBounds Bounds = EditorFloorComp->CalcBounds(EditorFloorComp->GetRelativeTransform());
+		PreviewComp->ConsumeRootMotion(Bounds.GetBox().Min, Bounds.GetBox().Max);
+	}
 }
 
 void FAnimationViewportClient::SetCameraTargetLocation(const FSphere &BoundSphere, float DeltaSeconds)
diff --git a/Engine/Source/Editor/Persona/Private/SAnimViewportToolBar.cpp b/Engine/Source/Editor/Persona/Private/SAnimViewportToolBar.cpp
--- a/Engine/Source/Editor/Persona/Private/SAnimViewportToolBar.cpp
+++ b/Engine/Source/Editor/Persona/Private/SAnimViewportToolBar.cpp
@@ -445,6 +445,12 @@ TSharedRef<SWidget> SAnimViewportToolBar::GenerateShowMenu() const
 			}
 			ShowMenuBuilder.EndSection();
 
+			ShowMenuBuilder.BeginSection("AnimViewportRootMotion", LOCTEXT("Viewport_RootMotionLabel", "Root Motion"));
+			{
+				ShowMenuBuilder.AddMenuEntry(Actions.ProcessRootMotion);
+			}
+			ShowMenuBuilder.EndSection();
+
 			ShowMenuBuilder.BeginSection("AnimViewportMesh", LOCTEXT("ShowMenu_Actions_Mesh", "Mesh"));
 			{
 				ShowMenuBuilder.AddMenuEntry( Actions.ShowReferencePose );
diff --git a/Engine/Source/Editor/Persona/Private/SAnimationEditorViewport.cpp b/Engine/Source/Editor/Persona/Private/SAnimationEditorViewport.cpp
--- a/Engine/Source/Editor/Persona/Private/SAnimationEditorViewport.cpp
+++ b/Engine/Source/Editor/Persona/Private/SAnimationEditorViewport.cpp
@@ -685,6 +685,12 @@ void SAnimationEditorViewportTabBody::BindCommands()
 		FExecuteAction::CreateSP(this, &SAnimationEditorViewportTabBody::OnMuteAudio),
 		FCanExecuteAction(),
 		FIsActionChecked::CreateSP(this, &SAnimationEditorViewportTabBody::IsAudioMuted));
+
+	CommandList.MapAction(
+		ViewportShowMenuCommands.ProcessRootMotion,
+		FExecuteAction::CreateSP(this, &SAnimationEditorViewportTabBody::OnTogglePreviewRootMotion),
+		FCanExecuteAction(),
+		FIsActionChecked::CreateSP(this, &SAnimationEditorViewportTabBody::IsPreviewingRootMotion));
 }
 
 bool SAnimationEditorViewportTabBody::IsPreviewModeOn(int32 PreviewMode) const
@@ -1355,6 +1361,27 @@ bool SAnimationEditorViewportTabBody::IsAudioMuted()
 	return GetAnimationViewportClient()->IsAudioMuted();
 }
 
+void SAnimationEditorViewportTabBody::OnTogglePreviewRootMotion()
+{
+	UDebugSkelMeshComponent* PreviewComponent = PersonaPtr.Pin()->PreviewComponent;
+
+	if (PreviewComponent)
+	{
+		PreviewComponent->bPreviewRootMotion = !PreviewComponent->bPreviewRootMotion;
+	}
+}
+
+bool SAnimationEditorViewportTabBody::IsPreviewingRootMotion() const
+{
+	UDebugSkelMeshComponent* PreviewComponent = PersonaPtr.Pin()->PreviewComponent;
+
+	if (PreviewComponent)
+	{
+		return PreviewComponent->bPreviewRootMotion;
+	}
+	return false;
+}
+
 #if WITH_APEX_CLOTHING
 bool SAnimationEditorViewportTabBody::IsDisablingClothSimulation() const
 {
diff --git a/Engine/Source/Editor/Persona/Private/SAnimationEditorViewport.h b/Engine/Source/Editor/Persona/Private/SAnimationEditorViewport.h
--- a/Engine/Source/Editor/Persona/Private/SAnimationEditorViewport.h
+++ b/Engine/Source/Editor/Persona/Private/SAnimationEditorViewport.h
@@ -324,6 +324,12 @@ private:
 	/** Whether audio from the viewport is muted */
 	bool IsAudioMuted();
 
+	/** Function to set whether we are previewing root motion */
+	void OnTogglePreviewRootMotion();
+	
+	/** Whether or not we are previewing root motion */
+	bool IsPreviewingRootMotion() const;
+
 #if WITH_APEX_CLOTHING
 	/** 
 	 * clothing show options 
diff --git a/Engine/Source/Editor/UnrealEd/Classes/Animation/DebugSkelMeshComponent.h b/Engine/Source/Editor/UnrealEd/Classes/Animation/DebugSkelMeshComponent.h
--- a/Engine/Source/Editor/UnrealEd/Classes/Animation/DebugSkelMeshComponent.h
+++ b/Engine/Source/Editor/UnrealEd/Classes/Animation/DebugSkelMeshComponent.h
@@ -98,6 +98,9 @@ class UDebugSkelMeshComponent : public USkeletalMeshComponent
 	UPROPERTY(transient)
 	uint32 bDisplayBound:1;
 
+	UPROPERTY(transient)
+	uint32 bPreviewRootMotion:1;
+
 	/** Compressed SpaceBases for when bDisplayRawAnimation == true, as raw space bases are stored in SpaceBases **/
 	TArray<FTransform> CompressedSpaceBases;
 
@@ -196,6 +199,11 @@ class UDebugSkelMeshComponent : public USkeletalMeshComponent
 	 */
 	UNREALED_API bool CheckIfBoundsAreCorrrect();
 
+	/** 
+	 * Update components position based on animation root motion
+	 */
+	UNREALED_API void ConsumeRootMotion(const FVector& FloorMin, const FVector& FloorMax);
+
 #if WITH_EDITOR
 	//TODO - This is a really poor way to post errors to the user. Work out a better way.
 	struct FAnimNotifyErrors
diff --git a/Engine/Source/Editor/UnrealEd/Private/Animation/DebugSkelMeshComponent.cpp b/Engine/Source/Editor/UnrealEd/Private/Animation/DebugSkelMeshComponent.cpp
--- a/Engine/Source/Editor/UnrealEd/Private/Animation/DebugSkelMeshComponent.cpp
+++ b/Engine/Source/Editor/UnrealEd/Private/Animation/DebugSkelMeshComponent.cpp
@@ -146,6 +146,49 @@ bool UDebugSkelMeshComponent::CheckIfBoundsAreCorrrect()
 	return false;
 }
 
+float WrapInRange(float StartVal, float MinVal, float MaxVal)
+{
+	float Size = MaxVal - MinVal;
+	float EndVal = StartVal;
+	while (EndVal < MinVal)
+	{
+		EndVal += Size;
+	}
+
+	while (EndVal > MaxVal)
+	{
+		EndVal -= Size;
+	}
+	return EndVal;
+}
+
+void UDebugSkelMeshComponent::ConsumeRootMotion(const FVector& FloorMin, const FVector& FloorMax)
+{
+	if (bPreviewRootMotion)
+	{
+		if (UAnimInstance* AnimInst = GetAnimInstance())
+		{
+			FRootMotionMovementParams ExtractedRootMotion = AnimInst->ConsumeExtractedRootMotion();
+			if (ExtractedRootMotion.bHasRootMotion)
+			{
+				AddLocalTransform(ExtractedRootMotion.RootMotionTransform);
+
+				//Handle moving component so that it stays within the editor floor
+				FTransform CurrentTransform = GetRelativeTransform();
+				FVector Trans = CurrentTransform.GetTranslation();
+				Trans.X = WrapInRange(Trans.X, FloorMin.X, FloorMax.X);
+				Trans.Y = WrapInRange(Trans.Y, FloorMin.Y, FloorMax.Y);
+				CurrentTransform.SetTranslation(Trans);
+				SetRelativeTransform(CurrentTransform);
+			}
+		}
+	}
+	else
+	{
+		SetWorldTransform(FTransform());
+	}
+}
+
 FPrimitiveSceneProxy* UDebugSkelMeshComponent::CreateSceneProxy()
 {
 	FDebugSkelMeshSceneProxy* Result = NULL;
diff --git a/Engine/Source/Runtime/Engine/Classes/Animation/AnimInstance.h b/Engine/Source/Runtime/Engine/Classes/Animation/AnimInstance.h
--- a/Engine/Source/Runtime/Engine/Classes/Animation/AnimInstance.h
+++ b/Engine/Source/Runtime/Engine/Classes/Animation/AnimInstance.h
@@ -596,11 +596,17 @@ public:
 	/** Get current RootMotion FAnimMontageInstance if any. NULL otherwise. */
 	FAnimMontageInstance * GetRootMotionMontageInstance() const;
 
+	/** Get current accumulated root motion, removing it from the AnimInstance in the process */
+	FRootMotionMovementParams ConsumeExtractedRootMotion();
+
 private:
 	/** Active Root Motion Montage Instance, if any. */
 	struct FAnimMontageInstance * RootMotionMontageInstance;
 
 	/** Update RootMotionMontageInstance */
 	void UpdateRootMotionMontageInstance();
+
+	// Root motion extracted from animation since the last time ConsumeExtractedRootMotion was called
+	FRootMotionMovementParams ExtractedRootMotion;
 };
 
diff --git a/Engine/Source/Runtime/Engine/Classes/Animation/AnimationAsset.h b/Engine/Source/Runtime/Engine/Classes/Animation/AnimationAsset.h
--- a/Engine/Source/Runtime/Engine/Classes/Animation/AnimationAsset.h
+++ b/Engine/Source/Runtime/Engine/Classes/Animation/AnimationAsset.h
@@ -232,6 +232,56 @@ public:
 	}
 };
 
+/** Utility struct to accumulate root motion. */
+USTRUCT()
+struct FRootMotionMovementParams
+{
+	GENERATED_USTRUCT_BODY()
+
+		UPROPERTY()
+		bool bHasRootMotion;
+
+	UPROPERTY()
+		FTransform RootMotionTransform;
+
+	FRootMotionMovementParams()
+		: bHasRootMotion(false)
+		, RootMotionTransform(FTransform::Identity)
+	{
+	}
+
+	void Set(const FTransform & InTransform)
+	{
+		bHasRootMotion = true;
+		RootMotionTransform = InTransform;
+	}
+
+	void Accumulate(const FTransform & InTransform)
+	{
+		if (!bHasRootMotion)
+		{
+			Set(InTransform);
+		}
+		else
+		{
+			RootMotionTransform = InTransform * RootMotionTransform;
+		}
+	}
+
+	void Accumulate(const FRootMotionMovementParams & MovementParams)
+	{
+		if (MovementParams.bHasRootMotion)
+		{
+			Accumulate(MovementParams.RootMotionTransform);
+		}
+	}
+
+	void Clear()
+	{
+		bHasRootMotion = false;
+	}
+};
+
 // This structure is used to either advance or synchronize animation players
 struct FAnimAssetTickContext
 {
diff --git a/Engine/Source/Runtime/Engine/Classes/GameFramework/CharacterMovementComponent.h b/Engine/Source/Runtime/Engine/Classes/GameFramework/CharacterMovementComponent.h
--- a/Engine/Source/Runtime/Engine/Classes/GameFramework/CharacterMovementComponent.h
+++ b/Engine/Source/Runtime/Engine/Classes/GameFramework/CharacterMovementComponent.h
@@ -12,6 +12,7 @@
 #include "Interfaces/NetworkPredictionInterface.h"
 #include "Engine/EngineTypes.h"
 #include "AI/Navigation/NavigationAvoidanceTypes.h"
+#include "Animation/AnimationAsset.h"
 #include "CharacterMovementComponent.generated.h"
 
 struct FVector_NetQuantize100;
@@ -93,57 +94,6 @@ public:
 	void SetFromLineTrace(const FHitResult& InHit, const float InSweepFloorDist, const float InLineDist, const bool bIsWalkableFloor);
 };
 
-
-/** Utility struct to accumulate root motion. */
-USTRUCT()
-struct FRootMotionMovementParams
-{
-	GENERATED_USTRUCT_BODY()
-
-	UPROPERTY()
-	bool bHasRootMotion;
-
-	UPROPERTY()
-	FTransform RootMotionTransform;
-
-	FRootMotionMovementParams()
-		: bHasRootMotion(false)
-		, RootMotionTransform(FTransform::Identity)
-	{
-	}
-
-	void Set(const FTransform & InTransform)
-	{
-		bHasRootMotion = true;
-		RootMotionTransform = InTransform;
-	}
-
-	void Accumulate(const FTransform & InTransform)
-	{
-		if( !bHasRootMotion )
-		{
-			Set(InTransform);
-		}
-		else
-		{
-			RootMotionTransform = InTransform * RootMotionTransform;
-		}
-	}
-
-	void Accumulate(const FRootMotionMovementParams & MovementParams)
-	{
-		if( MovementParams.bHasRootMotion )
-		{
-			Accumulate(MovementParams.RootMotionTransform);
-		}
-	}
-
-	void Clear()
-	{
-		bHasRootMotion = false;
-	}
-};
-
 /** 
  * Tick function that calls UCharacterMovementComponent::PreClothTick
  **/
@@ -1368,6 +1318,9 @@ protected:
 	/** Unpack compressed flags from a saved move and set state accordingly. See FSavedMove_Character. */
 	virtual void UpdateFromCompressedFlags(uint8 Flags);
 
+	/** Ticks the characters pose and accumulates root motion */
+	void TickCharacterPose(float DeltaTime);
+
 public:
 
 	/** React to instantaneous change in position. Invalidates cached floor recomputes it if possible if there is a current movement base. */
diff --git a/Engine/Source/Runtime/Engine/Private/Animation/AnimInstance.cpp b/Engine/Source/Runtime/Engine/Private/Animation/AnimInstance.cpp
--- a/Engine/Source/Runtime/Engine/Private/Animation/AnimInstance.cpp
+++ b/Engine/Source/Runtime/Engine/Private/Animation/AnimInstance.cpp
@@ -1534,7 +1534,6 @@ void UAnimInstance::Montage_UpdateWeight(float DeltaSeconds)
 void UAnimInstance::Montage_Advance(float DeltaSeconds)
 {
 	bool bUpdateRootMotionMontageInstance = false;
-	FRootMotionMovementParams ExtractedRootMotion;
 
 	// go through all montage instances, and update them
 	// and make sure their weight is updated properly
@@ -1570,16 +1569,6 @@ void UAnimInstance::Montage_Advance(float DeltaSeconds)
 	{
 		UpdateRootMotionMontageInstance();
 	}
-
-	// If Root Motion has been extracted, forward it to character physics.
-	if( ExtractedRootMotion.bHasRootMotion )
-	{
-		ACharacter * CharacterOwner = Cast<ACharacter>(GetOwningActor());
-		if( CharacterOwner && CharacterOwner->CharacterMovement )
-		{
-			CharacterOwner->CharacterMovement->RootMotionParams.Accumulate(ExtractedRootMotion);
-		}
-	}
 }
 
 float UAnimInstance::PlaySlotAnimation(UAnimSequenceBase* Asset, FName SlotNodeName, float BlendInTime, float BlendOutTime, float InPlayRate)
@@ -1804,6 +1793,13 @@ FAnimMontageInstance * UAnimInstance::GetRootMotionMontageInstance() const
 	return RootMotionMontageInstance;
 }
 
+FRootMotionMovementParams UAnimInstance::ConsumeExtractedRootMotion()
+{
+	FRootMotionMovementParams RootMotion = ExtractedRootMotion;
+	ExtractedRootMotion.Clear();
+	return RootMotion;
+}
+
 void UAnimInstance::Montage_Stop(float InBlendOutTime)
 {
 	FAnimMontageInstance * CurMontageInstance = GetActiveMontageInstance();
diff --git a/Engine/Source/Runtime/Engine/Private/Animation/AnimMontage.cpp b/Engine/Source/Runtime/Engine/Private/Animation/AnimMontage.cpp
--- a/Engine/Source/Runtime/Engine/Private/Animation/AnimMontage.cpp
+++ b/Engine/Source/Runtime/Engine/Private/Animation/AnimMontage.cpp
@@ -8,7 +8,6 @@
 #include "Animation/AnimMontage.h"
 #include "AnimationUtils.h"
 #include "AnimationRuntime.h"
-#include "GameFramework/CharacterMovementComponent.h"
 
 ///////////////////////////////////////////////////////////////////////////
 //
diff --git a/Engine/Source/Runtime/Engine/Private/CharacterMovementComponent.cpp b/Engine/Source/Runtime/Engine/Private/CharacterMovementComponent.cpp
--- a/Engine/Source/Runtime/Engine/Private/CharacterMovementComponent.cpp
+++ b/Engine/Source/Runtime/Engine/Private/CharacterMovementComponent.cpp
@@ -896,7 +896,7 @@ void UCharacterMovementComponent::SimulatedTick(float DeltaSeconds)
 		// Tick animations before physics.
 		if( CharacterOwner->Mesh )
 		{
-			CharacterOwner->Mesh->TickPose(DeltaSeconds);
+			TickCharacterPose(DeltaSeconds);
 
 			// Make sure animation didn't trigger an event that destroyed us
 			if (!HasValidData())
@@ -1338,7 +1338,7 @@ void UCharacterMovementComponent::PerformMovement(float DeltaSeconds)
 		// If using RootMotion, tick animations before running physics.
 		if( !CharacterOwner->bClientUpdating && CharacterOwner->IsPlayingRootMotion() && CharacterOwner->Mesh )
 		{
-			CharacterOwner->Mesh->TickPose(DeltaSeconds);
+			TickCharacterPose(DeltaSeconds);
 
 			// Make sure animation didn't trigger an event that destroyed us
 			if (!HasValidData())
@@ -5424,7 +5424,7 @@ void UCharacterMovementComponent::MoveAutonomous
 	// If not playing root motion, tick animations after physics. We do this here to keep events, notifies, states and transitions in sync with client updates.
 	if( !CharacterOwner->bClientUpdating && !CharacterOwner->IsPlayingRootMotion() && CharacterOwner->Mesh )
 	{
-		CharacterOwner->Mesh->TickPose(DeltaTime);
+		TickCharacterPose(DeltaTime);
 	}
 }
 
@@ -5963,6 +5963,19 @@ void UCharacterMovementComponent::RegisterComponentTickFunctions(bool bRegister)
 	}
 }
 
+void UCharacterMovementComponent::TickCharacterPose(float DeltaTime)
+{
+	check(CharacterOwner && CharacterOwner->Mesh)
+
+	CharacterOwner->Mesh->TickPose(DeltaTime);
+
+	// Grab root motion now that we have ticked the pose
+	if (UAnimInstance* AnimInstance = CharacterOwner->Mesh->GetAnimInstance())
+	{
+		RootMotionParams.Accumulate(AnimInstance->ConsumeExtractedRootMotion());
+	}
+}
+
 FNetworkPredictionData_Client_Character::FNetworkPredictionData_Client_Character()
 	: ClientUpdateTime(0.f)
 	, CurrentTimeStamp(0.f)
diff --git a/Engine/Source/Runtime/Engine/Private/World.cpp b/Engine/Source/Runtime/Engine/Private/World.cpp
--- a/Engine/Source/Runtime/Engine/Private/World.cpp
+++ b/Engine/Source/Runtime/Engine/Private/World.cpp
@@ -30,7 +30,6 @@
 #include "VisualLog.h"
 #include "LevelUtils.h"
 #include "PhysicsPublic.h"
-#include "GameFramework/CharacterMovementComponent.h"
 #include "AI/AISystemBase.h"
 #if WITH_EDITOR
 	#include "Editor/UnrealEd/Public/Kismet2/KismetEditorUtilities.h"
