Enable BlackBody material node for ES2 (will probably be slow)

#ttp 335883 - Blackbody node not supported on ES2?

diff --git a/Engine/Source/Runtime/Engine/Private/Materials/HLSLMaterialTranslator.h b/Engine/Source/Runtime/Engine/Private/Materials/HLSLMaterialTranslator.h
--- a/Engine/Source/Runtime/Engine/Private/Materials/HLSLMaterialTranslator.h
+++ b/Engine/Source/Runtime/Engine/Private/Materials/HLSLMaterialTranslator.h
@@ -3277,11 +3277,6 @@ protected:
 
 	virtual int32 BlackBody( int32 Temp ) override
 	{
-		if (ErrorUnlessFeatureLevelSupported(ERHIFeatureLevel::SM3) == INDEX_NONE)
-		{
-			return INDEX_NONE;
-		}
-
 		if( Temp == INDEX_NONE )
 		{
 			return INDEX_NONE;
