Change DOF constraint to be on heap. Saves all PrimitiveComponents from having FConstraintInstance

diff --git a/Engine/Source/Runtime/Engine/Classes/PhysicsEngine/BodyInstance.h b/Engine/Source/Runtime/Engine/Classes/PhysicsEngine/BodyInstance.h
--- a/Engine/Source/Runtime/Engine/Classes/PhysicsEngine/BodyInstance.h
+++ b/Engine/Source/Runtime/Engine/Classes/PhysicsEngine/BodyInstance.h
@@ -4,13 +4,13 @@
 
 #include "PhysxUserData.h"
 #include "CollisionQueryParams.h"
-#include "ConstraintInstance.h"
 #include "BodyInstance.generated.h"
 
 
 #define UE_WITH_PHYSICS (WITH_PHYSX || WITH_BOX2D)
 
 struct FCollisionShape;
+class UPhysicsConstraintComponent;
 
 #if WITH_PHYSX
 namespace physx
@@ -197,7 +197,9 @@ public:
 	FVector GetLockedAxis() const;
 	void CreateDOFLock();
 
-	FConstraintInstance DOFConstraint;
+	/** Constraint used to allow for easy DOF setup per bodyinstance */
+	UPROPERTY()
+	UPhysicsConstraintComponent * DOFConstraint;
 
 protected:
 
diff --git a/Engine/Source/Runtime/Engine/Private/PhysicsEngine/BodyInstance.cpp b/Engine/Source/Runtime/Engine/Private/PhysicsEngine/BodyInstance.cpp
--- a/Engine/Source/Runtime/Engine/Private/PhysicsEngine/BodyInstance.cpp
+++ b/Engine/Source/Runtime/Engine/Private/PhysicsEngine/BodyInstance.cpp
@@ -532,24 +532,29 @@ FVector FBodyInstance::GetLockedAxis() const
 
 void FBodyInstance::CreateDOFLock()
 {
-	//TODO: This hasn't been tested with calling the function multiple times. It's really only expected to be used during InitBody
-	if (IsDynamic() == false) return;
-#if WITH_PHYSX
+
+	DOFConstraint = NULL;
 	//setup constraint based on DOF
 	FVector LockedAxis = GetLockedAxis();
-	if (LockedAxis.IsNearlyZero() == false)
+
+	if (IsDynamic() == false || LockedAxis.IsNearlyZero())
+	{
+		return;
+	}
+
+	DOFConstraint = ConstructObject<UPhysicsConstraintComponent>(UPhysicsConstraintComponent::StaticClass());
 	{
-		DOFConstraint.bSwingLimitSoft = false;
-		DOFConstraint.bTwistLimitSoft = false;
-		DOFConstraint.bLinearLimitSoft = false;
+		DOFConstraint->ConstraintInstance.bSwingLimitSoft = false;
+		DOFConstraint->ConstraintInstance.bTwistLimitSoft = false;
+		DOFConstraint->ConstraintInstance.bLinearLimitSoft = false;
 		//set all rotation to free
-		DOFConstraint.AngularSwing1Motion = EAngularConstraintMotion::ACM_Locked;
-		DOFConstraint.AngularSwing2Motion = EAngularConstraintMotion::ACM_Locked;
-		DOFConstraint.AngularTwistMotion = EAngularConstraintMotion::ACM_Free;
+		DOFConstraint->ConstraintInstance.AngularSwing1Motion = EAngularConstraintMotion::ACM_Locked;
+		DOFConstraint->ConstraintInstance.AngularSwing2Motion = EAngularConstraintMotion::ACM_Locked;
+		DOFConstraint->ConstraintInstance.AngularTwistMotion = EAngularConstraintMotion::ACM_Free;
 
-		DOFConstraint.LinearXMotion = ELinearConstraintMotion::LCM_Locked;
-		DOFConstraint.LinearYMotion = ELinearConstraintMotion::LCM_Free;
-		DOFConstraint.LinearZMotion = ELinearConstraintMotion::LCM_Free;
+		DOFConstraint->ConstraintInstance.LinearXMotion = ELinearConstraintMotion::LCM_Locked;
+		DOFConstraint->ConstraintInstance.LinearYMotion = ELinearConstraintMotion::LCM_Free;
+		DOFConstraint->ConstraintInstance.LinearZMotion = ELinearConstraintMotion::LCM_Free;
 
 		FVector Normal = LockedAxis.SafeNormal();
 		FVector Sec;
@@ -559,17 +564,16 @@ void FBodyInstance::CreateDOFLock()
 
 		FTransform TM = GetUnrealWorldTransform();
 
-		DOFConstraint.PriAxis1 = TM.InverseTransformVectorNoScale(Normal);
-		DOFConstraint.SecAxis1 = TM.InverseTransformVectorNoScale(Sec);
+		DOFConstraint->ConstraintInstance.PriAxis1 = TM.InverseTransformVectorNoScale(Normal);
+		DOFConstraint->ConstraintInstance.SecAxis1 = TM.InverseTransformVectorNoScale(Sec);
 
-		DOFConstraint.PriAxis2 = Normal;
-		DOFConstraint.SecAxis2 = Sec;
-		DOFConstraint.Pos2 = TM.GetLocation();
+		DOFConstraint->ConstraintInstance.PriAxis2 = Normal;
+		DOFConstraint->ConstraintInstance.SecAxis2 = Sec;
+		DOFConstraint->ConstraintInstance.Pos2 = TM.GetLocation();
 
 		// Create constraint instance based on DOF
-		DOFConstraint.InitConstraint(OwnerComponent.Get(), this, nullptr, 1.f);
+		DOFConstraint->ConstraintInstance.InitConstraint(OwnerComponent.Get(), this, nullptr, 1.f);
 	}
-#endif
 }
 
 ECollisionEnabled::Type FBodyInstance::GetCollisionEnabled() const
@@ -1402,6 +1406,9 @@ void FBodyInstance::TermBody()
 
 	BodySetup = NULL;
 	OwnerComponent = NULL;
+	DOFConstraint = NULL;
+
+	
 }
 
 #if WITH_BODY_WELDING
diff --git a/Engine/Source/Runtime/Engine/Private/SkeletalMeshComponentPhysics.cpp b/Engine/Source/Runtime/Engine/Private/SkeletalMeshComponentPhysics.cpp
--- a/Engine/Source/Runtime/Engine/Private/SkeletalMeshComponentPhysics.cpp
+++ b/Engine/Source/Runtime/Engine/Private/SkeletalMeshComponentPhysics.cpp
@@ -829,6 +829,10 @@ void USkeletalMeshComponent::InitArticulated(FPhysScene* PhysScene)
 				BodyInst->LockedAxisMode = BodyInstance.LockedAxisMode;
 				BodyInst->CustomLockedAxis = BodyInstance.CustomLockedAxis;
 			}
+			else
+			{
+				BodyInst->LockedAxisMode = ELockedAxis::None;
+			}
 
 #if WITH_PHYSX
 			// Create physics body instance.
