Fix post edit change crash in InstancedStaticMesh due to virtual parent not being called.

diff --git a/Engine/Source/Runtime/Engine/Private/InstancedStaticMesh.cpp b/Engine/Source/Runtime/Engine/Private/InstancedStaticMesh.cpp
--- a/Engine/Source/Runtime/Engine/Private/InstancedStaticMesh.cpp
+++ b/Engine/Source/Runtime/Engine/Private/InstancedStaticMesh.cpp
@@ -1556,6 +1556,8 @@ void UInstancedStaticMeshComponent::PostEditChangeChainProperty(FPropertyChanged
 
 		MarkRenderStateDirty();
 	}
+
+	Super::PostEditChangeChainProperty(PropertyChangedEvent);
 }
 #endif
 
