Make DOFConstraint transient

diff --git a/Engine/Source/Runtime/Engine/Classes/PhysicsEngine/BodyInstance.h b/Engine/Source/Runtime/Engine/Classes/PhysicsEngine/BodyInstance.h
--- a/Engine/Source/Runtime/Engine/Classes/PhysicsEngine/BodyInstance.h
+++ b/Engine/Source/Runtime/Engine/Classes/PhysicsEngine/BodyInstance.h
@@ -198,7 +198,7 @@ public:
 	void CreateDOFLock();
 
 	/** Constraint used to allow for easy DOF setup per bodyinstance */
-	UPROPERTY()
+	UPROPERTY(transient)
 	UPhysicsConstraintComponent * DOFConstraint;
 
 protected:
