RHI texture references are now managed with an FRenderResource.

Last render time for textures can now be exposed via an FLastRenderTimeContainer provided at construction time. Only implemented for texture references. Eventually this can unify last render time for normal textures and texture references.

diff --git a/Engine/Source/Runtime/Engine/Classes/Engine/Texture.h b/Engine/Source/Runtime/Engine/Classes/Engine/Texture.h
--- a/Engine/Source/Runtime/Engine/Classes/Engine/Texture.h
+++ b/Engine/Source/Runtime/Engine/Classes/Engine/Texture.h
@@ -5,6 +5,7 @@
 #include "Engine/TextureDefines.h"
 #include "TextureResource.h"
 #include "MaterialShared.h"
+#include "RenderResource.h"
 #include "Texture.generated.h"
 
 // This needs to be mirrored in EditorFactories.cpp.
@@ -507,7 +508,7 @@ public:
 	class FTextureResource*	Resource;
 
 	/** Stable RHI texture reference that refers to the current RHI texture. Note this is manually refcounted! */
-	class FRHITextureReference* TextureRefRHI;
+	FTextureReference TextureReference;
 
 	/** Release fence to know when resources have been freed on the rendering thread. */
 	FRenderCommandFence ReleaseFence;
diff --git a/Engine/Source/Runtime/Engine/Private/Materials/MaterialUniformExpressions.cpp b/Engine/Source/Runtime/Engine/Private/Materials/MaterialUniformExpressions.cpp
--- a/Engine/Source/Runtime/Engine/Private/Materials/MaterialUniformExpressions.cpp
+++ b/Engine/Source/Runtime/Engine/Private/Materials/MaterialUniformExpressions.cpp
@@ -294,7 +294,7 @@ FUniformBufferRHIRef FUniformExpressionSet::CreateUniformBuffer(const FMaterialR
 			Uniform2DTextureExpressions[ExpressionIndex]->GetTextureValue(MaterialRenderContext,MaterialRenderContext.Material,Value);
 			if(Value && Value->Resource)
 			{
-				*ResourceTable++ = Value->TextureRefRHI;
+				*ResourceTable++ = Value->TextureReference.TextureReferenceRHI;
 				*ResourceTable++ = Value->Resource->SamplerStateRHI;
 			}
 			else
@@ -311,7 +311,7 @@ FUniformBufferRHIRef FUniformExpressionSet::CreateUniformBuffer(const FMaterialR
 			UniformCubeTextureExpressions[ExpressionIndex]->GetTextureValue(MaterialRenderContext,MaterialRenderContext.Material,Value);
 			if(Value && Value->Resource)
 			{
-				*ResourceTable++ = Value->TextureRefRHI;
+				*ResourceTable++ = Value->TextureReference.TextureReferenceRHI;
 				*ResourceTable++ = Value->Resource->SamplerStateRHI;
 			}
 			else
diff --git a/Engine/Source/Runtime/Engine/Private/Texture.cpp b/Engine/Source/Runtime/Engine/Private/Texture.cpp
--- a/Engine/Source/Runtime/Engine/Private/Texture.cpp
+++ b/Engine/Source/Runtime/Engine/Private/Texture.cpp
@@ -57,14 +57,7 @@ UTexture::UTexture(const class FPostConstructInitializeProperties& PCIP)
 
 	if (FApp::CanEverRender() && !IsTemplate())
 	{
-		ENQUEUE_UNIQUE_RENDER_COMMAND_ONEPARAMETER(
-			FCreateTextureRefCommand,
-			UTexture*,Texture,this,
-		{
-			FTextureReferenceRHIRef TextureRef = RHICmdList.CreateTextureReference();
-			TextureRef->AddRef();
-			Texture->TextureRefRHI = TextureRef;
-		});
+		TextureReference.BeginInit_GameThread();
 	}
 }
 
@@ -318,22 +311,16 @@ void UTexture::PostLoad()
 void UTexture::BeginDestroy()
 {
 	Super::BeginDestroy();
-	if( !UpdateStreamingStatus() && (Resource || TextureRefRHI) )
+	if( !UpdateStreamingStatus() && (Resource || TextureReference.IsInitialized_GameThread()) )
 	{
 		// Send the rendering thread a release message for the texture's resource.
 		if (Resource)
 		{
 			BeginReleaseResource(Resource);
 		}
-		if (TextureRefRHI)
+		if (TextureReference.IsInitialized_GameThread())
 		{
-			ENQUEUE_UNIQUE_RENDER_COMMAND_ONEPARAMETER(
-				FDestroyTextureRefCommand,
-				UTexture*,Texture,this,
-			{
-				Texture->TextureRefRHI->Release();
-				Texture->TextureRefRHI = NULL;
-			});
+			TextureReference.BeginRelease_GameThread();
 		}
 		ReleaseFence.BeginFence();
 		// Keep track that we already kicked off the async release.
@@ -348,29 +335,24 @@ bool UTexture::IsReadyForFinishDestroy()
 	if( Super::IsReadyForFinishDestroy() && !UpdateStreamingStatus() )
 	{
 		// Kick off async resource release if we haven't already.
-		if( !bAsyncResourceReleaseHasBeenStarted && (Resource || TextureRefRHI) )
+		if( !bAsyncResourceReleaseHasBeenStarted && (Resource || TextureReference.IsInitialized_GameThread()) )
 		{
 			// Send the rendering thread a release message for the texture's resource.
 			if (Resource)
 			{
 				BeginReleaseResource(Resource);
 			}
-			if (TextureRefRHI)
+			if (TextureReference.IsInitialized_GameThread())
 			{
-				ENQUEUE_UNIQUE_RENDER_COMMAND_ONEPARAMETER(
-					FDestroyTextureRefCommand,
-					UTexture*,Texture,this,
-				{
-					Texture->TextureRefRHI->Release();
-					Texture->TextureRefRHI = NULL;
-				});
+				TextureReference.BeginRelease_GameThread();
 			}
 			ReleaseFence.BeginFence();
 			// Keep track that we already kicked off the async release.
 			bAsyncResourceReleaseHasBeenStarted = true;
 		}
+
 		// Only allow FinishDestroy to be called once the texture resource has finished its rendering thread cleanup.
-		else if( (!Resource && !TextureRefRHI) || ReleaseFence.IsFenceComplete() )
+		if( !bAsyncResourceReleaseHasBeenStarted || ReleaseFence.IsFenceComplete() )
 		{
 			bReadyForFinishDestroy = true;
 		}
@@ -382,10 +364,11 @@ void UTexture::FinishDestroy()
 {
 	Super::FinishDestroy();
 
+	check(!bAsyncResourceReleaseHasBeenStarted || ReleaseFence.IsFenceComplete());
+	check(TextureReference.IsInitialized_GameThread() == false);
+
 	if(Resource)
 	{
-		check(ReleaseFence.IsFenceComplete());
-
 		// Free the resource.
 		delete Resource;
 		Resource = NULL;
@@ -393,8 +376,6 @@ void UTexture::FinishDestroy()
 
 	CleanupCachedRunningPlatformData();
 	CleanupCachedCookedPlatformData();
-
-	check(TextureRefRHI == NULL);
 }
 
 void UTexture::PreSave()
diff --git a/Engine/Source/Runtime/Engine/Private/Texture2D.cpp b/Engine/Source/Runtime/Engine/Private/Texture2D.cpp
--- a/Engine/Source/Runtime/Engine/Private/Texture2D.cpp
+++ b/Engine/Source/Runtime/Engine/Private/Texture2D.cpp
@@ -318,11 +318,9 @@ float UTexture2D::GetLastRenderTimeForStreaming()
 	float LastRenderTime = -FLT_MAX;
 	if (Resource)
 	{
-		LastRenderTime = Resource->LastRenderTime;
-		if (TextureRefRHI)
-		{
-			LastRenderTime = FMath::Max(LastRenderTime,TextureRefRHI->GetLastCachedTime());
-		}
+		// The last render time is the last time the resource was directly bound or the last
+		// time the texture reference was cached in a resource table, whichever was later.
+		LastRenderTime = FMath::Max<double>(Resource->LastRenderTime,TextureReference.GetLastRenderTime());
 	}
 	return LastRenderTime;
 }
@@ -333,10 +331,7 @@ void UTexture2D::InvalidateLastRenderTimeForStreaming()
 	{
 		Resource->LastRenderTime = -FLT_MAX;
 	}
-	if (TextureRefRHI)
-	{
-		TextureRefRHI->SetLastCachedTime(-FLT_MAX);
-	}
+	TextureReference.InvalidateLastRenderTime();
 }
 
 #if WITH_EDITOR
@@ -1338,7 +1333,7 @@ void FTexture2DResource::InitRHI()
 
 				TextureRHI = Texture2DRHI;
 				RHIBindDebugLabelName(TextureRHI, *Owner->GetName());
-				RHIUpdateTextureReference(Owner->TextureRefRHI,TextureRHI);
+				RHIUpdateTextureReference(Owner->TextureReference.TextureReferenceRHI,TextureRHI);
 
 				return;
 			}
@@ -1348,7 +1343,7 @@ void FTexture2DResource::InitRHI()
 			Texture2DRHI	= RHICreateTexture2D( SizeX, SizeY, EffectiveFormat, Owner->RequestedMips, 1, TexCreateFlags, CreateInfo);
 			TextureRHI		= Texture2DRHI;
 			RHIBindDebugLabelName(TextureRHI, *Owner->GetName());
-			RHIUpdateTextureReference(Owner->TextureRefRHI,TextureRHI);
+			RHIUpdateTextureReference(Owner->TextureReference.TextureReferenceRHI,TextureRHI);
 
 			check(Owner->RequestedMips == Texture2DRHI->GetNumMips());
 			check(Owner->PlatformData->Mips.Num() == CurrentFirstMip + Owner->RequestedMips);
@@ -1398,7 +1393,7 @@ void FTexture2DResource::InitRHI()
 			Texture2DRHI	= RHICreateTexture2D( SizeX, SizeY, EffectiveFormat, Owner->RequestedMips, 1, TexCreateFlags, CreateInfo );
 			TextureRHI		= Texture2DRHI;
 			RHIBindDebugLabelName(TextureRHI, *Owner->GetName());
-			RHIUpdateTextureReference(Owner->TextureRefRHI,TextureRHI);
+			RHIUpdateTextureReference(Owner->TextureReference.TextureReferenceRHI,TextureRHI);
 			for( int32 MipIndex=CurrentFirstMip; MipIndex<OwnerMips.Num(); MipIndex++ )
 			{
 				if( MipData[MipIndex] != NULL )
@@ -1439,7 +1434,7 @@ void FTexture2DResource::ReleaseRHI()
 	check( Owner->PendingMipChangeRequestStatus.GetValue() == TexState_ReadyFor_Requests );	
 	FTextureResource::ReleaseRHI();
 	Texture2DRHI.SafeRelease();
-	RHIUpdateTextureReference(Owner->TextureRefRHI,FTextureRHIParamRef());
+	RHIUpdateTextureReference(Owner->TextureReference.TextureReferenceRHI,FTextureRHIParamRef());
 }
 
 void FTexture2DResource::CreateSamplerStates(float MipMapBias)
@@ -1709,7 +1704,7 @@ void FTexture2DResource::UpdateMipCount()
 			}
 
 			TextureRHI = Texture2DRHI;
-			RHIUpdateTextureReference(Owner->TextureRefRHI,TextureRHI);
+			RHIUpdateTextureReference(Owner->TextureReference.TextureReferenceRHI,TextureRHI);
 		}
 	}
 
@@ -2215,7 +2210,7 @@ void FTexture2DResource::FinalizeMipCount()
 			TextureRHI		= IntermediateTextureRHI;
 			Texture2DRHI	= IntermediateTextureRHI;
 			CurrentFirstMip = PendingFirstMip;
-			RHIUpdateTextureReference(Owner->TextureRefRHI,TextureRHI);
+			RHIUpdateTextureReference(Owner->TextureReference.TextureReferenceRHI,TextureRHI);
 
 			// Update mip-level fading.
 			EMipFadeSettings MipFadeSetting = (Owner->LODGroup == TEXTUREGROUP_Lightmap || Owner->LODGroup == TEXTUREGROUP_Shadowmap) ? MipFade_Slow : MipFade_Normal;
@@ -2315,7 +2310,7 @@ bool FTexture2DResource::TryReallocate( int32 OldMipCount, int32 NewMipCount )
 
 	Texture2DRHI = NewTextureRHI;
 	TextureRHI = NewTextureRHI;
-	RHIUpdateTextureReference(Owner->TextureRefRHI,TextureRHI);
+	RHIUpdateTextureReference(Owner->TextureReference.TextureReferenceRHI,TextureRHI);
 
 	PendingFirstMip = CurrentFirstMip = MipIndex;
 
diff --git a/Engine/Source/Runtime/Engine/Private/Texture2DDynamic.cpp b/Engine/Source/Runtime/Engine/Private/Texture2DDynamic.cpp
--- a/Engine/Source/Runtime/Engine/Private/Texture2DDynamic.cpp
+++ b/Engine/Source/Runtime/Engine/Private/Texture2DDynamic.cpp
@@ -59,13 +59,13 @@ void FTexture2DDynamicResource::InitRHI()
 	FRHIResourceCreateInfo CreateInfo;
 	Texture2DRHI = RHICreateTexture2D( Owner->SizeX, Owner->SizeY, Owner->Format, Owner->NumMips, 1, Flags, CreateInfo );
 	TextureRHI = Texture2DRHI;
-	RHIUpdateTextureReference(Owner->TextureRefRHI,TextureRHI);
+	RHIUpdateTextureReference(Owner->TextureReference.TextureReferenceRHI,TextureRHI);
 }
 
 /** Called when the resource is released. This is only called by the rendering thread. */
 void FTexture2DDynamicResource::ReleaseRHI()
 {
-	RHIUpdateTextureReference(Owner->TextureRefRHI,FTextureRHIParamRef());
+	RHIUpdateTextureReference(Owner->TextureReference.TextureReferenceRHI,FTextureRHIParamRef());
 	FTextureResource::ReleaseRHI();
 	Texture2DRHI.SafeRelease();
 }
diff --git a/Engine/Source/Runtime/Engine/Private/TextureCube.cpp b/Engine/Source/Runtime/Engine/Private/TextureCube.cpp
--- a/Engine/Source/Runtime/Engine/Private/TextureCube.cpp
+++ b/Engine/Source/Runtime/Engine/Private/TextureCube.cpp
@@ -228,7 +228,7 @@ public:
 		TextureCubeRHI = RHICreateTextureCube( Owner->GetSizeX(), Owner->GetPixelFormat(), Owner->GetNumMips(), TexCreateFlags, CreateInfo );
 		TextureRHI = TextureCubeRHI;
 		RHIBindDebugLabelName(TextureRHI, *Owner->GetName());
-		RHIUpdateTextureReference(Owner->TextureRefRHI,TextureRHI);
+		RHIUpdateTextureReference(Owner->TextureReference.TextureReferenceRHI,TextureRHI);
 
 		// Read the mip-levels into the RHI texture.
 		int32 NumMips = Owner->GetNumMips();
@@ -265,7 +265,7 @@ public:
 	{
 		DEC_DWORD_STAT_BY( STAT_TextureMemory, TextureSize );
 		DEC_DWORD_STAT_FNAME_BY( LODGroupStatName, TextureSize );
-		RHIUpdateTextureReference(Owner->TextureRefRHI,FTextureRHIParamRef());
+		RHIUpdateTextureReference(Owner->TextureReference.TextureReferenceRHI,FTextureRHIParamRef());
 		TextureCubeRHI.SafeRelease();
 		FTextureResource::ReleaseRHI();
 	}
diff --git a/Engine/Source/Runtime/Engine/Private/TextureMovie.cpp b/Engine/Source/Runtime/Engine/Private/TextureMovie.cpp
--- a/Engine/Source/Runtime/Engine/Private/TextureMovie.cpp
+++ b/Engine/Source/Runtime/Engine/Private/TextureMovie.cpp
@@ -367,7 +367,7 @@ void FTextureMovieResource::InitDynamicRHI()
 			Texture2DRHI
 			);
 		TextureRHI = (FTextureRHIRef&)Texture2DRHI;
-		RHIUpdateTextureReference(Owner->TextureRefRHI,TextureRHI);
+		RHIUpdateTextureReference(Owner->TextureReference.TextureReferenceRHI,TextureRHI);
 
 		// add to the list of global deferred updates (updated during scene rendering)
 		// since the latest decoded movie frame is rendered to this movie texture target
@@ -395,7 +395,7 @@ void FTextureMovieResource::ReleaseDynamicRHI()
 	// release the FTexture RHI resources here as well
 	ReleaseRHI();
 
-	RHIUpdateTextureReference(Owner->TextureRefRHI,FTextureRHIParamRef());
+	RHIUpdateTextureReference(Owner->TextureReference.TextureReferenceRHI,FTextureRHIParamRef());
 	Texture2DRHI.SafeRelease();
 	RenderTargetTextureRHI.SafeRelease();	
 
diff --git a/Engine/Source/Runtime/Engine/Private/TextureRenderTarget2D.cpp b/Engine/Source/Runtime/Engine/Private/TextureRenderTarget2D.cpp
--- a/Engine/Source/Runtime/Engine/Private/TextureRenderTarget2D.cpp
+++ b/Engine/Source/Runtime/Engine/Private/TextureRenderTarget2D.cpp
@@ -326,7 +326,7 @@ void FTextureRenderTarget2DResource::InitDynamicRHI()
 			Texture2DRHI
 			);
 		TextureRHI = (FTextureRHIRef&)Texture2DRHI;
-		RHIUpdateTextureReference(Owner->TextureRefRHI,TextureRHI);
+		RHIUpdateTextureReference(Owner->TextureReference.TextureReferenceRHI,TextureRHI);
 
 		AddToDeferredUpdateList(true);
 	}
@@ -352,7 +352,7 @@ void FTextureRenderTarget2DResource::ReleaseDynamicRHI()
 	// release the FTexture RHI resources here as well
 	ReleaseRHI();
 
-	RHIUpdateTextureReference(Owner->TextureRefRHI,FTextureRHIParamRef());
+	RHIUpdateTextureReference(Owner->TextureReference.TextureReferenceRHI,FTextureRHIParamRef());
 	Texture2DRHI.SafeRelease();
 	RenderTargetTextureRHI.SafeRelease();	
 
diff --git a/Engine/Source/Runtime/Engine/Private/TextureRenderTargetCube.cpp b/Engine/Source/Runtime/Engine/Private/TextureRenderTargetCube.cpp
--- a/Engine/Source/Runtime/Engine/Private/TextureRenderTargetCube.cpp
+++ b/Engine/Source/Runtime/Engine/Private/TextureRenderTargetCube.cpp
@@ -224,7 +224,7 @@ void FTextureRenderTargetCubeResource::InitDynamicRHI()
 		}
 
 		TextureRHI = TextureCubeRHI;
-		RHIUpdateTextureReference(Owner->TextureRefRHI,TextureRHI);
+		RHIUpdateTextureReference(Owner->TextureReference.TextureReferenceRHI,TextureRHI);
 
 		// Create the RHI target surface used for rendering to
 		{
@@ -267,7 +267,7 @@ void FTextureRenderTargetCubeResource::ReleaseDynamicRHI()
 	// release the FTexture RHI resources here as well
 	ReleaseRHI();
 
-	RHIUpdateTextureReference(Owner->TextureRefRHI,FTextureRHIParamRef());
+	RHIUpdateTextureReference(Owner->TextureReference.TextureReferenceRHI,FTextureRHIParamRef());
 	CubeFaceSurfaceRHI.SafeRelease();
 	RenderTargetCubeRHI.SafeRelease();
 	RenderTargetTextureRHI.SafeRelease();	
diff --git a/Engine/Source/Runtime/IOS/MetalRHI/Private/MetalTexture.cpp b/Engine/Source/Runtime/IOS/MetalRHI/Private/MetalTexture.cpp
--- a/Engine/Source/Runtime/IOS/MetalRHI/Private/MetalTexture.cpp
+++ b/Engine/Source/Runtime/IOS/MetalRHI/Private/MetalTexture.cpp
@@ -24,8 +24,8 @@ FORCEINLINE int32 GetMetalCubeFace(ECubeFace Face)
 class FMetalTextureReference : public FRHITextureReference
 {
 public:
-	explicit FMetalTextureReference()
-		: FRHITextureReference()
+	explicit FMetalTextureReference(FLastRenderTimeContainer* InLastRenderTime)
+		: FRHITextureReference(InLastRenderTime)
 	{}
 
 	// IRefCountedObject interface.
@@ -395,9 +395,9 @@ void FMetalDynamicRHI::RHIUnlockTextureCubeFace(FTextureCubeRHIParamRef TextureC
 
 
 
-FTextureReferenceRHIRef FMetalDynamicRHI::RHICreateTextureReference()
+FTextureReferenceRHIRef FMetalDynamicRHI::RHICreateTextureReference(FLastRenderTimeContainer* InLastRenderTime)
 {
-	return new FMetalTextureReference();
+	return new FMetalTextureReference(InLastRenderTime);
 }
 
 void FMetalDynamicRHI::RHIUpdateTextureReference(FTextureReferenceRHIParamRef TextureRefRHI, FTextureRHIParamRef NewTextureRHI)
diff --git a/Engine/Source/Runtime/IOS/MetalRHI/Private/MetalUniformBuffer.cpp b/Engine/Source/Runtime/IOS/MetalRHI/Private/MetalUniformBuffer.cpp
--- a/Engine/Source/Runtime/IOS/MetalRHI/Private/MetalUniformBuffer.cpp
+++ b/Engine/Source/Runtime/IOS/MetalRHI/Private/MetalUniformBuffer.cpp
@@ -284,7 +284,7 @@ void FMetalUniformBuffer::CacheResourcesInternal()
 					// todo: this does multiple virtual function calls to find the right type to cast to
 					// this is due to multiple inheritance nastiness, NEEDS CLEANUP
 					FRHITexture* TextureRHI = (FRHITexture*)Resources[i].GetReference();
-					TextureRHI->SetLastCachedTime(CurrentTime);
+					TextureRHI->SetLastRenderTime(CurrentTime);
 					FMetalSurface& Surface = GetMetalSurfaceFromRHITexture(TextureRHI);
 					RawResources[i] = &Surface;
 				}
diff --git a/Engine/Source/Runtime/OpenGLDrv/Private/OpenGLTexture.cpp b/Engine/Source/Runtime/OpenGLDrv/Private/OpenGLTexture.cpp
--- a/Engine/Source/Runtime/OpenGLDrv/Private/OpenGLTexture.cpp
+++ b/Engine/Source/Runtime/OpenGLDrv/Private/OpenGLTexture.cpp
@@ -1907,9 +1907,9 @@ void FOpenGLDynamicRHI::RHIVirtualTextureSetFirstMipVisible(FTexture2DRHIParamRe
 {
 }
 
-FTextureReferenceRHIRef FOpenGLDynamicRHI::RHICreateTextureReference()
+FTextureReferenceRHIRef FOpenGLDynamicRHI::RHICreateTextureReference(FLastRenderTimeContainer* InLastRenderTime)
 {
-	return new FOpenGLTextureReference();
+	return new FOpenGLTextureReference(InLastRenderTime);
 }
 
 void FOpenGLTextureReference::SetReferencedTexture(FRHITexture* InTexture)
diff --git a/Engine/Source/Runtime/OpenGLDrv/Private/OpenGLUniformBuffer.cpp b/Engine/Source/Runtime/OpenGLDrv/Private/OpenGLUniformBuffer.cpp
--- a/Engine/Source/Runtime/OpenGLDrv/Private/OpenGLUniformBuffer.cpp
+++ b/Engine/Source/Runtime/OpenGLDrv/Private/OpenGLUniformBuffer.cpp
@@ -409,7 +409,7 @@ void FOpenGLUniformBuffer::CacheResourcesInternal()
 		case UBMT_TEXTURE:
 			{
 				FRHITexture* TextureRHI = (FRHITexture*)Resources[i].GetReference();
-				TextureRHI->SetLastCachedTime(CurrentTime);
+				TextureRHI->SetLastRenderTime(CurrentTime);
 				// todo: this does multiple virtual function calls to find the right type to cast to
 				// this is due to multiple inheritance nastiness, NEEDS CLEANUP
 				FOpenGLTextureBase* TextureGL = GetOpenGLTextureFromRHITexture(TextureRHI);
diff --git a/Engine/Source/Runtime/OpenGLDrv/Public/OpenGLResources.h b/Engine/Source/Runtime/OpenGLDrv/Public/OpenGLResources.h
--- a/Engine/Source/Runtime/OpenGLDrv/Public/OpenGLResources.h
+++ b/Engine/Source/Runtime/OpenGLDrv/Public/OpenGLResources.h
@@ -1000,8 +1000,9 @@ class FOpenGLTextureReference : public FRHITextureReference
 	FOpenGLTextureBase* TexturePtr;
 
 public:
-	FOpenGLTextureReference()
-		: TexturePtr(NULL)
+	explicit FOpenGLTextureReference(FLastRenderTimeContainer* InLastRenderTime)
+		: FRHITextureReference(InLastRenderTime)
+		, TexturePtr(NULL)
 	{}
 
 	void SetReferencedTexture(FRHITexture* InTexture);
diff --git a/Engine/Source/Runtime/RHI/Public/RHIMethods.h b/Engine/Source/Runtime/RHI/Public/RHIMethods.h
--- a/Engine/Source/Runtime/RHI/Public/RHIMethods.h
+++ b/Engine/Source/Runtime/RHI/Public/RHIMethods.h
@@ -695,8 +695,9 @@ DEFINE_RHIMETHOD_5(
 	return,return false
 	);
 
-DEFINE_RHIMETHOD_0(
+DEFINE_RHIMETHOD_GLOBAL_1(
 	FTextureReferenceRHIRef,CreateTextureReference,
+	FLastRenderTimeContainer*,LastRenderTime,
 	return, return new FRHITextureReferenceNullImpl();
 	);
 
@@ -1046,24 +1047,24 @@ DEFINE_RHIMETHOD_GLOBAL_5(
 	);
 
 
-/**
-* Creates a Cube Array RHI texture resource
-* @param Size - width/height of the texture to create
-* @param ArraySize - number of array elements of the texture to create
-* @param Format - EPixelFormat texture format
-* @param NumMips - number of mips to generate or 0 for full mip pyramid
-* @param Flags - ETextureCreateFlags creation flags
-*/
+	/**
+	* Creates a Cube Array RHI texture resource
+	* @param Size - width/height of the texture to create
+	* @param ArraySize - number of array elements of the texture to create
+	* @param Format - EPixelFormat texture format
+	* @param NumMips - number of mips to generate or 0 for full mip pyramid
+	* @param Flags - ETextureCreateFlags creation flags
+	*/
 DEFINE_RHIMETHOD_GLOBAL_6(
 	FTextureCubeRHIRef,CreateTextureCubeArray,
-	uint32,Size,
-	uint32,ArraySize,
-	uint8,Format,
-	uint32,NumMips,
-	uint32,Flags,
+		uint32,Size,
+		uint32,ArraySize,
+		uint8,Format,
+		uint32,NumMips,
+		uint32,Flags,
 	FRHIResourceCreateInfo&,CreateInfo,
-	return,return new FRHITextureCube(Size,NumMips,(EPixelFormat)Format,Flags);
-);
+		return,return new FRHITextureCube(Size,NumMips,(EPixelFormat)Format,Flags);
+	);
 
 
 /**
diff --git a/Engine/Source/Runtime/RHI/Public/RHIResources.h b/Engine/Source/Runtime/RHI/Public/RHIResources.h
--- a/Engine/Source/Runtime/RHI/Public/RHIResources.h
+++ b/Engine/Source/Runtime/RHI/Public/RHIResources.h
@@ -194,17 +194,30 @@ private:
 // Textures
 //
 
+class RHI_API FLastRenderTimeContainer
+{
+public:
+	FLastRenderTimeContainer() : LastRenderTime(-FLT_MAX) {}
+
+	double GetLastRenderTime() const { return LastRenderTime; }
+	void SetLastRenderTime(double InLastRenderTime) { LastRenderTime = InLastRenderTime; }
+
+private:
+	/** The last time the resource was rendered. */
+	double LastRenderTime;
+};
+
 class RHI_API FRHITexture : public FRHIResource
 {
 public:
 	
 	/** Initialization constructor. */
-	FRHITexture(uint32 InNumMips,uint32 InNumSamples,EPixelFormat InFormat,uint32 InFlags)
+	FRHITexture(uint32 InNumMips,uint32 InNumSamples,EPixelFormat InFormat,uint32 InFlags,FLastRenderTimeContainer* InLastRenderTime)
 	: NumMips(InNumMips)
 	, NumSamples(InNumSamples)
 	, Format(InFormat)
 	, Flags(InFlags)
-	, LastCachedTime(-FLT_MAX)
+	, LastRenderTime(InLastRenderTime ? *InLastRenderTime : DefaultLastRenderTime)
 	{}
 
 	// Dynamic cast methods.
@@ -255,13 +268,20 @@ public:
 
 	FRHIResourceInfo ResourceInfo;
 
-	/** @return the last time this texture was cached in a resource table. */
-	inline float GetLastCachedTime() const { return LastCachedTime; }
-
 	/** sets the last time this texture was cached in a resource table. */
-	void SetLastCachedTime(float InLastCachedTime)
+	void SetLastRenderTime(float InLastRenderTime)
+	{
+		LastRenderTime.SetLastRenderTime(InLastRenderTime);
+	}
+
+	/** Returns the last render time container, or NULL if none were specified at creation. */
+	FLastRenderTimeContainer* GetLastRenderTimeContainer()
 	{
-		LastCachedTime = InLastCachedTime;
+		if (&LastRenderTime == &DefaultLastRenderTime)
+		{
+			return NULL;
+		}
+		return &LastRenderTime;
 	}
 
 private:
@@ -269,7 +289,8 @@ private:
 	uint32 NumSamples;
 	EPixelFormat Format;
 	uint32 Flags;
-	float LastCachedTime;
+	FLastRenderTimeContainer& LastRenderTime;
+	FLastRenderTimeContainer DefaultLastRenderTime;
 };
 
 class RHI_API FRHITexture2D : public FRHITexture
@@ -278,7 +299,7 @@ public:
 	
 	/** Initialization constructor. */
 	FRHITexture2D(uint32 InSizeX,uint32 InSizeY,uint32 InNumMips,uint32 InNumSamples,EPixelFormat InFormat,uint32 InFlags)
-	: FRHITexture(InNumMips,InNumSamples,InFormat,InFlags)
+	: FRHITexture(InNumMips,InNumSamples,InFormat,InFlags,NULL)
 	, SizeX(InSizeX)
 	, SizeY(InSizeY)
 	{}
@@ -304,7 +325,7 @@ public:
 	
 	/** Initialization constructor. */
 	FRHITexture2DArray(uint32 InSizeX,uint32 InSizeY,uint32 InSizeZ,uint32 InNumMips,EPixelFormat InFormat,uint32 InFlags)
-	: FRHITexture(InNumMips,1,InFormat,InFlags)
+	: FRHITexture(InNumMips,1,InFormat,InFlags,NULL)
 	, SizeX(InSizeX)
 	, SizeY(InSizeY)
 	, SizeZ(InSizeZ)
@@ -335,7 +356,7 @@ public:
 	
 	/** Initialization constructor. */
 	FRHITexture3D(uint32 InSizeX,uint32 InSizeY,uint32 InSizeZ,uint32 InNumMips,EPixelFormat InFormat,uint32 InFlags)
-	: FRHITexture(InNumMips,1,InFormat,InFlags)
+	: FRHITexture(InNumMips,1,InFormat,InFlags,NULL)
 	, SizeX(InSizeX)
 	, SizeY(InSizeY)
 	, SizeZ(InSizeZ)
@@ -366,7 +387,7 @@ public:
 	
 	/** Initialization constructor. */
 	FRHITextureCube(uint32 InSize,uint32 InNumMips,EPixelFormat InFormat,uint32 InFlags)
-	: FRHITexture(InNumMips,1,InFormat,InFlags)
+	: FRHITexture(InNumMips,1,InFormat,InFlags,NULL)
 	, Size(InSize)
 	{}
 	
@@ -384,8 +405,8 @@ private:
 class RHI_API FRHITextureReference : public FRHITexture
 {
 public:
-	explicit FRHITextureReference()
-		: FRHITexture(0,0,PF_Unknown,0)
+	explicit FRHITextureReference(FLastRenderTimeContainer* InLastRenderTime)
+		: FRHITexture(0,0,PF_Unknown,0,InLastRenderTime)
 	{}
 
 	virtual FRHITextureReference* GetTextureReference() override { return this; }
@@ -405,7 +426,7 @@ class RHI_API FRHITextureReferenceNullImpl : public FRHITextureReference
 {
 public:
 	FRHITextureReferenceNullImpl()
-		: FRHITextureReference()
+		: FRHITextureReference(NULL)
 	{}
 
 	void SetReferencedTexture(FRHITexture* InTexture)
diff --git a/Engine/Source/Runtime/RenderCore/Private/RenderResource.cpp b/Engine/Source/Runtime/RenderCore/Private/RenderResource.cpp
--- a/Engine/Source/Runtime/RenderCore/Private/RenderResource.cpp
+++ b/Engine/Source/Runtime/RenderCore/Private/RenderResource.cpp
@@ -114,6 +114,47 @@ void ReleaseResourceAndFlush(FRenderResource* Resource)
 	FlushRenderingCommands();
 }
 
+FTextureReference::FTextureReference()
+	: TextureReferenceRHI(NULL)
+{
+}
+
+FTextureReference::~FTextureReference()
+{
+}
+
+void FTextureReference::BeginInit_GameThread()
+{
+	bInitialized_GameThread = true;
+	BeginInitResource(this);
+}
+
+void FTextureReference::BeginRelease_GameThread()
+{
+	BeginReleaseResource(this);
+	bInitialized_GameThread = false;
+}
+
+void FTextureReference::InvalidateLastRenderTime()
+{
+	LastRenderTimeRHI.SetLastRenderTime(-FLT_MAX);
+}
+
+void FTextureReference::InitRHI()
+{
+	TextureReferenceRHI = RHICreateTextureReference(&LastRenderTimeRHI);
+}
+	
+void FTextureReference::ReleaseRHI()
+{
+	TextureReferenceRHI.SafeRelease();
+}
+
+FString FTextureReference::GetFriendlyName() const
+{
+	return TEXT("FTextureReference");
+}
+
 /** The global null color vertex buffer, which is set with a stride of 0 on meshes without a color component. */
 TGlobalResource<FNullColorVertexBuffer> GNullColorVertexBuffer;
 
diff --git a/Engine/Source/Runtime/RenderCore/Public/RenderResource.h b/Engine/Source/Runtime/RenderCore/Public/RenderResource.h
--- a/Engine/Source/Runtime/RenderCore/Public/RenderResource.h
+++ b/Engine/Source/Runtime/RenderCore/Public/RenderResource.h
@@ -333,6 +333,49 @@ public:
 	virtual FString GetFriendlyName() const { return TEXT("FTexture"); }
 };
 
+/** A texture reference resource. */
+class RENDERCORE_API FTextureReference : public FRenderResource
+{
+public:
+	/** The texture reference's RHI resource. */
+	FTextureReferenceRHIRef	TextureReferenceRHI;
+
+
+private:
+	/** The last time the texture has been rendered via this reference. */
+	FLastRenderTimeContainer LastRenderTimeRHI;
+
+	/** True if the texture reference has been initialized from the game thread. */
+	bool bInitialized_GameThread;
+
+public:
+	/** Default constructor. */
+	FTextureReference();
+
+	// Destructor
+	virtual ~FTextureReference();
+
+	/** Returns the last time the texture has been rendered via this reference. */
+	double GetLastRenderTime() const { return LastRenderTimeRHI.GetLastRenderTime(); }
+
+	/** Invalidates the last render time. */
+	void InvalidateLastRenderTime();
+
+	/** Returns true if the texture reference has been initialized from the game thread. */
+	bool IsInitialized_GameThread() const { return bInitialized_GameThread; }
+
+	/** Kicks off the initialization process on the game thread. */
+	void BeginInit_GameThread();
+
+	/** Kicks off the release process on the game thread. */
+	void BeginRelease_GameThread();
+
+	// FRenderResource interface.
+	virtual void InitRHI();
+	virtual void ReleaseRHI();
+	virtual FString GetFriendlyName() const;
+};
+
 /** A vertex buffer resource */
 class RENDERCORE_API FVertexBuffer : public FRenderResource
 {
diff --git a/Engine/Source/Runtime/Windows/D3D11RHI/Private/D3D11Texture.cpp b/Engine/Source/Runtime/Windows/D3D11RHI/Private/D3D11Texture.cpp
--- a/Engine/Source/Runtime/Windows/D3D11RHI/Private/D3D11Texture.cpp
+++ b/Engine/Source/Runtime/Windows/D3D11RHI/Private/D3D11Texture.cpp
@@ -1502,9 +1502,9 @@ void FD3D11DynamicRHI::RHIVirtualTextureSetFirstMipVisible(FTexture2DRHIParamRef
 {
 }
 
-FTextureReferenceRHIRef FD3D11DynamicRHI::RHICreateTextureReference()
+FTextureReferenceRHIRef FD3D11DynamicRHI::RHICreateTextureReference(FLastRenderTimeContainer* LastRenderTime)
 {
-	return new FD3D11TextureReference(this);
+	return new FD3D11TextureReference(this,LastRenderTime);
 }
 
 void FD3D11DynamicRHI::RHIUpdateTextureReference(FTextureReferenceRHIParamRef TextureRefRHI, FTextureRHIParamRef NewTextureRHI)
diff --git a/Engine/Source/Runtime/Windows/D3D11RHI/Private/D3D11UniformBuffer.cpp b/Engine/Source/Runtime/Windows/D3D11RHI/Private/D3D11UniformBuffer.cpp
--- a/Engine/Source/Runtime/Windows/D3D11RHI/Private/D3D11UniformBuffer.cpp
+++ b/Engine/Source/Runtime/Windows/D3D11RHI/Private/D3D11UniformBuffer.cpp
@@ -272,7 +272,7 @@ void FD3D11UniformBuffer::CacheResourcesInternal()
 				// todo: this does multiple virtual function calls to find the right type to cast to
 				// this is due to multiple inheritance nastiness, NEEDS CLEANUP
 				FRHITexture* TextureRHI = (FRHITexture*)Resources[i].GetReference();
-				TextureRHI->SetLastCachedTime(CurrentTime);
+				TextureRHI->SetLastRenderTime(CurrentTime);
 				FD3D11TextureBase* TextureD3D11 = GetD3D11TextureFromRHITexture(TextureRHI);
 				RawResources[i].ShaderResource = TextureD3D11->GetBaseShaderResource();
 				RawResources[i].D3D11Resource = (IUnknown*)TextureD3D11->GetShaderResourceView();
diff --git a/Engine/Source/Runtime/Windows/D3D11RHI/Public/D3D11Resources.h b/Engine/Source/Runtime/Windows/D3D11RHI/Public/D3D11Resources.h
--- a/Engine/Source/Runtime/Windows/D3D11RHI/Public/D3D11Resources.h
+++ b/Engine/Source/Runtime/Windows/D3D11RHI/Public/D3D11Resources.h
@@ -489,8 +489,8 @@ typedef TD3D11Texture2D<FD3D11BaseTextureCube>    FD3D11TextureCube;
 class FD3D11TextureReference : public FRHITextureReference, public FD3D11TextureBase
 {
 public:
-	explicit FD3D11TextureReference(class FD3D11DynamicRHI* InD3DRHI)
-		: FRHITextureReference()
+	FD3D11TextureReference(class FD3D11DynamicRHI* InD3DRHI, FLastRenderTimeContainer* LastRenderTime)
+		: FRHITextureReference(LastRenderTime)
 		, FD3D11TextureBase(InD3DRHI,NULL,NULL, 0, false,TArray<TRefCountPtr<ID3D11RenderTargetView> >(),NULL)
 	{
 		BaseShaderResource = NULL;
