Fixed translation issue with Brushes which have a non-zero baked prepivot.  Fixed issue where decals' initialized scale is not preserved on spawning.
#ttp 340007 - CRITICAL: Regression: EDITOR: GEOMETRY: Hightlight becomes offset after saving pivot location
#ttp 340361 - CRITICAL: Regression: EDITOR: DECAL: decals have improper scale
#branch UE4
#proj Runtime.Engine
#change In AActor::RerunConstructionScripts(), get the old transform from the RootComponent relative transform instead of the current transform - this is guaranteed not to have been modified by the custom location code in SceneComponent.  It also ensures that the relative scale set up by the decal component is preserved.
#change In AActor::PostSpawnInitialize(), when calling ExecuteConstruction, pass in a transform which preserves the original RelativeScale3D of the root component, instead of trashing it and resetting it to an identity.

diff --git a/Engine/Source/Runtime/Engine/Private/Actor.cpp b/Engine/Source/Runtime/Engine/Private/Actor.cpp
--- a/Engine/Source/Runtime/Engine/Private/Actor.cpp
+++ b/Engine/Source/Runtime/Engine/Private/Actor.cpp
@@ -2238,7 +2238,9 @@ void AActor::PostSpawnInitialize(FVector const& SpawnLocation, FRotator const& S
 	// After this, we can assume all components are created and assembled.
 	if (!bDeferConstruction)
 	{
-		ExecuteConstruction( FTransform(SpawnRotation, SpawnLocation), NULL );
+		// Preserve original root component scale
+		const FVector SpawnScale = GetRootComponent() ? GetRootComponent()->RelativeScale3D : FVector(1.0f, 1.0f, 1.0f);
+		ExecuteConstruction( FTransform(SpawnRotation, SpawnLocation, SpawnScale), NULL );
 		PostActorConstruction();
 	}
 }
diff --git a/Engine/Source/Runtime/Engine/Private/ActorConstruction.cpp b/Engine/Source/Runtime/Engine/Private/ActorConstruction.cpp
--- a/Engine/Source/Runtime/Engine/Private/ActorConstruction.cpp
+++ b/Engine/Source/Runtime/Engine/Private/ActorConstruction.cpp
@@ -238,7 +238,8 @@ void AActor::RerunConstructionScripts()
 					//detach it to remove any scaling 
 					RootComponent->DetachFromParent(true);
 				}
-				OldTransform = RootComponent->ComponentToWorld;
+
+				OldTransform = RootComponent->GetRelativeTransform();
 			}
 		}
 
diff --git a/Engine/Source/Runtime/Engine/Private/ActorEditor.cpp b/Engine/Source/Runtime/Engine/Private/ActorEditor.cpp
--- a/Engine/Source/Runtime/Engine/Private/ActorEditor.cpp
+++ b/Engine/Source/Runtime/Engine/Private/ActorEditor.cpp
@@ -224,7 +224,7 @@ AActor::FActorTransactionAnnotation::FActorTransactionAnnotation(const AActor* A
 	if (RootComponent && RootComponent->bCreatedByConstructionScript)
 	{
 		bRootComponentDataCached = true;
-		RootComponentData.Transform = RootComponent->ComponentToWorld;
+		RootComponentData.Transform = RootComponent->GetRelativeTransform();
 
 		if (RootComponent->AttachParent)
 		{
