Export GameNetworkManager.

Fixes https://answers.unrealengine.com/questions/65570/agamenetworkmanager-unresolved-externals.html

diff --git a/Engine/Source/Runtime/Engine/Classes/GameFramework/GameNetworkManager.h b/Engine/Source/Runtime/Engine/Classes/GameFramework/GameNetworkManager.h
--- a/Engine/Source/Runtime/Engine/Classes/GameFramework/GameNetworkManager.h
+++ b/Engine/Source/Runtime/Engine/Classes/GameFramework/GameNetworkManager.h
@@ -18,7 +18,7 @@ enum EStandbyType
 };
 
 UCLASS(config=Game, notplaceable)
-class AGameNetworkManager : public AInfo
+class ENGINE_API AGameNetworkManager : public AInfo
 {
 	GENERATED_UCLASS_BODY()
 
