Expose Instigator to blueprints and SpawnActor node

diff --git a/Engine/Source/Runtime/Engine/Classes/GameFramework/Actor.h b/Engine/Source/Runtime/Engine/Classes/GameFramework/Actor.h
--- a/Engine/Source/Runtime/Engine/Classes/GameFramework/Actor.h
+++ b/Engine/Source/Runtime/Engine/Classes/GameFramework/Actor.h
@@ -360,7 +360,7 @@ public:
 	uint32 bFindCameraComponentWhenViewTarget:1;
 
 	/** Pawn responsible for damage caused by this actor. */
-	UPROPERTY(replicatedUsing=OnRep_Instigator)
+	UPROPERTY(BlueprintReadWrite, replicatedUsing=OnRep_Instigator, meta=(ExposeOnSpawn=true), Category=Actor)
 	class APawn* Instigator;      
 
 	/** The time this actor was created, relative to World->GetTimeSeconds() */
