Automatically propagate Instigator values when calling Spawn Actor in blueprints.
If called from the Pawn context, pawn becomes the Instigator
If called from an Actor context, the actor's Instigator is used as the instigator.

diff --git a/Engine/Source/Runtime/Engine/Private/GameplayStatics.cpp b/Engine/Source/Runtime/Engine/Private/GameplayStatics.cpp
--- a/Engine/Source/Runtime/Engine/Private/GameplayStatics.cpp
+++ b/Engine/Source/Runtime/Engine/Private/GameplayStatics.cpp
@@ -282,8 +282,21 @@ AActor* UGameplayStatics::BeginSpawningActorFromClass(UObject* WorldContextObjec
 	UClass* Class = *ActorClass;
 	if (Class != NULL)
 	{
+		// If the WorldContextObject is a Pawn we will use that as the instigator.
+		// Otherwise if the WorldContextObject is an Actor we will share its instigator.
+		// If the value is set via the exposed parameter on SpawnNode it will be overwritten anyways, so this is safe to specify here
+		APawn* AutoInstigator = Cast<APawn>(WorldContextObject);
+		if (AutoInstigator == nullptr)
+		{
+			AActor* ContextActor = Cast<AActor>(WorldContextObject);
+			if (ContextActor)
+			{
+				AutoInstigator = ContextActor->Instigator;
+			}
+		}
+
 		UWorld* World = GEngine->GetWorldFromContextObject(WorldContextObject);
-		NewActor = World->SpawnActorDeferred<AActor>(Class, SpawnLoc, SpawnRot, NULL, NULL, bNoCollisionFail);
+		NewActor = World->SpawnActorDeferred<AActor>(Class, SpawnLoc, SpawnRot, NULL, AutoInstigator, bNoCollisionFail);
 	}
 
 	return NewActor;
