[337960] Potential fix for a crash that reportedly occurred while loading a few unique Blueprint assets.

- Modified AActor::SetRootComponent() to replace a check() with an ensure() since we're already handling the failure condition.
- Modified UEngine::CopyPropertiesForUnrelatedObjects() to get/set the RootComponent property value directly rather than calling SetRootComponent() in order to bypass some unnecessary validation. Also tested to make sure that the situation that initially required saving/restoring the RootComponent in CPFUO (e.g. making sure that reparenting from StaticMeshActor->SkeletalMeshActor and vice-versa preserves the RootComponent) is still being handled in the same way after the change.

diff --git a/Engine/Source/Runtime/Engine/Private/Actor.cpp b/Engine/Source/Runtime/Engine/Private/Actor.cpp
--- a/Engine/Source/Runtime/Engine/Private/Actor.cpp
+++ b/Engine/Source/Runtime/Engine/Private/Actor.cpp
@@ -2677,8 +2677,7 @@ UPrimitiveComponent* AActor::GetRootPrimitiveComponent() const
 bool AActor::SetRootComponent(class USceneComponent* NewRootComponent)
 {
 	/** Only components owned by this actor can be used as a its root component. */
-	check(NewRootComponent == NULL || NewRootComponent->GetOwner() == this);
-	if (NewRootComponent == NULL || NewRootComponent->GetOwner() == this)
+	if (ensure(NewRootComponent == NULL || NewRootComponent->GetOwner() == this))
 	{
 		RootComponent = NewRootComponent;
 		return true;
diff --git a/Engine/Source/Runtime/Engine/Private/UnrealEngine.cpp b/Engine/Source/Runtime/Engine/Private/UnrealEngine.cpp
--- a/Engine/Source/Runtime/Engine/Private/UnrealEngine.cpp
+++ b/Engine/Source/Runtime/Engine/Private/UnrealEngine.cpp
@@ -9446,9 +9446,14 @@ void UEngine::CopyPropertiesForUnrelatedObjects(UObject* OldObject, UObject* New
 
 	// If the new object is an Actor, save the root component reference, to be restored later
 	USceneComponent* SavedRootComponent = NULL;
+	UObjectProperty* RootComponentProperty = NULL;
 	if(NewActor != NULL)
 	{
-		SavedRootComponent = NewActor->GetRootComponent();
+		RootComponentProperty = FindField<UObjectProperty>(NewActor->GetClass(), "RootComponent");
+		if(RootComponentProperty != NULL)
+		{
+			SavedRootComponent = Cast<USceneComponent>(RootComponentProperty->GetObjectPropertyValue_InContainer(NewActor));
+		}
 	}
 
 	// Serialize out the modified properties on the old default object
@@ -9582,7 +9587,11 @@ void UEngine::CopyPropertiesForUnrelatedObjects(UObject* OldObject, UObject* New
 	// Restore the root component reference
 	if(NewActor != NULL)
 	{
-		NewActor->SetRootComponent(SavedRootComponent);
+		if(RootComponentProperty != NULL)
+		{
+			RootComponentProperty->SetObjectPropertyValue_InContainer(NewActor, SavedRootComponent);
+		}
+
 		NewActor->ResetOwnedComponents();
 	}
 
