Detect and remove broken network guids

diff --git a/Engine/Source/Runtime/CoreUObject/Public/UObject/CoreNet.h b/Engine/Source/Runtime/CoreUObject/Public/UObject/CoreNet.h
--- a/Engine/Source/Runtime/CoreUObject/Public/UObject/CoreNet.h
+++ b/Engine/Source/Runtime/CoreUObject/Public/UObject/CoreNet.h
@@ -143,6 +143,7 @@ class COREUOBJECT_API UPackageMap : public UObject
 
 	virtual void		LogDebugInfo( FOutputDevice & Ar) { }
 	virtual UObject *	GetObjectFromNetGUID( const FNetworkGUID & NetGUID, const bool bIgnoreMustBeMapped ) { return NULL; }
+	virtual bool		IsGUIDBroken( const FNetworkGUID & NetGUID, const bool bMustBeRegistered ) const { return false; }
 
 protected:
 
diff --git a/Engine/Source/Runtime/Engine/Classes/Engine/PackageMapClient.h b/Engine/Source/Runtime/Engine/Classes/Engine/PackageMapClient.h
--- a/Engine/Source/Runtime/Engine/Classes/Engine/PackageMapClient.h
+++ b/Engine/Source/Runtime/Engine/Classes/Engine/PackageMapClient.h
@@ -54,7 +54,7 @@ public:
 	void			CleanReferences();
 	bool			SupportsObject( const UObject * Object );
 	bool			IsDynamicObject( const UObject * Object );
-	bool			IsNetGUIDAuthority();
+	bool			IsNetGUIDAuthority() const;
 	FNetworkGUID	GetOrAssignNetGUID( const UObject * Object );
 	FNetworkGUID	AssignNewNetGUID_Server( const UObject * Object );
 	void			RegisterNetGUID_Internal( const FNetworkGUID & NetGUID, const FNetGuidCacheObject & CacheObject );
@@ -62,9 +62,10 @@ public:
 	void			RegisterNetGUID_Client( const FNetworkGUID & NetGUID, const UObject * Object );
 	void			RegisterNetGUIDFromPath_Client( const FNetworkGUID & NetGUID, const FString & PathName, const FNetworkGUID & OuterGUID, const bool bNoLoad, const bool bIgnoreWhenMissing );
 	UObject *		GetObjectFromNetGUID( const FNetworkGUID & NetGUID, const bool bIgnoreMustBeMapped );
-	bool			ShouldIgnoreWhenMissing( const FNetworkGUID & NetGUID );
-	bool			IsGUIDRegistered( const FNetworkGUID & NetGUID );
-	bool			IsGUIDLoaded( const FNetworkGUID & NetGUID );
+	bool			ShouldIgnoreWhenMissing( const FNetworkGUID & NetGUID ) const;
+	bool			IsGUIDRegistered( const FNetworkGUID & NetGUID ) const;
+	bool			IsGUIDLoaded( const FNetworkGUID & NetGUID ) const;
+	bool			IsGUIDBroken( const FNetworkGUID & NetGUID, const bool bMustBeRegistered ) const;
 
 	void			AsyncPackageCallback( const FString & PackageName, UPackage * Package );
 	
@@ -138,6 +139,7 @@ class UPackageMapClient : public UPackageMap
 
 	virtual void		LogDebugInfo( FOutputDevice & Ar) override;
 	virtual UObject *	GetObjectFromNetGUID( const FNetworkGUID & NetGUID, const bool bIgnoreMustBeMapped ) override;
+	virtual bool		IsGUIDBroken( const FNetworkGUID & NetGUID, const bool bMustBeRegistered ) const override { return GuidCache->IsGUIDBroken( NetGUID, bMustBeRegistered ); }
 
 	TArray< FNetworkGUID > & GetMustBeMappedGuidsInLastBunch() { return MustBeMappedGuidsInLastBunch; }
 
@@ -153,7 +155,7 @@ protected:
 
 	bool	ShouldSendFullPath(const UObject* Object, const FNetworkGUID &NetGUID);
 	
-	bool IsNetGUIDAuthority();
+	bool IsNetGUIDAuthority() const;
 
 	class UNetConnection* Connection;
 
diff --git a/Engine/Source/Runtime/Engine/Private/DataChannel.cpp b/Engine/Source/Runtime/Engine/Private/DataChannel.cpp
--- a/Engine/Source/Runtime/Engine/Private/DataChannel.cpp
+++ b/Engine/Source/Runtime/Engine/Private/DataChannel.cpp
@@ -1483,6 +1483,15 @@ void UActorChannel::Tick()
 		{
 			// This guid is now resolved, we can remove it from the pending guid list
 			It.RemoveCurrent();
+			continue;
+		}
+
+		if ( Connection->Driver->GuidCache->IsGUIDBroken( *It, true ) )
+		{
+			// This guid is broken, remove it, and warn
+			UE_LOG( LogNet, Warning, TEXT( "UActorChannel::Tick: Guid is broken. NetGUID: %s, ChIndex: %i, Actor: %s" ), *It->ToString(), ChIndex, Actor != NULL ? *Actor->GetPathName() : TEXT( "NULL" ) );
+			It.RemoveCurrent();
+			continue;
 		}
 	}
 
@@ -1495,7 +1504,7 @@ void UActorChannel::Tick()
 			delete QueuedBunches[i];
 		}
 
-		UE_LOG( LogNet, Log, TEXT( "UActorChannel::Tick: Flushing queued bunches: ChIndex: %i, Actor: %s, Queued: %i" ), ChIndex, Actor != NULL ? *Actor->GetPathName() : TEXT( "NULL" ), QueuedBunches.Num() );
+		UE_LOG( LogNet, Log, TEXT( "UActorChannel::Tick: Flushing queued bunches. ChIndex: %i, Actor: %s, Queued: %i" ), ChIndex, Actor != NULL ? *Actor->GetPathName() : TEXT( "NULL" ), QueuedBunches.Num() );
 
 		QueuedBunches.Empty();
 	}
diff --git a/Engine/Source/Runtime/Engine/Private/PackageMapClient.cpp b/Engine/Source/Runtime/Engine/Private/PackageMapClient.cpp
--- a/Engine/Source/Runtime/Engine/Private/PackageMapClient.cpp
+++ b/Engine/Source/Runtime/Engine/Private/PackageMapClient.cpp
@@ -1215,7 +1215,7 @@ bool UPackageMapClient::ObjectLevelHasFinishedLoading(UObject* Object)
  *  This is ugly but probably better than adding a shadow variable that has to be
  *  set/cleared at the net driver level.
  */
-bool UPackageMapClient::IsNetGUIDAuthority()
+bool UPackageMapClient::IsNetGUIDAuthority() const
 {
 	return GuidCache->IsNetGUIDAuthority();
 }
@@ -1444,7 +1444,7 @@ bool FNetGUIDCache::IsDynamicObject( const UObject * Object )
 	return !Object->IsFullNameStableForNetworking();
 }
 
-bool FNetGUIDCache::IsNetGUIDAuthority()
+bool FNetGUIDCache::IsNetGUIDAuthority() const
 {
 	return Driver == NULL || Driver->IsServer();
 }
@@ -2024,7 +2024,7 @@ UObject * FNetGUIDCache::GetObjectFromNetGUID( const FNetworkGUID & NetGUID, con
 	return Object;
 }
 
-bool FNetGUIDCache::ShouldIgnoreWhenMissing( const FNetworkGUID & NetGUID )
+bool FNetGUIDCache::ShouldIgnoreWhenMissing( const FNetworkGUID & NetGUID ) const
 {
 	if ( IsNetGUIDAuthority() )
 	{
@@ -2054,7 +2054,7 @@ bool FNetGUIDCache::ShouldIgnoreWhenMissing( const FNetworkGUID & NetGUID )
 	return CacheObject->bIgnoreWhenMissing || OutermostCacheObject->bIsPending;
 }
 
-bool FNetGUIDCache::IsGUIDRegistered( const FNetworkGUID & NetGUID )
+bool FNetGUIDCache::IsGUIDRegistered( const FNetworkGUID & NetGUID ) const
 {
 	if ( !NetGUID.IsValid() )
 	{
@@ -2069,7 +2069,7 @@ bool FNetGUIDCache::IsGUIDRegistered( const FNetworkGUID & NetGUID )
 	return ObjectLookup.Contains( NetGUID );
 }
 
-bool FNetGUIDCache::IsGUIDLoaded( const FNetworkGUID & NetGUID )
+bool FNetGUIDCache::IsGUIDLoaded( const FNetworkGUID & NetGUID ) const
 {
 	if ( !NetGUID.IsValid() )
 	{
@@ -2081,7 +2081,7 @@ bool FNetGUIDCache::IsGUIDLoaded( const FNetworkGUID & NetGUID )
 		return false;
 	}
 
-	FNetGuidCacheObject * CacheObjectPtr = ObjectLookup.Find( NetGUID );
+	const FNetGuidCacheObject * CacheObjectPtr = ObjectLookup.Find( NetGUID );
 
 	if ( CacheObjectPtr == NULL )
 	{
@@ -2091,6 +2091,28 @@ bool FNetGUIDCache::IsGUIDLoaded( const FNetworkGUID & NetGUID )
 	return CacheObjectPtr->Object != NULL;
 }
 
+bool FNetGUIDCache::IsGUIDBroken( const FNetworkGUID & NetGUID, const bool bMustBeRegistered ) const
+{
+	if ( !NetGUID.IsValid() )
+	{
+		return false;
+	}
+
+	if ( NetGUID.IsDefault() )
+	{
+		return false;
+	}
+
+	const FNetGuidCacheObject * CacheObjectPtr = ObjectLookup.Find( NetGUID );
+
+	if ( CacheObjectPtr == NULL )
+	{
+		return bMustBeRegistered;
+	}
+
+	return CacheObjectPtr->bIsBroken;
+}
+
 //------------------------------------------------------
 // Debug command to see how many times we've exported each NetGUID
 // Used for measuring inefficiencies. Some duplication is unavoidable since we cannot garuntee atomicicity across multiple channels.
diff --git a/Engine/Source/Runtime/Engine/Private/RepLayout.cpp b/Engine/Source/Runtime/Engine/Private/RepLayout.cpp
--- a/Engine/Source/Runtime/Engine/Private/RepLayout.cpp
+++ b/Engine/Source/Runtime/Engine/Private/RepLayout.cpp
@@ -1329,7 +1329,7 @@ void FRepLayout::UpdateUnmappedObjects_r(
 		if ( AbsOffset >= MaxAbsOffset )
 		{
 			// Array must have shrunk, we can remove this item
-			UE_LOG( LogNet, VeryVerbose, TEXT( "REMOVED unmapped property: AbsOffset >= MaxAbsOffset. Offset: %i" ), AbsOffset );
+			UE_LOG( LogNet, VeryVerbose, TEXT( "UpdateUnmappedObjects_r: REMOVED unmapped property: AbsOffset >= MaxAbsOffset. Offset: %i" ), AbsOffset );
 			It.RemoveCurrent();
 			continue;
 		}
@@ -1351,7 +1351,7 @@ void FRepLayout::UpdateUnmappedObjects_r(
 
 		if ( Object != NULL )
 		{
-			UE_LOG( LogNet, VeryVerbose, TEXT( "REMOVED unmapped property: Offset: %i, Guid: %s, PropName: %s, ObjName: %s" ), AbsOffset, *It.Value().Guid.ToString(), *Cmd.Property->GetName(), *Object->GetName() );
+			UE_LOG( LogNet, VeryVerbose, TEXT( "UpdateUnmappedObjects_r: REMOVED unmapped property: Offset: %i, Guid: %s, PropName: %s, ObjName: %s" ), AbsOffset, *It.Value().Guid.ToString(), *Cmd.Property->GetName(), *Object->GetName() );
 			UObjectPropertyBase * ObjProperty = CastChecked< UObjectPropertyBase>( Cmd.Property );
 
 			UObject * OldObject = ObjProperty->GetObjectPropertyValue( Data + AbsOffset );
@@ -1378,7 +1378,14 @@ void FRepLayout::UpdateUnmappedObjects_r(
 			continue;
 		}
 
-		UE_LOG( LogNet, VeryVerbose, TEXT( "UnmappedGuids: Offset: %i, Guid: %s" ), AbsOffset, *It.Value().Guid.ToString() );
+		if ( PackageMap->IsGUIDBroken( It.Value().Guid, false ) )
+		{
+			UE_LOG( LogNet, Warning, TEXT( "UpdateUnmappedObjects_r: Broken GUID. NetGuid: %s" ), *It.Value().Guid.ToString() );
+			It.RemoveCurrent();
+			continue;
+		}
+
+		UE_LOG( LogNet, VeryVerbose, TEXT( "UpdateUnmappedObjects_r: Still unmapped. Offset: %i, NetGuid: %s" ), AbsOffset, *It.Value().Guid.ToString() );
 
 		bOutHasMoreUnmapped = true;
 	}
