Property naming fix in Level Streaming Volume

diff --git a/Engine/Source/Runtime/Engine/Classes/Engine/LevelStreamingVolume.h b/Engine/Source/Runtime/Engine/Classes/Engine/LevelStreamingVolume.h
--- a/Engine/Source/Runtime/Engine/Classes/Engine/LevelStreamingVolume.h
+++ b/Engine/Source/Runtime/Engine/Classes/Engine/LevelStreamingVolume.h
@@ -28,7 +28,7 @@ class ALevelStreamingVolume : public AVolume
 	GENERATED_UCLASS_BODY()
 
 	/** Levels names affected by this level streaming volume. */
-	UPROPERTY(Category=LevelStreamingVolume, VisibleAnywhere, BlueprintReadOnly)
+	UPROPERTY(Category=LevelStreamingVolume, VisibleAnywhere, BlueprintReadOnly, meta=(DisplayName = "Streaming Levels"))
 	TArray<FName> StreamingLevelNames;
 
 	/** If true, this streaming volume should only be used for editor streaming level previs. */
