Fixed crash when exporting landscape into raw mesh

diff --git a/Engine/Source/Runtime/Engine/Private/Landscape/LandscapeEdit.cpp b/Engine/Source/Runtime/Engine/Private/Landscape/LandscapeEdit.cpp
--- a/Engine/Source/Runtime/Engine/Private/Landscape/LandscapeEdit.cpp
+++ b/Engine/Source/Runtime/Engine/Private/Landscape/LandscapeEdit.cpp
@@ -2310,7 +2310,13 @@ bool ALandscapeProxy::ExportToRawMesh(int32 InExportLOD, FRawMesh& OutRawMesh) c
 				{
 					int32 SubNumX, SubNumY, SubX, SubY;
 					CDI.ComponentXYToSubsectionXY(x, y, SubNumX, SubNumY, SubX, SubY);
-					bool bInvisible = VisDataMap[x + y * (ComponentSizeQuadsLOD + 1)] >= VisThreshold;
+
+					// Whether this vertex is in hole
+					bool bInvisible = false;
+					if (VisDataMap.Num()) 
+					{
+						bInvisible = (VisDataMap[x + y * (ComponentSizeQuadsLOD + 1)] >= VisThreshold);
+					}
 
 					// triangulation matches FLandscapeIndexBuffer constructor
 					Faces[0] = VertexIdx;
