Landscape exports normals in local space

diff --git a/Engine/Source/Runtime/Engine/Private/Landscape/LandscapeEdit.cpp b/Engine/Source/Runtime/Engine/Private/Landscape/LandscapeEdit.cpp
--- a/Engine/Source/Runtime/Engine/Private/Landscape/LandscapeEdit.cpp
+++ b/Engine/Source/Runtime/Engine/Private/Landscape/LandscapeEdit.cpp
@@ -2338,7 +2338,7 @@ bool ALandscapeProxy::ExportToRawMesh(int32 InExportLOD, FRawMesh& OutRawMesh) c
 					FVector LocalVertexPos = CDI.GetWorldVertex(VertexX, VertexY);
 
 					FVector LocalTangentX, LocalTangentY, LocalTangentZ;
-					CDI.GetWorldTangentVectors(VertexX, VertexY, LocalTangentX, LocalTangentY, LocalTangentZ);
+					CDI.GetLocalTangentVectors(VertexX, VertexY, LocalTangentX, LocalTangentY, LocalTangentZ);
 
 					OutRawMesh.VertexPositions[VertexIdx] = LocalVertexPos;
 					OutRawMesh.WedgeTangentX[VertexIdx] = LocalTangentX;
