ttp335480 LIVE: EDITOR: MIGRATE: CRASH: Migrating "Character" folder from StrategyGame into a new blank project

diff --git a/Engine/Source/Runtime/Engine/Private/AnimBlueprint.cpp b/Engine/Source/Runtime/Engine/Private/AnimBlueprint.cpp
--- a/Engine/Source/Runtime/Engine/Private/AnimBlueprint.cpp
+++ b/Engine/Source/Runtime/Engine/Private/AnimBlueprint.cpp
@@ -66,7 +66,7 @@ UAnimBlueprint* UAnimBlueprint::FindRootAnimBlueprint(UAnimBlueprint* DerivedBlu
 	UAnimBlueprint* ParentBP = NULL;
 
 	// Determine if there is an anim blueprint in the ancestry of this class
-	for (UClass* ParentClass = DerivedBlueprint->ParentClass; ParentClass != UObject::StaticClass(); ParentClass = ParentClass->GetSuperClass())
+	for (UClass* ParentClass = DerivedBlueprint->ParentClass; ParentClass && (UObject::StaticClass() != ParentClass); ParentClass = ParentClass->GetSuperClass())
 	{
 		if (UAnimBlueprint* TestBP = Cast<UAnimBlueprint>(ParentClass->ClassGeneratedBy))
 		{
