Fix for compiler errors when brute forcing STAT to be true in Shipping builds

diff --git a/Engine/Source/Runtime/Core/Public/HAL/MallocBinned.h b/Engine/Source/Runtime/Core/Public/HAL/MallocBinned.h
--- a/Engine/Source/Runtime/Core/Public/HAL/MallocBinned.h
+++ b/Engine/Source/Runtime/Core/Public/HAL/MallocBinned.h
@@ -1183,7 +1183,7 @@ public:
 			ValidateHeap();
 #if STATS
 			UpdateSlackStat();
-
+#if !NO_LOGGING
 			// This is all of the memory including stuff too big for the pools
 			BufferedOutput.CategorizedLogf( LogMemory.GetCategoryName(), ELogVerbosity::Log, TEXT( "Allocator Stats for %s:" ), GetDescriptiveName() );
 			// Waste is the total overhead of the memory system
@@ -1260,6 +1260,7 @@ public:
 			BufferedOutput.CategorizedLogf( LogMemory.GetCategoryName(), ELogVerbosity::Log, TEXT( "%iK allocated in pools (with %iK slack and %iK waste). Efficiency %.2f%%" ), TotalMemory, TotalSlack, TotalWaste, TotalMemory ? 100.0f * (TotalMemory - TotalWaste) / TotalMemory : 100.0f );
 			BufferedOutput.CategorizedLogf( LogMemory.GetCategoryName(), ELogVerbosity::Log, TEXT( "Allocations %i Current / %i Total (in %i pools)" ), TotalActiveRequests, TotalTotalRequests, TotalPools );
 			BufferedOutput.CategorizedLogf( LogMemory.GetCategoryName(), ELogVerbosity::Log, TEXT( "" ) );
+#endif
 #endif
 		}
 
diff --git a/Engine/Source/Runtime/Engine/Private/Particles/ParticleSystemRender.cpp b/Engine/Source/Runtime/Engine/Private/Particles/ParticleSystemRender.cpp
--- a/Engine/Source/Runtime/Engine/Private/Particles/ParticleSystemRender.cpp
+++ b/Engine/Source/Runtime/Engine/Private/Particles/ParticleSystemRender.cpp
@@ -87,6 +87,7 @@ void DumpParticleRenderingStats(FOutputDevice& Ar)
 #if STATS
 	if (GTrackParticleRenderingStats)
 	{
+#if ALLOW_DEBUG_FILES
 		{
 			// Have to keep this filename short enough to be valid for Xbox 360
 			const FString TemplateFileName = FString::Printf(TEXT("%sPT-%s.csv"), *FPaths::ProfilingDir(), *FDateTime::Now().ToString());
@@ -137,7 +138,7 @@ void DumpParticleRenderingStats(FOutputDevice& Ar)
 			}
 			Ar.Logf(TEXT("Component stats saved to %s"), *ComponentFileName);
 		}
-
+#endif
 		GTemplateRenderStats.Empty();
 		GComponentRenderStats.Empty();
 	}
@@ -190,6 +191,7 @@ void ENGINE_API FinishOneFrameParticleStats()
 
 		UE_LOG(LogParticles, Warning, TEXT("IFileManager::MakeDirectory %s %u"), *WritePath, bDirSuccess);
 
+#if ALLOW_DEBUG_FILES
 		// Have to keep this filename short enough to be valid for Xbox 360
 		const FString TemplateFileName = WritePath / TEXT("ParticleTemplates.csv");
 		FDiagnosticTableViewer ParticleTemplateViewer(*TemplateFileName);
@@ -240,6 +242,7 @@ void ENGINE_API FinishOneFrameParticleStats()
 		ParticleTemplateViewer.CycleRow();
 
 		UE_LOG(LogParticles, Warning, TEXT("One frame stats saved to %s"), *TemplateFileName);
+#endif
 
 		// Request a screenshot to be saved next frame
 		FScreenshotRequest::RequestScreenshot( PathFromScreenshots / TEXT("Shot.bmp"), false );
