Added bone driven controller node for animation blueprint graphs.

diff --git a/Engine/Source/Editor/AnimGraph/Classes/AnimGraphNode_BoneDrivenController.h b/Engine/Source/Editor/AnimGraph/Classes/AnimGraphNode_BoneDrivenController.h
new file mode 100644
--- /dev/null
+++ b/Engine/Source/Editor/AnimGraph/Classes/AnimGraphNode_BoneDrivenController.h
@@ -0,0 +1,34 @@
+// Copyright 1998-2014 Epic Games, Inc. All Rights Reserved.
+
+#pragma once
+
+#include "AnimGraphNode_SkeletalControlBase.h"
+#include "Animation/BoneControllers/AnimNode_BoneDrivenController.h"
+#include "AnimGraphNode_BoneDrivenController.generated.h"
+
+UCLASS(MinimalAPI)
+class UAnimGraphNode_BoneDrivenController : public UAnimGraphNode_SkeletalControlBase
+{
+	GENERATED_UCLASS_BODY()
+
+	UPROPERTY(EditAnywhere, Category = Settings)
+	FAnimNode_BoneDrivenController Node;
+
+public:
+
+	// UEdGraphNode interface
+	virtual FText GetNodeTitle(ENodeTitleType::Type TitleType) const override;
+	//virtual FString GetNodeNativeTitle(ENodeTitleType::Type TitleType) const override;
+	virtual FString GetTooltip() const override;
+	//////////////////////////////////////////////////////////////////////////
+
+	// UAnimGraphNode_SkeletalControlBase interface
+	ANIMGRAPH_API virtual void Draw(FPrimitiveDrawInterface* PDI, USkeletalMeshComponent* SkelMeshComp) const override;
+	//////////////////////////////////////////////////////////////////////////
+
+protected:
+
+	// UAnimGraphNode_SkeletalControlBase interface
+	virtual FText GetControllerDescription() const override;
+	//////////////////////////////////////////////////////////////////////////
+};
\ No newline at end of file
diff --git a/Engine/Source/Editor/AnimGraph/Private/AnimGraphNode_BoneDrivenController.cpp b/Engine/Source/Editor/AnimGraph/Private/AnimGraphNode_BoneDrivenController.cpp
new file mode 100644
--- /dev/null
+++ b/Engine/Source/Editor/AnimGraph/Private/AnimGraphNode_BoneDrivenController.cpp
@@ -0,0 +1,76 @@
+// Copyright 1998-2014 Epic Games, Inc. All Rights Reserved.
+
+#include "AnimGraphPrivatePCH.h"
+#include "AnimGraphNode_BoneDrivenController.h"
+
+#define LOCTEXT_NAMESPACE "A3Nodes"
+
+UAnimGraphNode_BoneDrivenController::UAnimGraphNode_BoneDrivenController(const FPostConstructInitializeProperties& PCIP)
+	:Super(PCIP)
+{
+
+}
+
+FString UAnimGraphNode_BoneDrivenController::GetTooltip() const
+{
+	return LOCTEXT("UAnimGraphNode_BoneDrivenController_ToolTip", "Drives the transform of a bone using the transform of another").ToString();
+}
+
+FText UAnimGraphNode_BoneDrivenController::GetNodeTitle(ENodeTitleType::Type TitleType) const
+{
+	FFormatNamedArguments Args;
+	Args.Add(TEXT("ControllerDesc"), GetControllerDescription());
+
+	if(TitleType == ENodeTitleType::ListView)
+	{
+		Args.Add(TEXT("Delim"), FText::FromString(TEXT(" - ")));
+	}
+	else
+	{
+		Args.Add(TEXT("Delim"), FText::FromString(TEXT("\n")));
+	}
+
+	Args.Add(TEXT("SourceBone"), FText::FromName(Node.SourceBone.BoneName));
+	Args.Add(TEXT("TargetBone"), FText::FromName(Node.TargetBone.BoneName));
+
+	return FText::Format(LOCTEXT("AnimGraphNode_BoneDrivenController_Title", "{ControllerDesc}{Delim}Driving Bone: {SourceBone}{Delim}Driven Bone: {TargetBone}"), Args);
+}
+
+FText UAnimGraphNode_BoneDrivenController::GetControllerDescription() const
+{
+	return LOCTEXT("BoneDrivenController", "Bone Driven Controller");
+}
+
+void UAnimGraphNode_BoneDrivenController::Draw(FPrimitiveDrawInterface* PDI, USkeletalMeshComponent* SkelMeshComp) const
+{	
+	static const float ArrowHeadWidth = 5.0f;
+	static const float ArrowHeadHeight = 8.0f;
+
+	int32 SourceIdx = SkelMeshComp->GetBoneIndex(Node.SourceBone.BoneName);
+	int32 TargetIdx = SkelMeshComp->GetBoneIndex(Node.TargetBone.BoneName);
+
+	if(SourceIdx != INDEX_NONE && TargetIdx != INDEX_NONE)
+	{
+		FTransform SourceTM = SkelMeshComp->SpaceBases[SourceIdx] * SkelMeshComp->ComponentToWorld;
+		FTransform TargetTM = SkelMeshComp->SpaceBases[TargetIdx] * SkelMeshComp->ComponentToWorld;
+
+		PDI->DrawLine(TargetTM.GetLocation(), SourceTM.GetLocation(), FLinearColor(0.0f, 0.0f, 1.0f), SDPG_Foreground, 0.5f);
+
+		FVector ToTarget = TargetTM.GetTranslation() - SourceTM.GetTranslation();
+		FVector UnitToTarget = ToTarget;
+		UnitToTarget.Normalize();
+		FVector Midpoint = SourceTM.GetTranslation() + 0.5f * ToTarget + 0.5f * UnitToTarget * ArrowHeadHeight;
+
+		FVector YAxis;
+		FVector ZAxis;
+		UnitToTarget.FindBestAxisVectors(YAxis, ZAxis);
+		FMatrix ArrowMatrix(UnitToTarget, YAxis, ZAxis, Midpoint);
+
+		DrawConnectedArrow(PDI, ArrowMatrix, FLinearColor(0.0f, 1.0f, 0.0), ArrowHeadHeight, ArrowHeadWidth, SDPG_Foreground);
+
+		PDI->DrawPoint(SourceTM.GetTranslation(), FLinearColor(0.8f, 0.8f, 0.2f), 5.0f, SDPG_Foreground);
+		PDI->DrawPoint(SourceTM.GetTranslation() + ToTarget, FLinearColor(0.8f, 0.8f, 0.2f), 5.0f, SDPG_Foreground);
+	}
+}
+
+#undef LOCTEXT_NAMESPACE
diff --git a/Engine/Source/Runtime/Engine/Classes/Animation/BoneControllers/AnimNode_BoneDrivenController.h b/Engine/Source/Runtime/Engine/Classes/Animation/BoneControllers/AnimNode_BoneDrivenController.h
new file mode 100644
--- /dev/null
+++ b/Engine/Source/Runtime/Engine/Classes/Animation/BoneControllers/AnimNode_BoneDrivenController.h
@@ -0,0 +1,79 @@
+// Copyright 1998-2014 Epic Games, Inc. All Rights Reserved.
+
+#pragma once
+
+#include "AnimNode_SkeletalControlBase.h"
+#include "AnimNode_BoneDrivenController.generated.h"
+
+// Evaluation of the bone transforms relies on the size and ordering of this
+// enum, if this needs to change make sure EvaluateBoneTransforms is updated.
+UENUM()
+namespace EComponentType
+{
+	enum Type
+	{
+		None = 0,
+		TranslationX,
+		TranslationY,
+		TranslationZ,
+		RotationX,
+		RotationY,
+		RotationZ,
+		Scale
+	};
+}
+
+USTRUCT()
+struct ENGINE_API FAnimNode_BoneDrivenController : public FAnimNode_SkeletalControlBase
+{
+	GENERATED_USTRUCT_BODY()
+
+	// Bone to use as controller input
+	UPROPERTY(EditAnywhere, BlueprintReadWrite, Category = Bones)
+	FBoneReference SourceBone;
+
+	// Bone to drive using controller input
+	UPROPERTY(EditAnywhere, BlueprintReadWrite, Category = Bones)
+	FBoneReference TargetBone;
+
+	// Transform component to use as input
+	UPROPERTY(EditAnywhere, BlueprintReadWrite, Category = Controller)
+	TEnumAsByte<EComponentType::Type> SourceComponent;
+
+	// Transform component to write to
+	UPROPERTY(EditAnywhere, BlueprintReadWrite, Category = Controller)
+	TEnumAsByte<EComponentType::Type> TargetComponent;
+
+	// Multiplier to apply to the input value
+	UPROPERTY(EditAnywhere, BlueprintReadWrite, Category=Controller)
+	float Multiplier;
+
+	// Whether or not to use the range limits
+	UPROPERTY(EditAnywhere, BlueprintReadWrite, Category=Range)
+	bool bUseRange;
+
+	// Min limit of the output value
+	UPROPERTY(EditAnywhere, BlueprintReadWrite, Category=Range)
+	float RangeMin;
+
+	// Max limit of the output value
+	UPROPERTY(EditAnywhere, BlueprintReadWrite, Category=Range)
+	float RangeMax;
+
+	FAnimNode_BoneDrivenController();
+
+	// FAnimNode_Base interface
+	virtual void GatherDebugData(FNodeDebugData& DebugData) override;
+	//////////////////////////////////////////////////////////////////////////
+
+	// FAnimNode_SkeletalControlBase interface
+	virtual void EvaluateBoneTransforms(USkeletalMeshComponent* SkelComp, const FBoneContainer & RequiredBones, FA2CSPose& MeshBases, TArray<FBoneTransform>& OutBoneTransforms) override;
+	virtual bool IsValidToEvaluate(const USkeleton * Skeleton, const FBoneContainer & RequiredBones) override;
+	//////////////////////////////////////////////////////////////////////////
+
+protected:
+	
+	// FAnimNode_SkeletalControlBase interface
+	virtual void InitializeBoneReferences(const FBoneContainer & RequiredBones) override;
+	//////////////////////////////////////////////////////////////////////////
+};
\ No newline at end of file
diff --git a/Engine/Source/Runtime/Engine/Private/Animation/AnimNode_BoneDrivenController.cpp b/Engine/Source/Runtime/Engine/Private/Animation/AnimNode_BoneDrivenController.cpp
new file mode 100644
--- /dev/null
+++ b/Engine/Source/Runtime/Engine/Private/Animation/AnimNode_BoneDrivenController.cpp
@@ -0,0 +1,127 @@
+// Copyright 1998-2014 Epic Games, Inc. All Rights Reserved.
+
+#include "EnginePrivate.h"
+#include "Animation/BoneControllers/AnimNode_BoneDrivenController.h"
+
+FAnimNode_BoneDrivenController::FAnimNode_BoneDrivenController()
+: SourceComponent(EComponentType::None)
+, TargetComponent(EComponentType::None)
+, Multiplier(1.0f)
+, bUseRange(false)
+, RangeMin(-1.0f)
+, RangeMax(1.0f)
+{
+
+}
+
+void FAnimNode_BoneDrivenController::GatherDebugData(FNodeDebugData& DebugData)
+{
+	FString DebugLine = DebugData.GetNodeName(this);
+	DebugLine += "(";
+	AddDebugNodeData(DebugLine);
+	DebugLine += FString::Printf(TEXT("  DrivingBone: %s\nDrivenBone: %s"), *SourceBone.BoneName.ToString(), *TargetBone.BoneName.ToString());
+	DebugData.AddDebugItem(DebugLine);
+
+	ComponentPose.GatherDebugData(DebugData);
+}
+
+void FAnimNode_BoneDrivenController::EvaluateBoneTransforms(USkeletalMeshComponent* SkelComp, const FBoneContainer & RequiredBones, FA2CSPose& MeshBases, TArray<FBoneTransform>& OutBoneTransforms)
+{
+	check(OutBoneTransforms.Num() == 0);
+	
+	// Early out if we're not driving from or to anything
+	if(SourceComponent == EComponentType::None || TargetComponent == EComponentType::None || Multiplier == 0.0f)
+	{
+		return;
+	}
+
+	// Get the Local space transform and the ref pose transform to see how the transform for the source bone has changed
+	const TArray<FTransform>& RefPose = RequiredBones.GetRefPoseArray();
+	FTransform SourceOrigRef = RefPose[SourceBone.BoneIndex];
+	FTransform SourceCurr = MeshBases.GetLocalSpaceTransform(SourceBone.BoneIndex);
+
+	// Difference from the ref pose
+	FVector RotationDiff = (SourceCurr.GetRotation() * SourceOrigRef.GetRotation().Inverse()).Euler();
+	FVector TranslationDiff = SourceCurr.GetLocation() - SourceOrigRef.GetLocation();
+	FVector CurrentScale = SourceCurr.GetScale3D();
+	FVector RefScale = SourceOrigRef.GetScale3D();
+	float ScaleDiff = FMath::Max3(CurrentScale[0], CurrentScale[1], CurrentScale[2]) - FMath::Max3(RefScale[0], RefScale[1], RefScale[2]);
+
+	// Starting point for the new transform
+	FTransform NewLocal = MeshBases.GetLocalSpaceTransform(TargetBone.BoneIndex);
+	
+	// Difference to apply after processing
+	FVector NewRotDiff(FVector::ZeroVector);
+	FVector NewTransDiff(FVector::ZeroVector);
+	float NewScaleDiff = 0.0f;
+
+	float* SourcePtr = nullptr;
+	float* DestPtr = nullptr;
+
+	// Resolve source value
+	if(SourceComponent < EComponentType::RotationX)
+	{
+		SourcePtr = &TranslationDiff[(int32)(SourceComponent - EComponentType::TranslationX)];
+	}
+	else if(SourceComponent < EComponentType::Scale)
+	{
+		SourcePtr = &RotationDiff[(int32)(SourceComponent - EComponentType::RotationX)];
+	}
+	else
+	{
+		SourcePtr = &ScaleDiff;
+	}
+	
+	// Resolve target value
+	if(TargetComponent < EComponentType::RotationX)
+	{
+		DestPtr = &NewTransDiff[(int32)(TargetComponent - EComponentType::TranslationX)];
+	}
+	else if(TargetComponent < EComponentType::Scale)
+	{
+		DestPtr = &NewRotDiff[(int32)(TargetComponent - EComponentType::RotationX)];
+	}
+	else
+	{
+		DestPtr = &NewScaleDiff;
+	}
+
+	// Set the value; clamped to the given range if necessary.
+	if(bUseRange)
+	{
+		*DestPtr = FMath::Clamp(*SourcePtr * Multiplier, RangeMin, RangeMax);
+	}
+	else
+	{
+		*DestPtr = *SourcePtr * Multiplier;
+	}
+
+	// Build final transform difference
+	FTransform FinalDiff(FQuat::MakeFromEuler(NewRotDiff), NewTransDiff, FVector(NewScaleDiff));
+	NewLocal.AccumulateWithAdditiveScale(FinalDiff);
+
+	// If we have a parent, concatenate the transform, otherwise just take the new transform
+	int32 PIdx = RequiredBones.GetParentBoneIndex(TargetBone.BoneIndex);
+	if(PIdx != INDEX_NONE)
+	{
+		FTransform ParentTM = MeshBases.GetComponentSpaceTransform(PIdx);
+		
+		OutBoneTransforms.Add(FBoneTransform(TargetBone.BoneIndex, NewLocal * ParentTM));
+	}
+	else
+	{
+		OutBoneTransforms.Add(FBoneTransform(TargetBone.BoneIndex, NewLocal));
+	}
+
+}
+
+bool FAnimNode_BoneDrivenController::IsValidToEvaluate(const USkeleton * Skeleton, const FBoneContainer & RequiredBones)
+{
+	return SourceBone.IsValid(RequiredBones) && TargetBone.IsValid(RequiredBones);
+}
+
+void FAnimNode_BoneDrivenController::InitializeBoneReferences(const FBoneContainer & RequiredBones)
+{
+	SourceBone.Initialize(RequiredBones);
+	TargetBone.Initialize(RequiredBones);
+}
diff --git a/Engine/Source/Runtime/Engine/Private/PrimitiveDrawingUtils.cpp b/Engine/Source/Runtime/Engine/Private/PrimitiveDrawingUtils.cpp
--- a/Engine/Source/Runtime/Engine/Private/PrimitiveDrawingUtils.cpp
+++ b/Engine/Source/Runtime/Engine/Private/PrimitiveDrawingUtils.cpp
@@ -1071,6 +1071,34 @@ void DrawDirectionalArrow(FPrimitiveDrawInterface* PDI,const FMatrix& ArrowToWor
 	PDI->DrawLine(ArrowToWorld.TransformPosition(FVector(Length,0,0)),ArrowToWorld.TransformPosition(FVector(Length-ArrowSize,-ArrowSize,-ArrowSize)),InColor,DepthPriority);
 }
 
+/**
+ * Draws a directional arrow with connected spokes.
+ *
+ * @param	PDI				Draw interface.
+ * @param	ArrowToWorld	Transform matrix for the arrow.
+ * @param	Color			Color of the arrow.
+ * @param	ArrowHeight		Height of the the arrow head.
+ * @param	ArrowWidth		Width of the arrow head.
+ * @param	DepthPriority	Depth priority for the arrow.
+ * @param	Thickness		Thickness of the lines used to draw the arrow.
+ * @param	NumSpokes		Number of spokes used to make the arrow head.
+ */
+void DrawConnectedArrow(class FPrimitiveDrawInterface* PDI, const FMatrix& ArrowToWorld, const FLinearColor& Color, float ArrowHeight, float ArrowWidth, uint8 DepthPriority, float Thickness, int32 NumSpokes)
+{
+	float RotPerSpoke = (2.0f * PI) / (float)NumSpokes;
+	FQuat Rotator(FVector(1.0f, 0.0f, 0.0f), RotPerSpoke);
+
+	FVector Origin = ArrowToWorld.GetOrigin();
+	FVector SpokePoint = FVector(-ArrowHeight, ArrowWidth, 0);
+	for(int32 CurrSpoke = 0 ; CurrSpoke < NumSpokes ; ++CurrSpoke)
+	{
+		PDI->DrawLine(Origin, ArrowToWorld.TransformPosition(SpokePoint), Color, DepthPriority, Thickness);
+		FVector PrevPoint = SpokePoint;
+		SpokePoint = Rotator.RotateVector(SpokePoint);
+		PDI->DrawLine(ArrowToWorld.TransformPosition(PrevPoint), ArrowToWorld.TransformPosition(SpokePoint), Color, DepthPriority, Thickness);
+	}
+}
+
 /**
  * Draws a axis-aligned 3 line star.
  *
diff --git a/Engine/Source/Runtime/Engine/Public/SceneManagement.h b/Engine/Source/Runtime/Engine/Public/SceneManagement.h
--- a/Engine/Source/Runtime/Engine/Public/SceneManagement.h
+++ b/Engine/Source/Runtime/Engine/Public/SceneManagement.h
@@ -1380,6 +1380,7 @@ extern ENGINE_API void DrawWireCone(class FPrimitiveDrawInterface* PDI, const FT
 extern ENGINE_API void DrawWireSphereCappedCone(FPrimitiveDrawInterface* PDI, const FTransform& Transform, float ConeRadius, float ConeAngle, int32 ConeSides, int32 ArcFrequency, int32 CapSegments, const FLinearColor& Color, uint8 DepthPriority);
 extern ENGINE_API void DrawOrientedWireBox(class FPrimitiveDrawInterface* PDI,const FVector& Base,const FVector& X,const FVector& Y,const FVector& Z, FVector Extent, const FLinearColor& Color,uint8 DepthPriority);
 extern ENGINE_API void DrawDirectionalArrow(class FPrimitiveDrawInterface* PDI,const FMatrix& ArrowToWorld,const FLinearColor& InColor,float Length,float ArrowSize,uint8 DepthPriority);
+extern ENGINE_API void DrawConnectedArrow(class FPrimitiveDrawInterface* PDI, const FMatrix& ArrowToWorld, const FLinearColor& Color, float ArrowHeight, float ArrowWidth, uint8 DepthPriority, float Thickness = 0.5f, int32 NumSpokes = 6);
 extern ENGINE_API void DrawWireStar(class FPrimitiveDrawInterface* PDI,const FVector& Position, float Size, const FLinearColor& Color,uint8 DepthPriority);
 extern ENGINE_API void DrawDashedLine(class FPrimitiveDrawInterface* PDI, const FVector& Start, const FVector& End, const FLinearColor& Color, float DashSize, uint8 DepthPriority, float DepthBias = 0.0f);
 extern ENGINE_API void DrawWireDiamond(class FPrimitiveDrawInterface* PDI,const FMatrix& DiamondMatrix, float Size, const FLinearColor& InColor,uint8 DepthPriority);
