Corrected typo in DisplayName.

diff --git a/Engine/Source/Runtime/Engine/Classes/Particles/Location/ParticleModuleLocationSkelVertSurface.h b/Engine/Source/Runtime/Engine/Classes/Particles/Location/ParticleModuleLocationSkelVertSurface.h
--- a/Engine/Source/Runtime/Engine/Classes/Particles/Location/ParticleModuleLocationSkelVertSurface.h
+++ b/Engine/Source/Runtime/Engine/Classes/Particles/Location/ParticleModuleLocationSkelVertSurface.h
@@ -8,7 +8,7 @@
 UENUM()
 enum ELocationSkelVertSurfaceSource
 {
-	VERTSURFACESOURCE_Vert UMETA(DisplayName="Verticies"),
+	VERTSURFACESOURCE_Vert UMETA(DisplayName="Vertices"),
 	VERTSURFACESOURCE_Surface UMETA(DisplayName="Surfaces"),
 	VERTSURFACESOURCE_MAX,
 };
