Allow Skel Vert / Surf Loc module to work for arbitrary actors with skel mesh components.

#ttp 320881 - UE4: RENDERING: Fix particles spawning on skeletal mesh to not be restricted to skeletal mesh actors and pawns

diff --git a/Engine/Source/Runtime/Engine/Private/Particles/ParticleModules_Location.cpp b/Engine/Source/Runtime/Engine/Private/Particles/ParticleModules_Location.cpp
--- a/Engine/Source/Runtime/Engine/Private/Particles/ParticleModules_Location.cpp
+++ b/Engine/Source/Runtime/Engine/Private/Particles/ParticleModules_Location.cpp
@@ -2361,7 +2361,7 @@ void UParticleModuleLocationSkelVertSurface::UpdateBoneIndicesList(FParticleEmit
 				for (int32 FindBoneIdx = 0; FindBoneIdx < ValidAssociatedBones.Num(); FindBoneIdx++)
 				{
 					const int32 BoneIdx = SkeletalMeshActor->SkeletalMeshComponent->SkeletalMesh->RefSkeleton.FindBoneIndex(ValidAssociatedBones[FindBoneIdx]);
-					if (BoneIdx != INDEX_NONE)
+					if (BoneIdx != INDEX_NONE && ValidAssociatedBones.Num() > InsertionIndex)
 					{
 						InstancePayload->ValidAssociatedBoneIndices[InsertionIndex++] = BoneIdx;
 					}
@@ -2370,14 +2370,13 @@ void UParticleModuleLocationSkelVertSurface::UpdateBoneIndicesList(FParticleEmit
 				InstancePayload->NumValidAssociatedBoneIndices = InsertionIndex;
 			}
 		}
-		// if we have a pawn
+		// If we have an arbitrary actor, search for skeletal mesh components
 		else 
 		{
-			APawn* Pawn = Cast<APawn>(ActorInst);
-			if( Pawn != NULL )
+			if(ActorInst != NULL)
 			{
 				TArray<USkeletalMeshComponent*> Components;
-				Pawn->GetComponents(Components);
+				ActorInst->GetComponents(Components);
 
 				int32 InsertionIndex = 0;
 				// look over all of the components looking for a SkelMeshComp and then if we find one we look at it to see if the bones match
@@ -2390,7 +2389,7 @@ void UParticleModuleLocationSkelVertSurface::UpdateBoneIndicesList(FParticleEmit
 						for (int32 FindBoneIdx = 0; FindBoneIdx < ValidAssociatedBones.Num(); FindBoneIdx++)
 						{
 							const int32 BoneIdx = SkelComp->SkeletalMesh->RefSkeleton.FindBoneIndex(ValidAssociatedBones[FindBoneIdx]);
-							if (BoneIdx != INDEX_NONE)
+							if (BoneIdx != INDEX_NONE && ValidAssociatedBones.Num() > InsertionIndex)
 							{
 								InstancePayload->ValidAssociatedBoneIndices[InsertionIndex++] = BoneIdx;
 							}
