CIS Fix - Fixed mismatched name between vectorized and non-vectorized FTransform for Android and HTML5

diff --git a/Engine/Source/Runtime/Core/Public/Math/TransformVectorized.h b/Engine/Source/Runtime/Core/Public/Math/TransformVectorized.h
--- a/Engine/Source/Runtime/Core/Public/Math/TransformVectorized.h
+++ b/Engine/Source/Runtime/Core/Public/Math/TransformVectorized.h
@@ -892,7 +892,7 @@ public:
 	 *
 	 * @param SourceAtom The other transform to accumulate into this one
 	 */
-	FORCEINLINE void AccumulateWithAdditiveScale(const FTransform& SourceAtom)
+	FORCEINLINE void AccumulateWithAdditiveScale3D(const FTransform& SourceAtom)
 	{
 		// Add ref pose relative animation to base animation, only if rotation is significant.
 		const VectorRegister RotationW = VectorReplicate(SourceAtom.Rotation, 3);
diff --git a/Engine/Source/Runtime/Engine/Private/Animation/AnimNode_BoneDrivenController.cpp b/Engine/Source/Runtime/Engine/Private/Animation/AnimNode_BoneDrivenController.cpp
--- a/Engine/Source/Runtime/Engine/Private/Animation/AnimNode_BoneDrivenController.cpp
+++ b/Engine/Source/Runtime/Engine/Private/Animation/AnimNode_BoneDrivenController.cpp
@@ -98,7 +98,7 @@ void FAnimNode_BoneDrivenController::EvaluateBoneTransforms(USkeletalMeshCompone
 
 	// Build final transform difference
 	FTransform FinalDiff(FQuat::MakeFromEuler(NewRotDiff), NewTransDiff, FVector(NewScaleDiff));
-	NewLocal.AccumulateWithAdditiveScale(FinalDiff);
+	NewLocal.AccumulateWithAdditiveScale3D(FinalDiff);
 
 	// If we have a parent, concatenate the transform, otherwise just take the new transform
 	int32 PIdx = RequiredBones.GetParentBoneIndex(TargetBone.BoneIndex);
