CPU optimizations in the GPU particle simulation code.

1) Micro-optimizations when building the tile vertex buffer.
2) Added unsafe FP32 -> FP16 conversion that doesn't perform checks. This is now used when building a GPU particle vertex buffer because the inputs are known to be within [0,1].

Integrated from GitHub pull requests 102 and 122 by [REDACTED_PERSON]

diff --git a/Engine/Source/Runtime/Core/Public/Math/Float16.h b/Engine/Source/Runtime/Core/Public/Math/Float16.h
--- a/Engine/Source/Runtime/Core/Public/Math/Float16.h
+++ b/Engine/Source/Runtime/Core/Public/Math/Float16.h
@@ -59,6 +59,16 @@ public:
 
 	/** Convert from Fp32 to Fp16. */
 	void Set( float FP32Value );
+	
+	/**
+	 * Convert from Fp32 to Fp16 without doing any checks if
+	 * the Fp32 exponent is too large or too small. This is a
+	 * faster alternative to Set() when you know the values
+	 * within the single precision float don't need the checks.
+	 *
+	 * @param FP32Value Single precision float to be set as half precision.
+	 */
+	void SetWithoutBoundsChecks( const float FP32Value );
 
 	/** Convert from Fp16 to Fp32. */
 	float GetFloat() const;
@@ -138,6 +148,19 @@ FORCEINLINE void FFloat16::Set( float FP32Value )
 	}
 }
 
+FORCEINLINE void FFloat16::SetWithoutBoundsChecks( const float FP32Value )
+{
+	const FFloat32 FP32(FP32Value);
+
+	// Make absolutely sure that you never pass in a single precision floating
+	// point value that may actually need the checks. If you are not 100% sure
+	// of that just use Set().
+
+	Components.Sign = FP32.Components.Sign;
+	Components.Exponent = int32(FP32.Components.Exponent) - 127 + 15;
+	Components.Mantissa = uint16(FP32.Components.Mantissa >> 13);
+}
+
 FORCEINLINE float FFloat16::GetFloat() const
 {
 	FFloat32	Result;
diff --git a/Engine/Source/Runtime/Engine/Private/Particles/ParticleGpuSimulation.cpp b/Engine/Source/Runtime/Engine/Private/Particles/ParticleGpuSimulation.cpp
--- a/Engine/Source/Runtime/Engine/Private/Particles/ParticleGpuSimulation.cpp
+++ b/Engine/Source/Runtime/Engine/Private/Particles/ParticleGpuSimulation.cpp
@@ -1119,20 +1119,18 @@ FORCEINLINE int32 ComputeAlignedTileCount(int32 TileCount)
  */
 static void BuildTileVertexBuffer( FParticleBufferParamRef TileOffsetsRef, const uint32* Tiles, int32 TileCount )
 {
-	int32 Index;
 	const int32 AlignedTileCount = ComputeAlignedTileCount(TileCount);
 	FVector2D* TileOffset = (FVector2D*)RHILockVertexBuffer( TileOffsetsRef, 0, AlignedTileCount * sizeof(FVector2D), RLM_WriteOnly );
-	for ( Index = 0; Index < TileCount; ++Index )
+	for ( int32 Index = 0; Index < TileCount; ++Index )
 	{
 		const uint32 TileIndex = Tiles[Index];
-		TileOffset[Index] = FVector2D(
-			FMath::Fractional( (float)TileIndex / (float)GParticleSimulationTileCountX ),
-			FMath::Fractional( FMath::TruncToFloat( (float)TileIndex / (float)GParticleSimulationTileCountX ) / (float)GParticleSimulationTileCountY )
-									  );
+		TileOffset[Index].X = FMath::Fractional( (float)TileIndex / (float)GParticleSimulationTileCountX );
+		TileOffset[Index].Y = FMath::Fractional( FMath::TruncToFloat( (float)TileIndex / (float)GParticleSimulationTileCountX ) / (float)GParticleSimulationTileCountY );
 	}
-	for ( ; Index < AlignedTileCount; ++Index )
+	for ( int32 Index = TileCount; Index < AlignedTileCount; ++Index )
 	{
-		TileOffset[Index] = FVector2D(100.0f, 100.0f);
+		TileOffset[Index].X = 100.0f;
+		TileOffset[Index].Y = 100.0f;
 	}
 	RHIUnlockVertexBuffer( TileOffsetsRef );
 }
@@ -1827,8 +1825,9 @@ static void BuildParticleVertexBuffer( FVertexBufferRHIParamRef VertexBufferRHI,
 				// on some platforms, union and/or bitfield writes to Locked memory are really slow, so use a forced int write instead
 				// and in fact one 32-bit write is faster than two uint16 writes (i.e. using .Encoded)
 				FParticleIndex Temp;
-				Temp.X = IndexX;
-				Temp.Y = IndexY;
+				// We use the unsafe version of FP32 -> FP16 conversion because we know all values are in [0,1].
+				Temp.X.SetWithoutBoundsChecks(IndexX);
+				Temp.Y.SetWithoutBoundsChecks(IndexY);
 				*(uint32*)ParticleIndices = *(uint32*)&Temp;
 
 				// move to next particle
