#ue4 - Bug fixes for Character movement on rotating surfaces, especially for networking.

- Fixed incorrect rotation delta computed in UpdateBasedMovement().
- Changed AActor::PostNetReceiveLocation to also include rotation (PostNetReceiveLocationAndRotation). Also a perf gain since we don't set location and rotation separately now.
- Fixed replicated rotation overwriting relative rotation for Characters.
- Fixed simulated proxies not matching the server location when the velocity reached zero. Proxies would move forward one more frame in that case.
- Fixed bIgnoreBaseRotation being ignored for simulated clients. Absolute rotation is sent in this case instead.
- Changed some MovementBaseUtility functions to conform to engine conventions ("location" rather than "position").
- Initial support for Characters walking on skeletal meshes (disabled until a few client networking features are completed soon).

diff --git a/Engine/Source/Runtime/Engine/Classes/GameFramework/Actor.h b/Engine/Source/Runtime/Engine/Classes/GameFramework/Actor.h
--- a/Engine/Source/Runtime/Engine/Classes/GameFramework/Actor.h
+++ b/Engine/Source/Runtime/Engine/Classes/GameFramework/Actor.h
@@ -1480,7 +1480,11 @@ public:
 	virtual void OnRep_ReplicatedMovement();
 
 	/** Update and smooth location, not called for simulated physics! */
-	virtual void PostNetReceiveLocation();
+	DEPRECATED(4.4, "PostNetReceiveLocation() has been replaced by PostNetReceiveLocationAndRotation().")
+	virtual void PostNetReceiveLocation() {}
+
+	/** Update location and rotation from ReplicatedMovement. Not called for simulated physics! */
+	virtual void PostNetReceiveLocationAndRotation();
 
 	/** Update velocity - typically from ReplicatedMovement, not called for simulated physics! */
 	virtual void PostNetReceiveVelocity(const FVector& NewVelocity);
diff --git a/Engine/Source/Runtime/Engine/Classes/GameFramework/Character.h b/Engine/Source/Runtime/Engine/Classes/GameFramework/Character.h
--- a/Engine/Source/Runtime/Engine/Classes/GameFramework/Character.h
+++ b/Engine/Source/Runtime/Engine/Classes/GameFramework/Character.h
@@ -42,6 +42,10 @@ struct FRepRootMotionMontage
 	UPROPERTY()
 	bool bRelativePosition;
 
+	/** Whether rotation is relative or absolute. */
+	UPROPERTY()
+	bool bRelativeRotation;
+
 	/** Clear the montage */
 	void Clear()
 	{
@@ -70,47 +74,64 @@ struct FSimulatedRootMotionReplicatedMove
 };
 
 
-/** Utility for determining if we should use relative positioning when based on a component (because it may move). */
+/** MovementBaseUtility provides utilities for working with movement bases, for which we may need relative positioning info. */
 namespace MovementBaseUtility
 {
+	/** Determine whether MovementBase can possibly move. */
 	FORCEINLINE bool IsDynamicBase(const class UPrimitiveComponent* MovementBase)
 	{
 		return (MovementBase && MovementBase->Mobility == EComponentMobility::Movable);
 	}
 
-	FORCEINLINE bool UseRelativePosition(const class UPrimitiveComponent* MovementBase)
+	/** Determine if we should use relative positioning when based on a component (because it may move). */
+	FORCEINLINE bool UseRelativeLocation(const class UPrimitiveComponent* MovementBase)
 	{
 		return IsDynamicBase(MovementBase);
 	}
 
 	/** Ensure that BasedObjectTick ticks after NewBase */
-	void AddTickDependency(FTickFunction& BasedObjectTick, class UPrimitiveComponent* NewBase);
+	ENGINE_API void AddTickDependency(FTickFunction& BasedObjectTick, class UPrimitiveComponent* NewBase);
 
 	/** Remove tick dependency of BasedObjectTick on OldBase */
-	void RemoveTickDependency(FTickFunction& BasedObjectTick, class UPrimitiveComponent* OldBase);
+	ENGINE_API void RemoveTickDependency(FTickFunction& BasedObjectTick, class UPrimitiveComponent* OldBase);
 
 	/** Get the velocity of the given component, first checking the ComponentVelocity and falling back to the physics velocity if necessary. */
-	FVector GetMovementBaseVelocity(const class UPrimitiveComponent* MovementBase);
+	ENGINE_API FVector GetMovementBaseVelocity(const class UPrimitiveComponent* MovementBase);
 
 	/** Get the tangential velocity at WorldLocation for the given component. */
-	FVector GetMovementBaseTangentialVelocity(const class UPrimitiveComponent* MovementBase, const FVector& WorldLocation);
+	ENGINE_API FVector GetMovementBaseTangentialVelocity(const class UPrimitiveComponent* MovementBase, const FVector& WorldLocation);
+
+	/** Get the transforms for the given MovementBase, optionally at the location of a bone. Returns false if MovementBase is NULL, or if BoneName is not a valid bone. */
+	ENGINE_API bool GetMovementBaseTransform(const UPrimitiveComponent* MovementBase, const FName BoneName, FVector& OutLocation, FQuat& OutQuat);
+
+
+	// Deprecated function
+	DEPRECATED(4.4, "MovementBaseUtility::UseRelativePosition() is deprecated, use MovementBaseUtility::UseRelativeLocation() instead.")
+	FORCEINLINE bool UseRelativePosition(const class UPrimitiveComponent* MovementBase)
+	{
+		return UseRelativeLocation(MovementBase);
+	}
 }
 
-/** Struct to hold relative position information from the server. */
+/** Struct to hold information about the "base" object the character is standing on. */
 USTRUCT()
-struct FRepRelativeMovement
+struct FBasedMovementInfo
 {
 	GENERATED_USTRUCT_BODY()
 
 	/** Component we are based on */
 	UPROPERTY()
-	class UPrimitiveComponent* MovementBase;
+	UPrimitiveComponent* MovementBase;
 
-	/** Location relative to MovementBase. Only valid if HasRelativePosition() is true. */
+	/** Bone name on component, for skeletal meshes. NAME_None if not a skeletal mesh or if bone is invalid. */
+	UPROPERTY()
+	FName BoneName;
+
+	/** Location relative to MovementBase. Only valid if HasRelativeLocation() is true. */
 	UPROPERTY()
 	FVector_NetQuantize100 Location;
 
-	/** Rotation relative to MovementBase. Only valid if HasRelativePosition() is true. */
+	/** Rotation: relative to MovementBase if HasRelativeRotation() is true, absolute otherwise. */
 	UPROPERTY()
 	FRotator Rotation;
 
@@ -118,11 +139,24 @@ struct FRepRelativeMovement
 	UPROPERTY()
 	bool bServerHasBaseComponent;
 
+	/** Whether rotation is relative to the base or absolute. It can only be relative if location is also relative. */
+	UPROPERTY()
+	bool bRelativeRotation;
 
-	/** Is position relative? */
-	FORCEINLINE bool HasRelativePosition() const
+	/** Whether there is a velocity on the server. Used for forcing replication when velocity goes to zero. */
+	UPROPERTY()
+	bool bServerHasVelocity;
+
+	/** Is location relative? */
+	FORCEINLINE bool HasRelativeLocation() const
+	{
+		return MovementBaseUtility::UseRelativeLocation(MovementBase);
+	}
+
+	/** Is rotation relative or absolute? It can only be relative if location is also relative. */
+	FORCEINLINE bool HasRelativeRotation() const
 	{
-		return MovementBaseUtility::UseRelativePosition(MovementBase);
+		return bRelativeRotation && HasRelativeLocation();
 	}
 
 	/** Return true if the client should have MovementBase, but it hasn't replicated (possibly component has not streamed in). */
@@ -130,8 +164,21 @@ struct FRepRelativeMovement
 	{
 		return (MovementBase == NULL) && bServerHasBaseComponent;
 	}
+
+
+	// Deprecated function
+	DEPRECATED(4.4, "FBasedMovementInfo::HasRelativePosition() is deprecated, use FBasedMovementInfo::HasRelativeLocation() instead.")
+	FORCEINLINE bool HasRelativePosition() const
+	{
+		return HasRelativeLocation();
+	}
 };
 
+// Allow the old name to continue to work for one release
+DEPRECATED(4.4, "FRepRelativeMovement has been renamed to FBasedMovementInfo")
+typedef FBasedMovementInfo FRepRelativeMovement;
+
+
 //=============================================================================
 // Characters are Pawns that have a mesh, collision, and physics and are responsible for all
 // physical interaction between the player or AI and the world. They are strongly tied
@@ -171,15 +218,21 @@ class ENGINE_API ACharacter : public APawn
 	static FName CapsuleComponentName;
 
 	/** Sets the MovementBase used by CharacterMovement walking movement. */
-	virtual void SetBase(UPrimitiveComponent* NewBase, bool bNotifyActor=true);
-
+	virtual void SetBase(UPrimitiveComponent* NewBase, const FName BoneName = NAME_None, bool bNotifyActor=true);
+	 
 protected:
-	// CharacterMovement basing related variables (here for replication purposes).
-	//
 
-	/** Component we are based on */
+	/** Info about our current movement base (object we are standing on). */
 	UPROPERTY()
-	class UPrimitiveComponent* MovementBase;
+	struct FBasedMovementInfo BasedMovement;
+
+	/** Replicated version of relative movement. Read-only on simulated proxies! */
+	UPROPERTY(ReplicatedUsing=OnRep_ReplicatedBasedMovement)
+	struct FBasedMovementInfo ReplicatedBasedMovement;
+
+	/** Rep notify for ReplicatedBasedMovement */
+	UFUNCTION()
+	virtual void OnRep_ReplicatedBasedMovement();
 
 	/** Desired translation offset of mesh. */
 	UPROPERTY()
@@ -188,25 +241,31 @@ protected:
 	/** Event called after actor's base changes (if SetBase was requested to notify us with bNotifyPawn). */
 	virtual void BaseChange();
 
-	// Always called immediately after properties are received from the remote.
-	virtual void PostNetReceiveBase();
-
 	/** CharacterMovement MovementMode (and custom mode) replicated for simulated proxies. Use CharacterMovementComponent::UnpackNetworkMovementMode() to translate it. */
 	UPROPERTY(Replicated)
 	uint8 ReplicatedMovementMode;
 
+	/** Flag that we are receiving replication of the based movement. */
+	UPROPERTY()
+	bool bInBaseReplication;
+
 public:	
 
+	/** Accessor for BasedMovement */
+	FORCEINLINE const FBasedMovementInfo& GetBasedMovement() const { return BasedMovement; }
+	
+	/** Accessor for ReplicatedBasedMovement */
+	FORCEINLINE const FBasedMovementInfo& GetReplicatedBasedMovement() const { return ReplicatedBasedMovement; }
+
+	/** Save a new relative location in BasedMovement and a new rotation with is either relative or absolute. */
+	void SaveRelativeBasedMovement(const FVector& NewRelativeLocation, const FRotator& NewRotation, bool bRelativeRotation);
+
 	/** Returns ReplicatedMovementMode */
 	uint8 GetReplicatedMovementMode() const { return ReplicatedMovementMode; }
 
 	/** @return Desired translation offset of mesh. */
 	const FVector& GetBaseTranslationOffset() const { return BaseTranslationOffset; }
 
-	/** location, rotation and velocity relative to base/bone (valid if base exists) */
-	UPROPERTY(replicated)
-	struct FRepRelativeMovement RelativeMovement;
-
 	/** Default crouched eye height */
 	UPROPERTY(EditAnywhere, BlueprintReadWrite, Category=Camera)
 	float CrouchedEyeHeight;
@@ -261,7 +320,7 @@ public:
 	virtual void PreNetReceive() override;
 	virtual void PostNetReceive() override;
 	virtual void OnRep_ReplicatedMovement() override;
-	virtual void PostNetReceiveLocation() override;
+	virtual void PostNetReceiveLocationAndRotation() override;
 	virtual void GetSimpleCollisionCylinder(float& CollisionRadius, float& CollisionHalfHeight) const override;
 	virtual void ApplyWorldOffset(const FVector& InOffset, bool bWorldShift) override;
 	virtual UActorComponent* FindComponentByClass(const TSubclassOf<UActorComponent> ComponentClass) const override;
@@ -277,7 +336,7 @@ public:
 	// Begin APawn Interface.
 	virtual void PostInitializeComponents() override;
 	virtual class UPawnMovementComponent* GetMovementComponent() const override { return CharacterMovement; }
-	virtual class UPrimitiveComponent* GetMovementBase() const override final { return MovementBase; }
+	virtual class UPrimitiveComponent* GetMovementBase() const override final { return BasedMovement.MovementBase; }
 	virtual float GetDefaultHalfHeight() const override;
 	virtual void TurnOff() override;
 	virtual void Restart() override;
diff --git a/Engine/Source/Runtime/Engine/Classes/GameFramework/CharacterMovementComponent.h b/Engine/Source/Runtime/Engine/Classes/GameFramework/CharacterMovementComponent.h
--- a/Engine/Source/Runtime/Engine/Classes/GameFramework/CharacterMovementComponent.h
+++ b/Engine/Source/Runtime/Engine/Classes/GameFramework/CharacterMovementComponent.h
@@ -537,7 +537,7 @@ public:
 	uint32 bCrouchMovesCharacterDown:1;
 
 	/** If true, the pawn ignores the effects of changes in its base's rotation on its rotation. */
-	UPROPERTY(EditAnywhere, BlueprintReadWrite, Category=Attachment)
+	UPROPERTY(Category="Character Movement", EditAnywhere, BlueprintReadWrite)
 	uint32 bIgnoreBaseRotation:1;
 
 	/** 
@@ -735,10 +735,10 @@ public:
 	/** Update controller's view rotation as pawn's base rotates */
 	virtual void UpdateBasedRotation(FRotator &FinalRotation, const FRotator& ReducedRotation);
 
-	/** Update (or defer updating) OldBaseLocation and OldBaseRotation if there is a valid movement base. */
+	/** Update (or defer updating) OldBaseLocation and OldBaseQuat if there is a valid movement base. */
 	virtual void MaybeSaveBaseLocation();
 
-	/** Update OldBaseLocation and OldBaseRotation if there is a valid movement base. */
+	/** Update OldBaseLocation and OldBaseQuat if there is a valid movement base. */
 	virtual void SaveBaseLocation();
 
 	/** changes physics based on MovementMode */
@@ -871,7 +871,7 @@ public:
 	virtual bool StepUp(const FVector& GravDir, const FVector& Delta, const FHitResult &Hit, struct UCharacterMovementComponent::FStepDownResult* OutStepDownResult = NULL);
 
 	/** Update the base of the actor */
-	virtual void SetBase(UPrimitiveComponent* NewBase, bool bNotifyActor=true);
+	virtual void SetBase(UPrimitiveComponent* NewBase, const FName BoneName = NAME_None, bool bNotifyActor=true);
 
 	/** Applies repulsion force to all touched components */
 	void ApplyRepulsionForce(float DeltaTime);
@@ -1436,6 +1436,7 @@ public:
 	FFindFloorResult StartFloor;
 	FRotator StartRotation;
 	FRotator StartControlRotation;
+	FQuat StartBaseRotation; // rotation of the base component, only saved if it can move.
 	float StartCapsuleRadius;
 	float StartCapsuleHalfHeight;
 
diff --git a/Engine/Source/Runtime/Engine/Classes/GameFramework/Pawn.h b/Engine/Source/Runtime/Engine/Classes/GameFramework/Pawn.h
--- a/Engine/Source/Runtime/Engine/Classes/GameFramework/Pawn.h
+++ b/Engine/Source/Runtime/Engine/Classes/GameFramework/Pawn.h
@@ -171,7 +171,7 @@ public:
 	virtual float GetNetPriority(const FVector& ViewPos, const FVector& ViewDir, APlayerController* Viewer, UActorChannel* InChannel, float Time, bool bLowBandwidth) override;
 	virtual bool ShouldTickIfViewportsOnly() const override;
 	virtual bool IsNetRelevantFor(APlayerController* RealViewer, AActor* Viewer, const FVector& SrcLocation) override;
-	virtual void PostNetReceiveLocation() override;
+	virtual void PostNetReceiveLocationAndRotation() override;
 	virtual void PostNetReceiveVelocity(const FVector& NewVelocity) override;
 	virtual void DisplayDebug(class UCanvas* Canvas, const FDebugDisplayInfo& DebugDisplay, float& YL, float& YPos) override;
 	virtual void GetActorEyesViewPoint( FVector& Location, FRotator& Rotation ) const override;
diff --git a/Engine/Source/Runtime/Engine/Private/ActorReplication.cpp b/Engine/Source/Runtime/Engine/Private/ActorReplication.cpp
--- a/Engine/Source/Runtime/Engine/Private/ActorReplication.cpp
+++ b/Engine/Source/Runtime/Engine/Private/ActorReplication.cpp
@@ -92,22 +92,17 @@ void AActor::OnRep_ReplicatedMovement()
 		if (Role == ROLE_SimulatedProxy)
 		{
 			PostNetReceiveVelocity(ReplicatedMovement.LinearVelocity);
-			PostNetReceiveLocation();
-
-			if( RootComponent && (ReplicatedMovement.Rotation != GetActorRotation()) )
-			{
-				RootComponent->MoveComponent(FVector::ZeroVector, ReplicatedMovement.Rotation, false);
-			}
+			PostNetReceiveLocationAndRotation();
 		}
 	}
 }
 
-void AActor::PostNetReceiveLocation()
+void AActor::PostNetReceiveLocationAndRotation()
 {
-	USceneComponent* RootSceneComp = Cast<USceneComponent>(RootComponent);
-	if( RootSceneComp && RootSceneComp->IsRegistered() && (ReplicatedMovement.Location != GetActorLocation()) )
+	if( RootComponent && RootComponent->IsRegistered() && (ReplicatedMovement.Location != GetActorLocation() || ReplicatedMovement.Rotation != GetActorRotation()) )
 	{
-		TeleportTo(ReplicatedMovement.Location, GetActorRotation(), false, true);
+		TeleportTo(ReplicatedMovement.Location, ReplicatedMovement.Rotation, false);
+		//SetActorLocationAndRotation(ReplicatedMovement.Location, ReplicatedMovement.Rotation); <-- preferred, but awaiting answer to question about UpdateNavOctree() missing in SceneComponent::MoveComponent
 	}
 }
 
diff --git a/Engine/Source/Runtime/Engine/Private/Character.cpp b/Engine/Source/Runtime/Engine/Private/Character.cpp
--- a/Engine/Source/Runtime/Engine/Private/Character.cpp
+++ b/Engine/Source/Runtime/Engine/Private/Character.cpp
@@ -357,7 +357,7 @@ void ACharacter::TeleportSucceeded(bool bIsATest)
 
 void ACharacter::ClearCrossLevelReferences()
 {
-	if( MovementBase != NULL && GetOutermost() != MovementBase->GetOutermost() )
+	if( BasedMovement.MovementBase != NULL && GetOutermost() != BasedMovement.MovementBase->GetOutermost() )
 	{
 		SetBase( NULL );
 	}
@@ -366,12 +366,12 @@ void ACharacter::ClearCrossLevelReferences()
 }
 
 
-// Helper to set up tick dependencies on a movement base
+
 namespace MovementBaseUtility
 {
 	void AddTickDependency(FTickFunction& BasedObjectTick, class UPrimitiveComponent* NewBase)
 	{
-		if (NewBase && MovementBaseUtility::UseRelativePosition(NewBase))
+		if (NewBase && MovementBaseUtility::UseRelativeLocation(NewBase))
 		{
 			BasedObjectTick.AddPrerequisite(NewBase, NewBase->PrimaryComponentTick);
 			AActor* NewBaseOwner = NewBase->GetOwner();
@@ -391,7 +391,7 @@ namespace MovementBaseUtility
 
 	void RemoveTickDependency(FTickFunction& BasedObjectTick, class UPrimitiveComponent* OldBase)
 	{
-		if (OldBase && MovementBaseUtility::UseRelativePosition(OldBase))
+		if (OldBase && MovementBaseUtility::UseRelativeLocation(OldBase))
 		{
 			BasedObjectTick.RemovePrerequisite(OldBase, OldBase->PrimaryComponentTick);
 			AActor* OldBaseOwner = OldBase->GetOwner();
@@ -409,6 +409,7 @@ namespace MovementBaseUtility
 		}
 	}
 
+	// TODO: GetPhysicsLinearVelocity(BoneName)
 	FVector GetMovementBaseVelocity(const class UPrimitiveComponent* MovementBase)
 	{
 		FVector BaseVelocity = FVector::ZeroVector;
@@ -454,13 +455,52 @@ namespace MovementBaseUtility
 		
 		return FVector::ZeroVector;
 	}
+
+	bool GetMovementBaseTransform(const UPrimitiveComponent* MovementBase, const FName BoneName, FVector& OutLocation, FQuat& OutQuat)
+	{
+		if (MovementBase)
+		{
+			if (BoneName != NAME_None)
+			{
+				const USkeletalMeshComponent* SkeletalBase = Cast<USkeletalMeshComponent>(MovementBase);
+				if (SkeletalBase)
+				{
+					const int32 BoneIndex = SkeletalBase->GetBoneIndex(BoneName);
+					if (BoneIndex != INDEX_NONE)
+					{
+						const FTransform BoneTransform = SkeletalBase->GetBoneTransform(BoneIndex);
+						OutLocation = BoneTransform.GetLocation();
+						OutQuat = BoneTransform.GetRotation();
+						return true;
+					}
+
+					UE_LOG(LogCharacter, Warning, TEXT("GetMovementBaseTransform(): Invalid bone '%s' for SkeletalMeshComponent base %s"), *BoneName.ToString(), *GetPathNameSafe(MovementBase));
+					return false;
+				}
+				// TODO: warn if not a skeletal mesh but providing bone index.
+			}
+
+			// No bone supplied
+			OutLocation = MovementBase->GetComponentLocation();
+			OutQuat = MovementBase->GetComponentQuat();
+			return true;
+		}
+
+		// NULL MovementBase
+		OutLocation = FVector::ZeroVector;
+		OutQuat = FQuat::Identity;
+		return false;
+	}
 }
 
 
 /**	Change the Pawn's base. */
-void ACharacter::SetBase( UPrimitiveComponent* NewBaseComponent, bool bNotifyPawn )
+void ACharacter::SetBase( UPrimitiveComponent* NewBaseComponent, const FName BoneName, bool bNotifyPawn )
 {
-	if (NewBaseComponent != MovementBase)
+	const bool bBaseChanged = (NewBaseComponent != BasedMovement.MovementBase);
+	const bool bBoneChanged = (BoneName != BasedMovement.BoneName);
+
+	if (bBaseChanged || bBoneChanged)
 	{
 		// Verify no recursion.
 		APawn* Loop = (NewBaseComponent ? Cast<APawn>(NewBaseComponent->GetOwner()) : NULL);
@@ -474,59 +514,68 @@ void ACharacter::SetBase( UPrimitiveComponent* NewBaseComponent, bool bNotifyPaw
 		}
 
 		// Set base.
-		UPrimitiveComponent* OldBase = MovementBase;
-		MovementBase = NewBaseComponent;
-
-		// Update relative location/rotation
-		// Only do this for non-simulated proxies, because we don't want to overwrite RelativeMovement for those (it could about to be updated by a network update).
-		if ( Role == ROLE_Authority || Role == ROLE_AutonomousProxy )
-		{
-			UE_LOG(LogCharacter, Verbose, TEXT("Setting base on %s for '%s' to '%s'"), Role == ROLE_Authority ? TEXT("Server") : TEXT("AutoProxy"), *GetName(), *GetFullNameSafe(NewBaseComponent));
-			RelativeMovement.MovementBase = NewBaseComponent;
-			RelativeMovement.bServerHasBaseComponent = (NewBaseComponent != NULL); // Flag whether the server had a non-null base.
-
-			CharacterMovement->SaveBaseLocation();
-		}
-		else
-		{
-			UE_LOG(LogCharacter, Verbose, TEXT("Setting base on Client for '%s' to '%s'"), *GetName(), *GetFullNameSafe(NewBaseComponent));
-		}
-
-		// Notify this actor of his new floor.
-		if ( bNotifyPawn )
-		{
-			BaseChange();
-		}
+		UPrimitiveComponent* OldBase = BasedMovement.MovementBase;
+		BasedMovement.MovementBase = NewBaseComponent;
+		BasedMovement.BoneName = BoneName;
 
 		if (CharacterMovement)
 		{
-			MovementBaseUtility::RemoveTickDependency(CharacterMovement->PrimaryComponentTick, OldBase);
-			MovementBaseUtility::AddTickDependency(CharacterMovement->PrimaryComponentTick, MovementBase);			
+			if (bBaseChanged)
+			{
+				MovementBaseUtility::RemoveTickDependency(CharacterMovement->PrimaryComponentTick, OldBase);
+				MovementBaseUtility::AddTickDependency(CharacterMovement->PrimaryComponentTick, NewBaseComponent);
+			}
 
-			if (MovementBase)
+			if (NewBaseComponent)
 			{
 				// Update OldBaseLocation/Rotation as those were referring to a different base
-				CharacterMovement->OldBaseLocation = MovementBase->GetComponentLocation();
-				CharacterMovement->OldBaseQuat = MovementBase->GetComponentQuat();
-
-				// Enable pre-cloth tick if we are standing on a physics object, as we need to to use 
-				// post-physics transforms
-				if(MovementBase->IsSimulatingPhysics())
-				{
-					CharacterMovement->PreClothComponentTick.SetTickFunctionEnable(true);
+				// ... but not when handling replication for proxies (since they are going to copy this data from the replicated values anyway)
+				if (!bInBaseReplication)
+				{					
+					CharacterMovement->SaveBaseLocation();
 				}
+
+				// Enable pre-cloth tick if we are standing on a physics object, as we need to to use post-physics transforms
+				CharacterMovement->PreClothComponentTick.SetTickFunctionEnable(NewBaseComponent->IsSimulatingPhysics());
 			}
 			else
 			{
-				// Base is NULL, invalidate floor.
+				BasedMovement.BoneName = NAME_None; // None, regardless of whether user tried to set a bone name, since we have no base component.
+				BasedMovement.bRelativeRotation = false;
 				CharacterMovement->CurrentFloor.Clear();
-
 				CharacterMovement->PreClothComponentTick.SetTickFunctionEnable(false);
 			}
+
+			if (Role == ROLE_Authority || Role == ROLE_AutonomousProxy)
+			{
+				BasedMovement.bServerHasBaseComponent = (BasedMovement.MovementBase != NULL); // Also set on proxies for nicer debugging.
+				UE_LOG(LogCharacter, Verbose, TEXT("Setting base on %s for '%s' to '%s'"), Role == ROLE_Authority ? TEXT("Server") : TEXT("AutoProxy"), *GetName(), *GetFullNameSafe(NewBaseComponent));
+			}
+			else
+			{
+				UE_LOG(LogCharacter, Verbose, TEXT("Setting base on Client for '%s' to '%s'"), *GetName(), *GetFullNameSafe(NewBaseComponent));
+			}
+
+		}
+
+		// Notify this actor of his new floor.
+		if ( bNotifyPawn )
+		{
+			BaseChange();
 		}
 	}
 }
 
+
+void ACharacter::SaveRelativeBasedMovement(const FVector& NewRelativeLocation, const FRotator& NewRotation, bool bRelativeRotation)
+{
+	checkSlow(BasedMovement.HasRelativeLocation());
+	BasedMovement.Location = NewRelativeLocation;
+	BasedMovement.Rotation = NewRotation;
+	BasedMovement.bRelativeRotation = bRelativeRotation;
+}
+
+
 void ACharacter::TurnOff()
 {
 	if (CharacterMovement.Get() != NULL)
@@ -617,17 +666,17 @@ void ACharacter::DisplayDebug(UCanvas* Canvas, const FDebugDisplayInfo& DebugDis
 		FIndenter PhysicsIndent(Indent);
 
 		FString BaseString;
-		if ( CharacterMovement.Get() == NULL || MovementBase == NULL )
+		if ( CharacterMovement.Get() == NULL || BasedMovement.MovementBase == NULL )
 		{
 			BaseString = "Not Based";
 		}
 		else
 		{
-			BaseString = MovementBase->IsWorldGeometry() ? "World Geometry" : MovementBase->GetName();
+			BaseString = BasedMovement.MovementBase->IsWorldGeometry() ? "World Geometry" : BasedMovement.MovementBase->GetName();
 			BaseString = FString::Printf(TEXT("Based On %s"), *BaseString);
 		}
 		
-		Canvas->DrawText(RenderFont, FString::Printf(TEXT("RelativeLoc: %s RelativeRot: %s %s"), *RelativeMovement.Location.ToString(), *RelativeMovement.Rotation.ToString(), *BaseString), Indent, YPos);
+		Canvas->DrawText(RenderFont, FString::Printf(TEXT("RelativeLoc: %s Rot: %s %s"), *BasedMovement.Location.ToString(), *BasedMovement.Rotation.ToString(), *BaseString), Indent, YPos);
 		YPos += YL;
 
 		if ( CharacterMovement.Get() != NULL )
@@ -805,16 +854,11 @@ void ACharacter::ClientAdjustRootMotionPosition_Implementation(float TimeStamp,
 //
 // Static variables for networking.
 //
-static FVector_NetQuantize100	SavedRelativeLocation;
-static FRotator					SavedRelativeRotation;
-static uint8					SavedMovementMode;
+static uint8 SavedMovementMode;
 
 void ACharacter::PreNetReceive()
 {
-	SavedRelativeLocation = RelativeMovement.Location;
-	SavedRelativeRotation = RelativeMovement.Rotation;
 	SavedMovementMode = ReplicatedMovementMode;
-
 	Super::PreNetReceive();
 }
 
@@ -827,43 +871,64 @@ void ACharacter::PostNetReceive()
 	}
 
 	Super::PostNetReceive();
-
-	// Skip base updates while playing root motion, it is handled inside of OnRep_RootMotion
-	if (Role == ROLE_SimulatedProxy && !IsPlayingRootMotion())
-	{
-		PostNetReceiveBase();
-	}
 }
 
-void ACharacter::PostNetReceiveBase()
+void ACharacter::OnRep_ReplicatedBasedMovement()
 {	
 	if (Role != ROLE_SimulatedProxy)
 	{
 		return;
 	}
 
-	const bool bBaseChanged = (RelativeMovement.MovementBase != MovementBase);
+	// Skip base updates while playing root motion, it is handled inside of OnRep_RootMotion
+	if (IsPlayingRootMotion())
+	{
+		return;
+	}
+
+	TGuardValue<bool> bInBaseReplicationGuard(bInBaseReplication, true);
+
+	const bool bBaseChanged = (BasedMovement.MovementBase != ReplicatedBasedMovement.MovementBase || BasedMovement.BoneName != ReplicatedBasedMovement.BoneName);
 	if (bBaseChanged)
 	{
-		SetBase(RelativeMovement.MovementBase);
+		// Even though we will copy the replicated based movement info, we need to use SetBase() to set up tick dependencies and trigger notifications.
+		SetBase(ReplicatedBasedMovement.MovementBase, ReplicatedBasedMovement.BoneName);
 	}
 
-	if (RelativeMovement.HasRelativePosition())
+	// Make sure to use the values of relative location/rotation etc from the server.
+	BasedMovement = ReplicatedBasedMovement;
+
+	if (ReplicatedBasedMovement.HasRelativeLocation())
 	{
-		if (bBaseChanged || (RelativeMovement.Location != SavedRelativeLocation) || (RelativeMovement.Rotation != SavedRelativeRotation))
+		// Update transform relative to movement base
+		const FVector OldLocation = GetActorLocation();
+		MovementBaseUtility::GetMovementBaseTransform(ReplicatedBasedMovement.MovementBase, ReplicatedBasedMovement.BoneName, CharacterMovement->OldBaseLocation, CharacterMovement->OldBaseQuat);
+		const FVector NewLocation = CharacterMovement->OldBaseLocation + ReplicatedBasedMovement.Location;
+
+		if (ReplicatedBasedMovement.HasRelativeRotation())
 		{
-			const FVector OldLocation = GetActorLocation();
-			const FRotator NewRotation = (FRotationMatrix(RelativeMovement.Rotation) * FRotationMatrix(MovementBase->GetComponentRotation())).Rotator();
-			SetActorLocationAndRotation( MovementBase->GetComponentLocation() + RelativeMovement.Location, NewRotation, false);
+			// Relative location, relative rotation
+			FRotator NewRotation = (FRotationMatrix(ReplicatedBasedMovement.Rotation) * FQuatRotationTranslationMatrix(CharacterMovement->OldBaseQuat, FVector::ZeroVector)).Rotator();
 			
-			// When position or base changes, movement mode will need to be updated. This assumes rotation changes don't affect that.
-			CharacterMovement->bJustTeleported |= (bBaseChanged || GetActorLocation() != OldLocation);
+			// TODO: need a better way to not assume we only use Yaw.
+			NewRotation.Pitch = 0.f;
+			NewRotation.Roll = 0.f;
 
-			INetworkPredictionInterface* PredictionInterface = InterfaceCast<INetworkPredictionInterface>(GetMovementComponent());
-			if (PredictionInterface)
-			{
-				PredictionInterface->SmoothCorrection(OldLocation);
-			}
+			SetActorLocationAndRotation(NewLocation, NewRotation);	
+		}
+		else
+		{
+			// Relative location, absolute rotation
+			SetActorLocationAndRotation(NewLocation, ReplicatedBasedMovement.Rotation);
+		}
+
+		// When position or base changes, movement mode will need to be updated. This assumes rotation changes don't affect that.
+		CharacterMovement->bJustTeleported |= (bBaseChanged || GetActorLocation() != OldLocation);
+
+		INetworkPredictionInterface* PredictionInterface = InterfaceCast<INetworkPredictionInterface>(GetMovementComponent());
+		if (PredictionInterface)
+		{
+			PredictionInterface->SmoothCorrection(OldLocation);
 		}
 	}
 }
@@ -1014,10 +1079,22 @@ bool ACharacter::RestoreReplicatedMove(const FSimulatedRootMotionReplicatedMove
 	if( RootMotionRepMove.RootMotion.bRelativePosition )
 	{
 		bool bSuccess = false;
-		if( MovementBaseUtility::UseRelativePosition(ServerBase) )
+		if( MovementBaseUtility::UseRelativeLocation(ServerBase) )
 		{
 			const FVector ServerLocation = ServerBase->GetComponentLocation() + RootMotionRepMove.RootMotion.Location;
-			const FRotator ServerRotation = (FRotationMatrix(RootMotionRepMove.RootMotion.Rotation) * FRotationMatrix(ServerBase->GetComponentRotation())).Rotator();
+			FRotator ServerRotation;
+			if (RootMotionRepMove.RootMotion.bRelativeRotation)
+			{
+				// Relative rotation
+				// TODO: Bone rotation
+				ServerRotation = (FRotationMatrix(RootMotionRepMove.RootMotion.Rotation) * FRotationMatrix(ServerBase->GetComponentRotation())).Rotator();
+			}
+			else
+			{
+				// Absolute rotation
+				ServerRotation = RootMotionRepMove.RootMotion.Rotation;
+			}
+
 			UpdateSimulatedPosition(ServerLocation, ServerRotation);
 			bSuccess = true;
 		}
@@ -1034,9 +1111,9 @@ bool ACharacter::RestoreReplicatedMove(const FSimulatedRootMotionReplicatedMove
 	}
 
 	// Set base
-	if( MovementBase != ServerBase )
+	if( BasedMovement.MovementBase != ServerBase /* TODO: or bone changes */)
 	{
-		SetBase( ServerBase );
+		SetBase( ServerBase /* TODO: bone */ );
 	}
 
 	return true;
@@ -1067,12 +1144,12 @@ void ACharacter::UpdateSimulatedPosition(const FVector & NewLocation, const FRot
 	}
 }
 
-void ACharacter::PostNetReceiveLocation()
+void ACharacter::PostNetReceiveLocationAndRotation()
 {
 	if( Role == ROLE_SimulatedProxy )
 	{
-		// Don't change position if using relative position (it should be the same anyway)
-		if (!RelativeMovement.HasRelativePosition())
+		// Don't change transform if using relative position (it should be nearly the same anyway, or base may be slightly out of sync)
+		if (!ReplicatedBasedMovement.HasRelativeLocation())
 		{
 			const FVector OldLocation = GetActorLocation();
 			UpdateSimulatedPosition(ReplicatedMovement.Location, ReplicatedMovement.Rotation);
@@ -1095,10 +1172,11 @@ void ACharacter::PreReplication( IRepChangedPropertyTracker & ChangedPropertyTra
 	if ( RootMotionMontageInstance )
 	{
 		// Is position stored in local space?
-		RepRootMotion.bRelativePosition = RelativeMovement.HasRelativePosition();
-		RepRootMotion.Location			= RepRootMotion.bRelativePosition ? RelativeMovement.Location : GetActorLocation();
-		RepRootMotion.Rotation			= RepRootMotion.bRelativePosition ? RelativeMovement.Rotation : GetActorRotation();
-		RepRootMotion.MovementBase		= RelativeMovement.MovementBase;
+		RepRootMotion.bRelativePosition = BasedMovement.HasRelativeLocation();
+		RepRootMotion.bRelativeRotation = BasedMovement.HasRelativeRotation();
+		RepRootMotion.Location			= RepRootMotion.bRelativePosition ? BasedMovement.Location : GetActorLocation();
+		RepRootMotion.Rotation			= RepRootMotion.bRelativeRotation ? BasedMovement.Rotation : GetActorRotation();
+		RepRootMotion.MovementBase		= BasedMovement.MovementBase;
 		RepRootMotion.AnimMontage		= RootMotionMontageInstance->Montage;
 		RepRootMotion.Position			= RootMotionMontageInstance->Position;
 
@@ -1111,17 +1189,31 @@ void ACharacter::PreReplication( IRepChangedPropertyTracker & ChangedPropertyTra
 		DOREPLIFETIME_ACTIVE_OVERRIDE( ACharacter, RepRootMotion, false );
 	}
 
-	ReplicatedMovementMode = CharacterMovement->PackNetworkMovementMode();
+	ReplicatedMovementMode = CharacterMovement->PackNetworkMovementMode();	
+	ReplicatedBasedMovement = BasedMovement;
+
+	// Optimization: only update and replicate these values if they are actually going to be used.
+	if (BasedMovement.HasRelativeLocation())
+	{
+		// When velocity becomes zero, force replication so the position is updated to match the server (it may have moved due to simulation on the client).
+		ReplicatedBasedMovement.bServerHasVelocity = !CharacterMovement->Velocity.IsZero();
+
+		// Make sure absolute rotations are updated in case rotation occurred after the base info was saved.
+		if (!BasedMovement.HasRelativeRotation())
+		{
+			ReplicatedBasedMovement.Rotation = GetActorRotation();
+		}
+	}
 }
 
 void ACharacter::GetLifetimeReplicatedProps( TArray< FLifetimeProperty > & OutLifetimeProps ) const
 {
 	Super::GetLifetimeReplicatedProps( OutLifetimeProps );
 
-	DOREPLIFETIME_CONDITION( ACharacter, RepRootMotion,		COND_SimulatedOnly );
-	DOREPLIFETIME_CONDITION( ACharacter, RelativeMovement,	COND_SimulatedOnly );
-	DOREPLIFETIME_CONDITION( ACharacter, ReplicatedMovementMode, COND_SimulatedOnly );
-	DOREPLIFETIME_CONDITION( ACharacter, bIsCrouched,		COND_SimulatedOnly );
+	DOREPLIFETIME_CONDITION( ACharacter, RepRootMotion,				COND_SimulatedOnly );
+	DOREPLIFETIME_CONDITION( ACharacter, ReplicatedBasedMovement,	COND_SimulatedOnly );
+	DOREPLIFETIME_CONDITION( ACharacter, ReplicatedMovementMode,	COND_SimulatedOnly );
+	DOREPLIFETIME_CONDITION( ACharacter, bIsCrouched,				COND_SimulatedOnly );
 }
 
 bool ACharacter::IsPlayingRootMotion() const
diff --git a/Engine/Source/Runtime/Engine/Private/CharacterMovementComponent.cpp b/Engine/Source/Runtime/Engine/Private/CharacterMovementComponent.cpp
--- a/Engine/Source/Runtime/Engine/Private/CharacterMovementComponent.cpp
+++ b/Engine/Source/Runtime/Engine/Private/CharacterMovementComponent.cpp
@@ -615,7 +615,7 @@ void UCharacterMovementComponent::OnMovementModeChanged(EMovementMode PreviousMo
 		// make sure we update our new floor/base on initial entry of the walking physics
 		FindFloor(UpdatedComponent->GetComponentLocation(), CurrentFloor, false);
 		AdjustFloorHeight();
-		SetBase(CurrentFloor.HitResult.Component.Get());
+		SetBase(CurrentFloor.HitResult.Component.Get(), CurrentFloor.HitResult.BoneName);
 	}
 	else
 	{
@@ -1017,7 +1017,7 @@ void UCharacterMovementComponent::SimulateMovement(float DeltaSeconds)
 	}
 
 	// If base is not resolved on the client, we should not try to simulate at all
-	if (CharacterOwner->RelativeMovement.IsBaseUnresolved())
+	if (CharacterOwner->GetBasedMovement().IsBaseUnresolved())
 	{
 		UE_LOG(LogCharacterMovement, Verbose, TEXT("Base for simulated character '%s' is not resolved on client, skipping SimulateMovement"), *CharacterOwner->GetName());
 		return;
@@ -1098,7 +1098,7 @@ void UCharacterMovementComponent::SimulateMovement(float DeltaSeconds)
 				if (MovementMode == MOVE_Walking)
 				{
 					AdjustFloorHeight();
-					SetBase(CurrentFloor.HitResult.Component.Get());
+					SetBase(CurrentFloor.HitResult.Component.Get(), CurrentFloor.HitResult.BoneName);
 				}
 				else if (MovementMode == MOVE_Falling)
 				{
@@ -1130,11 +1130,12 @@ void UCharacterMovementComponent::SimulateMovement(float DeltaSeconds)
 	LastUpdateLocation = UpdatedComponent ? UpdatedComponent->GetComponentLocation() : FVector::ZeroVector;
 }
 
-void UCharacterMovementComponent::SetBase( UPrimitiveComponent* NewBase, bool bNotifyActor )
+void UCharacterMovementComponent::SetBase( UPrimitiveComponent* NewBase, const FName BoneName, bool bNotifyActor )
 {
 	if (CharacterOwner)
 	{
-		CharacterOwner->SetBase(NewBase,bNotifyActor);
+		// TODO: temporarily disabling movement on bone until networking is complete.
+		CharacterOwner->SetBase(NewBase, NAME_None /*BoneName*/, bNotifyActor);
 	}
 }
 
@@ -1143,7 +1144,7 @@ void UCharacterMovementComponent::MaybeUpdateBasedMovement(float DeltaSeconds)
 	bDeferUpdateBasedMovement = false;
 
 	UPrimitiveComponent* MovementBase = CharacterOwner->GetMovementBase();
-	if(MovementBaseUtility::UseRelativePosition(MovementBase))
+	if (MovementBaseUtility::UseRelativeLocation(MovementBase))
 	{
 		if(!MovementBase->IsSimulatingPhysics())
 		{
@@ -1173,8 +1174,8 @@ void UCharacterMovementComponent::UpdateBasedMovement(float DeltaSeconds)
 		return;
 	}
 
-	UPrimitiveComponent* MovementBase = CharacterOwner->GetMovementBase();
-	if (!MovementBaseUtility::UseRelativePosition(MovementBase))
+	const UPrimitiveComponent* MovementBase = CharacterOwner->GetMovementBase();
+	if (!MovementBaseUtility::UseRelativeLocation(MovementBase))
 	{
 		return;
 	}
@@ -1190,40 +1191,55 @@ void UCharacterMovementComponent::UpdateBasedMovement(float DeltaSeconds)
 
 	FQuat DeltaQuat = FQuat::Identity;
 	FVector DeltaPosition = FVector::ZeroVector;
-	const bool bRotationChanged = !OldBaseQuat.Equals(MovementBase->GetComponentQuat());
-	if( bRotationChanged )
+
+	FQuat NewBaseQuat;
+	FVector NewBaseLocation;
+	if (!MovementBaseUtility::GetMovementBaseTransform(MovementBase, CharacterOwner->GetBasedMovement().BoneName, NewBaseLocation, NewBaseQuat))
+	{
+		return;
+	}
+
+	// Find change in rotation
+	const bool bRotationChanged = !OldBaseQuat.Equals(NewBaseQuat);
+	if (bRotationChanged)
 	{
-		DeltaQuat = OldBaseQuat * MovementBase->GetComponentQuat().Inverse();
+		DeltaQuat = NewBaseQuat * OldBaseQuat.Inverse();
 	}
 
 	// only if base moved
-	if ( bRotationChanged || (OldBaseLocation != MovementBase->GetComponentLocation()) )
+	if (bRotationChanged || (OldBaseLocation != NewBaseLocation))
 	{
 		// Calculate new transform matrix of base actor (ignoring scale).
 		const FQuatRotationTranslationMatrix OldLocalToWorld(OldBaseQuat, OldBaseLocation);
-		const FQuatRotationTranslationMatrix NewLocalToWorld(MovementBase->GetComponentQuat(), MovementBase->GetComponentLocation());
+		const FQuatRotationTranslationMatrix NewLocalToWorld(NewBaseQuat, NewBaseLocation);
 
-		FVector RotMotion(0.f);
 		if( CharacterOwner->IsMatineeControlled() )
 		{
-			FRotationTranslationMatrix HardRelMatrix(CharacterOwner->RelativeMovement.Rotation, CharacterOwner->RelativeMovement.Location);
+			FRotationTranslationMatrix HardRelMatrix(CharacterOwner->GetBasedMovement().Rotation, CharacterOwner->GetBasedMovement().Location);
 			const FMatrix NewWorldTM = HardRelMatrix * NewLocalToWorld;
 			const FRotator NewWorldRot = bIgnoreBaseRotation ? CharacterOwner->GetActorRotation() : NewWorldTM.Rotator();
 			MoveUpdatedComponent( NewWorldTM.GetOrigin() - CharacterOwner->GetActorLocation(), NewWorldRot, true );
 		}
 		else
 		{
-			FQuat FinalQuat = DeltaQuat * CharacterOwner->GetActorQuat();
-			FQuat PawnOldQuat = CharacterOwner->GetActorQuat();
-			CharacterOwner->FaceRotation(FinalQuat.Rotator(), 0.f);
-			FinalQuat = CharacterOwner->GetActorQuat();
-
-			if( bRotationChanged )
+			FQuat FinalQuat = CharacterOwner->GetActorQuat();
+			
+			if (bRotationChanged && !bIgnoreBaseRotation)
 			{
-				const FQuat PawnDeltaRotation = PawnOldQuat * FinalQuat.Inverse();
-				FRotator FinalRotation = FinalQuat.Rotator();
-				UpdateBasedRotation(FinalRotation, PawnDeltaRotation.Rotator());
-				FinalQuat = FinalRotation.Quaternion();
+				// Apply change in rotation and pipe through FaceRotation to maintain axis restrictions
+				const FQuat PawnOldQuat = CharacterOwner->GetActorQuat();
+				FinalQuat = DeltaQuat * FinalQuat;
+				CharacterOwner->FaceRotation(FinalQuat.Rotator(), 0.f);
+				FinalQuat = CharacterOwner->GetActorQuat();
+
+				// Pipe through ControlRotation, to affect camera.
+				if (CharacterOwner->Controller)
+				{
+					const FQuat PawnDeltaRotation = FinalQuat * PawnOldQuat.Inverse();
+					FRotator FinalRotation = FinalQuat.Rotator();
+					UpdateBasedRotation(FinalRotation, PawnDeltaRotation.Rotator());
+					FinalQuat = FinalRotation.Quaternion();
+				}
 			}
 
 			// We need to offset the base of the character here, not its origin, so offset by half height
@@ -1472,14 +1488,27 @@ void UCharacterMovementComponent::SaveBaseLocation()
 		return;
 	}
 
-	UPrimitiveComponent* MovementBase = CharacterOwner->GetMovementBase();
-	if (MovementBaseUtility::UseRelativePosition(MovementBase) && !CharacterOwner->IsMatineeControlled() && !Cast<USkeletalMeshComponent>(MovementBase))
+	const UPrimitiveComponent* MovementBase = CharacterOwner->GetMovementBase();
+	if (MovementBaseUtility::UseRelativeLocation(MovementBase) && !CharacterOwner->IsMatineeControlled())
 	{
-		OldBaseLocation = MovementBase->GetComponentLocation();
-		CharacterOwner->RelativeMovement.Location = CharacterOwner->GetActorLocation() - OldBaseLocation;
+		// Read transforms into OldBaseLocation, OldBaseQuat
+		MovementBaseUtility::GetMovementBaseTransform(MovementBase, CharacterOwner->GetBasedMovement().BoneName, OldBaseLocation, OldBaseQuat);
+
+		// Location
+		const FVector RelativeLocation = CharacterOwner->GetActorLocation() - OldBaseLocation;
 
-		OldBaseQuat = MovementBase->GetComponentQuat();
-		CharacterOwner->RelativeMovement.Rotation = (FRotationMatrix(CharacterOwner->GetActorRotation()) * FQuatRotationTranslationMatrix(OldBaseQuat, FVector::ZeroVector).GetTransposed()).Rotator();
+		// Rotation
+		if (bIgnoreBaseRotation)
+		{
+			// Absolute rotation
+			CharacterOwner->SaveRelativeBasedMovement(RelativeLocation, CharacterOwner->GetActorRotation(), false);
+		}
+		else
+		{
+			// Relative rotation
+			const FRotator RelativeRotation = (FRotationMatrix(CharacterOwner->GetActorRotation()) * FQuatRotationTranslationMatrix(OldBaseQuat, FVector::ZeroVector).GetTransposed()).Rotator();
+			CharacterOwner->SaveRelativeBasedMovement(RelativeLocation, RelativeRotation, true);	
+		}
 	}
 }
 
@@ -3019,7 +3048,7 @@ void UCharacterMovementComponent::RevertMove(const FVector& OldLocation, UPrimit
 	   )
 	{
 		CurrentFloor = OldFloor;
-		SetBase(OldBase);
+		SetBase(OldBase, OldFloor.HitResult.BoneName);
 	}
 	else
 	{
@@ -3302,7 +3331,7 @@ void UCharacterMovementComponent::PhysWalking(float deltaTime, int32 Iterations)
 				}
 
 				AdjustFloorHeight();
-				SetBase(CurrentFloor.HitResult.Component.Get());
+				SetBase(CurrentFloor.HitResult.Component.Get(), CurrentFloor.HitResult.BoneName);
 			}
 			else if (CurrentFloor.HitResult.bStartPenetrating && remainingTime <= 0.f)
 			{
@@ -3481,7 +3510,7 @@ void UCharacterMovementComponent::OnTeleported()
 	UPrimitiveComponent* OldBase = CharacterOwner->GetMovementBase();
 	UPrimitiveComponent* NewBase = NULL;
 	
-	if (OldBase && CurrentFloor.IsWalkableFloor() && CurrentFloor.FloorDist <= MAX_FLOOR_DIST)
+	if (OldBase && CurrentFloor.IsWalkableFloor() && CurrentFloor.FloorDist <= MAX_FLOOR_DIST && Velocity.Z <= 0.f)
 	{
 		// Close enough to land or just keep walking.
 		NewBase = CurrentFloor.HitResult.Component.Get();
@@ -4064,14 +4093,13 @@ void UCharacterMovementComponent::FindFloor(const FVector& CapsuleLocation, FFin
 				MutableThis->bForceNextFloorCheck = !MovementBase->IsCollisionEnabled()
 				|| MovementBase->GetCollisionResponseToChannel(CollisionChannel) != ECR_Block
 				|| (MovementBase->Mobility == EComponentMobility::Movable)
+				|| MovementBaseUtility::IsDynamicBase(MovementBase)
 				|| (Cast<const ADestructibleActor>(BaseActor) != NULL);
 			}
 
 			const bool IsActorBasePendingKill = BaseActor && BaseActor->IsPendingKill();
 
-			if ( !bForceNextFloorCheck && !IsActorBasePendingKill && MovementBase &&
-				(!MovementBaseUtility::IsDynamicBase(MovementBase) || (CharacterOwner->RelativeMovement.Location == CharacterOwner->GetActorLocation() - MovementBase->GetComponentLocation()))
-				)
+			if ( !bForceNextFloorCheck && !IsActorBasePendingKill && MovementBase )
 			{
 				//UE_LOG(LogCharacterMovement, Log, TEXT("%s SKIP check for floor"), *CharacterOwner->GetName());
 				OutFloorResult = CurrentFloor;
@@ -5008,8 +5036,8 @@ void UCharacterMovementComponent::ReplicateMoveToServer(float DeltaTime, const F
 
 			if (ClientData->PendingMove->StartBase != CharacterOwner->GetMovementBase())
 			{
-				SetBase(ClientData->PendingMove->StartBase.Get());
-			}
+				SetBase(ClientData->PendingMove->StartBase.Get() /* TODO: BoneName on saved move */);
+			}			
 			CurrentFloor = ClientData->PendingMove->StartFloor;
 
 			// Now that we have reverted to the old position, prepare a new move from that position,
@@ -5114,7 +5142,7 @@ void UCharacterMovementComponent::CallServerMove
 
 	// Determine if we send absolute or relative location
 	UPrimitiveComponent* ClientMovementBase = NewMove->EndBase.Get();
-	const FVector SendLocation = MovementBaseUtility::UseRelativePosition(ClientMovementBase) ? NewMove->SavedRelativeLocation : NewMove->SavedLocation;
+	const FVector SendLocation = MovementBaseUtility::UseRelativeLocation(ClientMovementBase) ? NewMove->SavedRelativeLocation : NewMove->SavedLocation;
 
 	// send old move if it exists
 	if (OldMove)
@@ -5327,8 +5355,9 @@ void UCharacterMovementComponent::ServerMoveHandleClientError(float TimeStamp, f
 
 	// Offset may be relative to base component
 	FVector ClientLoc = RelativeClientLoc;
-	if (MovementBaseUtility::UseRelativePosition(ClientMovementBase))
+	if (MovementBaseUtility::UseRelativeLocation(ClientMovementBase))
 	{
+		// TODO: use MovementBaseUtility::GetMovementBaseTransform(ClientMovementBase) instead.
 		ClientLoc += ClientMovementBase->GetComponentLocation();
 	}
 
@@ -5342,9 +5371,10 @@ void UCharacterMovementComponent::ServerMoveHandleClientError(float TimeStamp, f
 		UPrimitiveComponent* MovementBase = CharacterOwner->GetMovementBase();
 		ServerData->PendingAdjustment.NewVel = Velocity;
 		ServerData->PendingAdjustment.NewBase = MovementBase;
-		ServerData->PendingAdjustment.bBaseRelativePosition = MovementBaseUtility::UseRelativePosition(MovementBase);
+		ServerData->PendingAdjustment.bBaseRelativePosition = MovementBaseUtility::UseRelativeLocation(MovementBase);
 		if (ServerData->PendingAdjustment.bBaseRelativePosition)
 		{
+			// TODO: use MovementBaseUtility::GetMovementBaseTransform() instead.
 			ServerData->PendingAdjustment.NewLoc = CharacterOwner->GetActorLocation() - MovementBase->GetComponentLocation();
 		}
 		else
@@ -5384,7 +5414,7 @@ void UCharacterMovementComponent::ServerMoveHandleClientError(float TimeStamp, f
 				ApplyNetworkMovementMode(ClientMovementMode);
 
 				// Update base and floor at new location.
-				SetBase(ClientMovementBase);
+				SetBase(ClientMovementBase /* TODO: client bone? */);
 				UpdateFloorFromAdjustment();
 
 				// Even if base has not changed, we need to recompute the relative offsets (since we've moved).
@@ -5438,7 +5468,7 @@ void UCharacterMovementComponent::UpdateFloorFromAdjustment()
 
 	// If walking, try to update the cached floor so it is current. This is necessary for UpdateBasedMovement() and MoveAlongFloor() to work properly.
 	// If base is now NULL, presumably we are no longer walking. If we had a valid floor but don't find one now, we'll likely start falling.
-	if (CharacterOwner->GetMovementBase() && Velocity.Z <= 0.f)
+	if (CharacterOwner->GetMovementBase())
 	{
 		FindFloor(UpdatedComponent->GetComponentLocation(), CurrentFloor, false);
 	}
@@ -5589,6 +5619,7 @@ void UCharacterMovementComponent::ClientAdjustPosition_Implementation
 	//  Received Location is relative to non-world geometry (moving) bases
 	if (bBaseRelativePosition)
 	{
+		// TODO: fix if there is a bone
 		NewLocation += NewBase->GetComponentLocation();
 	}
 
@@ -5627,7 +5658,7 @@ void UCharacterMovementComponent::ClientAdjustPosition_Implementation
 			FinalBase = CurrentFloor.HitResult.Component.Get();
 		}
 	}
-	SetBase(FinalBase);
+	SetBase(FinalBase /* TODO: bone? */);
 
 	// Update floor at new location
 	UpdateFloorFromAdjustment();
@@ -6169,6 +6200,7 @@ void FSavedMove_Character::Clear()
 	StartFloor = FFindFloorResult();
 	StartRotation = FRotator::ZeroRotator;
 	StartControlRotation = FRotator::ZeroRotator;
+	StartBaseRotation = FQuat::Identity;
 	StartCapsuleRadius = 0.f;
 	StartCapsuleHalfHeight = 0.f;
 
@@ -6213,14 +6245,18 @@ void FSavedMove_Character::SetInitialPosition(ACharacter* Character)
 	StartVelocity = Character->CharacterMovement->Velocity;
 	UPrimitiveComponent* const MovementBase = Character->GetMovementBase();
 	StartBase = MovementBase;
+	StartBaseRotation = FQuat::Identity;
 	StartFloor = Character->CharacterMovement->CurrentFloor;
 	CustomTimeDilation = Character->CustomTimeDilation;
 
 	StartControlRotation = Character->GetControlRotation().Clamp();
 
-	if (MovementBaseUtility::UseRelativePosition(MovementBase))
+	if (MovementBaseUtility::UseRelativeLocation(MovementBase))
 	{
-		StartRelativeLocation = Character->RelativeMovement.Location;
+		StartRelativeLocation = Character->GetBasedMovement().Location;
+		// TODO: bone
+		FVector StartBaseLocation;
+		MovementBaseUtility::GetMovementBaseTransform(MovementBase, NAME_None, StartBaseLocation, StartBaseRotation);
 	}
 
 	Character->CapsuleComponent->GetScaledCapsuleSize(StartCapsuleRadius, StartCapsuleHalfHeight);
@@ -6234,9 +6270,9 @@ void FSavedMove_Character::PostUpdate(ACharacter* Character, FSavedMove_Characte
 		SavedRotation = Character->GetActorRotation();
 		UPrimitiveComponent* const MovementBase = Character->GetMovementBase();
 		EndBase = MovementBase;
-		if (MovementBaseUtility::UseRelativePosition(MovementBase))
+		if (MovementBaseUtility::UseRelativeLocation(MovementBase))
 		{
-			SavedRelativeLocation = Character->RelativeMovement.Location;
+			SavedRelativeLocation = Character->GetBasedMovement().Location;
 		}
 
 		SavedControlRotation = Character->GetControlRotation().Clamp();
@@ -6282,8 +6318,9 @@ bool FSavedMove_Character::IsImportantMove(const FSavedMovePtr& LastAckedMove) c
 
 FVector FSavedMove_Character::GetStartLocation(const class UPrimitiveComponent* MovementBase) const
 {
-	if (MovementBaseUtility::UseRelativePosition(MovementBase))
+	if (MovementBaseUtility::UseRelativeLocation(MovementBase))
 	{
+		// TODO: use MovementBaseUtility::GetMovementBaseTransform() instead.
 		return MovementBase->GetComponentLocation() + StartRelativeLocation;
 	}
 
@@ -6316,6 +6353,7 @@ bool FSavedMove_Character::CanCombineWith(const FSavedMovePtr& NewMove, ACharact
 			&& MovementMode == NewMove->MovementMode
 			&& StartCapsuleRadius == NewMove->StartCapsuleRadius
 			&& StartCapsuleHalfHeight == NewMove->StartCapsuleHalfHeight
+			&& StartBaseRotation.Equals(NewMove->StartBaseRotation) // only if base hasn't rotated
 			&& (CustomTimeDilation == NewMove->CustomTimeDilation);
 	}
 	else
@@ -6328,6 +6366,7 @@ bool FSavedMove_Character::CanCombineWith(const FSavedMovePtr& NewMove, ACharact
 			&& StartCapsuleRadius == NewMove->StartCapsuleRadius
 			&& StartCapsuleHalfHeight == NewMove->StartCapsuleHalfHeight
 			&& ((AccelNormal | NewMove->AccelNormal) > 0.99f)
+			&& StartBaseRotation.Equals(NewMove->StartBaseRotation) // only if base hasn't rotated
 			&& (CustomTimeDilation == NewMove->CustomTimeDilation);
 	}
 }
diff --git a/Engine/Source/Runtime/Engine/Private/Pawn.cpp b/Engine/Source/Runtime/Engine/Private/Pawn.cpp
--- a/Engine/Source/Runtime/Engine/Private/Pawn.cpp
+++ b/Engine/Source/Runtime/Engine/Private/Pawn.cpp
@@ -878,10 +878,10 @@ void APawn::PostNetReceiveVelocity(const FVector& NewVelocity)
 	}
 }
 
-void APawn::PostNetReceiveLocation()
+void APawn::PostNetReceiveLocationAndRotation()
 {
 	// always consider Location as changed if we were spawned this tick as in that case our replicated Location was set as part of spawning, before PreNetReceive()
-	if( (ReplicatedMovement.Location == GetActorLocation()) && (CreationTime != GetWorld()->TimeSeconds) )
+	if( (ReplicatedMovement.Location == GetActorLocation() && ReplicatedMovement.Rotation == GetActorRotation()) && (CreationTime != GetWorld()->TimeSeconds) )
 	{
 		return;
 	}
@@ -891,8 +891,9 @@ void APawn::PostNetReceiveLocation()
 		// Correction to make sure pawn doesn't penetrate floor after replication rounding
 		ReplicatedMovement.Location.Z += 0.01f;
 
-		FVector OldLocation = GetActorLocation();
-		TeleportTo( ReplicatedMovement.Location, GetActorRotation(), false, true );
+		const FVector OldLocation = GetActorLocation();
+		TeleportTo( ReplicatedMovement.Location, ReplicatedMovement.Rotation, false, true );
+		// SetActorLocationAndRotation(ReplicatedMovement.Location, ReplicatedMovement.Rotation); <-- preferred, but awaiting answer to question about UpdateNavOctree() missing in SceneComponent::MoveComponent
 
 		INetworkPredictionInterface* PredictionInterface = InterfaceCast<INetworkPredictionInterface>(GetMovementComponent());
 		if (PredictionInterface)
