Changed FMaterialInstanceBasePropertyOverrides native serialization to proper tagged serialization.

diff --git a/Engine/Source/Editor/UnrealEd/Private/PreviewMaterial.cpp b/Engine/Source/Editor/UnrealEd/Private/PreviewMaterial.cpp
--- a/Engine/Source/Editor/UnrealEd/Private/PreviewMaterial.cpp
+++ b/Engine/Source/Editor/UnrealEd/Private/PreviewMaterial.cpp
@@ -574,14 +574,8 @@ void UMaterialEditorInstanceConstant::CopyToSourceInstance()
 			}
 		}
 
-		//Create the source instance base overrides too. The source instance owns this from this point on.
-		if( !SourceInstance->BasePropertyOverrides )
-		{
-			SourceInstance->BasePropertyOverrides = new FMaterialInstanceBasePropertyOverrides();
-			SourceInstance->BasePropertyOverrides->Init(*SourceInstance);
-		}
 		//Check for changes to see if we need to force a recompile
-		bool bForceRecompile = SourceInstance->BasePropertyOverrides->Update(BasePropertyOverrides);
+		bool bForceRecompile = SourceInstance->BasePropertyOverrides.Update(BasePropertyOverrides);
 		bForceRecompile |= SourceInstance->bOverrideBaseProperties != bOverrideBaseProperties;
 		SourceInstance->bOverrideBaseProperties = bOverrideBaseProperties;
 		
@@ -680,10 +674,7 @@ void UMaterialEditorInstanceConstant::SetSourceInstance(UMaterialInstanceConstan
 	//Init the base property overrides.
 	BasePropertyOverrides.Init(*SourceInstance);
 	bOverrideBaseProperties = SourceInstance->bOverrideBaseProperties;
-	if( SourceInstance->BasePropertyOverrides )
-	{
-		BasePropertyOverrides = *SourceInstance->BasePropertyOverrides;
-	}
+	BasePropertyOverrides = SourceInstance->BasePropertyOverrides;
 
 	// Copy the Lightmass settings...
 	LightmassSettings.CastShadowAsMasked.bOverride = SourceInstance->GetOverrideCastShadowAsMasked();
diff --git a/Engine/Source/Runtime/Core/Public/UObject/ObjectVersion.h b/Engine/Source/Runtime/Core/Public/UObject/ObjectVersion.h
--- a/Engine/Source/Runtime/Core/Public/UObject/ObjectVersion.h
+++ b/Engine/Source/Runtime/Core/Public/UObject/ObjectVersion.h
@@ -576,6 +576,8 @@ enum EUnrealEngineObjectUE4Version
 	VER_UE4_ADD_LB_WEIGHTBLEND,
 	// Add root component to an foliage actor, all foliage cluster components will be attached to a root
 	VER_UE4_ADD_ROOTCOMPONENT_TO_FOLIAGEACTOR,
+	// FMaterialInstanceBasePropertyOverrides didn't use proper UObject serialize
+	VER_UE4_FIX_MATERIAL_PROPERTY_OVERRIDE_SERIALIZE,
 
 	// -----<new versions can be added before this line>-------------------------------------------------
 	// - this needs to be the last line (see note below)
diff --git a/Engine/Source/Runtime/Engine/Classes/Engine/EngineTypes.h b/Engine/Source/Runtime/Engine/Classes/Engine/EngineTypes.h
--- a/Engine/Source/Runtime/Engine/Classes/Engine/EngineTypes.h
+++ b/Engine/Source/Runtime/Engine/Classes/Engine/EngineTypes.h
@@ -65,6 +65,7 @@ enum ESceneDepthPriorityGroup
 };
 
 /** Note: This is mirrored in Lightmass, be sure to update the blend mode structure and logic there if this changes. */
+// Note: Check UMaterialInstance::Serialize if changed!!
 UENUM(BlueprintType)
 enum EBlendMode
 {
@@ -108,6 +109,7 @@ enum ETrailWidthMode
 	ETrailWidthMode_FromSecond UMETA(DisplayName = "From Second Socket"),
 };
 
+// Note: Check UMaterialInstance::Serialize if changed!!
 UENUM()
 enum EMaterialShadingModel
 {
diff --git a/Engine/Source/Runtime/Engine/Classes/Materials/MaterialInstance.h b/Engine/Source/Runtime/Engine/Classes/Materials/MaterialInstance.h
--- a/Engine/Source/Runtime/Engine/Classes/Materials/MaterialInstance.h
+++ b/Engine/Source/Runtime/Engine/Classes/Materials/MaterialInstance.h
@@ -1,15 +1,15 @@
 // Copyright 1998-2014 Epic Games, Inc. All Rights Reserved.
 
-
 #pragma once
+
 #include "Materials/MaterialInterface.h"
+#include "Materials/MaterialInstanceBasePropertyOverrides.h"
 #include "StaticParameterSet.h"
 #include "MaterialInstance.generated.h"
 
 //
 // Forward declarations.
 //
-struct FMaterialInstanceBasePropertyOverrides;
 class FMaterialShaderMap;
 class FMaterialShaderMapId;
 class FSHAHash;
@@ -169,7 +169,8 @@ class UMaterialInstance : public UMaterialInterface
 	UPROPERTY(EditAnywhere, BlueprintReadOnly, Category=MaterialInstance)
 	bool bOverrideBaseProperties;
 
-	FMaterialInstanceBasePropertyOverrides* BasePropertyOverrides;
+	UPROPERTY(EditAnywhere, BlueprintReadOnly, Category=MaterialInstance, meta=(editcondition = "bOverrideBaseProperties"))
+	struct FMaterialInstanceBasePropertyOverrides BasePropertyOverrides;
 
 	/** 
 	 * FMaterialRenderProxy derivatives that represent this material instance to the renderer, when the renderer needs to fetch parameter values. 
diff --git a/Engine/Source/Runtime/Engine/Classes/Materials/MaterialInstanceBasePropertyOverrides.h b/Engine/Source/Runtime/Engine/Classes/Materials/MaterialInstanceBasePropertyOverrides.h
--- a/Engine/Source/Runtime/Engine/Classes/Materials/MaterialInstanceBasePropertyOverrides.h
+++ b/Engine/Source/Runtime/Engine/Classes/Materials/MaterialInstanceBasePropertyOverrides.h
@@ -51,25 +51,4 @@ struct ENGINE_API FMaterialInstanceBasePropertyOverrides
 
 	bool operator==(const FMaterialInstanceBasePropertyOverrides& Other)const;
 	bool operator!=(const FMaterialInstanceBasePropertyOverrides& Other)const;
-
-	friend FArchive& operator<<(FArchive& Ar, FMaterialInstanceBasePropertyOverrides& Ref)
-	{
-		Ar	<< Ref.bOverride_OpacityMaskClipValue << Ref.OpacityMaskClipValue;
-
-		if( Ar.UE4Ver() >= VER_UE4_MATERIAL_INSTANCE_BASE_PROPERTY_OVERRIDES_PHASE_2 )
-		{
-			Ar	<< Ref.bOverride_BlendMode << Ref.BlendMode
-				<< Ref.bOverride_ShadingModel << Ref.ShadingModel
-				<< Ref.bOverride_TwoSided;
-
-			bool bTwoSided = Ref.TwoSided != 0;
-			Ar << bTwoSided;
-			if( Ar.IsLoading() )
-			{
-				Ref.TwoSided = bTwoSided;
-			}
-		}
-
-		return Ar;
-	}
 };
\ No newline at end of file
diff --git a/Engine/Source/Runtime/Engine/Private/Materials/MaterialInstance.cpp b/Engine/Source/Runtime/Engine/Private/Materials/MaterialInstance.cpp
--- a/Engine/Source/Runtime/Engine/Private/Materials/MaterialInstance.cpp
+++ b/Engine/Source/Runtime/Engine/Private/Materials/MaterialInstance.cpp
@@ -1582,24 +1582,40 @@ void UMaterialInstance::Serialize(FArchive& Ar)
 
 	if (Ar.UE4Ver() >= VER_UE4_MATERIAL_INSTANCE_BASE_PROPERTY_OVERRIDES )
 	{
-		Ar << bOverrideBaseProperties;
-		bool bHasPropertyOverrides = NULL != BasePropertyOverrides;
-		Ar << bHasPropertyOverrides;
-		if( bHasPropertyOverrides )
+		if( Ar.UE4Ver() < VER_UE4_FIX_MATERIAL_PROPERTY_OVERRIDE_SERIALIZE )
 		{
-			if( !BasePropertyOverrides )
+			// awful old native serialize of FMaterialInstanceBasePropertyOverrides UStruct
+			Ar << bOverrideBaseProperties;
+			bool bHasPropertyOverrides = false;
+			Ar << bHasPropertyOverrides;
+			if( bHasPropertyOverrides )
 			{
-				BasePropertyOverrides = new FMaterialInstanceBasePropertyOverrides();
-				BasePropertyOverrides->Init(*this);
+				Ar << BasePropertyOverrides.bOverride_OpacityMaskClipValue << BasePropertyOverrides.OpacityMaskClipValue;
+
+				if( Ar.UE4Ver() >= VER_UE4_MATERIAL_INSTANCE_BASE_PROPERTY_OVERRIDES_PHASE_2 )
+				{
+					Ar	<< BasePropertyOverrides.bOverride_BlendMode << BasePropertyOverrides.BlendMode
+						<< BasePropertyOverrides.bOverride_ShadingModel << BasePropertyOverrides.ShadingModel
+						<< BasePropertyOverrides.bOverride_TwoSided;
+
+					bool bTwoSided;
+					Ar << bTwoSided;
+					BasePropertyOverrides.TwoSided = bTwoSided;
+
+					// unrelated but closest change to bug
+					if( Ar.UE4Ver() < VER_UE4_STATIC_SHADOW_DEPTH_MAPS )
+					{
+						// switched enum order
+						switch( BasePropertyOverrides.ShadingModel )
+						{
+							case MSM_Unlit:			BasePropertyOverrides.ShadingModel = MSM_DefaultLit; break;
+							case MSM_DefaultLit:	BasePropertyOverrides.ShadingModel = MSM_Unlit; break;
+						}
+					}
+				}
 			}
-			Ar << *BasePropertyOverrides;
 		}
 	}
-	else
-	{
-		bOverrideBaseProperties = false;
-		BasePropertyOverrides = NULL;
-	}
 }
 
 void UMaterialInstance::PostLoad()
@@ -1639,7 +1655,7 @@ void UMaterialInstance::PostLoad()
 	}
 
 	// Update bHasStaticPermutationResource in case the parent was not found
-	bHasStaticPermutationResource = (!StaticParameters.IsEmpty() || (bOverrideBaseProperties && BasePropertyOverrides)) && Parent;
+	bHasStaticPermutationResource = (!StaticParameters.IsEmpty() || bOverrideBaseProperties) && Parent;
 
 	STAT(double MaterialLoadTime = 0);
 	{
@@ -1734,12 +1750,6 @@ void UMaterialInstance::FinishDestroy()
 
 	ClearAllCachedCookedPlatformData();
 
-	if( BasePropertyOverrides )
-	{
-		delete BasePropertyOverrides;
-		BasePropertyOverrides = NULL;
-	}
-
 	Super::FinishDestroy();
 }
 
@@ -1935,7 +1945,7 @@ void UMaterialInstance::UpdateStaticPermutation(const FStaticParameterSet& NewPa
 	TrimToOverriddenOnly(CompareParameters.StaticComponentMaskParameters);
 	TrimToOverriddenOnly(CompareParameters.TerrainLayerWeightParameters);
 
-	const bool bWantsStaticPermutationResource = (!CompareParameters.IsEmpty() || (BasePropertyOverrides && bOverrideBaseProperties) || bForceRecompile) && Parent;
+	const bool bWantsStaticPermutationResource = (!CompareParameters.IsEmpty() || bOverrideBaseProperties || bForceRecompile) && Parent;
 
 	if (bForceRecompile || bHasStaticPermutationResource != bWantsStaticPermutationResource || StaticParameters != CompareParameters)
 	{
@@ -2315,9 +2325,9 @@ void UMaterialInstance::GetBasePropertyOverridesHash(FSHAHash& OutHash)const
 {
 	FSHA1 HashState;
 
-	if( bOverrideBaseProperties && BasePropertyOverrides )
+	if( bOverrideBaseProperties )
 	{
-		BasePropertyOverrides->UpdateHash(HashState);
+		BasePropertyOverrides.UpdateHash(HashState);
 	}
 	
 	HashState.Final();
@@ -2327,18 +2337,18 @@ void UMaterialInstance::GetBasePropertyOverridesHash(FSHAHash& OutHash)const
 float UMaterialInstance::GetOpacityMaskClipValue_Internal() const
 {
 	checkSlow(IsInGameThread());
-	if( bOverrideBaseProperties && BasePropertyOverrides && BasePropertyOverrides->bOverride_OpacityMaskClipValue )
+	if( bOverrideBaseProperties && BasePropertyOverrides.bOverride_OpacityMaskClipValue )
 	{
-		return BasePropertyOverrides->OpacityMaskClipValue;
+		return BasePropertyOverrides.OpacityMaskClipValue;
 	}
 	return GetMaterial()->GetOpacityMaskClipValue();
 }
 EBlendMode UMaterialInstance::GetBlendMode_Internal() const
 {
 	checkSlow(IsInGameThread());
-	if( bOverrideBaseProperties && BasePropertyOverrides && BasePropertyOverrides->bOverride_BlendMode )
+	if( bOverrideBaseProperties && BasePropertyOverrides.bOverride_BlendMode )
 	{
-		return BasePropertyOverrides->BlendMode;
+		return BasePropertyOverrides.BlendMode;
 	}
 	return GetMaterial()->GetBlendMode();
 }
@@ -2346,9 +2356,9 @@ EBlendMode UMaterialInstance::GetBlendMode_Internal() const
 EMaterialShadingModel UMaterialInstance::GetShadingModel_Internal() const
 {
 	checkSlow(IsInGameThread());
-	if( bOverrideBaseProperties && BasePropertyOverrides && BasePropertyOverrides->bOverride_ShadingModel )
+	if( bOverrideBaseProperties && BasePropertyOverrides.bOverride_ShadingModel )
 	{
-		return BasePropertyOverrides->ShadingModel;
+		return BasePropertyOverrides.ShadingModel;
 	}
 	return GetMaterial()->GetShadingModel();
 }
@@ -2356,18 +2366,18 @@ EMaterialShadingModel UMaterialInstance::GetShadingModel_Internal() const
 bool UMaterialInstance::IsTwoSided_Internal() const
 {
 	checkSlow(IsInGameThread());
-	if( bOverrideBaseProperties && BasePropertyOverrides && BasePropertyOverrides->bOverride_TwoSided )
+	if( bOverrideBaseProperties && BasePropertyOverrides.bOverride_TwoSided )
 	{
-		return BasePropertyOverrides->TwoSided != 0;
+		return BasePropertyOverrides.TwoSided != 0;
 	}
 	return GetMaterial()->IsTwoSided();
 }
 
 bool UMaterialInstance::GetOpacityMaskClipValueOverride(float& OutResult) const
 {
-	if( bOverrideBaseProperties && BasePropertyOverrides && BasePropertyOverrides->bOverride_OpacityMaskClipValue )
+	if( bOverrideBaseProperties && BasePropertyOverrides.bOverride_OpacityMaskClipValue )
 	{
-		OutResult = BasePropertyOverrides->OpacityMaskClipValue;
+		OutResult = BasePropertyOverrides.OpacityMaskClipValue;
 		return true;
 	}
 	return false;
@@ -2375,9 +2385,9 @@ bool UMaterialInstance::GetOpacityMaskClipValueOverride(float& OutResult) const
 
 bool UMaterialInstance::GetBlendModeOverride(EBlendMode& OutResult) const
 {
-	if( bOverrideBaseProperties && BasePropertyOverrides && BasePropertyOverrides->bOverride_BlendMode )
+	if( bOverrideBaseProperties && BasePropertyOverrides.bOverride_BlendMode )
 	{
-		OutResult = BasePropertyOverrides->BlendMode;
+		OutResult = BasePropertyOverrides.BlendMode;
 		return true;
 	}
 	return false;
@@ -2385,9 +2395,9 @@ bool UMaterialInstance::GetBlendModeOverride(EBlendMode& OutResult) const
 
 bool UMaterialInstance::GetShadingModelOverride(EMaterialShadingModel& OutResult) const
 {
-	if( bOverrideBaseProperties && BasePropertyOverrides && BasePropertyOverrides->bOverride_ShadingModel )
+	if( bOverrideBaseProperties && BasePropertyOverrides.bOverride_ShadingModel )
 	{
-		OutResult = BasePropertyOverrides->ShadingModel;
+		OutResult = BasePropertyOverrides.ShadingModel;
 		return true;
 	}
 	return false;
@@ -2395,9 +2405,9 @@ bool UMaterialInstance::GetShadingModelOverride(EMaterialShadingModel& OutResult
 
 bool UMaterialInstance::IsTwoSidedOverride(bool& OutResult) const
 {
-	if( bOverrideBaseProperties && BasePropertyOverrides && BasePropertyOverrides->bOverride_TwoSided )
+	if( bOverrideBaseProperties && BasePropertyOverrides.bOverride_TwoSided )
 	{
-		OutResult = BasePropertyOverrides->TwoSided != 0;
+		OutResult = BasePropertyOverrides.TwoSided != 0;
 		return true;
 	}
 	return false;
