#ue4 - Add UI lower limits to many CharacterMovementComponent properties, and reorder them for more a logical grouping.

diff --git a/Engine/Source/Runtime/Engine/Classes/GameFramework/CharacterMovementComponent.h b/Engine/Source/Runtime/Engine/Classes/GameFramework/CharacterMovementComponent.h
--- a/Engine/Source/Runtime/Engine/Classes/GameFramework/CharacterMovementComponent.h
+++ b/Engine/Source/Runtime/Engine/Classes/GameFramework/CharacterMovementComponent.h
@@ -136,15 +136,15 @@ protected:
 public:
 
 	/** Maximum height character can step up */
-	UPROPERTY(Category="Character Movement", EditAnywhere, BlueprintReadWrite)
+	UPROPERTY(Category="Character Movement", EditAnywhere, BlueprintReadWrite, meta=(ClampMin="0", UIMin="0"))
 	float MaxStepHeight;
 
 	/** Initial velocity (instantaneous vertical acceleration) when jumping. */
-	UPROPERTY(Category="Character Movement", EditAnywhere, BlueprintReadWrite, meta=(DisplayName="Jump Z Velocity"))
+	UPROPERTY(Category="Character Movement", EditAnywhere, BlueprintReadWrite, meta=(DisplayName="Jump Z Velocity", ClampMin="0", UIMin="0"))
 	float JumpZVelocity;
 
 	/** Fraction of JumpZVelocity to use when automatically "jumping off" of a base actor that's not allowed to be a base for a character. (For example, if you're not allowed to stand on other players.) */
-	UPROPERTY(Category="Character Movement", EditAnywhere, BlueprintReadWrite, AdvancedDisplay)
+	UPROPERTY(Category="Character Movement", EditAnywhere, BlueprintReadWrite, AdvancedDisplay, meta=(ClampMin="0", UIMin="0"))
 	float JumpOffJumpZFactor;
 
 private:
@@ -167,7 +167,7 @@ public:
 	 * Don't allow the character to perch on the edge of a surface if the contact is this close to the edge of the capsule.
 	 * Note that characters will not fall off if they are within MaxStepHeight of a walkable surface below.
 	 */
-	UPROPERTY(Category="Character Movement", EditAnywhere, BlueprintReadWrite, AdvancedDisplay)
+	UPROPERTY(Category="Character Movement", EditAnywhere, BlueprintReadWrite, AdvancedDisplay, meta=(ClampMin="0", UIMin="0"))
 	float PerchRadiusThreshold;
 
 	/**
@@ -175,7 +175,7 @@ public:
 	 * Note that we still enforce MaxStepHeight to start the step up; this just allows the character to hang off the edge or step slightly higher off the floor.
 	 * (@see PerchRadiusThreshold)
 	 */
-	UPROPERTY(Category="Character Movement", EditAnywhere, BlueprintReadWrite, AdvancedDisplay)
+	UPROPERTY(Category="Character Movement", EditAnywhere, BlueprintReadWrite, AdvancedDisplay, meta=(ClampMin="0", UIMin="0"))
 	float PerchAdditionalHeight;
 
 	/** Actor's current physics mode. This is automatically replicated through the CharacterOwner and for client-server movement functions. */
@@ -196,34 +196,69 @@ public:
 	UPROPERTY(Category="Character Movement", EditAnywhere, BlueprintReadWrite)
 	float GravityScale;
 
-	/** Coefficient of friction. This property allows you to control how much friction pawns using CharacterMovement move across the ground. This can be used to simulate slippery surfaces such as ice or oil by changing the value (possibly based on the material pawn is standing on). */
-	UPROPERTY(Category="Character Movement", EditAnywhere, BlueprintReadWrite)
+	/**
+	 * Coefficient of friction. This property allows you to control how much friction is applied when moving across the ground, applying an opposing force that scales with current velocity.
+	 * This can be used to simulate slippery surfaces such as ice or oil by changing the value (possibly based on the material pawn is standing on).
+	 * See also: BrakingDecelerationWalking
+	 */
+	UPROPERTY(Category="Character Movement", EditAnywhere, BlueprintReadWrite, meta=(ClampMin="0", UIMin="0"))
 	float GroundFriction;
 
 	/** The maximum ground speed when walking. Also determines maximum lateral speed when falling. */
-	UPROPERTY(Category="Character Movement", EditAnywhere, BlueprintReadWrite)
+	UPROPERTY(Category="Character Movement", EditAnywhere, BlueprintReadWrite, meta=(ClampMin="0", UIMin="0"))
 	float MaxWalkSpeed;
 
 	/** The maximum ground speed when walking and crouched. */
-	UPROPERTY(Category="Character Movement", EditAnywhere, BlueprintReadWrite)
+	UPROPERTY(Category="Character Movement", EditAnywhere, BlueprintReadWrite, meta=(ClampMin="0", UIMin="0"))
 	float MaxWalkSpeedCrouched;
 
 	/** The maximum speed when using Custom movement mode. */
-	UPROPERTY(Category="Character Movement", EditAnywhere, BlueprintReadWrite)
+	UPROPERTY(Category="Character Movement", EditAnywhere, BlueprintReadWrite, meta=(ClampMin="0", UIMin="0"))
 	float MaxCustomMovementSpeed;
 
-	/** Collision half-height when crouching (component scale is applied separately) */
-	UPROPERTY(Category="Character Movement", EditAnywhere, BlueprintReadOnly)
-	float CrouchedHalfHeight;
-
 	/** The maximum swimming speed. */
-	UPROPERTY(Category="Character Movement", EditAnywhere, BlueprintReadWrite)
+	UPROPERTY(Category="Character Movement", EditAnywhere, BlueprintReadWrite, meta=(ClampMin="0", UIMin="0"))
 	float MaxSwimSpeed;
 
 	/** The maximum flying speed. */
-	UPROPERTY(Category="Character Movement", EditAnywhere, BlueprintReadWrite)
+	UPROPERTY(Category="Character Movement", EditAnywhere, BlueprintReadWrite, meta=(ClampMin="0", UIMin="0"))
 	float MaxFlySpeed;
 
+	/** Max Acceleration (rate of change of velocity) */
+	UPROPERTY(Category="Character Movement", EditAnywhere, BlueprintReadWrite, meta=(ClampMin="0", UIMin="0"))
+	float MaxAcceleration;
+
+	/**
+	 * Deceleration when walking and not applying acceleration.
+	 * See also: GroundFriction
+	 */
+	UPROPERTY(Category="Character Movement", EditAnywhere, BlueprintReadWrite, meta=(ClampMin="0", UIMin="0"))
+	float BrakingDecelerationWalking;
+
+	/** Lateral deceleration when falling and not applying acceleration. */
+	UPROPERTY(Category="Character Movement", EditAnywhere, BlueprintReadWrite, meta=(ClampMin="0", UIMin="0"))
+	float BrakingDecelerationFalling;
+
+	/** Deceleration when swimming and not applying acceleration. */
+	UPROPERTY(Category="Character Movement", EditAnywhere, BlueprintReadWrite, meta=(ClampMin="0", UIMin="0"))
+	float BrakingDecelerationSwimming;
+
+	/** Deceleration when flying and not applying acceleration. */
+	UPROPERTY(Category="Character Movement", EditAnywhere, BlueprintReadWrite, meta=(ClampMin="0", UIMin="0"))
+	float BrakingDecelerationFlying;
+
+	/** Amount of lateral movement control available to the pawn when falling. 0 = no control, 1 = full control at max speed of MaxWalkSpeed. */
+	UPROPERTY(Category="Character Movement", EditAnywhere, BlueprintReadWrite, meta=(ClampMin="0", UIMin="0"))
+	float AirControl;
+
+	/** Friction to apply to lateral air movement when falling. */
+	UPROPERTY(Category="Character Movement", EditAnywhere, BlueprintReadWrite, meta=(ClampMin="0", UIMin="0"))
+	float FallingLateralFriction;
+
+	/** Collision half-height when crouching (component scale is applied separately) */
+	UPROPERTY(Category="Character Movement", EditAnywhere, BlueprintReadOnly, meta=(ClampMin="0", UIMin="0"))
+	float CrouchedHalfHeight;
+
 	/** Water buoyancy. A ratio (1.0 = neutral buoyancy, 0.0 = no buoyancy) */
 	UPROPERTY(Category="Character Movement", EditAnywhere, BlueprintReadWrite)
 	float Buoyancy;
@@ -295,23 +330,15 @@ public:
 	class UPrimitiveComponent* DeferredUpdatedMoveComponent;
 
 	/** Maximum step height for getting out of water */
-	UPROPERTY(Category="Character Movement", EditAnywhere, BlueprintReadWrite, AdvancedDisplay)
+	UPROPERTY(Category="Character Movement", EditAnywhere, BlueprintReadWrite, AdvancedDisplay, meta=(ClampMin="0", UIMin="0"))
 	float MaxOutOfWaterStepHeight;
 
 	/** Z velocity applied when pawn tries to get out of water */
 	UPROPERTY(Category="Character Movement", EditAnywhere, BlueprintReadWrite, AdvancedDisplay)
 	float OutofWaterZ;
 
-	/** Amount of lateral movement control available to the pawn when falling. 0 = no control, 1 = full control at max speed of MaxWalkSpeed. */
-	UPROPERTY(Category="Character Movement", EditAnywhere, BlueprintReadWrite)
-	float AirControl;
-
-	/** Friction to apply to lateral air movement when falling. */
-	UPROPERTY(Category="Character Movement", EditAnywhere, BlueprintReadWrite)
-	float FallingLateralFriction;
-
 	/** Mass of pawn (for when momentum is imparted to it). */
-	UPROPERTY(Category="Character Movement", EditAnywhere, BlueprintReadWrite)
+	UPROPERTY(Category="Character Movement", EditAnywhere, BlueprintReadWrite, meta=(ClampMin="0", UIMin="0"))
 	float Mass;
 	
 	/** If enabled, the player will interact with physics objects when walking into them. */
@@ -434,26 +461,6 @@ public:
 	UPROPERTY(Category="Character Movement", EditAnywhere, BlueprintReadWrite, AdvancedDisplay, meta=(ClampMin="1", ClampMax="20", UIMin="1", UIMax="20"))
 	int32 MaxSimulationIterations;
 
-	/** Max Acceleration (rate of change of velocity) */
-	UPROPERTY(Category="Character Movement", EditAnywhere, BlueprintReadWrite)
-	float MaxAcceleration;
-
-	/** Deceleration when walking and not applying acceleration. */
-	UPROPERTY(Category="Character Movement", EditAnywhere, BlueprintReadWrite)
-	float BrakingDecelerationWalking;
-
-	/** Lateral deceleration when falling and not applying acceleration. */
-	UPROPERTY(Category="Character Movement", EditAnywhere, BlueprintReadWrite)
-	float BrakingDecelerationFalling;
-
-	/** Deceleration when swimming and not applying acceleration. */
-	UPROPERTY(Category="Character Movement", EditAnywhere, BlueprintReadWrite)
-	float BrakingDecelerationSwimming;
-
-	/** Deceleration when flying and not applying acceleration. */
-	UPROPERTY(Category="Character Movement", EditAnywhere, BlueprintReadWrite)
-	float BrakingDecelerationFlying;
-
 	/** Used in determining if pawn is going off ledge.  If the ledge is "shorter" than this value then the pawn will be able to walk off it. **/
 	UPROPERTY(Category="Character Movement", EditAnywhere, BlueprintReadWrite, AdvancedDisplay)
 	float LedgeCheckThreshold;
diff --git a/Engine/Source/Runtime/Engine/Private/CharacterMovementComponent.cpp b/Engine/Source/Runtime/Engine/Private/CharacterMovementComponent.cpp
--- a/Engine/Source/Runtime/Engine/Private/CharacterMovementComponent.cpp
+++ b/Engine/Source/Runtime/Engine/Private/CharacterMovementComponent.cpp
@@ -238,6 +238,7 @@ UCharacterMovementComponent::UCharacterMovementComponent(const class FPostConstr
 	bRequestedMoveUseAcceleration = true;
 	bUseRVOAvoidance = false;
 	bUseRVOPostProcess = false;
+	AvoidanceLockVelocity = FVector::ZeroVector;
 	AvoidanceLockTimer = 0.0f;
 	AvoidanceGroup.bGroup0 = true;
 	GroupsToAvoid.Packed = 0xFFFFFFFF;
@@ -2387,7 +2388,8 @@ void UCharacterMovementComponent::ApplyVelocityBraking(float DeltaTime, float Fr
 		return;
 	}
 
-	const	FVector OldVel = Velocity;
+	Friction = FMath::Max(0.f, Friction);
+	const FVector OldVel = Velocity;
 	FVector SumVel = FVector::ZeroVector;
 
 	// subdivide braking to get reasonably consistent results at lower frame rates
@@ -2396,7 +2398,7 @@ void UCharacterMovementComponent::ApplyVelocityBraking(float DeltaTime, float Fr
 	const float TimeStep = 0.03f;
 
 	// Decelerate to brake to a stop
-	const FVector RevAccel = (-1.f * BrakingDeceleration) * SafeNormalPrecise(Velocity);
+	const FVector RevAccel = (-1.f * FMath::Max(0.f,BrakingDeceleration)) * SafeNormalPrecise(Velocity);
 	while( RemainingTime >= MIN_TICK_TIME )
 	{
 		float dt = ((RemainingTime > TimeStep) ? TimeStep : RemainingTime);
