Remove unused network package loading code

diff --git a/Engine/Source/Runtime/Core/Public/Core.h b/Engine/Source/Runtime/Core/Public/Core.h
--- a/Engine/Source/Runtime/Core/Public/Core.h
+++ b/Engine/Source/Runtime/Core/Public/Core.h
@@ -94,7 +94,6 @@ class FCompactIndex;
 class FExec;
 class FGuid;
 class FMemStack;
-class FPackageInfo;
 class ITransaction;
 class FUnknown;
 class FRepLink;
diff --git a/Engine/Source/Runtime/CoreUObject/Private/Misc/PackageName.cpp b/Engine/Source/Runtime/CoreUObject/Private/Misc/PackageName.cpp
--- a/Engine/Source/Runtime/CoreUObject/Private/Misc/PackageName.cpp
+++ b/Engine/Source/Runtime/CoreUObject/Private/Misc/PackageName.cpp
@@ -720,30 +720,6 @@ bool FPackageName::SearchForPackageOnDisk(const FString& PackageName, FString* O
 	return bResult;
 }
 
-// A map from specific package names to a package's file (possibly a short package name)
-static TMap<FName, FName> PackageNameToPackageFileMapping;				
-
-/**
- * Add a mapping from a package's object name to a package file name, used for making aliases at runtime
- */
-void FPackageName::AddPackageNameMapping(FName PackageName, FName FileName)
-{
-	PackageNameToPackageFileMapping.Add(PackageName, FileName);
-}
-
-/**
- * Returns a mapped package file name, returns the passed in name if not in map
- */
-FName FPackageName::GetMappedPackageName(FName PackageName)
-{
-	FName *MappedName = PackageNameToPackageFileMapping.Find(PackageName);
-	if (MappedName != NULL)
-	{
-		return *MappedName;
-	}
-	return PackageName;
-}
-
 FString FPackageName::PackageFromPath( const TCHAR* InPathName )
 {
 	FString PackageName;
diff --git a/Engine/Source/Runtime/CoreUObject/Private/UObject/CoreNet.cpp b/Engine/Source/Runtime/CoreUObject/Private/UObject/CoreNet.cpp
--- a/Engine/Source/Runtime/CoreUObject/Private/UObject/CoreNet.cpp
+++ b/Engine/Source/Runtime/CoreUObject/Private/UObject/CoreNet.cpp
@@ -10,44 +10,6 @@ DEFINE_LOG_CATEGORY_STATIC(LogCoreNet, Log, All);
 
 DEFINE_STAT(STAT_NetSerializeFast_Array);
 
-/*-----------------------------------------------------------------------------
-	FPackageInfo implementation.
------------------------------------------------------------------------------*/
-
-//
-// FPackageInfo constructor.
-//
-FPackageInfo::FPackageInfo(UPackage* Package)
-:	PackageName		(Package != NULL ? Package->GetFName() : NAME_None)
-,	Parent			(Package)
-,	Guid			(Package != NULL ? Package->GetGuid() : FGuid(0,0,0,0))
-,	ObjectBase		( INDEX_NONE )
-,	ObjectCount		( 0 )
-,	LocalGeneration	( 0)
-,	RemoteGeneration( 0 )
-,	PackageFlags	(Package != NULL ? Package->PackageFlags : 0)
-,	ForcedExportBasePackageName(NAME_None)
-,	FileName		(Package != NULL ? Package->FileName : NAME_None)
-{
-	// if we have a pacakge, find it's source file so that we can send the extension of the file
-	if (Package != NULL)
-	{
-		FString PackageFile;
-		if (FPackageName::DoesPackageExist(Package->GetName(), NULL, &PackageFile))
-		{
-			Extension = FPaths::GetExtension(PackageFile);
-		}
-	}
-}
-
-//
-// FPackageInfo serializer.
-//
-FArchive& operator<<( FArchive& Ar, FPackageInfo& I )
-{
-	return Ar << I.Parent;
-}
-
 /*-----------------------------------------------------------------------------
 	FClassNetCache implementation.
 -----------------------------------------------------------------------------*/
diff --git a/Engine/Source/Runtime/CoreUObject/Public/Misc/PackageName.h b/Engine/Source/Runtime/CoreUObject/Public/Misc/PackageName.h
--- a/Engine/Source/Runtime/CoreUObject/Public/Misc/PackageName.h
+++ b/Engine/Source/Runtime/CoreUObject/Public/Misc/PackageName.h
@@ -187,16 +187,6 @@ public:
 	 **/
 	static bool SearchForPackageOnDisk(const FString& PackageName, FString* OutLongPackageName = NULL, FString* OutFilename = NULL, bool bUseLocalizedNames = false);
 
-	/**
-	 * Add a mapping from a package's object name to a package file name, used for making aliases at runtime
-	 */
-	static void AddPackageNameMapping(FName PackageName, FName FileName);
-
-	/**
-	 * Returns a mapped package file name, returns the passed in name if not in map
-	 */
-	static FName GetMappedPackageName(FName PackageName);
-
 	/**
 	 * Strips all path and extension information from a relative or fully qualified file name.
 	 *
diff --git a/Engine/Source/Runtime/CoreUObject/Public/UObject/CoreNet.h b/Engine/Source/Runtime/CoreUObject/Public/UObject/CoreNet.h
--- a/Engine/Source/Runtime/CoreUObject/Public/UObject/CoreNet.h
+++ b/Engine/Source/Runtime/CoreUObject/Public/UObject/CoreNet.h
@@ -80,31 +80,6 @@ private:
 	TMap< TWeakObjectPtr< UClass >, FClassNetCache * > ClassFieldIndices;
 };
 
-//
-// Ordered information of linker file requirements.
-//
-class COREUOBJECT_API FPackageInfo
-{
-public:
-	// Variables.
-	FName			PackageName;		// name of the package we need to request.
-	UPackage*		Parent;				// The parent package.
-	FGuid			Guid;				// Package identifier.
-	int32				ObjectBase;			// Net index of first object.
-	int32				ObjectCount;		// Number of objects, defined by server.
-	int32				LocalGeneration;	// This machine's generation of the package.
-	int32				RemoteGeneration;	// Remote machine's generation of the package.
-	uint32			PackageFlags;		// Package flags.
-	FName			ForcedExportBasePackageName; // for packages that were a forced export in another package (seekfree loading), the name of that base package, otherwise NAME_None
-	uint8			LoadingPhase;		// indicates if package was loaded during a seamless loading operation (e.g. seamless level change) to aid client in determining when to process it
-	FString			Extension;			// Extension of the package file, used so HTTP downloading can get the package
-	FName			FileName;			// Name of the file this package was loaded from if it is not equal to the PackageName
-
-	// Functions.
-	FPackageInfo(UPackage* Package = NULL);
-	friend COREUOBJECT_API FArchive& operator<<( FArchive& Ar, FPackageInfo& I );
-};
-
 //
 // Maps objects and names to and from indices for network communication.
 //
diff --git a/Engine/Source/Runtime/Engine/Classes/Engine/NetConnection.h b/Engine/Source/Runtime/Engine/Classes/Engine/NetConnection.h
--- a/Engine/Source/Runtime/Engine/Classes/Engine/NetConnection.h
+++ b/Engine/Source/Runtime/Engine/Classes/Engine/NetConnection.h
@@ -262,12 +262,6 @@ public:
 	 */
 	TSet<FNetworkGUID>	DestroyedStartupOrDormantActors;
 
-	/** 
-	 * list of packages the client has received info from the server on via "USES" but can't verify until async loading has completed
-	 * @see UWorld::NotifyReceivedText() and UNetDriver::VerifyPackageInfo()
-	 */
-	TArray<FPackageInfo> PendingPackageInfos;
-
 	/**
 	 * on the server, the world the client has told us it has loaded
 	 * used to make sure the client has traveled correctly, prevent replicating actors before level transitions are done, etc
@@ -482,19 +476,6 @@ public:
 	/** Returns number of bits left in current packet that can be used without causing a flush.  */
 	int64 GetFreeSendBufferBits();
 
-	/** 
-	 * parses the passed in control channel bunch and fills the given package info struct with that data
-	 * @param Bunch network data bunch received on control channel
-	 * @param Info (out) FPackageInfo that receives the parsed data
-	 */
-	void ParsePackageInfo(FInBunch& Bunch, FPackageInfo& Info);
-
-	/**
-	 * sends text describing the given package info to the client via the control channel
-	 * @note: parameter lacks const due to FArchive/FOutBunch interface; it won't be modified
-	 */
-	void SendPackageInfo(FPackageInfo& Info);
-
 	/** 
 	 * returns whether the client has initialized the level required for the given object
 	 * @return true if the client has initialized the level the object is in or the object is not in a level, false otherwise
diff --git a/Engine/Source/Runtime/Engine/Classes/Engine/NetDriver.h b/Engine/Source/Runtime/Engine/Classes/Engine/NetDriver.h
--- a/Engine/Source/Runtime/Engine/Classes/Engine/NetDriver.h
+++ b/Engine/Source/Runtime/Engine/Classes/Engine/NetDriver.h
@@ -489,14 +489,6 @@ public:
 	/** @return true if this netdriver is handling accepting connections */
 	ENGINE_API virtual bool IsServer() const;
 
-	/** verifies that the client has loaded or can load the package with the specified information
-	 * if found, sets the Info's Parent to the package and notifies the server of our generation of the package
-	 * if not, handles downloading the package, skipping it, or disconnecting, depending on the requirements of the package
-	 * @param Info the info on the package the client should have
-	 * @return true if we're done verifying this package, false if we're not done yet (because i.e. async loading is in progress)
-	 */
-	bool VerifyPackageInfo(FPackageInfo& Info);
-
 	ENGINE_API virtual void CleanPackageMaps();
 
 	ENGINE_API void PreSeamlessTravelGarbageCollect();
diff --git a/Engine/Source/Runtime/Engine/Private/DataChannel.cpp b/Engine/Source/Runtime/Engine/Private/DataChannel.cpp
--- a/Engine/Source/Runtime/Engine/Private/DataChannel.cpp
+++ b/Engine/Source/Runtime/Engine/Private/DataChannel.cpp
@@ -900,13 +900,10 @@ IMPLEMENT_CONTROL_CHANNEL_MESSAGE(Challenge);
 IMPLEMENT_CONTROL_CHANNEL_MESSAGE(Netspeed);
 IMPLEMENT_CONTROL_CHANNEL_MESSAGE(Login);
 IMPLEMENT_CONTROL_CHANNEL_MESSAGE(Failure);
-IMPLEMENT_CONTROL_CHANNEL_MESSAGE(Uses);
-IMPLEMENT_CONTROL_CHANNEL_MESSAGE(Have);
 IMPLEMENT_CONTROL_CHANNEL_MESSAGE(Join);
 IMPLEMENT_CONTROL_CHANNEL_MESSAGE(JoinSplit);
 IMPLEMENT_CONTROL_CHANNEL_MESSAGE(Skip);
 IMPLEMENT_CONTROL_CHANNEL_MESSAGE(Abort);
-IMPLEMENT_CONTROL_CHANNEL_MESSAGE(Unload);
 IMPLEMENT_CONTROL_CHANNEL_MESSAGE(PCSwap);
 IMPLEMENT_CONTROL_CHANNEL_MESSAGE(ActorChannelFailure);
 IMPLEMENT_CONTROL_CHANNEL_MESSAGE(DebugText);
@@ -1062,12 +1059,6 @@ void UControlChannel::ReceivedBunch( FInBunch& Bunch )
 				case NMT_Failure:
 					FNetControlMessage<NMT_Failure>::Discard(Bunch);
 					break;
-				case NMT_Uses:
-					FNetControlMessage<NMT_Uses>::Discard(Bunch);
-					break;
-				case NMT_Have:
-					FNetControlMessage<NMT_Have>::Discard(Bunch);
-					break;
 				case NMT_Join:
 					//FNetControlMessage<NMT_Join>::Discard(Bunch);
 					break;
@@ -1080,9 +1071,6 @@ void UControlChannel::ReceivedBunch( FInBunch& Bunch )
 				case NMT_Abort:
 					FNetControlMessage<NMT_Abort>::Discard(Bunch);
 					break;
-				case NMT_Unload:
-					FNetControlMessage<NMT_Unload>::Discard(Bunch);
-					break;
 				case NMT_PCSwap:
 					FNetControlMessage<NMT_PCSwap>::Discard(Bunch);
 					break;
diff --git a/Engine/Source/Runtime/Engine/Private/NetConnection.cpp b/Engine/Source/Runtime/Engine/Private/NetConnection.cpp
--- a/Engine/Source/Runtime/Engine/Private/NetConnection.cpp
+++ b/Engine/Source/Runtime/Engine/Private/NetConnection.cpp
@@ -177,7 +177,6 @@ void UNetConnection::Serialize( FArchive& Ar )
 	{
 		Children.CountBytes(Ar);
 		ClientVisibleLevelNames.CountBytes(Ar);
-		PendingPackageInfos.CountBytes(Ar);
 		QueuedAcks.CountBytes(Ar);
 		ResendAcks.CountBytes(Ar);
 		OpenChannels.CountBytes(Ar);
@@ -405,25 +404,6 @@ void UNetConnection::SendPackageMap()
 {
 }
 
-
-void UNetConnection::ParsePackageInfo(FInBunch& Bunch, FPackageInfo& Info)
-{
-	FString PackageName, FileName, ForcedExportName;
-	FNetControlMessage<NMT_Uses>::Receive(Bunch, Info.Guid, PackageName, FileName, Info.Extension, Info.PackageFlags, Info.RemoteGeneration, ForcedExportName, Info.LoadingPhase);
-	Info.PackageName = FName(*PackageName);
-	Info.FileName = FName(*FileName);
-	Info.ForcedExportBasePackageName = FName(*ForcedExportName);
-}
-
-void UNetConnection::SendPackageInfo(FPackageInfo& Info)
-{
-	FString PackageName(Info.PackageName.ToString());
-	FString FileName(Info.FileName.ToString());
-	FString ForcedExportName((Info.Parent != NULL && Info.Parent->GetForcedExportBasePackageName() != NAME_None) ? *Info.Parent->GetForcedExportBasePackageName().ToString() : TEXT(""));
-	UE_LOG(LogNet, Verbose, TEXT("   UNetConnection::SendPackageInfo: %s  HasMap: %d"), *PackageName, (int32)Info.Parent->ContainsMap());
-	FNetControlMessage<NMT_Uses>::Send(this, Info.Guid, PackageName, FileName, Info.Extension, Info.PackageFlags, Info.LocalGeneration, ForcedExportName, Info.LoadingPhase);
-}
-
 bool UNetConnection::ClientHasInitializedLevelFor(const UObject* TestObject)
 {
 	check(Driver);
diff --git a/Engine/Source/Runtime/Engine/Private/NetworkDriver.cpp b/Engine/Source/Runtime/Engine/Private/NetworkDriver.cpp
--- a/Engine/Source/Runtime/Engine/Private/NetworkDriver.cpp
+++ b/Engine/Source/Runtime/Engine/Private/NetworkDriver.cpp
@@ -146,16 +146,7 @@ void UNetDriver::AssertValid()
 
 void UNetDriver::TickFlush(float DeltaSeconds)
 {
-	if (ServerConnection != NULL)
-	{
-		// attempt to verify any packages that are pending
-		// it's important to verify packages in order so that we're not reshuffling replicated indices during gameplay, so abort as soon as one fails
-		while (ServerConnection->PendingPackageInfos.Num() > 0 && VerifyPackageInfo(ServerConnection->PendingPackageInfos[0]))
-		{
-			ServerConnection->PendingPackageInfos.RemoveAt(0);
-		}
-	}
-	else
+	if ( IsServer() )
 	{
 		// Update all clients.
 #if WITH_SERVER_CODE
@@ -2701,168 +2692,6 @@ void UNetDriver::AddClientConnection(UNetConnection * NewConnection)
 	}
 }
 
-bool UNetDriver::VerifyPackageInfo(FPackageInfo& Info)
-{
-	check(IsServer());
-
-	// We have filename and a packagename so we need to add a mapping
-	if (Info.FileName != NAME_None && Info.PackageName != NAME_None)
-	{
-		FPackageName::AddPackageNameMapping(Info.PackageName, Info.FileName);
-	}
-
-	//If the FileName is not none, then we will be looking for that file, otherwise we will be looking for the PackageName file
-	FName PackageFileToLoad = (Info.FileName != NAME_None ? Info.FileName : Info.PackageName);
-
-	if (GUseSeekFreeLoading)
-	{
-		FString PackageNameString = PackageFileToLoad.ToString();
-		FSeamlessTravelHandler &SeamlessTravelHandler = GEngine->SeamlessTravelHandlerForWorld( GetWorld() );
-
-		// try to find the package in memory
-		Info.Parent = FindPackage(NULL, *PackageNameString);
-		// zero GUID indicates runtime-created package object that no data has been loaded into (possibly placeholder or previous failed load)
-		if (Info.Parent == NULL || !Info.Parent->GetGuid().IsValid())
-		{
-			if (IsAsyncLoading())
-			{
-				// delay until async loading is complete
-				return false;
-			}
-			else if (Info.LoadingPhase == 1 && SeamlessTravelHandler.IsInTransition() && !SeamlessTravelHandler.HasSwitchedToDefaultMap())
-			{
-				// delay until seamless level transition has cleared old level
-				return false;
-			}
-			else if (Info.ForcedExportBasePackageName == NAME_None)
-			{
-				// attempt to async load the package
-				FString Filename;
-				if ( FPackageName::DoesPackageExist(PackageNameString, &Info.Guid, &Filename))
-				{
-					// check for localized package first as that may be required to load the base package
-					FString LocalizedPackageName = PackageNameString + LOCALIZED_SEEKFREE_SUFFIX;
-					FString LocalizedFileName;
-					if (FPackageName::DoesPackageExist(LocalizedPackageName, NULL, &LocalizedFileName))
-					{
-						LoadPackageAsync(*LocalizedPackageName);
-					}
-					// load base package
-					LoadPackageAsync(PackageNameString, &Info.Guid);
-					return false;
-				}
-			}
-			if (Info.ForcedExportBasePackageName != NAME_None)
-			{
-				//@FIXME: we may need the GUID of the base package here...
-
-				// this package is a forced export inside another package, so try loading that other package
-				FString BasePackageNameString(Info.ForcedExportBasePackageName.ToString());
-				FString FileName;
-				if (FPackageName::DoesPackageExist(BasePackageNameString, NULL, &FileName))
-				{
-					// check for localized package first as that may be required to load the base package
-					FString LocalizedPackageName = BasePackageNameString + LOCALIZED_SEEKFREE_SUFFIX;
-					FString LocalizedFileName;
-					if (FPackageName::DoesPackageExist(LocalizedPackageName, NULL, &LocalizedFileName))
-					{
-						LoadPackageAsync(*LocalizedPackageName);
-					}
-					// load base package
-					LoadPackageAsync(*BasePackageNameString);
-					return false;
-				}
-			}
-
-			FString Error = FString::Printf(TEXT("Failed to find required package '%s'"), *PackageNameString);
-			UE_LOG(LogNet, Log, TEXT("%s"), *Error);
-			GEngine->BroadcastTravelFailure(GetWorld(), ETravelFailure::PackageMissing, Error);
-			ServerConnection->Close();
-		}
-		else if (Info.Parent->GetGuid() != Info.Guid)
-		{
-			FString Error = FString::Printf(TEXT("Package '%s' mismatched - Server GUID: %s Client GUID: %s"), *Info.Parent->GetName(), *Info.Guid.ToString(), *Info.Parent->GetGuid().ToString());
-			UE_LOG(LogNet, Log, TEXT("%s"), *Error);
-			GEngine->BroadcastTravelFailure(GetWorld(), ETravelFailure::PackageVersion, Error);
-			ServerConnection->Close();
-		}
-		else
-		{
-			Info.LocalGeneration = 0;
-			// tell the server what we have
-			FGuid Guid = Info.Parent->GetGuid();
-			FNetControlMessage<NMT_Have>::Send(ServerConnection, Guid, Info.LocalGeneration);
-		}
-	}
-	else
-	{
-		// verify that we have this package, or it is downloadable
-		FString Filename;
-		if (FPackageName::DoesPackageExist(PackageFileToLoad.ToString(), &Info.Guid, &Filename))
-		{
-			Info.Parent = FindPackage(NULL, *Info.PackageName.ToString());
-			if (Info.Parent == NULL)
-			{
-				if (IsAsyncLoading())
-				{
-					// delay until async loading is complete
-					return false;
-				}
-				Info.Parent = CreatePackage(NULL, *Info.PackageName.ToString());
-			}
-			// check that the GUID matches (meaning it is the same package or it has been conformed)
-			if (!Info.Parent->GetGuid().IsValid())
-			{
-				if (IsAsyncLoading())
-				{
-					// delay until async loading is complete
-					return false;
-				}
-				uint32 LoadFlags = LOAD_NoWarn | LOAD_NoVerify | LOAD_Quiet;
-
-				if (Info.FileName != NAME_None)
-				{
-					Info.Parent->SetGuid(Info.Guid);
-				}
-
-				// we need to load the linker to get the info
-				BeginLoad();
-				GetPackageLinker(Info.Parent, *PackageFileToLoad.ToString(), LoadFlags, NULL, &Info.Guid);
-				EndLoad();
-			}
-			if (Info.Parent->GetGuid() != Info.Guid)
-			{
-				// incompatible files
-				//@todo FIXME: we need to be able to handle this better - have the client ignore this version of the package and download the correct one
-				FString Error = FString::Printf(TEXT("Package '%s' mismatched - Server GUID: %s Client GUID: %s"), *Info.Parent->GetName(), *Info.Guid.ToString(), *Info.Parent->GetGuid().ToString());
-				UE_LOG(LogNet, Log, TEXT("%s"), *Error);
-
-				GEngine->BroadcastTravelFailure(GetWorld(), ETravelFailure::PackageVersion, Error);
-				ServerConnection->Close();
-			}
-			else
-			{
-				Info.LocalGeneration = 0;
-				// tell the server what we have
-				FGuid Guid = Info.Parent->GetGuid();
-				FNetControlMessage<NMT_Have>::Send(ServerConnection, Guid, Info.LocalGeneration);
-			}
-		}
-		else
-		{
-			// we need to download this package
-			//@fixme: FIXME: handle
-			FString Error = FString::Printf(TEXT("Failed to find required package name '%s' GUID '%s'"), *Info.PackageName.ToString(), *Info.Guid.ToString());
-			UE_LOG(LogNet, Log, TEXT("%s"), *Error);
-
-			GEngine->BroadcastTravelFailure(GetWorld(), ETravelFailure::NoDownload, Error);
-			ServerConnection->Close();
-		} 
-	}
-
-	return true;
-}
-
 void UNetDriver::SetWorld(class UWorld* InWorld)
 {
 	if (World)
diff --git a/Engine/Source/Runtime/Engine/Private/PendingNetGame.cpp b/Engine/Source/Runtime/Engine/Private/PendingNetGame.cpp
--- a/Engine/Source/Runtime/Engine/Private/PendingNetGame.cpp
+++ b/Engine/Source/Runtime/Engine/Private/PendingNetGame.cpp
@@ -128,101 +128,6 @@ void UPendingNetGame::NotifyControlMessage(UNetConnection* Connection, uint8 Mes
 			}
 			break;
 
-		case NMT_Uses:
-		{
-			// Dependency information.
-			FPackageInfo Info(NULL);
-			Connection->ParsePackageInfo(Bunch, Info);
-
-#if !UE_BUILD_SHIPPING
-			UE_LOG(LogNet, Verbose, TEXT(" ---> PackageName: %s, GUID: %s, FileName: %s, Generation: %i, BasePkg: %s"), *Info.PackageName.ToString(), *Info.Guid.ToString(), *Info.FileName.ToString(), Info.RemoteGeneration, *Info.ForcedExportBasePackageName.ToString());
-#endif
-
-			// verify that we have this package (and Guid matches) 
-			FString Filename;
-			bool bWasPackageFound = false;
-
-			// We have filename and a packagename so we need to add a mapping
-			if (Info.FileName != NAME_None && Info.PackageName != NAME_None)
-			{
-				FPackageName::AddPackageNameMapping(Info.PackageName, Info.FileName);
-			}
-
-			//If the FileName is not none, then we will be looking for that file, otherwise we will be looking for the PackageName file
-			FName PackageFileToLoad = (Info.FileName != NAME_None ? Info.FileName : Info.PackageName);
-
-			// if the package guid wasn't found using GuidCache above, then look for it on disk (for mods, etc, or the PC)
-			if (!bWasPackageFound)
-			{
-				if (FPackageName::DoesPackageExist(PackageFileToLoad.ToString(), &Info.Guid, &Filename))
-				{
-					// if we are not doing seekfree loading, then open the package and tell the server we have it
-					// (seekfree loading case will open the packages in UGameEngine::LoadMap)
-					// If the packagename is not the same as the filename, then we default into this case regardless of seekfreeloading status
-					if ( !(GUseSeekFreeLoading && Info.FileName == NAME_None) )
-					{
-						Info.Parent = CreatePackage(NULL, *Info.PackageName.ToString());
-
-						uint32 LoadFlags = LOAD_NoWarn | LOAD_NoVerify | LOAD_Quiet;
-
-						if (Info.FileName != NAME_None)
-						{
-							Info.Parent->SetGuid(Info.Guid);
-						}
-
-						BeginLoad();
-						ULinkerLoad* Linker = GetPackageLinker(Info.Parent, *PackageFileToLoad.ToString(), LoadFlags, NULL, &Info.Guid);
-						EndLoad();
-
-						if (Linker)
-						{
-							Info.LocalGeneration = Linker->Summary.Generations.Num();
-							// tell the server what we have
-							FNetControlMessage<NMT_Have>::Send(NetDriver->ServerConnection, Linker->Summary.Guid, Info.LocalGeneration);
-
-							bWasPackageFound = true;
-						}
-						else
-						{
-							// there must be some problem with the cache if this happens, as the package was found with the right Guid,
-							// but then it failed to load, so log a message and mark that we didn't find it
-							UE_LOG(LogNet, Warning, TEXT("Linker for package %s failed to load, even though it should have [from %s]"), *Info.PackageName.ToString(), *Filename);
-						}
-					}
-					else
-					{
-						bWasPackageFound = true;
-					}
-				}
-			}
-			// if the package guid wasn't found above, then look for it in memory as a compiled in package
-			if (!bWasPackageFound && Info.Guid.IsValid())
-			{
-				UPackage *PossiblyCompiledIn = Cast<UPackage>(StaticFindObjectFast(UPackage::StaticClass(), NULL, Info.PackageName, false, false));
-				if (PossiblyCompiledIn && PossiblyCompiledIn->GetGuid() == Info.Guid)
-				{
-					// that is the correct compiled in package
-					Info.Parent = PossiblyCompiledIn;
-					Info.LocalGeneration = 1;
-					// tell the server what we have
-					FNetControlMessage<NMT_Have>::Send(NetDriver->ServerConnection, Info.Guid, Info.LocalGeneration);
-					bWasPackageFound = true;
-					UE_LOG(LogNet, Verbose, TEXT("Found compiled in package %s [%s] "), *Info.PackageName.ToString(), *Info.Guid.ToString());
-				}
-			}
-			if (!bWasPackageFound)
-			{
-				UE_LOG(LogNet, Warning, TEXT("Failed to match package %s [%s] in guid cache, disk or compiled-in, downloading not supported."), *Info.PackageName.ToString(), *Info.Guid.ToString());
-				ConnectionError = FText::Format( NSLOCTEXT("Engine", "NoDownload", "You cannot join the match because the host has downloadable content '{0}' you do not."), FText::FromName( Info.PackageName ) ).ToString();
-				NetDriver->ServerConnection->Close();
-				return;
-			}
-			break;
-		}
-		case NMT_Unload:
-		{
-			break;
-		}
 		case NMT_Failure:
 		{
 			// our connection attempt failed for some reason, for example a synchronization mismatch (bad GUID, etc) or because the server rejected our join attempt (too many players, etc)
diff --git a/Engine/Source/Runtime/Engine/Private/World.cpp b/Engine/Source/Runtime/Engine/Private/World.cpp
--- a/Engine/Source/Runtime/Engine/Private/World.cpp
+++ b/Engine/Source/Runtime/Engine/Private/World.cpp
@@ -3393,26 +3393,6 @@ void UWorld::NotifyControlMessage(UNetConnection* Connection, uint8 MessageType,
 				}
 				break;
 			}
-			case NMT_Uses:
-			{
-				// Dependency information.
-				FPackageInfo Info(NULL);
-				Connection->ParsePackageInfo(Bunch, Info);
-#if !(UE_BUILD_SHIPPING || UE_BUILD_TEST)
-				UE_LOG(LogNet, Verbose, TEXT(" ---> Package: %s, GUID: %s, Generation: %i, BasePkg: %s"), *Info.PackageName.ToString(), *Info.Guid.ToString(), Info.RemoteGeneration, *Info.ForcedExportBasePackageName.ToString());
-#endif
-				// it's important to verify packages in order so that we're not reshuffling replicated indices during gameplay, so don't try if there are already others pending
-				if (Connection->PendingPackageInfos.Num() > 0 || !NetDriver->VerifyPackageInfo(Info))
-				{
-					// we can't verify the package until we have finished async loading
-					Connection->PendingPackageInfos.Add(Info);
-				}
-				break;
-			}
-			case NMT_Unload:
-			{
-				break;
-			}
 			case NMT_DebugText:
 			{
 				// debug text message
@@ -3496,33 +3476,12 @@ void UWorld::NotifyControlMessage(UNetConnection* Connection, uint8 MessageType,
 				UE_LOG(LogNet, Log, TEXT("Client netspeed is %i"), Connection->CurrentNetSpeed);
 				break;
 			}
-			case NMT_Have:
-			{
-				// Client specifying his generation.
-				FGuid Guid;
-				int32 RemoteGeneration;
-				FNetControlMessage<NMT_Have>::Receive(Bunch, Guid, RemoteGeneration);
-#if !(UE_BUILD_SHIPPING || UE_BUILD_TEST)
-				UE_LOG(LogNet, Verbose, TEXT(" ---> GUID: %s, Generation: %i"), *Guid.ToString(), RemoteGeneration);
-#endif
-				bool bFound = true;
-				if (!bFound)
-				{
-					UE_LOG(LogNet, Log, TEXT("NOT FOUND GUID: %s, Generation: %i"), *Guid.ToString(), RemoteGeneration);
-					// the client specified a package with an incorrect GUID or one that it should not be using; kick it out
-					FString ErrorMsg = NSLOCTEXT("NetworkErrors", "IncompatiblePackage", "Incompatible package.").ToString();
-					FNetControlMessage<NMT_Failure>::Send(Connection, ErrorMsg);
-					Connection->Close();
-				}
-				break;
-			}
 			case NMT_Abort:
 			{
 				break;
 			}
 			case NMT_Skip:
 			{
-				
 				break;
 			}
 			case NMT_Login:
@@ -4237,9 +4196,7 @@ UWorld* FSeamlessTravelHandler::Tick()
 
 	UNetDriver* NetDriver = CurrentWorld->GetNetDriver();
 
-	if ( (LoadedPackage != NULL || LoadedWorld != NULL) && CurrentWorld->NextURL == TEXT("") &&
-		( NetDriver == NULL || NetDriver->ServerConnection == NULL || NetDriver->ServerConnection->PendingPackageInfos.Num() == 0 ||
-		(!bSwitchedToDefaultMap && NetDriver->ServerConnection->PendingPackageInfos[0].LoadingPhase == 1) ) )
+	if ( ( LoadedPackage != NULL || LoadedWorld != NULL ) && CurrentWorld->NextURL == TEXT( "" ) && ( NetDriver == NULL || NetDriver->IsServer() ) )
 	{
 		// First some validity checks		
 		if( CurrentWorld == LoadedWorld )
diff --git a/Engine/Source/Runtime/Engine/Public/Net/DataChannel.h b/Engine/Source/Runtime/Engine/Public/Net/DataChannel.h
--- a/Engine/Source/Runtime/Engine/Public/Net/DataChannel.h
+++ b/Engine/Source/Runtime/Engine/Public/Net/DataChannel.h
@@ -379,13 +379,10 @@ DEFINE_CONTROL_CHANNEL_MESSAGE_TWOPARAM(Challenge, 3, int32, FString); // server
 DEFINE_CONTROL_CHANNEL_MESSAGE_ONEPARAM(Netspeed, 4, int32); // client sends requested transfer rate
 DEFINE_CONTROL_CHANNEL_MESSAGE_THREEPARAM(Login, 5, FString, FString, FUniqueNetIdRepl); // client requests to be admitted to the game
 DEFINE_CONTROL_CHANNEL_MESSAGE_ONEPARAM(Failure, 6, FString); // indicates connection failure
-DEFINE_CONTROL_CHANNEL_MESSAGE_EIGHTPARAM(Uses, 7, FGuid, FString, FString, FString, uint32, int32, FString, uint8); // server tells client about a package they should have/acquire
-DEFINE_CONTROL_CHANNEL_MESSAGE_TWOPARAM(Have, 8, FGuid, int32); // client tells server what version of a package it has
 DEFINE_CONTROL_CHANNEL_MESSAGE_ZEROPARAM(Join, 9); // final join request (spawns PlayerController)
 DEFINE_CONTROL_CHANNEL_MESSAGE_TWOPARAM(JoinSplit, 10, FString, FUniqueNetIdRepl); // child player (splitscreen) join request
 DEFINE_CONTROL_CHANNEL_MESSAGE_ONEPARAM(Skip, 12, FGuid); // client request to skip an optional package
 DEFINE_CONTROL_CHANNEL_MESSAGE_ONEPARAM(Abort, 13, FGuid); // client informs server that it aborted a not-yet-verified package due to an UNLOAD request
-DEFINE_CONTROL_CHANNEL_MESSAGE_ONEPARAM(Unload, 14, FGuid); // server tells client that a package is no longer needed
 DEFINE_CONTROL_CHANNEL_MESSAGE_ONEPARAM(PCSwap, 15, int32); // client tells server it has completed a swap of its Connection->Actor
 DEFINE_CONTROL_CHANNEL_MESSAGE_ONEPARAM(ActorChannelFailure, 16, int32); // client tells server that it failed to open an Actor channel sent by the server (e.g. couldn't serialize Actor archetype)
 DEFINE_CONTROL_CHANNEL_MESSAGE_ONEPARAM(DebugText, 17, FString); // debug text sent to all clients or to server
