#ue4 - Add inflated MTD checks for boxes and spheres when computing depenetration vector (similar to existing checks for capsules).

- Also fix capsules incorrectly adding inflation to half-height; just doing it for radius is enough (PhysX and UE4 treat half-height differently).

diff --git a/Engine/Source/Runtime/Engine/Private/Collision/CollisionConversions.cpp b/Engine/Source/Runtime/Engine/Private/Collision/CollisionConversions.cpp
--- a/Engine/Source/Runtime/Engine/Private/Collision/CollisionConversions.cpp
+++ b/Engine/Source/Runtime/Engine/Private/Collision/CollisionConversions.cpp
@@ -462,30 +462,65 @@ FVector FindBestOverlappingNormal(const PxGeometry& Geom, const PxTransform& Que
 	return BestPlaneNormal;
 }
 
+
+
+static bool ComputeInflatedMTD_Internal(const float MtdInflation, const PxLocationHit& PHit, FHitResult& OutResult, const PxTransform& QueryTM, const PxGeometry& Geom, const PxTransform& PShapeWorldPose)
+{
+	PxGeometry* InflatedGeom = NULL;
+
+	PxVec3 PxMtdNormal(0.f);
+	PxF32 PxMtdDepth = 0.f;
+	const PxGeometry& POtherGeom = PHit.shape->getGeometry().any();
+	const bool bMtdResult = PxGeometryQuery::computePenetration(PxMtdNormal, PxMtdDepth, Geom, QueryTM, POtherGeom, PShapeWorldPose);
+	if (bMtdResult)
+	{
+		if (PxMtdNormal.isFinite())
+		{
+			OutResult.ImpactNormal = P2UVector(PxMtdNormal);
+			OutResult.PenetrationDepth = FMath::Max(FMath::Abs(PxMtdDepth) - MtdInflation, 0.f) + KINDA_SMALL_NUMBER;
+			return true;
+		}
+		else
+		{
+			UE_LOG(LogPhysics, Warning, TEXT("ComputeInflatedMTD_Internal: MTD returned NaN :( normal: (X:%f, Y:%f, Z:%f)"), PxMtdNormal.x, PxMtdNormal.y, PxMtdNormal.z);
+		}
+	}
+
+	return false;
+}
+
+
 // Compute depenetration vector and distance if possible with a slightly larger geometry
 static bool ComputeInflatedMTD(const float MtdInflation, const PxLocationHit& PHit, FHitResult& OutResult, const PxTransform& QueryTM, const PxGeometry& Geom, const PxTransform& PShapeWorldPose)
 {
-	if (Geom.getType() == PxGeometryType::eCAPSULE)
+	switch(Geom.getType())
 	{
-		const PxCapsuleGeometry* InCapsule = static_cast<const PxCapsuleGeometry*>(&Geom);
-		PxCapsuleGeometry InflatedCapsule(InCapsule->radius + MtdInflation, InCapsule->halfHeight + MtdInflation);
+		case PxGeometryType::eCAPSULE:
+		{
+			const PxCapsuleGeometry* InCapsule = static_cast<const PxCapsuleGeometry*>(&Geom);
+			PxCapsuleGeometry InflatedCapsule(InCapsule->radius + MtdInflation, InCapsule->halfHeight); // don't inflate halfHeight, radius is added all around.
+			return ComputeInflatedMTD_Internal(MtdInflation, PHit, OutResult, QueryTM, InflatedCapsule, PShapeWorldPose);
+		}
 
-		PxVec3 PxMtdNormal(0.f);
-		PxF32 PxMtdDepth = 0.f;
-		const PxGeometry& POtherGeom = PHit.shape->getGeometry().any();
-		const bool bMtdResult = PxGeometryQuery::computePenetration(PxMtdNormal, PxMtdDepth, InflatedCapsule, QueryTM, POtherGeom, PShapeWorldPose);
-		if (bMtdResult)
+		case PxGeometryType::eBOX:
 		{
-			if (PxMtdNormal.isFinite())
-			{
-				OutResult.ImpactNormal = P2UVector(PxMtdNormal);
-				OutResult.PenetrationDepth = FMath::Max(FMath::Abs(PxMtdDepth) - MtdInflation, 0.f) + KINDA_SMALL_NUMBER;
-				return true;
-			}
+			const PxBoxGeometry* InBox = static_cast<const PxBoxGeometry*>(&Geom);
+			PxBoxGeometry InflatedBox(InBox->halfExtents + PxVec3(MtdInflation));
+			return ComputeInflatedMTD_Internal(MtdInflation, PHit, OutResult, QueryTM, InflatedBox, PShapeWorldPose);
 		}
-	}
 
-	return false;
+		case PxGeometryType::eSPHERE:
+		{
+			const PxSphereGeometry* InSphere = static_cast<const PxSphereGeometry*>(&Geom);
+			PxSphereGeometry InflatedSphere(InSphere->radius + MtdInflation);
+			return ComputeInflatedMTD_Internal(MtdInflation, PHit, OutResult, QueryTM, InflatedSphere, PShapeWorldPose);
+		}
+
+		default:
+		{
+			return false;
+		}
+	}
 }
 
 
@@ -531,7 +566,7 @@ static bool ConvertOverlappedShapeToImpactHit(const PxLocationHit& PHit, const F
 		OutResult.PenetrationDepth = 0.f;
 		if (!bFiniteNormal)
 		{
-			UE_LOG(LogPhysics, Warning, TEXT("ConvertOverlappedShapeToImpactHit, MTD returned NaN :( normal: (X:%f, Y:%f, Z:%f)"), PHit.normal.x, PHit.normal.y, PHit.normal.z);
+			UE_LOG(LogPhysics, Warning, TEXT("ConvertOverlappedShapeToImpactHit: MTD returned NaN :( normal: (X:%f, Y:%f, Z:%f)"), PHit.normal.x, PHit.normal.y, PHit.normal.z);
 		}
 	}
 
