Fix misinterpretation of logic after removing some dead code

diff --git a/Engine/Source/Runtime/Engine/Private/World.cpp b/Engine/Source/Runtime/Engine/Private/World.cpp
--- a/Engine/Source/Runtime/Engine/Private/World.cpp
+++ b/Engine/Source/Runtime/Engine/Private/World.cpp
@@ -4196,7 +4196,7 @@ UWorld* FSeamlessTravelHandler::Tick()
 
 	UNetDriver* NetDriver = CurrentWorld->GetNetDriver();
 
-	if ( ( LoadedPackage != NULL || LoadedWorld != NULL ) && CurrentWorld->NextURL == TEXT( "" ) && ( NetDriver == NULL || NetDriver->IsServer() ) )
+	if ( ( LoadedPackage != NULL || LoadedWorld != NULL ) && CurrentWorld->NextURL == TEXT( "" ) )
 	{
 		// First some validity checks		
 		if( CurrentWorld == LoadedWorld )
