Landscape Lightmap : Fixed Lightmass rendering for LOD Biased components

diff --git a/Engine/Source/Runtime/Engine/Private/Landscape/LandscapeLight.cpp b/Engine/Source/Runtime/Engine/Private/Landscape/LandscapeLight.cpp
--- a/Engine/Source/Runtime/Engine/Private/Landscape/LandscapeLight.cpp
+++ b/Engine/Source/Runtime/Engine/Private/Landscape/LandscapeLight.cpp
@@ -138,17 +138,15 @@ namespace
 		int32 ComponentSize = ((LandscapeComponent->SubsectionSizeQuads + 1) * LandscapeComponent->NumSubsections) >> LODValue;
 		int32 LODHeightmapSize = LandscapeComponent->HeightmapTexture->Source.GetSizeX() >> LODValue;
 		float Ratio = (float)(LODHeightmapSize) / (HeightmapStride);
-		float Offset = 0.5f * Ratio;
 
 		int32 CurrentHeightmapOffsetX = FMath::RoundToInt((float)(LODHeightmapSize)* LandscapeComponent->HeightmapScaleBias.Z);
 		int32 CurrentHeightmapOffsetY = FMath::RoundToInt((float)(LODHeightmapSize)* LandscapeComponent->HeightmapScaleBias.W);
 
-		// Need to match for component edge cases, otherwise it causes a little 
-		float XX = FMath::Clamp<float>((X - HeightmapOffsetX) * Ratio - Offset, 0.f, ComponentSize - 1.f) + CurrentHeightmapOffsetX;
+		float XX = FMath::Clamp<float>((X - HeightmapOffsetX) * Ratio, 0.f, ComponentSize - 1.f) + CurrentHeightmapOffsetX;
 		int32 XI = (int32)XX;
 		float XF = XX - XI;
 
-		float YY = FMath::Clamp<float>((Y - HeightmapOffsetY) * Ratio - Offset, 0.f, ComponentSize - 1.f) + CurrentHeightmapOffsetY;
+		float YY = FMath::Clamp<float>((Y - HeightmapOffsetY) * Ratio, 0.f, ComponentSize - 1.f) + CurrentHeightmapOffsetY;
 		int32 YI = (int32)YY;
 		float YF = YY - YI;
 
@@ -193,34 +191,52 @@ namespace
 		check(Info);
 
 		FIntPoint ComponentBase = LandscapeComponent->GetSectionBase() / LandscapeComponent->ComponentSizeQuads;
-		ULandscapeComponent* Neighbors[8] = { 0, 0, 0, 0, 0, 0, 0, 0 };
 		int32 NeighborLODs[8] = { 0, 0, 0, 0, 0, 0, 0, 0 };
 
-		Neighbors[0] = Info->XYtoComponentMap.FindRef(ComponentBase + FIntPoint(-1, -1));
-		Neighbors[1] = Info->XYtoComponentMap.FindRef(ComponentBase + FIntPoint(0, -1));
-		Neighbors[2] = Info->XYtoComponentMap.FindRef(ComponentBase + FIntPoint(1, -1));
-		Neighbors[3] = Info->XYtoComponentMap.FindRef(ComponentBase + FIntPoint(-1, 0));
-		Neighbors[4] = Info->XYtoComponentMap.FindRef(ComponentBase + FIntPoint(1, 0));
-		Neighbors[5] = Info->XYtoComponentMap.FindRef(ComponentBase + FIntPoint(-1, 1));
-		Neighbors[6] = Info->XYtoComponentMap.FindRef(ComponentBase + FIntPoint(0, 1));
-		Neighbors[7] = Info->XYtoComponentMap.FindRef(ComponentBase + FIntPoint(1, 1));
-
 		int32 MaxLOD = FMath::CeilLogTwo(LandscapeComponent->SubsectionSizeQuads + 1) - 1;
 		bool bNeedUpscaling = GeometryLOD > InLOD;
-		int32 NeigborMaxLOD = -1;
-
-		for (int32 i = 0; i < 8; ++i)
-		{
-			NeighborLODs[i] = Neighbors[i] ? FMath::Clamp<int32>(Neighbors[i]->ForcedLOD >= 0 ? Neighbors[i]->ForcedLOD : Neighbors[i]->LODBias, 0, MaxLOD) : -1;
-			bNeedUpscaling |= (NeighborLODs[i] > InLOD);
-			NeigborMaxLOD = FMath::Max(NeighborLODs[i], NeigborMaxLOD);
-		}
+		int32 NeighborIdx = 0;
 
-		for (int32 i = 0; i < 8; ++i)
+		for (int32 y = -1; y <= 1; y++)
 		{
-			if (Neighbors[i] == NULL)
+			for (int32 x = -1; x <= 1; x++)
 			{
-				NeighborLODs[i] = NeigborMaxLOD;
+				if (x == 0 && y == 0)
+				{
+					continue;
+				}
+
+				ULandscapeComponent* Neighbor = Info->XYtoComponentMap.FindRef(ComponentBase + FIntPoint(x, y));
+				int32 NeighborLOD = -1;
+
+				if (Neighbor)
+				{
+					NeighborLOD = FMath::Clamp<int32>(Neighbor->ForcedLOD >= 0 ? Neighbor->ForcedLOD : Neighbor->LODBias, 0, MaxLOD);
+				}
+				else
+				{
+					NeighborLOD = 0;
+					// Sample neighbor components to find maximum LOD
+					for (int32 yy = -1; yy <= 1; yy++)
+					{
+						for (int32 xx = -1; xx <= 1; xx++)
+						{
+							if (xx == 0 && yy == 0)
+							{
+								continue;
+							}
+
+							ULandscapeComponent* Neighbor = Info->XYtoComponentMap.FindRef(ComponentBase + FIntPoint(xx, yy));
+							if (Neighbor)
+							{
+								NeighborLOD = FMath::Max(FMath::Clamp<int32>(Neighbor->ForcedLOD >= 0 ? Neighbor->ForcedLOD : Neighbor->LODBias, 0, MaxLOD), NeighborLOD);
+							}
+						}
+					}
+				}
+
+				bNeedUpscaling |= (NeighborLOD > InLOD);
+				NeighborLODs[NeighborIdx++] = NeighborLOD;
 			}
 		}
 
