Use WindowedFullscreen on Mac when supplying -Fullscreen as a commandline argument as proper fullscreen mode is known to be broken on OS X. We need to comprehensively reconsider how we treat mouse versus display versus render coordinates to cope with the way OS X's fullscreen mode works and that hasn't happened yet.

diff --git a/Engine/Source/Runtime/Engine/Private/GameEngine.cpp b/Engine/Source/Runtime/Engine/Private/GameEngine.cpp
--- a/Engine/Source/Runtime/Engine/Private/GameEngine.cpp
+++ b/Engine/Source/Runtime/Engine/Private/GameEngine.cpp
@@ -180,7 +180,12 @@ void UGameEngine::ConditionallyOverrideSettings(int32& ResolutionX, int32& Resol
 	else if (FParse::Param(FCommandLine::Get(),TEXT("FullScreen")))
 	{
 		// -FullScreen
+#if PLATFORM_MAC // @TODO: Fullscreen mode isn't working fully on OS X as it requires explicit CGL back-buffer size and mouse<->display coordinate conversions.
+		// For now use the windowed fullscreen mode which merely stretches to fill the display, which is what OS X does for us by default.
+		WindowMode = EWindowMode::WindowedFullscreen;
+#else
 		WindowMode = EWindowMode::Fullscreen;
+#endif
 	}
 
 	//fullscreen is always supported, but don't allow windowed mode on platforms that dont' support it.
