EngineAutomationTests.cpp (TTP# 340900)
--Added the testflag ATF_NonNullRHI so that it now requires a video card
--May want to put a check in so that it won't run and Fail automation when ran from the builders.  Since this tool is more or less expected to be ran on a local machine by a tester.

diff --git a/Engine/Source/Runtime/Engine/Private/Tests/EngineAutomationTests.cpp b/Engine/Source/Runtime/Engine/Private/Tests/EngineAutomationTests.cpp
--- a/Engine/Source/Runtime/Engine/Private/Tests/EngineAutomationTests.cpp
+++ b/Engine/Source/Runtime/Engine/Private/Tests/EngineAutomationTests.cpp
@@ -297,7 +297,7 @@ bool FExecWorldStringLatentCommand::Update()
 /**
  * Automation test to load a map and capture FPS performance charts
  */
-IMPLEMENT_COMPLEX_AUTOMATION_TEST(FCinematicFPSPerfTest, "Engine.Cinematic FPS Perf Capture", EAutomationTestFlags::ATF_Game);
+IMPLEMENT_COMPLEX_AUTOMATION_TEST(FCinematicFPSPerfTest, "Engine.Cinematic FPS Perf Capture", EAutomationTestFlags::ATF_Game | EAutomationTestFlags::ATF_NonNullRHI);
 
 void FCinematicFPSPerfTest::GetTests(TArray<FString>& OutBeautifiedNames, TArray <FString>& OutTestCommands) const
 {
