[INTEGRATE] Change 2181419 by [REDACTED_PERSON][REDACTED_HANDLE] on 2014/07/09 10:31:13

	Fixed issue where attached actors' transforms are not set correctly.
	#ttp 340842 - CRITICAL: Regression: TEMPLATES: Ramps in Rolling Example Map are misplaced
	#branch UE4-Releases/4.3
	#proj Runtime.Engine
	#change Fetches transform from root component ComponentToWorld, but updates the translation component afterwards, to take into account any custom location if the root component is a BrushComponent.

diff --git a/Engine/Source/Runtime/Engine/Private/ActorConstruction.cpp b/Engine/Source/Runtime/Engine/Private/ActorConstruction.cpp
--- a/Engine/Source/Runtime/Engine/Private/ActorConstruction.cpp
+++ b/Engine/Source/Runtime/Engine/Private/ActorConstruction.cpp
@@ -239,7 +239,8 @@ void AActor::RerunConstructionScripts()
 					RootComponent->DetachFromParent(true);
 				}
 
-				OldTransform = RootComponent->GetRelativeTransform();
+				OldTransform = RootComponent->ComponentToWorld;
+				OldTransform.SetTranslation(RootComponent->GetComponentLocation()); // take into account any custom location
 			}
 		}
 
diff --git a/Engine/Source/Runtime/Engine/Private/ActorEditor.cpp b/Engine/Source/Runtime/Engine/Private/ActorEditor.cpp
--- a/Engine/Source/Runtime/Engine/Private/ActorEditor.cpp
+++ b/Engine/Source/Runtime/Engine/Private/ActorEditor.cpp
@@ -224,7 +224,8 @@ AActor::FActorTransactionAnnotation::FActorTransactionAnnotation(const AActor* A
 	if (RootComponent && RootComponent->bCreatedByConstructionScript)
 	{
 		bRootComponentDataCached = true;
-		RootComponentData.Transform = RootComponent->GetRelativeTransform();
+		RootComponentData.Transform = RootComponent->ComponentToWorld;
+		RootComponentData.Transform.SetTranslation(RootComponent->GetComponentLocation()); // take into account any custom location
 
 		if (RootComponent->AttachParent)
 		{
