Remove GPU memory readbacks in ribbon particles

diff --git a/Engine/Source/Runtime/Engine/Private/Particles/ParticleSystemRender.cpp b/Engine/Source/Runtime/Engine/Private/Particles/ParticleSystemRender.cpp
--- a/Engine/Source/Runtime/Engine/Private/Particles/ParticleSystemRender.cpp
+++ b/Engine/Source/Runtime/Engine/Private/Particles/ParticleSystemRender.cpp
@@ -6036,9 +6036,10 @@ int32 FDynamicRibbonEmitterData::FillVertexData(struct FAsyncBufferFillData& Dat
 						CurrTileU = Tex_U;
 					}
 
+					FVector FinalPos = InterpPos + InterpUp * InterpSize;
 					Vertex = (FParticleBeamTrailVertex*)(TempVertexData);
-					Vertex->Position = InterpPos + InterpUp * InterpSize;
-					Vertex->OldPosition = Vertex->Position;
+					Vertex->Position = FinalPos;
+					Vertex->OldPosition = FinalPos;
 					Vertex->ParticleId	= 0;
 					Vertex->Size.X = InterpSize;
 					Vertex->Size.Y = InterpSize;
@@ -6060,9 +6061,10 @@ int32 FDynamicRibbonEmitterData::FillVertexData(struct FAsyncBufferFillData& Dat
 					TempVertexData += VertexStride;
 					//PackedVertexCount++;
 
+					FinalPos = InterpPos - InterpUp * InterpSize;
 					Vertex = (FParticleBeamTrailVertex*)(TempVertexData);
-					Vertex->Position = InterpPos - InterpUp * InterpSize;
-					Vertex->OldPosition = Vertex->Position;
+					Vertex->Position = FinalPos;
+					Vertex->OldPosition = FinalPos;
 					Vertex->ParticleId	= 0;
 					Vertex->Size.X = InterpSize;
 					Vertex->Size.Y = InterpSize;
