UE4: Wrapping many editor-only UEngine events in #if WITH_EDITOR to prevent them from being subscribed to in non-editor builds under the impression that they would actually get fired.

diff --git a/Engine/Source/Developer/GameplayDebugger/Private/GameplayDebugger.cpp b/Engine/Source/Developer/GameplayDebugger/Private/GameplayDebugger.cpp
--- a/Engine/Source/Developer/GameplayDebugger/Private/GameplayDebugger.cpp
+++ b/Engine/Source/Developer/GameplayDebugger/Private/GameplayDebugger.cpp
@@ -23,8 +23,10 @@ public:
 
 	void WorldAdded(UWorld* InWorld);
 	void WorldDestroyed(UWorld* InWorld);
+#if WITH_EDITOR
 	void OnLevelActorAdded(AActor* InActor);
 	void OnLevelActorDeleted(AActor* InActor);
+#endif
 };
 
 IMPLEMENT_MODULE(FGameplayDebugger, GameplayDebugger)
@@ -37,8 +39,10 @@ void FGameplayDebugger::StartupModule()
 	{
 		GEngine->OnWorldAdded().AddRaw(this, &FGameplayDebugger::WorldAdded);
 		GEngine->OnWorldDestroyed().AddRaw(this, &FGameplayDebugger::WorldDestroyed);
+#if WITH_EDITOR
 		GEngine->OnLevelActorAdded().AddRaw(this, &FGameplayDebugger::OnLevelActorAdded);
 		GEngine->OnLevelActorDeleted().AddRaw(this, &FGameplayDebugger::OnLevelActorDeleted);
+#endif
 	}
 #endif
 }
@@ -52,8 +56,10 @@ void FGameplayDebugger::ShutdownModule()
 	{
 		GEngine->OnWorldAdded().RemoveAll(this);
 		GEngine->OnWorldDestroyed().RemoveAll(this);
+#if WITH_EDITOR
 		GEngine->OnLevelActorAdded().RemoveAll(this);
 		GEngine->OnLevelActorDeleted().RemoveAll(this);
+#endif
 	}
 #endif
 }
@@ -68,6 +74,7 @@ void FGameplayDebugger::WorldDestroyed(UWorld* InWorld)
 
 }
 
+#if WITH_EDITOR
 void FGameplayDebugger::OnLevelActorAdded(AActor* InActor)
 {
 #if !(UE_BUILD_SHIPPING || UE_BUILD_TEST)
@@ -89,3 +96,4 @@ void FGameplayDebugger::OnLevelActorDeleted(AActor* InActor)
 {
 
 }
+#endif
diff --git a/Engine/Source/Runtime/Engine/Classes/AI/Navigation/NavigationSystem.h b/Engine/Source/Runtime/Engine/Classes/AI/Navigation/NavigationSystem.h
--- a/Engine/Source/Runtime/Engine/Classes/AI/Navigation/NavigationSystem.h
+++ b/Engine/Source/Runtime/Engine/Classes/AI/Navigation/NavigationSystem.h
@@ -589,8 +589,10 @@ protected:
 
 	/** delegate handler for PostLoadMap event */
 	void OnPostLoadMap();
+#if WITH_EDITOR
 	/** delegate handler for ActorMoved events */
 	void OnActorMoved(AActor* Actor);
+#endif
 	/** delegate handler called when navigation is dirtied*/
 	void OnNavigationDirtied(const FBox& Bounds);
 	
diff --git a/Engine/Source/Runtime/Engine/Classes/Engine/Engine.h b/Engine/Source/Runtime/Engine/Classes/Engine/Engine.h
--- a/Engine/Source/Runtime/Engine/Classes/Engine/Engine.h
+++ b/Engine/Source/Runtime/Engine/Classes/Engine/Engine.h
@@ -1505,11 +1505,12 @@ public:
 
 	virtual bool IsInitialized() const { return bIsInitialized; }
 
+#if WITH_EDITOR
+
 	/** Editor-only event triggered when the actor list of the world has changed */
 	DECLARE_EVENT( UEngine, FLevelActorListChangedEvent );
 	FLevelActorListChangedEvent& OnLevelActorListChanged() { return LevelActorListChangedEvent; }
 
-
 	/** Called by internal engine systems after a world's actor list changes in a way not specifiable through other LevelActor__Events to notify other subsystems */
 	void BroadcastLevelActorListChanged() { LevelActorListChangedEvent.Broadcast(); }
 
@@ -1562,6 +1563,8 @@ public:
 	/** Called by internal engine systems after a level actor has been requested to be renamed */
 	void BroadcastLevelActorRequestRename(const AActor* InActor) { LevelActorRequestRenameEvent.Broadcast(InActor); }
 
+#endif // #if WITH_EDITOR
+
 	/** Event triggered after a server travel failure of any kind has occurred */
 	FOnTravelFailure& OnTravelFailure() { return TravelFailureEvent; }
 	/** Called by internal engine systems after a travel failure has occurred */
@@ -2163,6 +2166,8 @@ protected:
 	FWorldDestroyedEvent		WorldDestroyedEvent;
 private:
 
+#if WITH_EDITOR
+
 	/** Broadcasts whenever a world's actor list changes in a way not specifiable through other LevelActor__Events */
 	FLevelActorListChangedEvent LevelActorListChangedEvent;
 
@@ -2187,6 +2192,8 @@ private:
 	/** Broadcasts after an actor has been moved, rotated or scaled */
 	FOnActorMovedEvent		OnActorMovedEvent;
 
+#endif // #if WITH_EDITOR
+
 	/** Thread preventing screen saver from kicking. Suspend most of the time. */
 	FRunnableThread*		ScreenSaverInhibitor;
 	FScreenSaverInhibitor*  ScreenSaverInhibitorRunnable;
diff --git a/Engine/Source/Runtime/Engine/Private/AI/Navigation/NavigationSystem.cpp b/Engine/Source/Runtime/Engine/Private/AI/Navigation/NavigationSystem.cpp
--- a/Engine/Source/Runtime/Engine/Private/AI/Navigation/NavigationSystem.cpp
+++ b/Engine/Source/Runtime/Engine/Private/AI/Navigation/NavigationSystem.cpp
@@ -310,7 +310,12 @@ void UNavigationSystem::PostInitProperties()
 #endif
 
 		// register for any actor move change
-		GEngine->OnActorMoved().AddUObject(this, &UNavigationSystem::OnActorMoved);
+#if WITH_EDITOR
+		if ( GIsEditor )
+		{
+			GEngine->OnActorMoved().AddUObject(this, &UNavigationSystem::OnActorMoved);
+		}
+#endif
 		FCoreDelegates::PostLoadMap.AddUObject(this, &UNavigationSystem::OnPostLoadMap);
 #if WITH_NAVIGATION_GENERATOR
 		UNavigationSystem::NavigationDirtyEvent.AddUObject(this, &UNavigationSystem::OnNavigationDirtied);
@@ -2617,6 +2622,7 @@ void UNavigationSystem::OnPostLoadMap()
 	}
 }
 
+#if WITH_EDITOR
 void UNavigationSystem::OnActorMoved(AActor* Actor)
 {
 #if WITH_NAVIGATION_GENERATOR
@@ -2626,6 +2632,7 @@ void UNavigationSystem::OnActorMoved(AActor* Actor)
 	}
 #endif // WITH_NAVIGATION_GENERATOR
 }
+#endif // WITH_EDITOR
 
 void UNavigationSystem::OnNavigationDirtied(const FBox& Bounds)
 {
@@ -2638,10 +2645,13 @@ void UNavigationSystem::CleanUp()
 {
 	UE_LOG(LogNavigation, Log, TEXT("UNavigationSystem::CleanUp"));
 
-	if (GEngine)
+#if WITH_EDITOR
+	if (GIsEditor && GEngine)
 	{
 		GEngine->OnActorMoved().RemoveAll(this);
 	}
+#endif // WITH_EDITOR
+
 	FCoreDelegates::PostLoadMap.RemoveAll(this);
 #if WITH_NAVIGATION_GENERATOR
 	UNavigationSystem::NavigationDirtyEvent.RemoveAll(this);
diff --git a/Engine/Source/Runtime/Engine/Private/LevelActor.cpp b/Engine/Source/Runtime/Engine/Private/LevelActor.cpp
--- a/Engine/Source/Runtime/Engine/Private/LevelActor.cpp
+++ b/Engine/Source/Runtime/Engine/Private/LevelActor.cpp
@@ -333,10 +333,12 @@ AActor* UWorld::SpawnActor( UClass* Class, FVector const* Location, FRotator con
 	// Broadcast notification of spawn
 	OnActorSpawned.Broadcast(Actor);
 
+#if WITH_EDITOR
 	if (GIsEditor)
 	{
 		GEngine->BroadcastLevelActorAdded(Actor);
 	}
+#endif
 
 	return Actor;
 }
@@ -468,11 +470,12 @@ bool UWorld::DestroyActor( AActor* ThisActor, bool bNetForce, bool bShouldModify
 				{
 					ChildActor->DetachSceneComponentsFromParent(*SceneCompIter, true);
 				}
-
+#if WITH_EDITOR
 				if( GIsEditor )
 				{
 					GEngine->BroadcastLevelActorDetached(ChildActor, ThisActor);
 				}
+#endif
 			}
 		}
 	}
@@ -489,10 +492,12 @@ bool UWorld::DestroyActor( AActor* ThisActor, bool bNetForce, bool bShouldModify
 
 		ThisActor->DetachRootComponentFromParent(false);
 
+#if WITH_EDITOR
 		if( GIsEditor )
 		{
 			GEngine->BroadcastLevelActorDetached(ThisActor, OldParentActor);
 		}
+#endif
 	}
 
 	if( ThisActor->IsPendingKill() )
@@ -562,7 +567,9 @@ bool UWorld::DestroyActor( AActor* ThisActor, bool bNetForce, bool bShouldModify
 	if( GIsEditor )
 	{
 		ThisActor->InvalidateLightingCache();
+#if WITH_EDITOR
 		GEngine->BroadcastLevelActorDeleted(ThisActor);
+#endif
 	}
 
 	// Return success.
