Fixing mesh rotation rate with mesh emitters and non-velocity or AwayFromCenter alignment; need to keep track of the current continuous rotation for these to maintain the same rotation rate as for non-aligned particles, because we recompute orientation every tick

diff --git a/Engine/Source/Runtime/Engine/Private/Particles/ParticleEmitterInstances.cpp b/Engine/Source/Runtime/Engine/Private/Particles/ParticleEmitterInstances.cpp
--- a/Engine/Source/Runtime/Engine/Private/Particles/ParticleEmitterInstances.cpp
+++ b/Engine/Source/Runtime/Engine/Private/Particles/ParticleEmitterInstances.cpp
@@ -2937,27 +2937,24 @@ void FParticleMeshEmitterInstance::Tick(float DeltaTime, bool bSuppressSpawning)
 				PayloadData->Rotation.X	= PayloadData->InitialOrientation.X + Euler.X;
 				PayloadData->Rotation.Y = PayloadData->InitialOrientation.Y + Euler.Y;
 				PayloadData->Rotation.Z = PayloadData->InitialOrientation.Z + Euler.Z;
+				PayloadData->Rotation += PayloadData->CurContinuousRotation;
 			}
-	    }
-	}
-
-	// Call the standard tick
-	FParticleEmitterInstance::Tick(DeltaTime, bSuppressSpawning);
-
-	// Apply rotation if it is active
-	if (MeshRotationActive)
-	{
-		for (int32 i = 0; i < ActiveParticles; i++)
-		{
-			DECLARE_PARTICLE(Particle, ParticleData + ParticleStride * ParticleIndices[i]);
-			if ((Particle.Flags & STATE_Particle_FreezeRotation) == 0)
+			else // not PSA_Velocity or PSA_AwayfromCenter, so rotation is not reset every tick
 			{
-				FMeshRotationPayloadData* PayloadData	 = (FMeshRotationPayloadData*)((uint8*)&Particle + MeshRotationOffset);
-				PayloadData->Rotation					+= Particle.RelativeTime * PayloadData->RotationRate;
+				if ((Particle.Flags & STATE_Particle_FreezeRotation) == 0)
+				{
+					PayloadData->Rotation = PayloadData->CurContinuousRotation;
+				}
 			}
+
+			PayloadData->CurContinuousRotation += DeltaTime * PayloadData->RotationRate;
 		}
 	}
 
+
+	// Call the standard tick
+	FParticleEmitterInstance::Tick(DeltaTime, bSuppressSpawning);
+
 	// Remove from the Sprite count... happens because we use the Super::Tick
 	DEC_DWORD_STAT_BY(STAT_SpriteParticles, ActiveParticles);
 	INC_DWORD_STAT_BY(STAT_MeshParticles, ActiveParticles);
diff --git a/Engine/Source/Runtime/Engine/Public/ParticleHelper.h b/Engine/Source/Runtime/Engine/Public/ParticleHelper.h
--- a/Engine/Source/Runtime/Engine/Public/ParticleHelper.h
+++ b/Engine/Source/Runtime/Engine/Public/ParticleHelper.h
@@ -752,6 +752,7 @@ struct FMeshRotationPayloadData
 {
 	FVector	 InitialOrientation;
 	FVector  Rotation;
+	FVector	 CurContinuousRotation;
 	FVector  RotationRate;
 	FVector  RotationRateBase;
 };
