Change constraint renderer to show 3 fan shapes based on color per axis

diff --git a/Engine/Source/Runtime/Engine/Classes/Engine/Engine.h b/Engine/Source/Runtime/Engine/Classes/Engine/Engine.h
--- a/Engine/Source/Runtime/Engine/Classes/Engine/Engine.h
+++ b/Engine/Source/Runtime/Engine/Classes/Engine/Engine.h
@@ -886,6 +886,15 @@ public:
 	UPROPERTY()
 	class UMaterial* ConstraintLimitMaterial;
 
+	UPROPERTY()
+	class UMaterialInstanceDynamic * ConstraintLimitMaterialX;
+
+	UPROPERTY()
+	class UMaterialInstanceDynamic * ConstraintLimitMaterialY;
+
+	UPROPERTY()
+	class UMaterialInstanceDynamic * ConstraintLimitMaterialZ;
+
 	/** @todo document */
 	UPROPERTY(globalconfig)
 	FStringAssetReference ConstraintLimitMaterialName;
diff --git a/Engine/Source/Runtime/Engine/Private/PhysicsEngine/PhysDrawing.cpp b/Engine/Source/Runtime/Engine/Private/PhysicsEngine/PhysDrawing.cpp
--- a/Engine/Source/Runtime/Engine/Private/PhysicsEngine/PhysDrawing.cpp
+++ b/Engine/Source/Runtime/Engine/Private/PhysicsEngine/PhysDrawing.cpp
@@ -11,12 +11,15 @@ static const int32 DrawConeLimitSides = 40;
 static const float DebugJointPosSize = 5.0f;
 static const float DebugJointAxisSize = 20.0f;
 
-static const float JointRenderSize = 5.0f;
+static const float JointRenderThickness = 0.3f;
+static const float JointRenderSize = 10.f;
 static const float LimitRenderSize = 16.0f;
 
-static const FColor JointUnselectedColor(255,0,255);
-static const FColor	JointFrame1Color(255,0,0);
-static const FColor JointFrame2Color(0,0,255);
+static const FColor JointUnselectedColor(255, 0, 255);
+static const FColor JointRed(255, 0, 0);
+static const FColor JointGreen(0, 255, 0);
+static const FColor JointBlue(0, 0, 255);
+
 static const FColor	JointLimitColor(0,255,0);
 static const FColor	JointRefColor(255,255,0);
 static const FColor JointLockedColor(255,128,10);
@@ -628,21 +631,11 @@ void UPhysicsAsset::DrawConstraints(FPrimitiveDrawInterface* PDI, const USkeleta
 		Instance.DrawConstraint(PDI, Scale, 1.f, true, true, Con1Frame, Con2Frame, false);
 	}
 }
-
-static void DrawOrientedStar(FPrimitiveDrawInterface* PDI, const FTransform& Transform, float Size, const FColor Color)
-{
-	FVector Position = Transform.GetTranslation();
-
-	PDI->DrawLine(Position + Size * Transform.GetScaledAxis( EAxis::X ), Position - Size * Transform.GetScaledAxis( EAxis::X ), Color, SDPG_World);
-	PDI->DrawLine(Position + Size * Transform.GetScaledAxis( EAxis::Y ), Position - Size * Transform.GetScaledAxis( EAxis::Y ), Color, SDPG_World);
-	PDI->DrawLine(Position + Size * Transform.GetScaledAxis( EAxis::Z ), Position - Size * Transform.GetScaledAxis( EAxis::Z ), Color, SDPG_World);
-}
-
 static void DrawLinearLimit(FPrimitiveDrawInterface* PDI, const FVector& Origin, const FVector& Axis, const FVector& Orth, float LinearLimitRadius, bool bLinearLimited, float DrawScale)
 {
 	float ScaledLimitSize = LimitRenderSize * DrawScale;
 
-	if(bLinearLimited)
+	if (bLinearLimited)
 	{
 		FVector Start = Origin - LinearLimitRadius * Axis;
 		FVector End = Origin + LinearLimitRadius * Axis;
@@ -669,24 +662,57 @@ static void DrawLinearLimit(FPrimitiveDrawInterface* PDI, const FVector& Origin,
 	}
 }
 
+//creates fan shape along visualized axis for rotation axis of length Length
+FMatrix HelpBuildFan(const FTransform & Con1Frame, const FTransform & Con2Frame, EAxis::Type DrawOnAxis, EAxis::Type RotationAxis, float Length)
+{
+	FVector Con1DrawOnAxis = Con1Frame.GetScaledAxis(DrawOnAxis);
+	FVector Con2DrawOnAxis = Con2Frame.GetScaledAxis(DrawOnAxis);
+
+	FVector Con1RotationAxis = Con1Frame.GetScaledAxis(RotationAxis);
+	FVector Con2RotationAxis = Con2Frame.GetScaledAxis(RotationAxis);
+
+	// Rotate parent twist ref axis
+	FQuat Con2ToCon1Rot = FQuat::FindBetween(Con2RotationAxis, Con1RotationAxis);
+	FVector Con2InCon1DrawOnAxis = Con2ToCon1Rot.RotateVector(Con2DrawOnAxis);
+
+	FTransform ConeLimitTM(Con2InCon1DrawOnAxis, Con1RotationAxis ^ Con2InCon1DrawOnAxis, Con1RotationAxis, Con1Frame.GetTranslation());
+	FMatrix ConeToWorld = FScaleMatrix(FVector(Length * 0.9f)) * ConeLimitTM.ToMatrixWithScale();
+	return ConeToWorld;
+}
+
+//builds radians for limit based on limit type
+float HelpBuildAngle(float LimitAngle, EAngularConstraintMotion LimitType)
+{
+	switch (LimitType)
+	{
+		case ACM_Free: return PI;
+		case ACM_Locked: return 0.f;
+		default: return FMath::DegreesToRadians(LimitAngle);
+	}
+}
+
+
 void FConstraintInstance::DrawConstraint(	FPrimitiveDrawInterface* PDI, 
 											float Scale, float LimitDrawScale, bool bDrawLimits, bool bDrawSelected,
 											const FTransform& Con1Frame, const FTransform& Con2Frame, bool bDrawAsPoint ) const
 {
-	UMaterialInterface* LimitMaterial = GEngine->ConstraintLimitMaterial;
-
+	static UMaterialInterface * LimitMaterialX = GEngine->ConstraintLimitMaterialX;
+	static UMaterialInterface * LimitMaterialY = GEngine->ConstraintLimitMaterialY;
+	static UMaterialInterface * LimitMaterialZ = GEngine->ConstraintLimitMaterialZ;
+	
 	FVector Con1Pos = Con1Frame.GetTranslation();
 	FVector Con2Pos = Con2Frame.GetTranslation();
 
-	float ScaledLimitSize = LimitRenderSize * LimitDrawScale;
+	float Length = JointRenderSize;
+	float Thickness = LimitDrawScale * JointRenderThickness;
 
 	// Special mode for drawing joints just as points..
 	if(bDrawAsPoint)
 	{
 		if(bDrawSelected)
 		{
-			PDI->DrawPoint( Con1Frame.GetTranslation(), JointFrame1Color, 3.f, SDPG_World );
-			PDI->DrawPoint( Con2Frame.GetTranslation(), JointFrame2Color, 3.f, SDPG_World );
+			PDI->DrawPoint( Con1Frame.GetTranslation(), JointRed, 3.f, SDPG_World );
+			PDI->DrawPoint( Con2Frame.GetTranslation(), JointBlue, 3.f, SDPG_World );
 		}
 		else
 		{
@@ -698,16 +724,11 @@ void FConstraintInstance::DrawConstraint(	FPrimitiveDrawInterface* PDI,
 		return;
 	}
 
-	if(bDrawSelected)
-	{
-		DrawOrientedStar( PDI, Con1Frame, LimitDrawScale * JointRenderSize, JointFrame1Color );
-		DrawOrientedStar( PDI, Con2Frame, LimitDrawScale * JointRenderSize, JointFrame2Color );
-	}
-	else
-	{
-		DrawOrientedStar( PDI, Con1Frame, LimitDrawScale * JointRenderSize, JointUnselectedColor );
-		DrawOrientedStar( PDI, Con2Frame, LimitDrawScale * JointRenderSize, JointUnselectedColor );
-	}
+	FVector Position = Con1Frame.GetTranslation();
+
+	PDI->DrawLine(Position, Position + Length * Con1Frame.GetScaledAxis(EAxis::X), JointRed, SDPG_World, Thickness);
+	PDI->DrawLine(Position, Position + Length * Con1Frame.GetScaledAxis(EAxis::Y), JointGreen, SDPG_World, Thickness);
+	PDI->DrawLine(Position, Position + Length * Con1Frame.GetScaledAxis(EAxis::Z), JointBlue, SDPG_World, Thickness);
 
 	//////////////////////////////////////////////////////////////////////////
 	// LINEAR DRAWING
@@ -748,96 +769,27 @@ void FConstraintInstance::DrawConstraint(	FPrimitiveDrawInterface* PDI,
 	bool bLockAllSwing = bLockSwing1 && bLockSwing2;
 
 	bool bDrawnAxisLine = false;
-	FVector RefLineEnd = Con1Frame.GetTranslation() + (1.2f * ScaledLimitSize * Con1Frame.GetScaledAxis( EAxis::X ));
-
-	//If swing1 locked
-	if(bLockSwing1)
-	{
-		// Draw little red 'V' to indicate locked swing.
-		PDI->DrawLine( Con2Frame.GetTranslation(), Con2Frame.TransformPosition( 0.3f * ScaledLimitSize * FVector(1,1,0) ), JointLockedColor, SDPG_World);
-		PDI->DrawLine( Con2Frame.GetTranslation(), Con2Frame.TransformPosition( 0.3f * ScaledLimitSize * FVector(1,-1,0) ), JointLockedColor, SDPG_World);
-	}
+	FVector RefLineStart = Con1Frame.GetTranslation() + (Length * Con1Frame.GetScaledAxis(EAxis::X));
+	FVector RefLineEnd = Con1Frame.GetTranslation() + (Length * Con1Frame.GetScaledAxis(EAxis::X));
 
-	//If swing2 locked
-	if(bLockSwing2)
-	{
-		// Draw little red 'V' to indicate locked swing.
-		PDI->DrawLine( Con2Frame.GetTranslation(), Con2Frame.TransformPosition( 0.3f * ScaledLimitSize * FVector(1,0,1) ), JointLockedColor, SDPG_World);
-		PDI->DrawLine( Con2Frame.GetTranslation(), Con2Frame.TransformPosition( 0.3f * ScaledLimitSize * FVector(1,0,-1) ), JointLockedColor, SDPG_World);
-	}
-	
-	// If swing is limited (but not locked) - draw the limit cone.
-	if(AngularSwing1Motion==ACM_Limited||AngularSwing2Motion==ACM_Limited)
+	//swing1
 	{
-		FTransform ConeLimitTM = Con2Frame;
-		ConeLimitTM.SetTranslation( Con1Frame.GetTranslation() );
-
-		float ang1 = Swing1LimitAngle;
-		if(bLockSwing1)
-			ang1 = 0.f;
-		else
-			ang1 *= ((float)PI/180.f); // convert to radians
-
-		float ang2 = Swing2LimitAngle;
-		if(bLockSwing2)
-			ang2 = 0.f;
-		else
-			ang2 *= ((float)PI/180.f);
-
-		FMatrix ConeToWorld = FScaleMatrix( FVector(ScaledLimitSize) ) * ConeLimitTM.ToMatrixWithScale();
-		DrawCone(PDI, ConeToWorld, ang1, ang2, DrawConeLimitSides, true, JointLimitColor, LimitMaterial->GetRenderProxy(false), SDPG_World );
-
-		// Draw reference line
-		PDI->DrawLine( Con1Frame.GetTranslation(), RefLineEnd, JointRefColor, SDPG_World );
-		bDrawnAxisLine = true;
+		FMatrix ConeToWorld = HelpBuildFan(Con1Frame, Con2Frame, EAxis::X, EAxis::Z, Length);
+		float Limit = HelpBuildAngle(Swing1LimitAngle, AngularSwing1Motion);
+		DrawCone(PDI, ConeToWorld, Limit, 0, DrawConeLimitSides, false, JointLimitColor, LimitMaterialX->GetRenderProxy(false), SDPG_World);
 	}
 
-	//Handle the twist constraint draw
-	if(bLockTwist)
+	//swing2
 	{
-		// If there is no axis line draw already - add a little one now to sit the 'twist locked' cross on
-		if(!bDrawnAxisLine)
-			PDI->DrawLine( Con1Frame.GetTranslation(), RefLineEnd, JointLockedColor, SDPG_World );
-
-		PDI->DrawLine(  RefLineEnd + Con1Frame.TransformVector( 0.3f * ScaledLimitSize * FVector(0.f,-0.5f,-0.5f) ), 
-			RefLineEnd + Con1Frame.TransformVector( 0.3f * ScaledLimitSize * FVector(0.f, 0.5f, 0.5f) ), JointLockedColor, SDPG_World);
-
-		PDI->DrawLine(  RefLineEnd + Con1Frame.TransformVector( 0.3f * ScaledLimitSize * FVector(0.f, 0.5f,-0.5f) ), 
-			RefLineEnd + Con1Frame.TransformVector( 0.3f * ScaledLimitSize * FVector(0.f,-0.5f, 0.5f) ), JointLockedColor, SDPG_World);
+		FMatrix ConeToWorld = HelpBuildFan(Con1Frame, Con2Frame, EAxis::Z, EAxis::Y, Length);
+		float Limit = HelpBuildAngle(Swing2LimitAngle, AngularSwing2Motion);
+		DrawCone(PDI, ConeToWorld, Limit, 0, DrawConeLimitSides, false, JointLimitColor, LimitMaterialZ->GetRenderProxy(false), SDPG_World);
 	}
-	else if(AngularTwistMotion==ACM_Limited)
-	{
-		// If twist is limited, but not completely locked, draw 
-		// If no axis yet drawn - do it now
-		if(!bDrawnAxisLine)
-			PDI->DrawLine( Con1Frame.GetTranslation(), RefLineEnd, JointRefColor, SDPG_World );
-
-		// Draw twist limit.
-		FVector ChildTwistRef = Con1Frame.GetScaledAxis( EAxis::Y );
-		FVector ParentTwistRef = Con2Frame.GetScaledAxis( EAxis::Y );
-		PDI->DrawLine( RefLineEnd, RefLineEnd + ChildTwistRef * ScaledLimitSize, JointRefColor, SDPG_World );
-
-		// Rotate parent twist ref axis
-		FQuat ParentToChildRot = FQuat::FindBetween( Con2Frame.GetScaledAxis( EAxis::X ), Con1Frame.GetScaledAxis( EAxis::X ) );
-		FVector ChildTwistLimit = ParentToChildRot.RotateVector( ParentTwistRef );
 
-		FQuat RotateLimit = FQuat( Con1Frame.GetScaledAxis( EAxis::X ), TwistLimitAngle * (PI/180.0f) );
-		FVector WLimit = RotateLimit.RotateVector(ChildTwistLimit);
-		PDI->DrawLine( RefLineEnd, RefLineEnd + WLimit * ScaledLimitSize, JointLimitColor, SDPG_World );
-
-		RotateLimit = FQuat( Con1Frame.GetScaledAxis( EAxis::X ), -TwistLimitAngle * (PI/180.0f) );
-		WLimit = RotateLimit.RotateVector(ChildTwistLimit);
-		PDI->DrawLine( RefLineEnd, RefLineEnd + WLimit * ScaledLimitSize, JointLimitColor, SDPG_World );
-
-		DrawArc(PDI, RefLineEnd, ChildTwistLimit, -ChildTwistLimit ^ Con1Frame.GetScaledAxis( EAxis::X ), -TwistLimitAngle, TwistLimitAngle, 0.8f * ScaledLimitSize, 8, JointLimitColor, SDPG_World);
-	}
-	else
+	//twist
 	{
-		// For convenience, in the hinge case where swing is locked and twist is unlimited, draw the twist axis.
-		if(bLockAllSwing)
-			PDI->DrawLine(  Con2Frame.GetTranslation() - ScaledLimitSize * Con2Frame.GetScaledAxis( EAxis::X ), 
-							Con2Frame.GetTranslation() + ScaledLimitSize * Con2Frame.GetScaledAxis( EAxis::X ), JointLimitColor, SDPG_World );
+		FMatrix ConeToWorld = HelpBuildFan(Con1Frame, Con2Frame, EAxis::Y, EAxis::X, Length);
+		float Limit = HelpBuildAngle(TwistLimitAngle, AngularTwistMotion);
+		DrawCone(PDI, ConeToWorld, Limit, 0, DrawConeLimitSides, false, JointLimitColor, LimitMaterialY->GetRenderProxy(false), SDPG_World);
 	}
-
-
 }
diff --git a/Engine/Source/Runtime/Engine/Private/UnrealEngine.cpp b/Engine/Source/Runtime/Engine/Private/UnrealEngine.cpp
--- a/Engine/Source/Runtime/Engine/Private/UnrealEngine.cpp
+++ b/Engine/Source/Runtime/Engine/Private/UnrealEngine.cpp
@@ -1209,6 +1209,13 @@ void UEngine::InitializeObjectReferences()
 		LoadSpecialMaterial(PreviewShadowsIndicatorMaterialName.AssetLongPathname, PreviewShadowsIndicatorMaterial, false);
 		LoadSpecialMaterial(ConstraintLimitMaterialName.AssetLongPathname, ConstraintLimitMaterial, false);
 
+		ConstraintLimitMaterialX = UMaterialInstanceDynamic::Create(ConstraintLimitMaterial, NULL);
+		ConstraintLimitMaterialX->SetVectorParameterValue(FName("Color"), FLinearColor::Red);
+		ConstraintLimitMaterialY = UMaterialInstanceDynamic::Create(ConstraintLimitMaterial, NULL);
+		ConstraintLimitMaterialY->SetVectorParameterValue(FName("Color"), FLinearColor::Green);
+		ConstraintLimitMaterialZ = UMaterialInstanceDynamic::Create(ConstraintLimitMaterial, NULL);
+		ConstraintLimitMaterialZ->SetVectorParameterValue(FName("Color"), FLinearColor::Blue);
+
 		//@TODO: This should move into the editor (used in editor modes exclusively)
 		if (DefaultBSPVertexTexture == NULL)
 		{
