Remove vestigial bits of UE3 music system

diff --git a/Engine/Source/Runtime/Engine/Classes/Engine/World.h b/Engine/Source/Runtime/Engine/Classes/Engine/World.h
--- a/Engine/Source/Runtime/Engine/Classes/Engine/World.h
+++ b/Engine/Source/Runtime/Engine/Classes/Engine/World.h
@@ -9,7 +9,6 @@
 #include "CollisionQueryParams.h"
 #include "WorldCollision.h"
 #include "PendingNetGame.h"
-#include "GameFramework/MusicTrackDataStructures.h"
 #include "EngineDefines.h"
 #include "Engine/LatentActionManager.h"
 #include "Runtime/RHI/Public/RHIDefinitions.h"
@@ -519,10 +518,6 @@ class ENGINE_API UWorld : public UObject, public FNetworkNotify
 	UPROPERTY(Transient)
 	class UPhysicsCollisionHandler*				PhysicsCollisionHandler;
 
-	/** Audio component used for playing music tracks via SeqAct_PlayMusicTrack */
-	UPROPERTY(Transient)
-	class UAudioComponent*						MusicComp;
-
 	/** Array of any additional objects that need to be referenced by this world, to make sure they aren't GC'd */
 	UPROPERTY(Transient)
 	TArray<UObject*>							ExtraReferencedObjects;
@@ -881,9 +876,6 @@ public:
 	/** @todo document */
 	FName CommittedPersistentLevelName;
 
-	/** Param information for the currently playing MusicComp */
-	struct FMusicTrackStruct CurrentMusicTrack;
-	
 #if WITH_EDITORONLY_DATA
 	/** Map of LandscapeInfos for all loaded levels, valid in the editor only */
 	TMap< FGuid, class ULandscapeInfo* > LandscapeInfoMap;
@@ -2245,8 +2237,6 @@ public:
 	 */
 	class AReverbVolume* GetAudioSettings( const FVector& ViewLocation, struct FReverbSettings* OutReverbSettings, struct FInteriorSettings* OutInteriorSettings );
 
-	void UpdateMusicTrack(FMusicTrackStruct NewMusicTrack);
-
 	/** Return the URL of this level on the local machine. */
 	virtual FString GetLocalURL() const;
 
diff --git a/Engine/Source/Runtime/Engine/Classes/GameFramework/MusicTrackDataStructures.h b/Engine/Source/Runtime/Engine/Classes/GameFramework/MusicTrackDataStructures.h
deleted file mode 100644
diff --git a/Engine/Source/Runtime/Engine/Classes/GameFramework/WorldSettings.h b/Engine/Source/Runtime/Engine/Classes/GameFramework/WorldSettings.h
--- a/Engine/Source/Runtime/Engine/Classes/GameFramework/WorldSettings.h
+++ b/Engine/Source/Runtime/Engine/Classes/GameFramework/WorldSettings.h
@@ -4,7 +4,6 @@
 
 #include "Sound/ReverbVolume.h"
 #include "Engine/PostProcessVolume.h"
-#include "MusicTrackDataStructures.h"
 
 #include "WorldSettings.generated.h"
 
@@ -396,18 +395,8 @@ class ENGINE_API AWorldSettings : public AInfo, public IInterface_AssetUserData
 	UPROPERTY()
 	TArray<struct FNetViewer> ReplicationViewers;
 
-	/** Version of a new music track request replicated to clients */
-	UPROPERTY(transient, replicatedUsing=OnRep_ReplicatedMusicTrack)
-	struct FMusicTrackStruct ReplicatedMusicTrack;
-
 	// ************************************
 
-	//
-	// Replication Notification Callbacks.
-	//
-	UFUNCTION()
-	virtual void OnRep_ReplicatedMusicTrack();
-
 	/** Maximum number of bookmarks	*/
 	static const int32 MAX_BOOKMARK_NUMBER = 10;
 
diff --git a/Engine/Source/Runtime/Engine/Private/EnginePrivate.h b/Engine/Source/Runtime/Engine/Private/EnginePrivate.h
--- a/Engine/Source/Runtime/Engine/Private/EnginePrivate.h
+++ b/Engine/Source/Runtime/Engine/Private/EnginePrivate.h
@@ -50,7 +50,6 @@
 #include "GameFramework/PlayerState.h"
 #include "Engine/SkyLight.h"
 #include "Engine/WindDirectionalSource.h"
-#include "GameFramework/MusicTrackDataStructures.h"
 #include "GameFramework/WorldSettings.h"
 #include "PhysicsEngine/BodyInstance.h"
 #include "Components/SceneComponent.h"
diff --git a/Engine/Source/Runtime/Engine/Private/MusicTrackDataStructures.cpp b/Engine/Source/Runtime/Engine/Private/MusicTrackDataStructures.cpp
deleted file mode 100644
diff --git a/Engine/Source/Runtime/Engine/Private/World.cpp b/Engine/Source/Runtime/Engine/Private/World.cpp
--- a/Engine/Source/Runtime/Engine/Private/World.cpp
+++ b/Engine/Source/Runtime/Engine/Private/World.cpp
@@ -121,7 +121,6 @@ void UWorld::Serialize( FArchive& Ar )
 		Ar << ForegroundLineBatcher;
 
 		Ar << MyParticleEventManager;
-		Ar << MusicComp;
 		Ar << GameState;
 		Ar << AuthorityGameMode;
 		Ar << NetworkManager;
@@ -184,7 +183,6 @@ void UWorld::AddReferencedObjects(UObject* InThis, FReferenceCollector& Collecto
 		Collector.AddReferencedObject( This->PersistentLineBatcher, This );
 		Collector.AddReferencedObject( This->ForegroundLineBatcher, This );
 		Collector.AddReferencedObject( This->MyParticleEventManager, This );
-		Collector.AddReferencedObject( This->MusicComp, This );
 		Collector.AddReferencedObject( This->GameState, This );
 		Collector.AddReferencedObject( This->AuthorityGameMode, This );
 		Collector.AddReferencedObject( This->NetworkManager, This );
@@ -4582,45 +4580,6 @@ int32 UWorld::GetDetailMode()
 	return GetCachedScalabilityCVars().DetailMode;
 }
 
-void UWorld::UpdateMusicTrack(FMusicTrackStruct NewMusicTrack)
-{
-	if (MusicComp != NULL)
-	{
-		// If attempting to play same track, don't.
-		if (NewMusicTrack.Sound == CurrentMusicTrack.Sound)
-		{
-			return;
-		}
-		else
-		{
-			// otherwise fade out the current track
-			MusicComp->FadeOut(CurrentMusicTrack.FadeOutTime,CurrentMusicTrack.FadeOutVolumeLevel);
-			MusicComp = NULL;
-		}
-	}
-
-	{
-		// create a new audio component to play music
-		MusicComp = FAudioDevice::CreateComponent( NewMusicTrack.Sound, this, NULL, false );
-		if (MusicComp != NULL)
-		{
-			// update the new component with the correct settings
-			MusicComp->bAutoDestroy = true;
-			MusicComp->bShouldRemainActiveIfDropped = true;
-			MusicComp->bIsMusic = true;
-			MusicComp->bAutoActivate = NewMusicTrack.bAutoPlay;
-			MusicComp->bIgnoreForFlushing = NewMusicTrack.bPersistentAcrossLevels;
-
-			// and finally fade in the new track
-			MusicComp->FadeIn( NewMusicTrack.FadeInTime, NewMusicTrack.FadeInVolumeLevel );
-		}
-	}
-
-	// set the properties for future fades as well as replication to clients
-	CurrentMusicTrack = NewMusicTrack;
-	GetWorldSettings()->ReplicatedMusicTrack = NewMusicTrack;
-}
-
 /**
  * Updates all physics constraint actor joint locations.
  */
diff --git a/Engine/Source/Runtime/Engine/Private/WorldSettings.cpp b/Engine/Source/Runtime/Engine/Private/WorldSettings.cpp
--- a/Engine/Source/Runtime/Engine/Private/WorldSettings.cpp
+++ b/Engine/Source/Runtime/Engine/Private/WorldSettings.cpp
@@ -106,11 +106,6 @@ void AWorldSettings::PostInitializeComponents()
 	}
 }
 
-void AWorldSettings::OnRep_ReplicatedMusicTrack()
-{
-	GetWorld()->UpdateMusicTrack(ReplicatedMusicTrack);
-}
-
 float AWorldSettings::GetGravityZ() const
 {
 	if (!bWorldGravitySet)
@@ -151,7 +146,6 @@ void AWorldSettings::GetLifetimeReplicatedProps( TArray< FLifetimeProperty > & O
 	DOREPLIFETIME( AWorldSettings, MatineeTimeDilation );
 	DOREPLIFETIME( AWorldSettings, WorldGravityZ );
 	DOREPLIFETIME( AWorldSettings, bHighPriorityLoading );
-	DOREPLIFETIME( AWorldSettings, ReplicatedMusicTrack );
 }
 
 void AWorldSettings::Serialize( FArchive& Ar )
diff --git a/Engine/Source/Runtime/Engine/Public/Engine.h b/Engine/Source/Runtime/Engine/Public/Engine.h
--- a/Engine/Source/Runtime/Engine/Public/Engine.h
+++ b/Engine/Source/Runtime/Engine/Public/Engine.h
@@ -71,7 +71,6 @@
 #include "GameFramework/PlayerState.h"
 #include "Engine/SkyLight.h"
 #include "Engine/WindDirectionalSource.h"
-#include "GameFramework/MusicTrackDataStructures.h"
 #include "GameFramework/WorldSettings.h"
 #include "PhysicsEngine/BodyInstance.h"
 #include "Components/SceneComponent.h"
