* part of change 2182500
* fixed TTP 340094 FIXIF: CRITICAL: ElementalDemo: XboxOne: Some particle corruption

diff --git a/Engine/Source/Runtime/Engine/Private/TextureDerivedData.cpp b/Engine/Source/Runtime/Engine/Private/TextureDerivedData.cpp
--- a/Engine/Source/Runtime/Engine/Private/TextureDerivedData.cpp
+++ b/Engine/Source/Runtime/Engine/Private/TextureDerivedData.cpp
@@ -34,7 +34,7 @@ enum
 // In case of merge conflicts with DDC versions, you MUST generate a new GUID and set this new
 // guid as version
 
-#define TEXTURE_DERIVEDDATA_VER		TEXT("4F83E7F4EC4E4AB788364F736C9E4311")
+#define TEXTURE_DERIVEDDATA_VER		TEXT("75F87CE100B546AAB8A1C4132AF9DB17")
 
 /*------------------------------------------------------------------------------
 	Timing of derived data operations.
