Fix for TTP 321351 - UE4: RENDERING: UDN: LICENSEE: Material Instance: Override Blend Mode ignores Opacity Mask

Changing an MIDs blend mode to masked now works correctly.

diff --git a/Engine/Source/Runtime/Engine/Classes/Materials/Material.h b/Engine/Source/Runtime/Engine/Classes/Materials/Material.h
--- a/Engine/Source/Runtime/Engine/Classes/Materials/Material.h
+++ b/Engine/Source/Runtime/Engine/Classes/Materials/Material.h
@@ -770,6 +770,7 @@ public:
 	ENGINE_API virtual EBlendMode GetBlendMode_Internal() const;
 	ENGINE_API virtual EMaterialShadingModel GetShadingModel_Internal() const;
 	ENGINE_API virtual bool IsTwoSided_Internal() const;
+	ENGINE_API virtual bool IsMasked_Internal() const;
 
 	ENGINE_API void SetShadingModel(EMaterialShadingModel NewModel) {ShadingModel = NewModel;}
 
diff --git a/Engine/Source/Runtime/Engine/Classes/Materials/MaterialInstance.h b/Engine/Source/Runtime/Engine/Classes/Materials/MaterialInstance.h
--- a/Engine/Source/Runtime/Engine/Classes/Materials/MaterialInstance.h
+++ b/Engine/Source/Runtime/Engine/Classes/Materials/MaterialInstance.h
@@ -244,6 +244,7 @@ public:
 	ENGINE_API virtual EBlendMode GetBlendMode_Internal() const;
 	ENGINE_API virtual EMaterialShadingModel GetShadingModel_Internal() const;
 	ENGINE_API virtual bool IsTwoSided_Internal() const;
+	ENGINE_API virtual bool IsMasked_Internal() const;
 
 	/** Returns true and sets Result if the property was overridden in the instance. Otherwise, returns false and the base material property should be used. */
 	ENGINE_API virtual bool GetOpacityMaskClipValueOverride(float& OutResult) const;
@@ -253,6 +254,8 @@ public:
 	ENGINE_API virtual bool GetShadingModelOverride(EMaterialShadingModel& OutResult) const;
 	/** Returns true and sets Result if the property was overridden in the instance. Otherwise, returns false and the base material property should be used. */
 	ENGINE_API virtual bool IsTwoSidedOverride(bool& OutResult) const;
+	/** Returns true and sets Result if the property was overridden in the instance. Otherwise, returns false and the base material property should be used. */
+	ENGINE_API virtual bool IsMaskedOverride(bool& OutResult) const;
 
 	/** Checks to see if an input property should be active, based on the state of the material */
 	ENGINE_API virtual bool IsPropertyActive(EMaterialProperty InProperty) const;
diff --git a/Engine/Source/Runtime/Engine/Classes/Materials/MaterialInterface.h b/Engine/Source/Runtime/Engine/Classes/Materials/MaterialInterface.h
--- a/Engine/Source/Runtime/Engine/Classes/Materials/MaterialInterface.h
+++ b/Engine/Source/Runtime/Engine/Classes/Materials/MaterialInterface.h
@@ -509,12 +509,14 @@ public:
 	ENGINE_API EBlendMode GetBlendMode() const;
 	ENGINE_API EMaterialShadingModel GetShadingModel() const;
 	ENGINE_API bool IsTwoSided() const;
+	ENGINE_API bool IsMasked() const;
 
 	/** Game thread versions of the accessors. On the render thread there are equivelant getters in FMaterial. */
 	ENGINE_API virtual float GetOpacityMaskClipValue_Internal() const;
 	ENGINE_API virtual EBlendMode GetBlendMode_Internal() const;
 	ENGINE_API virtual EMaterialShadingModel GetShadingModel_Internal() const;
 	ENGINE_API virtual bool IsTwoSided_Internal() const;
+	ENGINE_API virtual bool IsMasked_Internal() const;
 	/**
 	 * Force the streaming system to disregard the normal logic for the specified duration and
 	 * instead always load all mip-levels for all textures used by this material.
diff --git a/Engine/Source/Runtime/Engine/Private/Materials/Material.cpp b/Engine/Source/Runtime/Engine/Private/Materials/Material.cpp
--- a/Engine/Source/Runtime/Engine/Private/Materials/Material.cpp
+++ b/Engine/Source/Runtime/Engine/Private/Materials/Material.cpp
@@ -3417,6 +3417,11 @@ bool UMaterial::IsTwoSided_Internal() const
 	return TwoSided != 0;
 }
 
+bool UMaterial::IsMasked_Internal() const
+{
+	return bIsMasked != 0;
+}
+
 bool UMaterial::IsPropertyActive(EMaterialProperty InProperty)const 
 {
 	if(MaterialDomain == MD_PostProcess)
diff --git a/Engine/Source/Runtime/Engine/Private/Materials/MaterialInstance.cpp b/Engine/Source/Runtime/Engine/Private/Materials/MaterialInstance.cpp
--- a/Engine/Source/Runtime/Engine/Private/Materials/MaterialInstance.cpp
+++ b/Engine/Source/Runtime/Engine/Private/Materials/MaterialInstance.cpp
@@ -2373,6 +2373,16 @@ bool UMaterialInstance::IsTwoSided_Internal() const
 	return GetMaterial()->IsTwoSided();
 }
 
+bool UMaterialInstance::IsMasked_Internal() const
+{
+	checkSlow(IsInGameThread());
+	if (bOverrideBaseProperties && BasePropertyOverrides.bOverride_BlendMode)
+	{
+		return BasePropertyOverrides.BlendMode == EBlendMode::BLEND_Masked;
+	}
+	return GetMaterial()->IsMasked();
+}
+
 bool UMaterialInstance::GetOpacityMaskClipValueOverride(float& OutResult) const
 {
 	if( bOverrideBaseProperties && BasePropertyOverrides.bOverride_OpacityMaskClipValue )
@@ -2413,6 +2423,16 @@ bool UMaterialInstance::IsTwoSidedOverride(bool& OutResult) const
 	return false;
 }
 
+bool UMaterialInstance::IsMaskedOverride(bool& OutResult) const
+{
+	if (bOverrideBaseProperties && BasePropertyOverrides.BlendMode)
+	{
+		OutResult = BasePropertyOverrides.BlendMode == EBlendMode::BLEND_Masked;
+		return true;
+	}
+	return false;
+}
+
 /** Checks to see if an input property should be active, based on the state of the material */
 bool UMaterialInstance::IsPropertyActive(EMaterialProperty InProperty) const
 {
diff --git a/Engine/Source/Runtime/Engine/Private/Materials/MaterialInterface.cpp b/Engine/Source/Runtime/Engine/Private/Materials/MaterialInterface.cpp
--- a/Engine/Source/Runtime/Engine/Private/Materials/MaterialInterface.cpp
+++ b/Engine/Source/Runtime/Engine/Private/Materials/MaterialInterface.cpp
@@ -46,7 +46,7 @@ FMaterialRelevance UMaterialInterface::GetRelevance_Internal(const UMaterial* Ma
 		// Determine the material's view relevance.
 		FMaterialRelevance MaterialRelevance;
 		MaterialRelevance.bOpaque = !bIsTranslucent;
-		MaterialRelevance.bMasked = Material->bIsMasked;
+		MaterialRelevance.bMasked = IsMasked();
 		MaterialRelevance.bDistortion = Material->bUsesDistortion && bIsTranslucent;
 		MaterialRelevance.bSeparateTranslucency = bIsTranslucent && Material->bEnableSeparateTranslucency;
 		MaterialRelevance.bNormalTranslucency = bIsTranslucent && !Material->bEnableSeparateTranslucency;
@@ -261,6 +261,17 @@ bool UMaterialInterface::IsTwoSided()const
 	return GetRenderProxy(0)->GetMaterial(GRHIFeatureLevel)->IsTwoSided();
 }
 
+bool UMaterialInterface::IsMasked() const
+{
+	if (IsInGameThread())
+	{
+		return IsMasked_Internal();
+	}
+
+	//We're on the render thread so get it from the proxy.
+	return GetRenderProxy(0)->GetMaterial(GRHIFeatureLevel)->IsMasked();
+}
+
 EMaterialShadingModel UMaterialInterface::GetShadingModel()const
 {
 	if( IsInGameThread() )
@@ -287,6 +298,11 @@ bool UMaterialInterface::IsTwoSided_Internal()const
 	return false;
 }
 
+bool UMaterialInterface::IsMasked_Internal() const
+{
+	return false;
+}
+
 EMaterialShadingModel UMaterialInterface::GetShadingModel_Internal()const
 {
 	return MSM_DefaultLit;
diff --git a/Engine/Source/Runtime/Engine/Private/Materials/MaterialShared.cpp b/Engine/Source/Runtime/Engine/Private/Materials/MaterialShared.cpp
--- a/Engine/Source/Runtime/Engine/Private/Materials/MaterialShared.cpp
+++ b/Engine/Source/Runtime/Engine/Private/Materials/MaterialShared.cpp
@@ -858,6 +858,16 @@ bool FMaterialResource::IsTwoSided() const
 	return Ret;
 }
 
+bool FMaterialResource::IsMasked() const 
+{ 
+	bool Ret;
+	if (!(MaterialInstance && MaterialInstance->IsMaskedOverride(Ret)))
+	{
+		Ret = Material->bIsMasked != 0;
+	}
+	return Ret;
+}
+
 bool FMaterialResource::IsDistorted() const { return Material->bUsesDistortion; }
 float FMaterialResource::GetTranslucencyDirectionalLightingIntensity() const { return Material->TranslucencyDirectionalLightingIntensity; }
 float FMaterialResource::GetTranslucentShadowDensityScale() const { return Material->TranslucentShadowDensityScale; }
@@ -869,12 +879,6 @@ FLinearColor FMaterialResource::GetTranslucentMultipleScatteringExtinction() con
 float FMaterialResource::GetTranslucentShadowStartOffset() const { return Material->TranslucentShadowStartOffset; }
 float FMaterialResource::GetRefractionDepthBiasValue() const { return Material->RefractionDepthBias; }
 bool FMaterialResource::UseTranslucencyVertexFog() const {return Material->bUseTranslucencyVertexFog;}
-/**
- * Check if the material is masked and uses an expression or a constant that's not 1.0f for opacity.
- * @return true if the material uses opacity
- */
-bool FMaterialResource::IsMasked() const { return Material->bIsMasked; }
-
 FString FMaterialResource::GetFriendlyName() const { return *GetNameSafe(Material); }
 
 uint32 FMaterialResource::GetDecalBlendMode() const
