Added flag to TextRenderComponent that allows text to render unmodified when using debug visualization modes. Useful for debugging.

#ttp 335079 - UE4: RENDERING: LIVE: Regression: Editor: TextRenderActor renders as black box when using Level Coloration

diff --git a/Engine/Source/Runtime/Engine/Classes/Components/TextRenderComponent.h b/Engine/Source/Runtime/Engine/Classes/Components/TextRenderComponent.h
--- a/Engine/Source/Runtime/Engine/Classes/Components/TextRenderComponent.h
+++ b/Engine/Source/Runtime/Engine/Classes/Components/TextRenderComponent.h
@@ -81,6 +81,10 @@ class ENGINE_API UTextRenderComponent : public UPrimitiveComponent
 	UPROPERTY(EditAnywhere, BlueprintReadOnly, Category=Text)
 	float HorizSpacingAdjust;
 
+	/** Allows text to draw unmodified when using debug visualization modes. **/
+	UPROPERTY(EditAnywhere, BlueprintReadOnly, Category=Rendering)
+	uint32 bAlwaysRenderAsText:1;
+
 	// -----------------------------
 	
 	/** Change the text value and signal the primitives to be rebuilt */
diff --git a/Engine/Source/Runtime/Engine/Private/TextRenderComponent.cpp b/Engine/Source/Runtime/Engine/Private/TextRenderComponent.cpp
--- a/Engine/Source/Runtime/Engine/Private/TextRenderComponent.cpp
+++ b/Engine/Source/Runtime/Engine/Private/TextRenderComponent.cpp
@@ -384,6 +384,7 @@ private:
 	float HorizSpacingAdjust;
 	EHorizTextAligment HorizontalAlignment;
 	EVerticalTextAligment VerticalAlignment;
+	bool bAlwaysRenderAsText;
 };
 
 FTextRenderSceneProxy::FTextRenderSceneProxy( UTextRenderComponent* Component) :
@@ -395,7 +396,8 @@ FTextRenderSceneProxy::FTextRenderSceneProxy( UTextRenderComponent* Component) :
 	YScale(Component->WorldSize * Component->YScale * Component->InvDefaultSize),
 	HorizSpacingAdjust(Component->HorizSpacingAdjust),
 	HorizontalAlignment(Component->HorizontalAlignment),
-	VerticalAlignment(Component->VerticalAlignment)
+	VerticalAlignment(Component->VerticalAlignment),
+	bAlwaysRenderAsText(Component->bAlwaysRenderAsText)
 {
 	UMaterialInterface* EffectiveMaterial = 0;
 
@@ -470,16 +472,24 @@ void FTextRenderSceneProxy::DrawDynamicElements(FPrimitiveDrawInterface* PDI,con
 
 		const bool bIsWireframe = View->Family->EngineShowFlags.Wireframe;
 
-		uint32 NumPasses = DrawRichMesh(
-			PDI, 
-			Mesh, 
-			FLinearColor(1.0f, 0.0f, 0.0f),	//WireframeColor,
-			FLinearColor(1.0f, 1.0f, 0.0f),	//LevelColor,
-			FLinearColor(1.0f, 1.0f, 1.0f),	//PropertyColor,		
-			this,
-			bUseSelectedMaterial,
-			bIsWireframe
-			);
+		if (bAlwaysRenderAsText)
+		{ 
+			// Render text unmodified
+			PDI->DrawMesh(Mesh);
+		}
+		else
+		{
+			uint32 NumPasses = DrawRichMesh(
+				PDI,
+				Mesh,
+				FLinearColor(1.0f, 0.0f, 0.0f),	//WireframeColor,
+				FLinearColor(1.0f, 1.0f, 0.0f),	//LevelColor,
+				FLinearColor(1.0f, 1.0f, 1.0f),	//PropertyColor,		
+				this,
+				bUseSelectedMaterial,
+				bIsWireframe
+				);
+		}
 
 #if !(UE_BUILD_SHIPPING || UE_BUILD_TEST)
 		RenderBounds(PDI, View->Family->EngineShowFlags, GetBounds(), IsSelected());
