Expose a Get Display  Name for class variables

diff --git a/Engine/Source/Runtime/Engine/Classes/Kismet/KismetSystemLibrary.h b/Engine/Source/Runtime/Engine/Classes/Kismet/KismetSystemLibrary.h
--- a/Engine/Source/Runtime/Engine/Classes/Kismet/KismetSystemLibrary.h
+++ b/Engine/Source/Runtime/Engine/Classes/Kismet/KismetSystemLibrary.h
@@ -61,6 +61,10 @@ class UKismetSystemLibrary : public UBlueprintFunctionLibrary
 	UFUNCTION(BlueprintPure, Category="Utilities")
 	static ENGINE_API FString GetDisplayName(const UObject* Object);
 
+	// Returns the display name of a class
+	UFUNCTION(BlueprintPure, Category = "Utilities", meta = (FriendlyName = "Get Display Name"))
+	static ENGINE_API FString GetClassDisplayName(UClass* Class);
+
 	// Engine build number, for displaying to end users.
 	UFUNCTION(BlueprintPure, Category="Development")
 	static FString GetEngineVersion();
diff --git a/Engine/Source/Runtime/Engine/Private/KismetSystemLibrary.cpp b/Engine/Source/Runtime/Engine/Private/KismetSystemLibrary.cpp
--- a/Engine/Source/Runtime/Engine/Private/KismetSystemLibrary.cpp
+++ b/Engine/Source/Runtime/Engine/Private/KismetSystemLibrary.cpp
@@ -39,6 +39,11 @@ FString UKismetSystemLibrary::GetDisplayName(const UObject* Object)
 	return Object ? Object->GetName() : FString();
 }
 
+FString UKismetSystemLibrary::GetClassDisplayName(UClass* Class)
+{
+	return Class ? Class->GetName() : FString();
+}
+
 FString UKismetSystemLibrary::GetEngineVersion()
 {
 	return GEngineVersion.ToString();
