Mobile Preview In Editor
 - Material editor preview nodes now work when the editor is in es2 mode
 - Make editor primitives work in ES2 mode
 - Fix slate font rendering in ES2 mode
 - Remove ES2 shader de-gamma when in the editor
 - Add a new linear sampler type for linear textures
 - Various fixes for dynamic feature level switching
 - When in an emulated feature level mode, display some text in the level editor viewport saying what feature level we're running
 - Removed feature level parameter from USkeletalMesh::GetResourceForRendering. Not used for anything
 - ParticleVertexFactoryPool is now creates feature level specific vertex factories.

diff --git a/Engine/Shaders/BasePassPixelShader.usf b/Engine/Shaders/BasePassPixelShader.usf
--- a/Engine/Shaders/BasePassPixelShader.usf
+++ b/Engine/Shaders/BasePassPixelShader.usf
@@ -431,7 +431,7 @@ float ClipForEditorPrimitives(FMaterialPixelParameters MaterialParameters)
 {
 	float Ret = 1;
 
-#if EDITOR_PRIMITIVE_MATERIAL && FEATURE_LEVEL >= FEATURE_LEVEL_SM4
+#if EDITOR_PRIMITIVE_MATERIAL && (FEATURE_LEVEL >= FEATURE_LEVEL_SM4 || ES2_EMULATION)
 	// Depth test manually if compositing editor primitives since the depth buffer is different (MSAA only)
 	BRANCH
 	if (bEnableEditorPrimitiveDepthTest)
@@ -526,7 +526,7 @@ void Main(
 #endif
 
 
-#if EDITOR_PRIMITIVE_MATERIAL && FEATURE_LEVEL >= FEATURE_LEVEL_SM4
+#if EDITOR_PRIMITIVE_MATERIAL && (FEATURE_LEVEL >= FEATURE_LEVEL_SM4 || ES2_EMULATION)
 	const bool bEditorWeightedZBuffering = true;
 #else
 	const bool bEditorWeightedZBuffering = false;
diff --git a/Engine/Shaders/Common.usf b/Engine/Shaders/Common.usf
--- a/Engine/Shaders/Common.usf
+++ b/Engine/Shaders/Common.usf
@@ -128,10 +128,11 @@
 		#define FCOLOR_COMPONENT_SWIZZLE .bgra  
 	#else 
 		#define FCOLOR_COMPONENT_SWIZZLE 
-
-		// A8 textures when sampled have their component in R.
-		#undef A8_SAMPLE_MASK
-		#define A8_SAMPLE_MASK .r
+		#if COMPILER_GLSL
+			// A8 textures when sampled have their component in R
+			#undef A8_SAMPLE_MASK
+			#define A8_SAMPLE_MASK .r
+		#endif
 	#endif
 #else
 
diff --git a/Engine/Shaders/MaterialTemplate.usf b/Engine/Shaders/MaterialTemplate.usf
--- a/Engine/Shaders/MaterialTemplate.usf
+++ b/Engine/Shaders/MaterialTemplate.usf
@@ -541,24 +541,51 @@ MaterialFloat2 GetParticleMacroUV(FMaterialPixelParameters Parameters)
 	return (Parameters.ScreenPosition.xy / Parameters.ScreenPosition.w - Parameters.Particle.MacroUV.xy) * Parameters.Particle.MacroUV.zw + MaterialFloat2(.5, .5);
 }
 
+#define ES2_EMULATION (!COMPILER_GLSL_ES2 && COMPILE_SHADERS_FOR_DEVELOPMENT)
+
 MaterialFloat4 ProcessMaterialColorTextureLookup(MaterialFloat4 TextureValue)
 {
 #if ES2_PROFILE
+	#if ES2_EMULATION
+	if( View.ES2PreviewMode > 0.5f)
+	{
+		// undo HW srgb->lin
+		TextureValue = pow(TextureValue, 1.0f/2.2f); // TODO: replace with a more accurate lin -> sRGB conversion.
+	}
+	#endif
+
 	// sRGB read approximation
 	TextureValue.rgb *= TextureValue.rgb;
 #endif 
 	return TextureValue;
 }
 
+MaterialFloat4 ProcessMaterialLinearColorTextureLookup(MaterialFloat4 TextureValue)
+{
+	return TextureValue;
+}
+
 MaterialFloat ProcessMaterialGreyscaleTextureLookup(MaterialFloat TextureValue)
 {
 #if ES2_PROFILE
+	#if ES2_EMULATION
+	if( View.ES2PreviewMode > 0.5f )
+	{
+		// undo HW srgb->lin
+		TextureValue = pow(TextureValue, 1.0f/2.2f); // TODO: replace with a more accurate lin -> sRGB conversion.
+	}
+	#endif
 	// sRGB read approximation
 	TextureValue *= TextureValue;
 #endif 
 	return TextureValue;
 }
 
+MaterialFloat ProcessMaterialLinearGreyscaleTextureLookup(MaterialFloat TextureValue)
+{
+	return TextureValue;
+}
+
 /** Calculate a reflection vector about the specified world space normal. Optionally normalize this normal **/
 MaterialFloat3 ReflectionAboutCustomWorldNormal(FMaterialPixelParameters Parameters, MaterialFloat3 WorldNormal, bool bNormalizeInputNormal)
 {
diff --git a/Engine/Source/Editor/LevelEditor/Private/LevelEditorActions.cpp b/Engine/Source/Editor/LevelEditor/Private/LevelEditorActions.cpp
--- a/Engine/Source/Editor/LevelEditor/Private/LevelEditorActions.cpp
+++ b/Engine/Source/Editor/LevelEditor/Private/LevelEditorActions.cpp
@@ -450,14 +450,7 @@ bool FLevelEditorActionCallbacks::IsMaterialQualityLevelChecked( EMaterialQualit
 
 void FLevelEditorActionCallbacks::SetFeatureLevelPreview(ERHIFeatureLevel::Type InPreviewFeatureLevel)
 {
-	// try iterating over all worlds here
-	for (TObjectIterator<UWorld> It; It; ++It)
-	{
-		UWorld* World = *It;
-		World->ChangeFeatureLevel(InPreviewFeatureLevel);
-	}
-	
-	UWorld::ForceFeatureLevelUpdate(InPreviewFeatureLevel);
+	UWorld::ChangeAllWorldFeatureLevels(InPreviewFeatureLevel);
 }
 
 bool FLevelEditorActionCallbacks::IsFeatureLevelPreviewChecked(ERHIFeatureLevel::Type InPreviewFeatureLevel)
@@ -2835,10 +2828,9 @@ void FLevelEditorCommands::RegisterCommands()
 			FUICommandInfoDecl(
 			this->AsShared(),
 			Name,
-			//FText::Format(NSLOCTEXT("LevelEditorCommands", "FeatureLevelPreview", "Preview scenes using feature level {0}"), TEXT("Thing")),
 			FText::FromName(Name),
 			FText::Format(NSLOCTEXT("LevelEditorCommands", "FeatureLevelPreviewTooltip", "Preview scenes using feature level {0}"), FText::FromName(Name)))
-			.UserInterfaceType(EUserInterfaceActionType::Button)
+			.UserInterfaceType(EUserInterfaceActionType::RadioButton)
 			.DefaultGesture(FInputGesture());
 	}
 }
diff --git a/Engine/Source/Editor/LevelEditor/Private/LevelEditorToolBar.cpp b/Engine/Source/Editor/LevelEditor/Private/LevelEditorToolBar.cpp
--- a/Engine/Source/Editor/LevelEditor/Private/LevelEditorToolBar.cpp
+++ b/Engine/Source/Editor/LevelEditor/Private/LevelEditorToolBar.cpp
@@ -17,11 +17,6 @@
 #include "Matinee/MatineeActor.h"
 #include "Engine/LevelScriptBlueprint.h"
 
-static TAutoConsoleVariable<int32> CVarFeatureLevelPreview(
-	TEXT("r.FeatureLevelPreview"),
-	0,
-	TEXT("If 1 the quick settings menu will contain an option to enable feature level preview modes"),
-	ECVF_RenderThreadSafe);
 
 /**
  * Static: Creates a widget for the level editor tool bar
@@ -478,9 +473,7 @@ static void MakeShaderModelPreviewMenu(FMenuBuilder& MenuBuilder)
 {
 	MenuBuilder.BeginSection("LevelEditorShaderModelPreview", NSLOCTEXT("LevelToolBarViewMenu", "FeatureLevelPreviewHeading", "Feature Level Preview"));
 	{
-		static const ERHIFeatureLevel::Type MaxFeatureLevel = GMaxRHIFeatureLevel;
-
-		for (int32 i = MaxFeatureLevel; i >= 0; --i)
+		for (int32 i = GMaxRHIFeatureLevel; i >= 0; --i)
 		{
 			MenuBuilder.AddMenuEntry(FLevelEditorCommands::Get().FeatureLevelPreview[i]);
 		}
@@ -548,7 +541,8 @@ TSharedRef< SWidget > FLevelEditorToolBar::GenerateQuickSettingsMenu( TSharedRef
 			LOCTEXT( "MaterialQualityLevelSubMenu_ToolTip", "Sets the value of the CVar \"r.MaterialQualityLevel\" (low=0, high=1). This affects materials via the QualitySwitch material expression." ),
 			FNewMenuDelegate::CreateStatic( &MakeMaterialQualityLevelMenu ) );
 
-		if (CVarFeatureLevelPreview.GetValueOnGameThread() != 0)
+		static const auto CVarFeatureLevelPreview = IConsoleManager::Get().FindTConsoleVariableDataInt(TEXT("r.FeatureLevelPreview"));
+		if (CVarFeatureLevelPreview->GetValueOnGameThread() != 0)
 		{
 			MenuBuilder.AddSubMenu(
 				LOCTEXT("FeatureLevelPreviewSubMenu", "Feature Level Preview"),
diff --git a/Engine/Source/Editor/LevelEditor/Private/SLevelViewport.cpp b/Engine/Source/Editor/LevelEditor/Private/SLevelViewport.cpp
--- a/Engine/Source/Editor/LevelEditor/Private/SLevelViewport.cpp
+++ b/Engine/Source/Editor/LevelEditor/Private/SLevelViewport.cpp
@@ -175,29 +175,64 @@ void SLevelViewport::ConstructViewportOverlayContent()
 	.HAlign( HAlign_Right )
 	.Padding( 5.0f )
 	[
-		SNew( SHorizontalBox )
-		.Visibility( this, &SLevelViewport::GetCurrentLevelTextVisibility )
-		// Current level label
-		+ SHorizontalBox::Slot()
-		.AutoWidth()
-		.Padding( 2.0f, 1.0f, 2.0f, 1.0f )
+		SNew(SVerticalBox)
+		+ SVerticalBox::Slot()
+		.AutoHeight()
+		.Padding(2.0f, 1.0f, 2.0f, 1.0f)
 		[
-			SNew( STextBlock )
-			.Text( this, &SLevelViewport::GetCurrentLevelText, true )
-			.Font( FEditorStyle::GetFontStyle( TEXT( "MenuItem.Font" ) ) )
-			.ShadowOffset( FVector2D(1,1) )
-		]
+			SNew(SHorizontalBox)
+			.Visibility(this, &SLevelViewport::GetCurrentFeatureLevelPreviewTextVisibility)
+			// Current level label
+			+ SHorizontalBox::Slot()
+			.AutoWidth()
+			.Padding(2.0f, 1.0f, 2.0f, 1.0f)
+			[
+				SNew(STextBlock)
+				.Text(this, &SLevelViewport::GetCurrentFeatureLevelPreviewText, true)
+				.Font(FEditorStyle::GetFontStyle(TEXT("MenuItem.Font")))
+				.ShadowOffset(FVector2D(1, 1))
+			]
 
-		// Current level
-		+ SHorizontalBox::Slot()
-		.AutoWidth()
-		.Padding( 4.0f, 1.0f, 2.0f, 1.0f )
+			// Current level
+			+ SHorizontalBox::Slot()
+				.AutoWidth()
+				.Padding(4.0f, 1.0f, 2.0f, 1.0f)
+				[
+					SNew(STextBlock)
+					.Text(this, &SLevelViewport::GetCurrentFeatureLevelPreviewText, false)
+					.Font(FEditorStyle::GetFontStyle(TEXT("MenuItem.Font")))
+					.ColorAndOpacity(FLinearColor(0.4f, 1.0f, 1.0f))
+					.ShadowOffset(FVector2D(1, 1))
+				]
+		]
+		+ SVerticalBox::Slot()
+		.AutoHeight()
+		.Padding(2.0f, 1.0f, 2.0f, 1.0f)
 		[
-			SNew( STextBlock )
-			.Text( this, &SLevelViewport::GetCurrentLevelText, false )
-			.Font( FEditorStyle::GetFontStyle( TEXT( "MenuItem.Font" ) ) )
-			.ColorAndOpacity( FLinearColor( 0.4f, 1.0f, 1.0f ) )
-			.ShadowOffset( FVector2D(1,1) )
+			SNew(SHorizontalBox)
+			.Visibility(this, &SLevelViewport::GetCurrentLevelTextVisibility)
+			// Current level label
+			+ SHorizontalBox::Slot()
+			.AutoWidth()
+			.Padding(2.0f, 1.0f, 2.0f, 1.0f)
+			[
+				SNew(STextBlock)
+				.Text(this, &SLevelViewport::GetCurrentLevelText, true)
+				.Font(FEditorStyle::GetFontStyle(TEXT("MenuItem.Font")))
+				.ShadowOffset(FVector2D(1, 1))
+			]
+
+			// Current level
+			+ SHorizontalBox::Slot()
+				.AutoWidth()
+				.Padding(4.0f, 1.0f, 2.0f, 1.0f)
+				[
+					SNew(STextBlock)
+					.Text(this, &SLevelViewport::GetCurrentLevelText, false)
+					.Font(FEditorStyle::GetFontStyle(TEXT("MenuItem.Font")))
+					.ColorAndOpacity(FLinearColor(0.4f, 1.0f, 1.0f))
+					.ShadowOffset(FVector2D(1, 1))
+				]
 		]
 	];
 
@@ -2953,6 +2988,32 @@ FString SLevelViewport::GetDeviceProfileString( ) const
 	return DeviceProfile;
 }
 
+FText SLevelViewport::GetCurrentFeatureLevelPreviewText( bool bDrawOnlyLabel ) const
+{
+	FText LabelName;
+	FText FeatureLevelText;
+
+	if (bDrawOnlyLabel)
+	{
+		LabelName = LOCTEXT("FeatureLevelLabel", "Feature Level:");
+	}
+	else
+	{
+		if (GRHIFeatureLevel != GMaxRHIFeatureLevel)
+		{
+			FName FeatureLevelName;
+			GetFeatureLevelName(GRHIFeatureLevel, FeatureLevelName);
+			FeatureLevelText = FText::Format(LOCTEXT("FeatureLevel", "{0}"), FText::FromName(FeatureLevelName));
+		}
+	}
+	
+	if (bDrawOnlyLabel)
+	{
+		return LabelName;
+	}
+
+	return FeatureLevelText;
+}
 
 FText SLevelViewport::GetCurrentLevelText( bool bDrawOnlyLabel ) const
 {
@@ -3002,6 +3063,11 @@ EVisibility SLevelViewport::GetCurrentLevelTextVisibility() const
 	return (&GetLevelViewportClient() == GCurrentLevelEditingViewportClient) ? ContentVisibility : EVisibility::Collapsed;
 }
 
+EVisibility SLevelViewport::GetCurrentFeatureLevelPreviewTextVisibility() const
+{
+	return (GRHIFeatureLevel != GMaxRHIFeatureLevel) ? EVisibility::SelfHitTestInvisible : EVisibility::Collapsed;
+}
+
 void SLevelViewport::OnSetViewportConfiguration(FName ConfigurationName)
 {
 	TSharedPtr<FLevelViewportLayout> LayoutPinned = ParentLayout.Pin();
diff --git a/Engine/Source/Editor/LevelEditor/Public/SLevelViewport.h b/Engine/Source/Editor/LevelEditor/Public/SLevelViewport.h
--- a/Engine/Source/Editor/LevelEditor/Public/SLevelViewport.h
+++ b/Engine/Source/Editor/LevelEditor/Public/SLevelViewport.h
@@ -236,10 +236,16 @@ public:
 
 	/** Called to get the level text */
 	FText GetCurrentLevelText( bool bDrawOnlyLabel ) const;
+
+	/** Called to get the feature level preview text */
+	FText GetCurrentFeatureLevelPreviewText( bool bDrawOnlyLabel ) const;
 	
 	/** @return The visibility of the current level text display */
 	EVisibility GetCurrentLevelTextVisibility() const;
 
+	/** @return The visibility of the current feature level preview text display */
+	EVisibility GetCurrentFeatureLevelPreviewTextVisibility() const;
+
 	/**
 	 * Sets the current layout on the parent tab that this viewport belongs to.
 	 * 
diff --git a/Engine/Source/Editor/MaterialEditor/Private/MaterialEditor.cpp b/Engine/Source/Editor/MaterialEditor/Private/MaterialEditor.cpp
--- a/Engine/Source/Editor/MaterialEditor/Private/MaterialEditor.cpp
+++ b/Engine/Source/Editor/MaterialEditor/Private/MaterialEditor.cpp
@@ -85,15 +85,26 @@ bool FMatExpressionPreview::ShouldCache(EShaderPlatform Platform, const FShaderT
 		// we only need the non-light-mapped, base pass, local vertex factory shaders for drawing an opaque Material Tile
 		// @todo: Added a FindShaderType by fname or something"
 
-		if(FCString::Stristr(ShaderType->GetName(), TEXT("BasePassVSFNoLightMapPolicy")) ||
-			FCString::Stristr(ShaderType->GetName(), TEXT("BasePassHSFNoLightMapPolicy")) ||
-			FCString::Stristr(ShaderType->GetName(), TEXT("BasePassDSFNoLightMapPolicy")))
+		if (IsES2Platform(Platform))
 		{
-			return true;
+			if (FCString::Stristr(ShaderType->GetName(), TEXT("BasePassForForwardShadingVSFNoLightMapPolicy")) ||
+				FCString::Stristr(ShaderType->GetName(), TEXT("BasePassForForwardShadingPSFNoLightMapPolicy")))
+			{
+				return true;
+			}
 		}
-		else if(FCString::Stristr(ShaderType->GetName(), TEXT("BasePassPSFNoLightMapPolicy")))
+		else
 		{
-			return true;
+			if (FCString::Stristr(ShaderType->GetName(), TEXT("BasePassVSFNoLightMapPolicy")) ||
+				FCString::Stristr(ShaderType->GetName(), TEXT("BasePassHSFNoLightMapPolicy")) ||
+				FCString::Stristr(ShaderType->GetName(), TEXT("BasePassDSFNoLightMapPolicy")))
+			{
+				return true;
+			}
+			else if (FCString::Stristr(ShaderType->GetName(), TEXT("BasePassPSFNoLightMapPolicy")))
+			{
+				return true;
+			}
 		}
 	}
 
diff --git a/Engine/Source/Editor/StatsViewer/Private/StatsPages/PrimitiveStatsPage.cpp b/Engine/Source/Editor/StatsViewer/Private/StatsPages/PrimitiveStatsPage.cpp
--- a/Engine/Source/Editor/StatsViewer/Private/StatsPages/PrimitiveStatsPage.cpp
+++ b/Engine/Source/Editor/StatsViewer/Private/StatsPages/PrimitiveStatsPage.cpp
@@ -145,7 +145,7 @@ struct PrimitiveStatsGenerator
 			// Calculate vertex color usage for skeletal meshes
 			if( Mesh )
 			{
-				FSkeletalMeshResource* SkelMeshResource = Mesh->GetResourceForRendering(GRHIFeatureLevel);
+				FSkeletalMeshResource* SkelMeshResource = Mesh->GetResourceForRendering();
 				for( int32 LODIndex = 0; LODIndex < SkelMeshResource->LODModels.Num(); ++LODIndex )
 				{
 					const FStaticLODModel& LODModel = SkelMeshResource->LODModels[ LODIndex ];
@@ -313,7 +313,7 @@ struct PrimitiveStatsGenerator
 					USkeletalMesh* SkeletalMesh = SkeletalMeshComponent->SkeletalMesh;
 					if( SkeletalMesh )
 					{
-						FSkeletalMeshResource* SkelMeshResource = SkeletalMesh->GetResourceForRendering(GRHIFeatureLevel);
+						FSkeletalMeshResource* SkelMeshResource = SkeletalMesh->GetResourceForRendering();
 						if (SkelMeshResource->LODModels.Num())
 						{
 							const FStaticLODModel& BaseLOD = SkelMeshResource->LODModels[0];
diff --git a/Engine/Source/Editor/UnrealEd/Private/Animation/DebugSkelMeshComponent.cpp b/Engine/Source/Editor/UnrealEd/Private/Animation/DebugSkelMeshComponent.cpp
--- a/Engine/Source/Editor/UnrealEd/Private/Animation/DebugSkelMeshComponent.cpp
+++ b/Engine/Source/Editor/UnrealEd/Private/Animation/DebugSkelMeshComponent.cpp
@@ -192,8 +192,8 @@ void UDebugSkelMeshComponent::ConsumeRootMotion(const FVector& FloorMin, const F
 FPrimitiveSceneProxy* UDebugSkelMeshComponent::CreateSceneProxy()
 {
 	FDebugSkelMeshSceneProxy* Result = NULL;
-	ERHIFeatureLevel::Type SceneFeatureLevel = GRHIFeatureLevel;
-	FSkeletalMeshResource* SkelMeshResource = SkeletalMesh ? SkeletalMesh->GetResourceForRendering(SceneFeatureLevel) : NULL;
+	ERHIFeatureLevel::Type SceneFeatureLevel = GetWorld()->FeatureLevel;
+	FSkeletalMeshResource* SkelMeshResource = SkeletalMesh ? SkeletalMesh->GetResourceForRendering() : NULL;
 
 	// only create a scene proxy for rendering if
 	// properly initialized
diff --git a/Engine/Source/Editor/UnrealEd/Private/Lightmass/LightmassRender.cpp b/Engine/Source/Editor/UnrealEd/Private/Lightmass/LightmassRender.cpp
--- a/Engine/Source/Editor/UnrealEd/Private/Lightmass/LightmassRender.cpp
+++ b/Engine/Source/Editor/UnrealEd/Private/Lightmass/LightmassRender.cpp
@@ -166,7 +166,7 @@ class FLightmassMaterialProxy : public FMaterial, public FMaterialRenderProxy
 public:
 	FLightmassMaterialProxy(): FMaterial()
 	{
-		SetQualityLevelProperties(EMaterialQualityLevel::High, false, GRHIFeatureLevel);
+		SetQualityLevelProperties(EMaterialQualityLevel::High, false, GMaxRHIFeatureLevel);
 	}
 
 	/** Initializes the material proxy and kicks off async shader compiling. */
@@ -194,7 +194,7 @@ public:
 		// Special path for a MI with static parameters
 		if (MaterialInstance && MaterialInstance->bHasStaticPermutationResource && MaterialInstance->Parent)
 		{
-			FMaterialResource* MIResource = MaterialInstance->GetMaterialResource(GRHIFeatureLevel);
+			FMaterialResource* MIResource = MaterialInstance->GetMaterialResource(GMaxRHIFeatureLevel);
 
 			// Use the shader map Id from the static permutation
 			// This allows us to create a deterministic yet unique Id for the shader map that will be compiled for this FLightmassMaterialProxy
@@ -209,7 +209,7 @@ public:
 		}
 		else
 		{
-			FMaterialResource* MaterialResource = Material->GetMaterialResource(GRHIFeatureLevel);
+			FMaterialResource* MaterialResource = Material->GetMaterialResource(GMaxRHIFeatureLevel);
 
 			// Copy the material resource Id
 			// The FLightmassMaterialProxy's GetShaderMapUsage will set it apart from the MI's resource when it comes to finding a shader map
diff --git a/Engine/Source/Runtime/Core/Private/HAL/ConsoleManager.cpp b/Engine/Source/Runtime/Core/Private/HAL/ConsoleManager.cpp
--- a/Engine/Source/Runtime/Core/Private/HAL/ConsoleManager.cpp
+++ b/Engine/Source/Runtime/Core/Private/HAL/ConsoleManager.cpp
@@ -1859,3 +1859,9 @@ static TAutoConsoleVariable<int32> CVarPrecomputedVisibilityWarning(
 	0,
 	TEXT("If set to 1, a warning will be displayed when rendering a scene from a view point without precomputed visibility."),
 	ECVF_RenderThreadSafe);
+
+static TAutoConsoleVariable<int32> CVarFeatureLevelPreview(
+	TEXT("r.FeatureLevelPreview"),
+	0,
+	TEXT("If 1 the quick settings menu will contain an option to enable feature level preview modes"),
+	ECVF_RenderThreadSafe);
\ No newline at end of file
diff --git a/Engine/Source/Runtime/Core/Public/UObject/ObjectVersion.h b/Engine/Source/Runtime/Core/Public/UObject/ObjectVersion.h
--- a/Engine/Source/Runtime/Core/Public/UObject/ObjectVersion.h
+++ b/Engine/Source/Runtime/Core/Public/UObject/ObjectVersion.h
@@ -578,6 +578,8 @@ enum EUnrealEngineObjectUE4Version
 	VER_UE4_ADD_ROOTCOMPONENT_TO_FOLIAGEACTOR,
 	// FMaterialInstanceBasePropertyOverrides didn't use proper UObject serialize
 	VER_UE4_FIX_MATERIAL_PROPERTY_OVERRIDE_SERIALIZE,
+	// Addition of linear color sampler. color sample type is changed to linear sampler if source texture !sRGB
+	VER_UE4_ADD_LINEAR_COLOR_SAMPLER,
 
 	// -----<new versions can be added before this line>-------------------------------------------------
 	// - this needs to be the last line (see note below)
diff --git a/Engine/Source/Runtime/Engine/Classes/Engine/EngineTypes.h b/Engine/Source/Runtime/Engine/Classes/Engine/EngineTypes.h
--- a/Engine/Source/Runtime/Engine/Classes/Engine/EngineTypes.h
+++ b/Engine/Source/Runtime/Engine/Classes/Engine/EngineTypes.h
@@ -142,6 +142,8 @@ enum EMaterialSamplerType
 	SAMPLERTYPE_Alpha UMETA(DisplayName="Alpha"),
 	SAMPLERTYPE_Normal UMETA(DisplayName="Normal"),
 	SAMPLERTYPE_Masks UMETA(DisplayName="Masks"),
+	SAMPLERTYPE_LinearColor UMETA(DisplayName = "Linear Color"),
+	SAMPLERTYPE_LinearGrayscale UMETA(DisplayName = "Linear Grayscale"),
 	SAMPLERTYPE_MAX,
 };
 
diff --git a/Engine/Source/Runtime/Engine/Classes/Engine/SkeletalMesh.h b/Engine/Source/Runtime/Engine/Classes/Engine/SkeletalMesh.h
--- a/Engine/Source/Runtime/Engine/Classes/Engine/SkeletalMesh.h
+++ b/Engine/Source/Runtime/Engine/Classes/Engine/SkeletalMesh.h
@@ -781,7 +781,7 @@ public:
 	FORCEINLINE FSkeletalMeshResource* GetImportedResource() const { return ImportedResource.Get(); }
 
 	/** Get the resource to use for rendering. */
-	FORCEINLINE FSkeletalMeshResource* GetResourceForRendering(ERHIFeatureLevel::Type FeatureLevel) const { return GetImportedResource(); }
+	FORCEINLINE FSkeletalMeshResource* GetResourceForRendering() const { return GetImportedResource(); }
 
 	/** Skeleton of this skeletal mesh **/
 	UPROPERTY(Category=Mesh, AssetRegistrySearchable, VisibleAnywhere, BlueprintReadOnly)
diff --git a/Engine/Source/Runtime/Engine/Classes/Engine/World.h b/Engine/Source/Runtime/Engine/Classes/Engine/World.h
--- a/Engine/Source/Runtime/Engine/Classes/Engine/World.h
+++ b/Engine/Source/Runtime/Engine/Classes/Engine/World.h
@@ -598,6 +598,9 @@ public:
 	/** The interface to the scene manager for this world. */
 	class FSceneInterface*						Scene;
 
+	/** The current renderer feature level of this world */
+	ERHIFeatureLevel::Type						FeatureLevel;
+
 	/** Saved editor viewport states - one for each view type. Indexed using ELevelViewportType above.							*/
 	FLevelViewportInfo							EditorViews[4];
 
@@ -622,7 +625,7 @@ public:
 	/** Change the feature level that this world is current rendering with */
 	void ChangeFeatureLevel(ERHIFeatureLevel::Type InFeatureLevel);
 
-	static void ForceFeatureLevelUpdate(ERHIFeatureLevel::Type InFeatureLevel);
+	static void ChangeAllWorldFeatureLevels(ERHIFeatureLevel::Type InFeatureLevel);
 
 private:
 	/** List of all the controllers in the world. */
diff --git a/Engine/Source/Runtime/Engine/Private/Atmosphere/Atmosphere.h b/Engine/Source/Runtime/Engine/Private/Atmosphere/Atmosphere.h
--- a/Engine/Source/Runtime/Engine/Private/Atmosphere/Atmosphere.h
+++ b/Engine/Source/Runtime/Engine/Private/Atmosphere/Atmosphere.h
@@ -133,7 +133,7 @@ public:
 	 */
 	virtual void InitRHI() override
 	{
-		if ( TextureData && IsFeatureLevelSupported(GRHIShaderPlatform, ERHIFeatureLevel::SM3))
+		if (TextureData && GetFeatureLevel() >= ERHIFeatureLevel::SM3)
 		{
 			switch(TexType)
 			{
diff --git a/Engine/Source/Runtime/Engine/Private/GPUSort.cpp b/Engine/Source/Runtime/Engine/Private/GPUSort.cpp
--- a/Engine/Source/Runtime/Engine/Private/GPUSort.cpp
+++ b/Engine/Source/Runtime/Engine/Private/GPUSort.cpp
@@ -81,11 +81,11 @@ public:
 	 */
 	virtual void InitRHI() override
 	{
-		const int32 OffsetsCount = DIGIT_COUNT * MAX_GROUP_COUNT;
-		const int32 OffsetsBufferSize = OffsetsCount * sizeof(uint32);
-
-		if (GRHIFeatureLevel == ERHIFeatureLevel::SM5)
+		if (GetFeatureLevel() >= ERHIFeatureLevel::SM5)
 		{
+			const int32 OffsetsCount = DIGIT_COUNT * MAX_GROUP_COUNT;
+			const int32 OffsetsBufferSize = OffsetsCount * sizeof(uint32);
+		
 			for (int32 BufferIndex = 0; BufferIndex < 2; ++BufferIndex)
 			{
 				FRHIResourceCreateInfo CreateInfo;
@@ -182,7 +182,7 @@ public:
 	 */
 	virtual void InitRHI() override
 	{
-		if (GRHIFeatureLevel == ERHIFeatureLevel::SM5)
+		if (GetFeatureLevel() >= ERHIFeatureLevel::SM5)
 		{
 			FRHIResourceCreateInfo CreateInfo;
 			SortParametersBufferRHI = RHICreateVertexBuffer(
diff --git a/Engine/Source/Runtime/Engine/Private/InstancedStaticMesh.cpp b/Engine/Source/Runtime/Engine/Private/InstancedStaticMesh.cpp
--- a/Engine/Source/Runtime/Engine/Private/InstancedStaticMesh.cpp
+++ b/Engine/Source/Runtime/Engine/Private/InstancedStaticMesh.cpp
@@ -1122,7 +1122,7 @@ void UInstancedStaticMeshComponent::ApplyComponentInstanceData(TSharedPtr<FCompo
 FPrimitiveSceneProxy* UInstancedStaticMeshComponent::CreateSceneProxy()
 {
 	// We don't support instancing on ES2
-	if (GRHIFeatureLevel == ERHIFeatureLevel::ES2)
+	if (GetWorld()->FeatureLevel == ERHIFeatureLevel::ES2)
 	{
 		return NULL;
 	}
diff --git a/Engine/Source/Runtime/Engine/Private/Landscape/Landscape.cpp b/Engine/Source/Runtime/Engine/Private/Landscape/Landscape.cpp
--- a/Engine/Source/Runtime/Engine/Private/Landscape/Landscape.cpp
+++ b/Engine/Source/Runtime/Engine/Private/Landscape/Landscape.cpp
@@ -646,8 +646,9 @@ void ULandscapeComponent::BeginDestroy()
 
 FPrimitiveSceneProxy* ULandscapeComponent::CreateSceneProxy()
 {
+	const auto FeatureLevel = GetWorld()->FeatureLevel;
 	FPrimitiveSceneProxy* Proxy = NULL;
-	if (GRHIFeatureLevel >= ERHIFeatureLevel::SM3)
+	if (FeatureLevel >= ERHIFeatureLevel::SM3)
 	{
 #if WITH_EDITOR
 		if( EditToolRenderData == NULL )
@@ -659,7 +660,7 @@ FPrimitiveSceneProxy* ULandscapeComponent::CreateSceneProxy()
 		Proxy = new FLandscapeComponentSceneProxy(this, NULL);
 #endif
 	}
-	else if (GRHIFeatureLevel >= ERHIFeatureLevel::ES2)
+	else if (FeatureLevel >= ERHIFeatureLevel::ES2)
 	{
 #if WITH_EDITOR
 		if( !PlatformData.HasValidPlatformData() ) // Deferred generation
diff --git a/Engine/Source/Runtime/Engine/Private/Landscape/LandscapeRender.cpp b/Engine/Source/Runtime/Engine/Private/Landscape/LandscapeRender.cpp
--- a/Engine/Source/Runtime/Engine/Private/Landscape/LandscapeRender.cpp
+++ b/Engine/Source/Runtime/Engine/Private/Landscape/LandscapeRender.cpp
@@ -653,7 +653,7 @@ FLandscapeComponentSceneProxy::FLandscapeComponentSceneProxy(ULandscapeComponent
 	}
 #endif
 
-	bRequiresAdjacencyInformation = RequiresAdjacencyInformation( MaterialInterface, XYOffsetmapTexture == NULL ? &FLandscapeVertexFactory::StaticType : &FLandscapeXYOffsetVertexFactory::StaticType );
+	bRequiresAdjacencyInformation = RequiresAdjacencyInformation(MaterialInterface, XYOffsetmapTexture == NULL ? &FLandscapeVertexFactory::StaticType : &FLandscapeXYOffsetVertexFactory::StaticType, InComponent->GetWorld()->FeatureLevel);
 	SharedBuffersKey = SubsectionSizeQuads | (NumSubsections << 16) | ( XYOffsetmapTexture == NULL ? 0 : 0x80000000);
 
 	DynamicMesh.LCI = ComponentLightInfo; 
@@ -698,7 +698,7 @@ void FLandscapeComponentSceneProxy::CreateRenderThreadResources()
 	SharedBuffers = FLandscapeComponentSceneProxy::SharedBuffersMap.FindRef(SharedBuffersKey);
 	if( SharedBuffers == NULL )
 	{
-		SharedBuffers = new FLandscapeSharedBuffers(SharedBuffersKey, SubsectionSizeQuads, NumSubsections);
+		SharedBuffers = new FLandscapeSharedBuffers(SharedBuffersKey, SubsectionSizeQuads, NumSubsections, GetScene()->GetFeatureLevel());
 		FLandscapeComponentSceneProxy::SharedBuffersMap.Add(SharedBuffersKey, SharedBuffers);
 
 		if (!XYOffsetmapTexture)
@@ -1610,9 +1610,9 @@ void FLandscapeVertexBuffer::InitRHI()
 //
 
 template <typename INDEX_TYPE>
-void FLandscapeSharedBuffers::CreateIndexBuffers()
+void FLandscapeSharedBuffers::CreateIndexBuffers(ERHIFeatureLevel::Type InFeatureLevel)
 {
-	if (GRHIFeatureLevel == ERHIFeatureLevel::ES2)
+	if (InFeatureLevel == ERHIFeatureLevel::ES2)
 {
 		if (!bVertexScoresComputed)
 		{
@@ -1639,7 +1639,7 @@ void FLandscapeSharedBuffers::CreateIndexBuffers()
 		MaxIndexFull = 0;
 		MinIndexFull = MAX_int32;
 
-		if (GRHIFeatureLevel == ERHIFeatureLevel::ES2)
+		if (InFeatureLevel == ERHIFeatureLevel::ES2)
 		{
 			// ES2 version
 			float MipRatio = (float)SubsectionSizeQuads / (float)LodSubsectionSizeQuads; // Morph current MIP to base MIP
@@ -1816,15 +1816,16 @@ void FLandscapeSharedBuffers::CreateIndexBuffers()
 	}
 }
 
-FLandscapeSharedBuffers::FLandscapeSharedBuffers(int32 InSharedBuffersKey, int32 InSubsectionSizeQuads, int32 InNumSubsections)
+FLandscapeSharedBuffers::FLandscapeSharedBuffers(int32 InSharedBuffersKey, int32 InSubsectionSizeQuads, int32 InNumSubsections, ERHIFeatureLevel::Type InFeatureLevel)
 :	SharedBuffersKey(InSharedBuffersKey)
 ,	NumIndexBuffers(FMath::CeilLogTwo(InSubsectionSizeQuads+1))
 ,	SubsectionSizeVerts(InSubsectionSizeQuads+1)
 ,	NumSubsections(InNumSubsections)
 ,	VertexFactory(NULL)
+,	VertexBuffer(NULL)
 ,	AdjacencyIndexBuffers(NULL)
 {
-	if (GRHIFeatureLevel > ERHIFeatureLevel::ES2)
+	if (InFeatureLevel > ERHIFeatureLevel::ES2)
 	{
 		// Vertex Buffer cannot be shared
 		VertexBuffer = new FLandscapeVertexBuffer(SubsectionSizeVerts, NumSubsections);
@@ -1835,20 +1836,18 @@ FLandscapeSharedBuffers::FLandscapeSharedBuffers(int32 InSharedBuffersKey, int32
 	// See if we need to use 16 or 32-bit index buffers
 	if( FMath::Square(SubsectionSizeVerts) * FMath::Square(NumSubsections) > 65535 )
 	{
-		CreateIndexBuffers<uint32>();
+		CreateIndexBuffers<uint32>(InFeatureLevel);
 	}
 	else
 	{
-		CreateIndexBuffers<uint16>();
+		CreateIndexBuffers<uint16>(InFeatureLevel);
 	}
 }
 
 FLandscapeSharedBuffers::~FLandscapeSharedBuffers()
 {
-	if (GRHIFeatureLevel > ERHIFeatureLevel::ES2)
-	{
-		delete VertexBuffer;
-	}
+	delete VertexBuffer;
+	
 	for( int32 i=0;i<NumIndexBuffers;i++ )
 	{
 		IndexBuffers[i]->ReleaseResource();
diff --git a/Engine/Source/Runtime/Engine/Private/Landscape/LandscapeRenderMobile.cpp b/Engine/Source/Runtime/Engine/Private/Landscape/LandscapeRenderMobile.cpp
--- a/Engine/Source/Runtime/Engine/Private/Landscape/LandscapeRenderMobile.cpp
+++ b/Engine/Source/Runtime/Engine/Private/Landscape/LandscapeRenderMobile.cpp
@@ -197,7 +197,7 @@ void FLandscapeComponentSceneProxyMobile::CreateRenderThreadResources()
 	SharedBuffers = FLandscapeComponentSceneProxy::SharedBuffersMap.FindRef(SharedBuffersKey);
 	if( SharedBuffers == NULL )
 	{
-		SharedBuffers = new FLandscapeSharedBuffers(SharedBuffersKey, SubsectionSizeQuads, NumSubsections);
+		SharedBuffers = new FLandscapeSharedBuffers(SharedBuffersKey, SubsectionSizeQuads, NumSubsections, GetScene()->GetFeatureLevel());
 		FLandscapeComponentSceneProxy::SharedBuffersMap.Add(SharedBuffersKey, SharedBuffers);
 	}
 
diff --git a/Engine/Source/Runtime/Engine/Private/LocalPlayer.cpp b/Engine/Source/Runtime/Engine/Private/LocalPlayer.cpp
--- a/Engine/Source/Runtime/Engine/Private/LocalPlayer.cpp
+++ b/Engine/Source/Runtime/Engine/Private/LocalPlayer.cpp
@@ -1177,7 +1177,7 @@ bool ULocalPlayer::HandleListSkelMeshesCommand( const TCHAR* Cmd, FOutputDevice&
 		if( SkeletalMesh && SkeletalMeshComponents.Num() )
 		{
 			// Dump information about skeletal mesh.
-			FSkeletalMeshResource* SkelMeshResource = SkeletalMesh->GetResourceForRendering(GRHIFeatureLevel);
+			FSkeletalMeshResource* SkelMeshResource = SkeletalMesh->GetResourceForRendering();
 			check(SkelMeshResource->LODModels.Num());
 			UE_LOG(LogPlayerManagement, Log, TEXT("%5i Vertices for LOD 0 of %s"),SkelMeshResource->LODModels[0].NumVertices,*SkeletalMesh->GetFullName());
 
diff --git a/Engine/Source/Runtime/Engine/Private/Materials/HLSLMaterialTranslator.h b/Engine/Source/Runtime/Engine/Private/Materials/HLSLMaterialTranslator.h
--- a/Engine/Source/Runtime/Engine/Private/Materials/HLSLMaterialTranslator.h
+++ b/Engine/Source/Runtime/Engine/Private/Materials/HLSLMaterialTranslator.h
@@ -2112,7 +2112,11 @@ protected:
 			case SAMPLERTYPE_Color:
 				SampleCode = FString::Printf( TEXT("ProcessMaterialColorTextureLookup(%s)"), *SampleCode );
 				break;
-			
+
+			case SAMPLERTYPE_LinearColor:
+				SampleCode = FString::Printf(TEXT("ProcessMaterialLinearColorTextureLookup(%s)"), *SampleCode);
+			break;
+
 			case SAMPLERTYPE_Alpha:
 				// Sampling a single channel texture in D3D9 gives: (G,G,G)
 				// Sampling a single channel texture in D3D11 gives: (G,0,0)
@@ -2127,6 +2131,13 @@ protected:
 				SampleCode = FString::Printf( TEXT("ProcessMaterialGreyscaleTextureLookup((%s).r).rrrr"), *SampleCode );
 				break;
 
+			case SAMPLERTYPE_LinearGrayscale:
+				// Sampling a greyscale texture in D3D9 gives: (G,G,G)
+				// Sampling a greyscale texture in D3D11 gives: (G,0,0)
+				// This replication reproduces the D3D9 behavior in all cases.
+				SampleCode = FString::Printf(TEXT("ProcessMaterialLinearGreyscaleTextureLookup((%s).r).rrrr"), *SampleCode);
+				break;
+
 			case SAMPLERTYPE_Normal:
 				// Normal maps need to be unpacked in the pixel shader.
 				SampleCode = FString::Printf( TEXT("UnpackNormalMap(%s)"), *SampleCode );
diff --git a/Engine/Source/Runtime/Engine/Private/Materials/Material.cpp b/Engine/Source/Runtime/Engine/Private/Materials/Material.cpp
--- a/Engine/Source/Runtime/Engine/Private/Materials/Material.cpp
+++ b/Engine/Source/Runtime/Engine/Private/Materials/Material.cpp
@@ -1626,6 +1626,11 @@ void UMaterial::Serialize(FArchive& Ar)
 	{
 		GMaterialsThatNeedCommentFix.Set(this);
 	}
+
+	if (Ar.UE4Ver() < VER_UE4_ADD_LINEAR_COLOR_SAMPLER)
+	{
+		GMaterialsThatNeedSamplerFixup.Set(this);
+	}
 #endif // #if WITH_EDITOR
 
 	if( Ar.UE4Ver() < VER_UE4_MATERIAL_ATTRIBUTES_REORDERING )
@@ -1932,7 +1937,7 @@ void UMaterial::PostLoad()
 		const int32 ExpressionCount = Expressions.Num();
 		for ( int32 ExpressionIndex = 0; ExpressionIndex < ExpressionCount; ++ExpressionIndex )
 		{
-			UMaterialExpressionTextureSample* TextureExpression = Cast<UMaterialExpressionTextureSample>( Expressions[ ExpressionIndex ] );
+			UMaterialExpressionTextureBase* TextureExpression = Cast<UMaterialExpressionTextureBase>(Expressions[ExpressionIndex]);
 			if ( TextureExpression && TextureExpression->Texture )
 			{
 				switch( TextureExpression->Texture->CompressionSettings )
@@ -1942,7 +1947,7 @@ void UMaterial::PostLoad()
 					break;
 					
 				case TC_Grayscale:
-					TextureExpression->SamplerType = SAMPLERTYPE_Grayscale;
+					TextureExpression->SamplerType = TextureExpression->Texture->SRGB ? SAMPLERTYPE_Grayscale : SAMPLERTYPE_LinearGrayscale;
 					break;
 
 				case TC_Masks:
@@ -1953,7 +1958,7 @@ void UMaterial::PostLoad()
 					TextureExpression->SamplerType = SAMPLERTYPE_Alpha;
 					break;
 				default:
-					TextureExpression->SamplerType = SAMPLERTYPE_Color;
+					TextureExpression->SamplerType = TextureExpression->Texture->SRGB ? SAMPLERTYPE_Color : SAMPLERTYPE_LinearColor;
 					break;
 				}
 			}
diff --git a/Engine/Source/Runtime/Engine/Private/Materials/MaterialExpressions.cpp b/Engine/Source/Runtime/Engine/Private/Materials/MaterialExpressions.cpp
--- a/Engine/Source/Runtime/Engine/Private/Materials/MaterialExpressions.cpp
+++ b/Engine/Source/Runtime/Engine/Private/Materials/MaterialExpressions.cpp
@@ -154,6 +154,7 @@ if( ExpressionInput.Expression == ToBeRemovedExpression )										\
 FUObjectAnnotationSparseBool GMaterialFunctionsThatNeedExpressionsFlipped;
 FUObjectAnnotationSparseBool GMaterialFunctionsThatNeedCoordinateCheck;
 FUObjectAnnotationSparseBool GMaterialFunctionsThatNeedCommentFix;
+FUObjectAnnotationSparseBool GMaterialFunctionsThatNeedSamplerFixup;
 #endif // #if WITH_EDITOR
 
 /** Returns whether the given expression class is allowed. */
@@ -1081,13 +1082,13 @@ EMaterialSamplerType UMaterialExpressionTextureBase::GetSamplerTypeForTexture(co
 			case TC_Normalmap:
 				return SAMPLERTYPE_Normal;
 			case TC_Grayscale:
-				return SAMPLERTYPE_Grayscale;
+				return Texture->SRGB ? SAMPLERTYPE_Grayscale : SAMPLERTYPE_LinearGrayscale;
 			case TC_Alpha:
 				return SAMPLERTYPE_Alpha;
 			case TC_Masks:
 				return SAMPLERTYPE_Masks;
 			default:
-				return SAMPLERTYPE_Color;
+				return Texture->SRGB ? SAMPLERTYPE_Color : SAMPLERTYPE_LinearColor;
 		}
 	}
 	return SAMPLERTYPE_Color;
@@ -5399,8 +5400,8 @@ int32 UMaterialExpressionFontSample::Compile(class FMaterialCompiler* Compiler,
 				Texture = Texture = GEngine->DefaultTexture;
 			}
 			check(Texture);
-
-			if (!VerifySamplerType(Compiler, (Desc.Len() > 0 ? *Desc : TEXT("FontSample")), Texture, SAMPLERTYPE_Color))
+			EMaterialSamplerType ExpectedSamplerType = Texture->SRGB ? SAMPLERTYPE_Color : SAMPLERTYPE_LinearColor;
+			if (!VerifySamplerType(Compiler, (Desc.Len() > 0 ? *Desc : TEXT("FontSample")), Texture, ExpectedSamplerType))
 			{
 				return INDEX_NONE;
 			}
@@ -5409,7 +5410,7 @@ int32 UMaterialExpressionFontSample::Compile(class FMaterialCompiler* Compiler,
 			Result = Compiler->TextureSample(
 				TextureCodeIndex,
 				Compiler->TextureCoordinate(0, false, false),
-				SAMPLERTYPE_Color
+				ExpectedSamplerType
 			);
 		}
 		else
@@ -5499,7 +5500,8 @@ int32 UMaterialExpressionFontSampleParameter::Compile(class FMaterialCompiler* C
 			Texture = Texture = GEngine->DefaultTexture;
 		}
 		check(Texture);
-		if (!VerifySamplerType(Compiler, (Desc.Len() > 0 ? *Desc : TEXT("FontSampleParameter")), Texture, SAMPLERTYPE_Color))
+		EMaterialSamplerType ExpectedSamplerType = Texture->SRGB ? SAMPLERTYPE_Color : SAMPLERTYPE_LinearColor;
+		if (!VerifySamplerType(Compiler, (Desc.Len() > 0 ? *Desc : TEXT("FontSampleParameter")), Texture, ExpectedSamplerType))
 		{
 			return INDEX_NONE;
 		}
@@ -5507,7 +5509,7 @@ int32 UMaterialExpressionFontSampleParameter::Compile(class FMaterialCompiler* C
 		Result = Compiler->TextureSample(
 			TextureCodeIndex,
 			Compiler->TextureCoordinate(0, false, false),
-			SAMPLERTYPE_Color
+			ExpectedSamplerType
 		);
 	}
 	return Result;
@@ -6219,6 +6221,11 @@ void UMaterialFunction::Serialize(FArchive& Ar)
 	{
 		GMaterialFunctionsThatNeedCommentFix.Set(this);
 	}
+
+	if (Ar.UE4Ver() < VER_UE4_ADD_LINEAR_COLOR_SAMPLER)
+	{
+		GMaterialFunctionsThatNeedSamplerFixup.Set(this);
+	}
 #endif // #if WITH_EDITOR
 }
 
@@ -6278,6 +6285,40 @@ void UMaterialFunction::PostLoad()
 		GMaterialFunctionsThatNeedCommentFix.Clear(this);
 		UMaterial::FixCommentPositions(FunctionEditorComments);
 	}
+
+	if (GMaterialFunctionsThatNeedSamplerFixup.Get(this))
+	{
+		GMaterialFunctionsThatNeedSamplerFixup.Clear(this);
+		const int32 ExpressionCount = FunctionExpressions.Num();
+		for (int32 ExpressionIndex = 0; ExpressionIndex < ExpressionCount; ++ExpressionIndex)
+		{
+			UMaterialExpressionTextureBase* TextureExpression = Cast<UMaterialExpressionTextureBase>(FunctionExpressions[ExpressionIndex]);
+			if (TextureExpression && TextureExpression->Texture)
+			{
+				switch (TextureExpression->Texture->CompressionSettings)
+				{
+				case TC_Normalmap:
+					TextureExpression->SamplerType = SAMPLERTYPE_Normal;
+					break;
+
+				case TC_Grayscale:
+					TextureExpression->SamplerType = TextureExpression->Texture->SRGB ? SAMPLERTYPE_Grayscale : SAMPLERTYPE_LinearGrayscale;
+					break;
+
+				case TC_Masks:
+					TextureExpression->SamplerType = SAMPLERTYPE_Masks;
+					break;
+
+				case TC_Alpha:
+					TextureExpression->SamplerType = SAMPLERTYPE_Alpha;
+					break;
+				default:
+					TextureExpression->SamplerType = TextureExpression->Texture->SRGB ? SAMPLERTYPE_Color : SAMPLERTYPE_LinearColor;
+					break;
+				}
+			}
+		}
+	}
 #endif // #if WITH_EDITOR
 }
 
diff --git a/Engine/Source/Runtime/Engine/Private/MeshParticleVertexFactory.cpp b/Engine/Source/Runtime/Engine/Private/MeshParticleVertexFactory.cpp
--- a/Engine/Source/Runtime/Engine/Private/MeshParticleVertexFactory.cpp
+++ b/Engine/Source/Runtime/Engine/Private/MeshParticleVertexFactory.cpp
@@ -75,7 +75,7 @@ void FMeshParticleVertexFactory::InitRHI()
 {
 	FVertexDeclarationElementList Elements;
 
-	const bool bInstanced = GRHIFeatureLevel >= ERHIFeatureLevel::SM3;
+	const bool bInstanced = GetFeatureLevel() >= ERHIFeatureLevel::SM3;
 
 	if (Data.bInitialized)
 	{
diff --git a/Engine/Source/Runtime/Engine/Private/NiagaraComponent.cpp b/Engine/Source/Runtime/Engine/Private/NiagaraComponent.cpp
--- a/Engine/Source/Runtime/Engine/Private/NiagaraComponent.cpp
+++ b/Engine/Source/Runtime/Engine/Private/NiagaraComponent.cpp
@@ -32,6 +32,7 @@ public:
 	FNiagaraSceneProxy(const UNiagaraComponent* InComponent)
 		:	FPrimitiveSceneProxy(InComponent)
 		,	DynamicData(NULL)
+		, VertexFactory(PVFT_Sprite, InComponent->GetWorld()->FeatureLevel)
 	{
 		Material = InComponent->Material;
 		if (!Material || !Material->CheckMaterialUsage(MATUSAGE_ParticleSprites))
diff --git a/Engine/Source/Runtime/Engine/Private/Particles/ParticleGpuSimulation.cpp b/Engine/Source/Runtime/Engine/Private/Particles/ParticleGpuSimulation.cpp
--- a/Engine/Source/Runtime/Engine/Private/Particles/ParticleGpuSimulation.cpp
+++ b/Engine/Source/Runtime/Engine/Private/Particles/ParticleGpuSimulation.cpp
@@ -2209,7 +2209,7 @@ public:
 	/** Initialize RHI resources. */
 	virtual void InitRHI() override
 	{
-		if ( ParticleCount > 0 && CurrentRHISupportsGPUParticles() )
+		if ( ParticleCount > 0 && RHISupportsGPUParticles(GetFeatureLevel()) )
 		{
 			const int32 BufferStride = sizeof(FParticleIndex);
 			const int32 BufferSize = ParticleCount * BufferStride;
@@ -3904,7 +3904,7 @@ void FFXSystem::RemoveGPUSimulation(FParticleSimulationGPU* Simulation)
 
 int32 FFXSystem::AddSortedGPUSimulation(FParticleSimulationGPU* Simulation, const FVector& ViewOrigin)
 {
-	check(GRHIFeatureLevel == ERHIFeatureLevel::SM5);
+	check(FeatureLevel == ERHIFeatureLevel::SM5);
 	const int32 BufferOffset = ParticleSimulationResources->SortedParticleCount;
 	ParticleSimulationResources->SortedParticleCount += Simulation->VertexBuffer.ParticleCount;
 	FParticleSimulationSortInfo* SortInfo = new(ParticleSimulationResources->SimulationsToSort) FParticleSimulationSortInfo();
@@ -4384,7 +4384,7 @@ static void ClearGPUSpriteResourceData( FGPUSpriteResources* Resources )
 FGPUSpriteResources* BeginCreateGPUSpriteResources( const FGPUSpriteResourceData& InResourceData )
 {
 	FGPUSpriteResources* Resources = NULL;
-	if (CurrentRHISupportsGPUParticles())
+	if (RHISupportsGPUParticles(GMaxRHIFeatureLevel))
 	{
 		Resources = new FGPUSpriteResources;
 		SetGPUSpriteResourceData( Resources, InResourceData );
diff --git a/Engine/Source/Runtime/Engine/Private/Particles/ParticleModules_Collision.cpp b/Engine/Source/Runtime/Engine/Private/Particles/ParticleModules_Collision.cpp
--- a/Engine/Source/Runtime/Engine/Private/Particles/ParticleModules_Collision.cpp
+++ b/Engine/Source/Runtime/Engine/Private/Particles/ParticleModules_Collision.cpp
@@ -658,7 +658,7 @@ bool UParticleModuleCollisionGPU::IsValidForLODLevel(UParticleLODLevel* LODLevel
 	check(Material);
 
 	EBlendMode BlendMode = BLEND_Opaque;
-	const FMaterialResource* MaterialResource = Material->GetMaterialResource(GRHIFeatureLevel);
+	const FMaterialResource* MaterialResource = Material->GetMaterialResource(GetWorld()->FeatureLevel);
 	if(MaterialResource)
 	{
 		BlendMode = MaterialResource->GetBlendMode();
diff --git a/Engine/Source/Runtime/Engine/Private/Particles/ParticleResources.cpp b/Engine/Source/Runtime/Engine/Private/Particles/ParticleResources.cpp
--- a/Engine/Source/Runtime/Engine/Private/Particles/ParticleResources.cpp
+++ b/Engine/Source/Runtime/Engine/Private/Particles/ParticleResources.cpp
@@ -68,13 +68,13 @@ void FParticleScratchVertexBuffer::InitRHI()
 {
 	// Create a scratch vertex buffer for injecting particles and rendering tiles.
 	uint32 Flags = BUF_Volatile;
-	if (GRHIFeatureLevel >= ERHIFeatureLevel::SM4)
+	if (GetFeatureLevel() >= ERHIFeatureLevel::SM4)
 	{
 		Flags |= BUF_ShaderResource;
 	}
 	FRHIResourceCreateInfo CreateInfo;
 	VertexBufferRHI = RHICreateVertexBuffer(GParticleScratchVertexBufferSize, Flags, CreateInfo);
-	if (GRHIFeatureLevel >= ERHIFeatureLevel::SM4)
+	if (GetFeatureLevel() >= ERHIFeatureLevel::SM4)
 	{
 		VertexBufferSRV_G32R32F = RHICreateShaderResourceView( VertexBufferRHI, /*Stride=*/ sizeof(FVector2D), PF_G32R32F );
 	}
diff --git a/Engine/Source/Runtime/Engine/Private/Particles/ParticleResources.h b/Engine/Source/Runtime/Engine/Private/Particles/ParticleResources.h
--- a/Engine/Source/Runtime/Engine/Private/Particles/ParticleResources.h
+++ b/Engine/Source/Runtime/Engine/Private/Particles/ParticleResources.h
@@ -46,14 +46,17 @@ typedef FVertexBufferRHIParamRef FParticleBufferParamRef;
 class FParticleScratchVertexBuffer : public FVertexBuffer
 {
 public:
-	/** SRV in to the buffer as an array of FVector2D values. */
-	FShaderResourceViewRHIRef VertexBufferSRV_G32R32F;
-	
+
 	FParticleShaderParamRef GetShaderParam();
 	FParticleBufferParamRef GetBufferParam();
 
 	virtual void InitRHI() override;
 	virtual void ReleaseRHI() override;
+
+private:
+
+	/** SRV in to the buffer as an array of FVector2D values. */
+	FShaderResourceViewRHIRef VertexBufferSRV_G32R32F;
 };
 
 /** The global scratch vertex buffer. */
diff --git a/Engine/Source/Runtime/Engine/Private/Particles/ParticleSortingGPU.cpp b/Engine/Source/Runtime/Engine/Private/Particles/ParticleSortingGPU.cpp
--- a/Engine/Source/Runtime/Engine/Private/Particles/ParticleSortingGPU.cpp
+++ b/Engine/Source/Runtime/Engine/Private/Particles/ParticleSortingGPU.cpp
@@ -250,7 +250,7 @@ static int32 GenerateParticleSortKeys(
  */
 void FParticleSortBuffers::InitRHI()
 {
-	if (GRHIFeatureLevel == ERHIFeatureLevel::SM5)
+	if (GetFeatureLevel() >= ERHIFeatureLevel::SM5)
 	{
 		for (int32 BufferIndex = 0; BufferIndex < 2; ++BufferIndex)
 		{
diff --git a/Engine/Source/Runtime/Engine/Private/Particles/ParticleSystemRender.cpp b/Engine/Source/Runtime/Engine/Private/Particles/ParticleSystemRender.cpp
--- a/Engine/Source/Runtime/Engine/Private/Particles/ParticleSystemRender.cpp
+++ b/Engine/Source/Runtime/Engine/Private/Particles/ParticleSystemRender.cpp
@@ -272,22 +272,24 @@ FParticleOrderPool GParticleOrderPool;
 // Particle vertex factory pool
 FParticleVertexFactoryPool GParticleVertexFactoryPool;
 
-FParticleVertexFactoryBase* FParticleVertexFactoryPool::GetParticleVertexFactory(EParticleVertexFactoryType InType)
+FParticleVertexFactoryBase* FParticleVertexFactoryPool::GetParticleVertexFactory(EParticleVertexFactoryType InType, ERHIFeatureLevel::Type InFeatureLevel)
 {
 	SCOPE_CYCLE_COUNTER(STAT_ParticlePoolTime);
 	check(InType < PVFT_MAX);
 	FParticleVertexFactoryBase* VertexFactory = NULL;
-	if (VertexFactoriesAvailable[InType].Num() == 0)
+	auto& AvailableFactories = VertexFactoriesAvailable[InType][InFeatureLevel];
+
+	if (AvailableFactories.Num() == 0)
 	{
 		// If there are none in the pool, create a new one, add it to the in use list and return it
-		VertexFactory = CreateParticleVertexFactory(InType);
+		VertexFactory = CreateParticleVertexFactory(InType, InFeatureLevel);
 		VertexFactories.Add(VertexFactory);
 	}
 	else
 	{
 		// Otherwise, pull one out of the available array
-		VertexFactory = VertexFactoriesAvailable[InType][VertexFactoriesAvailable[InType].Num() - 1];
-		VertexFactoriesAvailable[InType].RemoveAt(VertexFactoriesAvailable[InType].Num() - 1);
+		VertexFactory = AvailableFactories[AvailableFactories.Num() - 1];
+		AvailableFactories.RemoveAt(AvailableFactories.Num() - 1);
 	}
 	check(VertexFactory);
 	// Set it to true to indicate it is in use
@@ -300,7 +302,7 @@ bool FParticleVertexFactoryPool::ReturnParticleVertexFactory(FParticleVertexFact
 	SCOPE_CYCLE_COUNTER(STAT_ParticlePoolTime);
 	// Set it to false to indicate it is not in use
 	InVertexFactory->SetInUse(false);
-	VertexFactoriesAvailable[InVertexFactory->GetParticleFactoryType()].Add(InVertexFactory);
+	VertexFactoriesAvailable[InVertexFactory->GetParticleFactoryType()][InVertexFactory->GetFeatureLevel()].Add(InVertexFactory);
 	return true;
 }
 
@@ -325,15 +327,18 @@ void FParticleVertexFactoryPool::ClearPoolInternal()
 	// We can't safely touched the 'in-use' ones... 
 	for (int32 PoolIdx = 0; PoolIdx < PVFT_MAX; PoolIdx++)
 	{
-		for (int32 RemoveIdx = VertexFactoriesAvailable[PoolIdx].Num() - 1; RemoveIdx >= 0; RemoveIdx--)
+		for (int32 FeatureLevelIdx = 0; FeatureLevelIdx < ERHIFeatureLevel::Num; FeatureLevelIdx++)
 		{
-			FParticleVertexFactoryBase* VertexFactory = VertexFactoriesAvailable[PoolIdx][RemoveIdx];
-			if(VertexFactory != NULL)
+			for (int32 RemoveIdx = VertexFactoriesAvailable[PoolIdx][FeatureLevelIdx].Num() - 1; RemoveIdx >= 0; RemoveIdx--)
 			{
-				VertexFactory->ReleaseResource();
-				delete VertexFactory;
+				FParticleVertexFactoryBase* VertexFactory = VertexFactoriesAvailable[PoolIdx][FeatureLevelIdx][RemoveIdx];
+				if (VertexFactory != NULL)
+				{
+					VertexFactory->ReleaseResource();
+					delete VertexFactory;
+				}
+				VertexFactoriesAvailable[PoolIdx][FeatureLevelIdx].RemoveAt(RemoveIdx);
 			}
-			VertexFactoriesAvailable[PoolIdx].RemoveAt(RemoveIdx);
 		}
 	}
 }
@@ -376,12 +381,15 @@ void FParticleVertexFactoryPool::DumpInfo(FOutputDevice& Ar)
 	int32 TotalMemory = 0;
 	for (int32 PoolIdx = 0; PoolIdx < PVFT_MAX; PoolIdx++)
 	{
-		int32 LocalMemory = GetTypeSize((EParticleVertexFactoryType)PoolIdx) * VertexFactoriesAvailable[PoolIdx].Num();
-		Ar.Logf(TEXT("%s,%d,%d"), 
-			GetTypeString((EParticleVertexFactoryType)PoolIdx), 
-			VertexFactoriesAvailable[PoolIdx].Num(),
-			LocalMemory);
-		TotalMemory += LocalMemory;
+		for (int32 FeatureLevelIdx = 0; FeatureLevelIdx < ERHIFeatureLevel::Num; FeatureLevelIdx++)
+		{
+			int32 LocalMemory = GetTypeSize((EParticleVertexFactoryType)PoolIdx) * VertexFactoriesAvailable[PoolIdx][FeatureLevelIdx].Num();
+			Ar.Logf(TEXT("%s,%d,%d"),
+				GetTypeString((EParticleVertexFactoryType)PoolIdx),
+				VertexFactoriesAvailable[PoolIdx][FeatureLevelIdx].Num(),
+				LocalMemory);
+			TotalMemory += LocalMemory;
+		}
 	}
 	Ar.Logf(TEXT("TotalMemory Taken in Pool: %d"), TotalMemory);
 	TotalMemory = 0;
@@ -419,19 +427,19 @@ void FParticleVertexFactoryPool::DumpInfo(FOutputDevice& Ar)
  *
  *	@return	FParticleVertexFactoryBase*	The created VF; NULL if invalid InType
  */
-FParticleVertexFactoryBase* FParticleVertexFactoryPool::CreateParticleVertexFactory(EParticleVertexFactoryType InType)
+FParticleVertexFactoryBase* FParticleVertexFactoryPool::CreateParticleVertexFactory(EParticleVertexFactoryType InType, ERHIFeatureLevel::Type InFeatureLevel)
 {
 	FParticleVertexFactoryBase* NewVertexFactory = NULL;
 	switch (InType)
 	{
 	case PVFT_Sprite:
-		NewVertexFactory = new FParticleSpriteVertexFactory();
+		NewVertexFactory = new FParticleSpriteVertexFactory(PVFT_Sprite, InFeatureLevel);
 		break;
 	case PVFT_BeamTrail:
-		NewVertexFactory = new FParticleBeamTrailVertexFactory();
+		NewVertexFactory = new FParticleBeamTrailVertexFactory(PVFT_BeamTrail, InFeatureLevel);
 		break;
 	case PVFT_Mesh:
-		NewVertexFactory = new FMeshParticleVertexFactory();
+		NewVertexFactory = new FMeshParticleVertexFactory(PVFT_Mesh, InFeatureLevel);
 		break;
 	default:
 		break;
@@ -1125,7 +1133,8 @@ bool FDynamicSpriteEmitterData::GetVertexAndIndexDataNonInstanced(void* VertexDa
  */
 void FDynamicSpriteEmitterData::PreRenderView(FParticleSystemSceneProxy* Proxy, const FSceneViewFamily* ViewFamily, const uint32 VisibilityMap, int32 FrameNumber)
 {
-	const bool bInstanced = GRHIFeatureLevel >= ERHIFeatureLevel::SM3;
+	const auto FeatureLevel = ViewFamily->Scene->GetFeatureLevel();
+	const bool bInstanced = FeatureLevel >= ERHIFeatureLevel::SM3;
 
 	// Sort and generate particles for this view.
 	const FDynamicSpriteEmitterReplayDataBase* SourceData = GetSourceData();
@@ -1133,7 +1142,7 @@ void FDynamicSpriteEmitterData::PreRenderView(FParticleSystemSceneProxy* Proxy,
 	{
 		// Determine how many vertices and indices are needed to render.
 		const int32 ParticleCount = SourceData->ActiveParticleCount;
-		const int32 VertexSize = GetDynamicVertexStride(ViewFamily->Scene->GetFeatureLevel());
+		const int32 VertexSize = GetDynamicVertexStride(FeatureLevel);
 		const int32 DynamicParameterVertexSize = sizeof(FParticleVertexDynamicParameter);
 		const int32 VertexPerParticle = bInstanced ? 1 : 4;
 		const int32 ViewCount = ViewFamily->Views.Num();
@@ -1173,7 +1182,7 @@ void FDynamicSpriteEmitterData::PreRenderView(FParticleSystemSceneProxy* Proxy,
 					if (SourceData->SortMode != PSORTMODE_None)
 					{
 						// If material is using unlit translucency and the blend mode is translucent then we need to sort (back to front)
-						const FMaterial* Material = MaterialResource[bSelected]->GetMaterial(GRHIFeatureLevel);
+						const FMaterial* Material = MaterialResource[bSelected]->GetMaterial(FeatureLevel);
 						if (Material && 
 							(Material->GetBlendMode() == BLEND_Translucent ||
 							((SourceData->SortMode == PSORTMODE_Age_OldestFirst) || (SourceData->SortMode == PSORTMODE_Age_NewestFirst)))
@@ -1454,7 +1463,7 @@ void FDynamicSpriteEmitterData::CreateRenderThreadResources(const FParticleSyste
 	// Create the vertex factory...
 	if (VertexFactory == NULL)
 	{
-		VertexFactory = GParticleVertexFactoryPool.GetParticleVertexFactory(PVFT_Sprite);
+		VertexFactory = GParticleVertexFactoryPool.GetParticleVertexFactory(PVFT_Sprite, InOwnerProxy->GetScene()->GetFeatureLevel());
 		check(VertexFactory);
 	}
 
@@ -1676,7 +1685,7 @@ void FDynamicMeshEmitterData::CreateRenderThreadResources(const FParticleSystemS
 	// Create the vertex factory
 	if (VertexFactory == NULL)
 	{
-		VertexFactory = GParticleVertexFactoryPool.GetParticleVertexFactory(PVFT_Mesh);
+		VertexFactory = GParticleVertexFactoryPool.GetParticleVertexFactory(PVFT_Mesh, InOwnerProxy->GetScene()->GetFeatureLevel());
 		check(VertexFactory);
 		SetupVertexFactory((FMeshParticleVertexFactory*)VertexFactory, StaticMesh->RenderData->LODResources[0]);
 	}
@@ -1721,7 +1730,7 @@ int32 FDynamicMeshEmitterData::Render(FParticleSystemSceneProxy* Proxy, FPrimiti
 		return 0;
 	}
 
-	const bool bInstanced = GRHIFeatureLevel >= ERHIFeatureLevel::SM3;
+	const bool bInstanced = View->GetFeatureLevel() >= ERHIFeatureLevel::SM3;
 
 	int32 NumDraws = 0;
 	if (Source.EmitterRenderMode == ERM_Normal)
@@ -1795,7 +1804,9 @@ void FDynamicMeshEmitterData::PreRenderView(FParticleSystemSceneProxy* Proxy, co
 		return;
 	}
 
-	const bool bInstanced = GRHIFeatureLevel >= ERHIFeatureLevel::SM3;
+	const auto FeatureLevel = ViewFamily->Scene->GetFeatureLevel();
+	const auto ShaderPlatform = GShaderPlatformForFeatureLevel[FeatureLevel];
+	const bool bInstanced = FeatureLevel >= ERHIFeatureLevel::SM3;
 
 	int32 ParticleCount = Source.ActiveParticleCount;
 	if ((Source.MaxDrawCount >= 0) && (ParticleCount > Source.MaxDrawCount))
@@ -1803,7 +1814,7 @@ void FDynamicMeshEmitterData::PreRenderView(FParticleSystemSceneProxy* Proxy, co
 		ParticleCount = Source.MaxDrawCount;
 	}
 
-	const int32 InstanceVertexStride  = GetDynamicVertexStride(ViewFamily->Scene->GetFeatureLevel());
+	const int32 InstanceVertexStride = GetDynamicVertexStride(FeatureLevel);
 	const int32 DynamicParameterVertexStride  = GetDynamicParameterVertexStride();
 	const int32 ViewCount = ViewFamily->Views.Num();
 	if(bInstanced)
@@ -1915,7 +1926,7 @@ void FDynamicMeshEmitterData::PreRenderView(FParticleSystemSceneProxy* Proxy, co
 					if (bIsWireframe)
 					{
 						if( LODModel.WireframeIndexBuffer.IsInitialized()
-							&& !(RHISupportsTessellation(GRHIShaderPlatform) && Mesh.VertexFactory->GetType()->SupportsTessellationShaders())
+							&& !(RHISupportsTessellation(ShaderPlatform) && Mesh.VertexFactory->GetType()->SupportsTessellationShaders())
 							)
 						{
 							Mesh.Type = PT_LineList;
@@ -2582,7 +2593,7 @@ void FDynamicBeam2EmitterData::CreateRenderThreadResources(const FParticleSystem
 	//@todo. Cache these??
 	if (VertexFactory == NULL)
 	{
-		VertexFactory = (FParticleBeamTrailVertexFactory*)(GParticleVertexFactoryPool.GetParticleVertexFactory(PVFT_BeamTrail));
+		VertexFactory = (FParticleBeamTrailVertexFactory*)(GParticleVertexFactoryPool.GetParticleVertexFactory(PVFT_BeamTrail, InOwnerProxy->GetScene()->GetFeatureLevel()));
 		check(VertexFactory);
 	}
 }
@@ -5415,7 +5426,7 @@ void FDynamicTrailsEmitterData::CreateRenderThreadResources(const FParticleSyste
 	// Create the vertex factory...
 	if (VertexFactory == NULL)
 	{
-		VertexFactory = (FParticleBeamTrailVertexFactory*)(GParticleVertexFactoryPool.GetParticleVertexFactory(PVFT_BeamTrail));
+		VertexFactory = (FParticleBeamTrailVertexFactory*)(GParticleVertexFactoryPool.GetParticleVertexFactory(PVFT_BeamTrail, InOwnerProxy->GetScene()->GetFeatureLevel()));
 		check(VertexFactory);
 	}
 }
diff --git a/Engine/Source/Runtime/Engine/Private/Particles/ParticleVertexFactory.cpp b/Engine/Source/Runtime/Engine/Private/Particles/ParticleVertexFactory.cpp
--- a/Engine/Source/Runtime/Engine/Private/Particles/ParticleVertexFactory.cpp
+++ b/Engine/Source/Runtime/Engine/Private/Particles/ParticleVertexFactory.cpp
@@ -62,14 +62,18 @@ public:
 
 	FVertexDeclarationRHIRef VertexDeclarationRHI;
 
+	// Constructor.
+	FParticleSpriteVertexDeclaration(bool bInInstanced)
+		: bInstanced(bInInstanced)
+	{
+
+	}
+
 	// Destructor.
 	virtual ~FParticleSpriteVertexDeclaration() {}
 
 	virtual void FillDeclElements(FVertexDeclarationElementList& Elements, int32& Offset)
 	{
-		// Likely need to switch to two decls when enabling runtime ES2 preview.
-		const bool bInstanced = GRHIFeatureLevel >= ERHIFeatureLevel::SM3;
-
 		uint32 InitialStride = sizeof(float) * 2;
 		uint32 PerParticleStride = (sizeof(float) * 4) * 4;
 
@@ -124,10 +128,27 @@ public:
 	{
 		VertexDeclarationRHI.SafeRelease();
 	}
+
+private:
+
+	bool bInstanced;
 };
 
 /** The simple element vertex declaration. */
-static TGlobalResource<FParticleSpriteVertexDeclaration> GParticleSpriteVertexDeclaration;
+static TGlobalResource<FParticleSpriteVertexDeclaration> GParticleSpriteVertexDeclarationInstanced(true);
+static TGlobalResource<FParticleSpriteVertexDeclaration> GParticleSpriteVertexDeclarationNonInstanced(false);
+
+inline TGlobalResource<FParticleSpriteVertexDeclaration>& GetParticleSpriteVertexDeclaration(ERHIFeatureLevel::Type InFeatureLevel)
+{
+	if (InFeatureLevel >= ERHIFeatureLevel::SM3)
+	{
+		return GParticleSpriteVertexDeclarationInstanced;
+	}
+	else
+	{
+		return GParticleSpriteVertexDeclarationNonInstanced;
+	}
+}
 
 bool FParticleSpriteVertexFactory::ShouldCache(EShaderPlatform Platform, const class FMaterial* Material, const class FShaderType* ShaderType)
 {
@@ -151,12 +172,12 @@ void FParticleSpriteVertexFactory::ModifyCompilationEnvironment(EShaderPlatform
 void FParticleSpriteVertexFactory::InitRHI()
 {
 	InitStreams();
-	SetDeclaration(GParticleSpriteVertexDeclaration.VertexDeclarationRHI);
+	SetDeclaration(GetParticleSpriteVertexDeclaration(GetFeatureLevel()).VertexDeclarationRHI);
 }
 
 void FParticleSpriteVertexFactory::InitStreams()
 {
-	const bool bInstanced = GRHIFeatureLevel >= ERHIFeatureLevel::SM3;
+	const bool bInstanced = GetFeatureLevel() >= ERHIFeatureLevel::SM3;
 
 	check(Streams.Num() == 0);
 	if(bInstanced) 
diff --git a/Engine/Source/Runtime/Engine/Private/ReflectionCaptureComponent.cpp b/Engine/Source/Runtime/Engine/Private/ReflectionCaptureComponent.cpp
--- a/Engine/Source/Runtime/Engine/Private/ReflectionCaptureComponent.cpp
+++ b/Engine/Source/Runtime/Engine/Private/ReflectionCaptureComponent.cpp
@@ -958,6 +958,9 @@ void UReflectionCaptureComponent::PostLoad()
 {
 	Super::PostLoad();
 
+	static const auto CVar = IConsoleManager::Get().FindTConsoleVariableDataInt(TEXT("r.FeatureLevelPreview"));
+	bool bRetainAllFeatureLevelData = CVar->GetValueOnGameThread() == 1 && GIsEditor && GMaxRHIFeatureLevel >= ERHIFeatureLevel::SM3;
+
 	// If we're loading on a platform that doesn't require cooked data, attempt to load missing data from the DDC
 	if (!FPlatformProperties::RequiresCookedData())
 	{
@@ -990,31 +993,36 @@ void UReflectionCaptureComponent::PostLoad()
 		// If we have full HDR data but not encoded HDR data, generate the encoded data now
 		if (FullHDRDerivedData 
 			&& !EncodedHDRDerivedData 
-			&& GRHIFeatureLevel == ERHIFeatureLevel::ES2)
+			&& (GRHIFeatureLevel == ERHIFeatureLevel::ES2 || bRetainAllFeatureLevelData))
 		{
 			EncodedHDRDerivedData = FReflectionCaptureEncodedHDRDerivedData::GenerateEncodedHDRData(*FullHDRDerivedData, StateId, Brightness);
 		}
 	}
 	
 	// Initialize rendering resources for the current feature level, and toss data only needed by other feature levels
-	if (FullHDRDerivedData && GRHIFeatureLevel >= ERHIFeatureLevel::SM4)
+	if (FullHDRDerivedData && (GRHIFeatureLevel >= ERHIFeatureLevel::SM4 || bRetainAllFeatureLevelData))
 	{
 		// Don't need encoded HDR data for rendering on this feature level
-		INC_MEMORY_STAT_BY(STAT_ReflectionCaptureMemory,FullHDRDerivedData->CompressedCapturedData.GetAllocatedSize());
-		EncodedHDRDerivedData = NULL;
+		INC_MEMORY_STAT_BY(STAT_ReflectionCaptureMemory, FullHDRDerivedData->CompressedCapturedData.GetAllocatedSize());
 
-		if (GRHIFeatureLevel == ERHIFeatureLevel::SM4)
+		if ((GRHIFeatureLevel == ERHIFeatureLevel::SM4 || bRetainAllFeatureLevelData))
 		{
 			SM4FullHDRCubemapTexture = new FReflectionTextureCubeResource();
 			SM4FullHDRCubemapTexture->SetupParameters(GReflectionCaptureSize, FMath::CeilLogTwo(GReflectionCaptureSize) + 1, PF_FloatRGBA, &FullHDRDerivedData->GetCapturedDataForSM4Load());
 			BeginInitResource(SM4FullHDRCubemapTexture);
 		}
+
+		if (!bRetainAllFeatureLevelData)
+		{
+			EncodedHDRDerivedData = NULL;
+		}
 	}
-	else if (EncodedHDRDerivedData && GRHIFeatureLevel == ERHIFeatureLevel::ES2)
+
+	if (EncodedHDRDerivedData && (GRHIFeatureLevel == ERHIFeatureLevel::ES2 || bRetainAllFeatureLevelData))
 	{
 		if (FPlatformProperties::RequiresCookedData())
 		{
-			INC_MEMORY_STAT_BY(STAT_ReflectionCaptureMemory,EncodedHDRDerivedData->CapturedData.GetAllocatedSize());
+			INC_MEMORY_STAT_BY(STAT_ReflectionCaptureMemory, EncodedHDRDerivedData->CapturedData.GetAllocatedSize());
 		}
 
 		// Create a cubemap texture out of the encoded HDR data
@@ -1022,9 +1030,12 @@ void UReflectionCaptureComponent::PostLoad()
 		EncodedHDRCubemapTexture->SetupParameters(GReflectionCaptureSize, FMath::CeilLogTwo(GReflectionCaptureSize) + 1, PF_B8G8R8A8, &EncodedHDRDerivedData->CapturedData);
 		BeginInitResource(EncodedHDRCubemapTexture);
 
-		// Don't need the full HDR data for rendering on this feature level
-		delete FullHDRDerivedData;
-		FullHDRDerivedData = NULL;
+		if (!bRetainAllFeatureLevelData)
+		{
+			// Don't need the full HDR data for rendering on this feature level
+			delete FullHDRDerivedData;
+			FullHDRDerivedData = NULL;
+		}
 	}
 }
 
@@ -1284,7 +1295,9 @@ void UReflectionCaptureComponent::UpdateReflectionCaptureContents(UWorld* WorldT
 			}
 		}
 
-		if (WorldToUpdate->Scene->GetFeatureLevel() == ERHIFeatureLevel::SM4)
+		const auto FeatureLevel = WorldToUpdate->Scene->GetFeatureLevel();
+
+		if (FeatureLevel == ERHIFeatureLevel::SM4)
 		{
 			for (int32 CaptureIndex = 0; CaptureIndex < WorldCombinedCaptures.Num(); CaptureIndex++)
 			{
@@ -1303,7 +1316,7 @@ void UReflectionCaptureComponent::UpdateReflectionCaptureContents(UWorld* WorldT
 
 		WorldToUpdate->Scene->AllocateReflectionCaptures(WorldCombinedCaptures);
 
-		if (!FPlatformProperties::RequiresCookedData())
+		if (FeatureLevel >= ERHIFeatureLevel::SM4 && !FPlatformProperties::RequiresCookedData())
 		{
 			for (int32 CaptureIndex = 0; CaptureIndex < WorldCapturesToUpdateForLoad.Num(); CaptureIndex++)
 			{
diff --git a/Engine/Source/Runtime/Engine/Private/RendererSupport.cpp b/Engine/Source/Runtime/Engine/Private/RendererSupport.cpp
--- a/Engine/Source/Runtime/Engine/Private/RendererSupport.cpp
+++ b/Engine/Source/Runtime/Engine/Private/RendererSupport.cpp
@@ -127,11 +127,11 @@ void RestoreReferencesToRendererModuleClasses(
 	{
 		UWorld* World = It.Key();
 
-		World->Scene = RendererModule.AllocateScene(World, World->RequiresHitProxies(), GRHIFeatureLevel);
+		World->Scene = RendererModule.AllocateScene(World, World->RequiresHitProxies(), World->FeatureLevel);
 
 		if (It.Value())
 		{
-			World->FXSystem = FFXSystemInterface::Create(GRHIFeatureLevel);
+			World->FXSystem = FFXSystemInterface::Create(World->FeatureLevel);
 			World->Scene->SetFXSystem( World->FXSystem );
 		}
 
diff --git a/Engine/Source/Runtime/Engine/Private/SkeletalMesh.cpp b/Engine/Source/Runtime/Engine/Private/SkeletalMesh.cpp
--- a/Engine/Source/Runtime/Engine/Private/SkeletalMesh.cpp
+++ b/Engine/Source/Runtime/Engine/Private/SkeletalMesh.cpp
@@ -3860,7 +3860,7 @@ FSkeletalMeshSceneProxy::FSkeletalMeshSceneProxy(const USkinnedMeshComponent* Co
 				MaterialRelevance |= Material->GetRelevance();
 			}
 
-			const bool bRequiresAdjacencyInformation = RequiresAdjacencyInformation( Material, &TGPUSkinVertexFactory<false>::StaticType );
+			const bool bRequiresAdjacencyInformation = RequiresAdjacencyInformation( Material, &TGPUSkinVertexFactory<false>::StaticType, GetScene()->GetFeatureLevel() );
 			if ( bRequiresAdjacencyInformation && LODModel.AdjacencyMultiSizeIndexContainer.IsIndexBufferValid() == false )
 			{
 				UE_LOG(LogSkeletalMesh, Warning, 
@@ -4208,7 +4208,7 @@ void FSkeletalMeshSceneProxy::DrawDynamicElementsSection(FPrimitiveDrawInterface
 
 	BatchElement.UserIndex = MeshObject->GPUSkinCacheKeys[Section.ChunkIndex];
 
-	const bool bRequiresAdjacencyInformation = RequiresAdjacencyInformation( SectionElementInfo.Material, Mesh.VertexFactory->GetType() );
+	const bool bRequiresAdjacencyInformation = RequiresAdjacencyInformation(SectionElementInfo.Material, Mesh.VertexFactory->GetType(), View->GetFeatureLevel());
 	if ( bRequiresAdjacencyInformation )
 	{
 		check( LODModel.AdjacencyMultiSizeIndexContainer.IsIndexBufferValid() );
@@ -4582,7 +4582,7 @@ SIZE_T USkinnedMeshComponent::GetResourceSize(EResourceSizeMode::Type Mode)
 
 FPrimitiveSceneProxy* USkinnedMeshComponent::CreateSceneProxy()
 {
-	ERHIFeatureLevel::Type SceneFeatureLevel = GRHIFeatureLevel;
+	ERHIFeatureLevel::Type SceneFeatureLevel = GetWorld()->FeatureLevel;
 	FSkeletalMeshSceneProxy* Result = NULL;
 	FSkeletalMeshResource* SkelMeshResource = GetSkeletalMeshResource();
 
@@ -4594,7 +4594,7 @@ FPrimitiveSceneProxy* USkinnedMeshComponent::CreateSceneProxy()
 	{
 		// Only create a scene proxy if the bone count being used is supported, or if we don't have a skeleton (this is the case with destructibles)
 		int32 MaxBonesPerChunk = SkelMeshResource->GetMaxBonesPerChunk();
-		if (MaxBonesPerChunk <= GetFeatureLevelMaxNumberOfBones(GRHIFeatureLevel))
+		if (MaxBonesPerChunk <= GetFeatureLevelMaxNumberOfBones(SceneFeatureLevel))
 		{
 			Result = ::new FSkeletalMeshSceneProxy(this,SkelMeshResource);
 		}
diff --git a/Engine/Source/Runtime/Engine/Private/SkinnedMeshComponent.cpp b/Engine/Source/Runtime/Engine/Private/SkinnedMeshComponent.cpp
--- a/Engine/Source/Runtime/Engine/Private/SkinnedMeshComponent.cpp
+++ b/Engine/Source/Runtime/Engine/Private/SkinnedMeshComponent.cpp
@@ -54,8 +54,8 @@ void USkinnedMeshComponent::CreateRenderState_Concurrent()
 		// No need to create the mesh object if we aren't actually rendering anything (see UPrimitiveComponent::Attach)
 		if ( FApp::CanEverRender() && ShouldComponentAddToScene() )
 		{
-			ERHIFeatureLevel::Type SceneFeatureLevel = GRHIFeatureLevel;
-			FSkeletalMeshResource* SkelMeshResource = SkeletalMesh->GetResourceForRendering(SceneFeatureLevel);
+			ERHIFeatureLevel::Type SceneFeatureLevel = GetWorld()->FeatureLevel;
+			FSkeletalMeshResource* SkelMeshResource = SkeletalMesh->GetResourceForRendering();
 
 			// Also check if skeletal mesh has too many bones/chunk for GPU skinning.
 			const bool bIsCPUSkinned = SkelMeshResource->RequiresCPUSkinning(SceneFeatureLevel) || (GIsEditor && ShouldCPUSkin());
@@ -761,8 +761,8 @@ FSkeletalMeshResource* USkinnedMeshComponent::GetSkeletalMeshResource() const
 	}
 	else if (SkeletalMesh)
 	{
-		ERHIFeatureLevel::Type SceneFeatureLevel = GRHIFeatureLevel;
-		return SkeletalMesh->GetResourceForRendering(SceneFeatureLevel);
+		ERHIFeatureLevel::Type SceneFeatureLevel = GetWorld()->FeatureLevel;
+		return SkeletalMesh->GetResourceForRendering();
 	}
 	else
 	{
diff --git a/Engine/Source/Runtime/Engine/Private/SkyLightComponent.cpp b/Engine/Source/Runtime/Engine/Private/SkyLightComponent.cpp
--- a/Engine/Source/Runtime/Engine/Private/SkyLightComponent.cpp
+++ b/Engine/Source/Runtime/Engine/Private/SkyLightComponent.cpp
@@ -18,7 +18,7 @@
 
 void FSkyTextureCubeResource::InitRHI()
 {
-	if (GRHIFeatureLevel >= ERHIFeatureLevel::SM3)
+	if (GetFeatureLevel() >= ERHIFeatureLevel::SM3)
 	{
 		FRHIResourceCreateInfo CreateInfo;
 		TextureCubeRHI = RHICreateTextureCube(Size, Format, NumMips, 0, CreateInfo);
diff --git a/Engine/Source/Runtime/Engine/Private/StaticMesh.cpp b/Engine/Source/Runtime/Engine/Private/StaticMesh.cpp
--- a/Engine/Source/Runtime/Engine/Private/StaticMesh.cpp
+++ b/Engine/Source/Runtime/Engine/Private/StaticMesh.cpp
@@ -607,8 +607,10 @@ void FStaticMeshLODResources::InitVertexFactory(
 
 void FStaticMeshLODResources::InitResources(UStaticMesh* Parent)
 {
+	const auto MaxShaderPlatform = GShaderPlatformForFeatureLevel[GMaxRHIFeatureLevel];
+
 	// Initialize the vertex and index buffers.
-	if (IsES2Platform(GRHIShaderPlatform))
+	if (IsES2Platform(MaxShaderPlatform))
 	{
 		if (IndexBuffer.Is32Bit())
 		{
@@ -636,7 +638,7 @@ void FStaticMeshLODResources::InitResources(UStaticMesh* Parent)
 		BeginInitResource(&DepthOnlyIndexBuffer);
 	}
 
-	if( RHISupportsTessellation(GRHIShaderPlatform) ) 
+	if (RHISupportsTessellation(MaxShaderPlatform))
 	{
 		BeginInitResource(&AdjacencyIndexBuffer);
 	}
@@ -651,7 +653,7 @@ void FStaticMeshLODResources::InitResources(UStaticMesh* Parent)
 		PositionVertexBuffer.GetStride() * PositionVertexBuffer.GetNumVertices();
 	const uint32 StaticMeshIndexMemory = IndexBuffer.GetAllocatedSize()
 		+ WireframeIndexBuffer.GetAllocatedSize()
-		+ (RHISupportsTessellation(GRHIShaderPlatform) ? AdjacencyIndexBuffer.GetAllocatedSize() : 0);
+		+ (AdjacencyIndexBuffer.IsInitialized() ? AdjacencyIndexBuffer.GetAllocatedSize() : 0);
 	const uint32 ResourceVertexColorMemory = ColorVertexBuffer.GetStride() * ColorVertexBuffer.GetNumVertices();
 
 	INC_DWORD_STAT_BY( STAT_StaticMeshVertexMemory, StaticMeshVertexMemory );
@@ -667,7 +669,7 @@ void FStaticMeshLODResources::ReleaseResources()
 		PositionVertexBuffer.GetStride() * PositionVertexBuffer.GetNumVertices();
 	const uint32 StaticMeshIndexMemory = IndexBuffer.GetAllocatedSize()
 		+ WireframeIndexBuffer.GetAllocatedSize()
-		+ (RHISupportsTessellation(GRHIShaderPlatform) ? AdjacencyIndexBuffer.GetAllocatedSize() : 0);
+		+ (AdjacencyIndexBuffer.IsInitialized() ? AdjacencyIndexBuffer.GetAllocatedSize() : 0);
 	const uint32 ResourceVertexColorMemory = ColorVertexBuffer.GetStride() * ColorVertexBuffer.GetNumVertices();
 
 	DEC_DWORD_STAT_BY( STAT_StaticMeshVertexMemory, StaticMeshVertexMemory );
@@ -675,7 +677,7 @@ void FStaticMeshLODResources::ReleaseResources()
 	DEC_DWORD_STAT_BY( STAT_StaticMeshIndexMemory, StaticMeshIndexMemory );
 
 	// Release the vertex and index buffers.
-	if( RHISupportsTessellation(GRHIShaderPlatform) ) 
+	if (AdjacencyIndexBuffer.IsInitialized())
 	{
 		BeginReleaseResource(&AdjacencyIndexBuffer);
 	}
@@ -1323,7 +1325,7 @@ SIZE_T FStaticMeshRenderData::GetResourceSize() const
 			LODRenderData.ColorVertexBuffer.GetStride()				* LODRenderData.ColorVertexBuffer.GetNumVertices();
 		const int32 IBSize = LODRenderData.IndexBuffer.GetAllocatedSize()
 			+ LODRenderData.WireframeIndexBuffer.GetAllocatedSize()
-			+ (RHISupportsTessellation(GRHIShaderPlatform) ? LODRenderData.AdjacencyIndexBuffer.GetAllocatedSize() : 0);
+			+ (RHISupportsTessellation(GShaderPlatformForFeatureLevel[GMaxRHIFeatureLevel]) ? LODRenderData.AdjacencyIndexBuffer.GetAllocatedSize() : 0);
 
 		ResourceSize += VBSize + IBSize;
 		ResourceSize += LODRenderData.Sections.GetAllocatedSize();
diff --git a/Engine/Source/Runtime/Engine/Private/StaticMeshRender.cpp b/Engine/Source/Runtime/Engine/Private/StaticMeshRender.cpp
--- a/Engine/Source/Runtime/Engine/Private/StaticMeshRender.cpp
+++ b/Engine/Source/Runtime/Engine/Private/StaticMeshRender.cpp
@@ -213,7 +213,7 @@ bool FStaticMeshSceneProxy::GetMeshElement(int32 LODIndex,int32 SectionIndex,uin
 	}
 
 	const bool bWireframe = false;
-	const bool bRequiresAdjacencyInformation = RequiresAdjacencyInformation( Material, OutMeshElement.VertexFactory->GetType() );
+	const bool bRequiresAdjacencyInformation = RequiresAdjacencyInformation( Material, OutMeshElement.VertexFactory->GetType(), GetScene()->GetFeatureLevel() );
 
 	SetIndexSource(LODIndex, SectionIndex, OutMeshElement, bWireframe, bRequiresAdjacencyInformation );
 
@@ -466,7 +466,7 @@ void FStaticMeshSceneProxy::DrawStaticElements(FStaticPrimitiveDrawInterface* PD
 					bool bAnySectionCastsShadow = false;
 					for (int32 SectionIndex = 0; bSafeToUseShadowOnlyMesh && SectionIndex < LODModel.Sections.Num(); SectionIndex++)
 					{
-						const FMaterial* Material = ProxyLODInfo.Sections[SectionIndex].Material->GetRenderProxy(false)->GetMaterial(GRHIFeatureLevel);
+						const FMaterial* Material = ProxyLODInfo.Sections[SectionIndex].Material->GetRenderProxy(false)->GetMaterial(GetScene()->GetFeatureLevel());
 						const FStaticMeshSection& Section = LODModel.Sections[SectionIndex];
 						bSafeToUseShadowOnlyMesh =
 							Section.bCastShadow
@@ -946,6 +946,8 @@ FStaticMeshSceneProxy::FLODInfo::FLODInfo(const UStaticMeshComponent* InComponen
 	ShadowMap(NULL),
 	bUsesMeshModifyingMaterials(false)
 {
+	const auto FeatureLevel = InComponent->GetWorld()->FeatureLevel;
+
 	FStaticMeshRenderData* RenderData = InComponent->StaticMesh->RenderData;
 	if(LODIndex < InComponent->LODData.Num())
 	{
@@ -1004,7 +1006,7 @@ FStaticMeshSceneProxy::FLODInfo::FLODInfo(const UStaticMeshComponent* InComponen
 			SectionInfo.Material = UMaterial::GetDefaultMaterial(MD_Surface);
 		}
 
-		const bool bRequiresAdjacencyInformation = RequiresAdjacencyInformation( SectionInfo.Material, LODModel.VertexFactory.GetType() );
+		const bool bRequiresAdjacencyInformation = RequiresAdjacencyInformation(SectionInfo.Material, LODModel.VertexFactory.GetType(), FeatureLevel);
 		if ( bRequiresAdjacencyInformation && !LODModel.bHasAdjacencyInfo )
 		{
 			UE_LOG(LogStaticMesh, Warning, TEXT("Adjacency information not built for static mesh with a material that requires it. Using default material instead.\n\tMaterial: %s\n\tStaticMesh: %s"),
@@ -1026,7 +1028,7 @@ FStaticMeshSceneProxy::FLODInfo::FLODInfo(const UStaticMeshComponent* InComponen
 
 		// Flag the entire LOD if any material modifies its mesh
 		UMaterialInterface::TMicRecursionGuard RecursionGuard;
-		FMaterialResource const* MaterialResource = const_cast<UMaterialInterface const*>(SectionInfo.Material)->GetMaterial_Concurrent(RecursionGuard)->GetMaterialResource(GRHIFeatureLevel);
+		FMaterialResource const* MaterialResource = const_cast<UMaterialInterface const*>(SectionInfo.Material)->GetMaterial_Concurrent(RecursionGuard)->GetMaterialResource(FeatureLevel);
 		if(MaterialResource)
 		{
 			if(MaterialResource->MaterialModifiesMeshPosition())
diff --git a/Engine/Source/Runtime/Engine/Private/SystemSettings.cpp b/Engine/Source/Runtime/Engine/Private/SystemSettings.cpp
--- a/Engine/Source/Runtime/Engine/Private/SystemSettings.cpp
+++ b/Engine/Source/Runtime/Engine/Private/SystemSettings.cpp
@@ -353,14 +353,6 @@ void FSystemSettings::Initialize( bool bSetupForEditor )
 void FSystemSettings::CVarSink()
 {
 	ApplyOverrides();
-
-	auto CVarAllowHighQualityLightMaps = IConsoleManager::Get().FindConsoleVariable(TEXT("r.HighQualityLightMaps"));
-	if (bInitialUseHighQualityLightmaps != (CVarAllowHighQualityLightMaps->GetInt() != 0))
-	{
-		// This setting can't be changed at runtime, so put it back to what it was originally
-		CVarAllowHighQualityLightMaps->Set(bInitialUseHighQualityLightmaps);
-		UE_LOG(LogSystemSettings, Warning, TEXT("r.HighQualityLightMaps cannot be changed at runtime. Ignoring setting change"));
-	}
 }
 
 bool FSystemSettings::HandleLowendCommand( const TCHAR* Cmd , FOutputDevice& Ar )
@@ -490,14 +482,6 @@ void FSystemSettings::SetTextureLODGroup(TextureGroup TextureGroupID, int32 MinL
 void FSystemSettings::LoadFromIni()
 {
 	FSystemSettingsData::LoadFromIni(GetSectionName(bIsEditor));
-
-	if (FParse::Param(FCommandLine::Get(), TEXT("FeatureLevelES2")) || !FPlatformProperties::SupportsHighQualityLightmaps())
-	{
-		// Hack: disable high quality lightmaps with feature level ES2
-		//@todo-mobile - replace this with a system where feature levels can clamp misc features
-		static auto CVarAllowHighQualityLightMaps = IConsoleManager::Get().FindConsoleVariable(TEXT("r.HighQualityLightMaps"));
-		CVarAllowHighQualityLightMaps->Set(0);
-	}
 }
 
 void FSystemSettings::SaveToIni()
diff --git a/Engine/Source/Runtime/Engine/Private/TessellationRendering.cpp b/Engine/Source/Runtime/Engine/Private/TessellationRendering.cpp
--- a/Engine/Source/Runtime/Engine/Private/TessellationRendering.cpp
+++ b/Engine/Source/Runtime/Engine/Private/TessellationRendering.cpp
@@ -4,7 +4,7 @@
 #include "EnginePrivate.h"
 
 /** Returns true if the Material and Vertex Factory combination require adjacency information. */
-bool RequiresAdjacencyInformation( UMaterialInterface* Material, const FVertexFactoryType* VertexFactoryType )
+bool RequiresAdjacencyInformation( UMaterialInterface* Material, const FVertexFactoryType* VertexFactoryType, ERHIFeatureLevel::Type InFeatureLevel )
 {
 	EMaterialTessellationMode TessellationMode = MTM_NoTessellation;
 	bool bEnableCrackFreeDisplacement = false;
@@ -14,7 +14,7 @@ bool RequiresAdjacencyInformation( UMaterialInterface* Material, const FVertexFa
 		{
 			FMaterialRenderProxy* MaterialRenderProxy = Material->GetRenderProxy( false, false );
 			check( MaterialRenderProxy );
-			const FMaterial* MaterialResource = MaterialRenderProxy->GetMaterial(GRHIFeatureLevel);
+			const FMaterial* MaterialResource = MaterialRenderProxy->GetMaterial(InFeatureLevel);
 			check( MaterialResource );
 			TessellationMode = MaterialResource->GetTessellationMode();
 			bEnableCrackFreeDisplacement = MaterialResource->IsCrackFreeDisplacementEnabled();
diff --git a/Engine/Source/Runtime/Engine/Private/TextureRenderTarget2D.cpp b/Engine/Source/Runtime/Engine/Private/TextureRenderTarget2D.cpp
--- a/Engine/Source/Runtime/Engine/Private/TextureRenderTarget2D.cpp
+++ b/Engine/Source/Runtime/Engine/Private/TextureRenderTarget2D.cpp
@@ -21,7 +21,9 @@ UTextureRenderTarget2D::UTextureRenderTarget2D(const class FPostConstructInitial
 
 FTextureResource* UTextureRenderTarget2D::CreateResource()
 {
-	if (GRHIFeatureLevel <= ERHIFeatureLevel::ES2)
+	UWorld* World = GetWorld();
+	ERHIFeatureLevel::Type FeatureLevel = World != nullptr ? World->FeatureLevel : GMaxRHIFeatureLevel;
+	if (FeatureLevel <= ERHIFeatureLevel::ES2)
 	{
 		EPixelFormat Format = GetFormat();
 		if ((!GSupportsRenderTargetFormat_PF_FloatRGBA && (Format == PF_FloatRGBA || Format == PF_FloatRGB))
diff --git a/Engine/Source/Runtime/Engine/Private/UnrealEngine.cpp b/Engine/Source/Runtime/Engine/Private/UnrealEngine.cpp
--- a/Engine/Source/Runtime/Engine/Private/UnrealEngine.cpp
+++ b/Engine/Source/Runtime/Engine/Private/UnrealEngine.cpp
@@ -9709,10 +9709,17 @@ static void SetNearClipPlane(const TArray<FString>& Args)
 }
 FAutoConsoleCommand GSetNearClipPlaneCmd(TEXT("r.SetNearClipPlane"),TEXT("Set the near clipping plane (in cm)"),FConsoleCommandWithArgsDelegate::CreateStatic(SetNearClipPlane));
 
+static TAutoConsoleVariable<int32> CVarAllowHighQualityLightMaps(
+	TEXT("r.HighQualityLightMaps"),
+	1,
+	TEXT("If set to 1, allow high quality lightmaps which don't bake in direct lighting of stationary lights"),
+	ECVF_RenderThreadSafe | ECVF_ReadOnly);
+
 bool AllowHighQualityLightmaps()
 {
-	static const auto CVar = IConsoleManager::Get().FindTConsoleVariableDataInt(TEXT("r.HighQualityLightMaps"));
-	return CVar->GetValueOnAnyThread() != 0;
+	return FPlatformProperties::SupportsHighQualityLightmaps()
+		&& (GRHIFeatureLevel > ERHIFeatureLevel::ES2)
+		&& (CVarAllowHighQualityLightMaps.GetValueOnAnyThread() != 0);
 }
 
 // Helper function for changing system resolution via the r.setres console command
diff --git a/Engine/Source/Runtime/Engine/Private/VectorField.cpp b/Engine/Source/Runtime/Engine/Private/VectorField.cpp
--- a/Engine/Source/Runtime/Engine/Private/VectorField.cpp
+++ b/Engine/Source/Runtime/Engine/Private/VectorField.cpp
@@ -182,7 +182,7 @@ public:
 	 */
 	virtual void InitRHI() override
 	{
-		if (VolumeData && GRHIFeatureLevel >= ERHIFeatureLevel::SM4)
+		if (VolumeData && GetFeatureLevel() >= ERHIFeatureLevel::SM4)
 		{
 			const uint32 DataSize = SizeX * SizeY * SizeZ * sizeof(FFloat16Color);
 			FVectorFieldStaticResourceBulkDataInterface BulkDataInterface(VolumeData, DataSize);
@@ -714,7 +714,7 @@ public:
 	 */
 	virtual void InitRHI() override
 	{
-		if (GRHIFeatureLevel >= ERHIFeatureLevel::SM4)
+		if (GetFeatureLevel() >= ERHIFeatureLevel::SM4)
 		{
 			check(SizeX > 0);
 			check(SizeY > 0);
@@ -722,7 +722,7 @@ public:
 			UE_LOG(LogVectorField,Verbose,TEXT("InitRHI for 0x%016x %dx%dx%d"),(PTRINT)this,SizeX,SizeY,SizeZ);
 
 			uint32 TexCreateFlags = 0;
-			if (GRHIFeatureLevel == ERHIFeatureLevel::SM5)
+			if (GetFeatureLevel() >= ERHIFeatureLevel::SM5)
 			{
 				TexCreateFlags = TexCreate_ShaderResource | TexCreate_UAV;
 			}
@@ -735,7 +735,7 @@ public:
 				TexCreateFlags,
 				CreateInfo);
 
-			if (GRHIFeatureLevel == ERHIFeatureLevel::SM5)
+			if (GetFeatureLevel() >= ERHIFeatureLevel::SM5)
 			{
 				VolumeTextureUAV = RHICreateUnorderedAccessView(VolumeTextureRHI);
 			}
diff --git a/Engine/Source/Runtime/Engine/Private/World.cpp b/Engine/Source/Runtime/Engine/Private/World.cpp
--- a/Engine/Source/Runtime/Engine/Private/World.cpp
+++ b/Engine/Source/Runtime/Engine/Private/World.cpp
@@ -69,6 +69,7 @@ FWorldDelegates::FOnLevelChanged FWorldDelegates::LevelRemovedFromWorld;
 
 UWorld::UWorld( const class FPostConstructInitializeProperties& PCIP )
 :	UObject(PCIP)
+,	FeatureLevel(GRHIFeatureLevel)
 ,	TickTaskLevel(FTickTaskManagerInterface::Get().AllocateTickTaskLevel())
 ,   bIsBuilt(false)
 ,	NextTravelType(TRAVEL_Relative)
@@ -82,6 +83,7 @@ UWorld::UWorld( const class FPostConstructInitializeProperties& PCIP )
 
 UWorld::UWorld( const class FPostConstructInitializeProperties& PCIP,const FURL& InURL )
 :	UObject(PCIP)
+,	FeatureLevel(GRHIFeatureLevel)
 ,	URL(InURL)
 ,	TickTaskLevel(FTickTaskManagerInterface::Get().AllocateTickTaskLevel())
 ,   bIsBuilt(false)
@@ -758,10 +760,10 @@ void UWorld::InitWorld(const InitializationValues IVS)
 		
 		// Save off the value used to create the scene, so this UWorld can recreate its scene later
 		bRequiresHitProxies = IVS.bRequiresHitProxies;
-		Scene = GetRendererModule().AllocateScene( this, bRequiresHitProxies, GRHIFeatureLevel );
+		Scene = GetRendererModule().AllocateScene( this, bRequiresHitProxies, FeatureLevel );
 		if ( !IsRunningDedicatedServer() && !IsRunningCommandlet() )
 		{
-			FXSystem = FFXSystemInterface::Create(GRHIFeatureLevel);
+			FXSystem = FFXSystemInterface::Create(FeatureLevel);
 			Scene->SetFXSystem( FXSystem );
 		}
 		else
@@ -5014,23 +5016,70 @@ void UWorld::GetLightMapsAndShadowMaps(ULevel* Level, TArray<UTexture2D*>& OutLi
 
 void UWorld::ChangeFeatureLevel(ERHIFeatureLevel::Type InFeatureLevel)
 {
-	if (InFeatureLevel != Scene->GetFeatureLevel())
+	if (InFeatureLevel != FeatureLevel)
 	{
-		Scene->ChangeFeatureLevel(InFeatureLevel);
-		
-		FFXSystemInterface::Destroy(FXSystem);
-		FXSystem = FFXSystemInterface::Create(InFeatureLevel);
-		Scene->SetFXSystem(FXSystem);
+		Scene->Release();
+		GetRendererModule().RemoveScene(Scene);
+
+		FeatureLevel = InFeatureLevel;
+
+		Scene = GetRendererModule().AllocateScene(this, bRequiresHitProxies, InFeatureLevel);
+
+ 		FFXSystemInterface::Destroy(FXSystem);
+ 		FXSystem = FFXSystemInterface::Create(InFeatureLevel);
+ 		Scene->SetFXSystem(FXSystem);
 	}
 }
 
-void UWorld::ForceFeatureLevelUpdate(ERHIFeatureLevel::Type InFeatureLevel)
+void UWorld::ChangeAllWorldFeatureLevels(ERHIFeatureLevel::Type InFeatureLevel)
 {
-	FGlobalComponentReregisterContext RecreateComponents;
-	UMaterial::AllMaterialsCacheResourceShadersForRendering();
-	UMaterialInstance::AllMaterialsCacheResourceShadersForRendering();
-	GetGlobalShaderMap(GShaderPlatformForFeatureLevel[InFeatureLevel], true);
-	RecompileGlobalShaders();
+	bool bSomethingChanged = false;
+
+	// try iterating over all worlds here
+	for (TObjectIterator<UWorld> It; It; ++It)
+	{
+		UWorld* World = *It;
+
+		if (World->FeatureLevel != InFeatureLevel)
+		{
+			bSomethingChanged = true;
+		}
+	}
+
+	if (bSomethingChanged)
+	{
+		FlushRenderingCommands();
+
+		FGlobalComponentReregisterContext RecreateComponents;
+
+		GCurrentRHIFeatureLevel = InFeatureLevel;
+		GRHIShaderPlatform = GShaderPlatformForFeatureLevel[InFeatureLevel];
+
+		UMaterial::AllMaterialsCacheResourceShadersForRendering();
+		UMaterialInstance::AllMaterialsCacheResourceShadersForRendering();
+		GetGlobalShaderMap(GRHIShaderPlatform, false);
+		GShaderCompilingManager->ProcessAsyncResults(false, true);
+		
+		//invalidate global bound shader states so they will be created with the new shaders the next time they are set (in SetGlobalBoundShaderState)
+		for (TLinkedList<FGlobalBoundShaderStateResource*>::TIterator It(FGlobalBoundShaderStateResource::GetGlobalBoundShaderStateList()); It; It.Next())
+		{
+			BeginUpdateResourceRHI(*It);
+		}
+
+		// try iterating over all worlds here
+		for (TObjectIterator<UWorld> It; It; ++It)
+		{
+			UWorld* World = *It;
+
+			if (World->FeatureLevel != InFeatureLevel)
+			{
+				World->ChangeFeatureLevel(InFeatureLevel);
+			}
+		}
+
+		FOutputDeviceNull Ar;
+		RecompileShaders(TEXT("CHANGED"), Ar);
+	}
 }
 
 /**
diff --git a/Engine/Source/Runtime/Engine/Public/FXSystem.h b/Engine/Source/Runtime/Engine/Public/FXSystem.h
--- a/Engine/Source/Runtime/Engine/Public/FXSystem.h
+++ b/Engine/Source/Runtime/Engine/Public/FXSystem.h
@@ -83,14 +83,6 @@ inline bool RHISupportsGPUParticles(ERHIFeatureLevel::Type InFeatureLevel)
 	return InFeatureLevel >= ERHIFeatureLevel::SM4;
 }
 
-/**
- * Returns true if the current RHI supports GPU particles.
- */
-inline bool CurrentRHISupportsGPUParticles()
-{
-	return RHISupportsGPUParticles(GRHIFeatureLevel);
-}
-
 /*-----------------------------------------------------------------------------
 	The interface to the FX system runtime.
 -----------------------------------------------------------------------------*/
diff --git a/Engine/Source/Runtime/Engine/Public/Landscape/LandscapeRender.h b/Engine/Source/Runtime/Engine/Public/Landscape/LandscapeRender.h
--- a/Engine/Source/Runtime/Engine/Public/Landscape/LandscapeRender.h
+++ b/Engine/Source/Runtime/Engine/Public/Landscape/LandscapeRender.h
@@ -294,10 +294,10 @@ public:
 	FLandscapeIndexRanges* IndexRanges;
 	FLandscapeSharedAdjacencyIndexBuffer* AdjacencyIndexBuffers;
 
-	FLandscapeSharedBuffers(int32 SharedBuffersKey, int32 SubsectionSizeQuads, int32 NumSubsections);
+	FLandscapeSharedBuffers(int32 SharedBuffersKey, int32 SubsectionSizeQuads, int32 NumSubsections, ERHIFeatureLevel::Type InFeatureLevel);
 
 	template <typename INDEX_TYPE>
-	void CreateIndexBuffers();
+	void CreateIndexBuffers(ERHIFeatureLevel::Type InFeatureLevel);
 
 	virtual ~FLandscapeSharedBuffers();
 };
diff --git a/Engine/Source/Runtime/Engine/Public/MeshParticleVertexFactory.h b/Engine/Source/Runtime/Engine/Public/MeshParticleVertexFactory.h
--- a/Engine/Source/Runtime/Engine/Public/MeshParticleVertexFactory.h
+++ b/Engine/Source/Runtime/Engine/Public/MeshParticleVertexFactory.h
@@ -72,8 +72,8 @@ public:
 	};
 
 	/** Default constructor. */
-	FMeshParticleVertexFactory( EParticleVertexFactoryType InType = PVFT_Mesh )
-		: FParticleVertexFactoryBase(InType)
+	FMeshParticleVertexFactory(EParticleVertexFactoryType InType, ERHIFeatureLevel::Type InFeatureLevel)
+		: FParticleVertexFactoryBase(InType, InFeatureLevel)
 	{
 	}
 
diff --git a/Engine/Source/Runtime/Engine/Public/ParticleBeamTrailVertexFactory.h b/Engine/Source/Runtime/Engine/Public/ParticleBeamTrailVertexFactory.h
--- a/Engine/Source/Runtime/Engine/Public/ParticleBeamTrailVertexFactory.h
+++ b/Engine/Source/Runtime/Engine/Public/ParticleBeamTrailVertexFactory.h
@@ -29,8 +29,8 @@ class FParticleBeamTrailVertexFactory : public FParticleVertexFactoryBase
 public:
 
 	/** Default constructor. */
-	FParticleBeamTrailVertexFactory( EParticleVertexFactoryType InType = PVFT_BeamTrail )
-		: FParticleVertexFactoryBase(InType)
+	FParticleBeamTrailVertexFactory( EParticleVertexFactoryType InType, ERHIFeatureLevel::Type InFeatureLevel )
+		: FParticleVertexFactoryBase(InType, InFeatureLevel)
 	{
 	}
 
diff --git a/Engine/Source/Runtime/Engine/Public/ParticleHelper.h b/Engine/Source/Runtime/Engine/Public/ParticleHelper.h
--- a/Engine/Source/Runtime/Engine/Public/ParticleHelper.h
+++ b/Engine/Source/Runtime/Engine/Public/ParticleHelper.h
@@ -1134,7 +1134,7 @@ public:
 		ClearPoolInternal();
 	}
 
-	FParticleVertexFactoryBase* GetParticleVertexFactory(EParticleVertexFactoryType InType);
+	FParticleVertexFactoryBase* GetParticleVertexFactory(EParticleVertexFactoryType InType, ERHIFeatureLevel::Type InFeatureLevel);
 
 	bool ReturnParticleVertexFactory(FParticleVertexFactoryBase* InVertexFactory);
 
@@ -1169,9 +1169,9 @@ protected:
 	 *
 	 *	@return	FParticleVertexFactoryBase*	The created VF; NULL if invalid InType
 	 */
-	FParticleVertexFactoryBase* CreateParticleVertexFactory(EParticleVertexFactoryType InType);
+	FParticleVertexFactoryBase* CreateParticleVertexFactory(EParticleVertexFactoryType InType, ERHIFeatureLevel::Type InFeatureLevel);
 
-	TArray<FParticleVertexFactoryBase*>	VertexFactoriesAvailable[PVFT_MAX];
+	TArray<FParticleVertexFactoryBase*>	VertexFactoriesAvailable[PVFT_MAX][ERHIFeatureLevel::Num];
 	TArray<FParticleVertexFactoryBase*>	VertexFactories;
 };
 
diff --git a/Engine/Source/Runtime/Engine/Public/ParticleVertexFactory.h b/Engine/Source/Runtime/Engine/Public/ParticleVertexFactory.h
--- a/Engine/Source/Runtime/Engine/Public/ParticleVertexFactory.h
+++ b/Engine/Source/Runtime/Engine/Public/ParticleVertexFactory.h
@@ -28,8 +28,9 @@ class FParticleVertexFactoryBase : public FVertexFactory
 public:
 
 	/** Default constructor. */
-	explicit FParticleVertexFactoryBase( EParticleVertexFactoryType Type )
-		: ParticleFactoryType(Type)
+	explicit FParticleVertexFactoryBase( EParticleVertexFactoryType Type, ERHIFeatureLevel::Type InFeatureLevel )
+		: FVertexFactory(InFeatureLevel)
+		, ParticleFactoryType(Type)
 		, bInUse(false)
 	{
 	}
@@ -96,8 +97,8 @@ class FParticleSpriteVertexFactory : public FParticleVertexFactoryBase
 public:
 
 	/** Default constructor. */
-	FParticleSpriteVertexFactory( EParticleVertexFactoryType InType = PVFT_Sprite )
-		: FParticleVertexFactoryBase(InType)
+	FParticleSpriteVertexFactory( EParticleVertexFactoryType InType, ERHIFeatureLevel::Type InFeatureLevel )
+		: FParticleVertexFactoryBase(InType, InFeatureLevel)
 	{
 	}
 
diff --git a/Engine/Source/Runtime/Engine/Public/SceneInterface.h b/Engine/Source/Runtime/Engine/Public/SceneInterface.h
--- a/Engine/Source/Runtime/Engine/Public/SceneInterface.h
+++ b/Engine/Source/Runtime/Engine/Public/SceneInterface.h
@@ -299,7 +299,6 @@ public:
 	virtual bool IsEditorScene() const { return false; }
 
 	virtual ERHIFeatureLevel::Type GetFeatureLevel() const { return GRHIFeatureLevel; }
-	virtual void ChangeFeatureLevel(ERHIFeatureLevel::Type InFeatureLevel) {}
 
 protected:
 	virtual ~FSceneInterface() {}
diff --git a/Engine/Source/Runtime/Engine/Public/SceneManagement.h b/Engine/Source/Runtime/Engine/Public/SceneManagement.h
--- a/Engine/Source/Runtime/Engine/Public/SceneManagement.h
+++ b/Engine/Source/Runtime/Engine/Public/SceneManagement.h
@@ -1337,6 +1337,7 @@ BEGIN_UNIFORM_BUFFER_STRUCT(FViewUniformShaderParameters,ENGINE_API)
 	DECLARE_UNIFORM_BUFFER_STRUCT_MEMBER(FVector4,SceneTextureMinMax)
 	DECLARE_UNIFORM_BUFFER_STRUCT_MEMBER(FLinearColor,SkyLightColor)
 	DECLARE_UNIFORM_BUFFER_STRUCT_MEMBER_ARRAY(FVector4,SkyIrradianceEnvironmentMap,[7])
+	DECLARE_UNIFORM_BUFFER_STRUCT_MEMBER(float, ES2PreviewMode)
 END_UNIFORM_BUFFER_STRUCT(FViewUniformShaderParameters)
 
 
@@ -1483,7 +1484,7 @@ extern ENGINE_API int32 DrawRichMesh(
 extern ENGINE_API void DrawUVs(FViewport* InViewport, FCanvas* InCanvas, int32 InTextYPos, const int32 LODLevel, int32 UVChannel, TArray<FVector2D> SelectedEdgeTexCoords, class FStaticMeshRenderData* StaticMeshRenderData, class FStaticLODModel* SkeletalMeshRenderData );
 
 /** Returns true if the Material and Vertex Factory combination require adjacency information. */
-bool RequiresAdjacencyInformation( class UMaterialInterface* Material, const class FVertexFactoryType* VertexFactoryType );
+bool RequiresAdjacencyInformation( class UMaterialInterface* Material, const class FVertexFactoryType* VertexFactoryType, ERHIFeatureLevel::Type InFeatureLevel );
 /**
  * Computes the screen size of a given sphere bounds in the given view
  * @param Origin - Origin of the bounds in world space
diff --git a/Engine/Source/Runtime/Engine/Public/TessellationRendering.h b/Engine/Source/Runtime/Engine/Public/TessellationRendering.h
--- a/Engine/Source/Runtime/Engine/Public/TessellationRendering.h
+++ b/Engine/Source/Runtime/Engine/Public/TessellationRendering.h
@@ -8,6 +8,6 @@
 #define __TESSELLATIONRENDERING_H__
 
 /** Returns true if the Material and Vertex Factory combination require adjacency information. */
-extern bool RequiresAdjacencyInformation( UMaterialInterface* Material, const FVertexFactoryType* VertexFactoryType );
+extern bool RequiresAdjacencyInformation(UMaterialInterface* Material, const FVertexFactoryType* VertexFactoryType, ERHIFeatureLevel::Type InFeatureLevel);
 
 #endif
diff --git a/Engine/Source/Runtime/OpenGLDrv/Private/OpenGLDevice.cpp b/Engine/Source/Runtime/OpenGLDrv/Private/OpenGLDevice.cpp
--- a/Engine/Source/Runtime/OpenGLDrv/Private/OpenGLDevice.cpp
+++ b/Engine/Source/Runtime/OpenGLDrv/Private/OpenGLDevice.cpp
@@ -607,7 +607,7 @@ static void InitRHICapabilitiesForGL()
 	const GLint MinorVersion = FOpenGL::GetMinorVersion();
 
 	// Shader platform & RHI feature level
-	SetMaxRHIFeatureLevel(FOpenGL::GetFeatureLevel());
+	GMaxRHIFeatureLevel = GCurrentRHIFeatureLevel = FOpenGL::GetFeatureLevel();
 	GRHIShaderPlatform = FOpenGL::GetShaderPlatform();
 
 	FString FeatureLevelName;
diff --git a/Engine/Source/Runtime/RHI/Private/RHI.cpp b/Engine/Source/Runtime/RHI/Private/RHI.cpp
--- a/Engine/Source/Runtime/RHI/Private/RHI.cpp
+++ b/Engine/Source/Runtime/RHI/Private/RHI.cpp
@@ -177,14 +177,12 @@ RHI_API EShaderPlatform GRHIShaderPlatform = SP_PCD3D_SM5;
 /** The maximum feature level supported on this machine */
 ERHIFeatureLevel::Type GMaxRHIFeatureLevel = ERHIFeatureLevel::SM5;
 
-RHI_API ERHIFeatureLevel::Type GetRHIFeatureLevel()
-{
-	return GMaxRHIFeatureLevel;
-}
+/** The current feature level being used on this machine - This is only temporary and will go away when feature level becomes dynamic */
+ERHIFeatureLevel::Type GCurrentRHIFeatureLevel = ERHIFeatureLevel::SM5;
 
-RHI_API void SetMaxRHIFeatureLevel(ERHIFeatureLevel::Type InType)
+RHI_API ERHIFeatureLevel::Type GetCurrentRHIFeatureLevel()
 {
-	GMaxRHIFeatureLevel = InType;
+	return GCurrentRHIFeatureLevel;
 }
 
 FName FeatureLevelNames[] = 
diff --git a/Engine/Source/Runtime/RHI/Public/RHI.h b/Engine/Source/Runtime/RHI/Public/RHI.h
--- a/Engine/Source/Runtime/RHI/Public/RHI.h
+++ b/Engine/Source/Runtime/RHI/Public/RHI.h
@@ -232,15 +232,19 @@ extern RHI_API void GetFeatureLevelName(ERHIFeatureLevel::Type InFeatureLevel, F
 extern RHI_API void GetFeatureLevelName(ERHIFeatureLevel::Type InFeatureLevel, FName& OutName);
 
 extern RHI_API ERHIFeatureLevel::Type GMaxRHIFeatureLevel;
+extern RHI_API ERHIFeatureLevel::Type GCurrentRHIFeatureLevel;
 
-/** The current feature level supported by the RHI. */
-extern RHI_API ERHIFeatureLevel::Type GetRHIFeatureLevel();
+/** Function for retrieving the current feature level. Only exists to provide a convient way of tracking access to this global during mobile preview work */
+extern RHI_API ERHIFeatureLevel::Type GetCurrentRHIFeatureLevel();
 
-//#define GRHIFeatureLevel GetRHIFeatureLevel()
-#define GRHIFeatureLevel GMaxRHIFeatureLevel
+/** treating GRHIFeatureLevel as a function allows for better usage tracking at small cost to performance. */
+#define RHI_FEATURE_LEVEL_AS_FUNCTION 0
 
-/** Set the current feature level supported by the RHI. */
-extern RHI_API void SetMaxRHIFeatureLevel(ERHIFeatureLevel::Type InType);
+#if RHI_FEATURE_LEVEL_AS_FUNCTION
+	#define GRHIFeatureLevel GetCurrentRHIFeatureLevel()
+#else
+	#define GRHIFeatureLevel GCurrentRHIFeatureLevel
+#endif
 
 /** Table for finding out which shader platform corresponds to a given feature level for this RHI. */
 extern RHI_API EShaderPlatform GShaderPlatformForFeatureLevel[ERHIFeatureLevel::Num];
diff --git a/Engine/Source/Runtime/RenderCore/Private/RenderUtils.cpp b/Engine/Source/Runtime/RenderCore/Private/RenderUtils.cpp
--- a/Engine/Source/Runtime/RenderCore/Private/RenderUtils.cpp
+++ b/Engine/Source/Runtime/RenderCore/Private/RenderUtils.cpp
@@ -237,7 +237,7 @@ public:
 	 */
 	virtual void InitRHI() override
 	{
-		if (IsFeatureLevelSupported(GRHIShaderPlatform, ERHIFeatureLevel::SM4))
+		if (GetFeatureLevel() >= ERHIFeatureLevel::SM4)
 		{
 			// Create the texture.
 			FBlackVolumeTextureResourceBulkDataInterface BlackTextureBulkData;
@@ -277,7 +277,7 @@ public:
 	// FResource interface.
 	virtual void InitRHI() override
 	{
-		if (IsFeatureLevelSupported(GRHIShaderPlatform, ERHIFeatureLevel::SM4))
+		if (GetFeatureLevel() >= ERHIFeatureLevel::SM4)
 		{
 			// Create the texture RHI.  		
 			FRHIResourceCreateInfo CreateInfo;
@@ -469,7 +469,7 @@ public:
 	// FResource interface.
 	virtual void InitRHI() override
 	{
-		if (IsFeatureLevelSupported(GRHIShaderPlatform, ERHIFeatureLevel::SM5))
+		if (GetFeatureLevel() >= ERHIFeatureLevel::SM5)
 		{
 			// Create the texture RHI.
 			FRHIResourceCreateInfo CreateInfo;
diff --git a/Engine/Source/Runtime/RenderCore/Public/RenderResource.h b/Engine/Source/Runtime/RenderCore/Public/RenderResource.h
--- a/Engine/Source/Runtime/RenderCore/Public/RenderResource.h
+++ b/Engine/Source/Runtime/RenderCore/Public/RenderResource.h
@@ -22,7 +22,14 @@ public:
 
 	/** Default constructor. */
 	FRenderResource()
-	: bInitialized(false)
+		: FeatureLevel(ERHIFeatureLevel::Num)
+		, bInitialized(false)
+	{}
+
+	/** Constructor when we know what feature level this resource should support */
+	FRenderResource(ERHIFeatureLevel::Type InFeatureLevel)
+		: FeatureLevel(InFeatureLevel)
+		, bInitialized(false)
 	{}
 
 	/** Destructor used to catch unreleased resources. */
@@ -82,8 +89,12 @@ public:
 	// Accessors.
 	FORCEINLINE bool IsInitialized() const { return bInitialized; }
 
+	FORCEINLINE ERHIFeatureLevel::Type GetFeatureLevel() const { return FeatureLevel == ERHIFeatureLevel::Num ? GMaxRHIFeatureLevel : FeatureLevel; }
+
 private:
 
+	ERHIFeatureLevel::Type FeatureLevel;
+
 	/** This resource's link in the global resource list. */
 	TLinkedList<FRenderResource*> ResourceLink;
 
diff --git a/Engine/Source/Runtime/Renderer/Private/IndirectLightingCache.cpp b/Engine/Source/Runtime/Renderer/Private/IndirectLightingCache.cpp
--- a/Engine/Source/Runtime/Renderer/Private/IndirectLightingCache.cpp
+++ b/Engine/Source/Runtime/Renderer/Private/IndirectLightingCache.cpp
@@ -144,16 +144,17 @@ bool CanIndirectLightingCacheUseVolumeTexture(ERHIFeatureLevel::Type InFeatureLe
 	return InFeatureLevel >= ERHIFeatureLevel::SM3 && GSupportsVolumeTextureRendering;
 }
 
-FIndirectLightingCache::FIndirectLightingCache()
-	:	bUpdateAllCacheEntries(true)
-	,	BlockAllocator(0, 0, 0, GLightingCacheDimension, GLightingCacheDimension, GLightingCacheDimension, false, false)
+FIndirectLightingCache::FIndirectLightingCache(ERHIFeatureLevel::Type InFeatureLevel)
+	: FRenderResource(InFeatureLevel)
+	, bUpdateAllCacheEntries(true)
+	, BlockAllocator(0, 0, 0, GLightingCacheDimension, GLightingCacheDimension, GLightingCacheDimension, false, false)
 {
 	CacheSize = GLightingCacheDimension;
 }
 
 void FIndirectLightingCache::InitDynamicRHI()
 {
-	if (CanIndirectLightingCacheUseVolumeTexture(GRHIFeatureLevel))
+	if (CanIndirectLightingCacheUseVolumeTexture(GetFeatureLevel()))
 	{
 		uint32 Flags = TexCreate_ShaderResource | TexCreate_NoTiling;
 
@@ -344,7 +345,7 @@ FIndirectLightingCacheAllocation* FIndirectLightingCache::FindPrimitiveAllocatio
 
 void FIndirectLightingCache::UpdateCache(FScene* Scene, FSceneRenderer& Renderer, bool bAllowUnbuiltPreview)
 {
-	if (IsIndirectLightingCacheAllowed(Scene->GetFeatureLevel()))
+	if (IsIndirectLightingCacheAllowed(GetFeatureLevel()))
 	{
 		bool bAnyViewAllowsIndirectLightingCache = false;
 
@@ -608,7 +609,7 @@ void FIndirectLightingCache::UpdateBlock(FScene* Scene, FViewInfo* DebugDrawingV
 	float DirectionalShadowing = 1;
 	FVector SkyBentNormal(0, 0, 1);
 
-	if (CanIndirectLightingCacheUseVolumeTexture(Scene->GetFeatureLevel()) && BlockInfo.Allocation->bOpaqueRelevance)
+	if (CanIndirectLightingCacheUseVolumeTexture(GetFeatureLevel()) && BlockInfo.Allocation->bOpaqueRelevance)
 	{
 		static TArray<float> AccumulatedWeight;
 		AccumulatedWeight.Reset(NumSamplesPerBlock);
diff --git a/Engine/Source/Runtime/Renderer/Private/PostProcess/PostProcessCompositeEditorPrimitives.cpp b/Engine/Source/Runtime/Renderer/Private/PostProcess/PostProcessCompositeEditorPrimitives.cpp
--- a/Engine/Source/Runtime/Renderer/Private/PostProcess/PostProcessCompositeEditorPrimitives.cpp
+++ b/Engine/Source/Runtime/Renderer/Private/PostProcess/PostProcessCompositeEditorPrimitives.cpp
@@ -38,7 +38,7 @@ class FPostProcessCompositeEditorPrimitivesPS : public FGlobalShader
 			}
 		}
 
-		return IsFeatureLevelSupported(Platform, ERHIFeatureLevel::SM3);
+		return IsPCPlatform(Platform);
 	}
 
 	static void ModifyCompilationEnvironment(EShaderPlatform Platform, FShaderCompilerEnvironment& OutEnvironment)
diff --git a/Engine/Source/Runtime/Renderer/Private/PostProcess/PostProcessing.cpp b/Engine/Source/Runtime/Renderer/Private/PostProcess/PostProcessing.cpp
--- a/Engine/Source/Runtime/Renderer/Private/PostProcess/PostProcessing.cpp
+++ b/Engine/Source/Runtime/Renderer/Private/PostProcess/PostProcessing.cpp
@@ -1480,6 +1480,16 @@ void FPostProcessing::ProcessES2(FRHICommandListImmediate& RHICmdList, const FVi
 			}
 		}
 
+		// Composite editor primitives if we had any to draw and compositing is enabled
+		// TODO: Move FDeferredShadingSceneRenderer::ShouldCompositeEditorPrimitives somewhere more generic
+		if (FDeferredShadingSceneRenderer::ShouldCompositeEditorPrimitives(View))
+		{
+			FRenderingCompositePass* EditorCompNode = Context.Graph.RegisterPass(new(FMemStack::Get()) FRCPassPostProcessCompositeEditorPrimitives());
+			EditorCompNode->SetInput(ePId_Input0, FRenderingCompositeOutputRef(Context.FinalOutput));
+			//Node->SetInput(ePId_Input1, FRenderingCompositeOutputRef(Context.SceneDepth));
+			Context.FinalOutput = FRenderingCompositeOutputRef(EditorCompNode);
+		}
+
 		// Must run to blit to back buffer even if post processing is off.
 		FRenderingCompositePass* PostProcessTonemap = Context.Graph.RegisterPass(new(FMemStack::Get()) FRCPassPostProcessTonemapES2(bUsedFramebufferFetch));
 		PostProcessTonemap->SetInput(ePId_Input0, Context.FinalOutput);
diff --git a/Engine/Source/Runtime/Renderer/Private/PostProcess/SceneRenderTargets.cpp b/Engine/Source/Runtime/Renderer/Private/PostProcess/SceneRenderTargets.cpp
--- a/Engine/Source/Runtime/Renderer/Private/PostProcess/SceneRenderTargets.cpp
+++ b/Engine/Source/Runtime/Renderer/Private/PostProcess/SceneRenderTargets.cpp
@@ -189,6 +189,9 @@ void FSceneRenderTargets::Allocate(const FSceneViewFamily& ViewFamily)
 		MinShadowResolution = CVar->GetValueOnRenderThread();
 	}
 
+	// If feature level has changed, release all previously allocated targets to the pool. If feature level has changed but
+	const auto NewFeatureLevel = ViewFamily.Scene->GetFeatureLevel();
+
 	if( (BufferSize.X != DesiredBufferSize.X) ||
 		(BufferSize.Y != DesiredBufferSize.Y) ||
 		(CurrentGBufferFormat != GBufferFormat) ||
@@ -200,7 +203,8 @@ void FSceneRenderTargets::Allocate(const FSceneViewFamily& ViewFamily)
 		(CurrentMobile32bpp != Mobile32bpp) ||
 		(CurrentMobileMSAA != MobileMSAA) ||
 		(bCurrentLightPropagationVolume != bLightPropagationVolume) ||
-		(CurrentMinShadowResolution != MinShadowResolution))
+		(CurrentMinShadowResolution != MinShadowResolution) ||
+		(NewFeatureLevel != CurrentFeatureLevel))
 	{
 		CurrentGBufferFormat = GBufferFormat;
 		CurrentSceneColorFormat = SceneColorFormat;
@@ -212,6 +216,7 @@ void FSceneRenderTargets::Allocate(const FSceneViewFamily& ViewFamily)
 		CurrentMobileMSAA = MobileMSAA;
 		CurrentMinShadowResolution = MinShadowResolution;
 		bCurrentLightPropagationVolume = bLightPropagationVolume;
+		CurrentFeatureLevel = NewFeatureLevel;
 		
 		// Reinitialize the render targets for the given size.
 		SetBufferSize(DesiredBufferSize.X, DesiredBufferSize.Y);
@@ -258,7 +263,7 @@ void FSceneRenderTargets::BeginRenderingSceneColor(FRHICommandListImmediate& RHI
 	AllocSceneColor();
 	
 	// Set the scene color surface as the render target, and the scene depth surface as the depth-stencil target.
-	if (bGBufferPass && GRHIFeatureLevel >= ERHIFeatureLevel::SM4)
+	if (bGBufferPass && CurrentFeatureLevel >= ERHIFeatureLevel::SM4)
 	{
 		FTextureRHIParamRef RenderTargets[6] = {0};
 		RenderTargets[0] = GetSceneColorSurface();
@@ -290,7 +295,7 @@ int32 FSceneRenderTargets::GetNumGBufferTargets() const
 {
 	int32 NumGBufferTargets = 1;
 
-	if (GRHIFeatureLevel >= ERHIFeatureLevel::SM4 && !IsSimpleDynamicLightingEnabled())
+	if (CurrentFeatureLevel >= ERHIFeatureLevel::SM4 && !IsSimpleDynamicLightingEnabled())
 	{
 		NumGBufferTargets = bAllowStaticLighting ? 6 : 5;
 	}
@@ -319,17 +324,20 @@ void FSceneRenderTargets::AllocSceneColor()
 
 		// with TexCreate_UAV it would allow better sharing with later elements but it might come at a high cost:
 		// GCNPerformanceTweets.pdf Tip 37: Warning: Causes additional synchronization between draw calls when using a render target allocated with this flag, use sparingly
-		if(GRHIFeatureLevel >= ERHIFeatureLevel::SM5 && !OptimizeForUAVPerformance)
+		if (CurrentFeatureLevel >= ERHIFeatureLevel::SM5 && !OptimizeForUAVPerformance)
 		{
 			Desc.TargetableFlags |= TexCreate_UAV;
 		}
 
-	    uint16 NumSamples = GRHIShaderPlatform == SP_OPENGL_ES2_IOS ? 1 : CVarMobileMSAA.GetValueOnRenderThread();
-        if (NumSamples != 1 && NumSamples != 2 && NumSamples != 4)
+		if (CurrentFeatureLevel == ERHIFeatureLevel::ES2)
 		{
-            NumSamples = 1;
+		    uint16 NumSamples = GRHIShaderPlatform == SP_OPENGL_ES2_IOS ? 1 : CVarMobileMSAA.GetValueOnRenderThread();
+	        if (NumSamples != 1 && NumSamples != 2 && NumSamples != 4)
+			{
+	            NumSamples = 1;
+			}
+			Desc.NumSamples = NumSamples;
 		}
-		Desc.NumSamples = NumSamples;
 
 		GRenderTargetPool.FindFreeElement(Desc, SceneColor, TEXT("SceneColor"));
 	}
@@ -361,7 +369,7 @@ void FSceneRenderTargets::AllocLightAttenuation()
 	check(LightAttenuation);
 }
 
-void FSceneRenderTargets::FreeGBufferTargets()
+void FSceneRenderTargets::ReleaseGBufferTargets()
 {
 	GBufferResourcesUniformBuffer.SafeRelease();
 	GBufferA.SafeRelease();
@@ -440,7 +448,7 @@ void FSceneRenderTargets::AllocGBufferTargets()
 	check(GBufferA);
 
 	// Create the required render targets if running Highend.
-	if (GRHIFeatureLevel >= ERHIFeatureLevel::SM4)
+	if (CurrentFeatureLevel >= ERHIFeatureLevel::SM4)
 	{
 
 		// Allocate the Gbuffer resource uniform buffer.
@@ -583,7 +591,7 @@ void FSceneRenderTargets::AdjustGBufferRefCount(int Delta)
 	}
 	else if(GBufferRefCount == 0)
 	{
-		FreeGBufferTargets(); 
+		ReleaseGBufferTargets();
 	}
 }
 
@@ -635,7 +643,7 @@ void FSceneRenderTargets::ResolveSceneColor(FRHICommandList& RHICmdList, const F
 /** Resolves the GBuffer targets so that their resolved textures can be sampled. */
 void FSceneRenderTargets::ResolveGBufferSurfaces(FRHICommandList& RHICmdList, const FResolveRect& ResolveRect)
 {
-	if (GRHIFeatureLevel >= ERHIFeatureLevel::SM4)
+	if (CurrentFeatureLevel >= ERHIFeatureLevel::SM4)
 	{
 		SCOPED_DRAW_EVENT(ResolveGBufferSurfaces, DEC_SCENE_ITEMS);
 
@@ -841,7 +849,7 @@ int32 FSceneRenderTargets::GetEditorMSAACompositingSampleCount() const
 	int32 Value = 1;
 
 	// only supported on SM5 yet (SM4 doesn't have MSAA sample load functionality which makes it harder to implement)
-	if (GRHIFeatureLevel >= ERHIFeatureLevel::SM5)
+	if (CurrentFeatureLevel >= ERHIFeatureLevel::SM5)
 	{
 		static const auto CVar = IConsoleManager::Get().FindTConsoleVariableDataInt(TEXT("r.MSAA.CompositingSampleCount"));
 
@@ -909,7 +917,7 @@ bool FSceneRenderTargets::IsSeparateTranslucencyActive(const FViewInfo& View) co
 	static const auto CVar = IConsoleManager::Get().FindTConsoleVariableDataInt(TEXT("r.SeparateTranslucency"));
 	int32 Value = CVar->GetValueOnRenderThread();
 
-	return (Value != 0) && GRHIFeatureLevel >= ERHIFeatureLevel::SM4
+	return (Value != 0) && CurrentFeatureLevel >= ERHIFeatureLevel::SM4
 		&& View.Family->EngineShowFlags.PostProcessing
 		&& View.Family->EngineShowFlags.SeparateTranslucency
 		&& View.State != NULL;	// We require a ViewState in order for separate translucency to work (it keeps track of our SeparateTranslucencyRT)
@@ -1073,7 +1081,7 @@ void FSceneRenderTargets::AllocateDeferredShadingPathRenderTargets()
 
 	const FIntPoint TranslucentShadowBufferResolution = GetTranslucentShadowDepthTextureResolution();
 
-	if (GRHIFeatureLevel >= ERHIFeatureLevel::SM4)
+	if (CurrentFeatureLevel >= ERHIFeatureLevel::SM4)
 	{
 		for (int32 SurfaceIndex = 0; SurfaceIndex < NumTranslucencyShadowSurfaces; SurfaceIndex++)
 		{
@@ -1086,7 +1094,7 @@ void FSceneRenderTargets::AllocateDeferredShadingPathRenderTargets()
 		}
 	}
 
-	if (GRHIFeatureLevel >= ERHIFeatureLevel::SM4)
+	if (CurrentFeatureLevel >= ERHIFeatureLevel::SM4)
 	{
 		// Create several shadow depth cube maps with different resolutions, to handle different sized shadows on the screen
 		for (int32 SurfaceIndex = 0; SurfaceIndex < NumCubeShadowDepthSurfaces; SurfaceIndex++)
@@ -1112,7 +1120,7 @@ void FSceneRenderTargets::AllocateDeferredShadingPathRenderTargets()
 	}
 
 	// Create the required render targets if running Highend.
-	if (GRHIFeatureLevel >= ERHIFeatureLevel::SM4)
+	if (CurrentFeatureLevel >= ERHIFeatureLevel::SM4)
 	{
 		// Create the screen space ambient occlusion buffer
 		{
@@ -1204,7 +1212,7 @@ void FSceneRenderTargets::AllocateDeferredShadingPathRenderTargets()
 		GRenderTargetPool.FindFreeElement(Desc, SkySHIrradianceMap, TEXT("SkySHIrradianceMap"));
 	}
 
-	if (GRHIFeatureLevel >= ERHIFeatureLevel::SM5)
+	if (CurrentFeatureLevel >= ERHIFeatureLevel::SM5)
 	{
 		// Create the reflective shadow map textures for LightPropagationVolume feature
 		if(bCurrentLightPropagationVolume)
@@ -1231,7 +1239,7 @@ EPixelFormat FSceneRenderTargets::GetSceneColorFormat() const
 {
 	EPixelFormat SceneColorBufferFormat = PF_FloatRGBA;
 
-	if (GRHIFeatureLevel <= ERHIFeatureLevel::ES2)
+	if (CurrentFeatureLevel <= ERHIFeatureLevel::ES2)
 	{
 		// Potentially allocate an alpha channel in th -fe scene color texture to store the resolved scene depth.
 		SceneColorBufferFormat = GSupportsRenderTargetFormat_PF_FloatRGBA ? PF_FloatRGBA : PF_B8G8R8A8;
@@ -1272,7 +1280,7 @@ void FSceneRenderTargets::InitDynamicRHI()
 		FRHICommandListImmediate& RHICmdList = FRHICommandListExecutor::GetImmediateCommandList();
 		RHICmdList.SetScissorRect(false, 0, 0, 0, 0);
 
-		if (GRHIFeatureLevel == ERHIFeatureLevel::ES2)
+		if (CurrentFeatureLevel == ERHIFeatureLevel::ES2)
 		{
 			AllocateForwardShadingPathRenderTargets();
 		}
@@ -1283,9 +1291,9 @@ void FSceneRenderTargets::InitDynamicRHI()
 	}
 }
 
-void FSceneRenderTargets::ReleaseDynamicRHI()
+void FSceneRenderTargets::ReleaseAllTargets()
 {
-	FreeGBufferTargets();
+	ReleaseGBufferTargets();
 
 	SceneColor.SafeRelease();
 	SceneAlphaCopy.SafeRelease();
@@ -1339,7 +1347,11 @@ void FSceneRenderTargets::ReleaseDynamicRHI()
 	
 	EditorPrimitivesColor.SafeRelease();
 	EditorPrimitivesDepth.SafeRelease();
+}
 
+void FSceneRenderTargets::ReleaseDynamicRHI()
+{
+	ReleaseAllTargets();
 	GRenderTargetPool.FreeUnusedResources();
 }
 
@@ -1555,11 +1567,13 @@ void FSceneTextureShaderParameters::Set(
 				);
 		}
 
-		if (GRHIFeatureLevel >= ERHIFeatureLevel::SM5)
+		const auto FeatureLevel = View.GetFeatureLevel();
+
+		if (FeatureLevel >= ERHIFeatureLevel::SM5)
 		{
 			SetTextureParameter(RHICmdList, ShaderRHI, SceneColorSurfaceParameter, GSceneRenderTargets.GetSceneColorSurface());
 		}
-		if(GRHIFeatureLevel >= ERHIFeatureLevel::SM4)
+		if (FeatureLevel >= ERHIFeatureLevel::SM4)
 		{
 			if(GSupportsDepthFetchDuringDepthTest)
 			{
@@ -1697,7 +1711,9 @@ void FDeferredPixelShaderParameters::Set(FRHICommandList& RHICmdList, const Shad
 		SetTextureParameter(RHICmdList, ShaderRHI, DBufferCTextureNonMS, DBufferC->GetRenderTargetItem().ShaderResourceTexture);
 	}
 
-	if (TextureMode == ESceneRenderTargetsMode::SetTextures && GRHIFeatureLevel >= ERHIFeatureLevel::SM3)
+	const auto FeatureLevel = View.GetFeatureLevel();
+
+	if (TextureMode == ESceneRenderTargetsMode::SetTextures && FeatureLevel >= ERHIFeatureLevel::SM3)
 	{
 		// if there is no ambient occlusion it's better to have white there
 		IPooledRenderTarget* ScreenSpaceAO = GSceneRenderTargets.ScreenSpaceAO;
@@ -1713,7 +1729,7 @@ void FDeferredPixelShaderParameters::Set(FRHICommandList& RHICmdList, const Shad
 			CustomDepth = GSystemTextures.BlackDummy;
 		}
 
-		if (GRHIFeatureLevel >= ERHIFeatureLevel::SM4)
+		if (FeatureLevel >= ERHIFeatureLevel::SM4)
 		{
 			//gilmerge
 			if (GBufferResources.IsBound())
diff --git a/Engine/Source/Runtime/Renderer/Private/PostProcess/SceneRenderTargets.h b/Engine/Source/Runtime/Renderer/Private/PostProcess/SceneRenderTargets.h
--- a/Engine/Source/Runtime/Renderer/Private/PostProcess/SceneRenderTargets.h
+++ b/Engine/Source/Runtime/Renderer/Private/PostProcess/SceneRenderTargets.h
@@ -126,7 +126,8 @@ protected:
 		CurrentMaxShadowResolution(0),
 		CurrentTranslucencyLightingVolumeDim(64),
 		CurrentMobile32bpp(0),
-		bCurrentLightPropagationVolume(false)
+		bCurrentLightPropagationVolume(false),
+		CurrentFeatureLevel(ERHIFeatureLevel::Num)
 		{}
 public:
 
@@ -257,7 +258,7 @@ public:
 	const FTexture2DRHIRef* GetActualDepthTexture() const
 	{
 		const FTexture2DRHIRef* DepthTexture = NULL;
-		if((GRHIFeatureLevel >= ERHIFeatureLevel::SM4) || IsPCPlatform(GRHIShaderPlatform))
+		if((CurrentFeatureLevel >= ERHIFeatureLevel::SM4) || IsPCPlatform(GRHIShaderPlatform))
 		{
 			if(GSupportsDepthFetchDuringDepthTest)
 			{
@@ -514,7 +515,10 @@ private:
 	void AllocLightAttenuation();
 
 	// internal method, used by AdjustGBufferRefCount()
-	void FreeGBufferTargets();
+	void ReleaseGBufferTargets();
+
+	// release all allocated targets to the pool
+	void ReleaseAllTargets();
 
 	EPixelFormat GetSceneColorFormat() const;
 
@@ -547,6 +551,8 @@ private:
 	int32 CurrentMinShadowResolution;
 	/** To detect a change of the CVar r.LightPropagationVolume */
 	bool bCurrentLightPropagationVolume;
+	/** Feature level we were initialized for */
+	ERHIFeatureLevel::Type CurrentFeatureLevel;
 };
 
 /** The global render targets used for scene rendering. */
diff --git a/Engine/Source/Runtime/Renderer/Private/ReflectionEnvironment.cpp b/Engine/Source/Runtime/Renderer/Private/ReflectionEnvironment.cpp
--- a/Engine/Source/Runtime/Renderer/Private/ReflectionEnvironment.cpp
+++ b/Engine/Source/Runtime/Renderer/Private/ReflectionEnvironment.cpp
@@ -60,7 +60,7 @@ bool IsReflectionEnvironmentAvailable(ERHIFeatureLevel::Type InFeatureLevel)
 
 void FReflectionEnvironmentCubemapArray::InitDynamicRHI()
 {
-	if (GRHIFeatureLevel == ERHIFeatureLevel::SM5)
+	if (GetFeatureLevel() >= ERHIFeatureLevel::SM5)
 	{
 		const int32 NumReflectionCaptureMips = FMath::CeilLogTwo(GReflectionCaptureSize) + 1;
 
diff --git a/Engine/Source/Runtime/Renderer/Private/ReflectionEnvironmentCapture.cpp b/Engine/Source/Runtime/Renderer/Private/ReflectionEnvironmentCapture.cpp
--- a/Engine/Source/Runtime/Renderer/Private/ReflectionEnvironmentCapture.cpp
+++ b/Engine/Source/Runtime/Renderer/Private/ReflectionEnvironmentCapture.cpp
@@ -1004,7 +1004,7 @@ void FScene::AllocateReflectionCaptures(const TArray<UReflectionCaptureComponent
 {
 	if (NewCaptures.Num() > 0)
 	{
-		if (FeatureLevel == ERHIFeatureLevel::SM5)
+		if (GetFeatureLevel() >= ERHIFeatureLevel::SM5)
 		{
 			for (int32 CaptureIndex = 0; CaptureIndex < NewCaptures.Num(); CaptureIndex++)
 			{
@@ -1074,7 +1074,7 @@ void FScene::AllocateReflectionCaptures(const TArray<UReflectionCaptureComponent
 				}
 			}
 		}
-		else if (FeatureLevel == ERHIFeatureLevel::SM4)
+		else if (GetFeatureLevel() == ERHIFeatureLevel::SM4)
 		{
 			for (int32 ComponentIndex = 0; ComponentIndex < NewCaptures.Num(); ComponentIndex++)
 			{
@@ -1103,7 +1103,7 @@ void FScene::AllocateReflectionCaptures(const TArray<UReflectionCaptureComponent
 /** Updates the contents of all reflection captures in the scene.  Must be called from the game thread. */
 void FScene::UpdateAllReflectionCaptures()
 {
-	if (IsReflectionEnvironmentAvailable(FeatureLevel))
+	if (IsReflectionEnvironmentAvailable(GetFeatureLevel()))
 	{
 		ENQUEUE_UNIQUE_RENDER_COMMAND_ONEPARAMETER( 
 			CaptureCommand,
@@ -1200,7 +1200,7 @@ void GetReflectionCaptureData_RenderingThread(FRHICommandListImmediate& RHICmdLi
 
 void FScene::GetReflectionCaptureData(UReflectionCaptureComponent* Component, FReflectionCaptureFullHDRDerivedData& OutDerivedData) 
 {
-	check(FeatureLevel == ERHIFeatureLevel::SM5);
+	check(GetFeatureLevel() >= ERHIFeatureLevel::SM5);
 
 	ENQUEUE_UNIQUE_RENDER_COMMAND_THREEPARAMETER(
 		GetReflectionDataCommand,
@@ -1463,7 +1463,7 @@ void CopyToComponentTexture(FRHICommandList& RHICmdList, FScene* Scene, FReflect
  */
 void FScene::UpdateReflectionCaptureContents(UReflectionCaptureComponent* CaptureComponent)
 {
-	if (IsReflectionEnvironmentAvailable(FeatureLevel))
+	if (IsReflectionEnvironmentAvailable(GetFeatureLevel()))
 	{
 		const FReflectionCaptureFullHDRDerivedData* DerivedData = CaptureComponent->GetCachedFullHDRDerivedData();
 
@@ -1471,7 +1471,7 @@ void FScene::UpdateReflectionCaptureContents(UReflectionCaptureComponent* Captur
 		if (DerivedData && DerivedData->CompressedCapturedData.Num() > 0)
 		{
 			// For other feature levels the reflection textures are stored on the component instead of in a scene-wide texture array
-			if (FeatureLevel == ERHIFeatureLevel::SM5)
+			if (GetFeatureLevel() >= ERHIFeatureLevel::SM5)
 			{
 				ENQUEUE_UNIQUE_RENDER_COMMAND_THREEPARAMETER( 
 					UploadCaptureCommand,
@@ -1502,7 +1502,7 @@ void FScene::UpdateReflectionCaptureContents(UReflectionCaptureComponent* Captur
 				CopyCommand,
 				FScene*, Scene, this,
 				FReflectionCaptureProxy*, ReflectionProxy, ReflectionProxy,
-				ERHIFeatureLevel::Type, FeatureLevel, FeatureLevel,
+				ERHIFeatureLevel::Type, FeatureLevel, GetFeatureLevel(),
 			{
 				if (FeatureLevel == ERHIFeatureLevel::SM5)
 				{
@@ -1540,7 +1540,7 @@ void CopyToSkyTexture(FRHICommandList& RHICmdList, FScene* Scene, FTexture* Proc
 
 void FScene::UpdateSkyCaptureContents(const USkyLightComponent* CaptureComponent, bool bCaptureEmissiveOnly, FTexture* OutProcessedTexture, FSHVectorRGB3& OutIrradianceEnvironmentMap)
 {
-	if (FeatureLevel >= ERHIFeatureLevel::SM4)
+	if (GetFeatureLevel() >= ERHIFeatureLevel::SM4)
 	{
 		ENQUEUE_UNIQUE_RENDER_COMMAND( 
 			ClearCommand,
diff --git a/Engine/Source/Runtime/Renderer/Private/RendererScene.cpp b/Engine/Source/Runtime/Renderer/Private/RendererScene.cpp
--- a/Engine/Source/Runtime/Renderer/Private/RendererScene.cpp
+++ b/Engine/Source/Runtime/Renderer/Private/RendererScene.cpp
@@ -192,7 +192,8 @@ FScene::FScene(UWorld* InWorld, bool bInRequiresHitProxies, bool bInIsEditorScen
 ,	SkyLight(NULL)
 ,	SimpleDirectionalLight(NULL)
 ,	SunLight(NULL)
-,	IndirectLightingCache()
+,	ReflectionSceneData(InFeatureLevel)
+,	IndirectLightingCache(InFeatureLevel)
 ,	SurfaceCacheResources(NULL)
 ,	PreshadowCacheLayout(0, 0, 0, 0, false, false)
 ,	AtmosphericFog(NULL)
@@ -206,7 +207,6 @@ FScene::FScene(UWorld* InWorld, bool bInRequiresHitProxies, bool bInIsEditorScen
 ,	LowerDynamicSkylightColor(FLinearColor::Black)
 ,	NumVisibleLights(0)
 ,	bHasSkyLight(false)
-,	FeatureLevel(InFeatureLevel)
 {
 	check(World);
 
@@ -1525,6 +1525,7 @@ void FScene::SetShaderMapsOnMaterialResources_RenderThread(FRHICommandListImmedi
 		MaterialArray.Add(Material);
 	}
 
+	const auto FeatureLevel = GetFeatureLevel();
 	bool bFoundAnyInitializedMaterials = false;
 
 	// Iterate through all loaded material render proxies and recache their uniform expressions if needed
@@ -1991,13 +1992,6 @@ void FScene::OnLevelAddedToWorld_RenderThread(FName InLevelName)
 	}
 }
 
-void FScene::ChangeFeatureLevel(ERHIFeatureLevel::Type InFeatureLevel)
-{
-	FeatureLevel = InFeatureLevel;
-	GRHIShaderPlatform = GShaderPlatformForFeatureLevel[InFeatureLevel];
-	SetMaxRHIFeatureLevel(InFeatureLevel);
-}
-
 /**
  * Dummy NULL scene interface used by dedicated servers.
  */
diff --git a/Engine/Source/Runtime/Renderer/Private/SceneCaptureRendering.cpp b/Engine/Source/Runtime/Renderer/Private/SceneCaptureRendering.cpp
--- a/Engine/Source/Runtime/Renderer/Private/SceneCaptureRendering.cpp
+++ b/Engine/Source/Runtime/Renderer/Private/SceneCaptureRendering.cpp
@@ -239,7 +239,7 @@ void FScene::UpdateSceneCaptureContents(USceneCaptureComponentCube* CaptureCompo
 
 	check(CaptureComponent);
 
-	if (FeatureLevel >= ERHIFeatureLevel::SM4 && CaptureComponent->TextureTarget)
+	if (GetFeatureLevel() >= ERHIFeatureLevel::SM4 && CaptureComponent->TextureTarget)
 	{
 		const float FOV = 90 * (float)PI / 360.0f;
 		for (int32 faceidx = 0; faceidx < (int32)ECubeFace::CubeFace_MAX; faceidx++)
diff --git a/Engine/Source/Runtime/Renderer/Private/SceneOcclusion.cpp b/Engine/Source/Runtime/Renderer/Private/SceneOcclusion.cpp
--- a/Engine/Source/Runtime/Renderer/Private/SceneOcclusion.cpp
+++ b/Engine/Source/Runtime/Renderer/Private/SceneOcclusion.cpp
@@ -519,7 +519,7 @@ void FHZBOcclusionTester::SetInvalidFrameNumber()
 
 void FHZBOcclusionTester::InitDynamicRHI()
 {
-	if (GRHIFeatureLevel >= ERHIFeatureLevel::SM3)
+	if (GetFeatureLevel() >= ERHIFeatureLevel::SM3)
 	{
 #if PLATFORM_MAC // Workaround radr://16096028 Texture Readback via glReadPixels + PBOs stalls on Nvidia GPUs
 		FPooledRenderTargetDesc Desc( FPooledRenderTargetDesc::Create2DDesc( FIntPoint( SizeX, SizeY ), PF_R8G8B8A8, TexCreate_CPUReadback | TexCreate_HideInVisualizeTexture, TexCreate_None, false ) );
@@ -532,7 +532,7 @@ void FHZBOcclusionTester::InitDynamicRHI()
 
 void FHZBOcclusionTester::ReleaseDynamicRHI()
 {
-	if (GRHIFeatureLevel >= ERHIFeatureLevel::SM3)
+	if (GetFeatureLevel() >= ERHIFeatureLevel::SM3)
 	{
 		GRenderTargetPool.FreeUnusedResource( ResultsTextureCPU );
 	}
diff --git a/Engine/Source/Runtime/Renderer/Private/ScenePrivate.h b/Engine/Source/Runtime/Renderer/Private/ScenePrivate.h
--- a/Engine/Source/Runtime/Renderer/Private/ScenePrivate.h
+++ b/Engine/Source/Runtime/Renderer/Private/ScenePrivate.h
@@ -728,8 +728,9 @@ class FReflectionEnvironmentCubemapArray : public FRenderResource
 {
 public:
 
-	FReflectionEnvironmentCubemapArray() : 
-		MaxCubemaps(0)
+	FReflectionEnvironmentCubemapArray(ERHIFeatureLevel::Type InFeatureLevel) : 
+		FRenderResource(InFeatureLevel)
+	,	MaxCubemaps(0)
 	{}
 
 	virtual void InitDynamicRHI() override;
@@ -801,8 +802,9 @@ public:
 	/** Game thread tracking of what size this scene has allocated for the cubemap array. */
 	int32 MaxAllocatedReflectionCubemapsGameThread;
 
-	FReflectionEnvironmentSceneData() :
+	FReflectionEnvironmentSceneData(ERHIFeatureLevel::Type InFeatureLevel) :
 		bRegisteredReflectionCapturesHasChanged(true),
+		CubemapArray(InFeatureLevel),
 		MaxAllocatedReflectionCubemapsGameThread(0)
 	{}
 };
@@ -850,7 +852,7 @@ class FIndirectLightingCache : public FRenderResource
 public:
 
 	/** true for the editor case where we want a better preview for object that have no valid lightmaps */
-	FIndirectLightingCache();
+	FIndirectLightingCache(ERHIFeatureLevel::Type InFeatureLevel);
 
 	// FRenderResource interface
 	virtual void InitDynamicRHI();
@@ -1402,8 +1404,7 @@ public:
 
 	virtual bool IsEditorScene() const override { return bIsEditorScene; }
 
-	virtual ERHIFeatureLevel::Type GetFeatureLevel() const override { return FeatureLevel; }
-	virtual void ChangeFeatureLevel(ERHIFeatureLevel::Type InFeatureLevel) override;
+	virtual ERHIFeatureLevel::Type GetFeatureLevel() const override { return GetWorld()->FeatureLevel; }
 
 private:
 
@@ -1516,11 +1517,6 @@ private:
 	 * Note: This is tracked on the game thread!
 	 */
 	bool bHasSkyLight;
-
-	/** 
-	 * Feature level that this scene should render using
-	 */
-	ERHIFeatureLevel::Type FeatureLevel;
 };
 
 #endif // __SCENEPRIVATE_H__
diff --git a/Engine/Source/Runtime/Renderer/Private/SceneRendering.cpp b/Engine/Source/Runtime/Renderer/Private/SceneRendering.cpp
--- a/Engine/Source/Runtime/Renderer/Private/SceneRendering.cpp
+++ b/Engine/Source/Runtime/Renderer/Private/SceneRendering.cpp
@@ -468,6 +468,7 @@ TUniformBufferRef<FViewUniformShaderParameters> FViewInfo::CreateUniformBuffer(
 	ViewUniformShaderParameters.SceneTextureMinMax = SceneTexMinMax;
 
 	FScene* Scene = (FScene*)Family->Scene;
+	ERHIFeatureLevel::Type FeatureLevel;
 
 	if (Scene && Scene->SkyLight)
 	{
@@ -481,17 +482,26 @@ TUniformBufferRef<FViewUniformShaderParameters> FViewInfo::CreateUniformBuffer(
 			&& SkyLight->bPrecomputedLightingIsValid;
 
 		ViewUniformShaderParameters.SkyLightParameters = bApplyPrecomputedBentNormalShadowing ? 1 : 0;
+
+		FeatureLevel = Scene->GetFeatureLevel();
 	}
 	else
 	{
 		ViewUniformShaderParameters.SkyLightColor = FLinearColor::Black;
 		ViewUniformShaderParameters.SkyLightParameters = 0;
+
+		FeatureLevel = GMaxRHIFeatureLevel;
 	}
 
 	// Make sure there's no padding since we're going to cast to FVector4*
 	checkSlow(sizeof(ViewUniformShaderParameters.SkyIrradianceEnvironmentMap) == sizeof(FVector4) * 7);
 	SetupSkyIrradianceEnvironmentMapConstants((FVector4*)&ViewUniformShaderParameters.SkyIrradianceEnvironmentMap);
 
+	ViewUniformShaderParameters.ES2PreviewMode =
+		(GIsEditor &&
+		FeatureLevel == ERHIFeatureLevel::ES2 &&
+		GMaxRHIFeatureLevel > ERHIFeatureLevel::ES2) ? 1.0f : 0.0f;
+
 	return TUniformBufferRef<FViewUniformShaderParameters>::CreateUniformBufferImmediate(ViewUniformShaderParameters, UniformBuffer_SingleFrame);
 }
 
diff --git a/Engine/Source/Runtime/ShaderCore/Public/VertexFactory.h b/Engine/Source/Runtime/ShaderCore/Public/VertexFactory.h
--- a/Engine/Source/Runtime/ShaderCore/Public/VertexFactory.h
+++ b/Engine/Source/Runtime/ShaderCore/Public/VertexFactory.h
@@ -331,6 +331,19 @@ public:
 		DataType() {}
 	};
 
+	/**
+	 * Default constructor 
+	 */
+	FVertexFactory() {}
+
+	/**
+	* Constructor specifiying feature level
+	*/
+	FVertexFactory(ERHIFeatureLevel::Type InFeatureLevel) 
+		: FRenderResource(InFeatureLevel)
+	{
+	}
+
 	/**
 	 * @return The vertex factory's type.
 	 */
diff --git a/Engine/Source/Runtime/Windows/D3D11RHI/Private/D3D11Device.cpp b/Engine/Source/Runtime/Windows/D3D11RHI/Private/D3D11Device.cpp
--- a/Engine/Source/Runtime/Windows/D3D11RHI/Private/D3D11Device.cpp
+++ b/Engine/Source/Runtime/Windows/D3D11RHI/Private/D3D11Device.cpp
@@ -81,17 +81,17 @@ FD3D11DynamicRHI::FD3D11DynamicRHI(IDXGIFactory* InDXGIFactory,D3D_FEATURE_LEVEL
 	// ES2 feature level emulation in D3D11
 	if (FParse::Param(FCommandLine::Get(), TEXT("FeatureLevelES2")))
 	{
-		SetMaxRHIFeatureLevel(ERHIFeatureLevel::ES2);
+		GMaxRHIFeatureLevel = GCurrentRHIFeatureLevel = ERHIFeatureLevel::ES2;
 		GRHIShaderPlatform = SP_PCD3D_ES2;
 	}
 	else if(FeatureLevel == D3D_FEATURE_LEVEL_11_0)
 	{
-		SetMaxRHIFeatureLevel(ERHIFeatureLevel::SM5);
+		GMaxRHIFeatureLevel = GCurrentRHIFeatureLevel = ERHIFeatureLevel::SM5;
 		GRHIShaderPlatform = SP_PCD3D_SM5;
 	}
 	else if(FeatureLevel == D3D_FEATURE_LEVEL_10_0)
 	{
-		SetMaxRHIFeatureLevel(ERHIFeatureLevel::SM4);
+		GMaxRHIFeatureLevel = GCurrentRHIFeatureLevel = ERHIFeatureLevel::SM4;
 		GRHIShaderPlatform = SP_PCD3D_SM4;
 	}
 
diff --git a/Engine/Source/Runtime/Windows/D3D11RHI/Private/D3D11Texture.cpp b/Engine/Source/Runtime/Windows/D3D11RHI/Private/D3D11Texture.cpp
--- a/Engine/Source/Runtime/Windows/D3D11RHI/Private/D3D11Texture.cpp
+++ b/Engine/Source/Runtime/Windows/D3D11RHI/Private/D3D11Texture.cpp
@@ -481,7 +481,7 @@ TD3D11Texture2D<BaseResourceType>* FD3D11DynamicRHI::CreateD3D11Texture2D(uint32
 
 	bool bPooledTexture = true;
 
-	if (GRHIFeatureLevel <= ERHIFeatureLevel::ES2)
+	if (GMaxRHIFeatureLevel <= ERHIFeatureLevel::ES2)
 	{
 		// Remove sRGB read flag when not supported
 		Flags &= ~TexCreate_SRGB;
@@ -997,7 +997,7 @@ FTexture2DRHIRef FD3D11DynamicRHI::RHIAsyncCreateTexture2D(uint32 SizeX,uint32 S
 	check(GRHISupportsAsyncTextureCreation);
 	check((Flags & InvalidFlags) == 0);
 
-	if (GRHIFeatureLevel <= ERHIFeatureLevel::ES2)
+	if (GMaxRHIFeatureLevel <= ERHIFeatureLevel::ES2)
 	{
 		// Remove sRGB read flag when not supported
 		Flags &= ~TexCreate_SRGB;
