Change graph tasks to reduce overhead on game thread when using parallel animation evaluation

diff --git a/Engine/Source/Runtime/Engine/Classes/Components/SkeletalMeshComponent.h b/Engine/Source/Runtime/Engine/Classes/Components/SkeletalMeshComponent.h
--- a/Engine/Source/Runtime/Engine/Classes/Components/SkeletalMeshComponent.h
+++ b/Engine/Source/Runtime/Engine/Classes/Components/SkeletalMeshComponent.h
@@ -1050,10 +1050,6 @@ private:
 	bool ShouldBlendPhysicsBones();	
 	void ClearAnimScriptInstance();
 
-	// Reference to our animation evaluation tick event
-	FGraphEventRef EvaluationTickEvent;
-	FGraphEventRef TickCompletionEvent;
-
 	//Handle parallel evaluation of animation
 	TArray<FTransform> PTSpaceBases;
 	TArray<FTransform> PTLocalAtoms;
@@ -1062,10 +1058,7 @@ private:
 	bool bPTDoInterpolation;
 	bool bPTDuplicateToCacheBones;
 
-	//Parallel Task Delegates
-	FSimpleDelegateGraphTask::FDelegate ParallelEvaluationDelegate;
-	FSimpleDelegateGraphTask::FDelegate ParallelCompletionDelegate;
-
+public:
 	// Parallel evaluation wrappers
 	void ParallelAnimationEvaluation() { PerformAnimationEvaluation(PTSpaceBases, PTLocalAtoms, PTVertexAnims, PTRootBoneTranslation); }
 	void CompleteParallelAnimationEvaluation()
diff --git a/Engine/Source/Runtime/Engine/Private/SkeletalMeshComponent.cpp b/Engine/Source/Runtime/Engine/Private/SkeletalMeshComponent.cpp
--- a/Engine/Source/Runtime/Engine/Private/SkeletalMeshComponent.cpp
+++ b/Engine/Source/Runtime/Engine/Private/SkeletalMeshComponent.cpp
@@ -22,6 +22,70 @@
 
 TAutoConsoleVariable<int32> CVarUseParallelAnimationEvaluation(TEXT("a.ParallelAnimEvaluation"), 0, TEXT("If 1, animation evaluation will be run across the task graph system. If 0, evaluation will run purely on the game thread"));
 
+class FParallelAnimationEvaluationTask
+{
+	TWeakObjectPtr<USkeletalMeshComponent> SkeletalMeshComponent;
+
+public:
+	FParallelAnimationEvaluationTask(TWeakObjectPtr<USkeletalMeshComponent> InSkeletalMeshComponent)
+		: SkeletalMeshComponent(InSkeletalMeshComponent)
+	{
+	}
+
+	FORCEINLINE TStatId GetStatId() const
+	{
+		RETURN_QUICK_DECLARE_CYCLE_STAT(FParallelAnimationEvaluationTask, STATGROUP_TaskGraphTasks);
+	}
+	static ENamedThreads::Type GetDesiredThread()
+	{
+		return ENamedThreads::AnyThread;
+	}
+	static ESubsequentsMode::Type GetSubsequentsMode()
+	{
+		return ESubsequentsMode::TrackSubsequents;
+	}
+
+	void DoTask(ENamedThreads::Type CurrentThread, const FGraphEventRef& MyCompletionGraphEvent)
+	{
+		if (USkeletalMeshComponent* Comp = SkeletalMeshComponent.Get())
+		{
+			Comp->ParallelAnimationEvaluation();
+		}
+	}
+};
+
+class FParallelAnimationCompletionTask
+{
+	TWeakObjectPtr<USkeletalMeshComponent> SkeletalMeshComponent;
+
+public:
+	FParallelAnimationCompletionTask(TWeakObjectPtr<USkeletalMeshComponent> InSkeletalMeshComponent)
+		: SkeletalMeshComponent(InSkeletalMeshComponent)
+	{
+	}
+
+	FORCEINLINE TStatId GetStatId() const
+	{
+		RETURN_QUICK_DECLARE_CYCLE_STAT(FParallelAnimationCompletionTask, STATGROUP_TaskGraphTasks);
+	}
+	static ENamedThreads::Type GetDesiredThread()
+	{
+		return ENamedThreads::GameThread;
+	}
+	static ESubsequentsMode::Type GetSubsequentsMode()
+	{
+		return ESubsequentsMode::TrackSubsequents;
+	}
+
+	void DoTask(ENamedThreads::Type CurrentThread, const FGraphEventRef& MyCompletionGraphEvent)
+	{
+		if (USkeletalMeshComponent* Comp = SkeletalMeshComponent.Get())
+		{
+			Comp->CompleteParallelAnimationEvaluation();
+		}
+	}
+};
+
 USkeletalMeshComponent::USkeletalMeshComponent(const class FPostConstructInitializeProperties& PCIP)
 	: Super(PCIP)
 {
@@ -66,9 +130,6 @@ USkeletalMeshComponent::USkeletalMeshComponent(const class FPostConstructInitial
 	bEnableUpdateRateOptimizations = false;
 
 	bTickInEditor = true;
-
-	ParallelEvaluationDelegate = FSimpleDelegateGraphTask::FDelegate::CreateUObject(this, &USkeletalMeshComponent::ParallelAnimationEvaluation);
-	ParallelCompletionDelegate = FSimpleDelegateGraphTask::FDelegate::CreateUObject(this, &USkeletalMeshComponent::CompleteParallelAnimationEvaluation);
 }
 
 
@@ -889,23 +950,14 @@ void USkeletalMeshComponent::RefreshBoneTransforms(FActorComponentTickFunction*
 		}
 
 		// start parallel work
-		FGraphEventRef EvaluationTickEvent = FSimpleDelegateGraphTask::CreateAndDispatchWhenReady
-			(
-			ParallelEvaluationDelegate
-			, TEXT("USkeletalMeshComponent::ParallelAnimationEvaluation")
-			);
+		FGraphEventRef EvaluationTickEvent = TGraphTask<FParallelAnimationEvaluationTask>::CreateTask().ConstructAndDispatchWhenReady(this);
 
 		// set up a task to run on the game thread to accept the results
-		TickCompletionEvent = FSimpleDelegateGraphTask::CreateAndDispatchWhenReady
-			(
-			ParallelCompletionDelegate
-			, TEXT("USkeletalMeshComponent::CompleteParallelAnimationEvaluation")
-			, EvaluationTickEvent
-			, ENamedThreads::GameThread
-			);
+		FGraphEventArray Prerequistes;
+		Prerequistes.Add(EvaluationTickEvent);
+		FGraphEventRef TickCompletionEvent = TGraphTask<FParallelAnimationCompletionTask>::CreateTask(&Prerequistes).ConstructAndDispatchWhenReady(this);
 
 		TickFunction->GetCompletionHandle()->DontCompleteUntil(TickCompletionEvent);
-
 	}
 	else
 	{
