Check package guid when loading network packages + derive bPackage status implicitly

diff --git a/Engine/Source/Runtime/CoreUObject/Public/UObject/CoreObject.h b/Engine/Source/Runtime/CoreUObject/Public/UObject/CoreObject.h
--- a/Engine/Source/Runtime/CoreUObject/Public/UObject/CoreObject.h
+++ b/Engine/Source/Runtime/CoreUObject/Public/UObject/CoreObject.h
@@ -273,7 +273,7 @@ public:
 	}
 	
 	/** returns our Guid */
-	FORCEINLINE FGuid GetGuid()
+	FORCEINLINE FGuid GetGuid() const
 	{
 		return Guid;
 	}
diff --git a/Engine/Source/Runtime/Engine/Classes/Engine/PackageMapClient.h b/Engine/Source/Runtime/Engine/Classes/Engine/PackageMapClient.h
--- a/Engine/Source/Runtime/Engine/Classes/Engine/PackageMapClient.h
+++ b/Engine/Source/Runtime/Engine/Classes/Engine/PackageMapClient.h
@@ -37,6 +37,7 @@ public:
 	FNetworkGUID				OuterGUID;
 	FName						PathName;
 	int32						GuidSequence;				// We remember the guid sequence this net guid was generated on so we can tell how old it is
+	FGuid						PackageGuid;				// If this is a package, this is the guid to expect
 
 	float						InitialQueryTime;
 
@@ -60,7 +61,7 @@ public:
 	void			RegisterNetGUID_Internal( const FNetworkGUID & NetGUID, const FNetGuidCacheObject & CacheObject );
 	void			RegisterNetGUID_Server( const FNetworkGUID & NetGUID, const UObject * Object );
 	void			RegisterNetGUID_Client( const FNetworkGUID & NetGUID, const UObject * Object );
-	void			RegisterNetGUIDFromPath_Client( const FNetworkGUID & NetGUID, const FString & PathName, const FNetworkGUID & OuterGUID, const bool bNoLoad, const bool bIgnoreWhenMissing );
+	void			RegisterNetGUIDFromPath_Client( const FNetworkGUID & NetGUID, const FString & PathName, const FNetworkGUID & OuterGUID, const FGuid & PackageGuid, const bool bNoLoad, const bool bIgnoreWhenMissing );
 	UObject *		GetObjectFromNetGUID( const FNetworkGUID & NetGUID, const bool bIgnoreMustBeMapped );
 	bool			ShouldIgnoreWhenMissing( const FNetworkGUID & NetGUID ) const;
 	bool			IsGUIDRegistered( const FNetworkGUID & NetGUID ) const;
diff --git a/Engine/Source/Runtime/Engine/Private/PackageMapClient.cpp b/Engine/Source/Runtime/Engine/Private/PackageMapClient.cpp
--- a/Engine/Source/Runtime/Engine/Private/PackageMapClient.cpp
+++ b/Engine/Source/Runtime/Engine/Private/PackageMapClient.cpp
@@ -409,7 +409,6 @@ struct FExportFlags
 		{
 			uint8 bHasPath		: 1;
 			uint8 bNoLoad		: 1;
-			uint8 bIsPackage	: 1;
 		};
 
 		uint8	Value;
@@ -472,7 +471,7 @@ void UPackageMapClient::InternalWriteObject( FArchive & Ar, FNetworkGUID NetGUID
 		// Only the client sends default guids
 		check( !IsNetGUIDAuthority() );
 		ExportFlags.bHasPath = 1;
-	}		
+	}
 	else if ( GuidCache->IsExportingNetGUIDBunch )
 	{
 		// Only the server should be exporting guids
@@ -487,8 +486,7 @@ void UPackageMapClient::InternalWriteObject( FArchive & Ar, FNetworkGUID NetGUID
 			ExportFlags.bHasPath = ObjectPathName.IsEmpty() ? 0 : 1;
 		}
 
-		ExportFlags.bNoLoad		= bNoLoad ? 1 : 0;
-		ExportFlags.bIsPackage	= ( Cast< const UPackage >( Object ) != NULL ) ? 1 : 0;
+		ExportFlags.bNoLoad	= bNoLoad ? 1 : 0;
 
 		Ar << ExportFlags.Value;
 	}
@@ -510,6 +508,10 @@ void UPackageMapClient::InternalWriteObject( FArchive & Ar, FNetworkGUID NetGUID
 			check( !ObjectPathName.IsEmpty() );
 		}
 
+		const bool bIsPackage = ( NetGUID.IsStatic() && Object != NULL && Object->GetOuter() == NULL );
+
+		check( bIsPackage == ( Cast< UPackage >( Object ) != NULL ) );		// Make sure it really is a package
+
 		// Serialize reference to outer. This is basically a form of compression.
 		FNetworkGUID OuterNetGUID = GuidCache->GetOrAssignNetGUID( ObjectOuter );
 
@@ -520,6 +522,12 @@ void UPackageMapClient::InternalWriteObject( FArchive & Ar, FNetworkGUID NetGUID
 		// Serialize Name of object
 		Ar << ObjectPathName;
 
+		if ( bIsPackage )
+		{
+			FGuid PackageGuid = CastChecked< const UPackage >( Object )->GetGuid();
+			Ar << PackageGuid;
+		}
+
 		if ( GuidCache->IsExportingNetGUIDBunch )
 		{
 			CurrentExportNetGUIDs.Add( NetGUID );
@@ -536,8 +544,39 @@ void UPackageMapClient::InternalWriteObject( FArchive & Ar, FNetworkGUID NetGUID
 //
 //--------------------------------------------------------------------
 
-static void SanityCheckExport( const UObject * Object, const FNetworkGUID & NetGUID, const FString & ExpectedPathName, const UObject * ExpectedOuter, const FExportFlags & ExportFlags )
+static void SanityCheckExport( 
+	const FNetGUIDCache *	GuidCache,
+	const UObject *			Object, 
+	const FNetworkGUID &	NetGUID, 
+	const FString &			ExpectedPathName, 
+	const FGuid &			ExpectedPackageGuid, 
+	const UObject *			ExpectedOuter, 
+	const FNetworkGUID &	ExpectedOuterGUID,
+	const FExportFlags &	ExportFlags )
 {
+	
+	check( GuidCache != NULL );
+	check( Object != NULL );
+
+	const FNetGuidCacheObject * CacheObject = GuidCache->ObjectLookup.Find( NetGUID );
+
+	if ( CacheObject != NULL )
+	{
+		if ( CacheObject->OuterGUID != ExpectedOuterGUID )
+		{
+			UE_LOG( LogNetPackageMap, Warning, TEXT( "SanityCheckExport: CacheObject->OuterGUID != ExpectedOuterGUID. NetGUID: %s, Object: %s, Expected: %s" ), *NetGUID.ToString(), *Object->GetPathName(), *ExpectedPathName );
+		}
+
+		if ( CacheObject->PackageGuid != ExpectedPackageGuid )
+		{
+			UE_LOG( LogNetPackageMap, Warning, TEXT( "SanityCheckExport: CacheObject->PackageGuid != ExpectedPackageGuid. NetGUID: %s, Object: %s, Expected: %s" ), *NetGUID.ToString(), *Object->GetPathName(), *ExpectedPathName );
+		}
+	}
+	else
+	{
+		UE_LOG( LogNetPackageMap, Warning, TEXT( "SanityCheckExport: CacheObject == NULL. NetGUID: %s, Object: %s, Expected: %s" ), *NetGUID.ToString(), *Object->GetPathName(), *ExpectedPathName );
+	}
+
 	if ( Object->GetName() != ExpectedPathName )
 	{
 		UE_LOG( LogNetPackageMap, Warning, TEXT( "SanityCheckExport: Name mismatch. NetGUID: %s, Object: %s, Expected: %s" ), *NetGUID.ToString(), *Object->GetPathName(), *ExpectedPathName );
@@ -550,14 +589,12 @@ static void SanityCheckExport( const UObject * Object, const FNetworkGUID & NetG
 		UE_LOG( LogNetPackageMap, Warning, TEXT( "SanityCheckExport: Outer mismatch. Object: %s, NetGUID: %s, Current: %s, Expected: %s" ), *Object->GetPathName(), *NetGUID.ToString(), *CurrentOuterName, *ExpectedOuterName );
 	}
 
-	if ( ExportFlags.bIsPackage )
-	{
-		const UPackage * Package = Cast< const UPackage >( Object );
+	const bool bIsPackage = ( NetGUID.IsStatic() && Object != NULL && Object->GetOuter() == NULL );
+	const UPackage * Package = Cast< const UPackage >( Object );
 
-		if ( Package == NULL )
-		{
-			UE_LOG( LogNetPackageMap, Warning, TEXT( "SanityCheckExport: Expected package. Object:%s, NetGUID: %s" ), *Object->GetPathName(), *NetGUID.ToString() );
-		}
+	if ( bIsPackage != ( Package != NULL ) )
+	{
+		UE_LOG( LogNetPackageMap, Warning, TEXT( "SanityCheckExport: Package type mismatch. Object:%s, NetGUID: %s" ), *Object->GetPathName(), *NetGUID.ToString() );
 	}
 }
 
@@ -628,8 +665,17 @@ FNetworkGUID UPackageMapClient::InternalLoadObject( FArchive & Ar, UObject *& Ob
 		FNetworkGUID OuterGUID = InternalLoadObject( Ar, ObjOuter, InternalLoadObjectRecursionCount + 1 );
 
 		FString PathName;
+		FGuid	PackageGuid;
+
 		Ar << PathName;
 
+		const bool bIsPackage = NetGUID.IsStatic() && !OuterGUID.IsValid();
+
+		if ( bIsPackage )
+		{
+			Ar << PackageGuid;
+		}
+
 		if ( Ar.IsError() )
 		{
 			UE_LOG( LogNetPackageMap, Error, TEXT( "InternalLoadObject: Failed to load path name" ) );
@@ -643,7 +689,7 @@ FNetworkGUID UPackageMapClient::InternalLoadObject( FArchive & Ar, UObject *& Ob
 		if ( Object != NULL )
 		{
 			// If we already have the object, just do some sanity checking and return
-			SanityCheckExport( Object, NetGUID, PathName, ObjOuter, ExportFlags );
+			SanityCheckExport( GuidCache.Get(), Object, NetGUID, PathName, PackageGuid, ObjOuter, OuterGUID, ExportFlags );
 			return NetGUID;
 		}
 
@@ -669,6 +715,25 @@ FNetworkGUID UPackageMapClient::InternalLoadObject( FArchive & Ar, UObject *& Ob
 				return NetGUID;
 			}
 
+			if ( bIsPackage )
+			{
+				UPackage * Package = Cast< UPackage >( Object );
+
+				if ( Package == NULL )
+				{
+					UE_LOG( LogNetPackageMap, Error, TEXT( "UPackageMapClient::InternalLoadObject: Default object not a package from client: PathName: %s, ObjOuter: %s " ), *PathName, ObjOuter != NULL ? *ObjOuter->GetPathName() : TEXT( "NULL" ) );
+					Object = NULL;
+					return NetGUID;
+				}
+
+				if ( Package->GetGuid() != PackageGuid )
+				{
+					UE_LOG( LogNetPackageMap, Error, TEXT( "UPackageMapClient::InternalLoadObject: Default object package guid mismatch! PathName: %s, ObjOuter: %s " ), *PathName, ObjOuter != NULL ? *ObjOuter->GetPathName() : TEXT( "NULL" ) );
+					Object = NULL;
+					return NetGUID;
+				}
+			}
+
 			// Assign the guid to the object
 			NetGUID = GuidCache->GetOrAssignNetGUID( Object );
 
@@ -692,7 +757,7 @@ FNetworkGUID UPackageMapClient::InternalLoadObject( FArchive & Ar, UObject *& Ob
 		const bool bIgnoreWhenMissing = ExportFlags.bNoLoad;
 
 		// Register this path and outer guid combo with the net guid
-		GuidCache->RegisterNetGUIDFromPath_Client( NetGUID, PathName, OuterGUID, ExportFlags.bNoLoad, bIgnoreWhenMissing );
+		GuidCache->RegisterNetGUIDFromPath_Client( NetGUID, PathName, OuterGUID, PackageGuid, ExportFlags.bNoLoad, bIgnoreWhenMissing );
 
 		// Try again now that we've registered the path
 		Object = GuidCache->GetObjectFromNetGUID( NetGUID, GuidCache->IsExportingNetGUIDBunch );
@@ -1707,7 +1772,7 @@ void FNetGUIDCache::RegisterNetGUID_Client( const FNetworkGUID & NetGUID, const
  *	Associates a net guid with a path, that can be loaded or found later
  *  This function is only called on the client
  */
-void FNetGUIDCache::RegisterNetGUIDFromPath_Client( const FNetworkGUID & NetGUID, const FString & PathName, const FNetworkGUID & OuterGUID, const bool bNoLoad, const bool bIgnoreWhenMissing )
+void FNetGUIDCache::RegisterNetGUIDFromPath_Client( const FNetworkGUID & NetGUID, const FString & PathName, const FNetworkGUID & OuterGUID, const FGuid & PackageGuid, const bool bNoLoad, const bool bIgnoreWhenMissing )
 {
 	check( !IsNetGUIDAuthority() );		// Server never calls this locally
 	check( !NetGUID.IsDefault() );
@@ -1748,6 +1813,7 @@ void FNetGUIDCache::RegisterNetGUIDFromPath_Client( const FNetworkGUID & NetGUID
 
 	CacheObject.PathName			= FName( *PathName );
 	CacheObject.OuterGUID			= OuterGUID;
+	CacheObject.PackageGuid			= PackageGuid;
 	CacheObject.bNoLoad				= bNoLoad;
 	CacheObject.bIgnoreWhenMissing	= bIgnoreWhenMissing;
 	CacheObject.GuidSequence		= GuidSequence;
@@ -1767,7 +1833,7 @@ void FNetGUIDCache::AsyncPackageCallback( const FString & PackageName, UPackage
 
 	if ( !NetGUID.IsValid() )
 	{
-		UE_LOG( LogNetPackageMap, Warning, TEXT( "AsyncPackageCallback: Could not find package. Path: %s" ), *PackageName );
+		UE_LOG( LogNetPackageMap, Error, TEXT( "AsyncPackageCallback: Could not find package. Path: %s" ), *PackageName );
 		return;
 	}
 
@@ -1775,13 +1841,13 @@ void FNetGUIDCache::AsyncPackageCallback( const FString & PackageName, UPackage
 
 	if ( CacheObject == NULL )
 	{
-		UE_LOG( LogNetPackageMap, Warning, TEXT( "AsyncPackageCallback: Could not find net guid. Path: %s, NetGUID: %s" ), *PackageName, *NetGUID.ToString() );
+		UE_LOG( LogNetPackageMap, Error, TEXT( "AsyncPackageCallback: Could not find net guid. Path: %s, NetGUID: %s" ), *PackageName, *NetGUID.ToString() );
 		return;
 	}
 
 	if ( !CacheObject->bIsPending )
 	{
-		UE_LOG( LogNetPackageMap, Warning, TEXT( "AsyncPackageCallback: Package wasn't pending. Path: %s, NetGUID: %s" ), *PackageName, *NetGUID.ToString() );
+		UE_LOG( LogNetPackageMap, Error, TEXT( "AsyncPackageCallback: Package wasn't pending. Path: %s, NetGUID: %s" ), *PackageName, *NetGUID.ToString() );
 	}
 
 	CacheObject->bIsPending = false;
@@ -1789,7 +1855,12 @@ void FNetGUIDCache::AsyncPackageCallback( const FString & PackageName, UPackage
 	if ( Package == NULL )
 	{
 		CacheObject->bIsBroken = true;
-		UE_LOG( LogNetPackageMap, Warning, TEXT( "AsyncPackageCallback: Package FAILED to load. Path: %s, NetGUID: %s" ), *PackageName, *NetGUID.ToString() );
+		UE_LOG( LogNetPackageMap, Error, TEXT( "AsyncPackageCallback: Package FAILED to load. Path: %s, NetGUID: %s" ), *PackageName, *NetGUID.ToString() );
+	}
+	else if ( Package->GetGuid() != CacheObject->PackageGuid )
+	{
+		CacheObject->bIsBroken = true;
+		UE_LOG( LogNetPackageMap, Error, TEXT( "AsyncPackageCallback: Package GUID mismatch! Path: %s, NetGUID: %s" ), *PackageName, *NetGUID.ToString() );
 	}
 }
 
@@ -1999,6 +2070,13 @@ UObject * FNetGUIDCache::GetObjectFromNetGUID( const FNetworkGUID & NetGUID, con
 		return NULL;
 	}
 	
+	if ( bIsPackage && CastChecked< UPackage >( Object )->GetGuid() != CacheObjectPtr->PackageGuid )
+	{
+		CacheObjectPtr->bIsBroken = true;
+		UE_LOG( LogNetPackageMap, Error, TEXT( "GetObjectFromNetGUID: Package GUID mismatch! Path: %s, NetGUID: %s" ), *CacheObjectPtr->PathName.ToString(), *NetGUID.ToString() );
+		return NULL;
+	}
+
 	// Assign the resolved object to this guid
 	CacheObjectPtr->Object = Object;		
 
