Fix spring arm lag jitter. The problem was we were using average delta time instead of real delta time.

diff --git a/Engine/Source/Runtime/Engine/Private/GameFramework/SpringArmComponent.cpp b/Engine/Source/Runtime/Engine/Private/GameFramework/SpringArmComponent.cpp
--- a/Engine/Source/Runtime/Engine/Private/GameFramework/SpringArmComponent.cpp
+++ b/Engine/Source/Runtime/Engine/Private/GameFramework/SpringArmComponent.cpp
@@ -60,7 +60,7 @@ void USpringArmComponent::UpdateDesiredArmLocation(bool bDoTrace, bool bDoLocati
 	// Apply 'lag' to rotation if desired
 	if(bDoRotationLag)
 	{
-		DesiredRot = FMath::RInterpTo(PreviousDesiredRot, DesiredRot, GAverageMS/1000.f, CameraRotationLagSpeed);
+		DesiredRot = FMath::RInterpTo(PreviousDesiredRot, DesiredRot, DeltaTime, CameraRotationLagSpeed);
 	}
 	PreviousDesiredRot = DesiredRot;
 
@@ -70,7 +70,7 @@ void USpringArmComponent::UpdateDesiredArmLocation(bool bDoTrace, bool bDoLocati
 	FVector DesiredLoc = ArmOrigin;
 	if (bDoLocationLag)
 	{
-		DesiredLoc = FMath::VInterpTo(PreviousDesiredLoc, DesiredLoc, GAverageMS/1000.f, CameraLagSpeed);
+		DesiredLoc = FMath::VInterpTo(PreviousDesiredLoc, DesiredLoc, DeltaTime, CameraLagSpeed);
 	}
 	PreviousDesiredLoc = DesiredLoc;
 
