Unregister destructible components from tick or post physics tick

diff --git a/Engine/Source/Runtime/Engine/Private/DestructibleComponent.cpp b/Engine/Source/Runtime/Engine/Private/DestructibleComponent.cpp
--- a/Engine/Source/Runtime/Engine/Private/DestructibleComponent.cpp
+++ b/Engine/Source/Runtime/Engine/Private/DestructibleComponent.cpp
@@ -19,7 +19,9 @@ UDestructibleComponent::UDestructibleComponent(const class FPostConstructInitial
 	, PhysxUserData(this)
 #endif
 {
-	PostPhysicsComponentTick.bCanEverTick = true;
+
+	PrimaryComponentTick.bCanEverTick = false;
+	PostPhysicsComponentTick.bCanEverTick = false;
 
 	bHasCustomNavigableGeometry = EHasCustomNavigableGeometry::EvenIfNotCollidable;
 
