#ue4 - Projectile movement should not adjust out of penetration if it does not bounce (so don't use SafeMoveUpdatedComponent).

- If not bouncing and initially penetrating, this should cause a hit event that explodes. Otherwise the projectile will pop out of the object (potentially very far) and then start moving. It is up to the spawner to correctly use the bNoCollisionFail checks.
- Added a new MoveComponent flag(MOVECOMP_NeverIgnoreBlockingOverlaps) that can disable code in ShouldIgnoreHitResult that allows movement out of penetration if the movement is not opposed by the normal.
- UMovementComponent::ResolvePenetration() temporarily disables the above flag when trying to explicitly move out of penetration.

diff --git a/Engine/Source/Runtime/Engine/Classes/Components/SceneComponent.h b/Engine/Source/Runtime/Engine/Classes/Components/SceneComponent.h
--- a/Engine/Source/Runtime/Engine/Classes/Components/SceneComponent.h
+++ b/Engine/Source/Runtime/Engine/Classes/Components/SceneComponent.h
@@ -59,6 +59,7 @@ enum EMoveComponentFlags
 	MOVECOMP_NoFlags					= 0x0000,	// no flags
 	MOVECOMP_IgnoreBases				= 0x0001,	// ignore collisions with things the Actor is based on
 	MOVECOMP_SkipPhysicsMove			= 0x0002,	// when moving this component, do not move the physics representation. Used internally to avoid looping updates when syncing with physics.
+	MOVECOMP_NeverIgnoreBlockingOverlaps= 0x0004,	// never ignore initial blocking overlaps during movement, which are usually ignored when moving out of an object. MOVECOMP_IgnoreBases is still respected.
 };
 
 FORCEINLINE EMoveComponentFlags operator|(EMoveComponentFlags Arg1,EMoveComponentFlags Arg2)	{ return EMoveComponentFlags(uint32(Arg1) | uint32(Arg2)); }
diff --git a/Engine/Source/Runtime/Engine/Private/MovementComponent.cpp b/Engine/Source/Runtime/Engine/Private/MovementComponent.cpp
--- a/Engine/Source/Runtime/Engine/Private/MovementComponent.cpp
+++ b/Engine/Source/Runtime/Engine/Private/MovementComponent.cpp
@@ -340,6 +340,9 @@ bool UMovementComponent::ResolvePenetration(const FVector& ProposedAdjustment, c
 		}
 		else
 		{
+			// Disable MOVECOMP_NeverIgnoreBlockingOverlaps if it is enabled, otherwise we wouldn't be able to sweep out of the object to fix the penetration.
+			TGuardValue<EMoveComponentFlags> ScopedFlagRestore(MoveComponentFlags, EMoveComponentFlags(MoveComponentFlags & (~MOVECOMP_NeverIgnoreBlockingOverlaps)));
+
 			// Try sweeping as far as possible...
 			bool bMoved = MoveUpdatedComponent(Adjustment, NewRotation, true);
 			UE_LOG(LogMovement, Verbose, TEXT("ResolvePenetration: sweep %s by %s (success = %d)"), *ActorOwner->GetName(), *Adjustment.ToString(), bMoved);
diff --git a/Engine/Source/Runtime/Engine/Private/PrimitiveComponent.cpp b/Engine/Source/Runtime/Engine/Private/PrimitiveComponent.cpp
--- a/Engine/Source/Runtime/Engine/Private/PrimitiveComponent.cpp
+++ b/Engine/Source/Runtime/Engine/Private/PrimitiveComponent.cpp
@@ -1159,7 +1159,7 @@ static bool ShouldIgnoreHitResult(const UWorld* InWorld, FHitResult const& TestH
 	
 	// If we started penetrating, we may want to ignore it if we are moving out of penetration.
 	// This helps prevent getting stuck in walls.
-	if ( TestHit.bStartPenetrating && TestHit.bBlockingHit )
+	if ( TestHit.bStartPenetrating && TestHit.bBlockingHit && !(MoveFlags & MOVECOMP_NeverIgnoreBlockingOverlaps) )
 	{
 		const float DotTolerance = CVarInitialOverlapTolerance.GetValueOnGameThread();
 
diff --git a/Engine/Source/Runtime/Engine/Private/ProjectileMovementComponent.cpp b/Engine/Source/Runtime/Engine/Private/ProjectileMovementComponent.cpp
--- a/Engine/Source/Runtime/Engine/Private/ProjectileMovementComponent.cpp
+++ b/Engine/Source/Runtime/Engine/Private/ProjectileMovementComponent.cpp
@@ -113,12 +113,27 @@ void UProjectileMovementComponent::TickComponent(float DeltaTime, enum ELevelTic
 
 		const FVector TmpVelocity = Velocity;
 		const FRotator NewRotation = bRotationFollowsVelocity ? Velocity.Rotation() : ActorOwner->GetActorRotation();
-		SafeMoveUpdatedComponent( MoveDelta, NewRotation, true, Hit );
+
+		// Move the component
+		if (bShouldBounce)
+		{
+			// If we can bounce, we are allowed to move out of penetrations, so use SafeMoveUpdatedComponent which does that automatically.
+			SafeMoveUpdatedComponent( MoveDelta, NewRotation, true, Hit );
+		}
+		else
+		{
+			// If we can't bounce, then we shouldn't adjust if initially penetrating, because that should be a blocking hit that causes a hit event and stop simulation.
+			TGuardValue<EMoveComponentFlags> ScopedFlagRestore(MoveComponentFlags, MoveComponentFlags | MOVECOMP_NeverIgnoreBlockingOverlaps);
+			MoveUpdatedComponent(MoveDelta, NewRotation, true, &Hit );
+		}
+		
+		// If we hit a trigger that destroyed us, abort.
 		if( ActorOwner->IsPendingKill() || !UpdatedComponent )
 		{
 			return;
 		}
 
+		// Handle hit result after movement
 		if (Hit.Time == 1.f)
 		{
 			bSliding = false;
