Fixed static shadow depth map memory stats when building lighting
Fixed changing light intensity unbuilding a Stationary light

diff --git a/Engine/Source/Editor/UnrealEd/Private/Lightmass/Lightmass.cpp b/Engine/Source/Editor/UnrealEd/Private/Lightmass/Lightmass.cpp
--- a/Engine/Source/Editor/UnrealEd/Private/Lightmass/Lightmass.cpp
+++ b/Engine/Source/Editor/UnrealEd/Private/Lightmass/Lightmass.cpp
@@ -3575,6 +3575,9 @@ void FLightmassProcessor::ImportStaticShadowDepthMap(ULightComponent* Light)
 		Lightmass::FStaticShadowDepthMapData ShadowMapData;
 		Swarm.ReadChannel(Channel, &ShadowMapData, sizeof(ShadowMapData));
 
+		BeginReleaseResource(&DepthMap);
+		DepthMap.Empty();
+
 		DepthMap.WorldToLight = ShadowMapData.WorldToLight;
 		DepthMap.ShadowMapSizeX = ShadowMapData.ShadowMapSizeX;
 		DepthMap.ShadowMapSizeY = ShadowMapData.ShadowMapSizeY;
@@ -3582,8 +3585,7 @@ void FLightmassProcessor::ImportStaticShadowDepthMap(ULightComponent* Light)
 		ReadArray(Channel, DepthMap.DepthSamples);
 		Swarm.CloseChannel(Channel);
 
-		BeginReleaseResource(&DepthMap);
-		BeginInitResource(&DepthMap);
+		DepthMap.InitializeAfterImport();
 	}
 	else
 	{
diff --git a/Engine/Source/Runtime/Engine/Classes/Components/LightComponent.h b/Engine/Source/Runtime/Engine/Classes/Components/LightComponent.h
--- a/Engine/Source/Runtime/Engine/Classes/Components/LightComponent.h
+++ b/Engine/Source/Runtime/Engine/Classes/Components/LightComponent.h
@@ -30,7 +30,11 @@ public:
 
 	virtual void InitRHI();
 
-	void Empty();
+	/** Frees the CPU backing of the shadowmap. */
+	ENGINE_API void Empty();
+
+	/** Called after being imported during a lighting build. */
+	ENGINE_API void InitializeAfterImport();
 
 	friend FArchive& operator<<(FArchive& Ar, FStaticShadowDepthMap& ShadowMap);
 };
diff --git a/Engine/Source/Runtime/Engine/Private/LightComponent.cpp b/Engine/Source/Runtime/Engine/Private/LightComponent.cpp
--- a/Engine/Source/Runtime/Engine/Private/LightComponent.cpp
+++ b/Engine/Source/Runtime/Engine/Private/LightComponent.cpp
@@ -43,6 +43,12 @@ void FStaticShadowDepthMap::Empty()
 	DepthSamples.Empty();
 }
 
+void FStaticShadowDepthMap::InitializeAfterImport()
+{
+	INC_DWORD_STAT_BY(STAT_PrecomputedShadowDepthMapMemory, DepthSamples.GetAllocatedSize());
+	BeginInitResource(this);
+}
+
 FArchive& operator<<(FArchive& Ar, FStaticShadowDepthMap& ShadowMap)
 {
 	Ar << ShadowMap.WorldToLight;
@@ -526,8 +532,8 @@ void ULightComponent::PostEditChangeProperty(FPropertyChangedEvent& PropertyChan
 		PropertyName != GET_MEMBER_NAME_STRING_CHECKED(UDirectionalLightComponent, OcclusionDepthRange) &&
 		PropertyName != GET_MEMBER_NAME_STRING_CHECKED(UDirectionalLightComponent, LightShaftOverrideDirection) &&
 		// Properties that should only unbuild lighting for a Static light (can be changed dynamically on a Stationary light)
-		(PropertyName != GET_MEMBER_NAME_STRING_CHECKED(ULightComponent, BloomScale) || Mobility == EComponentMobility::Static) &&
-		(PropertyName != GET_MEMBER_NAME_STRING_CHECKED(ULightComponent, BloomScale) || Mobility == EComponentMobility::Static))
+		(PropertyName != GET_MEMBER_NAME_STRING_CHECKED(ULightComponent, Intensity) || Mobility == EComponentMobility::Static) &&
+		(PropertyName != GET_MEMBER_NAME_STRING_CHECKED(ULightComponent, LightColor) || Mobility == EComponentMobility::Static))
 	{
 		InvalidateLightingCache();
 	}
